ElectraPlayer: RELEASE. Moved from NFL to public. Note that this is done via a Delete and "Add New" operation as discussed.

#rnx

diff --git a/Engine/Plugins/Media/ElectraPlayer/ElectraPlayer.uplugin b/Engine/Plugins/Media/ElectraPlayer/ElectraPlayer.uplugin
new file mode 100644
--- /dev/null
+++ b/Engine/Plugins/Media/ElectraPlayer/ElectraPlayer.uplugin
@@ -0,0 +1,53 @@
+{
+	"FileVersion": 1,
+	"Version": 0,
+	"VersionName": "0.1",
+	"FriendlyName": "Electra Player",
+	"Description": "Next-generation playback capability.",
+	"Category": "Media Players",
+	"CreatedBy": "Epic Games, Inc.",
+	"CreatedByURL": "http://epicgames.com",
+	"DocsURL": "https://docs.unrealengine.com/en-US/Engine/MediaFramework/Overview",
+	"MarketplaceURL": "",
+	"SupportURL": "",
+	"EnabledByDefault": false,
+	"CanContainContent": false,
+	"IsBetaVersion": true,
+	"Installed": false,
+	"Modules": [
+		{
+			"Name": "ElectraPlayerRuntime",
+			"Type": "RuntimeNoCommandlet",
+			"LoadingPhase": "PreLoadingScreen",
+			"WhitelistPlatforms": [ "Win64", "Mac", "IOS", "Android" ],
+			"BlacklistTargets": [ "Server" ]
+		},
+		{
+			"Name": "ElectraPlayerPlugin",
+			"Type": "RuntimeNoCommandlet",
+			"LoadingPhase": "PreLoadingScreen",
+			"WhitelistPlatforms": [ "Win64", "Mac", "IOS", "Android" ],
+			"BlacklistTargets": [ "Server" ]
+		},
+		{
+			"Name": "ElectraPlayerFactory",
+			"Type": "Editor",
+			"LoadingPhase": "PostEngineInit",
+			"WhitelistPlatforms": [ "Win64", "Mac" ],
+			"BlacklistTargets": [ "Server" ]
+		},
+		{
+			"Name": "ElectraPlayerFactory",
+			"Type": "RuntimeNoCommandlet",
+			"LoadingPhase": "PostEngineInit",
+			"WhitelistPlatforms": [ "Win64", "Mac", "IOS", "Android" ],
+			"BlacklistTargets": [ "Server" ]
+		}
+	],
+	"Plugins": [
+		{
+			"Name": "ElectraUtil",
+			"Enabled": true
+		}
+	]
+}
\ No newline at end of file
diff --git a/Engine/Plugins/Media/ElectraPlayer/Source/ElectraPlayerPlugin/ElectraPlayerPlugin.Build.cs b/Engine/Plugins/Media/ElectraPlayer/Source/ElectraPlayerPlugin/ElectraPlayerPlugin.Build.cs
new file mode 100644
--- /dev/null
+++ b/Engine/Plugins/Media/ElectraPlayer/Source/ElectraPlayerPlugin/ElectraPlayerPlugin.Build.cs
@@ -0,0 +1,59 @@
+// Copyright Epic Games, Inc. All Rights Reserved.
+using EpicGames.Core;
+using System.IO;
+
+namespace UnrealBuildTool.Rules
+{
+	public class ElectraPlayerPlugin: ModuleRules
+	{
+		public ElectraPlayerPlugin(ReadOnlyTargetRules Target) : base(Target)
+		{
+
+            bLegalToDistributeObjectCode = true;
+            PCHUsage = PCHUsageMode.NoPCHs;
+
+            DynamicallyLoadedModuleNames.AddRange(
+				new string[] {
+					"Media",
+				});
+
+			PrivateDependencyModuleNames.AddRange(
+				new string[] {
+					"Core",
+					"CoreUObject",
+					"MediaUtils",
+					"RHI",
+					"AnalyticsET",
+					"ElectraPlayerFactory",
+					"ElectraPlayerRuntime",
+					"ElectraSamples",
+					"ElectraBase"
+				});
+
+			PrivateIncludePathModuleNames.AddRange(
+			new string[] {
+				"Media",
+            });
+
+			if (Target.bCompileAgainstEngine)
+			{
+				PrivateDependencyModuleNames.Add("Engine");
+			}
+
+			if (Target.Platform == UnrealTargetPlatform.Mac)
+            {
+                PrivateDependencyModuleNames.Add("MetalRHI");
+            }
+
+			if (Target.Platform == UnrealTargetPlatform.IOS )
+            {
+                PrivateDependencyModuleNames.Add("MetalRHI");
+            }
+
+			if (Target.Type == TargetType.Editor)
+			{
+				PrivateDependencyModuleNames.Add("DeveloperSettings");
+			}
+		}
+	}
+}
diff --git a/Engine/Plugins/Media/ElectraPlayer/Source/ElectraPlayerPlugin/Private/ElectraPlayerAudioSample.cpp b/Engine/Plugins/Media/ElectraPlayer/Source/ElectraPlayerPlugin/Private/ElectraPlayerAudioSample.cpp
new file mode 100644
--- /dev/null
+++ b/Engine/Plugins/Media/ElectraPlayer/Source/ElectraPlayerPlugin/Private/ElectraPlayerAudioSample.cpp
@@ -0,0 +1,13 @@
+// Copyright Epic Games, Inc. All Rights Reserved.
+
+#include "ElectraPlayerAudioSample.h"
+
+void FElectraPlayerAudioSample::InitializePoolable()
+{
+}
+
+void FElectraPlayerAudioSample::ShutdownPoolable()
+{
+	DecoderOutput.Reset();
+}
+
diff --git a/Engine/Plugins/Media/ElectraPlayer/Source/ElectraPlayerPlugin/Private/ElectraPlayerAudioSample.h b/Engine/Plugins/Media/ElectraPlayer/Source/ElectraPlayerPlugin/Private/ElectraPlayerAudioSample.h
new file mode 100644
--- /dev/null
+++ b/Engine/Plugins/Media/ElectraPlayer/Source/ElectraPlayerPlugin/Private/ElectraPlayerAudioSample.h
@@ -0,0 +1,83 @@
+// Copyright Epic Games, Inc. All Rights Reserved.
+#pragma once
+
+#include "IMediaAudioSample.h"
+#include "MediaObjectPool.h"
+#include "Misc/Timespan.h"
+#include "Templates/SharedPointer.h"
+
+#include "MediaAudioDecoderOutput.h"
+
+ 
+// IMediaAudioSample impl
+// is used as audio decoder output
+// contains decoded audio samples ready for playback along with corresponding metadata
+// owns binary data so can be cached
+class FElectraPlayerAudioSample :
+	public IMediaAudioSample,
+	public IMediaPoolable
+{
+public:
+
+	FElectraPlayerAudioSample()
+	{
+	}
+
+	~FElectraPlayerAudioSample()
+	{
+	}
+
+	void Initialize(const IAudioDecoderOutputPtr & InDecoderOutput)
+	{
+		DecoderOutput = InDecoderOutput;
+	}
+
+	const void* GetBuffer() override
+	{ return DecoderOutput ? DecoderOutput->GetBuffer() : nullptr; }
+
+	uint32 GetChannels() const override
+	{ return DecoderOutput ? DecoderOutput->GetChannels() : 0; }
+
+	FTimespan GetDuration() const override
+	{ return DecoderOutput ? DecoderOutput->GetDuration() : FTimespan::Zero(); }
+
+	EMediaAudioSampleFormat GetFormat() const override
+	{ 
+		static_assert((int32)EMediaAudioSampleFormat::Undefined == (int32)IAudioDecoderOutput::ESampleFormat::Undefined, "check enums are identical");
+		static_assert((int32)EMediaAudioSampleFormat::Double == (int32)IAudioDecoderOutput::ESampleFormat::Double, "check enums are identical");
+		static_assert((int32)EMediaAudioSampleFormat::Float == (int32)IAudioDecoderOutput::ESampleFormat::Float, "check enums are identical");
+		static_assert((int32)EMediaAudioSampleFormat::Int8 == (int32)IAudioDecoderOutput::ESampleFormat::Int8, "check enums are identical");
+		static_assert((int32)EMediaAudioSampleFormat::Int16 == (int32)IAudioDecoderOutput::ESampleFormat::Int16, "check enums are identical");
+		static_assert((int32)EMediaAudioSampleFormat::Int32 == (int32)IAudioDecoderOutput::ESampleFormat::Int32, "check enums are identical");
+
+		return DecoderOutput ? (EMediaAudioSampleFormat)DecoderOutput->GetFormat() : EMediaAudioSampleFormat::Undefined;
+	}
+
+	uint32 GetFrames() const override
+	{ return DecoderOutput ? DecoderOutput->GetFrames() : 0; }
+
+	uint32 GetSampleRate() const override
+	{ return DecoderOutput ? DecoderOutput->GetSampleRate() : 0; }
+
+	FMediaTimeStamp GetTime() const override
+	{
+		if (!DecoderOutput)
+		{
+			return FMediaTimeStamp();
+		}
+		FDecoderTimeStamp TimeStamp = DecoderOutput->GetTime();
+		return FMediaTimeStamp(TimeStamp.Time, TimeStamp.SequenceIndex);
+	}
+
+	uint32 GetMaxBufferBytes() const
+	{ return DecoderOutput ? DecoderOutput->GetReservedBufferBytes() : 0; }
+	
+	virtual void InitializePoolable() override;
+	virtual void ShutdownPoolable() override;
+
+private:
+	IAudioDecoderOutputPtr DecoderOutput;
+};
+
+using FElectraPlayerAudioSamplePtr = TSharedPtr<FElectraPlayerAudioSample, ESPMode::ThreadSafe>;
+using FElectraPlayerAudioSamplePool = TMediaObjectPool<FElectraPlayerAudioSample>;
diff --git a/Engine/Plugins/Media/ElectraPlayer/Source/ElectraPlayerPlugin/Private/ElectraPlayerPlugin.cpp b/Engine/Plugins/Media/ElectraPlayer/Source/ElectraPlayerPlugin/Private/ElectraPlayerPlugin.cpp
new file mode 100644
--- /dev/null
+++ b/Engine/Plugins/Media/ElectraPlayer/Source/ElectraPlayerPlugin/Private/ElectraPlayerPlugin.cpp
@@ -0,0 +1,687 @@
+// Copyright Epic Games, Inc. All Rights Reserved.
+
+#include "Misc/Optional.h"
+#include "IMediaEventSink.h"
+#include "IMediaOptions.h"
+#include "MediaSamples.h"
+#include "MediaPlayerOptions.h"
+
+#include "ElectraPlayerPlugin.h"
+#include "IElectraPlayerRuntimeModule.h"
+#include "IElectraPlayerPluginModule.h"
+
+using namespace Electra;
+
+//-----------------------------------------------------------------------------
+
+FElectraPlayerPlugin::FElectraPlayerPlugin()
+	: OptionInterface(nullptr)
+	, EventSink(nullptr)
+{
+	// Make sure a few assumptions are correct...
+	static_assert((int32)EMediaEvent::MediaBuffering == (int32)IElectraPlayerAdapterDelegate::EPlayerEvent::MediaBuffering, "check alignment of both enums");
+	static_assert((int32)EMediaEvent::MediaClosed == (int32)IElectraPlayerAdapterDelegate::EPlayerEvent::MediaClosed, "check alignment of both enums");
+	static_assert((int32)EMediaEvent::MediaConnecting == (int32)IElectraPlayerAdapterDelegate::EPlayerEvent::MediaConnecting, "check alignment of both enums");
+	static_assert((int32)EMediaEvent::MediaOpened == (int32)IElectraPlayerAdapterDelegate::EPlayerEvent::MediaOpened, "check alignment of both enums");
+	static_assert((int32)EMediaEvent::MediaOpenFailed == (int32)IElectraPlayerAdapterDelegate::EPlayerEvent::MediaOpenFailed, "check alignment of both enums");
+	static_assert((int32)EMediaEvent::PlaybackEndReached == (int32)IElectraPlayerAdapterDelegate::EPlayerEvent::PlaybackEndReached, "check alignment of both enums");
+	static_assert((int32)EMediaEvent::PlaybackResumed == (int32)IElectraPlayerAdapterDelegate::EPlayerEvent::PlaybackResumed, "check alignment of both enums");
+	static_assert((int32)EMediaEvent::PlaybackSuspended == (int32)IElectraPlayerAdapterDelegate::EPlayerEvent::PlaybackSuspended, "check alignment of both enums");
+	static_assert((int32)EMediaEvent::SeekCompleted == (int32)IElectraPlayerAdapterDelegate::EPlayerEvent::SeekCompleted, "check alignment of both enums");
+	static_assert((int32)EMediaEvent::TracksChanged == (int32)IElectraPlayerAdapterDelegate::EPlayerEvent::TracksChanged, "check alignment of both enums");
+	static_assert((int32)EMediaEvent::Internal_PurgeVideoSamplesHint == (int32)IElectraPlayerAdapterDelegate::EPlayerEvent::Internal_PurgeVideoSamplesHint, "check alignment of both enums");
+	static_assert((int32)EMediaEvent::Internal_ResetForDiscontinuity == (int32)IElectraPlayerAdapterDelegate::EPlayerEvent::Internal_ResetForDiscontinuity, "check alignment of both enums");
+	static_assert((int32)EMediaEvent::Internal_RenderClockStart == (int32)IElectraPlayerAdapterDelegate::EPlayerEvent::Internal_RenderClockStart, "check alignment of both enums");
+	static_assert((int32)EMediaEvent::Internal_RenderClockStop == (int32)IElectraPlayerAdapterDelegate::EPlayerEvent::Internal_RenderClockStop, "check alignment of both enums");
+	static_assert((int32)EMediaEvent::Internal_VideoSamplesAvailable == (int32)IElectraPlayerAdapterDelegate::EPlayerEvent::Internal_VideoSamplesAvailable, "check alignment of both enums");
+	static_assert((int32)EMediaEvent::Internal_VideoSamplesUnavailable == (int32)IElectraPlayerAdapterDelegate::EPlayerEvent::Internal_VideoSamplesUnavailable, "check alignment of both enums");
+	static_assert((int32)EMediaEvent::Internal_AudioSamplesAvailable == (int32)IElectraPlayerAdapterDelegate::EPlayerEvent::Internal_AudioSamplesAvailable, "check alignment of both enums");
+	static_assert((int32)EMediaEvent::Internal_AudioSamplesUnavailable == (int32)IElectraPlayerAdapterDelegate::EPlayerEvent::Internal_AudioSamplesUnavailable, "check alignment of both enums");
+
+	static_assert((int32)EMediaState::Closed == (int32)IElectraPlayerInterface::EPlayerState::Closed, "check alignment of both enums");
+	static_assert((int32)EMediaState::Error == (int32)IElectraPlayerInterface::EPlayerState::Error, "check alignment of both enums");
+	static_assert((int32)EMediaState::Paused == (int32)IElectraPlayerInterface::EPlayerState::Paused, "check alignment of both enums");
+	static_assert((int32)EMediaState::Playing == (int32)IElectraPlayerInterface::EPlayerState::Playing, "check alignment of both enums");
+	static_assert((int32)EMediaState::Preparing == (int32)IElectraPlayerInterface::EPlayerState::Preparing, "check alignment of both enums");
+	static_assert((int32)EMediaState::Stopped == (int32)IElectraPlayerInterface::EPlayerState::Stopped, "check alignment of both enums");
+
+	static_assert((int32)EMediaStatus::None == (int32)IElectraPlayerInterface::EPlayerStatus::None, "check alignment of both enums");
+	static_assert((int32)EMediaStatus::Buffering == (int32)IElectraPlayerInterface::EPlayerStatus::Buffering, "check alignment of both enums");
+	static_assert((int32)EMediaStatus::Connecting == (int32)IElectraPlayerInterface::EPlayerStatus::Connecting, "check alignment of both enums");
+
+	static_assert((int32)EMediaTrackType::Audio == (int32)IElectraPlayerInterface::EPlayerTrackType::Audio, "check alignment of both enums");
+	static_assert((int32)EMediaTrackType::Video == (int32)IElectraPlayerInterface::EPlayerTrackType::Video, "check alignment of both enums");
+
+	static_assert(IMediaPlayerLifecycleManagerDelegate::ResourceFlags_Decoder == IElectraPlayerInterface::ResourceFlags_Decoder, "check alignment of both enums");
+	static_assert(IMediaPlayerLifecycleManagerDelegate::ResourceFlags_OutputBuffers == IElectraPlayerInterface::ResourceFlags_OutputBuffers, "check alignment of both enums");
+	static_assert(IMediaPlayerLifecycleManagerDelegate::ResourceFlags_Any == IElectraPlayerInterface::ResourceFlags_Any, "check alignment of both enums");
+	static_assert(IMediaPlayerLifecycleManagerDelegate::ResourceFlags_All == IElectraPlayerInterface::ResourceFlags_All, "check alignment of both enums");
+}
+
+bool FElectraPlayerPlugin::Initialize(IMediaEventSink& InEventSink,
+	FElectraPlayerSendAnalyticMetricsDelegate& InSendAnalyticMetricsDelegate,
+	FElectraPlayerSendAnalyticMetricsPerMinuteDelegate& InSendAnalyticMetricsPerMinuteDelegate,
+	FElectraPlayerReportVideoStreamingErrorDelegate& InReportVideoStreamingErrorDelegate,
+	FElectraPlayerReportSubtitlesMetricsDelegate& InReportSubtitlesFileMetricsDelegate)
+{
+	EventSink = &InEventSink;
+
+	MediaSamples.Reset(new FMediaSamples);
+
+	PlayerResourceDelegate.Reset(PlatformCreatePlayerResourceDelegate());
+
+	PlayerDelegate = MakeShareable(new FPlayerAdapterDelegate(AsShared()));
+	Player = MakeShareable(FElectraPlayerRuntimeFactory::CreatePlayer(PlayerDelegate, InSendAnalyticMetricsDelegate, InSendAnalyticMetricsPerMinuteDelegate, InReportVideoStreamingErrorDelegate, InReportSubtitlesFileMetricsDelegate));
+
+	return true;
+}
+
+FElectraPlayerPlugin::~FElectraPlayerPlugin()
+{
+}
+
+//-----------------------------------------------------------------------------
+
+Electra::FVariantValue FElectraPlayerPlugin::FPlayerAdapterDelegate::QueryOptions(EOptionType Type, const Electra::FVariantValue & Param)
+{
+	if (TSharedPtr<FElectraPlayerPlugin, ESPMode::ThreadSafe> PinnedHost = Host.Pin())
+	{
+		switch (Type)
+		{
+		case	EOptionType::MaxVerticalStreamResolution:
+			if (PinnedHost->OptionInterface)
+			{
+				return FVariantValue((int64)PinnedHost->OptionInterface->GetMediaOption(TEXT("MaxResolutionForMediaStreaming"), (int64)0));
+			}
+			break;
+
+		case	EOptionType::MaxBandwidthForStreaming:
+			if (PinnedHost->OptionInterface)
+			{
+				return FVariantValue((int64)PinnedHost->OptionInterface->GetMediaOption(TEXT("ElectraMaxStreamingBandwidth"), (int64)0));
+			}
+			break;
+
+		case EOptionType::PlayListData:
+			if (PinnedHost->OptionInterface)
+			{
+				const FName PlaylistOptionKey = TEXT("ElectraGetPlaylistData");
+				if (PinnedHost->OptionInterface->HasMediaOption(PlaylistOptionKey))
+				{
+					check(Param.IsType(FVariantValue::EDataType::TypeFString));
+					return FVariantValue(PinnedHost->OptionInterface->GetMediaOption(PlaylistOptionKey, Param.GetFString()));
+				}
+			}
+			break;
+
+		case EOptionType::LicenseKeyData:
+			if (PinnedHost->OptionInterface)
+			{
+				const FName LicenseKeyDataOptionKey = TEXT("ElectraGetLicenseKeyData");
+				if (PinnedHost->OptionInterface->HasMediaOption(LicenseKeyDataOptionKey))
+				{
+					check(Param.IsType(FVariantValue::EDataType::TypeFString));
+					return FVariantValue(PinnedHost->OptionInterface->GetMediaOption(LicenseKeyDataOptionKey, Param.GetFString()));
+				}
+			}
+			break;
+
+		case EOptionType::PlaystartPosFromSeekPositions:
+			if (PinnedHost->OptionInterface)
+			{
+				const FName PlaystartOptionKey = TEXT("ElectraGetPlaystartPosFromSeekPositions");
+				if (PinnedHost->OptionInterface->HasMediaOption(PlaystartOptionKey))
+				{
+					check(Param.IsType(FVariantValue::EDataType::TypeSharedPointer));
+
+					TSharedPtr<TArray<FTimespan>, ESPMode::ThreadSafe> PosArray = Param.GetSharedPointer<TArray<FTimespan>>();
+					if (PosArray.IsValid())
+					{
+						TSharedPtr<FElectraSeekablePositions, ESPMode::ThreadSafe> Res = StaticCastSharedPtr<FElectraSeekablePositions, IMediaOptions::FDataContainer, ESPMode::ThreadSafe>(PinnedHost->OptionInterface->GetMediaOption(PlaystartOptionKey, MakeShared<FElectraSeekablePositions, ESPMode::ThreadSafe>(*PosArray)));
+						if (Res.IsValid() && Res->Data.Num())
+						{
+							return FVariantValue(int64(Res->Data[0].GetTicks())); // return HNS
+						}
+					}
+					return FVariantValue();
+				}
+			}
+			break;
+
+		default:
+			break;
+		}
+	}
+	return FVariantValue();
+}
+
+
+void FElectraPlayerPlugin::FPlayerAdapterDelegate::SendMediaEvent(EPlayerEvent Event)
+{
+	if (TSharedPtr<FElectraPlayerPlugin, ESPMode::ThreadSafe> PinnedHost = Host.Pin())
+	{
+		if (PinnedHost->EventSink)
+		{
+			PinnedHost->EventSink->ReceiveMediaEvent((EMediaEvent)Event);
+		}
+	}
+}
+
+
+void FElectraPlayerPlugin::FPlayerAdapterDelegate::OnVideoFlush()
+{
+	if (TSharedPtr<FElectraPlayerPlugin, ESPMode::ThreadSafe> PinnedHost = Host.Pin())
+	{
+		TRange<FTimespan> AllTime(FTimespan::MinValue(), FTimespan::MaxValue());
+		TSharedPtr<IMediaTextureSample, ESPMode::ThreadSafe> FlushSample;
+		while (PinnedHost->GetSamples().FetchVideo(AllTime, FlushSample))
+			;
+	}
+}
+
+
+void FElectraPlayerPlugin::FPlayerAdapterDelegate::OnAudioFlush()
+{
+	if (TSharedPtr<FElectraPlayerPlugin, ESPMode::ThreadSafe> PinnedHost = Host.Pin())
+	{
+		TRange<FTimespan> AllTime(FTimespan::MinValue(), FTimespan::MaxValue());
+		TSharedPtr<IMediaAudioSample, ESPMode::ThreadSafe> FlushSample;
+		while (PinnedHost->GetSamples().FetchAudio(AllTime, FlushSample))
+			;
+	}
+}
+
+
+void FElectraPlayerPlugin::FPlayerAdapterDelegate::PresentVideoFrame(const FVideoDecoderOutputPtr & InVideoFrame)
+{
+	if (TSharedPtr<FElectraPlayerPlugin, ESPMode::ThreadSafe> PinnedHost = Host.Pin())
+	{
+		FElectraTextureSampleRef TextureSample = PinnedHost->OutputTexturePool.AcquireShared();
+		TextureSample->Initialize(InVideoFrame.Get());
+		PinnedHost->MediaSamples->AddVideo(TextureSample);
+	}
+}
+
+
+void FElectraPlayerPlugin::FPlayerAdapterDelegate::PresentAudioFrame(const IAudioDecoderOutputPtr& InAudioFrame)
+{
+	if (TSharedPtr<FElectraPlayerPlugin, ESPMode::ThreadSafe> PinnedHost = Host.Pin())
+	{
+		TSharedRef<FElectraPlayerAudioSample, ESPMode::ThreadSafe> AudioSample = PinnedHost->OutputAudioPool.AcquireShared();
+		AudioSample->Initialize(InAudioFrame);
+		PinnedHost->MediaSamples->AddAudio(AudioSample);
+	}
+}
+
+
+bool FElectraPlayerPlugin::FPlayerAdapterDelegate::CanReceiveVideoSamples(int32 NumFrames)
+{
+	if (TSharedPtr<FElectraPlayerPlugin, ESPMode::ThreadSafe> PinnedHost = Host.Pin())
+	{
+		return PinnedHost->MediaSamples->CanReceiveVideoSamples(NumFrames);
+	}
+	return false;
+}
+
+
+bool FElectraPlayerPlugin::FPlayerAdapterDelegate::CanReceiveAudioSamples(int32 NumFrames)
+{
+	if (TSharedPtr<FElectraPlayerPlugin, ESPMode::ThreadSafe> PinnedHost = Host.Pin())
+	{
+		return PinnedHost->MediaSamples->CanReceiveAudioSamples(NumFrames);
+	}
+	return false;
+}
+
+void FElectraPlayerPlugin::FPlayerAdapterDelegate::PrepareForDecoderShutdown()
+{
+	if (TSharedPtr<FElectraPlayerPlugin, ESPMode::ThreadSafe> PinnedHost = Host.Pin())
+	{
+		PinnedHost->OutputTexturePool.PrepareForDecoderShutdown();
+	}
+}
+
+
+FString FElectraPlayerPlugin::FPlayerAdapterDelegate::GetVideoAdapterName() const
+{
+	return GRHIAdapterName;
+}
+
+
+IElectraPlayerResourceDelegate* FElectraPlayerPlugin::FPlayerAdapterDelegate::GetResourceDelegate() const
+{
+	if (TSharedPtr<FElectraPlayerPlugin, ESPMode::ThreadSafe> PinnedHost = Host.Pin())
+	{
+		return PinnedHost->PlayerResourceDelegate.Get();
+	}
+	return nullptr;
+}
+
+//-----------------------------------------------------------------------------
+
+// IMediaPlayer interface
+
+FGuid FElectraPlayerPlugin::GetPlayerPluginGUID() const
+{
+	static FGuid PlayerPluginGUID(0x94ee3f80, 0x8e604292, 0xb4d24dd5, 0xfdade1c2);
+	return PlayerPluginGUID;
+}
+
+FString FElectraPlayerPlugin::GetInfo() const
+{
+	return TEXT("No information available");
+}
+
+IMediaSamples& FElectraPlayerPlugin::GetSamples()
+{
+	return *MediaSamples;
+}
+
+FString FElectraPlayerPlugin::GetStats() const
+{
+	return TEXT("ElectraPlayer: GetStats: <empty>?");
+}
+
+IMediaTracks& FElectraPlayerPlugin::GetTracks()
+{
+	return *this;
+}
+
+bool FElectraPlayerPlugin::Open(const FString& Url, const IMediaOptions* Options)
+{
+	return Open(Url, Options, nullptr);
+}
+
+bool FElectraPlayerPlugin::Open(const FString& Url, const IMediaOptions* Options, const FMediaPlayerOptions* InPlayerOptions)
+{
+	// Remember the option interface to poll for changes during playback.
+	OptionInterface = Options;
+
+	IElectraPlayerInterface::FPlaystartOptions LocalPlaystartOptions;
+
+	// Get playstart options from passed options, if they exist.
+	if (InPlayerOptions)
+	{
+		LocalPlaystartOptions.TimeOffset = InPlayerOptions->SeekTime;
+		LocalPlaystartOptions.AudioTrackIndex = InPlayerOptions->Tracks.Audio;
+	}
+
+	// Set up options to initialize the internal player with.
+	Electra::FParamDict PlayerOptions;
+	//	PlayerOptions.Set("initial_bitrate", Electra::FVariantValue((int64)2000000));
+
+		// Check for one-time initialization options that can't be changed during playback.
+		/*
+		int64 MinStreamFPS = Options->GetMediaOption(TEXT("MinStreamFPSForMediaStreaming"), (int64)-1);
+		if (MinStreamFPS > 0)
+		{
+			PlayerOptions.Set("min_stream_fps", Electra::FVariantValue(MinStreamFPS));
+			UE_LOG(LogElectraPlayerPlugin, Log, TEXT("[%p][%p] IMediaPlayer::Open: Limiting min stream fps to %d"), this, AdaptivePlayer.Get(), (int32)MinStreamFPS);
+		}
+
+		int64 MaxStreamFPS = Options->GetMediaOption(TEXT("MaxStreamFPSForMediaStreaming"), (int64)-1);
+		if (MaxStreamFPS > 0)
+		{
+			PlayerOptions.Set("max_stream_fps", Electra::FVariantValue(MaxStreamFPS));
+			UE_LOG(LogElectraPlayerPlugin, Log, TEXT("[%p][%p] IMediaPlayer::Open: Limiting max stream fps to %d"), this, AdaptivePlayer.Get(), (int32)MaxStreamFPS);
+		}
+		*/
+	int64 MaxVerticalHeight = Options->GetMediaOption(TEXT("MaxElectraVerticalResolution"), (int64)-1);
+	if (MaxVerticalHeight > 0)
+	{
+		PlayerOptions.Set("max_resoY", Electra::FVariantValue(MaxVerticalHeight));
+		UE_LOG(LogElectraPlayerPlugin, Log, TEXT("[%p] IMediaPlayer::Open: Limiting vertical resolution to %d for all streams"), this, (int32)MaxVerticalHeight);
+	}
+	int64 MaxVerticalHeightAt60 = Options->GetMediaOption(TEXT("MaxElectraVerticalResolutionOf60fpsVideos"), (int64)-1);
+	if (MaxVerticalHeightAt60 > 0)
+	{
+		PlayerOptions.Set("max_resoY_above_30fps", Electra::FVariantValue(MaxVerticalHeightAt60));
+		UE_LOG(LogElectraPlayerPlugin, Log, TEXT("[%p] IMediaPlayer::Open: Limiting vertical resolution to %d for streams >30fps"), this, (int32)MaxVerticalHeightAt60);
+	}
+
+	int64 MaxVerticalHeightForWindowsSoftware = Options->GetMediaOption(TEXT("MaxElectraVerticalResolutionOfWindowsSWD"), (int64)-1);
+	if (MaxVerticalHeightForWindowsSoftware > 0)
+	{
+		PlayerOptions.Set("max_resoY_windows_software", Electra::FVariantValue(MaxVerticalHeightForWindowsSoftware));
+		UE_LOG(LogElectraPlayerPlugin, Log, TEXT("[%p] IMediaPlayer::Open: Limiting vertical resolution to %d for Windows software decoding"), this, (int32)MaxVerticalHeightForWindowsSoftware);
+	}
+
+	double LiveEdgeDistanceForNormalPresentation = Options->GetMediaOption(TEXT("ElectraLivePresentationOffset"), (double)-1.0);
+	if (LiveEdgeDistanceForNormalPresentation > 0.0)
+	{
+		PlayerOptions.Set("seekable_range_live_end_offset", Electra::FVariantValue(Electra::FTimeValue().SetFromSeconds(LiveEdgeDistanceForNormalPresentation)));
+		UE_LOG(LogElectraPlayerPlugin, Log, TEXT("[%p] IMediaPlayer::Open: Setting distance to live edge for normal presentations to %.3f seconds"), this, LiveEdgeDistanceForNormalPresentation);
+	}
+	double LiveEdgeDistanceForAudioOnlyPresentation = Options->GetMediaOption(TEXT("ElectraLiveAudioPresentationOffset"), (double)-1.0);
+	if (LiveEdgeDistanceForAudioOnlyPresentation > 0.0)
+	{
+		PlayerOptions.Set("seekable_range_live_end_offset_audioonly", Electra::FVariantValue(Electra::FTimeValue().SetFromSeconds(LiveEdgeDistanceForAudioOnlyPresentation)));
+		UE_LOG(LogElectraPlayerPlugin, Log, TEXT("[%p] IMediaPlayer::Open: Setting distance to live edge for audio-only presentation to %.3f seconds"), this, LiveEdgeDistanceForAudioOnlyPresentation);
+	}
+	bool bUseConservativeLiveEdgeDistance = Options->GetMediaOption(TEXT("ElectraLiveUseConservativePresentationOffset"), (bool)false);
+	if (bUseConservativeLiveEdgeDistance)
+	{
+		PlayerOptions.Set("seekable_range_live_end_offset_conservative", Electra::FVariantValue(bUseConservativeLiveEdgeDistance));
+		UE_LOG(LogElectraPlayerPlugin, Log, TEXT("[%p] IMediaPlayer::Open: Using conservative live edge for distance calculation"), this);
+	}
+	bool bDontAllowInsecureForMediaSegments = Options->GetMediaOption(TEXT("ElectraDontAllowInsecureForEncryptedMediaSegments"), (bool)false);
+	if (bDontAllowInsecureForMediaSegments)
+	{
+		PlayerOptions.Set("dont_use_insecure_for_media_segments", Electra::FVariantValue(bDontAllowInsecureForMediaSegments));
+		UE_LOG(LogElectraPlayerPlugin, Log, TEXT("[%p] IMediaPlayer::Open: Do not allow downgrade to http for encrypted media segments"), this);
+	}
+	bool bDontAllowInsecureForInitSegments = Options->GetMediaOption(TEXT("ElectraDontAllowInsecureForInitSegments"), (bool)false);
+	if (bDontAllowInsecureForInitSegments)
+	{
+		PlayerOptions.Set("dont_use_insecure_for_init_segments", Electra::FVariantValue(bDontAllowInsecureForInitSegments));
+		UE_LOG(LogElectraPlayerPlugin, Log, TEXT("[%p] IMediaPlayer::Open: Do not allow downgrade to http for init segments"), this);
+	}
+	bool bThrowErrorWhenRebuffering = Options->GetMediaOption(TEXT("ElectraThrowErrorWhenRebuffering"), (bool)false);
+	if (bThrowErrorWhenRebuffering)
+	{
+		PlayerOptions.Set("throw_error_when_rebuffering", Electra::FVariantValue(bThrowErrorWhenRebuffering));
+		UE_LOG(LogElectraPlayerPlugin, Log, TEXT("[%p] IMediaPlayer::Open: Throw playback error when rebuffering"), this);
+	}
+
+	// Check for options that can be changed during playback and apply them at startup already.
+	// If a media source supports the MaxResolutionForMediaStreaming option then we can override the max resolution.
+	int64 DefaultValue = 0;
+	int64 MaxVerticalStreamResolution = Options->GetMediaOption(TEXT("MaxResolutionForMediaStreaming"), DefaultValue);
+	if (MaxVerticalStreamResolution != 0)
+	{
+		UE_LOG(LogElectraPlayerPlugin, Log, TEXT("[%p] IMediaPlayer::Open: Limiting max resolution to %d"), this, (int32)MaxVerticalStreamResolution);
+		LocalPlaystartOptions.MaxVerticalStreamResolution = (int32)MaxVerticalStreamResolution;
+	}
+
+	int64 MaxBandwidthForStreaming = Options->GetMediaOption(TEXT("ElectraMaxStreamingBandwidth"), (int64)0);
+	if (MaxBandwidthForStreaming > 0)
+	{
+		UE_LOG(LogElectraPlayerPlugin, Log, TEXT("[%p] Limiting max streaming bandwidth to %d bps"), this, (int32)MaxBandwidthForStreaming);
+		LocalPlaystartOptions.MaxBandwidthForStreaming = (int32)MaxBandwidthForStreaming;
+	}
+
+	return Player->OpenInternal(Url, PlayerOptions, LocalPlaystartOptions);
+}
+
+//-----------------------------------------------------------------------------
+/**
+ *
+ */
+bool FElectraPlayerPlugin::Open(const TSharedRef<FArchive, ESPMode::ThreadSafe>& Archive, const FString& OriginalUrl, const IMediaOptions* /*Options*/)
+{
+	// we support playback only from an external file, not from a "resource" (e.g. a packaged asset)
+	UE_LOG(LogElectraPlayerPlugin, Error, TEXT("[%p] IMediaPlayer::Archive"), this);
+	unimplemented();
+	return false;
+}
+
+
+//-----------------------------------------------------------------------------
+/**
+*	Internal Close / Shutdown player
+*/
+void FElectraPlayerPlugin::Close()
+{
+	Player->CloseInternal(true);
+}
+
+//-----------------------------------------------------------------------------
+/**
+ *	Tick the player from the game thread
+ */
+void FElectraPlayerPlugin::TickInput(FTimespan DeltaTime, FTimespan Timecode)
+{
+	OutputTexturePool.Tick();
+	Player->Tick(DeltaTime, Timecode);
+}
+
+//-----------------------------------------------------------------------------
+/**
+	Get special feature flags states
+*/
+bool FElectraPlayerPlugin::GetPlayerFeatureFlag(EFeatureFlag flag) const
+{
+	switch(flag)
+	{
+		case EFeatureFlag::AllowShutdownOnClose:
+			return Player->IsKillAfterCloseAllowed();
+		case EFeatureFlag::UsePlaybackTimingV2:
+			return true;
+		default:
+			break;
+	}
+	return IMediaPlayer::GetPlayerFeatureFlag(flag);
+}
+
+//-----------------------------------------------------------------------------
+/**
+	Set a notification to be signaled once any async tear down of the instance is done
+*/
+bool FElectraPlayerPlugin::SetAsyncResourceReleaseNotification(IAsyncResourceReleaseNotificationRef AsyncResourceReleaseNotification)
+{
+	class FAsyncResourceReleaseNotifyContainer : public IElectraPlayerInterface::IAsyncResourceReleaseNotifyContainer
+	{
+	public:
+		FAsyncResourceReleaseNotifyContainer(IAsyncResourceReleaseNotificationRef InAsyncResourceReleaseNotification) : AsyncResourceReleaseNotification(InAsyncResourceReleaseNotification) {}
+		virtual void Signal(uint32 ResourceFlags) override { AsyncResourceReleaseNotification->Signal(ResourceFlags); }
+	private:
+		IAsyncResourceReleaseNotificationRef AsyncResourceReleaseNotification;
+	};
+
+	Player->SetAsyncResourceReleaseNotification(new FAsyncResourceReleaseNotifyContainer(AsyncResourceReleaseNotification));
+	return true;
+}
+
+uint32 FElectraPlayerPlugin::GetNewResourcesOnOpen() const
+{
+	// Electra will recreate all decoder related resources on each open call
+	// (a simplification: we also recreate the texture pool should it change sizes on SOME platforms - but we reoprt the release only per instance, so thisb matches that)
+	return IMediaPlayerLifecycleManagerDelegate::ResourceFlags_Decoder;
+}
+
+//////////////////////////////////////////////////////////////////////////
+// IMediaControl impl
+
+//-----------------------------------------------------------------------------
+/**
+ *	Currently, we cannot do anything.. well, at least we can play!
+ */
+bool FElectraPlayerPlugin::CanControl(EMediaControl Control) const
+{
+	EMediaState CurrentState = GetState();
+	if (Control == EMediaControl::Pause)
+	{
+		return CurrentState == EMediaState::Playing;
+	}
+
+	if (Control == EMediaControl::Resume)
+	{
+		return CurrentState == EMediaState::Paused || CurrentState == EMediaState::Stopped;
+	}
+
+	if (Control == EMediaControl::Seek)
+	{
+		return CurrentState == EMediaState::Playing || CurrentState == EMediaState::Paused || CurrentState == EMediaState::Stopped ;//|| CurrentState == Preparing;
+	}
+
+	return false;
+}
+
+//-----------------------------------------------------------------------------
+/**
+ *	Rate is only real-time
+ */
+float FElectraPlayerPlugin::GetRate() const
+{
+	return Player->GetRate();
+}
+
+//-----------------------------------------------------------------------------
+/**
+ *	Expose player state
+ */
+EMediaState FElectraPlayerPlugin::GetState() const
+{
+	return (EMediaState)Player->GetState();
+}
+
+//-----------------------------------------------------------------------------
+/**
+ *	Expose player status
+ */
+EMediaStatus FElectraPlayerPlugin::GetStatus() const
+{
+	return (EMediaStatus)Player->GetStatus();
+}
+
+
+bool FElectraPlayerPlugin::IsLooping() const
+{
+	return Player->IsLooping();
+}
+
+
+bool FElectraPlayerPlugin::SetLooping(bool bLooping)
+{
+	return Player->SetLooping(bLooping);
+}
+
+//-----------------------------------------------------------------------------
+/**
+ *	Only return real-time playback for the moment..
+ */
+TRangeSet<float> FElectraPlayerPlugin::GetSupportedRates(EMediaRateThinning Thinning) const
+{
+	TRangeSet<float> Res;
+	Res.Add(TRange<float>{1.0f}); // only normal (real-time) playback rate
+	Res.Add(TRange<float>{0.0f}); // and pause
+	return Res;
+}
+
+
+FTimespan FElectraPlayerPlugin::GetTime() const
+{
+	return Player->GetTime();
+}
+
+
+FTimespan FElectraPlayerPlugin::GetDuration() const
+{
+	return Player->GetDuration();
+}
+
+
+bool FElectraPlayerPlugin::SetRate(float Rate)
+{
+	UE_LOG(LogElectraPlayerPlugin, Log, TEXT("[%p] IMediaControls::SetRate(%f)"), this, Rate);
+	CSV_EVENT(ElectraPlayer, TEXT("Setting Rate"));
+
+	return Player->SetRate(Rate);
+}
+
+
+bool FElectraPlayerPlugin::Seek(const FTimespan& Time)
+{
+	UE_LOG(LogElectraPlayerPlugin, Log, TEXT("[%p] IMediaControls::Seek() to %s"), this, *Time.ToString(TEXT("%h:%m:%s.%f")));
+	CSV_EVENT(ElectraPlayer, TEXT("Seeking"));
+
+	return Player->Seek(Time);
+}
+
+
+bool FElectraPlayerPlugin::GetAudioTrackFormat(int32 TrackIndex, int32 FormatIndex, FMediaAudioTrackFormat& OutFormat) const
+{
+	IElectraPlayerInterface::FAudioTrackFormat Format;
+	if (!Player->GetAudioTrackFormat(TrackIndex, FormatIndex, Format))
+	{
+		return false;
+	}
+	OutFormat.BitsPerSample = Format.BitsPerSample;
+	OutFormat.NumChannels = Format.NumChannels;
+	OutFormat.SampleRate = Format.SampleRate;
+	OutFormat.TypeName = Format.TypeName;
+	return true;
+}
+
+
+bool FElectraPlayerPlugin::GetVideoTrackFormat(int32 TrackIndex, int32 FormatIndex, FMediaVideoTrackFormat& OutFormat) const
+{
+	IElectraPlayerInterface::FVideoTrackFormat Format;
+	if (!Player->GetVideoTrackFormat(TrackIndex, FormatIndex, Format))
+	{
+		return false;
+	}
+	OutFormat.Dim = Format.Dim;
+	OutFormat.FrameRate = Format.FrameRate;
+	OutFormat.FrameRates = Format.FrameRates;
+	OutFormat.TypeName = Format.TypeName;
+	return true;
+}
+
+
+int32 FElectraPlayerPlugin::GetNumTracks(EMediaTrackType TrackType) const
+{
+	return Player->GetNumTracks((IElectraPlayerInterface::EPlayerTrackType)TrackType);
+}
+
+
+int32 FElectraPlayerPlugin::GetNumTrackFormats(EMediaTrackType TrackType, int32 TrackIndex) const
+{
+	return Player->GetNumTrackFormats((IElectraPlayerInterface::EPlayerTrackType)TrackType, TrackIndex);
+}
+
+
+int32 FElectraPlayerPlugin::GetSelectedTrack(EMediaTrackType TrackType) const
+{
+	return Player->GetSelectedTrack((IElectraPlayerInterface::EPlayerTrackType)TrackType);
+}
+
+
+FText FElectraPlayerPlugin::GetTrackDisplayName(EMediaTrackType TrackType, int32 TrackIndex) const
+{
+	return Player->GetTrackDisplayName((IElectraPlayerInterface::EPlayerTrackType)TrackType, TrackIndex);
+}
+
+
+int32 FElectraPlayerPlugin::GetTrackFormat(EMediaTrackType TrackType, int32 TrackIndex) const
+{
+	return Player->GetTrackFormat((IElectraPlayerInterface::EPlayerTrackType)TrackType, TrackIndex);
+}
+
+
+FString FElectraPlayerPlugin::GetTrackLanguage(EMediaTrackType TrackType, int32 TrackIndex) const
+{
+	return Player->GetTrackLanguage((IElectraPlayerInterface::EPlayerTrackType)TrackType, TrackIndex);
+}
+
+
+FString FElectraPlayerPlugin::GetTrackName(EMediaTrackType TrackType, int32 TrackIndex) const
+{
+	return Player->GetTrackName((IElectraPlayerInterface::EPlayerTrackType)TrackType, TrackIndex);
+}
+
+
+bool FElectraPlayerPlugin::SelectTrack(EMediaTrackType TrackType, int32 TrackIndex)
+{
+	return Player->SelectTrack((IElectraPlayerInterface::EPlayerTrackType)TrackType, TrackIndex);
+}
+
+
+bool FElectraPlayerPlugin::SetTrackFormat(EMediaTrackType TrackType, int32 TrackIndex, int32 FormatIndex)
+{
+	return false;
+}
+
+
+bool FElectraPlayerPlugin::SetVideoTrackFrameRate(int32 TrackIndex, int32 FormatIndex, float FrameRate)
+{
+	return false;
+}
+
+
+void FElectraPlayerPlugin::SetLastAudioRenderedSampleTime(FTimespan SampleTime)
+{
+	// No-op.
+}
+
+
diff --git a/Engine/Plugins/Media/ElectraPlayer/Source/ElectraPlayerPlugin/Private/ElectraPlayerPlugin.h b/Engine/Plugins/Media/ElectraPlayer/Source/ElectraPlayerPlugin/Private/ElectraPlayerPlugin.h
new file mode 100644
--- /dev/null
+++ b/Engine/Plugins/Media/ElectraPlayer/Source/ElectraPlayerPlugin/Private/ElectraPlayerPlugin.h
@@ -0,0 +1,198 @@
+// Copyright Epic Games, Inc. All Rights Reserved.
+
+#pragma once
+
+#include "IElectraPlayerInterface.h"
+
+#include "Containers/UnrealString.h"
+#include "Containers/Queue.h"
+#include "Misc/Guid.h"
+#include "IMediaCache.h"
+#include "IMediaPlayer.h"
+#include "IMediaView.h"
+#include "IMediaControls.h"
+#include "IMediaTracks.h"
+#include "IMediaEventSink.h"
+#include "IMediaPlayerLifecycleManager.h"
+#include "MediaSampleQueue.h"
+#include "Templates/SharedPointer.h"
+#include "Logging/LogMacros.h"
+
+class IElectraPlayerRuntimeModule;
+class FElectraPlayerResourceDelegate;
+
+DECLARE_LOG_CATEGORY_EXTERN(LogElectraPlayerPlugin, Log, All);
+
+//-----------------------------------------------------------------------------
+
+#include "ElectraTextureSample.h"
+#include "ElectraPlayerAudioSample.h"
+
+//-----------------------------------------------------------------------------
+
+class IMediaOptions;
+class FMediaSamples;
+class FArchive;
+
+/**
+	* Implements a media player
+	* Supports multiple platforms.
+	*/
+class FElectraPlayerPlugin
+	: public IMediaPlayer
+	, protected IMediaCache
+	, protected IMediaView
+	, protected IMediaControls
+	, public IMediaTracks
+	, public TSharedFromThis<FElectraPlayerPlugin, ESPMode::ThreadSafe>
+{
+public:
+
+	FElectraPlayerPlugin();
+
+	virtual ~FElectraPlayerPlugin();
+
+	bool Initialize(IMediaEventSink& InEventSink,
+		FElectraPlayerSendAnalyticMetricsDelegate& InSendAnalyticMetricsDelegate,
+		FElectraPlayerSendAnalyticMetricsPerMinuteDelegate& InSendAnalyticMetricsPerMinuteDelegate,
+		FElectraPlayerReportVideoStreamingErrorDelegate& InReportVideoStreamingErrorDelegate,
+		FElectraPlayerReportSubtitlesMetricsDelegate& InReportSubtitlesFileMetricsDelegate);
+
+public:
+	// IMediaPlayer impl
+
+	void Close() override;
+
+	IMediaCache& GetCache() override
+	{
+		return *this;
+	}
+
+	IMediaControls& GetControls() override
+	{
+		return *this;
+	}
+
+	FString GetUrl() const override
+	{
+		return Player->GetUrl();
+	}
+
+	IMediaView& GetView() override
+	{
+		return *this;
+	}
+
+	void SetGuid(const FGuid& Guid) override
+	{
+		Player->SetGuid(Guid);
+	}
+
+	FString GetInfo() const override;
+	FGuid GetPlayerPluginGUID() const override;
+	IMediaSamples& GetSamples() override;
+	FString GetStats() const override;
+	IMediaTracks& GetTracks() override;
+	bool Open(const FString& Url, const IMediaOptions* Options) override;
+	bool Open(const FString& Url, const IMediaOptions* Options, const FMediaPlayerOptions* PlayerOptions) override;
+	bool Open(const TSharedRef<FArchive, ESPMode::ThreadSafe>& Archive, const FString& OriginalUrl, const IMediaOptions* Options) override;
+	void TickInput(FTimespan DeltaTime, FTimespan Timecode) override;
+	void SetLastAudioRenderedSampleTime(FTimespan SampleTime) override;
+	bool FlushOnSeekStarted() const override
+	{
+		return true;
+	}
+	bool FlushOnSeekCompleted() const override
+	{
+		return true;
+	}
+
+	bool GetPlayerFeatureFlag(EFeatureFlag flag) const override;
+
+	bool SetAsyncResourceReleaseNotification(IAsyncResourceReleaseNotificationRef AsyncResourceReleaseNotification) override;
+	uint32 GetNewResourcesOnOpen() const override;
+
+private:
+	friend class FElectraPlayerPluginModule;
+	friend class FElectraPlayerResourceDelegate;
+
+	// IMediaControls impl
+	bool CanControl(EMediaControl Control) const override;
+
+	FTimespan GetDuration() const override;
+
+	bool IsLooping() const override;
+	bool SetLooping(bool bLooping) override;
+
+	EMediaState GetState() const override;
+	EMediaStatus GetStatus() const override;
+
+	TRangeSet<float> GetSupportedRates(EMediaRateThinning Thinning) const override;
+	FTimespan GetTime() const override;
+	float GetRate() const override;
+	bool SetRate(float Rate) override;
+	bool Seek(const FTimespan& Time) override;
+
+
+	// From IMediaTracks
+	bool GetAudioTrackFormat(int32 TrackIndex, int32 FormatIndex, FMediaAudioTrackFormat& OutFormat) const override;
+	int32 GetNumTracks(EMediaTrackType TrackType) const override;
+	int32 GetNumTrackFormats(EMediaTrackType TrackType, int32 TrackIndex) const override;
+	int32 GetSelectedTrack(EMediaTrackType TrackType) const override;
+	FText GetTrackDisplayName(EMediaTrackType TrackType, int32 TrackIndex) const override;
+	int32 GetTrackFormat(EMediaTrackType TrackType, int32 TrackIndex) const override;
+	FString GetTrackLanguage(EMediaTrackType TrackType, int32 TrackIndex) const override;
+	FString GetTrackName(EMediaTrackType TrackType, int32 TrackIndex) const override;
+	bool GetVideoTrackFormat(int32 TrackIndex, int32 FormatIndex, FMediaVideoTrackFormat& OutFormat) const override;
+	bool SelectTrack(EMediaTrackType TrackType, int32 TrackIndex) override;
+	bool SetTrackFormat(EMediaTrackType TrackType, int32 TrackIndex, int32 FormatIndex) override;
+	bool SetVideoTrackFrameRate(int32 TrackIndex, int32 FormatIndex, float FrameRate) override;
+
+	IElectraPlayerResourceDelegate* PlatformCreatePlayerResourceDelegate();
+	static void PlatformSetupResourceParams(Electra::FParamDict& Params);
+
+	// IElectraPlayerAdapterDelegate impl
+	class FPlayerAdapterDelegate : public IElectraPlayerAdapterDelegate
+	{
+	public:
+		FPlayerAdapterDelegate(const TSharedPtr<FElectraPlayerPlugin, ESPMode::ThreadSafe> & InHost) : Host(InHost) {}
+
+		Electra::FVariantValue QueryOptions(EOptionType Type, const Electra::FVariantValue & Param) override;
+		void SendMediaEvent(EPlayerEvent Event) override;
+		void OnVideoFlush() override;
+		void OnAudioFlush() override;
+		void PresentVideoFrame(const FVideoDecoderOutputPtr& InVideoFrame) override;
+		void PresentAudioFrame(const IAudioDecoderOutputPtr& InAudioFrame) override;
+		bool CanReceiveVideoSamples(int32 NumFrames) override;
+		bool CanReceiveAudioSamples(int32 NumFrames) override;
+		void PrepareForDecoderShutdown() override;
+		FString GetVideoAdapterName() const override;
+		IElectraPlayerResourceDelegate* GetResourceDelegate() const override;
+
+	private:
+		TWeakPtr<FElectraPlayerPlugin, ESPMode::ThreadSafe> Host;
+	};
+	friend class FPlayerAdapterDelegate;
+	TSharedPtr<FPlayerAdapterDelegate, ESPMode::ThreadSafe> PlayerDelegate;
+
+
+private:
+	/** Output queues as needed by MediaFramework */
+	TUniquePtr<FMediaSamples> MediaSamples;
+
+	/** Option interface */
+	const IMediaOptions* OptionInterface;
+
+	/** The media event handler */
+	IMediaEventSink* EventSink;
+
+	/** The actual player */
+	TSharedPtr<IElectraPlayerInterface, ESPMode::ThreadSafe> Player;
+
+	/** Output sample pools */
+	FElectraTextureSamplePool OutputTexturePool;
+	FElectraPlayerAudioSamplePool OutputAudioPool;
+
+	TUniquePtr<IElectraPlayerResourceDelegate> PlayerResourceDelegate;
+};
+
diff --git a/Engine/Plugins/Media/ElectraPlayer/Source/ElectraPlayerPlugin/Private/ElectraPlayerPluginModule.cpp b/Engine/Plugins/Media/ElectraPlayer/Source/ElectraPlayerPlugin/Private/ElectraPlayerPluginModule.cpp
new file mode 100644
--- /dev/null
+++ b/Engine/Plugins/Media/ElectraPlayer/Source/ElectraPlayerPlugin/Private/ElectraPlayerPluginModule.cpp
@@ -0,0 +1,123 @@
+// Copyright Epic Games, Inc. All Rights Reserved.
+
+#include "CoreMinimal.h"
+
+#include "Modules/ModuleManager.h"
+#include "RHI.h"
+
+#include "ElectraPlayerMisc.h"
+#include "IElectraPlayerPluginModule.h"
+#include "IElectraPlayerRuntimeModule.h"
+#include "ElectraPlayerPlugin.h"
+#include "ParameterDictionary.h"
+
+#define LOCTEXT_NAMESPACE "ElectraPlayerPluginModule"
+
+DEFINE_LOG_CATEGORY(LogElectraPlayerPlugin);
+
+// -----------------------------------------------------------------------------------------------------------------------------------
+
+class FElectraPlayerPluginModule : public IElectraPlayerPluginModule
+{
+public:
+	// IElectraPlayerPluginModule interface
+
+	virtual bool IsInitialized() const override
+	{
+		return bInitialized;
+	}
+
+	virtual TSharedPtr<IMediaPlayer, ESPMode::ThreadSafe> CreatePlayer(IMediaEventSink& EventSink) override
+	{
+		if (bInitialized)
+		{
+			FElectraPlayerPlugin* NewRawPlayer = new FElectraPlayerPlugin();
+			if (NewRawPlayer)
+			{
+				TSharedPtr<IMediaPlayer, ESPMode::ThreadSafe> NewPlayer = MakeShareable(NewRawPlayer);
+				if (NewRawPlayer->Initialize(EventSink, SendAnalyticMetricsDelegate, SendAnalyticMetricsPerMinuteDelegate, ReportVideoStreamingErrorDelegate, ReportSubtitlesMetricsDelegate))
+				{
+					return NewPlayer;
+				}
+			}
+		}
+		return nullptr;
+	}
+
+	virtual void SendAnalyticMetrics(const TSharedPtr<IAnalyticsProviderET>& AnalyticsProvider, const FGuid& PlayerGuid) override
+	{
+		SendAnalyticMetricsDelegate.Broadcast(AnalyticsProvider, PlayerGuid);
+	}
+
+	virtual void SendAnalyticMetricsPerMinute(const TSharedPtr<IAnalyticsProviderET>& AnalyticsProvider) override
+	{
+		SendAnalyticMetricsPerMinuteDelegate.Broadcast(AnalyticsProvider);
+	}
+
+	virtual void ReportVideoStreamingError(const FGuid& PlayerGuid, const FString& LastError) override
+	{
+		ReportVideoStreamingErrorDelegate.Broadcast(PlayerGuid, LastError);
+	}
+
+	virtual void ReportSubtitlesMetrics(const FGuid& PlayerGuid, const FString& URL, double ResponseTime, const FString& LastError)
+	{
+		ReportSubtitlesMetricsDelegate.Broadcast(PlayerGuid, URL, ResponseTime, LastError);
+	}
+
+public:
+	// IModuleInterface interface
+
+	void StartupModule() override
+	{
+		// Check that we have the player module and that it has initialized successfully.
+		if (FModuleManager::Get().GetModule("ElectraPlayerRuntime"))
+		{
+			IElectraPlayerRuntimeModule* ElectraPlayer = &FModuleManager::Get().GetModuleChecked<IElectraPlayerRuntimeModule>("ElectraPlayerRuntime");
+			if (!ElectraPlayer || !ElectraPlayer->IsInitialized())
+			{
+				return;
+			}
+		}
+		else
+		{
+			return;
+		}
+
+		if (!bInitialized)
+		{
+//WE COULD ALSO CHECK FOR SINGLE THREADING ISSUES HERE!? -- kinda piggy back on this?
+			// to detect cooking and other commandlets that run with NullRHI
+			if (GDynamicRHI == nullptr || GDynamicRHI->GetName() == FString(TEXT("Null")) || GDynamicRHI->GetName() == FString("Empty"))
+			{
+				UE_LOG(LogElectraPlayerPlugin, Log, TEXT("Dummy Dynamic RHI detected. Electra Player plugin is not initialised"));
+				return;
+			}
+
+			Electra::FParamDict Params;
+			Params.Set("DeviceName", Electra::FVariantValue(FString(GDynamicRHI->GetName())));
+			FElectraPlayerPlugin::PlatformSetupResourceParams(Params);
+			FElectraPlayerPlatform::StartupPlatformResources(Params);
+
+			bInitialized = true;
+		}
+	}
+
+	void ShutdownModule() override
+	{
+		bInitialized = false;
+	}
+
+private:
+	bool bInitialized = false;
+
+	FElectraPlayerSendAnalyticMetricsDelegate			SendAnalyticMetricsDelegate;
+	FElectraPlayerSendAnalyticMetricsPerMinuteDelegate	SendAnalyticMetricsPerMinuteDelegate;
+	FElectraPlayerReportVideoStreamingErrorDelegate		ReportVideoStreamingErrorDelegate;
+	FElectraPlayerReportSubtitlesMetricsDelegate		ReportSubtitlesMetricsDelegate;
+};
+
+IMPLEMENT_MODULE(FElectraPlayerPluginModule, ElectraPlayerPlugin);
+
+#undef LOCTEXT_NAMESPACE
+
+
diff --git a/Engine/Plugins/Media/ElectraPlayer/Source/ElectraPlayerPlugin/Private/ElectraPlayerPlugin_Android.cpp b/Engine/Plugins/Media/ElectraPlayer/Source/ElectraPlayerPlugin/Private/ElectraPlayerPlugin_Android.cpp
new file mode 100644
--- /dev/null
+++ b/Engine/Plugins/Media/ElectraPlayer/Source/ElectraPlayerPlugin/Private/ElectraPlayerPlugin_Android.cpp
@@ -0,0 +1,34 @@
+// Copyright Epic Games, Inc. All Rights Reserved.
+
+#include "CoreTypes.h"
+
+#if PLATFORM_ANDROID
+
+#include <ElectraPlayerPlugin.h>
+
+
+class FElectraPlayerResourceDelegate : public IElectraPlayerResourceDelegate
+{
+public:
+	FElectraPlayerResourceDelegate(FElectraPlayerPlugin* InOwner) : Owner(InOwner) {}
+
+	virtual jobject GetCodecSurface() override
+	{
+		return (jobject)Owner->OutputTexturePool.GetCodecSurface();
+	}
+
+private:
+	FElectraPlayerPlugin* Owner;
+};
+
+
+IElectraPlayerResourceDelegate* FElectraPlayerPlugin::PlatformCreatePlayerResourceDelegate()
+{
+	return new FElectraPlayerResourceDelegate(this);
+}
+
+void FElectraPlayerPlugin::PlatformSetupResourceParams(Electra::FParamDict& Params)
+{
+}
+
+#endif // PLATFORM_ANDROID
diff --git a/Engine/Plugins/Media/ElectraPlayer/Source/ElectraPlayerPlugin/Private/ElectraPlayerPlugin_Apple.cpp b/Engine/Plugins/Media/ElectraPlayer/Source/ElectraPlayerPlugin/Private/ElectraPlayerPlugin_Apple.cpp
new file mode 100644
--- /dev/null
+++ b/Engine/Plugins/Media/ElectraPlayer/Source/ElectraPlayerPlugin/Private/ElectraPlayerPlugin_Apple.cpp
@@ -0,0 +1,18 @@
+// Copyright Epic Games, Inc. All Rights Reserved.
+
+#include "CoreTypes.h"
+
+#if PLATFORM_MAC || PLATFORM_IOS
+
+#include <ElectraPlayerPlugin.h>
+
+IElectraPlayerResourceDelegate* FElectraPlayerPlugin::PlatformCreatePlayerResourceDelegate()
+{
+	return nullptr;
+}
+
+void FElectraPlayerPlugin::PlatformSetupResourceParams(Electra::FParamDict& Params)
+{
+}
+
+#endif // PLATFORM_MAC || PLATFORM_IOS
diff --git a/Engine/Plugins/Media/ElectraPlayer/Source/ElectraPlayerPlugin/Private/ElectraPlayerPlugin_Windows.cpp b/Engine/Plugins/Media/ElectraPlayer/Source/ElectraPlayerPlugin/Private/ElectraPlayerPlugin_Windows.cpp
new file mode 100644
--- /dev/null
+++ b/Engine/Plugins/Media/ElectraPlayer/Source/ElectraPlayerPlugin/Private/ElectraPlayerPlugin_Windows.cpp
@@ -0,0 +1,19 @@
+// Copyright Epic Games, Inc. All Rights Reserved.
+
+#include "CoreTypes.h"
+
+#if PLATFORM_WINDOWS
+
+#include <ElectraPlayerPlugin.h>
+
+IElectraPlayerResourceDelegate* FElectraPlayerPlugin::PlatformCreatePlayerResourceDelegate()
+{
+	return nullptr;
+}
+
+void FElectraPlayerPlugin::PlatformSetupResourceParams(Electra::FParamDict& Params)
+{
+	Params.Set("Device", Electra::FVariantValue(GDynamicRHI->RHIGetNativeDevice()));
+}
+
+#endif // PLATFORM_WINDOWS
\ No newline at end of file
diff --git a/Engine/Plugins/Media/ElectraPlayer/Source/ElectraPlayerRuntime/ElectraPlayerRuntime.Build.cs b/Engine/Plugins/Media/ElectraPlayer/Source/ElectraPlayerRuntime/ElectraPlayerRuntime.Build.cs
new file mode 100644
--- /dev/null
+++ b/Engine/Plugins/Media/ElectraPlayer/Source/ElectraPlayerRuntime/ElectraPlayerRuntime.Build.cs
@@ -0,0 +1,168 @@
+// Copyright Epic Games, Inc. All Rights Reserved.
+using EpicGames.Core;
+using System.IO;
+
+namespace UnrealBuildTool.Rules
+{
+	public class ElectraPlayerRuntime: ModuleRules
+	{
+		public ElectraPlayerRuntime(ReadOnlyTargetRules Target) : base(Target)
+		{
+			//
+			// Common setup...
+			//
+
+            bLegalToDistributeObjectCode = true;
+            PCHUsage = PCHUsageMode.NoPCHs;
+
+			PrivateDependencyModuleNames.AddRange(
+				new string[] {
+					"Core",
+					"HTTP",
+					"WebSockets",
+					"Json",
+					"ElectraBase",
+				});
+			if (Target.bCompileAgainstEngine)
+			{
+				// Added to allow debug rendering if used in UE context
+				PrivateDependencyModuleNames.Add("Engine");
+			}
+
+			PrivateIncludePaths.AddRange(
+				new string[] {
+					"ElectraPlayerRuntime/Private/Runtime",
+				});
+
+			//
+			// Common platform setup...
+			//
+
+			if (Target.Platform.IsInGroup(UnrealPlatformGroup.Windows))
+            {
+                string DirectXSDKDir = Target.UEThirdPartySourceDirectory + "Windows/DirectX";
+                PublicSystemIncludePaths.Add(DirectXSDKDir + "/include");
+
+				if (Target.Platform == UnrealTargetPlatform.Win64)
+                {
+                    DirectXSDKDir += "/Lib/x64/";
+                }
+                else if (Target.Platform == UnrealTargetPlatform.Win32)
+                {
+                    DirectXSDKDir += "/Lib/x86/";
+                }
+
+				PrivateDefinitions.Add("_CRT_SECURE_NO_WARNINGS=1");
+
+                PrivateDefinitions.Add("ELECTRA_ENABLE_MFDECODER");
+                PrivateDefinitions.Add("ELECTRA_ENABLE_SWDECODE");
+
+                AddEngineThirdPartyPrivateStaticDependencies(Target, "WinHttp");
+
+				if (WinSupportsDX9())
+				{
+					AddEngineThirdPartyPrivateStaticDependencies(Target, "DX9");
+					PrivateDefinitions.Add("ELECTRA_HAVE_DX9");     // video decoding on Win7 enabled - uses DX9
+				}
+
+				if (WinSupportsDX11())
+				{
+					PrivateDefinitions.Add("ELECTRA_HAVE_DX11");    // video decoding for DX11 enabled (Win8+)
+				}
+
+				if (Target.Platform == UnrealTargetPlatform.Win64 || Target.Platform == UnrealTargetPlatform.Win32)
+				{
+					PublicAdditionalLibraries.AddRange(new string[] {
+						DirectXSDKDir + "dxerr.lib",
+					});
+				}
+
+				PublicSystemLibraries.AddRange(new string[] {
+                    "strmiids.lib",
+                    "legacy_stdio_definitions.lib",
+                    "Dxva2.lib",
+                });
+
+                // Delay-load all MF DLLs to be able to check Windows version for compatibility in `StartupModule` before loading them manually
+                PublicSystemLibraries.Add("mfplat.lib");
+                PublicDelayLoadDLLs.Add("mfplat.dll");
+                PublicSystemLibraries.Add("mfuuid.lib");
+
+				PublicIncludePaths.Add("$(ModuleDir)/Public/Windows");
+				PrivateIncludePaths.Add("ElectraPlayerRuntime/Private/Runtime/Decoder/Windows");
+				PrivateIncludePaths.Add("ElectraPlayerRuntime/Private/Windows");
+			}
+			else if (Target.Platform == UnrealTargetPlatform.Mac)
+            {
+                PublicFrameworks.AddRange(
+				new string[] {
+								"CoreMedia",
+								"CoreVideo",
+								"AVFoundation",
+								"AudioToolbox",
+                                "VideoToolbox",
+								"QuartzCore"
+				});
+
+				PublicIncludePaths.Add("$(ModuleDir)/Public/Apple");
+
+				PrivateIncludePaths.Add("ElectraPlayerRuntime/Private/Runtime/Decoder/Apple");
+			}
+			else if (Target.Platform == UnrealTargetPlatform.IOS )
+            {
+				PublicFrameworks.AddRange(
+				new string[] {
+								"CoreMedia",
+								"CoreVideo",
+								"AVFoundation",
+								"AudioToolbox",
+                                "VideoToolbox",
+								"QuartzCore"
+				});
+
+				PublicIncludePaths.Add("$(ModuleDir)/Public/Apple");
+
+				PrivateIncludePaths.Add("ElectraPlayerRuntime/Private/Runtime/Decoder/Apple");
+			}
+			else if (Target.IsInPlatformGroup(UnrealPlatformGroup.Android) )
+            {
+				PublicDefinitions.Add("CURL_ENABLE_DEBUG_CALLBACK=1");
+
+                if (Target.Configuration != UnrealTargetConfiguration.Shipping)
+                {
+                    PublicDefinitions.Add("CURL_ENABLE_NO_TIMEOUTS_OPTION=1");
+                }
+
+                string PluginPath = Utils.MakePathRelativeTo(ModuleDirectory, Target.RelativeEnginePath);
+                AdditionalPropertiesForReceipt.Add("AndroidPlugin", Path.Combine(PluginPath, "ElectraPlayerRuntime_UPL.xml"));
+
+				PublicIncludePaths.Add("$(ModuleDir)/Public/Android");
+
+				PrivateIncludePaths.Add("ElectraPlayerRuntime/Private/Runtime/Decoder/Android");
+			}
+			else if (Target.IsInPlatformGroup(UnrealPlatformGroup.Unix) )
+            {
+				PublicDefinitions.Add("CURL_ENABLE_DEBUG_CALLBACK=1");
+
+                if (Target.Configuration != UnrealTargetConfiguration.Shipping)
+                {
+                    PublicDefinitions.Add("CURL_ENABLE_NO_TIMEOUTS_OPTION=1");
+                }
+
+				PublicIncludePaths.Add("$(ModuleDir)/Public/Unix");
+
+				PrivateIncludePaths.Add("ElectraPlayerRuntime/Private/Runtime/Decoder/GStreamer");
+			}
+		}
+
+		protected virtual bool WinSupportsDX9()
+		{
+			return true;
+		}
+
+		protected virtual bool WinSupportsDX11()
+		{
+			return true;
+		}
+	}
+}
diff --git a/Engine/Plugins/Media/ElectraPlayer/Source/ElectraPlayerRuntime/ElectraPlayerRuntime_UPL.xml b/Engine/Plugins/Media/ElectraPlayer/Source/ElectraPlayerRuntime/ElectraPlayerRuntime_UPL.xml
new file mode 100644
--- /dev/null
+++ b/Engine/Plugins/Media/ElectraPlayer/Source/ElectraPlayerRuntime/ElectraPlayerRuntime_UPL.xml
@@ -0,0 +1,131 @@
+<?xml version="1.0" encoding="utf-8"?>
+<!--ElectraPlayer plugin additions-->
+<root xmlns:android="http://schemas.android.com/apk/res/android">
+	<!-- init section is always evaluated once per architecture -->
+	<init>
+		<log text="ElectraPlayer init"/>
+		<!-- currently supports all architectures -->
+		<setBool result="bSupported" value="true"/>
+		<if condition="bSupported">
+			<true>
+			</true>
+			<false>
+				<log text="ElectraPlayer not supported for this architecture, disabled."/>
+			</false>
+		</if>
+	</init>
+
+	<!-- optional updates applied to AndroidManifest.xml -->
+	<androidManifestUpdates>
+		<if condition="bSupported">
+			<true>
+			</true>
+		</if>
+	</androidManifestUpdates>
+
+	<!-- optional additions to proguard -->
+	<proguardAdditions>
+		<insert>
+-keep class com.epicgames.ue4.ElectraAudioDecoderAAC {
+	public *;
+}
+-keep class com.epicgames.ue4.ElectraAudioDecoderAAC$FOutputFormatInfo {
+	public *;
+}
+-keep class com.epicgames.ue4.ElectraAudioDecoderAAC$FOutputBufferInfo {
+	public *;
+}
+-keep class com.epicgames.ue4.ElectraVideoDecoderH264 {
+	public *;
+}
+-keep class com.epicgames.ue4.ElectraVideoDecoderH264$FCreateParameters {
+	public *;
+}
+-keep class com.epicgames.ue4.ElectraVideoDecoderH264$FDecoderInformation {
+	public *;
+}
+-keep class com.epicgames.ue4.ElectraVideoDecoderH264$FOutputFormatInfo {
+	public *;
+}
+-keep class com.epicgames.ue4.ElectraVideoDecoderH264$FOutputBufferInfo {
+	public *;
+}
+		</insert>
+	</proguardAdditions>
+
+	<!-- optional files or directories to copy to Intermediate/Android/APK -->
+	<resourceCopies>
+		<if condition="bSupported">
+			<true>
+				<copyDir src="$S(PluginDir)/Private/Runtime/Decoder/Android/Java" dst="$S(BuildDir)" />
+			</true>
+		</if>
+	</resourceCopies>
+
+	<!-- optional additions to the GameActivity imports in GameActivity.java -->
+	<gameActivityImportAdditions>
+		<if condition="bSupported">
+			<true>
+			</true>
+		</if>
+	</gameActivityImportAdditions>
+
+	<!-- optional additions to the GameActivity class in GameActivity.java -->
+	<gameActivityClassAdditions>
+		<if condition="bSupported">
+			<true>
+			</true>
+		</if>
+	</gameActivityClassAdditions>
+
+	<!-- optional additions to GameActivity onCreate in GameActivity.java
+	<gameActivityOnCreateAdditions>
+	</gameActivityOnCreateAdditions>
+	-->
+
+	<!-- optional additions to GameActivity onDestroy in GameActivity.java
+	<gameActivityOnDestroyAdditions>
+	</gameActivityOnDestroyAdditions>
+	-->
+
+	<!-- optional additions to GameActivity onStop in GameActivity.java -->
+	<gameActivityOnStopAdditions>
+		<if condition="bSupported">
+			<true>
+				<insert>
+				</insert>
+			</true>
+		</if>
+	</gameActivityOnStopAdditions>
+
+
+	<!-- optional additions to GameActivity onPause in GameActivity.java -->
+	<gameActivityOnPauseAdditions>
+		<if condition="bSupported">
+			<true>
+				<insert>
+				</insert>
+			</true>
+		</if>
+	</gameActivityOnPauseAdditions>
+
+	<!-- optional additions to GameActivity onResume in GameActivity.java -->
+	<gameActivityOnResumeAdditions>
+		<if condition="bSupported">
+			<true>
+				<insert>
+				</insert>
+			</true>
+		</if>
+	</gameActivityOnResumeAdditions>
+
+	<!-- optional additions to GameActivity onActivityResult in GameActivity.java
+	<gameActivityOnActivityResultAdditions>
+	</gameActivityOnActivityResultAdditions>
+	-->
+
+	<!-- optional libraries to load in GameActivity.java before libUE4.so -->
+	<soLoadLibrary>
+	</soLoadLibrary>
+</root>
+
