ElectraPlayer: RELEASE. Moved from NFL to public. Note that this is done via a Delete and "Add New" operation as discussed.

diff --git a/Engine/Plugins/Media/ElectraPlayer/Source/ElectraPlayerRuntime/Private/ElectraPlayer.h b/Engine/Plugins/Media/ElectraPlayer/Source/ElectraPlayerRuntime/Private/ElectraPlayer.h
new file mode 100644
--- /dev/null
+++ b/Engine/Plugins/Media/ElectraPlayer/Source/ElectraPlayerRuntime/Private/ElectraPlayer.h
@@ -0,0 +1,484 @@
+// Copyright Epic Games, Inc. All Rights Reserved.
+
+#pragma once
+
+
+#include "Containers/UnrealString.h"
+#include "Containers/Queue.h"
+#include "Misc/Guid.h"
+#include "Templates/SharedPointer.h"
+#include "IAnalyticsProviderET.h"
+
+#include "IElectraPlayerInterface.h"
+
+#include "Player/AdaptiveStreamingPlayer.h"
+
+class FVideoDecoderOutput;
+using FVideoDecoderOutputPtr = TSharedPtr<FVideoDecoderOutput, ESPMode::ThreadSafe>;
+class IAudioDecoderOutput;
+using IAudioDecoderOutputPtr = TSharedPtr<IAudioDecoderOutput, ESPMode::ThreadSafe>;
+
+namespace Electra
+{
+class IVideoDecoderResourceDelegate;
+}
+
+class FElectraRendererVideo;
+class FElectraRendererAudio;
+
+using namespace Electra;
+
+DECLARE_MULTICAST_DELEGATE_TwoParams(FElectraPlayerSendAnalyticMetricsDelegate, const TSharedPtr<IAnalyticsProviderET>& /*AnalyticsProvider*/, const FGuid& /*PlayerGuid*/);
+DECLARE_MULTICAST_DELEGATE_OneParam(FElectraPlayerSendAnalyticMetricsPerMinuteDelegate, const TSharedPtr<IAnalyticsProviderET>& /*AnalyticsProvider*/);
+DECLARE_MULTICAST_DELEGATE_TwoParams(FElectraPlayerReportVideoStreamingErrorDelegate, const FGuid& /*PlayerGuid*/, const FString& /*LastError*/);
+DECLARE_MULTICAST_DELEGATE_FourParams(FElectraPlayerReportSubtitlesMetricsDelegate, const FGuid& /*PlayerGuid*/, const FString& /*URL*/, double /*ResponseTime*/, const FString& /*LastError*/);
+
+class FElectraPlayer
+	: public IElectraPlayerInterface
+	, public IAdaptiveStreamingPlayerMetrics
+{
+public:
+	FElectraPlayer(const TSharedPtr<IElectraPlayerAdapterDelegate, ESPMode::ThreadSafe>& AdapterDelegate,
+			  FElectraPlayerSendAnalyticMetricsDelegate& InSendAnalyticMetricsDelegate,
+			  FElectraPlayerSendAnalyticMetricsPerMinuteDelegate& InSendAnalyticMetricsPerMinuteDelegate,
+			  FElectraPlayerReportVideoStreamingErrorDelegate& InReportVideoStreamingErrorDelegate,
+			  FElectraPlayerReportSubtitlesMetricsDelegate& InReportSubtitlesFileMetricsDelegate);
+	~FElectraPlayer();
+
+	void OnVideoDecoded(const FVideoDecoderOutputPtr& DecoderOutput, bool bDoNotRender);
+	void OnVideoFlush();
+	void OnAudioDecoded(const IAudioDecoderOutputPtr& DecoderOutput);
+	void OnAudioFlush();
+
+	void OnVideoRenderingStarted();
+	void OnVideoRenderingStopped();
+	void OnAudioRenderingStarted();
+	void OnAudioRenderingStopped();
+
+	void SendAnalyticMetrics(const TSharedPtr<IAnalyticsProviderET>& AnalyticsProvider, const FGuid& InPlayerGuid);
+	void SendAnalyticMetricsPerMinute(const TSharedPtr<IAnalyticsProviderET>& AnalyticsProvider);
+	void SendPendingAnalyticMetrics(const TSharedPtr<IAnalyticsProviderET>& AnalyticsProvider);
+	void ReportVideoStreamingError(const FGuid& InPlayerGuid, const FString& LastError);
+	void ReportSubtitlesMetrics(const FGuid& InPlayerGuid, const FString& URL, double ResponseTime, const FString& LastError);
+
+	void DropOldFramesFromPresentationQueue();
+
+	bool CanPresentVideoFrames(uint64 NumFrames);
+	bool CanPresentAudioFrames(uint64 NumFrames);
+
+	FString GetUrl() const
+	{
+		return MediaUrl;
+	}
+
+	void SetGuid(const FGuid& Guid)
+	{
+		PlayerGuid = Guid;
+	}
+
+	void SetAsyncResourceReleaseNotification(IAsyncResourceReleaseNotifyContainer* AsyncResourceReleaseNotification) override;
+
+	// -------- PlayerAdapter (Plugin/Native) API
+
+	bool OpenInternal(const FString& Url, const FParamDict & PlayerOptions, const FPlaystartOptions & InPlaystartOptions) override;
+	void CloseInternal(bool bKillAfterClose) override;
+
+	void Tick(FTimespan DeltaTime, FTimespan Timecode) override;
+
+	bool IsKillAfterCloseAllowed() const override { return bAllowKillAfterCloseEvent;  }
+
+	EPlayerState GetState() const override;
+	EPlayerStatus GetStatus() const override;
+
+	bool IsLooping() const override;
+	bool SetLooping(bool bLooping) override;
+
+	FTimespan GetTime() const override;
+	FTimespan GetDuration() const override;
+
+	bool IsLive() const override;
+	FTimespan GetSeekableDuration() const override;
+
+	float GetRate() const override;
+	bool SetRate(float Rate) override;
+
+	bool Seek(const FTimespan& Time) override;
+
+	bool GetAudioTrackFormat(int32 TrackIndex, int32 FormatIndex, FAudioTrackFormat& OutFormat) const override;
+	bool GetVideoTrackFormat(int32 TrackIndex, int32 FormatIndex, FVideoTrackFormat& OutFormat) const override;
+
+	int32 GetNumTracks(EPlayerTrackType TrackType) const override;
+	int32 GetNumTrackFormats(EPlayerTrackType TrackType, int32 TrackIndex) const override;
+	int32 GetSelectedTrack(EPlayerTrackType TrackType) const override;
+	FText GetTrackDisplayName(EPlayerTrackType TrackType, int32 TrackIndex) const override;
+	int32 GetTrackFormat(EPlayerTrackType TrackType, int32 TrackIndex) const override;
+	FString GetTrackLanguage(EPlayerTrackType TrackType, int32 TrackIndex) const override;
+	FString GetTrackName(EPlayerTrackType TrackType, int32 TrackIndex) const override;
+	bool SelectTrack(EPlayerTrackType TrackType, int32 TrackIndex) override;
+
+	void NotifyOfOptionChange() override;
+
+private:
+	void CalculateTargetSeekTime(FTimespan& OutTargetTime, const FTimespan& InTime);
+
+	bool PresentVideoFrame(const FVideoDecoderOutputPtr& InVideoFrame);
+	bool PresentAudioFrame(const IAudioDecoderOutputPtr& DecoderOutput);
+	
+	void PlatformNotifyOfOptionChange();
+
+	// Methods from IAdaptiveStreamingPlayerMetrics
+	virtual void ReportOpenSource(const FString& URL) override;
+	virtual void ReportReceivedMasterPlaylist(const FString& EffectiveURL) override;
+	virtual void ReportReceivedPlaylists() override;
+	virtual void ReportPlaylistDownload(const Metrics::FPlaylistDownloadStats& PlaylistDownloadStats) override;
+	virtual void ReportBufferingStart(Metrics::EBufferingReason BufferingReason) override;
+	virtual void ReportBufferingEnd(Metrics::EBufferingReason BufferingReason) override;
+	virtual void ReportBandwidth(int64 EffectiveBps, int64 ThroughputBps, double LatencyInSeconds) override;
+	virtual void ReportBufferUtilization(const Metrics::FBufferStats& BufferStats) override;
+	virtual void ReportSegmentDownload(const Metrics::FSegmentDownloadStats& SegmentDownloadStats) override;
+	virtual void ReportLicenseKey(const Metrics::FLicenseKeyStats& LicenseKeyStats) override;
+	virtual void ReportDataAvailabilityChange(const Metrics::FDataAvailabilityChange& DataAvailability) override;
+	virtual void ReportVideoQualityChange(int32 NewBitrate, int32 PreviousBitrate, bool bIsDrasticDownswitch) override;
+	virtual void ReportPrerollStart() override;
+	virtual void ReportPrerollEnd() override;
+	virtual void ReportPlaybackStart() override;
+	virtual void ReportPlaybackPaused() override;
+	virtual void ReportPlaybackResumed() override;
+	virtual void ReportPlaybackEnded() override;
+	virtual void ReportJumpInPlayPosition(const FTimeValue& ToNewTime, const FTimeValue& FromTime, Metrics::ETimeJumpReason TimejumpReason) override;
+	virtual void ReportPlaybackStopped() override;
+	virtual void ReportError(const FString& ErrorReason) override;
+	virtual void ReportLogMessage(IInfoLog::ELevel InLogLevel, const FString& InLogMessage, int64 InPlayerWallclockMilliseconds) override;
+	virtual void ReportDroppedVideoFrame() override;
+	virtual void ReportDroppedAudioFrame() override;
+	
+	void LogPresentationFramesQueues(FTimespan DeltaTime);
+
+	void MediaStateOnPreparingFinished();
+	bool MediaStateOnPlay();
+	bool MediaStateOnPause();
+	void MediaStateOnEndReached();
+	void MediaStateOnSeekFinished();
+	void MediaStateOnError();
+
+	TSharedPtr<FStreamMetadata, ESPMode::ThreadSafe> GetTrackStreamMetadata(EPlayerTrackType TrackType, int32 TrackIndex) const;
+
+	// Delegate to talk back to adapter host
+	TWeakPtr<IElectraPlayerAdapterDelegate, ESPMode::ThreadSafe>	AdapterDelegate;
+
+	// Contains number of audio tracks available to expose it later.
+	int32											NumTracksAudio;
+	int32											NumTracksVideo;
+	int32											SelectedQuality;
+	int32											SelectedVideoTrackIndex;
+	int32											SelectedAudioTrackIndex;
+
+	FIntPoint										LastPresentedFrameDimension;
+
+
+	/** Media player Guid */
+	FGuid											PlayerGuid;
+	/** Metric delegates */
+	FElectraPlayerSendAnalyticMetricsDelegate&			SendAnalyticMetricsDelegate;
+	FElectraPlayerSendAnalyticMetricsPerMinuteDelegate&	SendAnalyticMetricsPerMinuteDelegate;
+	FElectraPlayerReportVideoStreamingErrorDelegate&		ReportVideoStreamingErrorDelegate;
+	FElectraPlayerReportSubtitlesMetricsDelegate&		ReportSubtitlesMetricsDelegate;
+
+	/** Option interface **/
+	FPlaystartOptions								PlaystartOptions;
+
+	/**  */
+	TAtomic<EPlayerState>							State;
+	TAtomic<EPlayerStatus>							Status;
+	bool											bWasClosedOnError;
+
+	bool											bAllowKillAfterCloseEvent;
+
+	/** Queued events */
+	TQueue<IElectraPlayerAdapterDelegate::EPlayerEvent>	DeferredEvents;
+
+	/** The URL of the currently opened media. */
+	FString											MediaUrl;
+
+	class FInternalPlayerImpl
+	{
+	public:
+		/** The media player itself **/
+		TSharedPtr<IAdaptiveStreamingPlayer, ESPMode::ThreadSafe>	AdaptivePlayer;
+
+		/** Renderers to use **/
+		TSharedPtr<FElectraRendererVideo, ESPMode::ThreadSafe>		RendererVideo;
+		TSharedPtr<FElectraRendererAudio, ESPMode::ThreadSafe>		RendererAudio;
+
+		/** */
+		static void DoCloseAsync(TSharedPtr<FInternalPlayerImpl, ESPMode::ThreadSafe> && Player, TSharedPtr<IAsyncResourceReleaseNotifyContainer, ESPMode::ThreadSafe> AsyncDestructNotification);
+	};
+
+	TSharedPtr<FInternalPlayerImpl, ESPMode::ThreadSafe>			CurrentPlayer;
+
+	TSharedPtr<IAsyncResourceReleaseNotifyContainer, ESPMode::ThreadSafe> AsyncResourceReleaseNotification;
+
+	class FAdaptiveStreamingPlayerResourceProvider : public IAdaptiveStreamingPlayerResourceProvider
+	{
+	public:
+		FAdaptiveStreamingPlayerResourceProvider(const TWeakPtr<IElectraPlayerAdapterDelegate, ESPMode::ThreadSafe> & AdapterDelegate);
+		virtual ~FAdaptiveStreamingPlayerResourceProvider() = default;
+		
+		virtual void ProvideStaticPlaybackDataForURL(TSharedPtr<IAdaptiveStreamingPlayerResourceRequest, ESPMode::ThreadSafe> InOutRequest) override;
+
+		void ProcessPendingStaticResourceRequests();
+		void ClearPendingRequests();
+
+	private:
+		/** Requests for static resource fetches we want to perform on the main thread **/
+		TQueue<TSharedPtr<IAdaptiveStreamingPlayerResourceRequest, ESPMode::ThreadSafe>, EQueueMode::Mpsc> PendingStaticResourceRequests;
+
+		// Player adapter delegate
+		TWeakPtr<IElectraPlayerAdapterDelegate, ESPMode::ThreadSafe>	AdapterDelegate;
+	};
+
+	TSharedPtr<FAdaptiveStreamingPlayerResourceProvider, ESPMode::ThreadSafe> StaticResourceProvider;
+	TSharedPtr<IVideoDecoderResourceDelegate, ESPMode::ThreadSafe> VideoDecoderResourceDelegate;
+
+
+	class FAverageValue
+	{
+	public:
+		FAverageValue()
+		: Samples(nullptr)
+		, NumSamples(0)
+		, MaxSamples(0)
+		{
+		}
+		~FAverageValue()
+		{
+			delete[] Samples;
+		}
+		void SetNumSamples(int32 InMaxSamples)
+		{
+			check(InMaxSamples > 0);
+			delete[] Samples;
+			NumSamples = 0;
+			MaxSamples = InMaxSamples;
+			Samples = new double[MaxSamples];
+		}
+		void AddValue(double Value)
+		{
+			Samples[NumSamples % MaxSamples] = Value;
+			++NumSamples;
+		}
+		void Reset()
+		{
+			NumSamples = 0;
+		}
+		double GetAverage() const
+		{
+			double Avg = 0.0;
+			if (NumSamples > 0)
+			{
+				double Sum = 0.0;
+				int32 Last = NumSamples <= MaxSamples ? NumSamples : MaxSamples;
+				for (int32 i = 0; i < Last; ++i)
+				{
+					Sum += Samples[i];
+				}
+				Avg = Sum / Last;
+			}
+			return Avg;
+		}
+	private:
+		double*	Samples;
+		int32	NumSamples;
+		int32	MaxSamples;
+	};
+
+	struct FDroppedFrameStats
+	{
+		enum class EFrameType
+		{
+			Undefined,
+			VideoFrame,
+			AudioFrame
+		};
+		FDroppedFrameStats()
+		{
+			FrameType = EFrameType::Undefined;
+			Reset();
+		}
+		void SetFrameType(EFrameType InFrameType)
+		{
+			FrameType = InFrameType;
+		}
+		void Reset()
+		{
+			NumTotalDropped = 0;
+			WorstDeltaTime = FTimespan::Zero();
+			PlayerTimeAtLastReport = FTimespan::MinValue();
+			SystemTimeAtLastReport = 0.0;
+			LogWarningAfterSeconds = 10.0;
+		}
+		void SetLogWarningInterval(double InSecondsBetweenWarnings)
+		{
+			LogWarningAfterSeconds = InSecondsBetweenWarnings;
+		}
+		void AddNewDrop(const FTimespan& InFrameTime, const FTimespan& InPlayerTime, void* InPtrElectraPlayer, void* InPtrCurrentPlayer);
+
+		EFrameType	FrameType;
+		uint32		NumTotalDropped;
+		FTimespan	WorstDeltaTime;
+		FTimespan	PlayerTimeAtLastReport;
+		double		SystemTimeAtLastReport;
+		double		LogWarningAfterSeconds;
+	};
+
+	struct FStatistics
+	{
+		struct FBandwidth
+		{
+			FBandwidth()
+			{
+				Bandwidth.SetNumSamples(3);
+				Latency.SetNumSamples(3);
+				Reset();
+			}
+			void Reset()
+			{
+				Bandwidth.Reset();
+				Latency.Reset();
+			}
+			void AddSample(double InBytesPerSecond, double InLatency)
+			{
+				Bandwidth.AddValue(InBytesPerSecond);
+				Latency.AddValue(InLatency);
+			}
+			double GetAverageBandwidth() const
+			{
+				return Bandwidth.GetAverage();
+			}
+			double GetAverageLatency() const
+			{
+				return Latency.GetAverage();
+			}
+			FAverageValue	Bandwidth;
+			FAverageValue	Latency;
+		};
+		FStatistics()
+		{
+			DroppedVideoFrames.SetFrameType(FDroppedFrameStats::EFrameType::VideoFrame);
+			DroppedAudioFrames.SetFrameType(FDroppedFrameStats::EFrameType::AudioFrame);
+			Reset();
+		}
+		void Reset()
+		{
+			InitialURL.Empty();
+			CurrentlyActivePlaylistURL.Empty();
+			LastError.Empty();
+			LastState = "Empty";
+			TimeAtOpen  					= -1.0;
+			TimeToLoadMasterPlaylist		= -1.0;
+			TimeToLoadStreamPlaylists   	= -1.0;
+			InitialBufferingDuration		= -1.0;
+			InitialStreamBitrate			= 0;
+			TimeAtPrerollBegin  			= -1.0;
+			TimeForInitialPreroll   		= -1.0;
+			NumTimesRebuffered  			= 0;
+			NumTimesForwarded   			= 0;
+			NumTimesRewound 				= 0;
+			NumTimesLooped  				= 0;
+			TimeAtBufferingBegin			= 0.0;
+			TotalRebufferingDuration		= 0.0;
+			LongestRebufferingDuration  	= 0.0;
+			PlayPosAtStart  				= -1.0;
+			PlayPosAtEnd					= -1.0;
+			NumQualityUpswitches			= 0;
+			NumQualityDownswitches  		= 0;
+			NumQualityDrasticDownswitches   = 0;
+			NumVideoDatabytesStreamed   	= 0;
+			NumAudioDatabytesStreamed   	= 0;
+			NumSegmentDownloadsAborted  	= 0;
+			CurrentlyActiveResolutionWidth  = 0;
+			CurrentlyActiveResolutionHeight = 0;
+			VideoSegmentBitratesStreamed.Empty();
+			InitialBufferingBandwidth.Reset();
+			bIsInitiallyDownloading 		= false;
+			bDidPlaybackEnd 				= false;
+			SubtitlesURL.Empty();
+			SubtitlesResponseTime			= 0.0;
+			SubtitlesLastError.Empty();
+			DroppedVideoFrames.Reset();
+			DroppedAudioFrames.Reset();
+			MediaTimelineAtStart.Reset();
+			MediaTimelineAtEnd.Reset();
+			MediaDuration = 0.0;
+		}
+
+		FString					InitialURL;
+		FString					CurrentlyActivePlaylistURL;
+		FString					LastError;
+		FString					LastState;		// "Empty", "Opening", "Preparing", "Buffering", "Idle", "Ready", "Playing", "Paused", "Seeking", "Rebuffering", "Ended"
+		double					TimeAtOpen;
+		double					TimeToLoadMasterPlaylist;
+		double					TimeToLoadStreamPlaylists;
+		double					InitialBufferingDuration;
+		int32					InitialStreamBitrate;
+		double					TimeAtPrerollBegin;
+		double					TimeForInitialPreroll;
+		int32					NumTimesRebuffered;
+		int32					NumTimesForwarded;
+		int32					NumTimesRewound;
+		int32					NumTimesLooped;
+		double					TimeAtBufferingBegin;
+		double					TotalRebufferingDuration;
+		double					LongestRebufferingDuration;
+		double					PlayPosAtStart;
+		double					PlayPosAtEnd;
+		int32					NumQualityUpswitches;
+		int32					NumQualityDownswitches;
+		int32					NumQualityDrasticDownswitches;
+		int64					NumVideoDatabytesStreamed;
+		int64					NumAudioDatabytesStreamed;
+		int32					NumSegmentDownloadsAborted;
+		int32					CurrentlyActiveResolutionWidth;
+		int32					CurrentlyActiveResolutionHeight;
+		TMap<int32, uint32>		VideoSegmentBitratesStreamed;		// key=video stream bitrate, value=number of segments loaded at this rate
+		FBandwidth				InitialBufferingBandwidth;
+		bool					bIsInitiallyDownloading;
+		bool					bDidPlaybackEnd;
+		FString					SubtitlesURL;
+		double					SubtitlesResponseTime;
+		FString					SubtitlesLastError;
+		FDroppedFrameStats		DroppedVideoFrames;
+		FDroppedFrameStats		DroppedAudioFrames;
+		FTimeRange				MediaTimelineAtStart;
+		FTimeRange				MediaTimelineAtEnd;
+		double					MediaDuration;
+	};
+
+	struct FAnalyticsEvent
+	{
+		FString EventName;
+		TArray<FAnalyticsEventAttribute> ParamArray;
+	};
+
+	void UpdatePlayEndStatistics();
+	void LogStatistics();
+	void AddCommonAnalyticsAttributes(TArray<FAnalyticsEventAttribute>& InOutParamArray);
+	TSharedPtr<FAnalyticsEvent> CreateAnalyticsEvent(FString InEventName);
+	void EnqueueAnalyticsEvent(TSharedPtr<FAnalyticsEvent> InAnalyticEvent);
+
+	FCriticalSection						StatisticsLock;
+	FStatistics								Statistics;
+
+	TQueue<TSharedPtr<FAnalyticsEvent>>		QueuedAnalyticEvents;
+	int32									NumQueuedAnalyticEvents = 0;
+	FString									AnalyticsOSVersion;
+	FString									AnalyticsGPUType;
+
+	/** Unique player instance GUID sent with each analytics event. This allows finding all events of a particular playback session. **/
+	FString									AnalyticsInstanceGuid;
+	/** Sequential analytics event number. Helps sorting events. **/
+	uint32									AnalyticsInstanceEventCount;
+};
+
+ENUM_CLASS_FLAGS(FElectraPlayer::EPlayerStatus);
+
diff --git a/Engine/Plugins/Media/ElectraPlayer/Source/ElectraPlayerRuntime/Private/ElectraPlayerAudioDecoderOutput.cpp b/Engine/Plugins/Media/ElectraPlayer/Source/ElectraPlayerRuntime/Private/ElectraPlayerAudioDecoderOutput.cpp
new file mode 100644
--- /dev/null
+++ b/Engine/Plugins/Media/ElectraPlayer/Source/ElectraPlayerRuntime/Private/ElectraPlayerAudioDecoderOutput.cpp
@@ -0,0 +1,151 @@
+// Copyright Epic Games, Inc. All Rights Reserved.
+
+#include "MediaAudioDecoderOutput.h"
+#include "Renderer/RendererAudio.h"
+
+class FElectraPlayerAudioDecoderOutput : public IAudioDecoderOutput
+{
+public:
+	FElectraPlayerAudioDecoderOutput()
+	: BufferMaxSizeBytes(0)
+	, BufferUsedSizeBytes(0)
+	, Buffer(nullptr)
+	{
+	}
+
+	~FElectraPlayerAudioDecoderOutput() override
+	{
+		Reserve(0);
+	}
+
+	void Reserve(uint32 InBufferSizeBytes) override
+	{
+		if (InBufferSizeBytes != BufferMaxSizeBytes)
+		{
+			if (Buffer)
+			{
+				FMemory::Free(Buffer);
+				Buffer = nullptr;
+			}
+			BufferMaxSizeBytes = InBufferSizeBytes;
+			if (BufferMaxSizeBytes)
+			{
+				Buffer = FMemory::Malloc(InBufferSizeBytes);
+				check(Buffer);
+			}
+		}
+	}
+
+	void SetOwner(const TSharedPtr<IDecoderOutputOwner, ESPMode::ThreadSafe>& Renderer) override
+	{
+		OwningRenderer = Renderer;
+	}
+
+	void Initialize(ESampleFormat InFormat, uint32 InNumChannels, uint32 InSampleRate, FTimespan InDuration, const FDecoderTimeStamp & InPts, uint32 InBufferSizeBytes) override
+	{
+		check(BufferMaxSizeBytes >= InBufferSizeBytes);
+
+		Format = InFormat;
+		NumChannels = InNumChannels;
+		SampleRate = InSampleRate;
+
+		Time = InPts;
+		Duration = InDuration;
+
+		BufferUsedSizeBytes = InBufferSizeBytes;
+	}
+
+	const void* GetBuffer() const override
+	{
+		return Buffer;
+	}
+
+	uint32 GetUsedBufferBytes() const override
+	{
+		return BufferUsedSizeBytes;
+	}
+
+	uint32 GetReservedBufferBytes() const override
+	{
+		return BufferMaxSizeBytes;
+	}
+
+	uint32 GetChannels() const override
+	{
+		return NumChannels;
+	}
+
+	uint32 GetFrames() const override
+	{
+		if (Format == ESampleFormat::Undefined)
+		{
+			return 0;
+		}
+
+		check((int32)Format <= (int32)ESampleFormat::Int32);
+
+		static_assert((int32)ESampleFormat::Undefined == 0, "Check order of enum");
+		static_assert((int32)ESampleFormat::Double == 1, "Check order of enum");
+		static_assert((int32)ESampleFormat::Float == 2, "Check order of enum");
+		static_assert((int32)ESampleFormat::Int8 == 3, "Check order of enum");
+		static_assert((int32)ESampleFormat::Int16 == 4, "Check order of enum");
+		static_assert((int32)ESampleFormat::Int32 == 5, "Check order of enum");
+
+		static const uint32 _SampleSize[] = {0, 8, 4, 1, 2, 4};
+		return BufferUsedSizeBytes / (NumChannels * _SampleSize[(uint32)Format]);
+	}
+
+	ESampleFormat GetFormat() const override
+	{
+		return Format;
+	}
+
+	FDecoderTimeStamp GetTime() const override
+	{
+		return Time;
+	}
+
+	FTimespan GetDuration() const override
+	{
+		return Duration;
+	}
+
+	uint32 GetSampleRate() const override
+	{
+		return SampleRate;
+	}
+
+	void ShutdownPoolable() override
+	{
+		TSharedPtr<IDecoderOutputOwner, ESPMode::ThreadSafe> lockedAudioRenderer = OwningRenderer.Pin();
+		if (lockedAudioRenderer.IsValid())
+		{
+			lockedAudioRenderer->SampleReleasedToPool(Duration);
+		}
+	}
+
+private:
+	uint32 BufferMaxSizeBytes;
+	uint32 BufferUsedSizeBytes;
+	void *Buffer;
+
+	ESampleFormat Format;
+	uint32 NumChannels;
+	uint32 SampleRate;
+
+	FDecoderTimeStamp Time;
+	FTimespan Duration;
+
+	TWeakPtr<IDecoderOutputOwner, ESPMode::ThreadSafe> OwningRenderer;
+};
+
+// =================================================================================================================================
+
+/*
+	Create one and the same audio sample for all platforms - we so far do not need any specialization here (otherwise this should be with the respective decoder - much like the class definition / impl above)
+*/
+IAudioDecoderOutput *FElectraPlayerPlatformAudioDecoderOutputFactory::Create()
+{
+	return new FElectraPlayerAudioDecoderOutput();
+}
+
diff --git a/Engine/Plugins/Media/ElectraPlayer/Source/ElectraPlayerRuntime/Private/ElectraPlayerPlatform.h b/Engine/Plugins/Media/ElectraPlayer/Source/ElectraPlayerRuntime/Private/ElectraPlayerPlatform.h
new file mode 100644
--- /dev/null
+++ b/Engine/Plugins/Media/ElectraPlayer/Source/ElectraPlayerRuntime/Private/ElectraPlayerPlatform.h
@@ -0,0 +1,15 @@
+// Copyright Epic Games, Inc. All Rights Reserved.
+
+#pragma once
+
+namespace Electra {
+
+bool PlatformEarlyStartup();
+bool PlatformMemorySetup();
+bool PlatformShutdown();
+
+class IVideoDecoderResourceDelegate;
+TSharedPtr<IVideoDecoderResourceDelegate, ESPMode::ThreadSafe> PlatformCreateVideoDecoderResourceDelegate(const TWeakPtr<IElectraPlayerAdapterDelegate, ESPMode::ThreadSafe>& AdapterDelegate);
+
+}; // namespace Electra
+
diff --git a/Engine/Plugins/Media/ElectraPlayer/Source/ElectraPlayerRuntime/Private/ElectraPlayerRuntimeModule.cpp b/Engine/Plugins/Media/ElectraPlayer/Source/ElectraPlayerRuntime/Private/ElectraPlayerRuntimeModule.cpp
new file mode 100644
--- /dev/null
+++ b/Engine/Plugins/Media/ElectraPlayer/Source/ElectraPlayerRuntime/Private/ElectraPlayerRuntimeModule.cpp
@@ -0,0 +1,113 @@
+// Copyright Epic Games, Inc. All Rights Reserved.
+
+#include "CoreTypes.h"
+
+#include "Modules/ModuleManager.h"
+#include "Misc/ConfigCacheIni.h"
+
+#include "PlayerRuntimeGlobal.h"
+#include "Core/MediaInterlocked.h"
+#include "Renderer/RendererVideo.h"
+
+#include "Renderer/RendererAudio.h"
+#include "Decoder/VideoDecoderH264.h"
+#include "Decoder/AudioDecoderAAC.h"
+
+#include "ElectraPlayerPrivate.h"
+#include "IElectraPlayerRuntimeModule.h"
+#include "ElectraPlayer.h"
+#include "ElectraPlayerPlatform.h"
+
+DEFINE_LOG_CATEGORY(LogElectraPlayer);
+DEFINE_LOG_CATEGORY(LogElectraPlayerStats);
+
+#define LOCTEXT_NAMESPACE "ElectraPlayerRuntimeModule"
+
+// -----------------------------------------------------------------------------------------------------------------------------------
+
+// Implements the ElectraPlayer Runtime Module module.
+class FElectraPlayerRuntimeModule : public IElectraPlayerRuntimeModule
+{
+public:
+	// IElectraPlayerRuntimeModule interface
+
+	virtual bool IsInitialized() const override
+	{
+		return bInitialized;
+	}
+
+public:
+	// IModuleInterface interface
+
+	void StartupModule() override
+	{
+		check(!bInitialized);
+
+		if (!Electra::PlatformEarlyStartup())
+		{
+
+			return;
+		}
+
+		// Hook-up platform memory system (if any)
+		Electra::PlatformMemorySetup();
+
+		// Core
+		Electra::Configuration ElectraConfig;
+		FString AnalyticsEventsFromIni;
+	
+		// Read Analytics Events from ini configuring the events sent during playback
+		if (GConfig->GetString(TEXT("ElectraPlayer"), TEXT("AnalyticsEvents"), AnalyticsEventsFromIni, GEngineIni))
+		{
+			// Parse comma delimited strings into arrays
+			TArray<FString> EnabledAnalyticsEvents;
+			AnalyticsEventsFromIni.ParseIntoArray(EnabledAnalyticsEvents, TEXT(","), /*bCullEmpty=*/true);
+			for (auto EnabledEvent : EnabledAnalyticsEvents)
+			{
+				EnabledEvent.TrimStartAndEndInline();
+				ElectraConfig.EnabledAnalyticsEvents.Add(EnabledEvent, true);
+			}
+		}
+		Electra::Startup(ElectraConfig);
+
+		// Video renderer
+		FElectraRendererVideo::SystemConfiguration SysCfgVideo;
+		FElectraRendererVideo::Startup(SysCfgVideo);
+		// Audio renderer
+		FElectraRendererAudio::SystemConfiguration SysCfgAudio;
+		FElectraRendererAudio::Startup(SysCfgAudio);
+
+		// H.264 decoder
+		IVideoDecoderH264::FSystemConfiguration SysCfgH264;
+		IVideoDecoderH264::Startup(SysCfgH264);
+
+		// AAC decoder
+		IAudioDecoderAAC::FSystemConfiguration SysCfgAAC;
+		IAudioDecoderAAC::Startup(SysCfgAAC);
+
+		bInitialized = true;
+	}
+
+	void ShutdownModule() override
+	{
+		if (bInitialized)
+		{
+			Electra::WaitForAllPlayersToHaveTerminated();
+			IAudioDecoderAAC::Shutdown();
+			IVideoDecoderH264::Shutdown();
+			FElectraRendererAudio::Shutdown();
+			FElectraRendererVideo::Shutdown();
+			Electra::Shutdown();
+			Electra::PlatformShutdown();
+			bInitialized = false;
+		}
+	}
+
+private:
+	bool bInitialized = false;
+};
+
+IMPLEMENT_MODULE(FElectraPlayerRuntimeModule, ElectraPlayerRuntime);
+
+#undef LOCTEXT_NAMESPACE
+
diff --git a/Engine/Plugins/Media/ElectraPlayer/Source/ElectraPlayerRuntime/Private/ElectraPlayer_Android.cpp b/Engine/Plugins/Media/ElectraPlayer/Source/ElectraPlayerRuntime/Private/ElectraPlayer_Android.cpp
new file mode 100644
--- /dev/null
+++ b/Engine/Plugins/Media/ElectraPlayer/Source/ElectraPlayerRuntime/Private/ElectraPlayer_Android.cpp
@@ -0,0 +1,99 @@
+// Copyright Epic Games, Inc. All Rights Reserved.
+
+#include "CoreTypes.h"
+
+#if PLATFORM_ANDROID
+
+#include "IElectraPlayerInterface.h"
+#include "ElectraPlayerPrivate.h"
+#include "ElectraPlayer.h"
+#include "ParameterDictionary.h"
+#include "VideoDecoderResourceDelegate.h"
+
+// ----------------------------------------------------------------------------------------------------------------------
+// ----------------------------------------------------------------------------------------------------------------------
+
+namespace Electra {
+
+/**
+ * Early startup during module load process
+ */
+bool PlatformEarlyStartup()
+{
+	return true;
+}
+
+/**
+ * Any platform specific memory hooks
+ */
+bool PlatformMemorySetup()
+{
+	return true;
+}
+
+/**
+ * Shutdown of module
+ */
+bool PlatformShutdown()
+{
+	return true;
+}
+
+class FVideoDecoderResourceDelegate : public IVideoDecoderResourceDelegate
+{
+public:
+	FVideoDecoderResourceDelegate(const TWeakPtr<IElectraPlayerAdapterDelegate, ESPMode::ThreadSafe>& InAdapterDelegate) : PlayerAdapterDelegate(InAdapterDelegate) {}
+	virtual ~FVideoDecoderResourceDelegate() {}
+
+	virtual jobject VideoDecoderResourceDelegate_GetCodecSurface() override
+	{
+		if (const TSharedPtr<IElectraPlayerAdapterDelegate, ESPMode::ThreadSafe>& PinnedAdapterDelegate = PlayerAdapterDelegate.Pin())
+		{
+			IElectraPlayerResourceDelegate* ResourceDelegate = PinnedAdapterDelegate->GetResourceDelegate();
+			if (ResourceDelegate)
+			{
+				return ResourceDelegate->GetCodecSurface();
+			}
+		}
+		return nullptr;
+	}
+
+private:
+	TWeakPtr<IElectraPlayerAdapterDelegate, ESPMode::ThreadSafe> PlayerAdapterDelegate;
+};
+
+TSharedPtr<IVideoDecoderResourceDelegate, ESPMode::ThreadSafe> PlatformCreateVideoDecoderResourceDelegate(const TWeakPtr<IElectraPlayerAdapterDelegate, ESPMode::ThreadSafe> & AdapterDelegate)
+{
+	return MakeShared<FVideoDecoderResourceDelegate, ESPMode::ThreadSafe>(AdapterDelegate);
+}
+
+} //namespace Electra
+
+// ----------------------------------------------------------------------------------------------------------------------
+// ----------------------------------------------------------------------------------------------------------------------
+
+/**
+ * Public call for pre-use initialization after module is loaded, allowing for user passed parameters
+ */
+bool FElectraPlayerPlatform::StartupPlatformResources(const Electra::FParamDict& Params)
+{
+	return true;
+}
+
+
+void FElectraPlayer::PlatformNotifyOfOptionChange()
+{
+	if (TSharedPtr<IElectraPlayerAdapterDelegate, ESPMode::ThreadSafe> PinnedAdapterDelegate = AdapterDelegate.Pin())
+	{
+		FVariantValue Value = PinnedAdapterDelegate->QueryOptions(IElectraPlayerAdapterDelegate::EOptionType::Android_VideoSurface);
+		if (Value.IsValid())
+		{
+			if (CurrentPlayer.IsValid() && CurrentPlayer->AdaptivePlayer.IsValid())
+			{
+				CurrentPlayer->AdaptivePlayer->Android_UpdateSurface(Value.GetSharedPointer<IOptionPointerValueContainer>());
+			}
+		}
+	}
+}
+
+#endif
diff --git a/Engine/Plugins/Media/ElectraPlayer/Source/ElectraPlayerRuntime/Private/ElectraPlayer_Apple.cpp b/Engine/Plugins/Media/ElectraPlayer/Source/ElectraPlayerRuntime/Private/ElectraPlayer_Apple.cpp
new file mode 100644
--- /dev/null
+++ b/Engine/Plugins/Media/ElectraPlayer/Source/ElectraPlayerRuntime/Private/ElectraPlayer_Apple.cpp
@@ -0,0 +1,66 @@
+// Copyright Epic Games, Inc. All Rights Reserved.
+
+#include "CoreTypes.h"
+
+#if PLATFORM_IOS || PLATFORM_MAC
+
+#include "IElectraPlayerInterface.h"
+#include "ElectraPlayerPrivate.h"
+#include "ElectraPlayer.h"
+#include "ParameterDictionary.h"
+
+// ----------------------------------------------------------------------------------------------------------------------
+// ----------------------------------------------------------------------------------------------------------------------
+
+namespace Electra {
+
+/**
+ * Early startup during module load process
+ */
+bool PlatformEarlyStartup()
+{
+	return true;
+}
+
+/**
+ * Any platform specific memory hooks
+ */
+bool PlatformMemorySetup()
+{
+	return true;
+}
+
+/**
+ * Shutdown of module
+ */
+bool PlatformShutdown()
+{
+	return true;
+}
+
+TSharedPtr<IVideoDecoderResourceDelegate, ESPMode::ThreadSafe> PlatformCreateVideoDecoderResourceDelegate(const TWeakPtr<IElectraPlayerAdapterDelegate, ESPMode::ThreadSafe>& AdapterDelegate)
+{
+	return nullptr;
+}
+
+} //namespace Electra
+
+// ----------------------------------------------------------------------------------------------------------------------
+// ----------------------------------------------------------------------------------------------------------------------
+
+/**
+ * Public call for pre-use initialization after module is loaded, allowing for user passed parameters
+ */
+bool FElectraPlayerPlatform::StartupPlatformResources(const Electra::FParamDict& Params)
+{
+	return true;
+}
+
+
+void FElectraPlayer::PlatformNotifyOfOptionChange()
+{
+}
+
+#endif
+
+
diff --git a/Engine/Plugins/Media/ElectraPlayer/Source/ElectraPlayerRuntime/Private/ElectraPlayer_PC.cpp b/Engine/Plugins/Media/ElectraPlayer/Source/ElectraPlayerRuntime/Private/ElectraPlayer_PC.cpp
new file mode 100644
--- /dev/null
+++ b/Engine/Plugins/Media/ElectraPlayer/Source/ElectraPlayerRuntime/Private/ElectraPlayer_PC.cpp
@@ -0,0 +1,319 @@
+// Copyright Epic Games, Inc. All Rights Reserved.
+
+#include "CoreTypes.h"
+
+#if PLATFORM_WINDOWS
+
+
+#include "IElectraPlayerInterface.h"
+#include "ElectraPlayerPrivate.h"
+#include "ElectraPlayerPrivate_Platform.h"
+
+#include "ElectraPlayer.h"
+
+#include "ParameterDictionary.h"
+
+// ----------------------------------------------------------------------------------------------------------------------
+
+THIRD_PARTY_INCLUDES_START
+#include <mfapi.h>
+THIRD_PARTY_INCLUDES_END
+#include <d3d12.h>
+#include <dxgi1_4.h>
+#ifdef ELECTRA_HAVE_DX11
+#pragma comment(lib, "D3D11.lib")
+#endif
+#pragma comment(lib, "D3D12.lib")
+#pragma comment(lib, "DXGI.lib")
+
+// ----------------------------------------------------------------------------------------------------------------------
+// ----------------------------------------------------------------------------------------------------------------------
+
+namespace Electra {
+
+// ----------------------------------------------------------------------------------------------------------------------
+	
+TUniquePtr<FDXDeviceInfo> FDXDeviceInfo::s_DXDeviceInfo;
+
+// ----------------------------------------------------------------------------------------------------------------------
+	
+static bool bHaveDirectXResources = false;
+
+// ----------------------------------------------------------------------------------------------------------------------
+
+bool CreateDirectXResources(const FParamDict& Params)
+{
+	if (!GIsClient && !bHaveDirectXResources)
+	{
+		// We claim success in all cases - avoiding an error code if we are not a client
+		return true;
+	}
+
+	Electra::FDXDeviceInfo::s_DXDeviceInfo = MakeUnique<Electra::FDXDeviceInfo>();
+	Electra::FDXDeviceInfo::s_DXDeviceInfo->DxVersion = Electra::FDXDeviceInfo::ED3DVersion::VersionUnknown;
+
+	FDXDeviceInfo::s_DXDeviceInfo->RenderingDx11Device = nullptr;
+
+	if (Electra::IsWindows8Plus())
+	{
+		static const FString DeviceKey("Device");
+		static const FString DeviceNameKey("DeviceName");
+
+		if (!Params.HaveKey(DeviceKey) || !Params.HaveKey(DeviceNameKey))
+		{
+			UE_LOG(LogElectraPlayer, Warning, TEXT("ElectraPlayer: Missing rendering device info"));
+			return false;
+		}
+
+
+		UINT ResetToken = 0;
+		CHECK_HR(MFCreateDXGIDeviceManager(&ResetToken, Electra::FDXDeviceInfo::s_DXDeviceInfo->DxDeviceManager.GetInitReference()));
+		if (!Electra::FDXDeviceInfo::s_DXDeviceInfo->DxDeviceManager.IsValid())
+		{
+			UE_LOG(LogElectraPlayer, Warning, TEXT("ElectraPlayer: Could not create DXGI device manager"));
+			return false;
+		}
+
+		// Create device from same adapter as already existing device
+
+		uint32 UE4DxDeviceCreationFlags = 0;
+		TRefCountPtr<IDXGIAdapter> DXGIAdapter;
+
+		FString DeviceName = Params.GetValue(DeviceNameKey).GetFString();
+
+		if (DeviceName == TEXT("D3D11"))
+		{
+			ID3D11Device* UE4DxDevice = static_cast<ID3D11Device*>(Params.GetValue(DeviceKey).GetPointer());
+
+			TRefCountPtr<IDXGIDevice> DXGIDevice;
+			UE4DxDevice->QueryInterface(__uuidof(IDXGIDevice), (void**)DXGIDevice.GetInitReference());
+
+			DXGIDevice->GetAdapter((IDXGIAdapter**)DXGIAdapter.GetInitReference());
+
+			UE4DxDeviceCreationFlags = UE4DxDevice->GetCreationFlags();
+
+			FDXDeviceInfo::s_DXDeviceInfo->RenderingDx11Device = TRefCountPtr<ID3D11Device>(UE4DxDevice, true);
+
+			Electra::FDXDeviceInfo::s_DXDeviceInfo->DxVersion = Electra::FDXDeviceInfo::ED3DVersion::Version11Win8;
+		}
+		else if (DeviceName == TEXT("D3D12"))
+		{
+			TRefCountPtr<IDXGIFactory4> DXGIFactory;
+			CHECK_HR(CreateDXGIFactory(__uuidof(IDXGIFactory4), (void**)DXGIFactory.GetInitReference()));
+
+			ID3D12Device* UE4DxDevice = static_cast<ID3D12Device*>(Params.GetValue(DeviceKey).GetPointer());
+
+			LUID Luid = UE4DxDevice->GetAdapterLuid();
+
+			CHECK_HR(DXGIFactory->EnumAdapterByLuid(Luid, __uuidof(IDXGIAdapter), (void**)DXGIAdapter.GetInitReference()));
+
+			Electra::FDXDeviceInfo::s_DXDeviceInfo->DxVersion = Electra::FDXDeviceInfo::ED3DVersion::Version12Win10;
+		}
+		else
+		{
+			UE_LOG(LogElectraPlayer, Warning, TEXT("ElectraPlayer: Dynamic RHI is not for DX11 or DX12"));
+			return false;
+		}
+
+#ifdef ELECTRA_HAVE_DX11
+		D3D_FEATURE_LEVEL FeatureLevel;
+
+		uint32 DeviceCreationFlags = 0;
+		if ((UE4DxDeviceCreationFlags & D3D11_CREATE_DEVICE_DEBUG) != 0)
+		{
+			DeviceCreationFlags |= D3D11_CREATE_DEVICE_DEBUG;
+		}
+
+		CHECK_HR(D3D11CreateDevice(
+			DXGIAdapter,
+			D3D_DRIVER_TYPE_UNKNOWN,
+			nullptr,
+			DeviceCreationFlags,
+			nullptr,
+			0,
+			D3D11_SDK_VERSION,
+			Electra::FDXDeviceInfo::s_DXDeviceInfo->DxDevice.GetInitReference(),
+			&FeatureLevel,
+			Electra::FDXDeviceInfo::s_DXDeviceInfo->DxDeviceContext.GetInitReference()));
+
+		if (FeatureLevel < D3D_FEATURE_LEVEL_9_3)
+		{
+			Electra::FDXDeviceInfo::s_DXDeviceInfo->DxVersion = Electra::FDXDeviceInfo::ED3DVersion::VersionUnknown;
+			UE_LOG(LogElectraPlayer, Warning, TEXT("ElectraPlayer: Unable to Create D3D11 Device with feature level 9.3 or above"));
+			return false;
+		}
+
+		if (!Electra::FDXDeviceInfo::s_DXDeviceInfo->DxDevice.IsValid())
+		{
+			UE_LOG(LogElectraPlayer, Warning, TEXT("ElectraPlayer: Could not create DX11 device"));
+			return false;
+		}
+
+		CHECK_HR(Electra::FDXDeviceInfo::s_DXDeviceInfo->DxDeviceManager->ResetDevice(Electra::FDXDeviceInfo::s_DXDeviceInfo->DxDevice, ResetToken));
+
+		// multithread protect the newly created device as we're going to use it from decoding thread and from render thread for texture
+		// sharing between decoding and rendering DX devices
+		TRefCountPtr<ID3D10Multithread> DxMultithread;
+		Electra::FDXDeviceInfo::s_DXDeviceInfo->DxDevice->QueryInterface(__uuidof(ID3D10Multithread), (void**)DxMultithread.GetInitReference());
+		DxMultithread->SetMultithreadProtected(1);
+
+		UE_LOG(LogElectraPlayer, Log, TEXT("D3D11 Device for h/w accelerated decoding created: %p"), Electra::FDXDeviceInfo::s_DXDeviceInfo->DxDevice.GetReference());
+#endif
+		bHaveDirectXResources = true;
+	}
+	else
+	{
+#ifdef ELECTRA_HAVE_DX9
+		UINT ResetToken = 0;
+		CHECK_HR_DX9(DXVA2CreateDirect3DDeviceManager9(&ResetToken, Electra::FDXDeviceInfo::s_DXDeviceInfo->Dx9DeviceManager.GetInitReference()));
+
+		Electra::FDXDeviceInfo::s_DXDeviceInfo->Dx9 = Direct3DCreate9(D3D_SDK_VERSION);
+		check(Electra::FDXDeviceInfo::s_DXDeviceInfo->Dx9);
+		if (!Electra::FDXDeviceInfo::s_DXDeviceInfo->Dx9)
+		{
+			UE_LOG(LogElectraPlayer, Warning, TEXT("Direct3DCreate9(D3D_SDK_VERSION) failed"));
+			return false;
+		}
+
+		Electra::FDXDeviceInfo::s_DXDeviceInfo->Dx9Device = nullptr;
+		D3DPRESENT_PARAMETERS PresentParam;
+		FMemory::Memzero(PresentParam);
+		PresentParam.BackBufferWidth = 1;
+		PresentParam.BackBufferHeight = 1;
+		PresentParam.BackBufferFormat = D3DFMT_UNKNOWN;
+		PresentParam.BackBufferCount = 1;
+		PresentParam.SwapEffect = D3DSWAPEFFECT_DISCARD;
+		PresentParam.hDeviceWindow = nullptr;
+		PresentParam.Windowed = true;
+		PresentParam.Flags = D3DPRESENTFLAG_VIDEO;
+		CHECK_HR_DX9(Electra::FDXDeviceInfo::s_DXDeviceInfo->Dx9->CreateDevice(D3DADAPTER_DEFAULT, D3DDEVTYPE_HAL, nullptr, D3DCREATE_MULTITHREADED | D3DCREATE_MIXED_VERTEXPROCESSING, &PresentParam, Electra::FDXDeviceInfo::s_DXDeviceInfo->Dx9Device.GetInitReference()));
+
+		CHECK_HR_DX9(Electra::FDXDeviceInfo::s_DXDeviceInfo->Dx9DeviceManager->ResetDevice(Electra::FDXDeviceInfo::s_DXDeviceInfo->Dx9Device, ResetToken));
+
+		Electra::FDXDeviceInfo::s_DXDeviceInfo->DxVersion = Electra::FDXDeviceInfo::ED3DVersion::Version9Win7;
+		bHaveDirectXResources = true;
+#endif
+	}
+	return bHaveDirectXResources;
+}
+
+
+bool DestroyDirectXResources()
+{
+	Electra::FDXDeviceInfo::s_DXDeviceInfo.Reset();
+	bHaveDirectXResources = false;
+	return true;
+}
+
+
+bool IsWindows8Plus()
+{
+	return FPlatformMisc::VerifyWindowsVersion(6, 2);
+}
+
+
+bool IsWindows7Plus()
+{
+	return FPlatformMisc::VerifyWindowsVersion(6, 1);
+}
+
+// ----------------------------------------------------------------------------------------------------------------------
+
+// some Windows versions don't have Media Foundation preinstalled. We configure MF DLLs as delay-loaded and load them manually here
+// checking the result and avoiding error message box if failed
+static bool LoadMediaFoundationDLLs()
+{
+	// Ensure that all required modules are preloaded so they are not loaded just-in-time, causing a hitch.
+	if (Electra::IsWindows8Plus())
+	{
+		return FPlatformProcess::GetDllHandle(TEXT("mf.dll"))
+			&& FPlatformProcess::GetDllHandle(TEXT("mfplat.dll"))
+			&& FPlatformProcess::GetDllHandle(TEXT("msmpeg2vdec.dll"))
+			&& FPlatformProcess::GetDllHandle(TEXT("MSAudDecMFT.dll"));
+	}
+	else // Windows 7
+	{
+		return FPlatformProcess::GetDllHandle(TEXT("mf.dll"))
+			&& FPlatformProcess::GetDllHandle(TEXT("mfplat.dll"))
+			&& FPlatformProcess::GetDllHandle(TEXT("msmpeg2vdec.dll"))
+			&& FPlatformProcess::GetDllHandle(TEXT("msmpeg2adec.dll"));
+	}
+}
+
+// ----------------------------------------------------------------------------------------------------------------------
+
+/**
+ * Early startup during module load process
+ */
+bool PlatformEarlyStartup()
+{
+	// Win7+ only
+	if (!IsWindows7Plus())
+	{
+		UE_LOG(LogElectraPlayer, Log, TEXT("Electra is incompatible with Windows prior to 7.0 version: %s"), *FPlatformMisc::GetOSVersion());
+		return false;
+	}
+
+	if (!LoadMediaFoundationDLLs())
+	{
+		UE_LOG(LogElectraPlayer, Log, TEXT("Electra can't load Media Foundation, %s"), *FPlatformMisc::GetOSVersion());
+		return false;
+	}
+
+	HRESULT Res = MFStartup(MF_VERSION);
+	checkf(SUCCEEDED(Res), TEXT("MFStartup failed: %d"), Res);
+
+	//SetEnvironmentVariable(L"http_proxy", L"10.29.4.18:8888");
+	//_putenv_s("http_proxy", "10.29.4.18:8888");
+
+	return true;
+}
+
+/**
+ * Any platform specific memory hooks
+ */
+bool PlatformMemorySetup()
+{
+	return true;
+}
+
+/**
+ * Shutdown of module
+ */
+bool PlatformShutdown()
+{
+	bool bRes = Electra::DestroyDirectXResources();
+	if (!bRes)
+	{
+		UE_LOG(LogElectraPlayer, Warning, TEXT("ElectraPlayer: Failed to unintialize DXGI Manager and Device"));
+	}
+	MFShutdown();
+	return bRes;
+}
+
+
+TSharedPtr<IVideoDecoderResourceDelegate, ESPMode::ThreadSafe> PlatformCreateVideoDecoderResourceDelegate(const TWeakPtr<IElectraPlayerAdapterDelegate, ESPMode::ThreadSafe>& AdapterDelegate)
+{
+	return nullptr;
+}
+
+} //namespace Electra
+
+// ----------------------------------------------------------------------------------------------------------------------
+// ----------------------------------------------------------------------------------------------------------------------
+
+/**
+ * Public call for pre-use initialization after module is loaded, allowing for user passed parameters
+ */
+bool FElectraPlayerPlatform::StartupPlatformResources(const Electra::FParamDict& Params)
+{
+	return Electra::CreateDirectXResources(Params);
+}
+
+void FElectraPlayer::PlatformNotifyOfOptionChange()
+{
+}
+
+#endif
+
+
diff --git a/Engine/Plugins/Media/ElectraPlayer/Source/ElectraPlayerRuntime/Private/Runtime/BitDataStream.h b/Engine/Plugins/Media/ElectraPlayer/Source/ElectraPlayerRuntime/Private/Runtime/BitDataStream.h
new file mode 100644
--- /dev/null
+++ b/Engine/Plugins/Media/ElectraPlayer/Source/ElectraPlayerRuntime/Private/Runtime/BitDataStream.h
@@ -0,0 +1,226 @@
+// Copyright Epic Games, Inc. All Rights Reserved.
+
+#pragma once
+
+#include "PlayerCore.h"
+
+#if defined(_MSC_VER)
+#include <intrin.h>
+#endif
+
+
+class FBitDataStream
+{
+public:
+	FBitDataStream()
+		: Data(nullptr), DataSize(0), BytePosition(0), BitPosition(0)
+	{
+	}
+
+	FBitDataStream(const void* InData, uint64 InDataSize, uint64 StartBytePosition = 0, uint32 StartBitPosition = 0)
+	{
+		SetData(InData, InDataSize, StartBytePosition, StartBitPosition);
+	}
+
+	~FBitDataStream()
+	{
+	}
+
+	void SetData(const void* InData, uint64 InDataSize, uint64 StartBytePosition = 0, uint32 StartBitPosition = 0)
+	{
+		check(InData && InDataSize);
+		Data = (const uint8*)InData;
+		DataSize = InDataSize;
+		BytePosition = StartBytePosition;
+		BitPosition = StartBitPosition;
+	}
+
+	bool IsByteAligned() const
+	{
+		return (BitPosition & 7) == 0;
+	}
+
+	const void* GetRemainingData() const
+	{
+		return Data + BytePosition;
+	}
+
+	uint64 GetRemainingByteLength() const
+	{
+		return DataSize - BytePosition;
+	}
+
+	uint64 GetRemainingBits() const
+	{
+		return GetRemainingByteLength() * 8 + (BitPosition ? 8 - BitPosition : 0);
+	}
+
+	void SkipBytes(uint64 nBytes)
+	{
+		BytePosition += nBytes;
+	}
+
+	void SkipBits(uint64 nBits)
+	{
+		SkipBytes(nBits >> 3);
+		GetBits(nBits & 7);
+	}
+
+	uint32 GetBits(uint64 nBits)
+	{
+		uint32 rv = PeekBits(nBits);
+		uint32 bp = BitPosition + nBits;
+		BytePosition += bp >> 3;
+		BitPosition = bp & 7;
+		return rv;
+	}
+
+	uint32 PeekBits(uint64 nBits)
+	{
+		checkf(nBits <= 32, TEXT("This function can return at most 32 bits!"));
+		if (nBits == 0)
+		{
+			return 0;
+		}
+		else
+		{
+			uint32 bytePos = BytePosition + (BitPosition >> 3);
+			uint32 bitPos = BitPosition & 7;
+
+			uint32 w;
+			uint64 rl = GetRemainingByteLength();
+			if (rl >= 4)
+			{
+#if defined(_MSC_VER)
+				w = _byteswap_ulong(*(const uint32*)(Data + bytePos));
+#else
+				w = __builtin_bswap32(*(const uint32*)(Data + bytePos));
+#endif
+			}
+			else
+			{
+				if (rl == 3)
+				{
+					w = ((uint32)Data[bytePos + 0] << 24) | ((uint32)Data[bytePos + 1] << 16) | ((uint32)Data[bytePos + 2] << 8);
+				}
+				else if (rl == 2)
+				{
+					w = ((uint32)Data[bytePos + 0] << 24) | ((uint32)Data[bytePos + 1] << 16);
+				}
+				else if (rl == 1)
+				{
+					w = ((uint32)Data[bytePos + 0] << 24);
+				}
+				else
+				{
+					w = 0;
+				}
+
+			}
+
+			if (bitPos == 0)
+			{
+				w = w >> (32 - nBits);
+				return w;
+			}
+			else
+			{
+				uint32 w0 = w;
+				uint8 w1;
+				if (rl > 4)
+				{
+					w1 = *(const uint8*)(Data + bytePos + 4);
+				}
+				else
+				{
+					w1 = 0;
+				}
+				w = (w0 << bitPos) | uint32(w1 >> (8 - bitPos));
+				w = w >> (32 - nBits);
+				return w;
+			}
+		}
+	}
+
+	uint64 GetBits64(uint64 nBits)
+	{
+		if (nBits <= 32)
+		{
+			return GetBits(nBits);
+		}
+		else
+		{
+			uint32 upper = GetBits(nBits - 32);
+			uint32 lower = GetBits(32);
+			return (uint64(upper) << 32ULL) | uint64(lower);
+		}
+	}
+
+private:
+	const uint8*	Data;
+	uint64			DataSize;
+	uint64			BytePosition;
+	uint32			BitPosition;
+};
+
+
+
+
+namespace Electra
+{
+
+	struct FByteReader
+	{
+		FByteReader(const void* InData, int64 DataSize)
+			: Data((const uint8*)InData)
+			, NumBytesToGo(DataSize)
+		{
+		}
+
+		bool ReadByte(uint8& To)
+		{
+			if (NumBytesToGo)
+			{
+				To = *Data++;
+				--NumBytesToGo;
+				return true;
+			}
+			return false;
+		}
+
+		bool ReadByte(uint16& To)
+		{
+			if (NumBytesToGo >= sizeof(uint16))
+			{
+				uint16 hi = *Data++;
+				uint16 lo = *Data++;
+				NumBytesToGo -= sizeof(uint16);
+				To = (hi << 8) | lo;
+				return true;
+			}
+			return false;
+		}
+		bool ReadBytes(uint8* To, int64 num)
+		{
+			if (NumBytesToGo >= num)
+			{
+				FMemory::Memcpy(To, Data, num);
+				Data += num;
+				NumBytesToGo -= num;
+				return true;
+			}
+			return false;
+		}
+		int64 BytesRemaining() const
+		{
+			return NumBytesToGo;
+		}
+
+	private:
+		const uint8* Data;
+		int64			NumBytesToGo;
+	};
+
+} // namespace Electra
+
+
diff --git a/Engine/Plugins/Media/ElectraPlayer/Source/ElectraPlayerRuntime/Private/Runtime/Core/MediaEventSignal.h b/Engine/Plugins/Media/ElectraPlayer/Source/ElectraPlayerRuntime/Private/Runtime/Core/MediaEventSignal.h
new file mode 100644
--- /dev/null
+++ b/Engine/Plugins/Media/ElectraPlayer/Source/ElectraPlayerRuntime/Private/Runtime/Core/MediaEventSignal.h
@@ -0,0 +1,77 @@
+// Copyright Epic Games, Inc. All Rights Reserved.
+
+#pragma once
+
+#include "Misc/Build.h"
+
+#include "HAL/Event.h"
+#include "HAL/PlatformProcess.h"
+
+#include "Core/MediaNoncopyable.h"
+
+
+class FMediaEvent : private TMediaNoncopyable<FMediaEvent>
+{
+public:
+	FMediaEvent()
+		: Event(FPlatformProcess::GetSynchEventFromPool(true))
+	{
+	}
+
+	~FMediaEvent()
+	{
+		Event->Trigger();
+		FPlatformProcess::ReturnSynchEventToPool(Event);
+	}
+
+	//! Waits for the event signal to get set. If already set, this function returns immediately.
+	void Wait()
+	{
+		Event->Wait();
+	}
+
+	//! Waits for the event signal to get set and clears it immediately. If already set, this function returns immediately.
+	void WaitAndReset()
+	{
+		Event->Wait();
+		Event->Reset();
+	}
+
+	//! Signals the event to be set. If already set, this does nothing.
+	void Signal()
+	{
+		Event->Trigger();
+	}
+
+	//! Clears the signal. Threads will now wait again until the signal is set once more.
+	void Reset()
+	{
+		Event->Reset();
+	}
+
+	//! Returns state of the signal.
+	bool IsSignaled() const
+	{
+		return Event->Wait(0);
+	}
+
+	//! Waits for the event signal to get set within the given time limit. If already set, this function returns immediately. Returns true when event was set, false when wait timed out.
+	bool WaitTimeout(int64 MicroSeconds)
+	{
+		return Event->Wait(FTimespan::FromMicroseconds(MicroSeconds));
+	}
+
+	//! Waits for the event signal to get set within the given time limit and clears it immediately. If already set, this function returns immediately. Returns true when event was set, false when wait timed out.
+	bool WaitTimeoutAndReset(int64 MicroSeconds)
+	{
+		check(MicroSeconds > 0);	// 0 is forbidden. either call Wait() without timeout, or check the signal using IsSignaled()
+		bool bGot = Event->Wait(FTimespan::FromMicroseconds(MicroSeconds));
+		if (bGot)
+			Event->Reset();
+		return bGot;
+	}
+
+private:
+	mutable FEvent* Event;
+
+};
diff --git a/Engine/Plugins/Media/ElectraPlayer/Source/ElectraPlayerRuntime/Private/Runtime/Core/MediaFastDelegate.h b/Engine/Plugins/Media/ElectraPlayer/Source/ElectraPlayerRuntime/Private/Runtime/Core/MediaFastDelegate.h
new file mode 100644
--- /dev/null
+++ b/Engine/Plugins/Media/ElectraPlayer/Source/ElectraPlayerRuntime/Private/Runtime/Core/MediaFastDelegate.h
@@ -0,0 +1,179 @@
+// Copyright Epic Games, Inc. All Rights Reserved.
+
+#pragma once
+
+#include <Core/MediaTypes.h>
+#include <Core/MediaMacros.h>
+#include <Core/MediaNoncopyable.h>
+#include <Core/MediaLock.h>
+
+namespace Electra
+{
+#if defined(_MSC_VER)
+#pragma warning(push)
+#pragma warning(disable : 4191)
+#endif
+
+#include "3rdParty/FastDelegate/FastDelegate.h"
+	using namespace fastdelegate;
+
+#if defined(_MSC_VER)
+#pragma warning( pop )
+#endif
+
+
+
+
+
+	/**
+	 * Internal variables for the safe delegate.
+	 * Those will be passed via TSharedPtr<> to the service functions that shall invoke the delegate.
+	**/
+	template <typename CB>
+	struct TMediaSafeDelegateVars : private TMediaNoncopyable<TMediaSafeDelegateVars<CB> >
+	{
+		CB							Callback;
+		FMediaLockCriticalSection	Lock;
+	};
+
+	/**
+	 * Template functions to invoke the safe delegate for functions returning <void>.
+	 * Used only by the caller.
+	**/
+	template<typename D>
+	void TMediaSafeDelegateCall(D delegate)
+	{
+		FMediaLockCriticalSection::ScopedLock lock(delegate->Lock); if (!delegate->Callback.empty()) delegate->Callback();
+	}
+
+	template<typename D, typename A1>
+	void TMediaSafeDelegateCall(D delegate, A1 arg1)
+	{
+		FMediaLockCriticalSection::ScopedLock lock(delegate->Lock); if (!delegate->Callback.empty()) delegate->Callback(arg1);
+	}
+
+	template<typename D, typename A1, typename A2>
+	void TMediaSafeDelegateCall(D delegate, A1 arg1, A2 arg2)
+	{
+		FMediaLockCriticalSection::ScopedLock lock(delegate->Lock); if (!delegate->Callback.empty()) delegate->Callback(arg1, arg2);
+	}
+
+	template<typename D, typename A1, typename A2, typename A3>
+	void TMediaSafeDelegateCall(D delegate, A1 arg1, A2 arg2, A3 arg3)
+	{
+		FMediaLockCriticalSection::ScopedLock lock(delegate->Lock); if (!delegate->Callback.empty()) delegate->Callback(arg1, arg2, arg3);
+	}
+
+	template<typename D, typename A1, typename A2, typename A3, typename A4>
+	void TMediaSafeDelegateCall(D delegate, A1 arg1, A2 arg2, A3 arg3, A4 arg4)
+	{
+		FMediaLockCriticalSection::ScopedLock lock(delegate->Lock); if (!delegate->Callback.empty()) delegate->Callback(arg1, arg2, arg3, arg4);
+	}
+
+	template<typename D, typename A1, typename A2, typename A3, typename A4, typename A5>
+	void TMediaSafeDelegateCall(D delegate, A1 arg1, A2 arg2, A3 arg3, A4 arg4, A5 arg5)
+	{
+		FMediaLockCriticalSection::ScopedLock lock(delegate->Lock); if (!delegate->Callback.empty()) delegate->Callback(arg1, arg2, arg3, arg4, arg5);
+	}
+
+	template<typename D, typename A1, typename A2, typename A3, typename A4, typename A5, typename A6>
+	void TMediaSafeDelegateCall(D delegate, A1 arg1, A2 arg2, A3 arg3, A4 arg4, A5 arg5, A6 arg6)
+	{
+		FMediaLockCriticalSection::ScopedLock lock(delegate->Lock); if (!delegate->Callback.empty()) delegate->Callback(arg1, arg2, arg3, arg4, arg5, arg6);
+	}
+
+	template<typename D, typename A1, typename A2, typename A3, typename A4, typename A5, typename A6, typename A7>
+	void TMediaSafeDelegateCall(D delegate, A1 arg1, A2 arg2, A3 arg3, A4 arg4, A5 arg5, A6 arg6, A7 arg7)
+	{
+		FMediaLockCriticalSection::ScopedLock lock(delegate->Lock); if (!delegate->Callback.empty()) delegate->Callback(arg1, arg2, arg3, arg4, arg5, arg6, arg7);
+	}
+
+	template<typename D, typename A1, typename A2, typename A3, typename A4, typename A5, typename A6, typename A7, typename A8>
+	void TMediaSafeDelegateCall(D delegate, A1 arg1, A2 arg2, A3 arg3, A4 arg4, A5 arg5, A6 arg6, A7 arg7, A8 arg8)
+	{
+		FMediaLockCriticalSection::ScopedLock lock(delegate->Lock); if (!delegate->Callback.empty()) delegate->Callback(arg1, arg2, arg3, arg4, arg5, arg6, arg7, arg8);
+	}
+
+	template<typename D, typename A1, typename A2, typename A3, typename A4, typename A5, typename A6, typename A7, typename A8, typename A9>
+	void TMediaSafeDelegateCall(D delegate, A1 arg1, A2 arg2, A3 arg3, A4 arg4, A5 arg5, A6 arg6, A7 arg7, A8 arg8, A9 arg9)
+	{
+		FMediaLockCriticalSection::ScopedLock lock(delegate->Lock); if (!delegate->Callback.empty()) delegate->Callback(arg1, arg2, arg3, arg4, arg5, arg6, arg7, arg8, arg9);
+	}
+
+
+
+
+
+	/**
+	 * Safe delegate.
+	 *
+	 * Instantiate as a member in your class and set your callback function via Set().
+	 * Pass the embedded Proxy object to any safe-delegate aware service function.
+	 *
+	 * If your class is destroyed before the callback is fired the following will happen:
+	 *  - if the callback function has no return value nothing will happen
+	 *  - if the callback function returns some value an assert will trigger because no meaningful
+	 *    return value can be provided to the caller.
+	**/
+	template <typename CB>
+	class TMediaSafeDelegate
+	{
+	public:
+		//! Type of proxy to be passed into service functions.
+		typedef Electra::TSharedPtrTS<TMediaSafeDelegateVars<CB> > Proxy;
+
+		//! Empty callback, if none should be passed to a service function.
+		static Proxy None(void)
+		{
+			Proxy none(new TMediaSafeDelegateVars<CB>); return(none);
+		}
+
+
+		//! Destructor. Clears the callback function in the proxy.
+		~TMediaSafeDelegate()
+		{
+			if (Vars)
+			{
+				FMediaLockCriticalSection::ScopedLock lock(Vars->Lock);
+				Vars->Callback.clear();
+			}
+		}
+
+		//! Assigns the callback function to be invoked.
+		void Set(CB callback)
+		{
+			if (Vars == nullptr)
+			{
+				Vars = MakeShareable(new TMediaSafeDelegateVars<CB>);
+				Vars->Callback = callback;
+			}
+			else
+			{
+				FMediaLockCriticalSection::ScopedLock lock(Vars->Lock);
+				Vars->Callback = callback;
+			}
+		}
+
+		//! Cast operator to return the embedded proxy to a service function.
+		operator Proxy()
+		{
+			checkf(Vars, TEXT("No callback function set!"));
+			return(Vars);
+		}
+
+		//! Function operator to return the embedded proxy to a service function.
+		Proxy operator() (void)
+		{
+			checkf(Vars, TEXT("No callback function set!"));
+			return(Vars);
+		}
+
+	private:
+		Proxy	Vars;			//!< Embedded proxy object.
+	};
+
+
+};
+
+
+
diff --git a/Engine/Plugins/Media/ElectraPlayer/Source/ElectraPlayerRuntime/Private/Runtime/Core/MediaQueue.h b/Engine/Plugins/Media/ElectraPlayer/Source/ElectraPlayerRuntime/Private/Runtime/Core/MediaQueue.h
new file mode 100644
--- /dev/null
+++ b/Engine/Plugins/Media/ElectraPlayer/Source/ElectraPlayerRuntime/Private/Runtime/Core/MediaQueue.h
@@ -0,0 +1,957 @@
+// Copyright Epic Games, Inc. All Rights Reserved.
+
+#pragma once
+
+#include "Core/MediaTypes.h"
+#include "Core/MediaMacros.h"
+
+#include "Core/MediaLock.h"
+
+#include "Core/MediaInterlocked.h"
+
+
+
+
+/**
+ * Default memory allocator.
+**/
+template <typename T>
+class TMediaQueueAllocator
+{
+public:
+	static T* Allocate(uint32 NumElements)
+	{
+		return(reinterpret_cast<T*>(FMemory::Malloc(NumElements * sizeof(T),  __alignof(T))));
+	}
+	static void Deallocate(T* Address)
+	{
+		if (Address)
+		{
+			FMemory::Free((void*)Address);
+		}
+	}
+};
+
+
+
+/**
+ * Object safe FIFO container class.
+ *
+ * Implements essential FIFO functions only to avoid working with references.
+ *
+ * Methods like
+ *  T& Push()
+ *  T& operator[](int)
+ *  T& Front()
+ *
+ * are dangerous when used in a multi-threaded environment since the elements cannot be protected by an internal mutex.
+ * If you are using an external mutex to ensure object safety you may add such functions, but please derive a new class for this!!!
+ *
+**/
+template <typename T, typename L = FMediaLockNone, typename M = TMediaQueueAllocator<T> >
+class TMediaQueue : public L
+{
+public:
+	using MyType = TMediaQueue<T, L, M>;
+	using ElementType = T;
+
+	TMediaQueue(SIZE_T InitialCapacity = 0)
+		: Elements(NULL)
+		, MaxNum(0)
+		, NumIn(0)
+		, IdxIn(0)
+		, IdxOut(0)
+	{
+		Resize(InitialCapacity);
+	}
+
+	~TMediaQueue()
+	{
+		Resize(0);
+	}
+
+	void Clear()
+	{
+		L::Lock();
+		for (uint32 i = 0; i < NumIn; ++i)
+		{
+			Elements[IdxOut].~T();
+			if (++IdxOut == MaxNum)
+				IdxOut -= MaxNum;
+		}
+		IdxIn = IdxOut = NumIn = 0;
+		L::Unlock();
+	}
+
+	void Resize(SIZE_T newCapacity)
+	{
+		L::Lock();
+		Clear();
+		if (newCapacity != MaxNum)
+		{
+			M::Deallocate(Elements);
+			if (newCapacity)
+				Elements = M::Allocate(newCapacity);
+			else
+				Elements = NULL;
+			MaxNum = newCapacity;
+		}
+		L::Unlock();
+	}
+
+	void Swap(MyType& rhs)
+	{
+		L::Lock();
+		// Remember our values.
+		T* pElements = Elements;
+		SIZE_T  nCapacity = MaxNum;
+		SIZE_T  nNumIn = NumIn;
+		SIZE_T  nIdxIn = IdxIn;
+		SIZE_T  nIdxOut = IdxOut;
+		// Get the values from the other FIFO
+		Elements = rhs.Elements;
+		MaxNum = rhs.MaxNum;
+		NumIn = rhs.NumIn;
+		IdxIn = rhs.IdxIn;
+		IdxOut = rhs.IdxOut;
+		// Set our old values in the other FIFO.
+		rhs.Elements = pElements;
+		rhs.MaxNum = nCapacity;
+		rhs.NumIn = nNumIn;
+		rhs.IdxIn = nIdxIn;
+		rhs.IdxOut = nIdxOut;
+		L::Unlock();
+	}
+
+	SIZE_T Capacity() const
+	{
+		return MaxNum;
+	}
+
+	SIZE_T Num() const
+	{
+		return(NumIn);
+	}
+
+	bool IsEmpty() const
+	{
+		return Num() == 0;
+	}
+
+	bool IsFull() const
+	{
+		return Num() == Capacity();
+	}
+
+	void Push(const T& Element)
+	{
+		L::Lock();
+		check(!IsFull());
+		SIZE_T pos = IdxIn;
+		if (++IdxIn == MaxNum)
+			IdxIn = 0;
+		++NumIn;
+		// Copying the element can NOT take place outside the lock in case multiple threads push at the same time and a Resize() will happen.
+		new ((void*)(Elements + pos)) T(Element);
+		L::Unlock();
+	}
+
+	void PushFront(const T& Element)
+	{
+		L::Lock();
+		check(!IsFull());
+		if (IdxOut == 0)
+			IdxOut = MaxNum - 1;
+		else
+			--IdxOut;
+		SIZE_T pos = IdxOut;
+		++NumIn;
+		// Copying the element can NOT take place outside the lock in case multiple threads push at the same time and a Resize() will happen.
+		new ((void*)(Elements + pos)) T(Element);
+		L::Unlock();
+	}
+
+	T Pop()
+	{
+		L::Lock();
+		check(!IsEmpty());
+		T r(Elements[IdxOut]);
+		Elements[IdxOut].~T();
+		if (++IdxOut == MaxNum)
+			IdxOut = 0;
+		--NumIn;
+		L::Unlock();
+		return r;
+	}
+
+	T Front() const
+	{
+		L::Lock();
+		check(!IsEmpty());
+		T r(Elements[IdxOut]);
+		L::Unlock();
+		return r;
+	}
+
+protected:
+	T* 		Elements;
+	SIZE_T	MaxNum;
+	SIZE_T	NumIn;
+	SIZE_T	IdxIn;
+	SIZE_T	IdxOut;
+
+	// Not copyable (for now) because of the possible mutex class.
+	TMediaQueue(const MyType& rhs);
+	MyType& operator = (const MyType& rhs);
+};
+
+
+
+
+
+/**
+ * A mutex-less version of TMediaQueue to enable working with references.
+ *
+ * Since this is not thread-safe you must ensure access with an external mutex.
+**/
+template <typename T, typename M = TMediaQueueAllocator<T> >
+class TMediaQueueNoLock : public TMediaQueue<T, FMediaLockNone, M>
+{
+public:
+	using MyType = TMediaQueueNoLock<T, M>;
+	using ElementType = T;
+	typedef TMediaQueue<T, FMediaLockNone, M>	B;
+
+	TMediaQueueNoLock() : B() {}
+	~TMediaQueueNoLock() {}
+
+	T& PushRef()
+	{
+		check(!B::IsFull());
+		SIZE_T pos = B::IdxIn;
+		if (++B::IdxIn == B::Capacity())
+			B::IdxIn = 0;
+		++B::NumIn;
+		new ((void*)(B::Elements + pos)) T();
+		return B::Elements[pos];
+	}
+
+	const T& FrontRef() const
+	{
+		check(!B::IsEmpty());
+		return B::Elements[B::IdxOut];
+	}
+
+	T& FrontRef()
+	{
+		check(!B::IsEmpty());
+		return B::Elements[B::IdxOut];
+	}
+
+	const T& BackRef() const
+	{
+		check(!B::IsEmpty());
+		return B::Elements[(B::IdxIn + B::Capacity() - 1) % B::Capacity()];
+	}
+
+	T& BackRef()
+	{
+		check(!B::IsEmpty());
+		return B::Elements[(B::IdxIn + B::Capacity() - 1) % B::Capacity()];
+	}
+
+	void Erase(SIZE_T i)
+	{
+		check(!B::IsEmpty());
+		check(i < B::NumIn);
+		--B::NumIn;
+		SIZE_T i0 = (i + B::IdxOut) % B::Capacity();
+		SIZE_T i1 = (i + B::IdxOut + 1) % B::Capacity();
+		for (; i < B::NumIn; ++i)
+		{
+			B::Elements[i0] = B::Elements[i1];
+			i0 = i1++;
+			if (i1 >= B::Capacity())
+				i1 -= B::Capacity();
+		}
+		B::Elements[i0].~T();
+		if (B::IdxIn == 0)
+			B::IdxIn = B::Capacity() - 1;
+		else
+			--B::IdxIn;
+	}
+
+	const T& operator[](SIZE_T i) const
+	{
+		check(!B::IsEmpty());
+		check(i < B::NumIn);
+		i += B::IdxOut;
+		if (i >= B::Capacity())
+			i -= B::Capacity();
+		return B::Elements[i];
+	}
+
+	T& operator[](SIZE_T i)
+	{
+		check(!B::IsEmpty());
+		check(i < B::NumIn);
+		i += B::IdxOut;
+		if (i >= B::Capacity())
+			i -= B::Capacity();
+		return B::Elements[i];
+	}
+
+	void AppendFirstElements(const MyType& FromOther, SIZE_T MaxElements = ~SIZE_T(0))
+	{
+		// NOTE: Your code needs to make sure the source queue is protected by some external lock so does not get
+		//       modified while we're appending its elements.
+		SIZE_T NumAvail = B::Capacity() - B::Num();
+		SIZE_T NumToAppend = FromOther.Num();
+		if (NumToAppend > NumAvail)
+			NumToAppend = NumAvail;
+		if (NumToAppend > MaxElements)
+			NumToAppend = MaxElements;
+		for (SIZE_T i = 0; i < NumToAppend; ++i)
+			B::Push(FromOther[i]);
+	}
+
+	void AppendLastElements(const MyType& FromOther, SIZE_T MaxElements = ~SIZE_T(0))
+	{
+		// NOTE: Your code needs to make sure the source queue is protected by some external lock so does not get
+		//       modified while we're appending its elements.
+		SIZE_T NumAvail = B::Capacity() - B::Num();
+		SIZE_T NumToAppend = FromOther.Num();
+		if (NumToAppend > NumAvail)
+			NumToAppend = NumAvail;
+		if (NumToAppend > MaxElements)
+			NumToAppend = MaxElements;
+		for (SIZE_T i = 0, j = FromOther.Num() - NumToAppend; i < NumToAppend; ++i, ++j)
+			B::Push(FromOther[j]);
+	}
+
+};
+
+
+
+
+
+
+
+
+
+
+
+
+
+
+
+
+/**
+ * Object safe FIFO container class.
+ *
+ * Implements essential FIFO functions only to avoid working with references.
+ *
+ * Methods like
+ *  T& Push()
+ *  T& operator[](int)
+ *  T& Front()
+ *
+ * are dangerous when used in a multi-threaded environment since the elements cannot be protected by an internal mutex.
+ * If you are using an external mutex to ensure object safety you may add such functions, but please derive a new class for this!!!
+ * (or use the one below)
+ *
+**/
+template <typename T, uint32 CAPACITY, typename L = FMediaLockCriticalSection >
+class TMediaQueueFixedStatic : public L
+{
+public:
+	using MyType = TMediaQueueFixedStatic<T, CAPACITY, L>;
+	using ElementType = T;
+
+	TMediaQueueFixedStatic()
+		: Elements((T*)maElements)
+		, NumIn(0)
+		, IdxIn(0)
+		, IdxOut(0)
+	{
+	}
+
+	~TMediaQueueFixedStatic()
+	{
+		Clear();
+	}
+
+	void Clear()
+	{
+		L::Lock();
+		for (uint32 i = 0; i < NumIn; ++i)
+		{
+			Elements[IdxOut].~T();
+			if (++IdxOut == CAPACITY)
+				IdxOut -= CAPACITY;
+		}
+		IdxIn = IdxOut = NumIn = 0;
+		L::Unlock();
+	}
+
+	SIZE_T Capacity() const
+	{
+		return(CAPACITY);
+	}
+
+	SIZE_T Num() const
+	{
+		return(NumIn);
+	}
+
+	bool IsEmpty() const
+	{
+		return(Num() == 0);
+	}
+
+	bool IsFull() const
+	{
+		return(Num() == Capacity());
+	}
+
+	void Push(const T& element)
+	{
+		L::Lock();
+		check(!IsFull());
+		SIZE_T pos = IdxIn;
+		if (++IdxIn == CAPACITY)
+			IdxIn = 0;
+		++NumIn;
+		// Copying the element can NOT take place outside the lock in case multiple threads push at the same time and a Resize() will happen.
+		new ((void*)(Elements + pos)) T(element);
+		L::Unlock();
+	}
+
+	void PushFront(const T& element)
+	{
+		L::Lock();
+		check(!IsFull());
+		if (IdxOut == 0)
+			IdxOut = CAPACITY - 1;
+		else
+			--IdxOut;
+		SIZE_T pos = IdxOut;
+		++NumIn;
+		// Copying the element can NOT take place outside the lock in case multiple threads push at the same time and a Resize() will happen.
+		new ((void*)(Elements + pos)) T(element);
+		L::Unlock();
+	}
+
+	T Pop()
+	{
+		L::Lock();
+		check(!IsEmpty());
+		T r(Elements[IdxOut]);
+		Elements[IdxOut].~T();
+		if (++IdxOut == CAPACITY)
+			IdxOut = 0;
+		--NumIn;
+		L::Unlock();
+		return(r);
+	}
+
+	T Front() const
+	{
+		L::Lock();
+		check(!IsEmpty());
+		T r(Elements[IdxOut]);
+		L::Unlock();
+		return(r);
+	}
+
+protected:
+
+	MS_ALIGN(16)
+	uint8	maElements[CAPACITY * sizeof(T)]
+	GCC_ALIGN(16);
+	T*		Elements;
+	SIZE_T	NumIn;
+	SIZE_T	IdxIn;
+	SIZE_T	IdxOut;
+
+	// Not copyable (for now) because of the possible mutex class.
+	TMediaQueueFixedStatic(const MyType& rhs);
+	MyType& operator = (const MyType& rhs);
+};
+
+
+
+template <typename T, uint32 CAPACITY>
+class TMediaQueueFixedStaticNoLock : public TMediaQueueFixedStatic<T, CAPACITY, FMediaLockNone>
+{
+public:
+	using MyType = TMediaQueueFixedStaticNoLock<T, CAPACITY>;
+	using ElementType = T;
+	typedef TMediaQueueFixedStatic<T, CAPACITY, FMediaLockNone>	B;
+
+	TMediaQueueFixedStaticNoLock() : B() {}
+	~TMediaQueueFixedStaticNoLock() {}
+
+	T& PushRef()
+	{
+		check(!B::IsFull());
+		SIZE_T pos = B::IdxIn;
+		if (++B::IdxIn == CAPACITY)
+			B::IdxIn = 0;
+		++B::NumIn;
+		new ((void*)(B::Elements + pos)) T();
+		return(B::Elements[pos]);
+	}
+
+	const T& FrontRef() const
+	{
+		check(!B::IsEmpty());
+		return(B::Elements[B::IdxOut]);
+	}
+
+	T& FrontRef()
+	{
+		check(!B::IsEmpty());
+		return(B::Elements[B::IdxOut]);
+	}
+
+	const T& BackRef() const
+	{
+		check(!B::IsEmpty());
+		return(B::Elements[(B::IdxIn + B::Capacity() - 1) % B::Capacity()]);
+	}
+
+	T& BackRef()
+	{
+		check(!B::IsEmpty());
+		return(B::Elements[(B::IdxIn + B::Capacity() - 1) % B::Capacity()]);
+	}
+
+	void Erase(SIZE_T i)
+	{
+		check(!B::IsEmpty());
+		check(i < B::NumIn);
+		--B::NumIn;
+		SIZE_T i0 = (i + B::IdxOut) % CAPACITY;
+		SIZE_T i1 = (i + B::IdxOut + 1) % CAPACITY;
+		for (; i < B::NumIn; ++i)
+		{
+			B::Elements[i0] = B::Elements[i1];
+			i0 = i1++;
+			if (i1 >= CAPACITY)
+				i1 -= CAPACITY;
+		}
+		B::Elements[i0].~T();
+		if (B::IdxIn == 0)
+			B::IdxIn = CAPACITY - 1;
+		else
+			--B::IdxIn;
+	}
+
+	const T& operator[](SIZE_T i) const
+	{
+		check(!B::IsEmpty());
+		check(i < B::NumIn);
+		i += B::IdxOut;
+		if (i >= CAPACITY)
+			i -= CAPACITY;
+		return(B::Elements[i]);
+	}
+
+	T& operator[](SIZE_T i)
+	{
+		check(!B::IsEmpty());
+		check(i < B::NumIn);
+		i += B::IdxOut;
+		if (i >= CAPACITY)
+			i -= CAPACITY;
+		return(B::Elements[i]);
+	}
+};
+
+
+
+
+
+
+
+
+
+
+
+
+
+
+
+
+/**
+ * Dynamically sized FIFO.
+ *
+ * Will grow when capacity is reached.
+ * Therefore a lock is REQUIRED! Never use this with a null-mutex unless you know what you're doing.
+**/
+template <typename T, typename L = FMediaLockCriticalSection, typename M = TMediaQueueAllocator<T> >
+class TMediaQueueDynamic : public L, private TMediaNoncopyable<TMediaQueueDynamic<T> >
+{
+public:
+	using MyType = TMediaQueueDynamic<T, L, M>;
+	using ElementType = T;
+
+	TMediaQueueDynamic(SIZE_T initialCapacity = 0, SIZE_T increment = 32)
+		: Elements(NULL)
+		, MaxNum(0)
+		, NumIn(0)
+		, IdxIn(0)
+		, IdxOut(0)
+		, mIncrement(increment)
+	{
+		Resize(initialCapacity);
+	}
+
+	~TMediaQueueDynamic()
+	{
+		L::Lock();
+		Resize(0);
+		L::Unlock();
+	}
+
+	void SetIncrement(SIZE_T newIncrement)
+	{
+		L::Lock();
+		mIncrement = newIncrement;
+		L::Unlock();
+	}
+
+	void Clear()
+	{
+		L::Lock();
+		for (uint32 i = 0; i < NumIn; ++i)
+		{
+			Elements[IdxOut].~T();
+			if (++IdxOut == MaxNum)
+				IdxOut -= MaxNum;
+		}
+		IdxIn = IdxOut = NumIn = 0;
+		L::Unlock();
+	}
+
+	void Resize(SIZE_T newCapacity)
+	{
+		L::Lock();
+		if (newCapacity == 0)
+		{
+			Clear();
+			M::Deallocate(Elements);
+			Elements = NULL;
+			MaxNum = 0;
+		}
+		else
+		{
+			// Create a new array and copy all active elements over.
+			checkf(newCapacity >= Num(), TEXT("Cannot shrink the dynamic FIFO to %u with %u elements currently managed!"), (uint32)newCapacity, (uint32)Num());
+			T* pNewElements = M::Allocate(newCapacity);
+			for (SIZE_T i = 0, j = IdxOut; i < NumIn; ++i)
+			{
+				new ((void*)(pNewElements + i)) T(Elements[j]);
+				Elements[j].~T();
+				if (++j == MaxNum)
+					j = 0;
+			}
+			M::Deallocate(Elements);
+			Elements = pNewElements;
+			IdxIn = NumIn;
+			IdxOut = 0;
+			MaxNum = newCapacity;
+		}
+		L::Unlock();
+	}
+
+	SIZE_T Increment() const
+	{
+		return(mIncrement);
+	}
+
+	SIZE_T Capacity() const
+	{
+		return(MaxNum);
+	}
+
+	SIZE_T Num() const
+	{
+		return(NumIn);
+	}
+
+	bool IsEmpty() const
+	{
+		return(Num() == 0);
+	}
+
+	bool IsFull() const
+	{
+		return(Num() == Capacity());
+	}
+
+	void Push(const T& element)
+	{
+		L::Lock();
+		if (IsFull())
+		{
+			Resize(Capacity() + Increment());
+		}
+		SIZE_T pos = IdxIn;
+		if (++IdxIn == MaxNum)
+			IdxIn = 0;
+		++NumIn;
+		// Copying the element can NOT take place outside the lock in case multiple threads push at the same time and a Resize() will happen.
+		new ((void*)(Elements + pos)) T(element);
+		L::Unlock();
+	}
+
+	void PushFront(const T& element)
+	{
+		L::Lock();
+		if (IsFull())
+		{
+			Resize(Capacity() + Increment());
+		}
+		if (IdxOut == 0)
+			IdxOut = MaxNum - 1;
+		else
+			--IdxOut;
+		SIZE_T pos = IdxOut;
+		++NumIn;
+		// Copying the element can NOT take place outside the lock in case multiple threads push at the same time and a Resize() will happen.
+		new ((void*)(Elements + pos)) T(element);
+		L::Unlock();
+	}
+
+	T Pop()
+	{
+		L::Lock();
+		check(!IsEmpty());
+		T r(Elements[IdxOut]);
+		Elements[IdxOut].~T();
+		if (++IdxOut == MaxNum)
+			IdxOut = 0;
+		--NumIn;
+		L::Unlock();
+		return(r);
+	}
+
+	T Front() const
+	{
+		L::Lock();
+		check(!IsEmpty());
+		T r(Elements[IdxOut]);
+		L::Unlock();
+		return(r);
+	}
+
+protected:
+	T* Elements;
+	SIZE_T	MaxNum;
+	SIZE_T	NumIn;
+	SIZE_T	IdxIn;
+	SIZE_T	IdxOut;
+	SIZE_T	mIncrement;
+};
+
+
+
+
+
+
+/**
+ *
+**/
+template <typename T, typename M = TMediaQueueAllocator<T> >
+class TMediaQueueDynamicNoLock : public TMediaQueueDynamic<T, FMediaLockNone, M>
+{
+public:
+	using MyType = TMediaQueueDynamicNoLock<T, M>;
+	using ElementType = T;
+	typedef TMediaQueueDynamic<T, FMediaLockNone, M>	B;
+
+	TMediaQueueDynamicNoLock() : B() {}
+	~TMediaQueueDynamicNoLock() {}
+	/*
+		T & PushRef(void)
+		{
+			check(!B::IsFull());
+			SIZE_T pos = B::IdxIn;
+			if (++B::IdxIn == B::MaxNum)
+				B::IdxIn = 0;
+			++B::NumIn;
+			new ((void *)(B::Elements + pos)) T();
+			return(B::Elements[pos]);
+		}
+	*/
+	const T& FrontRef() const
+	{
+		check(!B::IsEmpty());
+		return(B::Elements[B::IdxOut]);
+	}
+
+	T& FrontRef()
+	{
+		check(!B::IsEmpty());
+		return(B::Elements[B::IdxOut]);
+	}
+
+	const T& BackRef() const
+	{
+		check(!B::IsEmpty());
+		return(B::Elements[(B::IdxIn + B::Capacity() - 1) % B::Capacity()]);
+	}
+
+	T& BackRef()
+	{
+		check(!B::IsEmpty());
+		return(B::Elements[(B::IdxIn + B::Capacity() - 1) % B::Capacity()]);
+	}
+	/*
+		void Erase(SIZE_T i)
+		{
+			check(!B::IsEmpty());
+			check(i < B::NumIn);
+			--B::NumIn;
+			SIZE_T i0 = (i + B::IdxOut    ) % B::MaxNum;
+			SIZE_T i1 = (i + B::IdxOut + 1) % B::MaxNum;
+			for(; i<B::NumIn; ++i)
+			{
+				B::Elements[i0] = B::Elements[i1];
+				i0 = i1++;
+				if (i1 >= B::MaxNum)
+					i1 -= B::MaxNum;
+			}
+			B::Elements[i0].~T();
+			if (B::IdxIn == 0)
+				B::IdxIn = B::MaxNum - 1;
+			else
+				--B::IdxIn;
+		}
+	*/
+	const T& operator[](SIZE_T i) const
+	{
+		check(!B::IsEmpty());
+		check(i < B::NumIn);
+		i += B::IdxOut;
+		if (i >= B::MaxNum)
+			i -= B::MaxNum;
+		return(B::Elements[i]);
+	}
+
+	T& operator[](SIZE_T i)
+	{
+		check(!B::IsEmpty());
+		check(i < B::NumIn);
+		i += B::IdxOut;
+		if (i >= B::MaxNum)
+			i -= B::MaxNum;
+		return(B::Elements[i]);
+	}
+};
+
+
+
+
+
+
+
+
+
+
+template <typename T, uint32 CAPACITY>
+class TMediaQueueFixedStaticInterlocked
+{
+public:
+	using MyType = TMediaQueueFixedStaticInterlocked<T, CAPACITY>;
+	using ElementType = T;
+
+	TMediaQueueFixedStaticInterlocked()
+		: Elements((T*)maElements)
+		, NumIn(0)
+		, IdxIn(0)
+		, IdxOut(0)
+	{
+	}
+
+	~TMediaQueueFixedStaticInterlocked()
+	{
+		Clear();
+	}
+
+	void Clear()
+	{
+		for (uint32 i = 0; i < NumIn; ++i)
+		{
+			Elements[IdxOut].~T();
+			if (++IdxOut == CAPACITY)
+				IdxOut -= CAPACITY;
+		}
+		IdxIn = IdxOut = NumIn = 0;
+	}
+
+	uint32 Capacity() const
+	{
+		return(CAPACITY);
+	}
+
+	uint32 Num() const
+	{
+		return(FMediaInterlockedRead(NumIn));
+	}
+
+	bool IsEmpty() const
+	{
+		return(Num() == 0);
+	}
+
+	bool IsFull() const
+	{
+		return(Num() == Capacity());
+	}
+
+	void Push(const T& element)
+	{
+		check(!IsFull());
+		uint32 pos = FMediaInterlockedIncrement(IdxIn) % CAPACITY;
+		new ((void*)(Elements + pos)) T(element);
+		FMediaInterlockedIncrement(NumIn);
+	}
+
+	T Pop()
+	{
+		check(!IsEmpty());
+		uint32 pos = FMediaInterlockedIncrement(IdxOut) % CAPACITY;
+		T r(Elements[pos]);
+		Elements[pos].~T();
+		FMediaInterlockedDecrement(NumIn);
+		return(r);
+	}
+
+	const T& FrontRef() const
+	{
+		check(!IsEmpty());
+		return(Elements[IdxOut % CAPACITY]);
+	}
+
+	T& FrontRef()
+	{
+		check(!IsEmpty());
+		return(Elements[IdxOut % CAPACITY]);
+	}
+
+protected:
+	MS_ALIGN(16)
+	uint8	maElements[CAPACITY * sizeof(T)]
+	GCC_ALIGN(16);
+	T* Elements;
+	mutable volatile uint32	NumIn;
+	volatile uint32			IdxIn;
+	volatile uint32			IdxOut;
+
+};
+
+
+
diff --git a/Engine/Plugins/Media/ElectraPlayer/Source/ElectraPlayerRuntime/Private/Runtime/Core/MediaThreads.h b/Engine/Plugins/Media/ElectraPlayer/Source/ElectraPlayerRuntime/Private/Runtime/Core/MediaThreads.h
new file mode 100644
--- /dev/null
+++ b/Engine/Plugins/Media/ElectraPlayer/Source/ElectraPlayerRuntime/Private/Runtime/Core/MediaThreads.h
@@ -0,0 +1,178 @@
+// Copyright Epic Games, Inc. All Rights Reserved.
+
+#pragma once
+
+#include "PlayerCore.h"
+
+#include "Core/MediaTypes.h"
+#include "Core/MediaMacros.h"
+#include "Core/MediaNoncopyable.h"
+#include "Core/MediaEventSignal.h"
+
+#include "HAL/PlatformAffinity.h"
+#include "HAL/Runnable.h"
+#include "HAL/PlatformProcess.h"
+#include "HAL/CriticalSection.h"
+
+#include "Containers/UnrealString.h"
+#include "Containers/StringConv.h"
+#include "Templates/SharedPointer.h"
+
+/**
+ *
+ *
+*/
+class FMediaRunnable : private TMediaNoncopyable<FMediaRunnable>, public FRunnable
+{
+public:
+
+	//! Common thread configuration parameters.
+	struct Param
+	{
+		Param()
+		{
+			// Set some standard values as defaults.
+			Priority = TPri_Normal;
+			StackSize = 65536;
+			CoreAffinity = -1;
+		}
+		EThreadPriority		Priority;
+		SIZE_T				StackSize;
+		int32				CoreAffinity;
+	};
+
+	static FMediaRunnable* Create(int32 CoreAffinityMask, EThreadPriority Priority, SIZE_T StackSize, const FString& InThreadName);
+	static void Destroy(FMediaRunnable* Thread);
+
+	typedef Electra::FastDelegate0<> StartDelegate0;
+
+	void Start(StartDelegate0 Entry, bool bWaitRunning = false);
+
+	void SetDoneSignal(FMediaEvent* DoneSignal);
+
+	void SetName(const FString& InThreadName);
+
+	EThreadPriority ChangePriority(EThreadPriority NewPriority);
+
+	EThreadPriority PriorityGet() const
+	{
+		return ThreadPriority;
+	}
+
+	//! Returns the size of the stack.
+	SIZE_T StackSizeGet() const
+	{
+		return StackSize;
+	}
+
+	//! Returns the default stack size passed to Startup()
+	static SIZE_T StackSizeGetDefault()
+	{
+		return 0;
+	}
+
+	static void SleepSeconds(uint32 Seconds)
+	{
+		FPlatformProcess::SleepNoStats((float)Seconds);
+	}
+
+	static void SleepMilliseconds(uint32 Milliseconds)
+	{
+		FPlatformProcess::SleepNoStats(Milliseconds / 1000.0f);
+	}
+
+	static void SleepMicroseconds(uint32 Microseconds)
+	{
+		FPlatformProcess::SleepNoStats(Microseconds / 1000000.0f);
+	}
+
+private:
+
+	FRunnableThread*		MediaThreadRunnable;
+	EThreadPriority			ThreadPriority;
+	uint64					InitialCoreAffinity;
+	uint32					StackSize;
+	FString					ThreadName;
+	FCriticalSection		StateAccessMutex;
+	FEvent*					SignalRunning;
+	StartDelegate0			EntryFunction0;
+	FMediaEvent*			DoneSignal;
+	bool					bIsStarted;
+
+	FMediaRunnable();
+	~FMediaRunnable();
+
+	void StartInternal();
+	uint32 Run() override;
+	void Exit() override;
+};
+
+
+/**
+ * A thread base class to either inherit from or use as a variable.
+ *
+ * Thread parameters are given to constructor but can be changed before starting the thread
+ * with the ThreadSet..() functions.
+ * This allows the constructor to use system default values so an instance of this class
+ * can be used as a member variable of some class.
+ *
+ * To start the thread on some function of type void(*)(void) call ThreadStart() with
+ * an appropriate delegate.
+ *
+ * IMPORTANT: If used as a member variable you have to ensure for your class to stay alive
+ *            i.e. not be destroyed before the thread function has finished. Otherwise your
+ *            function will most likely crash.
+ *            Either wait for thread completion by calling ThreadWaitDone() *OR*
+ *            have the destructor wait for the thread itself by calling ThreadWaitDoneOnDelete(true)
+ *            at some point prior to destruction - preferably before starting the thread.
+**/
+class FMediaThread : private TMediaNoncopyable<FMediaThread>
+{
+public:
+	virtual ~FMediaThread();
+
+	//! Default constructor. Uses system defaults if no values given.
+	FMediaThread(const char* AnsiName = nullptr);
+
+	//! Set a thread priority other than the one given to the constructor before starting the thread.
+	void ThreadSetPriority(EThreadPriority Priority);
+
+	//! Set a core affinity before starting the thread. Defaults to -1 for no affinity (run on any core).
+	void ThreadSetCoreAffinity(int32 CoreAffinity);
+
+	//! Set a thread stack size other than the one given to the constructor before starting the thread.
+	void ThreadSetStackSize(SIZE_T StackSize);
+
+	//! Set a thread name other than the one given to the constructor before starting the thread.
+	void ThreadSetName(const char* InAnsiThreadName);
+
+	//! Sets a new thread name once the thread is running.
+	void ThreadRename(const char* InAnsiThreadName);
+
+	//! Sets whether or not the destructor needs to wait for the thread to have finished. Defaults to false. Useful when using this as a member variable.
+	void ThreadWaitDoneOnDelete(bool bWait);
+
+	//! Waits for the thread to have finished.
+	void ThreadWaitDone();
+
+	//! Starts the thread at the given function void(*)(void)
+	void ThreadStart(FMediaRunnable::StartDelegate0 EntryFunction);
+
+	//! Resets the thread to be started again. Must have waited for thread termination using ThreadWaitDone() first!
+	void ThreadReset();
+
+private:
+	FMediaEvent						SigDone;
+	FMediaRunnable*					MediaRunnable;
+	EThreadPriority					Priority;
+	int32							CoreAffinity;
+	SIZE_T							StackSize;
+	FString 						ThreadName;
+	bool							bIsStarted;
+	bool							bWaitDoneOnDelete;
+
+};
+
+
+
+
