ElectraPlayer: RELEASE. Moved from NFL to public. Note that this is done via a Delete and "Add New" operation as discussed.

diff --git a/Engine/Plugins/Media/ElectraPlayer/Source/ElectraPlayerRuntime/Private/Runtime/Decoder/Android/AudioDecoderAAC_Android.cpp b/Engine/Plugins/Media/ElectraPlayer/Source/ElectraPlayerRuntime/Private/Runtime/Decoder/Android/AudioDecoderAAC_Android.cpp
new file mode 100644
--- /dev/null
+++ b/Engine/Plugins/Media/ElectraPlayer/Source/ElectraPlayerRuntime/Private/Runtime/Decoder/Android/AudioDecoderAAC_Android.cpp
@@ -0,0 +1,921 @@
+// Copyright Epic Games, Inc. All Rights Reserved.
+
+#include "PlayerCore.h"
+
+#include "StreamAccessUnitBuffer.h"
+#include "Decoder/AudioDecoderAAC.h"
+#include "Renderer/RendererBase.h"
+#include "Player/PlayerSessionServices.h"
+#include "Utilities/Utilities.h"
+#include "Utilities/UtilsMPEGAudio.h"
+#include "Utilities/StringHelpers.h"
+#include "DecoderErrors_Android.h"
+#include "HAL/LowLevelMemTracker.h"
+
+#include "AudioDecoderAAC_JavaWrapper_Android.h"
+
+namespace Electra
+{
+
+
+/**
+ * AAC audio decoder class implementation.
+**/
+class FAudioDecoderAAC : public IAudioDecoderAAC, public FMediaThread
+{
+public:
+	static bool Startup(const IAudioDecoderAAC::FSystemConfiguration& InConfig);
+	static void Shutdown();
+
+	FAudioDecoderAAC();
+	virtual ~FAudioDecoderAAC();
+
+	virtual void SetPlayerSessionServices(IPlayerSessionServices* SessionServices) override;
+
+	virtual void Open(const FInstanceConfiguration& InConfig) override;
+	virtual void Close() override;
+
+	virtual void SetRenderer(TSharedPtr<IMediaRenderer, ESPMode::ThreadSafe> InRenderer) override;
+
+	virtual void SetAUInputBufferListener(IAccessUnitBufferListener* Listener) override;
+
+	virtual void SetReadyBufferListener(IDecoderOutputBufferListener* Listener) override;
+
+	virtual IAccessUnitBufferInterface::EAUpushResult AUdataPushAU(FAccessUnit* InAccessUnit) override;
+	virtual void AUdataPushEOD() override;
+	virtual void AUdataFlushEverything() override;
+
+private:
+	void StartThread();
+	void StopThread();
+	void WorkerThread();
+
+	bool InternalDecoderCreate();
+	void InternalDecoderDestroy();
+
+	bool CreateDecodedSamplePool();
+	void DestroyDecodedSamplePool();
+
+	void ReturnUnusedOutputBuffer();
+
+	void NotifyReadyBufferListener(bool bHaveOutput);
+
+	bool Decode(FAccessUnit* InAccessUnit);
+	bool FlushDecoder();
+	bool IsDifferentFormat(const FAccessUnit* Data);
+
+	void PostError(int32 ApiReturnValue, const FString& Message, uint16 Code, UEMediaError Error = UEMEDIA_ERROR_OK);
+	void LogMessage(IInfoLog::ELevel Level, const FString& Message);
+
+private:
+	FInstanceConfiguration													Config;
+
+	FAccessUnitBuffer														AccessUnits;
+
+	FMediaEvent																TerminateThreadSignal;
+	FMediaEvent																FlushDecoderSignal;
+	FMediaEvent																DecoderFlushedSignal;
+	bool																	bThreadStarted;
+
+	TSharedPtr<IMediaRenderer, ESPMode::ThreadSafe>							Renderer;
+	int32																	MaxDecodeBufferSize;
+
+	FMediaCriticalSection													ListenerMutex;
+	IAccessUnitBufferListener*												InputBufferListener;
+	IDecoderOutputBufferListener*											ReadyBufferListener;
+
+	IPlayerSessionServices* 												SessionServices;
+
+	TSharedPtr<MPEG::FAACDecoderConfigurationRecord, ESPMode::ThreadSafe>	ConfigRecord;
+	TSharedPtrTS<FAccessUnit::CodecData>									CurrentCodecData;
+
+	TSharedPtr<IAndroidJavaAACAudioDecoder, ESPMode::ThreadSafe>			DecoderInstance;
+	IAndroidJavaAACAudioDecoder::FOutputFormatInfo							CurrentOutputFormatInfo;
+	bool																	bInDummyDecodeMode;
+	bool																	bHaveDiscontinuity;
+
+	IMediaRenderer::IBuffer*												CurrentOutputBuffer;
+	FParamDict																BufferAcquireOptions;
+	FParamDict																OutputBufferSampleProperties;
+
+public:
+	static FSystemConfiguration												SystemConfig;
+};
+
+IAudioDecoderAAC::FSystemConfiguration	FAudioDecoderAAC::SystemConfig;
+
+
+/***************************************************************************************************************************************************/
+/***************************************************************************************************************************************************/
+/***************************************************************************************************************************************************/
+
+IAudioDecoderAAC::FSystemConfiguration::FSystemConfiguration()
+{
+	ThreadConfig.Decoder.Priority 	   = TPri_Normal;
+	ThreadConfig.Decoder.StackSize	   = 65536;
+	ThreadConfig.Decoder.CoreAffinity    = -1;
+	// Not needed, but setting meaningful values anyway.
+	ThreadConfig.PassOn.Priority  	   = TPri_Normal;
+	ThreadConfig.PassOn.StackSize 	   = 32768;
+	ThreadConfig.PassOn.CoreAffinity     = -1;
+}
+
+IAudioDecoderAAC::FInstanceConfiguration::FInstanceConfiguration()
+	: ThreadConfig(FAudioDecoderAAC::SystemConfig.ThreadConfig)
+{
+}
+
+bool IAudioDecoderAAC::Startup(const IAudioDecoderAAC::FSystemConfiguration& InConfig)
+{
+	return FAudioDecoderAAC::Startup(InConfig);
+}
+
+void IAudioDecoderAAC::Shutdown()
+{
+	FAudioDecoderAAC::Shutdown();
+}
+
+IAudioDecoderAAC* IAudioDecoderAAC::Create()
+{
+	return new FAudioDecoderAAC;
+}
+
+/***************************************************************************************************************************************************/
+/***************************************************************************************************************************************************/
+/***************************************************************************************************************************************************/
+
+
+//-----------------------------------------------------------------------------
+/**
+ * Decoder system startup
+ *
+ * @param config
+ *
+ * @return
+ */
+bool FAudioDecoderAAC::Startup(const IAudioDecoderAAC::FSystemConfiguration& InConfig)
+{
+	SystemConfig = InConfig;
+	return true;
+}
+
+
+//-----------------------------------------------------------------------------
+/**
+ * Decoder system shutdown.
+ */
+void FAudioDecoderAAC::Shutdown()
+{
+}
+
+
+//-----------------------------------------------------------------------------
+/**
+ * Constructor
+ */
+FAudioDecoderAAC::FAudioDecoderAAC()
+	: FMediaThread("ElectraPlayer::AAC decoder")
+	, bThreadStarted(false)
+	, Renderer(nullptr)
+	, MaxDecodeBufferSize(0)
+	, InputBufferListener(nullptr)
+	, ReadyBufferListener(nullptr)
+	, SessionServices(nullptr)
+	, bInDummyDecodeMode(false)
+	, CurrentOutputBuffer(nullptr)
+{
+}
+
+
+//-----------------------------------------------------------------------------
+/**
+ * Destructor
+ */
+FAudioDecoderAAC::~FAudioDecoderAAC()
+{
+	Close();
+}
+
+
+//-----------------------------------------------------------------------------
+/**
+ * Sets the owning player's session service interface.
+ *
+ * @param InSessionServices
+ *
+ * @return
+ */
+void FAudioDecoderAAC::SetPlayerSessionServices(IPlayerSessionServices* InSessionServices)
+{
+	SessionServices = InSessionServices;
+}
+
+
+//-----------------------------------------------------------------------------
+/**
+ * Creates the decoded sample blocks output buffers.
+ *
+ * @return
+ */
+bool FAudioDecoderAAC::CreateDecodedSamplePool()
+{
+	check(Renderer);
+	FParamDict poolOpts;
+	// Assume LC-AAC with at most 6 channels and 16 bit decoded PCM samples. This is ok for HE-AAC as well which can at most be stereo.
+	uint32 frameSize = sizeof(int16) * 6 * 1024;
+
+	poolOpts.Set("max_buffer_size",  FVariantValue((int64) frameSize));
+	poolOpts.Set("num_buffers",      FVariantValue((int64) 8));
+
+	UEMediaError Error = Renderer->CreateBufferPool(poolOpts);
+	check(Error == UEMEDIA_ERROR_OK);
+
+	if (Error != UEMEDIA_ERROR_OK)
+	{
+		PostError(0, "Failed to create sample pool", ERRCODE_INTERNAL_ANDROID_COULD_NOT_CREATE_SAMPLE_POOL, Error);
+	}
+
+	MaxDecodeBufferSize = (int32) Renderer->GetBufferPoolProperties().GetValue("max_buffers").GetInt64();
+
+	return Error == UEMEDIA_ERROR_OK;
+}
+
+
+//-----------------------------------------------------------------------------
+/**
+ * Destroys the pool of decoded sample blocks.
+ */
+void FAudioDecoderAAC::DestroyDecodedSamplePool()
+{
+	Renderer->ReleaseBufferPool();
+}
+
+
+//-----------------------------------------------------------------------------
+/**
+ * Opens a decoder instance
+ *
+ * @param InConfig
+ *
+ * @return
+ */
+void FAudioDecoderAAC::Open(const IAudioDecoderAAC::FInstanceConfiguration& InConfig)
+{
+	Config = InConfig;
+
+	// Set a large enough size to hold a single access unit. Since we will be asking for a new AU on
+	// demand there is no need to overly restrict ourselves here. But it must be large enough to
+	// hold at least the largest expected access unit.
+	AccessUnits.CapacitySet(FAccessUnitBuffer::FConfiguration(2 << 20, 60.0));
+
+	StartThread();
+}
+
+
+//-----------------------------------------------------------------------------
+/**
+ * Closes the decoder instance.
+ */
+void FAudioDecoderAAC::Close()
+{
+	StopThread();
+}
+
+
+//-----------------------------------------------------------------------------
+/**
+ * Sets a new renderer.
+ *
+ * @param InRenderer
+ */
+void FAudioDecoderAAC::SetRenderer(TSharedPtr<IMediaRenderer, ESPMode::ThreadSafe> InRenderer)
+{
+	Renderer = InRenderer;
+}
+
+
+//-----------------------------------------------------------------------------
+/**
+ * Sets an AU input buffer listener.
+ *
+ * @param InListener
+ */
+void FAudioDecoderAAC::SetAUInputBufferListener(IAccessUnitBufferListener* InListener)
+{
+	FMediaCriticalSection::ScopedLock lock(ListenerMutex);
+	InputBufferListener = InListener;
+}
+
+
+//-----------------------------------------------------------------------------
+/**
+ * Sets a buffer-ready listener.
+ *
+ * @param InListener
+ */
+void FAudioDecoderAAC::SetReadyBufferListener(IDecoderOutputBufferListener* InListener)
+{
+	FMediaCriticalSection::ScopedLock lock(ListenerMutex);
+	ReadyBufferListener = InListener;
+}
+
+
+//-----------------------------------------------------------------------------
+/**
+ * Creates and runs the decoder thread.
+ */
+void FAudioDecoderAAC::StartThread()
+{
+	ThreadSetPriority(Config.ThreadConfig.Decoder.Priority);
+	ThreadSetStackSize(Config.ThreadConfig.Decoder.StackSize);
+	ThreadSetCoreAffinity(Config.ThreadConfig.Decoder.CoreAffinity);
+	ThreadStart(Electra::MakeDelegate(this, &FAudioDecoderAAC::WorkerThread));
+	bThreadStarted = true;
+}
+
+
+//-----------------------------------------------------------------------------
+/**
+ * Stops the decoder thread.
+ */
+void FAudioDecoderAAC::StopThread()
+{
+	if (bThreadStarted)
+	{
+		TerminateThreadSignal.Signal();
+		ThreadWaitDone();
+		bThreadStarted = false;
+	}
+}
+
+
+//-----------------------------------------------------------------------------
+/**
+ * Called to receive a new input access unit for decoding.
+ *
+ * @param InAccessUnit
+ */
+IAccessUnitBufferInterface::EAUpushResult FAudioDecoderAAC::AUdataPushAU(FAccessUnit* InAccessUnit)
+{
+	// Add a ref count to the AU before adding it to the buffer. If it can be added successfully it could be
+	// used immediately by the worker thread and it would be too late for us to increase the count then.
+	InAccessUnit->AddRef();
+	bool bOk = AccessUnits.Push(InAccessUnit);
+	if (!bOk)
+	{
+		// Could not add. Undo the refcount increase.
+		FAccessUnit::Release(InAccessUnit);
+	}
+	return bOk ? IAccessUnitBufferInterface::EAUpushResult::Ok : IAccessUnitBufferInterface::EAUpushResult::Full;
+}
+
+
+//-----------------------------------------------------------------------------
+/**
+ * "Pushes" an End Of Data marker indicating no further access units will be added.
+ */
+void FAudioDecoderAAC::AUdataPushEOD()
+{
+	AccessUnits.PushEndOfData();
+}
+
+
+//-----------------------------------------------------------------------------
+/**
+ * Flushes the decoder and clears the input access unit buffer.
+ */
+void FAudioDecoderAAC::AUdataFlushEverything()
+{
+	FlushDecoderSignal.Signal();
+	DecoderFlushedSignal.WaitAndReset();
+}
+
+
+//-----------------------------------------------------------------------------
+/**
+ * Notifies the buffer-ready listener that we will now be producing output data
+ * and need an output buffer.
+ * This call is not intended to create or obtain an output buffer. It is merely
+ * indicating that new output will be produced.
+ *
+ * @param bHaveOutput
+ */
+void FAudioDecoderAAC::NotifyReadyBufferListener(bool bHaveOutput)
+{
+	if (ReadyBufferListener)
+	{
+		IDecoderOutputBufferListener::FDecodeReadyStats stats;
+		stats.MaxDecodedElementsReady = MaxDecodeBufferSize;
+		stats.NumElementsInDecoder    = CurrentOutputBuffer ? 1 : 0;
+		stats.bOutputStalled		  = !bHaveOutput;
+		stats.bEODreached   		  = AccessUnits.IsEndOfData() && stats.NumDecodedElementsReady == 0 && stats.NumElementsInDecoder == 0;
+		ListenerMutex.Lock();
+		if (ReadyBufferListener)
+		{
+			ReadyBufferListener->DecoderOutputReady(stats);
+		}
+		ListenerMutex.Unlock();
+	}
+}
+
+
+//-----------------------------------------------------------------------------
+/**
+ * Posts an error to the session service error listeners.
+ *
+ * @param ApiReturnValue
+ * @param Message
+ * @param Code
+ * @param Error
+ */
+void FAudioDecoderAAC::PostError(int32 ApiReturnValue, const FString& Message, uint16 Code, UEMediaError Error)
+{
+	check(SessionServices);	// there better be a session service interface to receive the error!
+	if (SessionServices)
+	{
+		FErrorDetail err;
+		err.SetError(Error != UEMEDIA_ERROR_OK ? Error : UEMEDIA_ERROR_DETAIL);
+		err.SetFacility(Facility::EFacility::AACDecoder);
+		err.SetCode(Code);
+		err.SetMessage(Message);
+		err.SetPlatformMessage(FString::Printf(TEXT("%d (0x%08x)"), (int32) ApiReturnValue, (int32) ApiReturnValue));
+		SessionServices->PostError(err);
+	}
+}
+
+
+//-----------------------------------------------------------------------------
+/**
+ * Sends a log message to the session service log.
+ *
+ * @param Level
+ * @param Message
+ */
+void FAudioDecoderAAC::LogMessage(IInfoLog::ELevel Level, const FString& Message)
+{
+	if (SessionServices)
+	{
+		SessionServices->PostLog(Facility::EFacility::AACDecoder, Level, Message);
+	}
+}
+
+
+//-----------------------------------------------------------------------------
+/**
+ * Creates an audio decoder.
+ *
+ * Note: This requires the configuration record to have been parsed successfully.
+ */
+bool FAudioDecoderAAC::InternalDecoderCreate()
+{
+	if (!ConfigRecord.IsValid())
+	{
+		PostError(0, "No CSD to create audio decoder with", ERRCODE_INTERNAL_ANDROID_COULD_NOT_CREATE_AUDIO_DECODER);
+		return false;
+	}
+
+	InternalDecoderDestroy();
+
+	DecoderInstance = IAndroidJavaAACAudioDecoder::Create();
+	int32 result = DecoderInstance->InitializeDecoder(*ConfigRecord);
+	if (result)
+	{
+		PostError(result, "Failed to create decoder", ERRCODE_INTERNAL_ANDROID_COULD_NOT_CREATE_AUDIO_DECODER);
+		return false;
+	}
+	// Start it.
+	result = DecoderInstance->Start();
+	if (result)
+	{
+		PostError(result, "Failed to start decoder", ERRCODE_INTERNAL_ANDROID_COULD_NOT_START_DECODER);
+		return false;
+	}
+	return true;
+}
+
+
+//-----------------------------------------------------------------------------
+/**
+ * Destroys the audio decoder.
+ */
+void FAudioDecoderAAC::InternalDecoderDestroy()
+{
+	if (DecoderInstance.IsValid())
+	{
+		DecoderInstance->Flush();
+		DecoderInstance->Stop();
+		DecoderInstance->ReleaseDecoder();
+		DecoderInstance.Reset();
+	}
+}
+
+
+//-----------------------------------------------------------------------------
+/**
+ * Returns the acquired sample output buffer back to the renderer without having it rendered.
+ */
+void FAudioDecoderAAC::ReturnUnusedOutputBuffer()
+{
+	if (CurrentOutputBuffer)
+	{
+		OutputBufferSampleProperties.Clear();
+		Renderer->ReturnBuffer(CurrentOutputBuffer, false, OutputBufferSampleProperties);
+		CurrentOutputBuffer = nullptr;
+	}
+}
+
+
+//-----------------------------------------------------------------------------
+/**
+ * Checks if the codec specific format has changed.
+ *
+ * @param Data
+ *
+ * @return
+ */
+bool FAudioDecoderAAC::IsDifferentFormat(const FAccessUnit* Data)
+{
+	if (Data->AUCodecData.IsValid() && (!CurrentCodecData.IsValid() || CurrentCodecData->CodecSpecificData != Data->AUCodecData->CodecSpecificData))
+	{
+		return true;
+	}
+	return false;
+}
+
+
+//-----------------------------------------------------------------------------
+/**
+ * Flushes the decoder
+ */
+bool FAudioDecoderAAC::FlushDecoder()
+{
+	if (DecoderInstance.IsValid())
+	{
+		int32 result = DecoderInstance->Flush();
+		return result == 0;
+	}
+	return true;
+}
+
+
+//-----------------------------------------------------------------------------
+/**
+ * Decodes an access unit
+ */
+bool FAudioDecoderAAC::Decode(FAccessUnit* InAccessUnit)
+{
+	uint32 nDataOffset = 0;
+	// Loop until all data has been consumed.
+	bool bAllInputConsumed = false;
+	while(!bAllInputConsumed && !TerminateThreadSignal.IsSignaled() && !FlushDecoderSignal.IsSignaled())
+	{
+		// Need a new output buffer?
+		if (CurrentOutputBuffer == nullptr)
+		{
+			UEMediaError bufResult = Renderer->AcquireBuffer(CurrentOutputBuffer, 0, BufferAcquireOptions);
+			check(bufResult == UEMEDIA_ERROR_OK || bufResult == UEMEDIA_ERROR_INSUFFICIENT_DATA);
+			if (bufResult != UEMEDIA_ERROR_OK && bufResult != UEMEDIA_ERROR_INSUFFICIENT_DATA)
+			{
+				PostError(0, "Failed to acquire sample buffer", ERRCODE_INTERNAL_ANDROID_COULD_NOT_GET_RENDER_BUFFER, bufResult);
+				return false;
+			}
+		}
+
+		bool bHaveAvailSmpBlk = CurrentOutputBuffer != nullptr;
+
+		// Check if the renderer can accept the output we will want to send to it.
+		// If it can't right now we treat this as if we do not have an available output buffer.
+		if (Renderer.IsValid() && !Renderer->CanReceiveOutputFrames(1))
+		{
+			bHaveAvailSmpBlk = false;
+		}
+
+		NotifyReadyBufferListener(bHaveAvailSmpBlk);
+		if (bHaveAvailSmpBlk)
+		{
+			FTimeValue CurrentPTS = InAccessUnit->PTS;
+
+			if (DecoderInstance.IsValid() && !InAccessUnit->bIsDummyData)
+			{
+				if (bInDummyDecodeMode)
+				{
+					bHaveDiscontinuity = true;
+					bInDummyDecodeMode = false;
+				}
+
+				int32 result;
+				// Check if there is an available input buffer.
+				int32 InputBufferIndex = DecoderInstance->DequeueInputBuffer(1000 * 2);
+				if (InputBufferIndex >= 0)
+				{
+					int64 pts = InAccessUnit->PTS.GetAsMicroseconds();
+					result = DecoderInstance->QueueInputBuffer(InputBufferIndex, Electra::AdvancePointer(InAccessUnit->AUData, nDataOffset), InAccessUnit->AUSize - nDataOffset, pts);
+					if (result == 0)
+					{
+						// Presently, with RAW AAC data, we pass the entire AU as this is how the data is packetized. With ADTS framing there could potentially be multiple packets that need to be fed individually.
+						nDataOffset += InAccessUnit->AUSize;
+					}
+					else
+					{
+						PostError(result, "Failed to submit decoder input buffer", ERRCODE_INTERNAL_ANDROID_FAILED_TO_DECODE_VIDEO);
+						return false;
+					}
+				}
+
+				// Check if there is available output.
+				IAndroidJavaAACAudioDecoder::FOutputBufferInfo OutputBufferInfo;
+				result = DecoderInstance->DequeueOutputBuffer(OutputBufferInfo, 1000 * 2);
+				if (result == 0)
+				{
+					if (OutputBufferInfo.BufferIndex >= 0)
+					{
+						int32 OutputByteCount    = OutputBufferInfo.Size;
+						int32 SamplingRate  	 = CurrentOutputFormatInfo.SampleRate;
+						int32 NumberOfChannels   = CurrentOutputFormatInfo.NumChannels;
+						int32 NumSamplesProduced = OutputByteCount / (CurrentOutputFormatInfo.BytesPerSample * NumberOfChannels);
+						FTimeValue Duration;
+						Duration.SetFromND(NumSamplesProduced, SamplingRate);
+						FTimeValue PTS;
+						PTS.SetFromMicroseconds(OutputBufferInfo.PresentationTimestamp);
+						bool bEOS = OutputBufferInfo.bIsEOS;
+
+						int32 CurrentRenderOutputBufferSize = (int32)CurrentOutputBuffer->GetBufferProperties().GetValue("size").GetInt64();
+						void* CurrentRenderOutputBufferAddress = CurrentOutputBuffer->GetBufferProperties().GetValue("address").GetPointer();
+						check(CurrentRenderOutputBufferSize >= OutputByteCount);
+						if (CurrentRenderOutputBufferSize >= OutputByteCount)
+						{
+							result = DecoderInstance->GetOutputBufferAndRelease(CurrentRenderOutputBufferAddress, CurrentRenderOutputBufferSize, OutputBufferInfo);
+							if (result == 0)
+							{
+								OutputBufferSampleProperties.Clear();
+								OutputBufferSampleProperties.Set("num_channels",  FVariantValue((int64) NumberOfChannels));
+								OutputBufferSampleProperties.Set("sample_rate",   FVariantValue((int64) SamplingRate));
+								OutputBufferSampleProperties.Set("byte_size",     FVariantValue((int64) OutputByteCount));
+								OutputBufferSampleProperties.Set("duration",      FVariantValue(Duration));
+								OutputBufferSampleProperties.Set("pts",           FVariantValue(PTS));
+								OutputBufferSampleProperties.Set("discontinuity", FVariantValue((bool)bHaveDiscontinuity));
+
+								Renderer->ReturnBuffer(CurrentOutputBuffer, true, OutputBufferSampleProperties);
+								CurrentOutputBuffer = nullptr;
+								bHaveDiscontinuity = false;
+							}
+							else
+							{
+								PostError(0, "Failed to get decoder output buffer", ERRCODE_INTERNAL_ANDROID_COULD_NOT_GET_OUTPUT_BUFFER);
+								return false;
+							}
+						}
+						else
+						{
+							PostError(0, "Output buffer not large enough", ERRCODE_INTERNAL_ANDROID_RENDER_BUFFER_TOO_SMALL);
+							return false;
+						}
+					}
+					else if (OutputBufferInfo.BufferIndex == IAndroidJavaAACAudioDecoder::FOutputBufferInfo::EBufferIndexValues::MediaCodec_INFO_TRY_AGAIN_LATER)
+					{
+					}
+					else if (OutputBufferInfo.BufferIndex == IAndroidJavaAACAudioDecoder::FOutputBufferInfo::EBufferIndexValues::MediaCodec_INFO_OUTPUT_FORMAT_CHANGED)
+					{
+						IAndroidJavaAACAudioDecoder::FOutputFormatInfo OutputFormatInfo;
+						result = DecoderInstance->GetOutputFormatInfo(OutputFormatInfo, -1);
+						if (result == 0)
+						{
+							CurrentOutputFormatInfo = OutputFormatInfo;
+						}
+						else
+						{
+							PostError(result, "Failed to get decoder output format", ERRCODE_INTERNAL_ANDROID_COULD_NOT_GET_OUTPUT_FORMAT);
+							return false;
+						}
+					}
+					else if (OutputBufferInfo.BufferIndex == IAndroidJavaAACAudioDecoder::FOutputBufferInfo::EBufferIndexValues::MediaCodec_INFO_OUTPUT_BUFFERS_CHANGED)
+					{
+					}
+					else
+					{
+						// What new value might this be?
+					}
+				}
+				else
+				{
+					PostError(result, "Failed to get decoder output buffer", ERRCODE_INTERNAL_ANDROID_COULD_NOT_GET_OUTPUT_BUFFER);
+					return false;
+				}
+				bAllInputConsumed = nDataOffset >= InAccessUnit->AUSize;
+			}
+			else
+			{
+				// Check if we are in dummy decode mode already. If not we need to flush the decoder if we have one.
+				if (!bInDummyDecodeMode)
+				{
+					if (DecoderInstance.IsValid())
+					{
+						// The decoder must have been started in order to flush it, so we must not stop it here!
+						// It does not require to be started again after flushing unless it is operating asynchronously, which we do not do.
+						DecoderInstance->Flush();
+					}
+					bInDummyDecodeMode = true;
+					bHaveDiscontinuity = true;
+				}
+
+				// Clear to silence
+				FMemory::Memzero(CurrentOutputBuffer->GetBufferProperties().GetValue("address").GetPointer(), CurrentOutputBuffer->GetBufferProperties().GetValue("size").GetInt64());
+
+				// Assume sensible defaults.
+				int64 NumChannels = 2;
+				int64 SampleRate = 48000;
+				int64 SamplesPerBlock = 1024;
+
+				// With a valid configuration record we can use the actual values.
+				if (ConfigRecord.IsValid())
+				{
+					// Parameteric stereo results in stereo (2 channels).
+					NumChannels = ConfigRecord->PSSignal > 0 ? 2 : ConfigRecord->ChannelConfiguration;
+					SampleRate = ConfigRecord->ExtSamplingFrequency ? ConfigRecord->ExtSamplingFrequency : ConfigRecord->SamplingRate;
+					SamplesPerBlock = ConfigRecord->SBRSignal > 0 ? 2048 : 1024;
+				}
+				else if (CurrentOutputFormatInfo.IsValid())
+				{
+					NumChannels = CurrentOutputFormatInfo.NumChannels;
+					SampleRate = CurrentOutputFormatInfo.SampleRate;
+					SamplesPerBlock = 1024;	// Whether this was SBR with 2048 samples or not we don't get back from the decoder :-(
+				}
+
+				FTimeValue Duration;
+				Duration.SetFromND(SamplesPerBlock, (uint32) SampleRate);
+
+				// Note: The duration in the AU should be an exact multiple of the "SamplesPerBlock / SampleRate" as this is how it gets set
+				//       up in the stream reader. Should there ever be any discrepancy such that "Duration > pData->mDuration" it is simply
+				//       possible to recalculate "SamplesPerBlock" such that we return a buffer with less than the usual 1024 or 2048 samples.
+				OutputBufferSampleProperties.Clear();
+				OutputBufferSampleProperties.Set("num_channels",  FVariantValue(NumChannels));
+				OutputBufferSampleProperties.Set("sample_rate",   FVariantValue(SampleRate));
+				OutputBufferSampleProperties.Set("byte_size",     FVariantValue((int64)(NumChannels * sizeof(int16) * SamplesPerBlock)));
+				OutputBufferSampleProperties.Set("duration",      FVariantValue(Duration));
+				OutputBufferSampleProperties.Set("pts",           FVariantValue(CurrentPTS));
+				OutputBufferSampleProperties.Set("discontinuity", FVariantValue((bool)bHaveDiscontinuity));
+
+				Renderer->ReturnBuffer(CurrentOutputBuffer, true, OutputBufferSampleProperties);
+				CurrentOutputBuffer = nullptr;
+				bHaveDiscontinuity = false;
+
+				CurrentPTS += Duration;
+				InAccessUnit->Duration -= Duration;
+				bAllInputConsumed = InAccessUnit->Duration <= FTimeValue::GetZero();
+			}
+		}
+		else
+		{
+			// No available buffer. Sleep for a bit. Can't sleep on a signal since we have to check two: abort and flush
+			// We sleep for 20ms since a sample block for LC AAC produces 1024 samples which amount to 21.3ms @48kHz (or more at lower sample rates).
+			FMediaRunnable::SleepMilliseconds(20);
+		}
+	}
+	return true;
+}
+
+
+//-----------------------------------------------------------------------------
+/**
+ * AAC audio decoder main threaded decode loop
+ */
+void FAudioDecoderAAC::WorkerThread()
+{
+	LLM_SCOPE(ELLMTag::MediaStreaming);
+
+	FAccessUnit*	CurrentAccessUnit = nullptr;
+	bool			bError			  = false;
+
+	CurrentOutputBuffer = nullptr;
+	bInDummyDecodeMode = false;
+	bHaveDiscontinuity = false;
+
+	bError = !CreateDecodedSamplePool();
+	check(!bError);
+
+	while(!TerminateThreadSignal.IsSignaled())
+	{
+		// Notify the buffer listener that we will now be needing an AU for our input buffer.
+		if (!bError && InputBufferListener && AccessUnits.Num() == 0)
+		{
+			FAccessUnitBufferInfo	sin;
+			IAccessUnitBufferListener::FBufferStats	stats;
+			AccessUnits.GetStats(sin);
+			stats.NumAUsAvailable  = sin.NumCurrentAccessUnits;
+			stats.NumBytesInBuffer = sin.CurrentMemInUse;
+			stats.MaxBytesOfBuffer = sin.MaxDataSize;
+			stats.bEODSignaled     = sin.bEndOfData;
+			stats.bEODReached      = sin.bEndOfData && sin.NumCurrentAccessUnits == 0;
+			ListenerMutex.Lock();
+			if (InputBufferListener)
+			{
+				InputBufferListener->DecoderInputNeeded(stats);
+			}
+			ListenerMutex.Unlock();
+		}
+		// Get the AU to be decoded if one is there.
+		bool bHaveData = AccessUnits.WaitForData(1000 * 10);
+		if (bHaveData)
+		{
+			AccessUnits.Pop(CurrentAccessUnit);
+
+			// Check if the format has changed such that we need to destroy and re-create the decoder.
+			if (CurrentAccessUnit->bTrackChangeDiscontinuity || IsDifferentFormat(CurrentAccessUnit) || (CurrentAccessUnit->bIsDummyData && !bInDummyDecodeMode))
+				{
+				if (DecoderInstance.IsValid())
+				{
+					InternalDecoderDestroy();
+					bHaveDiscontinuity = true;
+				}
+				ReturnUnusedOutputBuffer();
+				ConfigRecord.Reset();
+				CurrentCodecData.Reset();
+			}
+
+			// Parse the CSD into a configuration record.
+			if (!ConfigRecord.IsValid() && !bError)
+			{
+				if (CurrentAccessUnit->AUCodecData.IsValid())
+				{
+					CurrentCodecData = CurrentAccessUnit->AUCodecData;
+					ConfigRecord = MakeShared<MPEG::FAACDecoderConfigurationRecord, ESPMode::ThreadSafe>();
+					if (!ConfigRecord->ParseFrom(CurrentCodecData->CodecSpecificData.GetData(), CurrentCodecData->CodecSpecificData.Num()))
+					{
+						ConfigRecord.Reset();
+						CurrentCodecData.Reset();
+						PostError(0, "Failed to parse AAC configuration record", ERRCODE_INTERNAL_ANDROID_FAILED_TO_PARSE_CSD);
+						bError = true;
+					}
+				}
+			}
+
+			// Check if this audio packet is to be dropped
+			if (!CurrentAccessUnit->DropState)
+			{
+				// Need to create a decoder instance?
+				if (!DecoderInstance.IsValid() && !bError)
+				{
+					if (!InternalDecoderCreate())
+					{
+						bError = true;
+					}
+				}
+
+				// Decode if not in error, otherwise just spend some idle time.
+				if (!bError)
+				{
+					if (!Decode(CurrentAccessUnit))
+					{
+						bError = true;
+					}
+				}
+				else
+				{
+					// Pace ourselves in consuming any more data until we are being stopped.
+					FMediaRunnable::SleepMicroseconds(CurrentAccessUnit->Duration.GetAsMicroseconds());
+				}
+			}
+
+			// We're done with this access unit. Delete it.
+			FAccessUnit::Release(CurrentAccessUnit);
+			CurrentAccessUnit = nullptr;
+		}
+		else
+		{
+			// No data. Is the buffer at EOD?
+			if (AccessUnits.IsEndOfData())
+			{
+				NotifyReadyBufferListener(true);
+				FMediaRunnable::SleepMilliseconds(20);
+			}
+		}
+
+		// Flush?
+		if (FlushDecoderSignal.IsSignaled())
+		{
+			ReturnUnusedOutputBuffer();
+			AccessUnits.Flush();
+			FlushDecoder();
+
+			FlushDecoderSignal.Reset();
+			DecoderFlushedSignal.Signal();
+		}
+	}
+
+	ReturnUnusedOutputBuffer();
+	// Close the decoder.
+	InternalDecoderDestroy();
+	DestroyDecodedSamplePool();
+
+	CurrentCodecData.Reset();
+	ConfigRecord.Reset();
+
+	// Flush any remaining input data.
+	AccessUnits.Flush();
+	AccessUnits.CapacitySet(0);
+	}
+
+
+} // namespace Electra
+
+
diff --git a/Engine/Plugins/Media/ElectraPlayer/Source/ElectraPlayerRuntime/Private/Runtime/Decoder/Android/DecoderErrors_Android.h b/Engine/Plugins/Media/ElectraPlayer/Source/ElectraPlayerRuntime/Private/Runtime/Decoder/Android/DecoderErrors_Android.h
new file mode 100644
--- /dev/null
+++ b/Engine/Plugins/Media/ElectraPlayer/Source/ElectraPlayerRuntime/Private/Runtime/Decoder/Android/DecoderErrors_Android.h
@@ -0,0 +1,30 @@
+// Copyright Epic Games, Inc. All Rights Reserved.
+
+#pragma once
+
+
+#define ERRCODE_INTERNAL_ANDROID_OUT_OF_MEMORY						1
+#define ERRCODE_INTERNAL_ANDROID_COULD_NOT_CREATE_AUDIO_DECODER		2
+#define ERRCODE_INTERNAL_ANDROID_COULD_NOT_CREATE_VIDEO_DECODER		3
+#define ERRCODE_INTERNAL_ANDROID_COULD_NOT_CREATE_SAMPLE_POOL		4
+#define ERRCODE_INTERNAL_ANDROID_COULD_NOT_CREATE_IMAGE_POOL		5
+#define ERRCODE_INTERNAL_ANDROID_FAILED_TO_PARSE_CSD				6
+#define ERRCODE_INTERNAL_ANDROID_FAILED_TO_DECODE_AUDIO				7
+#define ERRCODE_INTERNAL_ANDROID_FAILED_TO_DECODE_VIDEO				8
+
+#define ERRCODE_INTERNAL_ANDROID_NO_VIDEO_CSD						20
+#define ERRCODE_INTERNAL_ANDROID_BAD_VIDEO_CSD						21
+
+#define ERRCODE_INTERNAL_ANDROID_COULD_NOT_CONFIGURE_DECODER		30
+#define ERRCODE_INTERNAL_ANDROID_COULD_NOT_START_DECODER			31
+#define ERRCODE_INTERNAL_ANDROID_COULD_NOT_STOP_DECODER				32
+#define ERRCODE_INTERNAL_ANDROID_COULD_NOT_GET_INPUT_BUFFER			33
+#define ERRCODE_INTERNAL_ANDROID_COULD_NOT_GET_OUTPUT_BUFFER		34
+#define ERRCODE_INTERNAL_ANDROID_COULD_NOT_GET_OUTPUT_FORMAT		35
+
+#define ERRCODE_INTERNAL_ANDROID_COULD_NOT_GET_RENDER_BUFFER		40
+#define ERRCODE_INTERNAL_ANDROID_RENDER_BUFFER_TOO_SMALL			41
+#define ERRCODE_INTERNAL_ANDROID_COULD_NOT_FIND_OUTPUT_INFO			42
+
+#define ERRCODE_INTERNAL_ANDROID_INTERNAL							10000
+
diff --git a/Engine/Plugins/Media/ElectraPlayer/Source/ElectraPlayerRuntime/Private/Runtime/Decoder/Android/Java/src/com/epicgames/ue4/ElectraAudioDecoderAAC.java b/Engine/Plugins/Media/ElectraPlayer/Source/ElectraPlayerRuntime/Private/Runtime/Decoder/Android/Java/src/com/epicgames/ue4/ElectraAudioDecoderAAC.java
new file mode 100644
--- /dev/null
+++ b/Engine/Plugins/Media/ElectraPlayer/Source/ElectraPlayerRuntime/Private/Runtime/Decoder/Android/Java/src/com/epicgames/ue4/ElectraAudioDecoderAAC.java
@@ -0,0 +1,348 @@
+// Copyright Epic Games, Inc. All Rights Reserved.
+
+package com.epicgames.ue4;
+
+import android.os.Build;
+import java.io.IOException;
+import java.util.ArrayList;
+import android.util.Log;
+import android.util.SparseArray;
+
+import android.media.AudioFormat;
+import android.media.MediaCodec;
+import android.media.MediaCodecInfo;
+import android.media.MediaFormat;
+import java.nio.ByteBuffer;
+
+
+/**
+ * AAC audio decoder
+ *
+ */
+public class ElectraAudioDecoderAAC
+{
+	private static final String TAG = "[ElectraPlayerSDK]";
+
+	public class FOutputFormatInfo
+	{
+		public int SampleRate = 0;
+		public int NumChannels = 0;
+		public int BytesPerSample = 0;
+	}
+
+	public class FOutputBufferInfo
+	{
+		public int BufferIndex = -1;
+		public long PresentationTimestamp = -1;
+		public int Size = 0;
+		public boolean bIsEOS = false;
+		public boolean bIsConfig = false;
+	}
+
+	private MediaCodec DecoderHandle = null;
+	private boolean bIsInitialized = false;
+
+	private SparseArray<MediaCodec.BufferInfo> OutputBufferInfos = null;
+	private FOutputFormatInfo CurrentOutputFormatInfo = null;
+
+
+
+	public ElectraAudioDecoderAAC()
+	{
+		DecoderHandle = null;
+		bIsInitialized = false;
+		OutputBufferInfos = new SparseArray<>();
+		CurrentOutputFormatInfo = new FOutputFormatInfo();
+	}
+
+
+	public int CreateDecoder(int InNumChannels, int InSampleRate, byte[] InCSD)
+	{
+		try
+		{
+			// Create decoder configuration object
+			MediaFormat Format = MediaFormat.createAudioFormat(MediaFormat.MIMETYPE_AUDIO_AAC, InSampleRate, InNumChannels);
+
+			// When decoding RAW AAC we need to have codec specific data.
+			ByteBuffer Csd = ByteBuffer.wrap(InCSD);
+			Format.setByteBuffer("csd-0", Csd);
+			// Request float output.
+			// NOTE: Bah, is not honored by the decoder anyway...
+			//		Format.setInteger(MediaFormat.KEY_PCM_ENCODING, AudioFormat.ENCODING_PCM_FLOAT);
+
+			// Create and configure the decoder
+			DecoderHandle = MediaCodec.createDecoderByType(MediaFormat.MIMETYPE_AUDIO_AAC);
+			// Configure the codec (no surface, no crypto, as decoder)
+			if (DecoderHandle != null)
+			{
+				DecoderHandle.configure(Format, null, null, 0);
+				bIsInitialized = true;
+			}
+		}
+		catch(Exception e)
+		{
+			Log.w(TAG, "ElectraAudioDecoderAAC: Failed to create decoder");
+			e.printStackTrace();
+			return 1;
+		}
+		return 0;
+	}
+
+
+	public int ReleaseDecoder()
+	{
+		if (bIsInitialized)
+		{
+			bIsInitialized = false;
+			try
+			{
+				DecoderHandle.release();
+			}
+			catch(Exception e)
+			{
+				Log.w(TAG, "ElectraAudioDecoderAAC: Failed to release decoder");
+				e.printStackTrace();
+				return 1;
+			}
+			return 0;
+		}
+		else
+		{
+			return 1;
+		}
+	}
+
+
+	public int Start()
+	{
+		if (bIsInitialized)
+		{
+			try
+			{
+				// Start the decoder
+				DecoderHandle.start();
+			}
+			catch(Exception e)
+			{
+				Log.w(TAG, "ElectraAudioDecoderAAC: Failed to start decoder");
+				e.printStackTrace();
+				return 1;
+			}
+			return 0;
+		}
+		else
+		{
+			return 1;
+		}
+	}
+
+
+	public int Stop()
+	{
+		if (bIsInitialized)
+		{
+			try
+			{
+				// Stop the decoder
+				DecoderHandle.stop();
+			}
+			catch(Exception e)
+			{
+				Log.w(TAG, "ElectraAudioDecoderAAC: Failed to stop decoder");
+				e.printStackTrace();
+				return 1;
+			}
+			return 0;
+		}
+		else
+		{
+			return 1;
+		}
+	}
+
+
+	public int Flush()
+	{
+		if (bIsInitialized)
+		{
+			try
+			{
+				// Flush the decoder
+				DecoderHandle.flush();
+			}
+			catch(Exception e)
+			{
+				Log.w(TAG, "ElectraAudioDecoderAAC: Failed to flush decoder");
+				e.printStackTrace();
+				return 1;
+			}
+
+			// Clear the current infos by creating new empty structures.
+			OutputBufferInfos = new SparseArray<>();
+			// NOTE: Flushing does not invalidate the current output configuration. Do not clear it.
+			// CurrentOutputFormatInfo = new FOutputFormatInfo();
+			return 0;
+		}
+		else
+		{
+			return 1;
+		}
+	}
+
+
+	public int DequeueInputBuffer(int InTimeoutUsec)
+	{
+		if (bIsInitialized)
+		{
+			int BufferIndex = -1;
+			try
+			{
+				BufferIndex = DecoderHandle.dequeueInputBuffer(InTimeoutUsec);
+			}
+			catch(Exception e)
+			{
+				Log.w(TAG, "ElectraAudioDecoderAAC: Failed to dequeue input buffer");
+				e.printStackTrace();
+				return -10000;
+			}
+			return BufferIndex;
+		}
+		else
+		{
+			return -10001;
+		}
+	}
+
+
+	public int QueueInputBuffer(int InBufferIndex, long InTimestampUSec, byte[] InAccessUnit)
+	{
+		if (bIsInitialized)
+		{
+			if (InBufferIndex >= 0)
+			{
+				try
+				{
+					int SizeOfAU = InAccessUnit.length;
+					ByteBuffer Buffer = DecoderHandle.getInputBuffer(InBufferIndex);
+					Buffer.put(InAccessUnit);
+					DecoderHandle.queueInputBuffer(InBufferIndex, 0, SizeOfAU, InTimestampUSec, 0);
+				}
+				catch(Exception e)
+				{
+					Log.w(TAG, "ElectraAudioDecoderAAC: Failed to queue input buffer");
+					e.printStackTrace();
+					return -10000;
+				}
+				return 0;
+			}
+			else
+			{
+				return -10002;
+			}
+		}
+		else
+		{
+			return -10001;
+		}
+	}
+
+
+	public FOutputFormatInfo GetOutputFormatInfo(int InOutputBufferIndex)
+	{
+		//FOutputFormatInfo OutputFormatInfo = new FOutputFormatInfo();
+		//return OutputFormatInfo;
+		return CurrentOutputFormatInfo;
+	}
+
+
+	public FOutputBufferInfo DequeueOutputBuffer(int InTimeoutUsec)
+	{
+		if (bIsInitialized)
+		{
+			FOutputBufferInfo OutputBufferInfo = new FOutputBufferInfo();
+			MediaCodec.BufferInfo Info = new MediaCodec.BufferInfo();
+			try
+			{
+				OutputBufferInfo.BufferIndex = DecoderHandle.dequeueOutputBuffer(Info, InTimeoutUsec);
+				// If there is output we need to remember the buffer metadata for later.
+				if (OutputBufferInfo.BufferIndex >= 0)
+				{
+					OutputBufferInfo.PresentationTimestamp = Info.presentationTimeUs;
+					OutputBufferInfo.Size = Info.size;
+					OutputBufferInfo.bIsEOS = (Info.flags & MediaCodec.BUFFER_FLAG_END_OF_STREAM) != 0;
+					OutputBufferInfo.bIsConfig = (Info.flags & MediaCodec.BUFFER_FLAG_CODEC_CONFIG) != 0;
+					OutputBufferInfos.put(OutputBufferInfo.BufferIndex, Info);
+				}
+				else if (OutputBufferInfo.BufferIndex == MediaCodec.INFO_TRY_AGAIN_LATER)				// -1
+				{
+					// Have no output yet.
+				}
+				else if (OutputBufferInfo.BufferIndex == MediaCodec.INFO_OUTPUT_FORMAT_CHANGED)			// -2
+				{
+					MediaFormat Format = DecoderHandle.getOutputFormat();
+					//GameActivity.Log.debug("Android Audio Decoder: New format " + Format);
+
+					// Mwahahaaa... this only works from API level 29 onward...
+					//		boolean bIsFloat = Format.getInteger(MediaFormat.KEY_PCM_ENCODING, AudioFormat.ENCODING_PCM_16BIT) == AudioFormat.ENCODING_PCM_FLOAT;
+
+					FOutputFormatInfo NewFormatInfo = new FOutputFormatInfo();
+					NewFormatInfo.SampleRate	 = Format.getInteger(MediaFormat.KEY_SAMPLE_RATE);
+					NewFormatInfo.NumChannels    = Format.getInteger(MediaFormat.KEY_CHANNEL_COUNT);
+					//NewFormatInfo.BytesPerSample = bIsFloat ? 4 : 2;
+					NewFormatInfo.BytesPerSample = 2;
+					if (Format.containsKey(MediaFormat.KEY_PCM_ENCODING))
+					{
+						boolean bIsFloat = Format.getInteger(MediaFormat.KEY_PCM_ENCODING) == AudioFormat.ENCODING_PCM_FLOAT;
+						if (bIsFloat)
+						{
+							NewFormatInfo.BytesPerSample = 4;
+						}
+					}
+					CurrentOutputFormatInfo = NewFormatInfo;
+				}
+				else if (OutputBufferInfo.BufferIndex == MediaCodec.INFO_OUTPUT_BUFFERS_CHANGED)		// -3
+				{
+					// This is deprecated with API level 21 and of no real value to us.
+				}
+			}
+			catch(Exception e)
+			{
+				Log.w(TAG, "ElectraAudioDecoderAAC: Failed to dequeue output buffer");
+				e.printStackTrace();
+				return null;
+			}
+			return OutputBufferInfo;
+		}
+		return null;
+	}
+
+
+	public byte[] GetOutputBufferAndRelease(int InOutputBufferIndex)
+	{
+		byte[] Data = null;
+		if (bIsInitialized && InOutputBufferIndex >= 0)
+		{
+			MediaCodec.BufferInfo Info = OutputBufferInfos.get(InOutputBufferIndex);
+			if (Info != null)
+			{
+				OutputBufferInfos.delete(InOutputBufferIndex);
+
+				try
+				{
+					ByteBuffer outBuffer = DecoderHandle.getOutputBuffer(InOutputBufferIndex);
+					Data = new byte[Info.size];
+					outBuffer.get(Data);
+					DecoderHandle.releaseOutputBuffer(InOutputBufferIndex, false);
+				}
+				catch(Exception e)
+				{
+					Log.w(TAG, "ElectraAudioDecoderAAC: Failed to get output buffer");
+					e.printStackTrace();
+					return null;
+				}
+			}
+		}
+		return Data;
+	}
+}
+
diff --git a/Engine/Plugins/Media/ElectraPlayer/Source/ElectraPlayerRuntime/Private/Runtime/Decoder/Android/Java/src/com/epicgames/ue4/ElectraVideoDecoderH264.java b/Engine/Plugins/Media/ElectraPlayer/Source/ElectraPlayerRuntime/Private/Runtime/Decoder/Android/Java/src/com/epicgames/ue4/ElectraVideoDecoderH264.java
new file mode 100644
--- /dev/null
+++ b/Engine/Plugins/Media/ElectraPlayer/Source/ElectraPlayerRuntime/Private/Runtime/Decoder/Android/Java/src/com/epicgames/ue4/ElectraVideoDecoderH264.java
@@ -0,0 +1,714 @@
+// Copyright Epic Games, Inc. All Rights Reserved.
+
+package com.epicgames.ue4;
+
+import android.os.Build;
+import java.io.IOException;
+import java.util.ArrayList;
+import android.util.Log;
+import android.util.SparseArray;
+
+import android.opengl.*;
+import android.graphics.SurfaceTexture;
+import android.view.Surface;
+import android.media.AudioFormat;
+import android.media.MediaCodecList;
+import android.media.MediaCodec;
+import android.media.MediaCodecInfo;
+import android.media.MediaCodecInfo.CodecCapabilities;
+import android.media.MediaCodecInfo.CodecProfileLevel;
+import android.media.MediaFormat;
+import java.nio.ByteBuffer;
+
+
+/**
+ * H.264 video decoder
+ *
+ */
+public class ElectraVideoDecoderH264
+{
+	private static final String TAG = "[ElectraPlayerSDK]";
+	/**
+	 * Parameters for creating a new decoder instance.
+	 * NOTE: Make sure this is a static inner class so it can be
+	 *       created from C++ through JNI.
+	 */
+	public static class FCreateParameters
+	{
+		public int MaxWidth;
+		public int MaxHeight;
+		public int MaxFPS;
+		public int Width;
+		public int Height;
+		public boolean bNeedSecure;
+		public boolean bNeedTunneling;
+		public byte[] CSD0;
+		public byte[] CSD1;
+		public int ExternalRenderTextureID;
+		public boolean bRetainRenderer;
+		public boolean bSwizzleTexture;
+		public int NativeDecoderID;
+		public android.view.Surface VideoCodecSurface;
+		public boolean bSurfaceIsView;
+		public FCreateParameters()
+		{
+			MaxWidth				= 0;
+			MaxHeight   			= 0;
+			Width   				= 0;
+			Height  				= 0;
+			MaxFPS  				= 0;
+			bNeedSecure 			= false;
+			bNeedTunneling  		= false;
+			ExternalRenderTextureID = 0;
+			bRetainRenderer 		= false;
+			bSwizzleTexture 		= false;
+			NativeDecoderID			= 0;
+			VideoCodecSurface		= null;
+		}
+	}
+
+	public static class FDecoderInformation
+	{
+		public boolean bIsAdaptive = false;
+	}
+
+	public class FOutputFormatInfo
+	{
+		public int Width = 0;
+		public int Height = 0;
+		public int CropTop = 0;
+		public int CropBottom = 0;
+		public int CropLeft = 0;
+		public int CropRight = 0;
+		public int Stride = 0;
+		public int SliceHeight = 0;
+		public int ColorFormat = 0;
+	}
+
+	public class FOutputBufferInfo
+	{
+		public int BufferIndex = -1;
+		public long PresentationTimestamp = -1;
+		public int Size = 0;
+		public boolean bIsEOS = false;
+		public boolean bIsConfig = false;
+	}
+
+	private static class FCodecInformation
+	{
+		MediaCodecInfo						Info;
+		MediaCodecInfo.CodecCapabilities	Caps;
+	}
+
+
+	public class FFrameUpdateInfo
+	{
+		public java.nio.Buffer Buffer;
+		public long Timestamp = 0;
+		public long Duration = 0;
+		public boolean bFrameReady = false;
+		public boolean bRegionChanged = false;
+		public float UScale = 0.0f;
+		public float UOffset = 0.0f;
+		public float VScale = 0.0f;
+		public float VOffset = 0.0f;
+		public int NumPending = 0;
+	}
+
+
+	private FCreateParameters CreationParameters = null;
+	private FDecoderInformation DecoderInformation = null;
+	private MediaCodec DecoderHandle = null;
+	private boolean bIsInitialized = false;
+
+	private SparseArray<MediaCodec.BufferInfo> OutputBufferInfos = null;
+	private FOutputFormatInfo CurrentOutputFormatInfo = null;
+
+	private boolean bDecoderHasSurface = false;
+
+	private static FCodecInformation selectCodec(String mimeType)
+	{
+		/*
+		If we knew what to put into the media format exactly we could perhaps use findDecoderForFormat().
+			@see: https://developer.android.com/reference/android/media/MediaCodec.html#creation
+		Since we do not have an extractor to get the track format from, after all we are decoding an elementary stream here,
+		we have no idea what would normally be found and what properties are _required_ in the format.
+
+			MediaFormat testFmt = new MediaFormat();
+			String thisOne = cl.findDecoderForFormat(testFmt);
+		*/
+
+		MediaCodecList cl = new MediaCodecList(MediaCodecList.ALL_CODECS);
+		MediaCodecInfo[] codecInfos = cl.getCodecInfos();
+		int numCodecs = codecInfos.length;
+		for(int i=0; i<numCodecs; ++i)
+		{
+			if (codecInfos[i].isEncoder())
+			{
+				continue;
+			}
+
+			String[] types = codecInfos[i].getSupportedTypes();
+			for(int j=0; j<types.length; ++j)
+			{
+				if (types[j].equalsIgnoreCase(mimeType))
+				{
+					// Rumor has it that the codec list is sorted from "best" to "worst".
+					// Assuming this to be true we just return the first one matching the MIME type.
+					FCodecInformation Info = new FCodecInformation();
+					Info.Info = codecInfos[i];
+					Info.Caps = codecInfos[i].getCapabilitiesForType(types[j]);
+					return Info;
+				}
+			}
+		}
+		return null;
+	}
+
+	public ElectraVideoDecoderH264()
+	{
+		DecoderHandle = null;
+		bIsInitialized = false;
+		OutputBufferInfos = new SparseArray<>();
+		CurrentOutputFormatInfo = new FOutputFormatInfo();
+	}
+
+	public int CreateDecoder(FCreateParameters InCreateParams)
+	{
+		try
+		{
+			CreationParameters = InCreateParams;
+
+			FCodecInformation DecoderInfo = selectCodec("video/avc");
+			if (DecoderInfo != null)
+			{
+				String DecoderName = DecoderInfo.Info.getName();
+
+				// This is how we get information on the decoder's capabilities.
+				// For now we do not use those. We merely assume the decoder is capable of what we want.
+				// Caution must be exercised with profile/level values as they are defined with other values in Android than they are by the standard!
+				// Eg. High profile per the standard is identified by the value 100 while Android defines it as 8 (MediaCodecInfo.CodecProfileLevel.AVCProfileHigh)
+				boolean bAdaptive   						= DecoderInfo.Caps.isFeatureSupported(MediaCodecInfo.CodecCapabilities.FEATURE_AdaptivePlayback);
+				boolean bSecure 							= DecoderInfo.Caps.isFeatureSupported(MediaCodecInfo.CodecCapabilities.FEATURE_SecurePlayback);
+				boolean bTunneled   						= DecoderInfo.Caps.isFeatureSupported(MediaCodecInfo.CodecCapabilities.FEATURE_TunneledPlayback);
+				int [] ColorFmts							= DecoderInfo.Caps.colorFormats;
+				MediaCodecInfo.CodecProfileLevel [] ProfLev = DecoderInfo.Caps.profileLevels;
+
+
+				// Create decoder configuration object
+				int MaxWidth  =  InCreateParams.MaxWidth > 0 ? InCreateParams.MaxWidth  : 1920;
+				int MaxHeight = InCreateParams.MaxHeight > 0 ? InCreateParams.MaxHeight : 1088;
+				int Width  = InCreateParams.Width > 0  ? InCreateParams.Width : MaxWidth;
+				int Height = InCreateParams.Height > 0 ? InCreateParams.Height : MaxHeight;
+				MediaFormat Format = MediaFormat.createVideoFormat(MediaFormat.MIMETYPE_VIDEO_AVC, Width, Height);
+				// Add additional keys
+				if (bAdaptive)
+				{
+					// See: https://developer.android.com/reference/android/media/MediaFormat.html
+					Format.setInteger(MediaFormat.KEY_MAX_WIDTH, MaxWidth);
+					Format.setInteger(MediaFormat.KEY_MAX_HEIGHT, MaxHeight);
+				}
+				else
+				{
+					// Non-adaptive decoders need the CSD set in the configuration.
+					if (InCreateParams.CSD0 != null && InCreateParams.CSD0.length != 0)
+					{
+						ByteBuffer Csd0 = ByteBuffer.wrap(InCreateParams.CSD0);
+						Format.setByteBuffer("csd-0", Csd0);
+						// Set CSD1 only when it exists and we have set CSD0
+						if (InCreateParams.CSD1 != null && InCreateParams.CSD1.length != 0)
+						{
+							ByteBuffer Csd1 = ByteBuffer.wrap(InCreateParams.CSD1);
+							Format.setByteBuffer("csd-1", Csd1);
+						}
+					}
+				}
+
+				if (InCreateParams.MaxFPS > 0)
+				{
+					Format.setInteger(MediaFormat.KEY_FRAME_RATE, InCreateParams.MaxFPS);
+				}
+// Wishful thinking?
+//				Format.setInteger(MediaFormat.KEY_COLOR_FORMAT, MediaCodecInfo.CodecCapabilities.COLOR_FormatSurface);
+
+				// Create and configure the decoder
+				DecoderHandle = MediaCodec.createByCodecName(DecoderName);
+				if (DecoderHandle != null)
+				{
+					android.view.Surface OutputSurface = null;
+
+					if (InCreateParams.bRetainRenderer == false)
+					{
+						OutputSurface = InCreateParams.VideoCodecSurface;
+					}
+
+					bDecoderHasSurface = (OutputSurface != null);
+
+					DecoderHandle.configure(Format, OutputSurface, null, 0);
+					bIsInitialized = true;
+				}
+
+				// We are not storing any outside surface reference (we only hand it through to the decoder)
+				CreationParameters.VideoCodecSurface = null;
+
+				// Remember the decoder capabilities
+				DecoderInformation = new FDecoderInformation();
+				DecoderInformation.bIsAdaptive = bAdaptive;
+			}
+			else
+			{
+				// No decoder found.
+				Log.w(TAG, "ElectraVideoDecoderH264: No suitable decoder found");
+				return 1;
+			}
+		}
+		catch(Exception e)
+		{
+			Log.w(TAG, "ElectraVideoDecoderH264: Failed to create decoder");
+			e.printStackTrace();
+			return 1;
+		}
+		return 0;
+	}
+
+
+	public int ReleaseDecoder()
+	{
+		if (bIsInitialized)
+		{
+			bIsInitialized = false;
+
+			// Note: Do not release the texture renderer here as we might be able to reuse it later.
+			try
+			{
+				// Release the decoder
+				DecoderHandle.release();
+			}
+			catch(Exception e)
+			{
+				Log.w(TAG, "ElectraVideoDecoderH264: Failed to release decoder");
+				e.printStackTrace();
+				return 1;
+			}
+			return 0;
+		}
+		else
+		{
+			return 1;
+		}
+	}
+
+	public int release()
+	{
+		ReleaseDecoder();
+		return 0;
+	}
+
+
+	public int Start()
+	{
+		if (bIsInitialized)
+		{
+			try
+			{
+				// Start the decoder
+				DecoderHandle.start();
+			}
+			catch(Exception e)
+			{
+				Log.w(TAG,"ElectraVideoDecoderH264: Failed to start decoder");
+				e.printStackTrace();
+				return 1;
+			}
+			return 0;
+		}
+		else
+		{
+			return 1;
+		}
+	}
+
+
+	public int Stop()
+	{
+		if (bIsInitialized)
+		{
+			try
+			{
+				// Stop the decoder
+				DecoderHandle.stop();
+			}
+			catch(Exception e)
+			{
+				Log.w(TAG,"ElectraVideoDecoderH264: Failed to stop decoder");
+				e.printStackTrace();
+				return 1;
+			}
+			return 0;
+		}
+		else
+		{
+			return 1;
+		}
+	}
+
+
+	public int Flush()
+	{
+		if (bIsInitialized)
+		{
+			try
+			{
+				// Flush the decoder
+				DecoderHandle.flush();
+			}
+			catch(Exception e)
+			{
+				Log.w(TAG,"ElectraVideoDecoderH264: Failed to flush decoder");
+				e.printStackTrace();
+				return 1;
+			}
+
+			// Clear the current infos by creating new empty structures.
+			OutputBufferInfos = new SparseArray<>();
+			// NOTE: Flushing does not invalidate the current output configuration. Do not clear it.
+			// CurrentOutputFormatInfo = new FOutputFormatInfo();
+			return 0;
+		}
+		else
+		{
+			return 1;
+		}
+	}
+
+	public int Reset()
+	{
+		if (bIsInitialized)
+		{
+			try
+			{
+				// Reset the decoder
+				DecoderHandle.reset();
+			}
+			catch(Exception e)
+			{
+				Log.w(TAG,"ElectraVideoDecoderH264: Failed to reset decoder");
+				e.printStackTrace();
+				return 1;
+			}
+			return 0;
+		}
+		else
+		{
+			return 1;
+		}
+	}
+
+
+	public int DequeueInputBuffer(int InTimeoutUsec)
+	{
+		if (bIsInitialized)
+		{
+			int BufferIndex = -1;
+			try
+			{
+				BufferIndex = DecoderHandle.dequeueInputBuffer(InTimeoutUsec);
+			}
+			catch(Exception e)
+			{
+				Log.w(TAG,"ElectraVideoDecoderH264: Failed to dequeue input buffer");
+				e.printStackTrace();
+				return -10000;
+			}
+			return BufferIndex;
+		}
+		else
+		{
+			return -10001;
+		}
+	}
+
+
+	public int QueueInputBuffer(int InBufferIndex, long InTimestampUSec, byte[] InAccessUnit)
+	{
+		if (bIsInitialized)
+		{
+			if (InBufferIndex >= 0)
+			{
+				try
+				{
+					int SizeOfAU = InAccessUnit.length;
+					ByteBuffer Buffer = DecoderHandle.getInputBuffer(InBufferIndex);
+					Buffer.put(InAccessUnit);
+					DecoderHandle.queueInputBuffer(InBufferIndex, 0, SizeOfAU, InTimestampUSec, 0);
+				}
+				catch(Exception e)
+				{
+					Log.w(TAG,"ElectraVideoDecoderH264: Failed to queue input buffer");
+					e.printStackTrace();
+					return -10000;
+				}
+				return 0;
+			}
+			else
+			{
+				return -10002;
+			}
+		}
+		else
+		{
+			return -10001;
+		}
+	}
+
+
+	public int QueueCSDInputBuffer(int InBufferIndex, long InTimestampUSec, byte[] InCSD)
+	{
+		if (bIsInitialized)
+		{
+			if (InBufferIndex >= 0)
+			{
+				try
+				{
+					int SizeOfCSD = InCSD.length;
+					ByteBuffer Buffer = DecoderHandle.getInputBuffer(InBufferIndex);
+					Buffer.put(InCSD);
+					DecoderHandle.queueInputBuffer(InBufferIndex, 0, SizeOfCSD, InTimestampUSec, MediaCodec.BUFFER_FLAG_CODEC_CONFIG);
+				}
+				catch(Exception e)
+				{
+					Log.w(TAG,"ElectraVideoDecoderH264: Failed to queue CSD input buffer");
+					e.printStackTrace();
+					return -10000;
+				}
+				return 0;
+			}
+			else
+			{
+				return -10002;
+			}
+		}
+		else
+		{
+			return -10001;
+		}
+	}
+
+
+	public int QueueEOSInputBuffer(int InBufferIndex, long InTimestampUSec)
+	{
+		if (bIsInitialized)
+		{
+			if (InBufferIndex >= 0)
+			{
+				try
+				{
+					ByteBuffer Buffer = DecoderHandle.getInputBuffer(InBufferIndex);
+					Buffer.clear();
+					DecoderHandle.queueInputBuffer(InBufferIndex, 0, 0, InTimestampUSec, MediaCodec.BUFFER_FLAG_END_OF_STREAM);
+				}
+				catch(Exception e)
+				{
+					Log.w(TAG,"ElectraVideoDecoderH264: Failed to queue EOS input buffer");
+					e.printStackTrace();
+					return -10000;
+				}
+				return 0;
+			}
+			else
+			{
+				return -10002;
+			}
+		}
+		else
+		{
+			return -10001;
+		}
+	}
+
+
+	public FDecoderInformation GetDecoderInformation()
+	{
+		return DecoderInformation;
+	}
+
+
+	private FOutputFormatInfo GetOutputFormatInfoFromFormat(MediaFormat Format)
+	{
+		FOutputFormatInfo NewFormatInfo = new FOutputFormatInfo();
+		if (Format.containsKey(MediaFormat.KEY_WIDTH))
+		{
+			NewFormatInfo.Width = Format.getInteger(MediaFormat.KEY_WIDTH);
+		}
+		if (Format.containsKey(MediaFormat.KEY_HEIGHT))
+		{
+			NewFormatInfo.Height = Format.getInteger(MediaFormat.KEY_HEIGHT);
+		}
+		if (Format.containsKey("slice-height"))	// MediaFormat.KEY_SLICE_HEIGHT  with API level 23 only...
+		{
+			NewFormatInfo.SliceHeight = Format.getInteger("slice-height");
+		}
+		if (Format.containsKey("stride"))	// MediaFormat.KEY_STRIDE  with API level 23 only...
+		{
+			NewFormatInfo.Stride = Format.getInteger("stride");
+		}
+		if (Format.containsKey(MediaFormat.KEY_COLOR_FORMAT))
+		{
+			NewFormatInfo.ColorFormat = Format.getInteger(MediaFormat.KEY_COLOR_FORMAT);
+		}
+		// See: https://developer.android.com/reference/android/media/MediaCodec#accessing-raw-video-bytebuffers-on-older-devices
+		if (Format.containsKey("crop-left") && Format.containsKey("crop-right") && Format.containsKey("crop-top") && Format.containsKey("crop-bottom"))
+		{
+			NewFormatInfo.CropLeft   = Format.getInteger("crop-left");
+			NewFormatInfo.CropRight  = Format.getInteger("crop-right");
+			NewFormatInfo.CropTop    = Format.getInteger("crop-top");
+			NewFormatInfo.CropBottom = Format.getInteger("crop-bottom");
+		}
+		else
+		{
+			NewFormatInfo.CropLeft   = 0;
+			NewFormatInfo.CropTop    = 0;
+			NewFormatInfo.CropRight  = NewFormatInfo.Width - 1;
+			NewFormatInfo.CropBottom = NewFormatInfo.Height - 1;
+		}
+		return NewFormatInfo;
+	}
+
+
+	public FOutputFormatInfo GetOutputFormatInfo(int InOutputBufferIndex)
+	{
+		if (InOutputBufferIndex < 0)
+		{
+			return CurrentOutputFormatInfo;
+		}
+		if (bIsInitialized)
+		{
+			MediaFormat Format = DecoderHandle.getOutputFormat(InOutputBufferIndex);
+			return GetOutputFormatInfoFromFormat(Format);
+		}
+		return null;
+	}
+
+
+	public FOutputBufferInfo DequeueOutputBuffer(int InTimeoutUsec)
+	{
+		if (bIsInitialized)
+		{
+			FOutputBufferInfo OutputBufferInfo = new FOutputBufferInfo();
+			MediaCodec.BufferInfo Info = new MediaCodec.BufferInfo();
+			try
+			{
+				OutputBufferInfo.BufferIndex = DecoderHandle.dequeueOutputBuffer(Info, InTimeoutUsec);
+				// If there is output we need to remember the buffer metadata for later.
+				if (OutputBufferInfo.BufferIndex >= 0)
+				{
+					OutputBufferInfo.PresentationTimestamp = Info.presentationTimeUs;
+					OutputBufferInfo.Size = Info.size;
+					OutputBufferInfo.bIsEOS = (Info.flags & MediaCodec.BUFFER_FLAG_END_OF_STREAM) != 0;
+					OutputBufferInfo.bIsConfig = (Info.flags & MediaCodec.BUFFER_FLAG_CODEC_CONFIG) != 0;
+					OutputBufferInfos.put(OutputBufferInfo.BufferIndex, Info);
+				}
+				else if (OutputBufferInfo.BufferIndex == MediaCodec.INFO_TRY_AGAIN_LATER)				// -1
+				{
+					// Have no output yet.
+				}
+				else if (OutputBufferInfo.BufferIndex == MediaCodec.INFO_OUTPUT_FORMAT_CHANGED)			// -2
+				{
+					MediaFormat Format = DecoderHandle.getOutputFormat();
+					CurrentOutputFormatInfo = GetOutputFormatInfoFromFormat(Format);
+				}
+				else if (OutputBufferInfo.BufferIndex == MediaCodec.INFO_OUTPUT_BUFFERS_CHANGED)		// -3
+				{
+					// This is deprecated with API level 21 and of no real value to us.
+				}
+			}
+			catch(Exception e)
+			{
+				Log.w(TAG,"ElectraVideoDecoderH264: Failed to dequeue output buffer");
+				e.printStackTrace();
+				return null;
+			}
+			return OutputBufferInfo;
+		}
+		return null;
+	}
+
+
+	public byte[] GetOutputBuffer(int InOutputBufferIndex)
+	{
+		byte[] Data = null;
+		if (bIsInitialized && InOutputBufferIndex >= 0)
+		{
+			MediaCodec.BufferInfo Info = OutputBufferInfos.get(InOutputBufferIndex);
+			if (Info != null)
+			{
+				try
+				{
+					ByteBuffer outBuffer = DecoderHandle.getOutputBuffer(InOutputBufferIndex);
+					Data = new byte[Info.size];
+					outBuffer.get(Data);
+				}
+				catch(Exception e)
+				{
+					Log.w(TAG,"ElectraVideoDecoderH264: Failed to get output buffer");
+					e.printStackTrace();
+					return null;
+				}
+			}
+		}
+		return Data;
+	}
+
+	public int ReleaseOutputBuffer(int InOutputBufferIndex, boolean bRender, long releaseAt)
+	{
+		if (bIsInitialized && InOutputBufferIndex >= 0)
+		{
+			MediaCodec.BufferInfo Info = OutputBufferInfos.get(InOutputBufferIndex);
+			if (Info != null)
+			{
+				OutputBufferInfos.delete(InOutputBufferIndex);
+
+				FOutputFormatInfo OutFmtInfo = GetOutputFormatInfo(InOutputBufferIndex);
+
+				// If the decoder has no output surface we cannot release the buffer to it!
+				bRender &= bDecoderHasSurface;
+
+				try
+				{
+					boolean bDoRender = bRender == true && (Info.size != 0) && ((Info.flags & MediaCodec.BUFFER_FLAG_CODEC_CONFIG) != MediaCodec.BUFFER_FLAG_CODEC_CONFIG);
+					if (releaseAt < 0 || bDoRender == false)
+					{
+						if (bDoRender && CreationParameters.bSurfaceIsView)
+						{
+							// We use a timed output instead of the un-timed one because that one seems to use PTS values gathered from the sample(??) that of course do
+							// not correclate to "realtime within 1s" - and trigger warnings in logcat...
+							DecoderHandle.releaseOutputBuffer(InOutputBufferIndex, System.nanoTime());
+						}
+						else
+						{
+							DecoderHandle.releaseOutputBuffer(InOutputBufferIndex, bDoRender);
+						}
+					}
+					else
+					{
+						DecoderHandle.releaseOutputBuffer(InOutputBufferIndex, releaseAt);
+					}
+				}
+				catch(Exception e)
+				{
+					Log.w(TAG,"ElectraVideoDecoderH264: Failed to release buffer to surface");
+					e.printStackTrace();
+					return 1;
+				}
+				return 0;
+			}
+			else
+			{
+				return -10000;
+			}
+		}
+		return 1;
+	}
+
+
+}
