ElectraPlayer: RELEASE. Moved from NFL to public. Note that this is done via a Delete and "Add New" operation as discussed.

diff --git a/Engine/Plugins/Media/ElectraPlayer/Source/ElectraPlayerRuntime/Private/Runtime/Decoder/Android/VideoDecoderH264_Android.cpp b/Engine/Plugins/Media/ElectraPlayer/Source/ElectraPlayerRuntime/Private/Runtime/Decoder/Android/VideoDecoderH264_Android.cpp
new file mode 100644
--- /dev/null
+++ b/Engine/Plugins/Media/ElectraPlayer/Source/ElectraPlayerRuntime/Private/Runtime/Decoder/Android/VideoDecoderH264_Android.cpp
@@ -0,0 +1,1855 @@
+// Copyright Epic Games, Inc. All Rights Reserved.
+
+#include "PlayerCore.h"
+
+#include "StreamAccessUnitBuffer.h"
+#include "Decoder/VideoDecoderH264.h"
+#include "Renderer/RendererBase.h"
+#include "Player/PlayerSessionServices.h"
+#include "Utilities/StringHelpers.h"
+#include "Utilities/UtilsMPEGVideo.h"
+#include "DecoderErrors_Android.h"
+#include "HAL/LowLevelMemTracker.h"
+#include "Android/AndroidPlatformMisc.h"
+
+#include "VideoDecoderH264_JavaWrapper_Android.h"
+#include "MediaVideoDecoderOutputAndroid.h"
+#include "Renderer/RendererVideo.h"
+
+#include "Android/AndroidPlatform.h"
+#include "Android/AndroidJava.h"
+
+namespace Electra
+{
+
+/**
+ * H264 video decoder class implementation.
+**/
+class FVideoDecoderH264 : public IVideoDecoderH264, public FMediaThread
+{
+public:
+	static bool Startup(const IVideoDecoderH264::FSystemConfiguration& InConfig);
+	static void Shutdown();
+
+	FVideoDecoderH264();
+	virtual ~FVideoDecoderH264();
+
+	virtual void SetPlayerSessionServices(IPlayerSessionServices* SessionServices) override;
+
+	virtual void Open(const FInstanceConfiguration& InConfig) override;
+	virtual void Close() override;
+
+	virtual void SetMaximumDecodeCapability(int32 MaxWidth, int32 MaxHeight, int32 MaxProfile, int32 MaxProfileLevel, const FParamDict& AdditionalOptions) override;
+
+	virtual void SetAUInputBufferListener(IAccessUnitBufferListener* Listener) override;
+
+	virtual void SetReadyBufferListener(IDecoderOutputBufferListener* Listener) override;
+
+	virtual void SetRenderer(TSharedPtr<IMediaRenderer, ESPMode::ThreadSafe> InRenderer) override;
+
+	virtual void SetResourceDelegate(const TSharedPtr<IVideoDecoderResourceDelegate, ESPMode::ThreadSafe>& ResourceDelegate) override;
+
+	virtual IAccessUnitBufferInterface::EAUpushResult AUdataPushAU(FAccessUnit* AccessUnit) override;
+	virtual void AUdataPushEOD() override;
+	virtual void AUdataFlushEverything() override;
+
+	virtual void Android_UpdateSurface(const TSharedPtr<IOptionPointerValueContainer>& Surface) override;
+
+	static void JavaCallback_NewDataAvailable(uint32 NativeDecoderID);
+
+	static void ReleaseToSurface(uint32 NativeDecoderID, const FDecoderTimeStamp & Time);
+
+private:
+	bool InternalDecoderCreate();
+	void InternalDecoderDestroy();
+	void RecreateDecoderSession();
+	void StartThread();
+	void StopThread();
+	void WorkerThread();
+	void RenderThreadFN();
+
+	bool CreateDecodedImagePool();
+	void DestroyDecodedImagePool();
+
+	void NotifyReadyBufferListener(bool bHaveOutput);
+
+	bool AcquireOutputBuffer(IMediaRenderer::IBuffer*& RenderOutputBuffer);
+
+	void PrepareAU(FAccessUnit* AccessUnit, bool& bOutIsIDR, bool& bOutIsDiscardable);
+
+	void ProcessReadyOutputBuffersToSurface();
+
+private:
+	bool DrainDecoder();
+
+	enum class EDecodeResult
+	{
+		Ok,
+		Fail,
+		SessionLost,
+		TryAgainLater,
+	};
+
+	EDecodeResult Decode();
+	bool DecodeDummy();
+
+	enum class EOutputResult
+	{
+		Ok,
+		Fail,
+		TryAgainLater,
+		EOS
+	};
+	EOutputResult GetOutput();
+
+	struct FDecodedImage;
+	bool ProcessOutput(const FDecodedImage & NextImage);
+
+
+	void PostError(int32_t ApiReturnValue, const FString& Message, uint16 Code, UEMediaError Error = UEMEDIA_ERROR_OK);
+	void LogMessage(IInfoLog::ELevel Level, const FString& Message);
+
+private:
+
+	struct FDecoderFormatInfo
+	{
+		void Reset()
+		{
+			CurrentCodecData.Reset();
+		}
+		bool IsDifferentFrom(const FAccessUnit* InputAccessUnit);
+
+		TSharedPtrTS<const FAccessUnit::CodecData>		CurrentCodecData;
+	};
+
+	struct FInDecoderInfo : public TSharedFromThis<FInDecoderInfo, ESPMode::ThreadSafe>
+	{
+		FStreamCodecInformation		ParsedInfo;
+		FTimeValue					PTS;
+		FTimeValue					DTS;
+		FTimeValue					Duration;
+		int64						DecoderPTS;
+	};
+
+	struct FDecodedImage
+	{
+		TSharedPtr<FInDecoderInfo, ESPMode::ThreadSafe>	SourceInfo;
+		IAndroidJavaH264VideoDecoder::FOutputFormatInfo	OutputFormat;
+		IAndroidJavaH264VideoDecoder::FOutputBufferInfo	OutputBufferInfo;
+		bool											bIsDummy;
+	};
+
+
+	enum class EDecoderState
+	{
+		Ready,
+		Create,
+		Flush
+	};
+	enum class EInputState
+	{
+		Idle,
+		FeedAU,
+		FeedEOS,
+		AwaitingEOS
+	};
+	enum class EDummyState
+	{
+		NotDummy,
+		Entering,
+		InDummy
+	};
+
+
+	FInstanceConfiguration															Config;
+
+	FMediaEvent																		TerminateThreadSignal;
+	FMediaEvent																		FlushDecoderSignal;
+	FMediaEvent																		DecoderFlushedSignal;
+	bool																			bThreadStarted;
+
+	IPlayerSessionServices*															PlayerSessionServices;
+
+	TSharedPtr<IMediaRenderer, ESPMode::ThreadSafe>									Renderer;
+
+	TWeakPtr<IVideoDecoderResourceDelegate, ESPMode::ThreadSafe>					ResourceDelegate;
+
+	FAccessUnitBuffer																AccessUnitBuffer;
+	FAccessUnitBuffer																ReplayAccessUnitBuffer;
+
+	FMediaCriticalSection															ListenerMutex;
+	IAccessUnitBufferListener*														InputBufferListener;
+	IDecoderOutputBufferListener*													ReadyBufferListener;
+
+	bool																			bSurfaceIsView;
+
+	FDecoderFormatInfo																CurrentStreamFormatInfo;
+	TSharedPtr<IAndroidJavaH264VideoDecoder, ESPMode::ThreadSafe>					DecoderInstance;
+	IAndroidJavaH264VideoDecoder::FDecoderInformation								DecoderInfo;
+	FAccessUnit* 																	CurrentAccessUnit;
+	bool																			bMustSendCSD;
+	int64																			LastSentPTS;
+	TArray<TSharedPtr<FInDecoderInfo, ESPMode::ThreadSafe>>							InDecoderInfos;
+
+	EDecoderState																	DecoderState;
+	EInputState																		InputState;
+	EDummyState																		DummyState;
+
+	int32																			MaxDecodeBufferSize;
+	bool																			bError;
+
+	FMediaEvent																		NewDataAvailable;
+
+	struct FOutputBufferInfo
+	{
+		FOutputBufferInfo() : BufferIndex(-1) {}
+		FOutputBufferInfo(const FDecoderTimeStamp& InTimeStamp, int32 InBufferIndex, int32 InValidCount) : Timestamp(InTimeStamp), BufferIndex(InBufferIndex), ValidCount(InValidCount) {}
+
+		FDecoderTimeStamp Timestamp;
+		int32 BufferIndex;
+		int32 ValidCount;
+	};
+
+	TArray<FOutputBufferInfo>														ReadyOutputBuffersToSurface;
+	FCriticalSection																OutputSurfaceTargetCS;
+	FDecoderTimeStamp																OutputSurfaceTargetPTS;
+
+	uint32																			NativeDecoderID;
+
+	static uint32																	NextNativeDecoderID;
+	static FCriticalSection															NativeDecoderMapCS;
+	static TMap<uint32, FVideoDecoderH264*>											NativeDecoderMap;
+
+public:
+	static FSystemConfiguration														SystemConfig;
+};
+
+IVideoDecoderH264::FSystemConfiguration	FVideoDecoderH264::SystemConfig;
+
+uint32									FVideoDecoderH264::NextNativeDecoderID = 0;
+FCriticalSection						FVideoDecoderH264::NativeDecoderMapCS;
+TMap<uint32, FVideoDecoderH264*>		FVideoDecoderH264::NativeDecoderMap;
+
+/***************************************************************************************************************************************************/
+/***************************************************************************************************************************************************/
+/***************************************************************************************************************************************************/
+
+bool IVideoDecoderH264::Startup(const IVideoDecoderH264::FSystemConfiguration& InConfig)
+{
+	return FVideoDecoderH264::Startup(InConfig);
+}
+
+void IVideoDecoderH264::Shutdown()
+{
+	FVideoDecoderH264::Shutdown();
+}
+
+bool IVideoDecoderH264::GetStreamDecodeCapability(FStreamDecodeCapability& OutResult, const FStreamDecodeCapability& InStreamParameter)
+{
+	return false;
+}
+
+IVideoDecoderH264::FSystemConfiguration::FSystemConfiguration()
+{
+	ThreadConfig.Decoder.Priority 	= TPri_Normal;
+	ThreadConfig.Decoder.StackSize	= 64 << 10;
+	ThreadConfig.Decoder.CoreAffinity = -1;
+	ThreadConfig.PassOn.Priority  	= TPri_Normal;
+	ThreadConfig.PassOn.StackSize 	= 64 << 10;
+	ThreadConfig.PassOn.CoreAffinity  = -1;
+}
+
+IVideoDecoderH264::FInstanceConfiguration::FInstanceConfiguration()
+	: ThreadConfig(FVideoDecoderH264::SystemConfig.ThreadConfig)
+	, MaxDecodedFrames(8)
+{
+}
+
+IVideoDecoderH264* IVideoDecoderH264::Create()
+{
+	return new FVideoDecoderH264;
+}
+
+
+/***************************************************************************************************************************************************/
+/***************************************************************************************************************************************************/
+/***************************************************************************************************************************************************/
+
+class FElectraPlayerVideoDecoderOutputAndroidImpl : public FVideoDecoderOutputAndroid
+{
+public:
+	FElectraPlayerVideoDecoderOutputAndroidImpl()
+	{
+	}
+
+	~FElectraPlayerVideoDecoderOutputAndroidImpl()
+	{
+	}
+
+	void Initialize(EOutputType InOutputType, int32 InBufferIndex, int32 InValidCount, uint32 InNativeDecoderID, FParamDict* InParamDict)
+	{
+		FVideoDecoderOutputAndroid::Initialize(InParamDict);
+
+		OutputType = InOutputType;
+
+		BufferIndex = InBufferIndex;
+		ValidCount = InValidCount;
+		NativeDecoderID = InNativeDecoderID;
+	}
+
+	void ReleaseToSurface() const override
+	{
+		FVideoDecoderH264::ReleaseToSurface(NativeDecoderID, GetTime());
+	}
+
+	void SetOwner(const TSharedPtr<IDecoderOutputOwner, ESPMode::ThreadSafe>& InOwningRenderer) override
+	{
+		OwningRenderer = InOwningRenderer;
+	}
+
+	void ShutdownPoolable() override
+	{
+		TSharedPtr<IDecoderOutputOwner, ESPMode::ThreadSafe> lockedVideoRenderer = OwningRenderer.Pin();
+		if (lockedVideoRenderer.IsValid())
+		{
+			lockedVideoRenderer->SampleReleasedToPool(GetDuration());
+		}
+	}
+
+	virtual EOutputType GetOutputType() const override
+	{
+		return OutputType;
+	}
+
+private:
+	// Decoder output type
+	EOutputType OutputType;
+
+	// Decoder output buffer index associated with this sample (Surface type)
+	int32 BufferIndex;
+
+	// Valid count to check for decoder Android changes within a single logical decoder (Surface type)
+	int32 ValidCount;
+
+	// ID of native decoder
+	uint32 NativeDecoderID;
+
+	// We hold a weak reference to the video renderer. During destruction the video renderer could be destroyed while samples are still out there..
+	TWeakPtr<IDecoderOutputOwner, ESPMode::ThreadSafe> OwningRenderer;
+};
+
+/***************************************************************************************************************************************************/
+/***************************************************************************************************************************************************/
+/***************************************************************************************************************************************************/
+//-----------------------------------------------------------------------------
+/**
+ * Decoder system startup
+ *
+ * @param config
+ *
+ * @return
+ */
+bool FVideoDecoderH264::Startup(const IVideoDecoderH264::FSystemConfiguration& InConfig)
+{
+	SystemConfig = InConfig;
+	return true;
+}
+
+
+//-----------------------------------------------------------------------------
+/**
+ * Decoder system shutdown.
+ */
+void FVideoDecoderH264::Shutdown()
+{
+}
+
+
+//-----------------------------------------------------------------------------
+/**
+ * Constructor
+ */
+FVideoDecoderH264::FVideoDecoderH264()
+	: FMediaThread("ElectraPlayer::H264 decoder")
+	, bThreadStarted(false)
+	, PlayerSessionServices(nullptr)
+	, Renderer(nullptr)
+	, InputBufferListener(nullptr)
+	, ReadyBufferListener(nullptr)
+	, MaxDecodeBufferSize(0)
+	, bError(false)
+	, NativeDecoderID(++NextNativeDecoderID)
+{
+	FScopeLock Lock(&NativeDecoderMapCS);
+	NativeDecoderMap.Add(NativeDecoderID, this);
+}
+
+
+//-----------------------------------------------------------------------------
+/**
+ * Destructor
+ */
+FVideoDecoderH264::~FVideoDecoderH264()
+{
+	Close();
+	FScopeLock Lock(&NativeDecoderMapCS);
+	NativeDecoderMap.Remove(NativeDecoderID);
+}
+
+
+//-----------------------------------------------------------------------------
+/**
+ * Sets an AU input buffer listener.
+ *
+ * @param Listener
+ */
+void FVideoDecoderH264::SetAUInputBufferListener(IAccessUnitBufferListener* Listener)
+{
+	FMediaCriticalSection::ScopedLock lock(ListenerMutex);
+	InputBufferListener = Listener;
+}
+
+
+//-----------------------------------------------------------------------------
+/**
+ * Sets a buffer-ready listener.
+ *
+ * @param Listener
+ */
+void FVideoDecoderH264::SetReadyBufferListener(IDecoderOutputBufferListener* Listener)
+{
+	FMediaCriticalSection::ScopedLock lock(ListenerMutex);
+	ReadyBufferListener = Listener;
+}
+
+
+//-----------------------------------------------------------------------------
+/**
+ * Sets the owning player's session service interface.
+ *
+ * @param InSessionServices
+ *
+ * @return
+ */
+void FVideoDecoderH264::SetPlayerSessionServices(IPlayerSessionServices* SessionServices)
+{
+	PlayerSessionServices = SessionServices;
+}
+
+
+//-----------------------------------------------------------------------------
+/**
+ * Opens a decoder instance
+ *
+ * @param config
+ *
+ * @return
+ */
+void FVideoDecoderH264::Open(const IVideoDecoderH264::FInstanceConfiguration& InConfig)
+{
+	Config = InConfig;
+
+	// Set a large enough size to hold a single access unit. Since we will be asking for a new AU on
+	// demand there is no need to overly restrict ourselves here. But it must be large enough to
+	// hold at least the largest expected access unit.
+	AccessUnitBuffer.CapacitySet(FAccessUnitBuffer::FConfiguration(16 << 20, 60.0));
+
+	StartThread();
+}
+
+
+//-----------------------------------------------------------------------------
+/**
+ * Closes the decoder instance.
+ */
+void FVideoDecoderH264::Close()
+{
+	StopThread();
+}
+
+
+//-----------------------------------------------------------------------------
+/**
+ * Sets a new decoder limit.
+ * As soon as a new sequence starts matching this limit the decoder will be
+ * destroyed and recreated to conserve memory.
+ * The assumption is that the data being streamed will not exceed this limit,
+ * but if it does, any access unit requiring a better decoder will take
+ * precedence and force a decoder capable of decoding it!
+ *
+ * @param MaxWidth
+ * @param MaxHeight
+ * @param MaxProfile
+ * @param MaxProfileLevel
+ * @param AdditionalOptions
+ */
+void FVideoDecoderH264::SetMaximumDecodeCapability(int32 MaxWidth, int32 MaxHeight, int32 MaxProfile, int32 MaxProfileLevel, const FParamDict& AdditionalOptions)
+{
+	// Not implemented
+}
+
+
+//-----------------------------------------------------------------------------
+/**
+ * Sets a new renderer.
+ *
+ * @param InRenderer
+ */
+void FVideoDecoderH264::SetRenderer(TSharedPtr<IMediaRenderer, ESPMode::ThreadSafe> InRenderer)
+{
+	Renderer = InRenderer;
+}
+
+//-----------------------------------------------------------------------------
+/**
+*/
+void FVideoDecoderH264::SetResourceDelegate(const TSharedPtr<IVideoDecoderResourceDelegate, ESPMode::ThreadSafe>& InResourceDelegate)
+{
+	ResourceDelegate = InResourceDelegate;
+}
+
+//-----------------------------------------------------------------------------
+/**
+ * Creates and runs the decoder thread.
+ */
+void FVideoDecoderH264::StartThread()
+{
+	ThreadSetPriority(Config.ThreadConfig.Decoder.Priority);
+	ThreadSetCoreAffinity(Config.ThreadConfig.Decoder.CoreAffinity);
+	ThreadSetStackSize(Config.ThreadConfig.Decoder.StackSize);
+	ThreadStart(Electra::MakeDelegate(this, &FVideoDecoderH264::WorkerThread));
+	bThreadStarted = true;
+}
+
+
+//-----------------------------------------------------------------------------
+/**
+ * Stops the decoder thread.
+ */
+void FVideoDecoderH264::StopThread()
+{
+	if (bThreadStarted)
+	{
+		TerminateThreadSignal.Signal();
+		ThreadWaitDone();
+		bThreadStarted = false;
+	}
+}
+
+
+//-----------------------------------------------------------------------------
+/**
+ * Posts an error to the session service error listeners.
+ *
+ * @param ApiReturnValue
+ * @param Message
+ * @param Code
+ * @param Error
+ */
+void FVideoDecoderH264::PostError(int32_t ApiReturnValue, const FString& Message, uint16 Code, UEMediaError Error)
+{
+	check(PlayerSessionServices);
+	if (PlayerSessionServices)
+	{
+		FErrorDetail err;
+		err.SetError(Error != UEMEDIA_ERROR_OK ? Error : UEMEDIA_ERROR_DETAIL);
+		err.SetFacility(Facility::EFacility::H264Decoder);
+		err.SetCode(Code);
+		err.SetMessage(Message);
+		err.SetPlatformMessage(FString::Printf(TEXT("%d (0x%08x)"), ApiReturnValue, ApiReturnValue));
+		PlayerSessionServices->PostError(err);
+	}
+	bError = true;
+}
+
+
+//-----------------------------------------------------------------------------
+/**
+ * Sends a log message to the session service log.
+ *
+ * @param Level
+ * @param Message
+ */
+void FVideoDecoderH264::LogMessage(IInfoLog::ELevel Level, const FString& Message)
+{
+	if (PlayerSessionServices)
+	{
+		PlayerSessionServices->PostLog(Facility::EFacility::H264Decoder, Level, Message);
+	}
+}
+
+
+//-----------------------------------------------------------------------------
+/**
+ * Create a pool of decoded images for the decoder.
+ *
+ * @return
+ */
+bool FVideoDecoderH264::CreateDecodedImagePool()
+{
+	check(Renderer);
+	FParamDict poolOpts;
+
+	poolOpts.Set("num_buffers", FVariantValue((int64) Config.MaxDecodedFrames));
+
+	UEMediaError Error = Renderer->CreateBufferPool(poolOpts);
+	check(Error == UEMEDIA_ERROR_OK);
+
+	MaxDecodeBufferSize = (int32) Renderer->GetBufferPoolProperties().GetValue("max_buffers").GetInt64();
+
+	if (Error != UEMEDIA_ERROR_OK)
+	{
+		PostError(0, "Failed to create image pool", ERRCODE_INTERNAL_ANDROID_COULD_NOT_CREATE_IMAGE_POOL, Error);
+	}
+
+	return Error == UEMEDIA_ERROR_OK;
+}
+
+
+//-----------------------------------------------------------------------------
+/**
+ * Destroys the pool of decoded images.
+ */
+void FVideoDecoderH264::DestroyDecodedImagePool()
+{
+	Renderer->ReleaseBufferPool();
+}
+
+
+//-----------------------------------------------------------------------------
+/**
+ * Called to receive a new input access unit for decoding.
+ *
+ * @param AccessUnit
+ */
+IAccessUnitBufferInterface::EAUpushResult FVideoDecoderH264::AUdataPushAU(FAccessUnit* AccessUnit)
+{
+	AccessUnit->AddRef();
+	bool bOk = AccessUnitBuffer.Push(AccessUnit);
+	if (!bOk)
+	{
+		FAccessUnit::Release(AccessUnit);
+	}
+
+	return bOk ? IAccessUnitBufferInterface::EAUpushResult::Ok : IAccessUnitBufferInterface::EAUpushResult::Full;
+}
+
+
+//-----------------------------------------------------------------------------
+/**
+ * "Pushes" an End Of Data marker indicating no further access units will be added.
+ */
+void FVideoDecoderH264::AUdataPushEOD()
+{
+	AccessUnitBuffer.PushEndOfData();
+}
+
+
+//-----------------------------------------------------------------------------
+/**
+ * Flushes the decoder and clears the input access unit buffer.
+ */
+void FVideoDecoderH264::AUdataFlushEverything()
+{
+	FlushDecoderSignal.Signal();
+	DecoderFlushedSignal.WaitAndReset();
+}
+
+
+//-----------------------------------------------------------------------------
+/**
+ * Create a decoder instance.
+ *
+ * @return true if successful, false on error
+ */
+bool FVideoDecoderH264::InternalDecoderCreate()
+{
+	int32 result;
+
+	DecoderInfo = {};
+	InDecoderInfos.Empty();
+
+	IAndroidJavaH264VideoDecoder::FCreateParameters cp;
+	cp.CodecData					= CurrentAccessUnit->AUCodecData;
+	cp.MaxWidth 					= Config.MaxFrameWidth;
+	cp.MaxHeight					= Config.MaxFrameHeight;
+	cp.MaxProfile   				= Config.ProfileIdc;
+	cp.MaxProfileLevel  			= Config.LevelIdc;
+	cp.MaxFrameRate 				= 60;
+	cp.NativeDecoderID				= NativeDecoderID;
+
+	// See if we should decode directly to a externally provided surface or not...
+
+	cp.VideoCodecSurface = nullptr;
+	if (Config.AdditionalOptions.HaveKey("videoDecoder_Android_UseSurface"))
+	{
+		cp.bUseVideoCodecSurface = Config.AdditionalOptions.GetValue("videoDecoder_Android_UseSurface").GetBool();
+		if (cp.bUseVideoCodecSurface && Config.AdditionalOptions.HaveKey("videoDecoder_Android_Surface"))
+		{
+			TSharedPtr<IOptionPointerValueContainer> Value = Config.AdditionalOptions.GetValue("videoDecoder_Android_Surface").GetSharedPointer<IOptionPointerValueContainer>();
+			cp.VideoCodecSurface = AndroidJavaEnv::GetJavaEnv()->NewLocalRef(reinterpret_cast<jweak>(Value->GetPointer()));
+			cp.bSurfaceIsView = true;
+		}
+	}
+	if (!cp.bSurfaceIsView)
+	{
+		if (TSharedPtr<IVideoDecoderResourceDelegate, ESPMode::ThreadSafe> PinnedResourceDelegate = ResourceDelegate.Pin())
+		{
+			cp.VideoCodecSurface = PinnedResourceDelegate->VideoDecoderResourceDelegate_GetCodecSurface();
+		}
+	}
+	if (!cp.VideoCodecSurface)
+	{
+		PostError(0, "No surface to create decoder with", ERRCODE_INTERNAL_ANDROID_COULD_NOT_CREATE_VIDEO_DECODER);
+		return false;
+	}
+
+	// Recall if we render to a view or off-screen here, too
+	bSurfaceIsView = cp.bSurfaceIsView;
+
+	// Check if there is an existing decoder instance we can re purpose
+	if (DecoderInstance.IsValid())
+	{
+		DecoderInstance->Flush();
+		DecoderInstance->Stop();
+		DecoderInstance->ReleaseDecoder();
+// FIXME: this probably needs a check if the existing renderer is still compatible (why wouldn't it be?)
+		cp.bRetainRenderer = true;
+	}
+	else
+	{
+		DecoderInstance = IAndroidJavaH264VideoDecoder::Create(PlayerSessionServices);
+		cp.bRetainRenderer = false;
+	}
+
+	result = DecoderInstance->InitializeDecoder(cp);
+	if (result)
+	{
+		PostError(result, "Failed to create decoder", ERRCODE_INTERNAL_ANDROID_COULD_NOT_CREATE_VIDEO_DECODER);
+		return false;
+	}
+	// Get the decoder information.
+	const IAndroidJavaH264VideoDecoder::FDecoderInformation* DecInf = DecoderInstance->GetDecoderInformation();
+	check(DecInf);
+	if (DecInf)
+	{
+		DecoderInfo = *DecInf;
+	}
+
+	// Start it.
+	result = DecoderInstance->Start();
+	if (result)
+	{
+		PostError(result, "Failed to start decoder", ERRCODE_INTERNAL_ANDROID_COULD_NOT_START_DECODER);
+		return false;
+	}
+	return true;
+}
+
+
+//-----------------------------------------------------------------------------
+/**
+ * Destroys the current decoder instance.
+ */
+void FVideoDecoderH264::InternalDecoderDestroy()
+{
+	if (DecoderInstance.IsValid())
+	{
+		DecoderInstance->Flush();
+		DecoderInstance->Stop();
+		DecoderInstance->ReleaseDecoder();
+		DecoderInstance.Reset();
+	}
+	DecoderInfo = {};
+	InDecoderInfos.Empty();
+}
+
+
+//-----------------------------------------------------------------------------
+/**
+ * Creates a new decoder and runs all AUs since the last IDR frame through
+ * without producing decoded images.
+ * Used to continue decoding after the application has been backgrounded
+ * which results in a decoder session loss when resumed.
+ * Since decoding cannot resume at arbitrary points in the stream everything
+ * from the last IDR frame needs to be decoded again.
+ * To speed this up AUs that have no dependencies are not added to the replay data.
+ */
+void FVideoDecoderH264::RecreateDecoderSession()
+{
+	check(!"TODO");
+}
+
+
+//-----------------------------------------------------------------------------
+/**
+ * Notifies the buffer-ready listener that we will now be producing output data
+ * and need an output buffer.
+ * This call is not intended to create or obtain an output buffer. It is merely
+ * indicating that new output will be produced.
+ *
+ * @param bHaveOutput
+ */
+void FVideoDecoderH264::NotifyReadyBufferListener(bool bHaveOutput)
+{
+	if (ReadyBufferListener)
+	{
+		IDecoderOutputBufferListener::FDecodeReadyStats stats;
+		stats.MaxDecodedElementsReady = MaxDecodeBufferSize;
+		stats.NumElementsInDecoder    = InDecoderInfos.Num();
+		stats.bOutputStalled		  = !bHaveOutput;
+		stats.bEODreached   		  = AccessUnitBuffer.IsEndOfData() && stats.NumDecodedElementsReady == 0 && stats.NumElementsInDecoder == 0;
+		ListenerMutex.Lock();
+		if (ReadyBufferListener)
+		{
+			ReadyBufferListener->DecoderOutputReady(stats);
+		}
+		ListenerMutex.Unlock();
+	}
+}
+
+
+//-----------------------------------------------------------------------------
+/**
+ * Prepares the passed access unit to be sent to the decoder.
+ *
+ * @param AccessUnit
+ * @param bOutIsIDR
+ * @param bOutIsDiscardable
+ */
+void FVideoDecoderH264::PrepareAU(FAccessUnit* AccessUnit, bool& bOutIsIDR, bool& bOutIsDiscardable)
+{
+	bOutIsIDR = false;
+	bOutIsDiscardable = true;
+	if (!AccessUnit->bHasBeenPrepared)
+	{
+		AccessUnit->bHasBeenPrepared = true;
+
+		// Process NALUs
+		uint32* CurrentNALU = (uint32 *)AccessUnit->AUData;
+		uint32* LastNALU    = (uint32 *)Electra::AdvancePointer(CurrentNALU, AccessUnit->AUSize);
+		while(CurrentNALU < LastNALU)
+		{
+			// Check the nal_ref_idc in the NAL unit for dependencies.
+			uint8 nal = *(const uint8 *)(CurrentNALU + 1);
+			check((nal & 0x80) == 0);
+			if ((nal >> 5) != 0)
+			{
+				bOutIsDiscardable = false;
+			}
+			// IDR frame?
+			if ((nal & 0x1f) == 5)
+			{
+				bOutIsIDR = true;
+			}
+
+			// SEI message(s)?
+			if ((nal & 0x1f) == 6)
+			{
+				// TODO: we might need to set aside any SEI messages carrying 608 or 708 caption data.
+			}
+
+// TODO: There are decoder implementations out there that do not like NALUs before the SPS. We may need to filter those out.
+
+// TODO: SPS and PPS may need to be removed for some decoders and only sent in a dedicated CSD buffer.
+
+			uint32 naluLen = MEDIA_FROM_BIG_ENDIAN(*CurrentNALU) + 4;
+			*CurrentNALU = MEDIA_TO_BIG_ENDIAN(0x00000001U);
+			CurrentNALU = Electra::AdvancePointer(CurrentNALU, naluLen);
+		}
+	}
+}
+
+
+//-----------------------------------------------------------------------------
+/**
+ * Gets an output buffer from the video renderer.
+ *
+ * @param RenderOutputBuffer
+ *
+ * @return true if successful, false on error
+ */
+bool FVideoDecoderH264::AcquireOutputBuffer(IMediaRenderer::IBuffer*& RenderOutputBuffer)
+{
+	RenderOutputBuffer = nullptr;
+	FParamDict BufferAcquireOptions;
+	while(!TerminateThreadSignal.IsSignaled() && !FlushDecoderSignal.IsSignaled())
+	{
+		UEMediaError bufResult = Renderer->AcquireBuffer(RenderOutputBuffer, 0, BufferAcquireOptions);
+		check(bufResult == UEMEDIA_ERROR_OK || bufResult == UEMEDIA_ERROR_INSUFFICIENT_DATA);
+		if (bufResult != UEMEDIA_ERROR_OK && bufResult != UEMEDIA_ERROR_INSUFFICIENT_DATA)
+		{
+			PostError(0, "Failed to acquire sample buffer", ERRCODE_INTERNAL_ANDROID_COULD_NOT_GET_OUTPUT_BUFFER, bufResult);
+			return false;
+		}
+		bool bHaveAvailSmpBlk = RenderOutputBuffer != nullptr;
+		NotifyReadyBufferListener(bHaveAvailSmpBlk);
+		if (bHaveAvailSmpBlk)
+		{
+			break;
+		}
+		else
+		{
+			// No available buffer. Sleep for a bit. Can't sleep on a signal since we have to check two: abort and flush
+			FMediaRunnable::SleepMilliseconds(5);
+		}
+
+		ProcessReadyOutputBuffersToSurface();
+	}
+	return true;
+}
+
+
+//-----------------------------------------------------------------------------
+/**
+ * Creates a dummy output image that is not to be displayed and has no image data.
+ * Dummy access units are created when stream data is missing to ensure the data
+ * pipeline does not run dry and exhibits no gaps in the timeline.
+ *
+ * @return true if successful, false otherwise
+ */
+bool FVideoDecoderH264::DecodeDummy()
+{
+	const int32 kTotalSleepTimeMsec = 5;
+
+	check(CurrentAccessUnit->Duration.IsValid());
+
+	while(!TerminateThreadSignal.IsSignaled() && !FlushDecoderSignal.IsSignaled())
+	{
+		FDecodedImage NextImage;
+		NextImage.bIsDummy  	 = true;
+		NextImage.SourceInfo = TSharedPtr<FInDecoderInfo, ESPMode::ThreadSafe>(new FInDecoderInfo);
+		NextImage.SourceInfo->DTS = CurrentAccessUnit->DTS;
+		NextImage.SourceInfo->PTS = CurrentAccessUnit->PTS;
+		NextImage.SourceInfo->Duration = CurrentAccessUnit->Duration;
+		if (CurrentAccessUnit->AUCodecData.IsValid())
+		{
+			NextImage.SourceInfo->ParsedInfo = CurrentAccessUnit->AUCodecData->ParsedInfo;
+		}
+		if (ProcessOutput(NextImage))
+		{
+			break;
+		}
+		FMediaRunnable::SleepMilliseconds(kTotalSleepTimeMsec);
+
+		ProcessReadyOutputBuffersToSurface();
+	}
+	FAccessUnit::Release(CurrentAccessUnit);
+	CurrentAccessUnit = nullptr;
+	return true;
+}
+
+
+//-----------------------------------------------------------------------------
+/**
+ * Checks if the codec specific data has changed.
+ *
+ * @return false if the format is still the same, true if it has changed.
+ */
+bool FVideoDecoderH264::FDecoderFormatInfo::IsDifferentFrom(const FAccessUnit* InputAccessUnit)
+{
+	if (InputAccessUnit->AUCodecData.IsValid() && InputAccessUnit->AUCodecData.Get() != CurrentCodecData.Get())
+	{
+		// Address of codec data is different. Are the contents too?
+		if (!CurrentCodecData.IsValid() ||
+			CurrentCodecData->CodecSpecificData.Num() != InputAccessUnit->AUCodecData->CodecSpecificData.Num() ||
+			memcmp(CurrentCodecData->CodecSpecificData.GetData(), InputAccessUnit->AUCodecData->CodecSpecificData.GetData(), CurrentCodecData->CodecSpecificData.Num()))
+		{
+			CurrentCodecData = InputAccessUnit->AUCodecData;
+			return true;
+		}
+	}
+	return false;
+}
+
+
+//-----------------------------------------------------------------------------
+/**
+ * Sends an access unit to the decoder for decoding.
+ *
+ * @return
+ */
+FVideoDecoderH264::EDecodeResult FVideoDecoderH264::Decode()
+{
+	// No input AU to decode?
+	if (!CurrentAccessUnit)
+	{
+		return EDecodeResult::Ok;
+	}
+	// We gotta have a decoder here.
+	if (!DecoderInstance.IsValid())
+	{
+		return EDecodeResult::Fail;
+	}
+
+	check(CurrentAccessUnit->Duration.IsValid());
+	check(CurrentAccessUnit->PTS.IsValid());
+
+	int32 result = -1;
+	int32 InputBufferIndex = -1;
+
+	int64 pts = CurrentAccessUnit->PTS.GetAsMicroseconds();
+
+	// Do we need to send the CSD first?
+	if (bMustSendCSD)
+	{
+		InputBufferIndex = DecoderInstance->DequeueInputBuffer(0);
+		if (InputBufferIndex >= 0)
+		{
+			result = DecoderInstance->QueueCSDInputBuffer(InputBufferIndex, CurrentAccessUnit->AUCodecData->CodecSpecificData.GetData(), CurrentAccessUnit->AUCodecData->CodecSpecificData.Num(), pts);
+			check(result == 0);
+			if (result == 0)
+			{
+				bMustSendCSD = false;
+			}
+			else
+			{
+				PostError(result, "Failed to submit decoder CSD input buffer", ERRCODE_INTERNAL_ANDROID_FAILED_TO_DECODE_VIDEO);
+				return EDecodeResult::Fail;
+			}
+		}
+		else if (InputBufferIndex == -1)
+		{
+			// No available input buffer. Try later.
+			return EDecodeResult::TryAgainLater;
+		}
+		else
+		{
+			PostError(InputBufferIndex, "Failed to get a decoder input buffer for CSD", ERRCODE_INTERNAL_ANDROID_COULD_NOT_GET_INPUT_BUFFER);
+			return EDecodeResult::Fail;
+		}
+	}
+
+	// Send actual AU data now.
+	InputBufferIndex = DecoderInstance->DequeueInputBuffer(0);
+	if (InputBufferIndex >= 0)
+	{
+		result = DecoderInstance->QueueInputBuffer(InputBufferIndex, CurrentAccessUnit->AUData, CurrentAccessUnit->AUSize, pts);
+		check(result == 0);
+		if (result == 0)
+		{
+			LastSentPTS = pts;
+
+			TSharedPtr<FInDecoderInfo, ESPMode::ThreadSafe> SourceInfo = TSharedPtr<FInDecoderInfo, ESPMode::ThreadSafe>(new FInDecoderInfo);
+			SourceInfo->DecoderPTS = pts;
+			SourceInfo->DTS  	  = CurrentAccessUnit->DTS;
+			SourceInfo->PTS  	  = CurrentAccessUnit->PTS;
+			SourceInfo->Duration   = CurrentAccessUnit->Duration;
+			if (CurrentAccessUnit->AUCodecData.IsValid())
+			{
+				SourceInfo->ParsedInfo = CurrentAccessUnit->AUCodecData->ParsedInfo;
+			}
+			InDecoderInfos.Add(SourceInfo);
+			InDecoderInfos.Sort([](const TSharedPtr<FInDecoderInfo, ESPMode::ThreadSafe>& a, const TSharedPtr<FInDecoderInfo, ESPMode::ThreadSafe>& b) { return a->DecoderPTS < b->DecoderPTS; });
+
+			// Done with this AU now.
+			FAccessUnit::Release(CurrentAccessUnit);
+			CurrentAccessUnit = nullptr;
+			return EDecodeResult::Ok;
+		}
+		else
+		{
+			PostError(result, "Failed to submit decoder input buffer", ERRCODE_INTERNAL_ANDROID_FAILED_TO_DECODE_VIDEO);
+			return EDecodeResult::Fail;
+		}
+	}
+	else if (InputBufferIndex == -1)
+	{
+		// No available input buffer. Try later.
+		return EDecodeResult::TryAgainLater;
+	}
+	else
+	{
+		PostError(InputBufferIndex, "Failed to get a decoder input buffer", ERRCODE_INTERNAL_ANDROID_COULD_NOT_GET_INPUT_BUFFER);
+		return EDecodeResult::Fail;
+	}
+	return EDecodeResult::Fail;
+}
+
+
+//-----------------------------------------------------------------------------
+/**
+ * Drains the decoder by sending it an EOS buffer.
+ *
+ * @return true if successful, false otherwise.
+ */
+bool FVideoDecoderH264::DrainDecoder()
+{
+	if (!DecoderInstance.IsValid())
+	{
+		return true;
+	}
+
+	int32 result = -1;
+	int32 InputBufferIndex = DecoderInstance->DequeueInputBuffer(0);
+	if (InputBufferIndex >= 0)
+	{
+		result = DecoderInstance->QueueEOSInputBuffer(InputBufferIndex, 0);//LastSentPTS);
+		check(result == 0);
+		if (result == 0)
+		{
+			InputState = EInputState::AwaitingEOS;
+			return true;
+		}
+		else
+		{
+			PostError(result, "Failed to submit decoder EOS input buffer", ERRCODE_INTERNAL_ANDROID_FAILED_TO_DECODE_VIDEO);
+			return false;
+		}
+	}
+	else if (InputBufferIndex == -1)
+	{
+		// No available input buffer. Try later.
+		return true;
+	}
+	else
+	{
+		PostError(InputBufferIndex, "Failed to get a decoder input buffer for EOS", ERRCODE_INTERNAL_ANDROID_COULD_NOT_GET_INPUT_BUFFER);
+		return false;
+	}
+}
+
+
+//-----------------------------------------------------------------------------
+/**
+ * Tries to get another completed output buffer from the decoder.
+ *
+ * @return
+ */
+FVideoDecoderH264::EOutputResult FVideoDecoderH264::GetOutput()
+{
+	// When there is no decoder yet there will be no output so all is well!
+	if (!DecoderInstance.IsValid())
+	{
+		return EOutputResult::TryAgainLater;
+	}
+
+	// See if the output queue can receive more...
+	if (!Renderer->CanReceiveOutputFrames(1))
+	{
+		// Nope, come back later...
+		return EOutputResult::TryAgainLater;
+	}
+
+	int32 result = -1;
+	IAndroidJavaH264VideoDecoder::FOutputBufferInfo OutputBufferInfo;
+	result = DecoderInstance->DequeueOutputBuffer(OutputBufferInfo, 0);
+	if (result == 0)
+	{
+		if (OutputBufferInfo.BufferIndex >= 0)
+		{
+			// Received our EOS buffer back?
+			if (OutputBufferInfo.bIsEOS)
+			{
+				// We have sent an empty buffer with only the EOS flag set, so we expect an empty buffer in return.
+				check(OutputBufferInfo.Size == 0);
+				// We still need to release the EOS buffer back to the decoder though.
+				result = DecoderInstance->ReleaseOutputBuffer(OutputBufferInfo.BufferIndex, OutputBufferInfo.ValidCount, false, -1);
+				return EOutputResult::EOS;
+			}
+			// We have not seen a config buffer coming back from a decoder, which makes sense. But we do not put it past some
+			// implementation to do that, so for good measure let's ignore it.
+			if (OutputBufferInfo.bIsConfig)
+			{
+				LogMessage(Electra::IInfoLog::Warning, "Got CSD buffer back?");
+				result = DecoderInstance->ReleaseOutputBuffer(OutputBufferInfo.BufferIndex, OutputBufferInfo.ValidCount, false, -1);
+				return EOutputResult::EOS;
+			}
+
+			IAndroidJavaH264VideoDecoder::FOutputFormatInfo OutputFormatInfo;
+			result = DecoderInstance->GetOutputFormatInfo(OutputFormatInfo, OutputBufferInfo.BufferIndex);
+			if (result == 0)
+			{
+				// Check the output for a matching PTS.
+				TSharedPtr<FInDecoderInfo, ESPMode::ThreadSafe> SourceInfo;
+				bool bFound = false;
+				if (InDecoderInfos.Num())
+				{
+					// We expect it to be the first element.
+					if (InDecoderInfos[0]->DecoderPTS == OutputBufferInfo.PresentationTimestamp)
+					{
+						// Yup.
+						SourceInfo = InDecoderInfos[0];
+						InDecoderInfos.RemoveAt(0);
+						bFound = true;
+					}
+					else
+					{
+						// Not the first value indicates a problem with the decoder.
+						// Is the exact value in the list somewhere?
+						for(int32 i=0, iMax=InDecoderInfos.Num(); i<iMax; ++i)
+						{
+							if (InDecoderInfos[i]->DecoderPTS == OutputBufferInfo.PresentationTimestamp)
+							{
+								SourceInfo = InDecoderInfos[i];
+								bFound = true;
+								// Remove all prior entries that are apparently outdated.
+								while(i >= 0)
+								{
+									InDecoderInfos.RemoveAt(0);
+									--i;
+								}
+								break;
+							}
+						}
+						// No exact match found. Use whichever timestamp is just prior to the value we got and remove any older ones.
+						if (!bFound)
+						{
+							do
+							{
+								SourceInfo = InDecoderInfos[0];
+								InDecoderInfos.RemoveAt(0);
+							}
+							while(InDecoderInfos.Num() && InDecoderInfos[0]->DecoderPTS < OutputBufferInfo.PresentationTimestamp);
+							bFound = true;
+						}
+					}
+				}
+				if (!bFound)
+				{
+					// No pushed entries. Now what?
+					PostError(0, "Could not find matching decoder output information", ERRCODE_INTERNAL_ANDROID_COULD_NOT_FIND_OUTPUT_INFO);
+					return EOutputResult::Fail;
+				}
+
+				FDecodedImage NextImage;
+				NextImage.bIsDummy  	   = false;
+				NextImage.SourceInfo	   = SourceInfo;
+				NextImage.OutputFormat     = OutputFormatInfo;
+				NextImage.OutputBufferInfo = OutputBufferInfo;
+				ProcessOutput(NextImage);
+
+				// NOTE: We do not release the output buffer here as this would immediately overwrite the last frame in the surface not yet rendered.
+				//	     Instead we leave returning the buffer and the subsequent update of the texture to the render task.
+					//result = DecoderInstance->ReleaseOutputBuffer(OutputBufferInfo, true, -1);
+
+				return EOutputResult::Ok;
+			}
+			else
+			{
+				PostError(result, "Failed to get decoder output format", ERRCODE_INTERNAL_ANDROID_COULD_NOT_GET_OUTPUT_FORMAT);
+				return EOutputResult::Fail;
+			}
+		}
+		else if (OutputBufferInfo.BufferIndex == IAndroidJavaH264VideoDecoder::FOutputBufferInfo::EBufferIndexValues::MediaCodec_INFO_TRY_AGAIN_LATER)
+		{
+			return EOutputResult::TryAgainLater;
+		}
+		else if (OutputBufferInfo.BufferIndex == IAndroidJavaH264VideoDecoder::FOutputBufferInfo::EBufferIndexValues::MediaCodec_INFO_OUTPUT_FORMAT_CHANGED)
+		{
+			// We do not care about the global format change here. When we need the format we get it from the actual buffer then.
+			// Instead let's try to get the following output right away.
+			return GetOutput();
+		}
+		else if (OutputBufferInfo.BufferIndex == IAndroidJavaH264VideoDecoder::FOutputBufferInfo::EBufferIndexValues::MediaCodec_INFO_OUTPUT_BUFFERS_CHANGED)
+		{
+			// No-op as this is the result of a deprecated API we are not using.
+			// Let's try to get the following output right away.
+			return GetOutput();
+		}
+		else
+		{
+			// What new value might this be?
+			PostError(OutputBufferInfo.BufferIndex, "Unhandled output buffer index value", ERRCODE_INTERNAL_ANDROID_INTERNAL);
+			return EOutputResult::Fail;
+		}
+	}
+	else
+	{
+		PostError(result, "Failed to get decoder output buffer", ERRCODE_INTERNAL_ANDROID_COULD_NOT_GET_OUTPUT_BUFFER);
+		return EOutputResult::Fail;
+	}
+	return EOutputResult::TryAgainLater;
+}
+
+
+bool FVideoDecoderH264::ProcessOutput(const FDecodedImage & NextImage)
+{
+	// Get an output buffer from the renderer to pass the image to.
+	IMediaRenderer::IBuffer* RenderOutputBuffer = nullptr;
+	if (AcquireOutputBuffer(RenderOutputBuffer))
+	{
+		if (RenderOutputBuffer)
+		{
+			FParamDict* OutputBufferSampleProperties = new FParamDict();
+			check(OutputBufferSampleProperties);
+			OutputBufferSampleProperties->Set("pts", FVariantValue(NextImage.SourceInfo->PTS));
+			OutputBufferSampleProperties->Set("duration", FVariantValue(NextImage.SourceInfo->Duration));
+			if (!NextImage.bIsDummy)
+			{
+				int32 w = NextImage.OutputFormat.CropRight - NextImage.OutputFormat.CropLeft + 1;
+				int32 h = NextImage.OutputFormat.CropBottom - NextImage.OutputFormat.CropTop + 1;
+
+				const FStreamCodecInformation::FAspectRatio& Aspect = NextImage.SourceInfo->ParsedInfo.GetAspectRatio();
+				double PixelAspectRatio;
+				int32 aw, ah;
+				if (Aspect.IsSet())
+				{
+					aw = Aspect.Width;
+					ah = Aspect.Height;
+					PixelAspectRatio = (double)Aspect.Width / (double)Aspect.Height;
+				}
+				else
+				{
+					aw = ah = 1;
+					PixelAspectRatio = 1.0;
+				}
+
+				OutputBufferSampleProperties->Set("width", FVariantValue((int64)w));
+				OutputBufferSampleProperties->Set("height", FVariantValue((int64)h));
+				OutputBufferSampleProperties->Set("crop_left", FVariantValue((int64)NextImage.OutputFormat.CropLeft));
+				OutputBufferSampleProperties->Set("crop_right", FVariantValue((int64)NextImage.OutputFormat.Width - (NextImage.OutputFormat.CropLeft + w))); // convert into crop-offset from image border
+				OutputBufferSampleProperties->Set("crop_top", FVariantValue((int64)NextImage.OutputFormat.CropTop));
+				OutputBufferSampleProperties->Set("crop_bottom", FVariantValue((int64)NextImage.OutputFormat.Height - (NextImage.OutputFormat.CropTop + h))); // convert into crop-offset from image border
+				OutputBufferSampleProperties->Set("aspect_ratio", FVariantValue((double)PixelAspectRatio));
+				OutputBufferSampleProperties->Set("aspect_w", FVariantValue((int64)aw));
+				OutputBufferSampleProperties->Set("aspect_h", FVariantValue((int64)ah));
+				OutputBufferSampleProperties->Set("fps_num", FVariantValue((int64)0));
+				OutputBufferSampleProperties->Set("fps_denom", FVariantValue((int64)0));
+				OutputBufferSampleProperties->Set("pixelfmt", FVariantValue((int64)EPixelFormat::PF_B8G8R8A8));
+
+				//OutputBufferSampleProperties->Set("pitch", FVariantValue((int64)NextImage.OutputFormat.Stride));
+				//OutputBufferSampleProperties->Set("slice_height", FVariantValue((int64)NextImage.OutputFormat.SliceHeight));
+				//OutputBufferSampleProperties->Set("colorfmt", FVariantValue((int64)NextImage.OutputFormat.ColorFormat));
+
+				TSharedPtr<FElectraPlayerVideoDecoderOutputAndroidImpl, ESPMode::ThreadSafe> DecoderOutput = RenderOutputBuffer->GetBufferProperties().GetValue("texture").GetSharedPointer<FElectraPlayerVideoDecoderOutputAndroidImpl>();
+				check(DecoderOutput);
+
+				DecoderOutput->Initialize(bSurfaceIsView ? FVideoDecoderOutputAndroid::EOutputType::DirectToSurfaceAsView : FVideoDecoderOutputAndroid::EOutputType::DirectToSurfaceAsQueue, NextImage.OutputBufferInfo.BufferIndex, NextImage.OutputBufferInfo.ValidCount, NativeDecoderID, OutputBufferSampleProperties);
+
+				if (!bSurfaceIsView)
+				{
+					// Release the decoder output buffer & hence enqueue it on our output surface
+					// (we are issuing an RHI thread based update to our texture for each of these, so we should always have a 1:1 mapping - assuming we are fast enough
+					//  to not make the surface drop a frame before we get to it)
+					DecoderInstance->ReleaseOutputBuffer(NextImage.OutputBufferInfo.BufferIndex, NextImage.OutputBufferInfo.ValidCount, true, -1);
+				}
+				else
+				{
+					// We decode right into a Surface. Queue up output buffers until we are ready to show them
+					ReadyOutputBuffersToSurface.Emplace(FOutputBufferInfo(FDecoderTimeStamp(FTimespan(NextImage.SourceInfo->PTS.GetAsHNS()), 0), NextImage.OutputBufferInfo.BufferIndex, NextImage.OutputBufferInfo.ValidCount));
+				}
+			}
+			else
+			{
+				OutputBufferSampleProperties->Set("is_dummy", FVariantValue(true));
+			}
+
+			// Note: we are returning the buffer to the renderer before we are done getting data
+			// (but: this will sync all up as the render command queues are all in order - and hence the async process will be done before MediaTextureResources sees this)
+			Renderer->ReturnBuffer(RenderOutputBuffer, true, *OutputBufferSampleProperties);
+			return true;
+		}
+	}
+	return false;
+}
+
+
+//-----------------------------------------------------------------------------
+/**
+ * Handle presenting output to the Surface for display
+ * (if in 'DirectToSurface' mode)
+ *
+ * We would love to do this in "rendering" code, bt that is impossible as that
+ * would mandate a multi-threaded control of the decoder... which seems impossible
+ * to get to be reliable.
+ */
+void FVideoDecoderH264::ProcessReadyOutputBuffersToSurface()
+{
+	if (!DecoderInstance.IsValid())
+	{
+		return;
+	}
+
+	uint32 I = 0;
+	{
+		FScopeLock Lock(&OutputSurfaceTargetCS);
+
+		// No presentation PTS known?
+		if (OutputSurfaceTargetPTS.Time < 0.0)
+		{
+			return;
+		}
+
+		// Do we have anything to output?
+		uint32 Num = ReadyOutputBuffersToSurface.Num();
+		if (Num == 0)
+		{
+			return;
+		}
+
+		// Look which frame is the newest we could show
+		// (this is a SIMPLE version to select frames - COULD BE BETTER)
+		for (; I < Num; ++I)
+		{
+			const FOutputBufferInfo& OI = ReadyOutputBuffersToSurface[I];
+			if (OI.Timestamp.Time > OutputSurfaceTargetPTS.Time)
+			{
+				// Too new, this one must stay...
+				break;
+			}
+		}
+	}
+
+	// Anything?
+	if (I > 0)
+	{
+		// Yes. Remove all deemed too old without any display...
+		--I;
+		for (uint32 J = 0; J < I; ++J)
+		{
+			const FOutputBufferInfo& OI = ReadyOutputBuffersToSurface[J];
+			DecoderInstance->ReleaseOutputBuffer(OI.BufferIndex, OI.ValidCount , false, -1);
+		}
+		// Display the one we selected
+		const FOutputBufferInfo& OI = ReadyOutputBuffersToSurface[I];
+		DecoderInstance->ReleaseOutputBuffer(OI.BufferIndex, OI.ValidCount, true, -1);
+
+		// Remove what we processed...
+		ReadyOutputBuffersToSurface.RemoveAt(0, I + 1);
+	}
+}
+
+//-----------------------------------------------------------------------------
+/**
+*/
+void FVideoDecoderH264::ReleaseToSurface(uint32 NativeDecoderID, const FDecoderTimeStamp & Time)
+{
+	FVideoDecoderH264** NativeDecoder = NativeDecoderMap.Find(NativeDecoderID);
+	if (NativeDecoder)
+	{
+		FScopeLock Lock(&(*NativeDecoder)->OutputSurfaceTargetCS);
+		(*NativeDecoder)->OutputSurfaceTargetPTS = Time;
+	}
+}
+
+//-----------------------------------------------------------------------------
+/**
+*/
+void FVideoDecoderH264::Android_UpdateSurface(const TSharedPtr<IOptionPointerValueContainer>& Surface)
+{
+	if (bSurfaceIsView)
+	{
+		//!!!! ARE WE 100% SURE ABOUT "CHAIN OF EVENTS"? CAN THIS NOT BE READ TOO EARLY? (esp. also: thread safety)
+		// Update config that will be used to recreate decoder
+		Config.AdditionalOptions.Set("videoDecoder_Android_Surface", FVariantValue(Surface));
+
+		// Current decoder is now unusable as its surface gone / about to be gone
+		AUdataFlushEverything();
+	}
+}
+
+//-----------------------------------------------------------------------------
+/**
+ * H264 video decoder main threaded decode loop
+ */
+void FVideoDecoderH264::WorkerThread()
+{
+	LLM_SCOPE(ELLMTag::ElectraPlayer);
+
+	int32 result = 0;
+	bool bDone  = false;
+
+	bError = false;
+
+	// If the application is suspended and resumed we will lose out decoder session and cannot continue
+	// decoding at the point of interruption. We have to keep the access units from the previous IDR frame
+	// and run them through a new decoder session again, discarding all output, after which we can resume
+	// decoding from the point of interruption.
+	// Configure the playback AU buffer to be large enough to hold the access units to be replayed.
+	ReplayAccessUnitBuffer.CapacitySet(FAccessUnitBuffer::FConfiguration(64 << 20, 360.0));
+
+	// Create decoded image pool.
+	if (!CreateDecodedImagePool())
+	{
+		bError = true;
+	}
+
+	CurrentAccessUnit = nullptr;
+	bMustSendCSD = false;
+	InputState = EInputState::Idle;
+	DecoderState = EDecoderState::Ready;
+	DummyState = EDummyState::NotDummy;
+
+	LastSentPTS = -1;
+
+	OutputSurfaceTargetPTS.Time = -1.0;
+
+	bool bBlockedOnInput = false;
+
+	int64 TimeLast = MEDIAutcTime::CurrentMSec();
+	while(!TerminateThreadSignal.IsSignaled())
+	{
+		// Because of the different paths this decode loop can take there is a possibility that
+		// it may go very fast and not wait for any resources.
+		// To prevent this from becoming a tight loop we make sure to sleep at least some time
+		// here to throttle down.
+		int64 TimeNow = MEDIAutcTime::CurrentMSec();
+		int32 elapsedMS = TimeNow - TimeLast;
+		const int32 kTotalSleepTimeMsec = 5;
+		if (elapsedMS < kTotalSleepTimeMsec)
+		{
+			FMediaRunnable::SleepMilliseconds(kTotalSleepTimeMsec - elapsedMS);
+			TimeLast = MEDIAutcTime::CurrentMSec();
+		}
+		else
+		{
+			TimeLast = TimeNow;
+		}
+
+		// Notify optional buffer listener that we will now be needing an AU for our input buffer.
+		if (!CurrentAccessUnit && !bError && InputBufferListener && AccessUnitBuffer.Num() == 0)
+		{
+			FAccessUnitBufferInfo	sin;
+			IAccessUnitBufferListener::FBufferStats	stats;
+			AccessUnitBuffer.GetStats(sin);
+			stats.NumAUsAvailable  = sin.NumCurrentAccessUnits;
+			stats.NumBytesInBuffer = sin.CurrentMemInUse;
+			stats.MaxBytesOfBuffer = sin.MaxDataSize;
+			stats.bEODSignaled     = sin.bEndOfData;
+			stats.bEODReached      = sin.bEndOfData && sin.NumCurrentAccessUnits == 0;
+			ListenerMutex.Lock();
+			if (InputBufferListener)
+			{
+				InputBufferListener->DecoderInputNeeded(stats);
+			}
+			ListenerMutex.Unlock();
+		}
+
+		// Try to pull output. Do this first to make room in the decoder for new input data.
+		if (!bError)
+		{
+			ProcessReadyOutputBuffersToSurface();
+
+			// Any new data available?
+			/*
+				Note that the onFrameAvailable callback that we route through into native code as signal
+				is pretty useless here. It only triggers as we "release" a decoder output buffer into the
+				BufferQueue inside the output surface, NOT when the decoder would have data for us to push
+				into said surface. Hence: we poll!
+			*/
+			EOutputResult OutputResult = GetOutput();
+
+			switch(OutputResult)
+			{
+				case EOutputResult::Ok:
+				{
+					break;
+				}
+				case EOutputResult::EOS:
+				{
+					// Check that we are actually awaiting the EOS.
+					check(InputState == EInputState::AwaitingEOS);
+					if (InputState == EInputState::AwaitingEOS)
+					{
+						// When not in dummy decode mode a new decoder is now needed.
+						if (DummyState == EDummyState::NotDummy)
+						{
+							DecoderState = EDecoderState::Create;
+							InputState = EInputState::Idle;
+						}
+						// Otherwise, if dummy mode is pending we can now switch to it.
+						else if (DummyState == EDummyState::Entering)
+						{
+							DummyState = EDummyState::InDummy;
+							InputState = EInputState::FeedAU;
+						}
+
+						// Either way there is nothing in the decoder anymore so for safeties sake clear out the info.
+						InDecoderInfos.Empty();
+					}
+					break;
+				}
+				case EOutputResult::TryAgainLater:
+				{
+					// If there is no output and the input is blocked we have probably allowed for more output frames than the decoder
+					// has internally. Normally this should occur only while prerolling, so let's signal that we are stalled on output.
+					// Hopefully this should get the party started.
+					if (bBlockedOnInput)
+					{
+						NotifyReadyBufferListener(false);
+					}
+					break;
+				}
+				default:
+				case EOutputResult::Fail:
+				{
+					bError = true;
+					break;
+				}
+			}
+		}
+
+		// Need to get a new AU?
+		if (!CurrentAccessUnit)
+		{
+			if (AccessUnitBuffer.WaitForData(1000 * 5))
+			{
+				// When there is data, even and especially after a previous EOD, we are no longer done and idling.
+				if (bDone)
+				{
+					bDone = false;
+				}
+				bool bOk = AccessUnitBuffer.Pop(CurrentAccessUnit);
+				MEDIA_UNUSED_VAR(bOk);
+				check(bOk);
+				if (!CurrentAccessUnit->bIsDummyData)
+				{
+					bool bIsKeyframe, bIsDiscardable;
+					PrepareAU(CurrentAccessUnit, bIsKeyframe, bIsDiscardable);
+					// An IDR frame means we can start decoding there, so we can purge any accumulated replay AUs.
+					if (bIsKeyframe)
+					{
+						ReplayAccessUnitBuffer.Flush();
+					}
+
+					// Add to the replay buffer if it is not a discardable access unit.
+					if (!bIsDiscardable)
+					{
+						CurrentAccessUnit->AddRef();
+						if (!ReplayAccessUnitBuffer.Push(CurrentAccessUnit))
+						{
+							// FIXME: Is this cause for a playback error? For now we just forget about the replay AU and take any possible decoding artefacts.
+							FAccessUnit::Release(CurrentAccessUnit);
+						}
+					}
+
+					bool bStreamFormatChanged = CurrentStreamFormatInfo.IsDifferentFrom(CurrentAccessUnit);
+
+					// A new AU following dummy data?
+					bool bPrevWasDummy = false;
+					if (DummyState == EDummyState::InDummy)
+					{
+						DummyState = EDummyState::NotDummy;
+						bPrevWasDummy = true;
+					}
+
+					if (DecoderInstance.IsValid())
+					{
+						// Coming out of dummy mode?
+						if (bPrevWasDummy)
+						{
+							if (DecoderInfo.bIsAdaptive)
+							{
+								// Decoder must be flushed and CSD sent.
+								DecoderState = EDecoderState::Flush;
+								bMustSendCSD = true;
+							}
+							else
+							{
+								// EOS was already sent getting into dummy mode. Go through the decoder create cycle now.
+								DecoderState = EDecoderState::Create;
+								InputState = EInputState::Idle;
+							}
+						}
+						else
+
+						// Did the stream format change?
+						if (bStreamFormatChanged)
+						{
+							// If the decoder claims to be adaptive we just need to send the CSD of the new format first.
+							// Otherwise the decoder needs to be flushed, destroyed and created anew.
+							if (DecoderInfo.bIsAdaptive)
+							{
+								bMustSendCSD = true;
+							}
+							else
+							{
+								InputState = EInputState::FeedEOS;
+							}
+						}
+					}
+					else
+					{
+						DecoderState = EDecoderState::Create;
+					}
+				}
+				else
+				{
+					// When there is a decoder it may need to be flushed first.
+					if (DecoderInstance.IsValid())
+					{
+						// Already in dummy state?
+						if (DummyState == EDummyState::NotDummy)
+						{
+							// No, need to transition into dummy mode.
+							DummyState = EDummyState::Entering;
+							// First flush the decoder.
+							InputState = EInputState::FeedEOS;
+						}
+					}
+					else
+					{
+						// Dummy data can be fed at any time, no decoder is needed.
+						InputState = EInputState::FeedAU;
+						DummyState = EDummyState::InDummy;
+					}
+				}
+			}
+		}
+
+		if (CurrentAccessUnit)
+		{
+			if (!bError)
+			{
+				if (DecoderState == EDecoderState::Create)
+				{
+					DecoderState = EDecoderState::Ready;
+					if (InternalDecoderCreate())
+					{
+						// Ok. Start feeding the current access unit.
+						InputState = EInputState::FeedAU;
+						// If the decoder is an adaptive decoder it needs to get the codec specific data first.
+						bMustSendCSD = DecoderInfo.bIsAdaptive;
+					}
+					else
+					{
+						bError = true;
+					}
+				}
+				else if (DecoderState == EDecoderState::Flush)
+				{
+					if (DecoderInstance.IsValid())
+					{
+						DecoderInstance->Flush();
+					}
+					DecoderState = EDecoderState::Ready;
+					InDecoderInfos.Empty();
+				}
+
+				switch(InputState)
+				{
+					case EInputState::Idle:
+					case EInputState::AwaitingEOS:
+						{
+							bBlockedOnInput = false;
+							break;
+						}
+					case EInputState::FeedAU:
+					{
+						if (!CurrentAccessUnit->bIsDummyData)
+						{
+							// Decode
+							EDecodeResult DecRes = Decode();
+							if (DecRes == EDecodeResult::Ok)
+							{
+								// Ok.
+								bBlockedOnInput = false;
+							}
+							else if (DecRes == EDecodeResult::Fail)
+							{
+								bError = true;
+							}
+							if (DecRes == EDecodeResult::TryAgainLater)
+							{
+								bBlockedOnInput = true;
+							}
+						}
+						else
+						{
+							if (!DecodeDummy())
+							{
+								bError = true;
+							}
+						}
+						break;
+					}
+					case EInputState::FeedEOS:
+					{
+						bBlockedOnInput = false;
+						DrainDecoder();
+						break;
+					}
+				}
+			}
+			else
+			{
+				// In error state we sleep it off until terminated.
+				FMediaRunnable::SleepMicroseconds(CurrentAccessUnit->Duration.GetAsMicroseconds());
+				FAccessUnit::Release(CurrentAccessUnit);
+				CurrentAccessUnit = nullptr;
+			}
+		}
+		else
+		{
+			// No data. Is the buffer at EOD?
+			if (AccessUnitBuffer.IsEndOfData())
+			{
+				NotifyReadyBufferListener(true);
+				// Are we done yet?
+				if (!bDone && !bError)
+				{
+					bError = !DrainDecoder();
+					bDone = InputState == EInputState::AwaitingEOS;
+				}
+				FMediaRunnable::SleepMilliseconds(10);
+			}
+		}
+
+		// Flush?
+		if (FlushDecoderSignal.IsSignaled())
+		{
+			// Have to destroy the decoder!
+			InternalDecoderDestroy();
+			AccessUnitBuffer.Flush();
+			ReplayAccessUnitBuffer.Flush();
+
+			OutputSurfaceTargetPTS.Time = -1.0;
+			ReadyOutputBuffersToSurface.Empty();
+
+			FlushDecoderSignal.Reset();
+			DecoderFlushedSignal.Signal();
+
+			// Reset done state.
+			bDone = false;
+			FAccessUnit::Release(CurrentAccessUnit);
+			CurrentAccessUnit = nullptr;
+			bMustSendCSD = false;
+			InputState = EInputState::Idle;
+			DecoderState = EDecoderState::Ready;
+		}
+	}
+
+	FAccessUnit::Release(CurrentAccessUnit);
+	CurrentAccessUnit = nullptr;
+	InternalDecoderDestroy();
+	DestroyDecodedImagePool();
+	AccessUnitBuffer.Flush();
+	AccessUnitBuffer.CapacitySet(0);
+	ReplayAccessUnitBuffer.Flush();
+	ReplayAccessUnitBuffer.CapacitySet(0);
+}
+
+// - - - - - - - - - - - - - - - - - - - - - - - - - - - - - - - - - - - - - - - - - - - - - - - - - - - - - - - - - - - - - - - - - - - - -
+
+/**
+ * Callback triggered when we release a decoder output buffer to the output surface
+ * (currently here from completeness, but not currently used! - the signal that is)
+ */
+void FVideoDecoderH264::JavaCallback_NewDataAvailable(uint32 NativeDecoderID)
+{
+	FScopeLock Lock(&NativeDecoderMapCS);
+	FVideoDecoderH264 **NativeDecoder = NativeDecoderMap.Find(NativeDecoderID);
+	if (NativeDecoder)
+	{
+		(*NativeDecoder)->NewDataAvailable.Signal();
+	}
+}
+
+} // namespace Electra
+
+// - - - - - - - - - - - - - - - - - - - - - - - - - - - - - - - - - - - - - - - - - - - - - - - - - - - - - - - - - - - - - - - - - - - - -
+
+FVideoDecoderOutput* FElectraPlayerPlatformVideoDecoderOutputFactory::Create()
+{
+	return new Electra::FElectraPlayerVideoDecoderOutputAndroidImpl();
+}
+
+// - - - - - - - - - - - - - - - - - - - - - - - - - - - - - - - - - - - - - - - - - - - - - - - - - - - - - - - - - - - - - - - - - - - - -
+
+#if USE_ANDROID_JNI
+JNI_METHOD void Java_com_epicgames_ue4_ElectraVideoDecoderH264_nativeSignalNewDataAvailable(JNIEnv* jenv, jobject thiz, jint NativeDecoderID)
+{
+	Electra::FVideoDecoderH264::JavaCallback_NewDataAvailable(NativeDecoderID);
+}
+#endif
diff --git a/Engine/Plugins/Media/ElectraPlayer/Source/ElectraPlayerRuntime/Private/Runtime/Decoder/Android/VideoDecoderH264_JavaWrapper_Decoder_Android.cpp b/Engine/Plugins/Media/ElectraPlayer/Source/ElectraPlayerRuntime/Private/Runtime/Decoder/Android/VideoDecoderH264_JavaWrapper_Decoder_Android.cpp
new file mode 100644
--- /dev/null
+++ b/Engine/Plugins/Media/ElectraPlayer/Source/ElectraPlayerRuntime/Private/Runtime/Decoder/Android/VideoDecoderH264_JavaWrapper_Decoder_Android.cpp
@@ -0,0 +1,893 @@
+// Copyright Epic Games, Inc. All Rights Reserved.
+
+#include "HAL/Platform.h"
+
+#if PLATFORM_ANDROID
+
+#include "VideoDecoderH264_JavaWrapper_Android.h"
+#include "Utilities/UtilsMPEGVideo.h"
+#include "Player/PlayerSessionServices.h"
+
+#if (defined(USE_ANDROID_JNI_WITHOUT_GAMEACTIVITY) && USE_ANDROID_JNI_WITHOUT_GAMEACTIVITY != 0) || USE_ANDROID_JNI
+#include "Android/AndroidPlatform.h"
+#include "Android/AndroidJava.h"
+
+#if UE_BUILD_SHIPPING
+// always clear any exceptions in shipping
+#define CHECK_JNI_RESULT(Id) if (Id == 0) { JEnv->ExceptionClear(); }
+#else
+#define CHECK_JNI_RESULT(Id)								\
+	if (Id == 0)											\
+	{														\
+		if (bIsOptional)									\
+		{													\
+			JEnv->ExceptionClear();							\
+		}													\
+		else												\
+		{													\
+			JEnv->ExceptionDescribe();						\
+			checkf(Id != 0, TEXT("Failed to find " #Id));	\
+		}													\
+	}
+#endif
+
+//----------------------------------------------------------------------------------------------------
+
+namespace Electra
+{
+
+	namespace
+	{
+		//-----------------------------------------------------------------------------
+		/**
+		 * Posts an error to the session service error listeners.
+		 *
+		 * @param PlayerSessionServices
+		 * @param ApiReturnValue
+		 * @param Message
+		 * @param Code
+		 * @param Error
+		 */
+		static void GPostError(IPlayerSessionServices* PlayerSessionServices, int32_t ApiReturnValue, const FString& Message, uint16 Code, UEMediaError Error)
+		{
+			if (PlayerSessionServices)
+			{
+				FErrorDetail err;
+				err.SetError(Error != UEMEDIA_ERROR_OK ? Error : UEMEDIA_ERROR_DETAIL);
+				err.SetFacility(Facility::EFacility::H264Decoder);
+				err.SetCode(Code);
+				err.SetMessage(Message);
+				err.SetPlatformMessage(FString::Printf(TEXT("%d (0x%08x)"), ApiReturnValue, ApiReturnValue));
+				PlayerSessionServices->PostError(err);
+			}
+		}
+
+
+		//-----------------------------------------------------------------------------
+		/**
+		 * Sends a log message to the session service log.
+		 *
+		 * @param PlayerSessionServices
+		 * @param Level
+		 * @param Message
+		 */
+		static void GLogMessage(IPlayerSessionServices* PlayerSessionServices, IInfoLog::ELevel Level, const FString& Message)
+		{
+			if (PlayerSessionServices)
+			{
+				PlayerSessionServices->PostLog(Facility::EFacility::H264Decoder, Level, Message);
+			}
+		}
+	}
+
+
+
+
+
+	class FAndroidJavaH264VideoDecoder : public IAndroidJavaH264VideoDecoder
+		, public FJavaClassObject
+	{
+	public:
+
+		FAndroidJavaH264VideoDecoder(IPlayerSessionServices* InPlayerSessionServices);
+		virtual ~FAndroidJavaH264VideoDecoder();
+
+		virtual int32 InitializeDecoder(const FCreateParameters& InCreateParams) override;
+		virtual int32 ReleaseDecoder() override;
+		virtual const FDecoderInformation* GetDecoderInformation() override;
+		virtual int32 Start() override;
+		virtual int32 Stop() override;
+		virtual int32 Flush() override;
+		virtual int32 Reset() override;
+		virtual int32 DequeueInputBuffer(int32 InTimeoutUsec) override;
+		virtual int32 QueueInputBuffer(int32 InBufferIndex, const void* InAccessUnitData, int32 InAccessUnitSize, int64 InTimestampUSec) override;
+		virtual int32 QueueCSDInputBuffer(int32 InBufferIndex, const void* InCSDData, int32 InCSDSize, int64 InTimestampUSec) override;
+		virtual int32 QueueEOSInputBuffer(int32 InBufferIndex, int64 InTimestampUSec) override;
+		virtual int32 GetOutputFormatInfo(FOutputFormatInfo& OutFormatInfo, int32 InOutputBufferIndex) override;
+		virtual int32 DequeueOutputBuffer(FOutputBufferInfo& OutBufferInfo, int32 InTimeoutUsec) override;
+		virtual int32 GetOutputBuffer(void*& OutBufferDataPtr, int32 OutBufferDataSize, const FOutputBufferInfo& InOutBufferInfo) override;
+		virtual int32 ReleaseOutputBuffer(int32 BufferIndex, int32 ValidCount, bool bRender, int64 releaseAt) override;
+
+	private:
+		static FName GetClassName()
+		{
+			return FName("com/epicgames/ue4/ElectraVideoDecoderH264");
+		}
+
+		void PostError(int32_t ApiReturnValue, const FString& Message, uint16 Code, UEMediaError Error = UEMEDIA_ERROR_OK)
+		{
+			GPostError(PlayerSessionServices, ApiReturnValue, Message, Code, Error);
+		}
+		void LogMessage(IInfoLog::ELevel Level, const FString& Message)
+		{
+			GLogMessage(PlayerSessionServices, Level, Message);
+		}
+
+		jfieldID FindField(JNIEnv* JEnv, jclass InClass, const ANSICHAR* InFieldName, const ANSICHAR* InFieldType, bool bIsOptional)
+		{
+			jfieldID Field = InClass == nullptr ? nullptr : JEnv->GetFieldID(InClass, InFieldName, InFieldType);
+			CHECK_JNI_RESULT(Field);
+			return Field;
+		}
+
+		jmethodID FindMethod(JNIEnv* JEnv, jclass InClass, const ANSICHAR* InMethodName, const ANSICHAR* InMethodSignature, bool bIsOptional)
+		{
+			jmethodID Method = InClass == nullptr ? nullptr : JEnv->GetMethodID(InClass, InMethodName, InMethodSignature);
+			CHECK_JNI_RESULT(Method);
+			return Method;
+		}
+
+		// Create a Java byte array. Must DeleteLocalRef() after use or handling over to Java.
+		jbyteArray MakeJavaByteArray(const uint8* InData, int32 InNumBytes)
+		{
+			JNIEnv* JEnv = AndroidJavaEnv::GetJavaEnv();
+			jbyteArray RawBuffer = JEnv->NewByteArray(InNumBytes);
+			JEnv->SetByteArrayRegion(RawBuffer, 0, InNumBytes, reinterpret_cast<const jbyte*>(InData));
+			return RawBuffer;
+		}
+
+		int32 release();
+
+		//
+		IPlayerSessionServices* PlayerSessionServices;
+
+	public:
+		// Java methods
+		FJavaClassMethod	CreateDecoderFN;
+		FJavaClassMethod	ReleaseDecoderFN;
+		FJavaClassMethod	ReleaseFN;
+		FJavaClassMethod	StartFN;
+		FJavaClassMethod	StopFN;
+		FJavaClassMethod	FlushFN;
+		FJavaClassMethod	ResetFN;
+		FJavaClassMethod	DequeueInputBufferFN;
+		FJavaClassMethod	QueueInputBufferFN;
+		FJavaClassMethod	QueueCSDInputBufferFN;
+		FJavaClassMethod	QueueEOSInputBufferFN;
+		FJavaClassMethod	GetDecoderInformationFN;
+		FJavaClassMethod	GetOutputFormatInfoFN;
+		FJavaClassMethod	DequeueOutputBufferFN;
+		FJavaClassMethod	GetOutputBufferFN;
+		FJavaClassMethod	ReleaseOutputBufferFN;
+
+		// FCreateParameters member field IDs
+		jclass				FCreateParametersClass;
+		jmethodID			FCreateParametersCTOR;
+		jfieldID			FCreateParameters_MaxWidth;
+		jfieldID			FCreateParameters_MaxHeight;
+		jfieldID			FCreateParameters_MaxFPS;
+		jfieldID			FCreateParameters_Width;
+		jfieldID			FCreateParameters_Height;
+		jfieldID			FCreateParameters_bNeedSecure;
+		jfieldID			FCreateParameters_bNeedTunneling;
+		jfieldID			FCreateParameters_CSD0;
+		jfieldID			FCreateParameters_CSD1;
+		jfieldID			FCreateParameters_ExternalRenderTextureID;
+		jfieldID			FCreateParameters_bRetainRenderer;
+		jfieldID			FCreateParameters_bSwizzleTexture;
+		jfieldID			FCreateParameters_NativeDecoderID;
+		jfieldID			FCreateParameters_VideoCodecSurface;
+		jfieldID			FCreateParameters_bSurfaceIsView;
+
+		// FDecoderInformation member field IDs
+		jclass				FDecoderInformationClass;
+		jfieldID			FDecoderInformation_bIsAdaptive;
+
+		// FOutputFormatInfo member field IDs
+		jclass				FOutputFormatInfoClass;
+		jfieldID			FOutputFormatInfo_Width;
+		jfieldID			FOutputFormatInfo_Height;
+		jfieldID			FOutputFormatInfo_CropTop;
+		jfieldID			FOutputFormatInfo_CropBottom;
+		jfieldID			FOutputFormatInfo_CropLeft;
+		jfieldID			FOutputFormatInfo_CropRight;
+		jfieldID			FOutputFormatInfo_Stride;
+		jfieldID			FOutputFormatInfo_SliceHeight;
+		jfieldID			FOutputFormatInfo_ColorFormat;
+
+		// FOutputBufferInfo member field IDs
+		jclass				FOutputBufferInfoClass;
+		jfieldID			FOutputBufferInfo_BufferIndex;
+		jfieldID			FOutputBufferInfo_PresentationTimestamp;
+		jfieldID			FOutputBufferInfo_Size;
+		jfieldID			FOutputBufferInfo_bIsEOS;
+		jfieldID			FOutputBufferInfo_bIsConfig;
+
+		// Internal state
+		FMediaCriticalSection						MutexLock;
+		TUniquePtr<FDecoderInformation>	CurrentDecoderInformation;
+		int32							CurrentValidCount;
+		bool							bHaveDecoder;
+		bool							bIsStarted;
+	};
+
+
+	//-----------------------------------------------------------------------------
+	/**
+	 * Create a Java video decoder wrapper.
+	 *
+	 * @return Decoder wrapper instance.
+	 */
+	TSharedPtr<IAndroidJavaH264VideoDecoder, ESPMode::ThreadSafe> IAndroidJavaH264VideoDecoder::Create(IPlayerSessionServices* InPlayerSessionServices)
+	{
+		return MakeShared<FAndroidJavaH264VideoDecoder>(InPlayerSessionServices);
+	}
+
+
+	//-----------------------------------------------------------------------------
+	/**
+	 * CTOR
+	 */
+	FAndroidJavaH264VideoDecoder::FAndroidJavaH264VideoDecoder(IPlayerSessionServices* InPlayerSessionServices)
+		: FJavaClassObject(GetClassName(), "()V")
+		, PlayerSessionServices(InPlayerSessionServices)
+		, CreateDecoderFN(GetClassMethod("CreateDecoder", "(Lcom/epicgames/ue4/ElectraVideoDecoderH264$FCreateParameters;)I"))
+		, ReleaseDecoderFN(GetClassMethod("ReleaseDecoder", "()I"))
+		, ReleaseFN(GetClassMethod("release", "()I"))
+		, StartFN(GetClassMethod("Start", "()I"))
+		, StopFN(GetClassMethod("Stop", "()I"))
+		, FlushFN(GetClassMethod("Flush", "()I"))
+		, ResetFN(GetClassMethod("Reset", "()I"))
+		, DequeueInputBufferFN(GetClassMethod("DequeueInputBuffer", "(I)I"))
+		, QueueInputBufferFN(GetClassMethod("QueueInputBuffer", "(IJ[B)I"))
+		, QueueCSDInputBufferFN(GetClassMethod("QueueCSDInputBuffer", "(IJ[B)I"))
+		, QueueEOSInputBufferFN(GetClassMethod("QueueEOSInputBuffer", "(IJ)I"))
+		, GetDecoderInformationFN(GetClassMethod("GetDecoderInformation", "()Lcom/epicgames/ue4/ElectraVideoDecoderH264$FDecoderInformation;"))
+		, GetOutputFormatInfoFN(GetClassMethod("GetOutputFormatInfo", "(I)Lcom/epicgames/ue4/ElectraVideoDecoderH264$FOutputFormatInfo;"))
+		, DequeueOutputBufferFN(GetClassMethod("DequeueOutputBuffer", "(I)Lcom/epicgames/ue4/ElectraVideoDecoderH264$FOutputBufferInfo;"))
+		, GetOutputBufferFN(GetClassMethod("GetOutputBuffer", "(I)[B"))
+		, ReleaseOutputBufferFN(GetClassMethod("ReleaseOutputBuffer", "(IZJ)I"))
+	{
+		JNIEnv* JEnv = AndroidJavaEnv::GetJavaEnv();
+
+		// Get field IDs for FCreateParameters class members
+		jclass localCreateParametersClass = AndroidJavaEnv::FindJavaClass("com/epicgames/ue4/ElectraVideoDecoderH264$FCreateParameters");
+		FCreateParametersClass = (jclass)JEnv->NewGlobalRef(localCreateParametersClass);
+		JEnv->DeleteLocalRef(localCreateParametersClass);
+		FCreateParametersCTOR = FindMethod(JEnv, FCreateParametersClass, "<init>", "()V", false);
+		FCreateParameters_MaxWidth = FindField(JEnv, FCreateParametersClass, "MaxWidth", "I", false);
+		FCreateParameters_MaxHeight = FindField(JEnv, FCreateParametersClass, "MaxHeight", "I", false);
+		FCreateParameters_MaxFPS = FindField(JEnv, FCreateParametersClass, "MaxFPS", "I", false);
+		FCreateParameters_Width = FindField(JEnv, FCreateParametersClass, "Width", "I", false);
+		FCreateParameters_Height = FindField(JEnv, FCreateParametersClass, "Height", "I", false);
+		FCreateParameters_bNeedSecure = FindField(JEnv, FCreateParametersClass, "bNeedSecure", "Z", false);
+		FCreateParameters_bNeedTunneling = FindField(JEnv, FCreateParametersClass, "bNeedTunneling", "Z", false);
+		FCreateParameters_CSD0 = FindField(JEnv, FCreateParametersClass, "CSD0", "[B", false);
+		FCreateParameters_CSD1 = FindField(JEnv, FCreateParametersClass, "CSD1", "[B", false);
+		FCreateParameters_ExternalRenderTextureID = FindField(JEnv, FCreateParametersClass, "ExternalRenderTextureID", "I", false);
+		FCreateParameters_bRetainRenderer = FindField(JEnv, FCreateParametersClass, "bRetainRenderer", "Z", false);
+		FCreateParameters_bSwizzleTexture = FindField(JEnv, FCreateParametersClass, "bSwizzleTexture", "Z", false);
+		FCreateParameters_NativeDecoderID = FindField(JEnv, FCreateParametersClass, "NativeDecoderID", "I", false);
+		FCreateParameters_VideoCodecSurface = FindField(JEnv, FCreateParametersClass, "VideoCodecSurface", "Landroid/view/Surface;", false);
+		FCreateParameters_bSurfaceIsView = FindField(JEnv, FCreateParametersClass, "bSurfaceIsView", "Z", false);
+
+		// Get field IDs for FDecoderInformation class members
+		jclass localDecoderInformationClass = AndroidJavaEnv::FindJavaClass("com/epicgames/ue4/ElectraVideoDecoderH264$FDecoderInformation");
+		FDecoderInformationClass = (jclass)JEnv->NewGlobalRef(localDecoderInformationClass);
+		JEnv->DeleteLocalRef(localDecoderInformationClass);
+		FDecoderInformation_bIsAdaptive = FindField(JEnv, FDecoderInformationClass, "bIsAdaptive", "Z", false);
+
+		// Get field IDs for FOutputFormatInfo class members
+		jclass localOutputFormatInfoClass = AndroidJavaEnv::FindJavaClass("com/epicgames/ue4/ElectraVideoDecoderH264$FOutputFormatInfo");
+		FOutputFormatInfoClass = (jclass)JEnv->NewGlobalRef(localOutputFormatInfoClass);
+		JEnv->DeleteLocalRef(localOutputFormatInfoClass);
+		FOutputFormatInfo_Width = FindField(JEnv, FOutputFormatInfoClass, "Width", "I", false);
+		FOutputFormatInfo_Height = FindField(JEnv, FOutputFormatInfoClass, "Height", "I", false);
+		FOutputFormatInfo_CropTop = FindField(JEnv, FOutputFormatInfoClass, "CropTop", "I", false);
+		FOutputFormatInfo_CropBottom = FindField(JEnv, FOutputFormatInfoClass, "CropBottom", "I", false);
+		FOutputFormatInfo_CropLeft = FindField(JEnv, FOutputFormatInfoClass, "CropLeft", "I", false);
+		FOutputFormatInfo_CropRight = FindField(JEnv, FOutputFormatInfoClass, "CropRight", "I", false);
+		FOutputFormatInfo_Stride = FindField(JEnv, FOutputFormatInfoClass, "Stride", "I", false);
+		FOutputFormatInfo_SliceHeight = FindField(JEnv, FOutputFormatInfoClass, "SliceHeight", "I", false);
+		FOutputFormatInfo_ColorFormat = FindField(JEnv, FOutputFormatInfoClass, "ColorFormat", "I", false);
+
+		// Get field IDs for FOutputBufferInfo class members
+		jclass localOutputBufferInfoClass = AndroidJavaEnv::FindJavaClass("com/epicgames/ue4/ElectraVideoDecoderH264$FOutputBufferInfo");
+		FOutputBufferInfoClass = (jclass)JEnv->NewGlobalRef(localOutputBufferInfoClass);
+		JEnv->DeleteLocalRef(localOutputBufferInfoClass);
+		FOutputBufferInfo_BufferIndex = FindField(JEnv, FOutputBufferInfoClass, "BufferIndex", "I", false);
+		FOutputBufferInfo_PresentationTimestamp = FindField(JEnv, FOutputBufferInfoClass, "PresentationTimestamp", "J", false);
+		FOutputBufferInfo_Size = FindField(JEnv, FOutputBufferInfoClass, "Size", "I", false);
+		FOutputBufferInfo_bIsEOS = FindField(JEnv, FOutputBufferInfoClass, "bIsEOS", "Z", false);
+		FOutputBufferInfo_bIsConfig = FindField(JEnv, FOutputBufferInfoClass, "bIsConfig", "Z", false);
+
+
+		// Internal state
+		CurrentValidCount = 0;
+		bHaveDecoder = false;
+		bIsStarted = false;
+	}
+
+
+	//-----------------------------------------------------------------------------
+	/**
+	 * DTOR
+	 */
+	FAndroidJavaH264VideoDecoder::~FAndroidJavaH264VideoDecoder()
+	{
+		FMediaCriticalSection::ScopedLock Lock(MutexLock);
+		release();
+
+		JNIEnv* JEnv = AndroidJavaEnv::GetJavaEnv();
+		JEnv->DeleteGlobalRef(FCreateParametersClass);
+		JEnv->DeleteGlobalRef(FDecoderInformationClass);
+		JEnv->DeleteGlobalRef(FOutputBufferInfoClass);
+		JEnv->DeleteGlobalRef(FOutputFormatInfoClass);
+
+		CurrentDecoderInformation.Reset();
+	}
+
+
+	//-----------------------------------------------------------------------------
+	/**
+	 * Creates and initializes a Java instance of an H.264 video decoder.
+	 *
+	 * @param InCreateParams
+	 *
+	 * @return 0 if successful, 1 on error.
+	 */
+	int32 FAndroidJavaH264VideoDecoder::InitializeDecoder(const FCreateParameters& InCreateParams)
+	{
+		FMediaCriticalSection::ScopedLock Lock(MutexLock);
+		int32 result = -1;
+
+		// Already have a decoder?
+		if (bHaveDecoder)
+		{
+			// Do not care about potential error return values here.
+			Stop();
+			ReleaseDecoder();
+			bHaveDecoder = false;
+		}
+
+		CurrentDecoderInformation.Reset();
+
+		// Create an instance of the init param structure and fill in the members.
+		JNIEnv* JEnv = AndroidJavaEnv::GetJavaEnv();
+		jobject CreateParams = JEnv->NewObject(FCreateParametersClass, FCreateParametersCTOR);
+		JEnv->SetIntField(CreateParams, FCreateParameters_MaxWidth, InCreateParams.MaxWidth);
+		JEnv->SetIntField(CreateParams, FCreateParameters_MaxHeight, InCreateParams.MaxHeight);
+		JEnv->SetIntField(CreateParams, FCreateParameters_MaxFPS, InCreateParams.MaxFrameRate);
+		//JEnv->SetBooleanField(CreateParams, FCreateParameters_bNeedSecure, false);
+		//JEnv->SetBooleanField(CreateParams, FCreateParameters_bNeedTunneling, false);
+		JEnv->SetIntField(CreateParams, FCreateParameters_NativeDecoderID, InCreateParams.NativeDecoderID);
+
+		JEnv->SetObjectField(CreateParams, FCreateParameters_VideoCodecSurface, InCreateParams.VideoCodecSurface);
+
+		if (InCreateParams.CodecData.IsValid())
+		{
+			JEnv->SetIntField(CreateParams, FCreateParameters_Width, InCreateParams.CodecData->ParsedInfo.GetResolution().Width);
+			JEnv->SetIntField(CreateParams, FCreateParameters_Height, InCreateParams.CodecData->ParsedInfo.GetResolution().Height);
+
+			MPEG::FAVCDecoderConfigurationRecord avc;
+			avc.SetRawData(InCreateParams.CodecData->RawCSD.GetData(), InCreateParams.CodecData->RawCSD.Num());
+			if (avc.Parse())
+			{
+				jbyteArray CSD0 = MakeJavaByteArray(avc.GetCodecSpecificDataSPS().GetData(), avc.GetCodecSpecificDataSPS().Num());
+				jbyteArray CSD1 = MakeJavaByteArray(avc.GetCodecSpecificDataPPS().GetData(), avc.GetCodecSpecificDataPPS().Num());
+				JEnv->SetObjectField(CreateParams, FCreateParameters_CSD0, CSD0);
+				JEnv->SetObjectField(CreateParams, FCreateParameters_CSD1, CSD1);
+				JEnv->DeleteLocalRef(CSD0);
+				JEnv->DeleteLocalRef(CSD1);
+			}
+		}
+
+		// Pass along decoder output surface
+		JEnv->SetBooleanField(CreateParams, FCreateParameters_bRetainRenderer, InCreateParams.bRetainRenderer);
+		JEnv->SetBooleanField(CreateParams, FCreateParameters_bSurfaceIsView, InCreateParams.bSurfaceIsView);
+
+		// Create and initialize a decoder instance.
+		result = CallMethod<int>(CreateDecoderFN, CreateParams);
+		JEnv->DeleteLocalRef(CreateParams);
+		if (JEnv->ExceptionCheck())
+		{
+			JEnv->ExceptionDescribe();
+			JEnv->ExceptionClear();
+		}
+		if (result != 0)
+		{
+			return 1;
+		}
+
+		// Get decoder information
+		jobject OutputInfo = JEnv->CallObjectMethod(Object, GetDecoderInformationFN.Method);
+		if (JEnv->ExceptionCheck())
+		{
+			JEnv->ExceptionDescribe();
+			JEnv->ExceptionClear();
+		}
+		// Failure will return no object.
+		check(OutputInfo != nullptr);
+		if (OutputInfo != nullptr)
+		{
+			CurrentDecoderInformation = MakeUnique<FDecoderInformation>();
+			CurrentDecoderInformation->bIsAdaptive = JEnv->GetBooleanField(OutputInfo, FDecoderInformation_bIsAdaptive);
+			JEnv->DeleteLocalRef(OutputInfo);
+		}
+
+		bHaveDecoder = true;
+		return 0;
+	}
+
+
+
+	//-----------------------------------------------------------------------------
+	/**
+	 * Releases (destroys) the Java video decoder instance.
+	 *
+	 * @return 0 if successful, 1 on error.
+	 */
+	int32 FAndroidJavaH264VideoDecoder::ReleaseDecoder()
+	{
+		FMediaCriticalSection::ScopedLock Lock(MutexLock);
+		if (bHaveDecoder)
+		{
+			int32 result = CallMethod<int>(ReleaseDecoderFN);
+			JNIEnv* JEnv = AndroidJavaEnv::GetJavaEnv();
+			if (JEnv->ExceptionCheck())
+			{
+				JEnv->ExceptionDescribe();
+				JEnv->ExceptionClear();
+			}
+			bHaveDecoder = false;
+			return result ? 1 : 0;
+		}
+		return 1;
+	}
+
+
+	//-----------------------------------------------------------------------------
+	/**
+	 * Releases (destroys) the Java texture renderer the Java decoder has created
+	 *
+	 * @return 0 if successful, 1 on error.
+	 */
+	int32 FAndroidJavaH264VideoDecoder::release()
+	{
+		Stop();
+
+		int32 result = CallMethod<int>(ReleaseFN);
+		JNIEnv* JEnv = AndroidJavaEnv::GetJavaEnv();
+		if (JEnv->ExceptionCheck())
+		{
+			JEnv->ExceptionDescribe();
+			JEnv->ExceptionClear();
+		}
+		return result ? 1 : 0;
+	}
+
+
+	//-----------------------------------------------------------------------------
+	/**
+	 * Returns decoder information after a successful InitializeDecoder().
+	 *
+	 * @return Pointer to decoder information or null when no decoder has been created.
+	 */
+	const IAndroidJavaH264VideoDecoder::FDecoderInformation* FAndroidJavaH264VideoDecoder::GetDecoderInformation()
+	{
+		return CurrentDecoderInformation.Get();
+	}
+
+
+	//-----------------------------------------------------------------------------
+	/**
+	 * Starts the decoder instance.
+	 *
+	 * @return 0 if successful, 1 on error.
+	 */
+	int32 FAndroidJavaH264VideoDecoder::Start()
+	{
+		FMediaCriticalSection::ScopedLock Lock(MutexLock);
+		if (bHaveDecoder && !bIsStarted)
+		{
+			int32 result = CallMethod<int>(StartFN);
+			JNIEnv* JEnv = AndroidJavaEnv::GetJavaEnv();
+			if (JEnv->ExceptionCheck())
+			{
+				JEnv->ExceptionDescribe();
+				JEnv->ExceptionClear();
+			}
+			if (result)
+			{
+				return result;
+			}
+			bIsStarted = true;
+			return 0;
+		}
+		return 1;
+	}
+
+
+	//-----------------------------------------------------------------------------
+	/**
+	 * Stops the decoder instance.
+	 *
+	 * @return 0 if successful, 1 on error.
+	 */
+	int32 FAndroidJavaH264VideoDecoder::Stop()
+	{
+		FMediaCriticalSection::ScopedLock Lock(MutexLock);
+		if (bHaveDecoder && bIsStarted)
+		{
+			int32 result = CallMethod<int>(StopFN);
+			JNIEnv* JEnv = AndroidJavaEnv::GetJavaEnv();
+			if (JEnv->ExceptionCheck())
+			{
+				JEnv->ExceptionDescribe();
+				JEnv->ExceptionClear();
+			}
+			if (result)
+			{
+				return result;
+			}
+			bIsStarted = false;
+			return 0;
+		}
+		return 1;
+	}
+
+
+	//-----------------------------------------------------------------------------
+	/**
+	 * Flushes the decoder instance.
+	 *
+	 * @return 0 if successful, 1 on error.
+	 */
+	int32 FAndroidJavaH264VideoDecoder::Flush()
+	{
+		FMediaCriticalSection::ScopedLock Lock(MutexLock);
+
+		// Need to increment the valid count to ensure we will not try to release or render output buffers
+		// that have become invalid.
+		FPlatformAtomics::InterlockedIncrement(&CurrentValidCount);
+
+		// Synchronously operating decoders must be in the started state to be flushed.
+		if (bHaveDecoder && bIsStarted)
+		{
+			int32 result = CallMethod<int>(FlushFN);
+			JNIEnv* JEnv = AndroidJavaEnv::GetJavaEnv();
+			if (JEnv->ExceptionCheck())
+			{
+				JEnv->ExceptionDescribe();
+				JEnv->ExceptionClear();
+			}
+			return result;
+		}
+		return 1;
+	}
+
+
+	//-----------------------------------------------------------------------------
+	/**
+	 * Resets the decoder instance.
+	 *
+	 * @return 0 if successful, 1 on error.
+	 */
+	int32 FAndroidJavaH264VideoDecoder::Reset()
+	{
+		FMediaCriticalSection::ScopedLock Lock(MutexLock);
+		// Synchronously operating decoders should (must?) be in the stopped state to be reset.
+		if (bHaveDecoder && !bIsStarted)
+		{
+			int32 result = CallMethod<int>(ResetFN);
+			JNIEnv* JEnv = AndroidJavaEnv::GetJavaEnv();
+			if (JEnv->ExceptionCheck())
+			{
+				JEnv->ExceptionDescribe();
+				JEnv->ExceptionClear();
+			}
+			return result ? 1 : 0;
+		}
+		return 1;
+	}
+
+
+	//-----------------------------------------------------------------------------
+	/**
+	 * Dequeues an input buffer.
+	 *
+	 * @param InTimeoutUsec Timeout in microseconds to wait for an available buffer.
+	 *
+	 * @return >= 0 returns the index of the successfully dequeued buffer, negative values indicate an error.
+	 */
+	int32 FAndroidJavaH264VideoDecoder::DequeueInputBuffer(int32 InTimeoutUsec)
+	{
+		FMediaCriticalSection::ScopedLock Lock(MutexLock);
+		if (bHaveDecoder)
+		{
+			int32 result = CallMethod<int>(DequeueInputBufferFN, InTimeoutUsec);
+			JNIEnv* JEnv = AndroidJavaEnv::GetJavaEnv();
+			if (JEnv->ExceptionCheck())
+			{
+				JEnv->ExceptionDescribe();
+				JEnv->ExceptionClear();
+			}
+			return result;
+		}
+		return -1;
+	}
+
+
+	//-----------------------------------------------------------------------------
+	/**
+	 * Queues input for decoding in the buffer with a previously dequeued (calling DequeueInputBuffer()) index.
+	 *
+	 * @param InBufferIndex Index of the buffer to put data into and enqueue for decoding (see DequeueInputBuffer()).
+	 * @param InAccessUnitData Data to be decoded.
+	 * @param InAccessUnitSize Size of the data to be decoded.
+	 * @param InTimestampUSec Timestamp (PTS) of the data, in microseconds.
+	 *
+	 * @return 0 if successful, 1 on error.
+	 */
+	int32 FAndroidJavaH264VideoDecoder::QueueInputBuffer(int32 InBufferIndex, const void* InAccessUnitData, int32 InAccessUnitSize, int64 InTimestampUSec)
+	{
+		FMediaCriticalSection::ScopedLock Lock(MutexLock);
+		if (bHaveDecoder)
+		{
+			JNIEnv* JEnv = AndroidJavaEnv::GetJavaEnv();
+			jbyteArray  InData = MakeJavaByteArray((const uint8*)InAccessUnitData, InAccessUnitSize);
+			int32 result = CallMethod<int>(QueueInputBufferFN, InBufferIndex, (jlong)InTimestampUSec, InData);
+			JEnv->DeleteLocalRef(InData);
+			if (JEnv->ExceptionCheck())
+			{
+				JEnv->ExceptionDescribe();
+				JEnv->ExceptionClear();
+			}
+			return result ? 1 : 0;
+		}
+		return 1;
+	}
+
+
+	//-----------------------------------------------------------------------------
+	/**
+	 * Queues codec specific data for the following to-be-decoded data buffers.
+	 *
+	 * @param InBufferIndex Index of the buffer to put data into and enqueue for decoding (see DequeueInputBuffer()).
+	 * @param InCSDData Codec specific data
+	 * @param InCSDSize Size of the codec specific data
+	 * @param InTimestampUSec Timestamp (PTS) of the data, in microseconds. Must be the same as the next data to be decoded.
+	 *
+	 * @return 0 if successful, 1 on error.
+	 */
+	int32 FAndroidJavaH264VideoDecoder::QueueCSDInputBuffer(int32 InBufferIndex, const void* InCSDData, int32 InCSDSize, int64 InTimestampUSec)
+	{
+		FMediaCriticalSection::ScopedLock Lock(MutexLock);
+		if (bHaveDecoder)
+		{
+			JNIEnv* JEnv = AndroidJavaEnv::GetJavaEnv();
+			jbyteArray  InData = MakeJavaByteArray((const uint8*)InCSDData, InCSDSize);
+			int32 result = CallMethod<int>(QueueCSDInputBufferFN, InBufferIndex, (jlong)InTimestampUSec, InData);
+			JEnv->DeleteLocalRef(InData);
+			if (JEnv->ExceptionCheck())
+			{
+				JEnv->ExceptionDescribe();
+				JEnv->ExceptionClear();
+			}
+			return result ? 1 : 0;
+		}
+		return 1;
+	}
+
+
+	//-----------------------------------------------------------------------------
+	/**
+	 * Queues end of stream for the buffer with a previously dequeued (calling DequeueInputBuffer()) index.
+	 *
+	 * @param InBufferIndex Index of the buffer to put the EOS flag into and enqueue for decoding (see DequeueInputBuffer()).
+	 * @param InTimestampUSec Timestamp the previous data had. Can be 0.
+	 *
+	 * @return 0 if successful, 1 on error.
+	 */
+	int32 FAndroidJavaH264VideoDecoder::QueueEOSInputBuffer(int32 InBufferIndex, int64 InTimestampUSec)
+	{
+		FMediaCriticalSection::ScopedLock Lock(MutexLock);
+		if (bHaveDecoder)
+		{
+			JNIEnv* JEnv = AndroidJavaEnv::GetJavaEnv();
+			int32 result = CallMethod<int>(QueueEOSInputBufferFN, InBufferIndex, (jlong)InTimestampUSec);
+			if (JEnv->ExceptionCheck())
+			{
+				JEnv->ExceptionDescribe();
+				JEnv->ExceptionClear();
+			}
+			return result ? 1 : 0;
+		}
+		return 1;
+	}
+
+
+	//-----------------------------------------------------------------------------
+	/**
+	 * Returns format information of the decoded samples.
+	 *
+	 * @param OutFormatInfo
+	 * @param InOutputBufferIndex RESERVED FOR NOW - Pass any negative value to get the output format after DequeueOutputBuffer() returns a BufferIndex of MediaCodec_INFO_OUTPUT_FORMAT_CHANGED.
+	 *
+	 * @return 0 if successful, 1 on error.
+	 */
+	int32 FAndroidJavaH264VideoDecoder::GetOutputFormatInfo(FOutputFormatInfo& OutFormatInfo, int32 InOutputBufferIndex)
+	{
+		FMediaCriticalSection::ScopedLock Lock(MutexLock);
+		if (bHaveDecoder)
+		{
+			JNIEnv* JEnv = AndroidJavaEnv::GetJavaEnv();
+			jobject OutputInfo = JEnv->CallObjectMethod(Object, GetOutputFormatInfoFN.Method, InOutputBufferIndex);
+			if (JEnv->ExceptionCheck())
+			{
+				JEnv->ExceptionDescribe();
+				JEnv->ExceptionClear();
+			}
+			// Failure will return no object.
+			if (OutputInfo != nullptr)
+			{
+				OutFormatInfo.Width = JEnv->GetIntField(OutputInfo, FOutputFormatInfo_Width);
+				OutFormatInfo.Height = JEnv->GetIntField(OutputInfo, FOutputFormatInfo_Height);
+				OutFormatInfo.CropTop = JEnv->GetIntField(OutputInfo, FOutputFormatInfo_CropTop);
+				OutFormatInfo.CropBottom = JEnv->GetIntField(OutputInfo, FOutputFormatInfo_CropBottom);
+				OutFormatInfo.CropLeft = JEnv->GetIntField(OutputInfo, FOutputFormatInfo_CropLeft);
+				OutFormatInfo.CropRight = JEnv->GetIntField(OutputInfo, FOutputFormatInfo_CropRight);
+				OutFormatInfo.Stride = JEnv->GetIntField(OutputInfo, FOutputFormatInfo_Stride);
+				OutFormatInfo.SliceHeight = JEnv->GetIntField(OutputInfo, FOutputFormatInfo_SliceHeight);
+				OutFormatInfo.ColorFormat = JEnv->GetIntField(OutputInfo, FOutputFormatInfo_ColorFormat);
+				JEnv->DeleteLocalRef(OutputInfo);
+				return 0;
+			}
+		}
+		return 1;
+	}
+
+
+	//-----------------------------------------------------------------------------
+	/**
+	 * Dequeues an output buffer.
+	 *
+	 * @param InTimeoutUsec Timeout in microseconds to wait for an available buffer.
+	 *
+	 * @return 0 on success, 1 on failure. The OutBufferInfo.BufferIndex indicates the buffer index.
+	 */
+	int32 FAndroidJavaH264VideoDecoder::DequeueOutputBuffer(FOutputBufferInfo& OutBufferInfo, int32 InTimeoutUsec)
+	{
+		FMediaCriticalSection::ScopedLock Lock(MutexLock);
+		if (bHaveDecoder)
+		{
+			JNIEnv* JEnv = AndroidJavaEnv::GetJavaEnv();
+			jobject OutputInfo = JEnv->CallObjectMethod(Object, DequeueOutputBufferFN.Method, InTimeoutUsec);
+			if (JEnv->ExceptionCheck())
+			{
+				JEnv->ExceptionDescribe();
+				JEnv->ExceptionClear();
+			}
+			// Failure will return no object.
+			if (OutputInfo != nullptr)
+			{
+				OutBufferInfo.ValidCount = CurrentValidCount;
+				OutBufferInfo.BufferIndex = JEnv->GetIntField(OutputInfo, FOutputBufferInfo_BufferIndex);
+				OutBufferInfo.PresentationTimestamp = JEnv->GetLongField(OutputInfo, FOutputBufferInfo_PresentationTimestamp);
+				OutBufferInfo.Size = JEnv->GetIntField(OutputInfo, FOutputBufferInfo_Size);
+				OutBufferInfo.bIsEOS = JEnv->GetBooleanField(OutputInfo, FOutputBufferInfo_bIsEOS);
+				OutBufferInfo.bIsConfig = JEnv->GetBooleanField(OutputInfo, FOutputBufferInfo_bIsConfig);
+				JEnv->DeleteLocalRef(OutputInfo);
+				return 0;
+			}
+		}
+		return 1;
+	}
+
+
+	//-----------------------------------------------------------------------------
+	/**
+	 * Returns the decoded samples from a decoder output buffer in the decoder native format!
+	 *
+	 * @param OutBufferDataPtr
+	 * @param OutBufferDataSize
+	 * @param InOutBufferInfo
+	 *
+	 * @return 0 on success, 1 on failure.
+	 */
+	int32 FAndroidJavaH264VideoDecoder::GetOutputBuffer(void*& OutBufferDataPtr, int32 OutBufferDataSize, const FOutputBufferInfo& InOutBufferInfo)
+	{
+		FMediaCriticalSection::ScopedLock Lock(MutexLock);
+		if (bHaveDecoder)
+		{
+			if (InOutBufferInfo.ValidCount == CurrentValidCount)
+			{
+				JNIEnv* JEnv = AndroidJavaEnv::GetJavaEnv();
+				jbyteArray RawDecoderArray = (jbyteArray)JEnv->CallObjectMethod(Object, GetOutputBufferFN.Method, InOutBufferInfo.BufferIndex);
+				if (JEnv->ExceptionCheck())
+				{
+					JEnv->ExceptionDescribe();
+					JEnv->ExceptionClear();
+				}
+
+				if (RawDecoderArray != nullptr && OutBufferDataPtr != nullptr)
+				{
+					jbyte* RawDataPtr = JEnv->GetByteArrayElements(RawDecoderArray, 0);
+					int32 RawBufferSize = JEnv->GetArrayLength(RawDecoderArray);
+
+					check(RawBufferSize == InOutBufferInfo.Size);
+					check(RawBufferSize <= OutBufferDataSize)
+						if (RawDataPtr != nullptr)
+						{
+							memcpy(OutBufferDataPtr, RawDataPtr, RawBufferSize <= OutBufferDataSize ? RawBufferSize : OutBufferDataSize);
+							JEnv->ReleaseByteArrayElements(RawDecoderArray, RawDataPtr, JNI_ABORT);
+						}
+
+					JEnv->DeleteLocalRef(RawDecoderArray);
+				}
+				return 0;
+			}
+			else
+			{
+				// Decoder got flushed, buffer index no longer valid.
+				return 1;
+			}
+		}
+		return 1;
+	}
+
+
+	//-----------------------------------------------------------------------------
+	/**
+	 * Releases the decoder output buffer back to the decoder.
+	 *
+	 * @param InOutBufferInfo
+	 * @param bRender
+	 * @param releaseAt
+	 *
+	 * @return 0 on success, 1 on failure.
+	 */
+	int32 FAndroidJavaH264VideoDecoder::ReleaseOutputBuffer(int32 BufferIndex, int32 ValidCount, bool bRender, int64 releaseAt)
+	{
+		FMediaCriticalSection::ScopedLock Lock(MutexLock);
+		if (bHaveDecoder)
+		{
+			// Still same decoder instance?
+			if (ValidCount == CurrentValidCount)
+			{
+				// Yes...
+				JNIEnv* JEnv = AndroidJavaEnv::GetJavaEnv();
+				int32 result = CallMethod<int>(ReleaseOutputBufferFN, BufferIndex, bRender, (long)releaseAt);
+				if (JEnv->ExceptionCheck())
+				{
+					JEnv->ExceptionDescribe();
+					JEnv->ExceptionClear();
+				}
+				return result ? 1 : 0;
+			}
+			else
+			{
+				// Decoder got flushed, buffer index no longer valid.
+				// This is NOT an error!
+				return 0;
+			}
+		}
+		return 1;
+
+	}
+
+} // namespace Electra
+
+#endif
+
+#endif
diff --git a/Engine/Plugins/Media/ElectraPlayer/Source/ElectraPlayerRuntime/Private/Runtime/Decoder/Android/VideoDecoderResourceDelegate.h b/Engine/Plugins/Media/ElectraPlayer/Source/ElectraPlayerRuntime/Private/Runtime/Decoder/Android/VideoDecoderResourceDelegate.h
new file mode 100644
--- /dev/null
+++ b/Engine/Plugins/Media/ElectraPlayer/Source/ElectraPlayerRuntime/Private/Runtime/Decoder/Android/VideoDecoderResourceDelegate.h
@@ -0,0 +1,20 @@
+// Copyright Epic Games, Inc. All Rights Reserved.
+
+#pragma once
+
+#include "Android/AndroidPlatform.h"
+#include "Android/AndroidJava.h"
+#include "Android/AndroidApplication.h"
+
+namespace Electra
+{
+
+	class IVideoDecoderResourceDelegate
+	{
+	public:
+		virtual ~IVideoDecoderResourceDelegate() {}
+
+		virtual jobject VideoDecoderResourceDelegate_GetCodecSurface() = 0;
+	};
+
+}
