ElectraPlayer: RELEASE. Moved from NFL to public. Note that this is done via a Delete and "Add New" operation as discussed.

diff --git a/Engine/Plugins/Media/ElectraPlayer/Source/ElectraPlayerRuntime/Private/Runtime/Decoder/Apple/AudioDecoderAAC_Apple.cpp b/Engine/Plugins/Media/ElectraPlayer/Source/ElectraPlayerRuntime/Private/Runtime/Decoder/Apple/AudioDecoderAAC_Apple.cpp
new file mode 100644
--- /dev/null
+++ b/Engine/Plugins/Media/ElectraPlayer/Source/ElectraPlayerRuntime/Private/Runtime/Decoder/Apple/AudioDecoderAAC_Apple.cpp
@@ -0,0 +1,994 @@
+// Copyright Epic Games, Inc. All Rights Reserved.
+
+#include "HAL/Platform.h"
+
+#if PLATFORM_MAC || PLATFORM_IOS || PLATFORM_TVOS
+
+#include "PlayerCore.h"
+
+#include "StreamAccessUnitBuffer.h"
+#include "Decoder/AudioDecoderAAC.h"
+#include "Renderer/RendererBase.h"
+#include "Player/PlayerSessionServices.h"
+#include "Utilities/Utilities.h"
+#include "Utilities/UtilsMPEGAudio.h"
+#include "Utilities/StringHelpers.h"
+#include "DecoderErrors_Apple.h"
+#include "HAL/LowLevelMemTracker.h"
+
+#include <AudioToolbox/AudioToolbox.h>
+
+
+#define MAGIC_INSUFFICIENT_INPUT_MARKER		4711
+
+namespace Electra
+{
+
+
+/**
+ * AAC audio decoder class implementation.
+**/
+class FAudioDecoderAAC : public IAudioDecoderAAC, public FMediaThread
+{
+public:
+	static bool Startup(const IAudioDecoderAAC::FSystemConfiguration& InConfig);
+	static void Shutdown();
+
+	FAudioDecoderAAC();
+	virtual ~FAudioDecoderAAC();
+
+	virtual void SetPlayerSessionServices(IPlayerSessionServices* SessionServices) override;
+
+	virtual void Open(const FInstanceConfiguration& InConfig) override;
+	virtual void Close() override;
+
+	virtual void SetRenderer(TSharedPtr<IMediaRenderer, ESPMode::ThreadSafe> InRenderer) override;
+
+	virtual void SetAUInputBufferListener(IAccessUnitBufferListener* Listener) override;
+
+	virtual void SetReadyBufferListener(IDecoderOutputBufferListener* Listener) override;
+
+	virtual IAccessUnitBufferInterface::EAUpushResult AUdataPushAU(FAccessUnit* InAccessUnit) override;
+	virtual void AUdataPushEOD() override;
+	virtual void AUdataFlushEverything() override;
+
+private:
+	void StartThread();
+	void StopThread();
+	void WorkerThread();
+
+	bool InternalDecoderCreate();
+	void InternalDecoderDestroy();
+
+	bool CreateDecodedSamplePool();
+	void DestroyDecodedSamplePool();
+
+	void ReturnUnusedOutputBuffer();
+
+	void NotifyReadyBufferListener(bool bHaveOutput);
+
+	bool Decode(FAccessUnit* InAccessUnit);
+	bool IsDifferentFormat(const FAccessUnit* Data);
+
+	void PostError(OSStatus ApiReturnValue, const FString& Message, uint16 Code, UEMediaError Error = UEMEDIA_ERROR_OK);
+	void LogMessage(IInfoLog::ELevel Level, const FString& Message);
+
+private:
+
+	struct FAudioConverterInstance
+	{
+		FAudioConverterInstance()
+		{
+			ConverterRef = nullptr;
+			InputDescr   = {};
+			OutputDescr  = {};
+			PacketDescr  = {};
+		}
+		struct FWorkData
+		{
+			FWorkData()
+			{
+				Clear();
+			}
+			void Clear()
+			{
+				InputData     = nullptr;
+				InputSize     = 0;
+				InputConsumed = 0;
+			}
+			void*	InputData;
+			uint32	InputSize;
+			uint32	InputConsumed;
+		};
+		AudioConverterRef				ConverterRef;
+		AudioStreamBasicDescription		InputDescr;
+		AudioStreamBasicDescription		OutputDescr;
+		AudioStreamPacketDescription	PacketDescr;
+		FWorkData						WorkData;
+
+		// Methods for Audio Toolbox
+		OSStatus AudioToolbox_ComplexInputCallback(AudioConverterRef InAudioConverter, UInt32* InOutNumberDataPackets, AudioBufferList* InOutData, AudioStreamPacketDescription** OutDataPacketDescription);
+		static OSStatus _AudioToolbox_ComplexInputCallback(AudioConverterRef InAudioConverter, UInt32* InOutNumberDataPackets, AudioBufferList* InOutData, AudioStreamPacketDescription** OutDataPacketDescription, void* InUserData)
+		{
+			return static_cast<FAudioDecoderAAC::FAudioConverterInstance*>(InUserData)->AudioToolbox_ComplexInputCallback(InAudioConverter, InOutNumberDataPackets, InOutData, OutDataPacketDescription);
+		}
+	};
+
+private:
+	FInstanceConfiguration													Config;
+
+	FAccessUnitBuffer														AccessUnits;
+
+	FMediaEvent																TerminateThreadSignal;
+	FMediaEvent																FlushDecoderSignal;
+	FMediaEvent																DecoderFlushedSignal;
+	bool																	bThreadStarted;
+
+	TSharedPtr<IMediaRenderer, ESPMode::ThreadSafe>							Renderer;
+	int32																	MaxDecodeBufferSize;
+
+	FMediaCriticalSection													ListenerMutex;
+	IAccessUnitBufferListener*												InputBufferListener;
+	IDecoderOutputBufferListener*											ReadyBufferListener;
+
+	IPlayerSessionServices* 												SessionServices;
+
+	TSharedPtr<MPEG::FAACDecoderConfigurationRecord, ESPMode::ThreadSafe>	ConfigRecord;
+	TSharedPtrTS<FAccessUnit::CodecData>									CurrentCodecData;
+	bool																	bInDummyDecodeMode;
+	bool																	bHaveDiscontinuity;
+
+	FAudioConverterInstance*												DecoderInstance;
+
+	IMediaRenderer::IBuffer*												CurrentOutputBuffer;
+	FParamDict																BufferAcquireOptions;
+	FParamDict																OutputBufferSampleProperties;
+	int32																	PCMBufferSize;
+	int16*																	PCMBuffer;
+
+public:
+	static FSystemConfiguration												SystemConfig;
+};
+
+IAudioDecoderAAC::FSystemConfiguration	FAudioDecoderAAC::SystemConfig;
+
+
+/***************************************************************************************************************************************************/
+/***************************************************************************************************************************************************/
+/***************************************************************************************************************************************************/
+
+IAudioDecoderAAC::FSystemConfiguration::FSystemConfiguration()
+{
+	ThreadConfig.Decoder.Priority 	   = TPri_Normal;
+	ThreadConfig.Decoder.StackSize	   = 65536;
+	ThreadConfig.Decoder.CoreAffinity    = -1;
+	// Not needed, but setting meaningful values anyway.
+	ThreadConfig.PassOn.Priority  	   = TPri_Normal;
+	ThreadConfig.PassOn.StackSize 	   = 32768;
+	ThreadConfig.PassOn.CoreAffinity     = -1;
+}
+
+IAudioDecoderAAC::FInstanceConfiguration::FInstanceConfiguration()
+	: ThreadConfig(FAudioDecoderAAC::SystemConfig.ThreadConfig)
+{
+}
+
+bool IAudioDecoderAAC::Startup(const IAudioDecoderAAC::FSystemConfiguration& InConfig)
+{
+	return FAudioDecoderAAC::Startup(InConfig);
+}
+
+void IAudioDecoderAAC::Shutdown()
+{
+	FAudioDecoderAAC::Shutdown();
+}
+
+IAudioDecoderAAC* IAudioDecoderAAC::Create()
+{
+	return new FAudioDecoderAAC;
+}
+
+/***************************************************************************************************************************************************/
+/***************************************************************************************************************************************************/
+/***************************************************************************************************************************************************/
+
+
+//-----------------------------------------------------------------------------
+/**
+ * Decoder system startup
+ *
+ * @param config
+ *
+ * @return
+ */
+bool FAudioDecoderAAC::Startup(const IAudioDecoderAAC::FSystemConfiguration& config)
+{
+	SystemConfig = config;
+	return true;
+}
+
+
+//-----------------------------------------------------------------------------
+/**
+ * Decoder system shutdown.
+ */
+void FAudioDecoderAAC::Shutdown()
+{
+}
+
+
+//-----------------------------------------------------------------------------
+/**
+ * Constructor
+ */
+FAudioDecoderAAC::FAudioDecoderAAC()
+	: FMediaThread("ElectraPlayer::AAC decoder")
+	, bThreadStarted(false)
+	, Renderer(nullptr)
+	, MaxDecodeBufferSize(0)
+	, InputBufferListener(nullptr)
+	, ReadyBufferListener(nullptr)
+	, SessionServices(nullptr)
+	, bInDummyDecodeMode(false)
+	, bHaveDiscontinuity(false)
+	, DecoderInstance(nullptr)
+	, CurrentOutputBuffer(nullptr)
+	, PCMBufferSize(0)
+	, PCMBuffer(nullptr)
+{
+}
+
+
+//-----------------------------------------------------------------------------
+/**
+ * Destructor
+ */
+FAudioDecoderAAC::~FAudioDecoderAAC()
+{
+	Close();
+}
+
+
+//-----------------------------------------------------------------------------
+/**
+ * Sets the owning player's session service interface.
+ *
+ * @param InSessionServices
+ *
+ * @return
+ */
+void FAudioDecoderAAC::SetPlayerSessionServices(IPlayerSessionServices* InSessionServices)
+{
+	SessionServices = InSessionServices;
+}
+
+
+//-----------------------------------------------------------------------------
+/**
+ * Creates the decoded sample blocks output buffers.
+ *
+ * @return
+ */
+bool FAudioDecoderAAC::CreateDecodedSamplePool()
+{
+	check(Renderer);
+	FParamDict poolOpts;
+	// Assume LC-AAC with at most 8 channels and 16 bit decoded PCM samples. This is ok for HE-AAC as well which can at most be stereo.
+	uint32 frameSize = sizeof(int16) * 8 * 1024;
+
+	poolOpts.Set("max_buffer_size",  FVariantValue((int64) frameSize));
+	poolOpts.Set("num_buffers",      FVariantValue((int64) 8));
+
+	UEMediaError Error = Renderer->CreateBufferPool(poolOpts);
+	check(Error == UEMEDIA_ERROR_OK);
+
+	if (Error != UEMEDIA_ERROR_OK)
+	{
+		PostError(0, "Failed to create sample pool", ERRCODE_INTERNAL_APPLE_COULD_NOT_CREATE_SAMPLE_POOL, Error);
+	}
+
+	MaxDecodeBufferSize = (int32) Renderer->GetBufferPoolProperties().GetValue("max_buffers").GetInt64();
+
+	return Error == UEMEDIA_ERROR_OK;
+}
+
+
+//-----------------------------------------------------------------------------
+/**
+ * Destroys the pool of decoded sample blocks.
+ */
+void FAudioDecoderAAC::DestroyDecodedSamplePool()
+{
+	Renderer->ReleaseBufferPool();
+}
+
+
+//-----------------------------------------------------------------------------
+/**
+ * Opens a decoder instance
+ *
+ * @param InConfig
+ *
+ * @return
+ */
+void FAudioDecoderAAC::Open(const IAudioDecoderAAC::FInstanceConfiguration& InConfig)
+{
+	Config = InConfig;
+
+	// Set a large enough size to hold a single access unit. Since we will be asking for a new AU on
+	// demand there is no need to overly restrict ourselves here. But it must be large enough to
+	// hold at least the largest expected access unit.
+	AccessUnits.CapacitySet(FAccessUnitBuffer::FConfiguration(2 << 20, 60.0));
+
+	StartThread();
+}
+
+
+//-----------------------------------------------------------------------------
+/**
+ * Closes the decoder instance.
+ */
+void FAudioDecoderAAC::Close()
+{
+	StopThread();
+}
+
+
+//-----------------------------------------------------------------------------
+/**
+ * Sets a new renderer.
+ *
+ * @param InRenderer
+ */
+void FAudioDecoderAAC::SetRenderer(TSharedPtr<IMediaRenderer, ESPMode::ThreadSafe> InRenderer)
+{
+	Renderer = InRenderer;
+}
+
+
+//-----------------------------------------------------------------------------
+/**
+ * Sets an AU input buffer listener.
+ *
+ * @param InListener
+ */
+void FAudioDecoderAAC::SetAUInputBufferListener(IAccessUnitBufferListener* InListener)
+{
+	FMediaCriticalSection::ScopedLock lock(ListenerMutex);
+	InputBufferListener = InListener;
+}
+
+
+//-----------------------------------------------------------------------------
+/**
+ * Sets a buffer-ready listener.
+ *
+ * @param InListener
+ */
+void FAudioDecoderAAC::SetReadyBufferListener(IDecoderOutputBufferListener* InListener)
+{
+	FMediaCriticalSection::ScopedLock lock(ListenerMutex);
+	ReadyBufferListener = InListener;
+}
+
+
+//-----------------------------------------------------------------------------
+/**
+ * Creates and runs the decoder thread.
+ */
+void FAudioDecoderAAC::StartThread()
+{
+	ThreadSetPriority(Config.ThreadConfig.Decoder.Priority);
+	ThreadSetStackSize(Config.ThreadConfig.Decoder.StackSize);
+	ThreadSetCoreAffinity(Config.ThreadConfig.Decoder.CoreAffinity);
+	ThreadStart(Electra::MakeDelegate(this, &FAudioDecoderAAC::WorkerThread));
+	bThreadStarted = true;
+}
+
+
+//-----------------------------------------------------------------------------
+/**
+ * Stops the decoder thread.
+ */
+void FAudioDecoderAAC::StopThread()
+{
+	if (bThreadStarted)
+	{
+		TerminateThreadSignal.Signal();
+		ThreadWaitDone();
+		bThreadStarted = false;
+	}
+}
+
+
+//-----------------------------------------------------------------------------
+/**
+ * Called to receive a new input access unit for decoding.
+ *
+ * @param InAccessUnit
+ */
+IAccessUnitBufferInterface::EAUpushResult FAudioDecoderAAC::AUdataPushAU(FAccessUnit* InAccessUnit)
+{
+	// Add a ref count to the AU before adding it to the buffer. If it can be added successfully it could be
+	// used immediately by the worker thread and it would be too late for us to increase the count then.
+	InAccessUnit->AddRef();
+	bool bOk = AccessUnits.Push(InAccessUnit);
+	if (!bOk)
+	{
+		// Could not add. Undo the refcount increase.
+		FAccessUnit::Release(InAccessUnit);
+	}
+	return bOk ? IAccessUnitBufferInterface::EAUpushResult::Ok : IAccessUnitBufferInterface::EAUpushResult::Full;
+}
+
+
+//-----------------------------------------------------------------------------
+/**
+ * "Pushes" an End Of Data marker indicating no further access units will be added.
+ */
+void FAudioDecoderAAC::AUdataPushEOD()
+{
+	AccessUnits.PushEndOfData();
+}
+
+
+//-----------------------------------------------------------------------------
+/**
+ * Flushes the decoder and clears the input access unit buffer.
+ */
+void FAudioDecoderAAC::AUdataFlushEverything()
+{
+	FlushDecoderSignal.Signal();
+	DecoderFlushedSignal.WaitAndReset();
+}
+
+
+//-----------------------------------------------------------------------------
+/**
+ * Notifies the buffer-ready listener that we will now be producing output data
+ * and need an output buffer.
+ * This call is not intended to create or obtain an output buffer. It is merely
+ * indicating that new output will be produced.
+ *
+ * @param bHaveOutput
+ */
+void FAudioDecoderAAC::NotifyReadyBufferListener(bool bHaveOutput)
+{
+	if (ReadyBufferListener)
+	{
+		IDecoderOutputBufferListener::FDecodeReadyStats stats;
+		stats.MaxDecodedElementsReady = MaxDecodeBufferSize;
+		stats.NumElementsInDecoder    = CurrentOutputBuffer ? 1 : 0;
+		stats.bOutputStalled		  = !bHaveOutput;
+		stats.bEODreached   		  = AccessUnits.IsEndOfData() && stats.NumDecodedElementsReady == 0 && stats.NumElementsInDecoder == 0;
+		ListenerMutex.Lock();
+		if (ReadyBufferListener)
+		{
+			ReadyBufferListener->DecoderOutputReady(stats);
+		}
+		ListenerMutex.Unlock();
+	}
+}
+
+
+//-----------------------------------------------------------------------------
+/**
+ * Posts an error to the session service error listeners.
+ *
+ * @param ApiReturnValue
+ * @param Message
+ * @param Code
+ * @param Error
+ */
+void FAudioDecoderAAC::PostError(OSStatus ApiReturnValue, const FString& Message, uint16 Code, UEMediaError Error)
+{
+	check(SessionServices);	// there better be a session service interface to receive the error!
+	if (SessionServices)
+	{
+		FErrorDetail err;
+		err.SetError(Error != UEMEDIA_ERROR_OK ? Error : UEMEDIA_ERROR_DETAIL);
+		err.SetFacility(Facility::EFacility::AACDecoder);
+		err.SetCode(Code);
+		err.SetMessage(Message);
+		err.SetPlatformMessage(FString::Printf(TEXT("%d (0x%08x)"), (int32) ApiReturnValue, (int32) ApiReturnValue));
+		SessionServices->PostError(err);
+	}
+}
+
+
+//-----------------------------------------------------------------------------
+/**
+ * Sends a log message to the session service log.
+ *
+ * @param Level
+ * @param Message
+ */
+void FAudioDecoderAAC::LogMessage(IInfoLog::ELevel Level, const FString& Message)
+{
+	if (SessionServices)
+	{
+		SessionServices->PostLog(Facility::EFacility::AACDecoder, Level, Message);
+	}
+}
+
+
+//-----------------------------------------------------------------------------
+/**
+ * Creates an audio decoder ("converter" in AudioToolBox terminology)
+ *
+ * Note: This requires the configuration record to have been parsed successfully.
+ */
+bool FAudioDecoderAAC::InternalDecoderCreate()
+{
+	InternalDecoderDestroy();
+
+	if (!ConfigRecord.IsValid())
+	{
+		PostError(0, "No CSD to create audio decoder with", ERRCODE_INTERNAL_APPLE_COULD_NOT_CREATE_AUDIO_DECODER);
+		return false;
+	}
+
+	DecoderInstance = new FAudioConverterInstance;
+	// Set input configuration specific to our current format.
+	memset(&DecoderInstance->InputDescr, 0, sizeof(AudioStreamBasicDescription));
+	DecoderInstance->InputDescr.mSampleRate		  = ConfigRecord->ExtSamplingFrequency ? ConfigRecord->ExtSamplingFrequency : ConfigRecord->SamplingRate;
+	DecoderInstance->InputDescr.mFormatID		  = ConfigRecord->PSSignal > 0 ? kAudioFormatMPEG4AAC_HE_V2 : ConfigRecord->SBRSignal > 0 ? kAudioFormatMPEG4AAC_HE : kAudioFormatMPEG4AAC;
+	DecoderInstance->InputDescr.mFormatFlags      = DecoderInstance->InputDescr.mFormatID == kAudioFormatMPEG4AAC ? kMPEG4Object_AAC_LC : kMPEG4Object_AAC_SBR;
+	DecoderInstance->InputDescr.mFramesPerPacket  = DecoderInstance->InputDescr.mFormatID == kAudioFormatMPEG4AAC ? 1024 : 2048;
+    // Decoding parametric stereo (PS) implies the output to go from 1 to 2 channels.
+	DecoderInstance->InputDescr.mChannelsPerFrame = ConfigRecord->PSSignal > 0 ? 2 : ConfigRecord->ChannelConfiguration;
+
+	// Want LPCM output, use convenience method to set this up.
+	memset(&DecoderInstance->OutputDescr, 0, sizeof(AudioStreamBasicDescription));
+	FillOutASBDForLPCM(DecoderInstance->OutputDescr, DecoderInstance->InputDescr.mSampleRate, DecoderInstance->InputDescr.mChannelsPerFrame, 16, 16, false, false, false);
+	// Try to create audio converter
+	OSStatus res = AudioConverterNew(&DecoderInstance->InputDescr, &DecoderInstance->OutputDescr, &DecoderInstance->ConverterRef);
+	check(res == 0);
+	if (res)
+	{
+		InternalDecoderDestroy();
+		PostError(res, "Failed to create decoder", ERRCODE_INTERNAL_APPLE_COULD_NOT_CREATE_AUDIO_DECODER);
+		return false;
+	}
+	return true;
+}
+
+
+//-----------------------------------------------------------------------------
+/**
+ * Destroys the audio decoder ("converter" in AudioToolBox terminology)
+ */
+void FAudioDecoderAAC::InternalDecoderDestroy()
+{
+	if (DecoderInstance)
+	{
+		OSStatus res = 0;
+		if (DecoderInstance->ConverterRef)
+		{
+			res = AudioConverterReset(DecoderInstance->ConverterRef);
+			check(res == 0);	// not really that important
+			res = AudioConverterDispose(DecoderInstance->ConverterRef);
+			check(res == 0);
+			DecoderInstance->ConverterRef = nullptr;
+		}
+		delete DecoderInstance;
+		DecoderInstance = nullptr;
+	}
+}
+
+
+//-----------------------------------------------------------------------------
+/**
+ * Returns the acquired sample output buffer back to the renderer without having it rendered.
+ */
+void FAudioDecoderAAC::ReturnUnusedOutputBuffer()
+{
+	if (CurrentOutputBuffer)
+	{
+		OutputBufferSampleProperties.Clear();
+		Renderer->ReturnBuffer(CurrentOutputBuffer, false, OutputBufferSampleProperties);
+		CurrentOutputBuffer = nullptr;
+	}
+}
+
+
+//-----------------------------------------------------------------------------
+/**
+ * Callback from the AudioConverter to get input to be converted.
+ *
+ * @param InAudioConverter
+ * @param InOutNumberDataPackets
+ * @param InOutData
+ * @param OutDataPacketDescription
+ */
+OSStatus FAudioDecoderAAC::FAudioConverterInstance::AudioToolbox_ComplexInputCallback(AudioConverterRef InAudioConverter, UInt32* InOutNumberDataPackets, AudioBufferList* InOutData, AudioStreamPacketDescription** OutDataPacketDescription)
+{
+	// Make sure we have input work data when being called
+	if (WorkData.InputData)
+	{
+		if (InOutNumberDataPackets && *InOutNumberDataPackets)
+		{
+			*InOutNumberDataPackets   		     = 1;
+			InOutData->mNumberBuffers 		     = 1;
+			InOutData->mBuffers[0].mData  	     = WorkData.InputData;
+			InOutData->mBuffers[0].mDataByteSize = WorkData.InputSize;
+
+			if (OutDataPacketDescription)
+			{
+				check(*OutDataPacketDescription == nullptr);
+				*OutDataPacketDescription = &PacketDescr;
+				PacketDescr.mStartOffset 		    = 0;
+				PacketDescr.mVariableFramesInPacket = 0;
+				PacketDescr.mDataByteSize		    = WorkData.InputSize;
+			}
+			WorkData.InputConsumed = WorkData.InputSize;
+			WorkData.InputData     = nullptr;
+			WorkData.InputSize     = 0;
+		}
+		return 0;
+	}
+	else
+	{
+		// When we have no pending input access unit we can't return data.
+		// We return some dummy value to indicate that we have no data at the moment.
+		// Returning 0 would indicate an EOF and make the decoder stop working.
+		*InOutNumberDataPackets = 0;
+		return MAGIC_INSUFFICIENT_INPUT_MARKER;
+	}
+}
+
+
+//-----------------------------------------------------------------------------
+/**
+ * Decodes an access unit
+ */
+bool FAudioDecoderAAC::Decode(FAccessUnit* InAccessUnit)
+{
+	FTimeValue CurrentPTS;
+	CurrentPTS = InAccessUnit->PTS;
+
+	uint32 nDataOffset = 0;
+	// Loop until all data has been consumed.
+	bool bAllInputConsumed = false;
+	while(!bAllInputConsumed && !TerminateThreadSignal.IsSignaled() && !FlushDecoderSignal.IsSignaled())
+	{
+		// Need a new output buffer?
+		if (CurrentOutputBuffer == nullptr)
+		{
+			UEMediaError bufResult = Renderer->AcquireBuffer(CurrentOutputBuffer, 0, BufferAcquireOptions);
+			check(bufResult == UEMEDIA_ERROR_OK || bufResult == UEMEDIA_ERROR_INSUFFICIENT_DATA);
+			if (bufResult != UEMEDIA_ERROR_OK && bufResult != UEMEDIA_ERROR_INSUFFICIENT_DATA)
+			{
+				PostError(0, "Failed to acquire sample buffer", ERRCODE_INTERNAL_APPLE_COULD_NOT_GET_SAMPLE_BUFFER, bufResult);
+				return false;
+			}
+		}
+
+		bool bHaveAvailSmpBlk = CurrentOutputBuffer != nullptr;
+
+		// Check if the renderer can accept the output we will want to send to it.
+		// If it can't right now we treat this as if we do not have an available output buffer.
+		if (Renderer.IsValid() && !Renderer->CanReceiveOutputFrames(1))
+		{
+			bHaveAvailSmpBlk = false;
+		}
+
+		NotifyReadyBufferListener(bHaveAvailSmpBlk);
+		if (bHaveAvailSmpBlk)
+		{
+			if (DecoderInstance && !InAccessUnit->bIsDummyData)
+			{
+				if (bInDummyDecodeMode)
+				{
+					bHaveDiscontinuity = true;
+					bInDummyDecodeMode = false;
+				}
+
+				DecoderInstance->WorkData.Clear();
+				DecoderInstance->WorkData.InputData = Electra::AdvancePointer(InAccessUnit->AUData, nDataOffset);
+				DecoderInstance->WorkData.InputSize = InAccessUnit->AUSize - nDataOffset;
+				int32 CurrentRenderOutputBufferSize = (int32)CurrentOutputBuffer->GetBufferProperties().GetValue("size").GetInt64();
+				void* CurrentRenderOutputBufferAddress = CurrentOutputBuffer->GetBufferProperties().GetValue("address").GetPointer();
+				int32 CurrentRenderOutputBufferOffset = 0;
+				int32 NumSamplesProduced = 0;
+				uint32 NumberOfChannels   = DecoderInstance->OutputDescr.mChannelsPerFrame;
+				uint32 SamplingRate       = DecoderInstance->InputDescr.mSampleRate;
+				while(1)
+				{
+					AudioBuffer		OutputBuffer;
+					AudioBufferList OutputBufferList;
+					OutputBuffer.mNumberChannels    = DecoderInstance->OutputDescr.mChannelsPerFrame;
+					OutputBuffer.mDataByteSize	    = PCMBufferSize;
+					OutputBuffer.mData			    = PCMBuffer;
+					OutputBufferList.mNumberBuffers = 1;
+					OutputBufferList.mBuffers[0]    = OutputBuffer;
+
+					// Get all possible output in LC size blocks. This may return fewer samples or we loop a number of times here.
+					UInt32 InOutPacketSize = 1024;
+
+					OSStatus res = AudioConverterFillComplexBuffer(DecoderInstance->ConverterRef, FAudioConverterInstance::_AudioToolbox_ComplexInputCallback, DecoderInstance, &InOutPacketSize, &OutputBufferList, nullptr);
+					if (res != 0 && res != MAGIC_INSUFFICIENT_INPUT_MARKER)
+					{
+						PostError(res, "Failed to decode", ERRCODE_INTERNAL_APPLE_FAILED_TO_DECODE_AUDIO);
+						return false;
+					}
+
+					if (InOutPacketSize)
+					{
+						NumSamplesProduced += InOutPacketSize;
+
+						int32 OutputByteCount = InOutPacketSize * NumberOfChannels * sizeof(int16);
+						check(OutputByteCount + CurrentRenderOutputBufferOffset <= CurrentRenderOutputBufferSize);
+						if (OutputByteCount + CurrentRenderOutputBufferOffset <= CurrentRenderOutputBufferSize)
+						{
+							FMemory::Memcpy(Electra::AdvancePointer(CurrentRenderOutputBufferAddress, CurrentRenderOutputBufferOffset), PCMBuffer, OutputByteCount);
+							CurrentRenderOutputBufferOffset += OutputByteCount;
+						}
+						else
+						{
+							PostError(0, "Output buffer not large enough", ERRCODE_INTERNAL_APPLE_AUDIO_OUTPUT_BUFFER_TOO_SMALL);
+							return false;
+						}
+					}
+					else
+					{
+						// All output exhausted. This should happen only when all input was consumed as well.
+						check(res == MAGIC_INSUFFICIENT_INPUT_MARKER);
+						break;
+					}
+
+					if (res == MAGIC_INSUFFICIENT_INPUT_MARKER)
+					{
+						break;
+					}
+
+				}
+				nDataOffset += DecoderInstance->WorkData.InputConsumed;
+				bAllInputConsumed = nDataOffset >= InAccessUnit->AUSize;
+
+				if (NumSamplesProduced)
+				{
+					FTimeValue Duration;
+					Duration.SetFromND(NumSamplesProduced, SamplingRate);
+
+					OutputBufferSampleProperties.Clear();
+					OutputBufferSampleProperties.Set("num_channels",  FVariantValue((int64) NumberOfChannels));
+					OutputBufferSampleProperties.Set("sample_rate",   FVariantValue((int64) SamplingRate));
+					OutputBufferSampleProperties.Set("byte_size",     FVariantValue((int64) CurrentRenderOutputBufferOffset));
+					OutputBufferSampleProperties.Set("duration",      FVariantValue(Duration));
+					OutputBufferSampleProperties.Set("pts",           FVariantValue(CurrentPTS));
+					OutputBufferSampleProperties.Set("discontinuity", FVariantValue((bool)bHaveDiscontinuity));
+
+					Renderer->ReturnBuffer(CurrentOutputBuffer, true, OutputBufferSampleProperties);
+					CurrentOutputBuffer = nullptr;
+					bHaveDiscontinuity = false;
+
+					CurrentPTS += Duration;
+				}
+
+			}
+			else
+			{
+				// Check if we are in dummy decode mode already.
+				if (!bInDummyDecodeMode)
+				{
+					bInDummyDecodeMode = true;
+					bHaveDiscontinuity = true;
+				}
+
+				// Clear to silence
+				FMemory::Memzero(CurrentOutputBuffer->GetBufferProperties().GetValue("address").GetPointer(), CurrentOutputBuffer->GetBufferProperties().GetValue("size").GetInt64());
+
+				// Assume sensible defaults.
+				int64 NumChannels = 2;
+				int64 SampleRate = 48000;
+				int64 SamplesPerBlock = 1024;
+
+				// With a valid configuration record we can use the actual values.
+				if (ConfigRecord.IsValid())
+				{
+					NumChannels = ConfigRecord->PSSignal > 0 ? 2 : ConfigRecord->ChannelConfiguration;
+					SampleRate = ConfigRecord->ExtSamplingFrequency ? ConfigRecord->ExtSamplingFrequency : ConfigRecord->SamplingRate;
+					SamplesPerBlock = ConfigRecord->SBRSignal > 0 ? 2048 : 1024;
+				}
+
+				FTimeValue Duration;
+				Duration.SetFromND(SamplesPerBlock, (uint32) SampleRate);
+
+				// Note: The duration in the AU should be an exact multiple of the "SamplesPerBlock / SampleRate" as this is how it gets set
+				//       up in the stream reader. Should there ever be any discrepancy such that "Duration > pData->mDuration" it is simply
+				//       possible to recalculate "SamplesPerBlock" such that we return a buffer with less than the usual 1024 or 2048 samples.
+				OutputBufferSampleProperties.Clear();
+				OutputBufferSampleProperties.Set("num_channels",  FVariantValue(NumChannels));
+				OutputBufferSampleProperties.Set("sample_rate",   FVariantValue(SampleRate));
+				OutputBufferSampleProperties.Set("byte_size",     FVariantValue((int64)(NumChannels * sizeof(int16) * SamplesPerBlock)));
+				OutputBufferSampleProperties.Set("duration",      FVariantValue(Duration));
+				OutputBufferSampleProperties.Set("pts",           FVariantValue(CurrentPTS));
+				OutputBufferSampleProperties.Set("discontinuity", FVariantValue((bool)bHaveDiscontinuity));
+
+				Renderer->ReturnBuffer(CurrentOutputBuffer, true, OutputBufferSampleProperties);
+				CurrentOutputBuffer = nullptr;
+				bHaveDiscontinuity = false;
+
+				CurrentPTS += Duration;
+				InAccessUnit->Duration -= Duration;
+				bAllInputConsumed = InAccessUnit->Duration <= FTimeValue::GetZero();
+			}
+		}
+		else
+		{
+			// No available buffer. Sleep for a bit. Can't sleep on a signal since we have to check two: abort and flush
+			// We sleep for 20ms since a sample block for LC AAC produces 1024 samples which amount to 21.3ms @48kHz (or more at lower sample rates).
+			FMediaRunnable::SleepMilliseconds(20);
+		}
+	}
+	return true;
+}
+
+
+//-----------------------------------------------------------------------------
+/**
+ * Checks if the codec specific format has changed.
+ *
+ * @param Data
+ *
+ * @return
+ */
+bool FAudioDecoderAAC::IsDifferentFormat(const FAccessUnit* Data)
+{
+	if (Data->AUCodecData.IsValid() && (!CurrentCodecData.IsValid() || CurrentCodecData->CodecSpecificData != Data->AUCodecData->CodecSpecificData))
+	{
+		return true;
+	}
+	return false;
+}
+
+
+//-----------------------------------------------------------------------------
+/**
+ * AAC audio decoder main threaded decode loop
+ */
+void FAudioDecoderAAC::WorkerThread()
+{
+	LLM_SCOPE(ELLMTag::MediaStreaming);
+
+	FAccessUnit*	CurrentAccessUnit = nullptr;
+	bool			bError			  = false;
+
+	PCMBufferSize   = sizeof(int16) * 8 * 2048;			// max. of 8 channels decoding HE-AAC
+	PCMBuffer   	= (int16*)FMemory::Malloc(PCMBufferSize, 32);
+	CurrentOutputBuffer = nullptr;
+	bInDummyDecodeMode = false;
+	bHaveDiscontinuity = false;
+
+	bError = !CreateDecodedSamplePool();
+	check(!bError);
+
+	while(!TerminateThreadSignal.IsSignaled())
+	{
+		// Notify the buffer listener that we will now be needing an AU for our input buffer.
+		if (!bError && InputBufferListener && AccessUnits.Num() == 0)
+		{
+			FAccessUnitBufferInfo	sin;
+			IAccessUnitBufferListener::FBufferStats	stats;
+			AccessUnits.GetStats(sin);
+			stats.NumAUsAvailable  = sin.NumCurrentAccessUnits;
+			stats.NumBytesInBuffer = sin.CurrentMemInUse;
+			stats.MaxBytesOfBuffer = sin.MaxDataSize;
+			stats.bEODSignaled     = sin.bEndOfData;
+			stats.bEODReached      = sin.bEndOfData && sin.NumCurrentAccessUnits == 0;
+			ListenerMutex.Lock();
+			if (InputBufferListener)
+			{
+				InputBufferListener->DecoderInputNeeded(stats);
+			}
+			ListenerMutex.Unlock();
+		}
+		// Get the AU to be decoded if one is there.
+		bool bHaveData = AccessUnits.WaitForData(1000 * 10);
+		if (bHaveData)
+		{
+			AccessUnits.Pop(CurrentAccessUnit);
+
+			// Check if the format has changed such that we need to destroy and re-create the decoder.
+			if (CurrentAccessUnit->bTrackChangeDiscontinuity || IsDifferentFormat(CurrentAccessUnit) || (CurrentAccessUnit->bIsDummyData && !bInDummyDecodeMode))
+				{
+				if (DecoderInstance)
+				{
+					InternalDecoderDestroy();
+					bHaveDiscontinuity = true;
+				}
+				ReturnUnusedOutputBuffer();
+				ConfigRecord.Reset();
+				CurrentCodecData.Reset();
+			}
+
+			// Parse the CSD into a configuration record.
+			if (!ConfigRecord.IsValid() && !bError)
+			{
+				if (CurrentAccessUnit->AUCodecData.IsValid())
+				{
+					CurrentCodecData = CurrentAccessUnit->AUCodecData;
+					ConfigRecord = MakeShared<MPEG::FAACDecoderConfigurationRecord, ESPMode::ThreadSafe>();
+					if (!ConfigRecord->ParseFrom(CurrentCodecData->CodecSpecificData.GetData(), CurrentCodecData->CodecSpecificData.Num()))
+					{
+						ConfigRecord.Reset();
+						CurrentCodecData.Reset();
+						PostError(0, "Failed to parse AAC configuration record", ERRCODE_INTERNAL_APPLE_FAILED_TO_PARSE_CSD);
+						bError = true;
+					}
+				}
+			}
+
+			// Check if this audio packet is to be dropped
+			if (!CurrentAccessUnit->DropState)
+			{
+				// Need to create a decoder instance?
+				if (DecoderInstance == nullptr && !bError)
+				{
+					if (!InternalDecoderCreate())
+					{
+						bError = true;
+					}
+				}
+
+				// Decode if not in error, otherwise just spend some idle time.
+				if (!bError)
+				{
+					if (!Decode(CurrentAccessUnit))
+					{
+						bError = true;
+					}
+				}
+				else
+				{
+					// Pace ourselves in consuming any more data until we are being stopped.
+					FMediaRunnable::SleepMicroseconds(CurrentAccessUnit->Duration.GetAsMicroseconds());
+				}
+			}
+
+			// We're done with this access unit. Delete it.
+			FAccessUnit::Release(CurrentAccessUnit);
+			CurrentAccessUnit = nullptr;
+		}
+		else
+		{
+			// No data. Is the buffer at EOD?
+			if (AccessUnits.IsEndOfData())
+			{
+				NotifyReadyBufferListener(true);
+				FMediaRunnable::SleepMilliseconds(20);
+			}
+		}
+
+		// Flush?
+		if (FlushDecoderSignal.IsSignaled())
+		{
+			ReturnUnusedOutputBuffer();
+			AccessUnits.Flush();
+
+			FlushDecoderSignal.Reset();
+			DecoderFlushedSignal.Signal();
+		}
+	}
+
+	ReturnUnusedOutputBuffer();
+	// Close the decoder.
+	InternalDecoderDestroy();
+	DestroyDecodedSamplePool();
+
+	// Flush any remaining input data.
+	AccessUnits.Flush();
+	AccessUnits.CapacitySet(0);
+
+	CurrentCodecData.Reset();
+	ConfigRecord.Reset();
+
+	FMemory::Free(PCMBuffer);
+	}
+
+
+} // namespace Electra
+
+
+#endif
diff --git a/Engine/Plugins/Media/ElectraPlayer/Source/ElectraPlayerRuntime/Private/Runtime/Decoder/Apple/DecoderErrors_Apple.h b/Engine/Plugins/Media/ElectraPlayer/Source/ElectraPlayerRuntime/Private/Runtime/Decoder/Apple/DecoderErrors_Apple.h
new file mode 100644
--- /dev/null
+++ b/Engine/Plugins/Media/ElectraPlayer/Source/ElectraPlayerRuntime/Private/Runtime/Decoder/Apple/DecoderErrors_Apple.h
@@ -0,0 +1,23 @@
+// Copyright Epic Games, Inc. All Rights Reserved.
+
+#pragma once
+
+
+#define ERRCODE_INTERNAL_APPLE_OUT_OF_MEMORY						1
+#define ERRCODE_INTERNAL_APPLE_COULD_NOT_CREATE_AUDIO_DECODER		2
+#define ERRCODE_INTERNAL_APPLE_COULD_NOT_CREATE_SAMPLE_POOL			3
+#define ERRCODE_INTERNAL_APPLE_COULD_NOT_GET_SAMPLE_BUFFER			4
+#define ERRCODE_INTERNAL_APPLE_FAILED_TO_PARSE_CSD					5
+#define ERRCODE_INTERNAL_APPLE_FAILED_TO_DECODE_AUDIO				6
+#define ERRCODE_INTERNAL_APPLE_AUDIO_OUTPUT_BUFFER_TOO_SMALL		7
+
+#define ERRCODE_INTERNAL_APPLE_COULD_NOT_CREATE_VIDEO_DECODER		100
+#define ERRCODE_INTERNAL_APPLE_COULD_NOT_CREATE_IMAGE_POOL			101
+#define ERRCODE_INTERNAL_APPLE_COULD_NOT_GET_OUTPUT_BUFFER			102
+#define ERRCODE_INTERNAL_APPLE_FAILED_TO_DECODE_VIDEO				103
+#define ERRCODE_INTERNAL_APPLE_NO_VIDEO_CSD							104
+#define ERRCODE_INTERNAL_APPLE_BAD_VIDEO_CSD						105
+#define ERRCODE_INTERNAL_APPLE_FAILED_TO_CREATE_BLOCK_BUFFER		110
+#define ERRCODE_INTERNAL_APPLE_FAILED_TO_CREATE_SAMPLE_BUFFER		111
+
+
diff --git a/Engine/Plugins/Media/ElectraPlayer/Source/ElectraPlayerRuntime/Private/Runtime/Decoder/Apple/VideoDecoderH264_Apple.cpp b/Engine/Plugins/Media/ElectraPlayer/Source/ElectraPlayerRuntime/Private/Runtime/Decoder/Apple/VideoDecoderH264_Apple.cpp
new file mode 100644
--- /dev/null
+++ b/Engine/Plugins/Media/ElectraPlayer/Source/ElectraPlayerRuntime/Private/Runtime/Decoder/Apple/VideoDecoderH264_Apple.cpp
@@ -0,0 +1,1631 @@
+// Copyright Epic Games, Inc. All Rights Reserved.
+
+#include "HAL/Platform.h"
+
+#if PLATFORM_MAC || PLATFORM_IOS || PLATFORM_TVOS
+
+#include "PlayerCore.h"
+
+#include "StreamAccessUnitBuffer.h"
+#include "Decoder/VideoDecoderH264.h"
+#include "Renderer/RendererBase.h"
+#include "Player/PlayerSessionServices.h"
+#include "Utilities/StringHelpers.h"
+#include "Utilities/UtilsMPEGVideo.h"
+#include "DecoderErrors_Apple.h"
+#include "HAL/LowLevelMemTracker.h"
+#include "MediaVideoDecoderOutputApple.h"
+#include "Renderer/RendererVideo.h"
+
+#include <VideoToolbox/VideoToolbox.h>
+
+namespace Electra
+{
+
+const uint32 NumImagesHoldBackForPTSOrdering = 5; // Number of frames held back to ensure proper PTS-ordering of decoder output
+const uint32 MaxImagesHoldBackForPTSOrdering = 5; // Maximum number of frames to be held in the buffer before we stall the decoder
+
+/**
+ * H264 video decoder class implementation.
+**/
+class FVideoDecoderH264 : public IVideoDecoderH264, public FMediaThread
+{
+public:
+	static bool Startup(const IVideoDecoderH264::FSystemConfiguration& InConfig);
+	static void Shutdown();
+
+	FVideoDecoderH264();
+	virtual ~FVideoDecoderH264();
+
+	virtual void SetPlayerSessionServices(IPlayerSessionServices* SessionServices) override;
+
+	virtual void Open(const FInstanceConfiguration& InConfig) override;
+	virtual void Close() override;
+
+	virtual void SetMaximumDecodeCapability(int32 MaxWidth, int32 MaxHeight, int32 MaxProfile, int32 MaxProfileLevel, const FParamDict& AdditionalOptions) override;
+
+	virtual void SetAUInputBufferListener(IAccessUnitBufferListener* Listener) override;
+
+	virtual void SetReadyBufferListener(IDecoderOutputBufferListener* Listener) override;
+
+	virtual void SetRenderer(TSharedPtr<IMediaRenderer, ESPMode::ThreadSafe> InRenderer) override;
+    
+    virtual void SetResourceDelegate(const TSharedPtr<IVideoDecoderResourceDelegate, ESPMode::ThreadSafe>& ResourceDelegate) override;
+
+	virtual IAccessUnitBufferInterface::EAUpushResult AUdataPushAU(FAccessUnit* AccessUnit) override;
+	virtual void AUdataPushEOD() override;
+	virtual void AUdataFlushEverything() override;
+
+private:
+	bool InternalDecoderCreate(CMFormatDescriptionRef InputFormatDescription);
+	void InternalDecoderDestroy();
+	void RecreateDecoderSession();
+	void StartThread();
+	void StopThread();
+	void WorkerThread();
+
+	bool CreateDecodedImagePool();
+	void DestroyDecodedImagePool();
+
+	void NotifyReadyBufferListener(bool bHaveOutput);
+
+	bool AcquireOutputBuffer(IMediaRenderer::IBuffer*& RenderOutputBuffer);
+
+	void PrepareAU(FAccessUnit* AccessUnit, bool& bOutIsIDR, bool& bOutIsDiscardable);
+
+	struct FDecodedImage;
+	void ProcessOutput(bool bFlush = false);
+
+private:
+	bool FlushDecoder();
+	void ClearInDecoderInfos();
+	void FlushPendingImages();
+
+	enum EDecodeResult
+	{
+		Ok,
+		Fail,
+		SessionLost
+	};
+
+	EDecodeResult Decode(FAccessUnit* AccessUnit, bool bRecreatingSession);
+	bool DecodeDummy(FAccessUnit* AccessUnit);
+
+	void PostError(int32_t ApiReturnValue, const FString& Message, uint16 Code, UEMediaError Error = UEMEDIA_ERROR_OK);
+	void LogMessage(IInfoLog::ELevel Level, const FString& Message);
+
+private:
+	void DecodeCallback(void* pSrcRef, OSStatus status, VTDecodeInfoFlags infoFlags, CVImageBufferRef imageBuffer, CMTime presentationTimeStamp, CMTime presentationDuration);
+	static void _DecodeCallback(void* pUser, void* pSrcRef, OSStatus status, VTDecodeInfoFlags infoFlags, CVImageBufferRef imageBuffer, CMTime presentationTimeStamp, CMTime presentationDuration)
+	{
+		static_cast<FVideoDecoderH264*>(pUser)->DecodeCallback(pSrcRef, status, infoFlags, imageBuffer, presentationTimeStamp, presentationDuration);
+	}
+
+	bool CreateFormatDescription(CMFormatDescriptionRef& OutFormatDescription, const FAccessUnit* InputAccessUnit);
+
+	struct FDecoderFormatInfo
+	{
+		void Reset()
+		{
+			CurrentCodecData.Reset();
+		}
+		bool IsDifferentFrom(const FAccessUnit* InputAccessUnit);
+
+		TSharedPtrTS<const FAccessUnit::CodecData>		CurrentCodecData;
+	};
+
+	struct FDecoderHandle
+	{
+		FDecoderHandle()
+			: FormatDescription(nullptr)
+			, DecompressionSession(nullptr)
+		{
+		}
+		~FDecoderHandle()
+		{
+			Close();
+		}
+
+		void Close()
+		{
+			if (DecompressionSession)
+			{
+				// Note: Do not finish delayed frames here. We don't want them and have to destroy the session immediately!
+				//VTDecompressionSessionFinishDelayedFrames(DecompressionSession);
+				VTDecompressionSessionWaitForAsynchronousFrames(DecompressionSession);
+				VTDecompressionSessionInvalidate(DecompressionSession);
+				CFRelease(DecompressionSession);
+				DecompressionSession = nullptr;
+			}
+			if (FormatDescription)
+			{
+				CFRelease(FormatDescription);
+				FormatDescription = nullptr;
+			}
+		}
+
+		bool IsCompatibleWith(CMFormatDescriptionRef NewFormatDescription)
+		{
+			if (DecompressionSession)
+			{
+				Boolean bIsCompatible = VTDecompressionSessionCanAcceptFormatDescription(DecompressionSession, NewFormatDescription);
+				return bIsCompatible;
+			}
+			return false;
+		}
+		CMFormatDescriptionRef		FormatDescription;
+		VTDecompressionSessionRef	DecompressionSession;
+	};
+
+	struct FInDecoderInfo : public TSharedFromThis<FInDecoderInfo, ESPMode::ThreadSafe>
+	{
+		FStreamCodecInformation		ParsedInfo;
+		FTimeValue					PTS;
+		FTimeValue					DTS;
+		FTimeValue					Duration;
+	};
+
+	struct FDecodedImage
+	{
+		FDecodedImage()
+			: ImageBufferRef(nullptr)
+		{
+		}
+
+		FDecodedImage(const FDecodedImage& rhs)
+			: ImageBufferRef(nullptr)
+		{
+			InternalCopy(rhs);
+		}
+		FDecodedImage& operator=(const FDecodedImage& rhs)
+		{
+			if (this != &rhs)
+			{
+				InternalCopy(rhs);
+			}
+			return *this;
+		}
+
+		bool operator<(const FDecodedImage& rhs) const
+		{
+			return(SourceInfo->PTS < rhs.SourceInfo->PTS);
+		}
+
+		void SetImageBufferRef(CVImageBufferRef InImageBufferRef)
+		{
+			if (ImageBufferRef)
+			{
+				CFRelease(ImageBufferRef);
+				ImageBufferRef = nullptr;
+			}
+			if (InImageBufferRef)
+			{
+				CFRetain(InImageBufferRef);
+				ImageBufferRef = InImageBufferRef;
+			}
+		}
+		CVImageBufferRef GetImageBufferRef()
+		{
+			return ImageBufferRef;
+		}
+
+		CVImageBufferRef ReleaseImageBufferRef()
+		{
+			CVImageBufferRef ref = ImageBufferRef;
+			ImageBufferRef = nullptr;
+			return ref;
+		}
+
+		~FDecodedImage()
+		{
+			if (ImageBufferRef)
+			{
+				CFRelease(ImageBufferRef);
+			}
+		}
+		TSharedPtr<FInDecoderInfo, ESPMode::ThreadSafe>	SourceInfo;
+	private:
+		void InternalCopy(const FDecodedImage& rhs)
+		{
+			SourceInfo = rhs.SourceInfo;
+			SetImageBufferRef(rhs.ImageBufferRef);
+		}
+
+		CVImageBufferRef			ImageBufferRef;
+	};
+
+	FInstanceConfiguration								Config;
+
+	FMediaEvent											TerminateThreadSignal;
+	FMediaEvent											FlushDecoderSignal;
+	FMediaEvent											DecoderFlushedSignal;
+	bool												bThreadStarted;
+
+	IPlayerSessionServices*								PlayerSessionServices;
+
+	TSharedPtr<IMediaRenderer, ESPMode::ThreadSafe>		Renderer;
+    
+    TWeakPtr<IVideoDecoderResourceDelegate, ESPMode::ThreadSafe> ResourceDelegate;
+
+	FAccessUnitBuffer									AccessUnitBuffer;
+	FAccessUnitBuffer									ReplayAccessUnitBuffer;
+
+	FMediaCriticalSection								ListenerMutex;
+	IAccessUnitBufferListener*							InputBufferListener;
+	IDecoderOutputBufferListener*						ReadyBufferListener;
+
+	FDecoderFormatInfo									CurrentStreamFormatInfo;
+	FDecoderHandle*										DecoderHandle;
+	TArray<TSharedPtr<FInDecoderInfo, ESPMode::ThreadSafe>> InDecoderInfos;
+	FMediaCriticalSection								InDecoderInfoMutex;
+
+	int32												MaxDecodeBufferSize;
+	bool												bError;
+
+	FMediaCriticalSection								ReadyImageMutex;
+	TArray<FDecodedImage>								ReadyImages;
+
+public:
+	static FSystemConfiguration							SystemConfig;
+};
+
+IVideoDecoderH264::FSystemConfiguration			FVideoDecoderH264::SystemConfig;
+
+/***************************************************************************************************************************************************/
+/***************************************************************************************************************************************************/
+/***************************************************************************************************************************************************/
+
+bool IVideoDecoderH264::Startup(const IVideoDecoderH264::FSystemConfiguration& InConfig)
+{
+	return FVideoDecoderH264::Startup(InConfig);
+}
+
+void IVideoDecoderH264::Shutdown()
+{
+	FVideoDecoderH264::Shutdown();
+}
+
+bool IVideoDecoderH264::GetStreamDecodeCapability(FStreamDecodeCapability& OutResult, const FStreamDecodeCapability& InStreamParameter)
+{
+	return false;
+}
+
+IVideoDecoderH264::FSystemConfiguration::FSystemConfiguration()
+{
+	ThreadConfig.Decoder.Priority 	= TPri_Normal;
+	ThreadConfig.Decoder.StackSize	= 64 << 10;
+	ThreadConfig.Decoder.CoreAffinity = -1;
+	// Render thread is needed here!
+	ThreadConfig.PassOn.Priority  	= TPri_Normal;
+	ThreadConfig.PassOn.StackSize 	= 64 << 10;
+	ThreadConfig.PassOn.CoreAffinity  = -1;
+}
+
+IVideoDecoderH264::FInstanceConfiguration::FInstanceConfiguration()
+	: MaxDecodedFrames(8)
+	, ThreadConfig(FVideoDecoderH264::SystemConfig.ThreadConfig)
+{
+}
+
+IVideoDecoderH264* IVideoDecoderH264::Create()
+{
+	return new FVideoDecoderH264;
+}
+
+
+/***************************************************************************************************************************************************/
+/***************************************************************************************************************************************************/
+/***************************************************************************************************************************************************/
+
+class FElectraPlayerVideoDecoderOutputApple : public FVideoDecoderOutputApple
+{
+public:
+	FElectraPlayerVideoDecoderOutputApple()
+		: ImageBufferRef(nullptr)
+	{
+	}
+
+	~FElectraPlayerVideoDecoderOutputApple()
+	{
+		if (ImageBufferRef)
+		{
+			CFRelease(ImageBufferRef);
+		}
+	}
+
+    void Initialize(CVImageBufferRef InImageBufferRef, FParamDict* InParamDict)
+	{
+		FVideoDecoderOutputApple::Initialize(InParamDict);
+
+		if (ImageBufferRef)
+		{
+			CFRelease(ImageBufferRef);
+		}
+		ImageBufferRef = InImageBufferRef;
+		if (ImageBufferRef)
+		{
+			CFRetain(ImageBufferRef);
+
+			Stride = CVPixelBufferGetBytesPerRow(ImageBufferRef);
+		}
+	}
+
+	void SetOwner(const TSharedPtr<IDecoderOutputOwner, ESPMode::ThreadSafe>& InOwningRenderer) override
+	{
+		OwningRenderer = InOwningRenderer;
+	}
+
+	void ShutdownPoolable() override
+	{
+		// release image buffer (we currently realloc it every time anyway)
+        if (ImageBufferRef)
+        {
+            CFRelease(ImageBufferRef);
+            ImageBufferRef = nullptr;
+        }
+
+		TSharedPtr<IDecoderOutputOwner, ESPMode::ThreadSafe> lockedVideoRenderer = OwningRenderer.Pin();
+		if (lockedVideoRenderer.IsValid())
+		{
+			lockedVideoRenderer->SampleReleasedToPool(GetDuration());
+		}
+	}
+
+	virtual uint32 GetStride() const override
+	{
+		return Stride;
+	}
+
+	virtual CVImageBufferRef GetImageBuffer() const override
+	{
+		return ImageBufferRef;
+	}
+
+private:
+	uint32 Stride;
+
+	CVImageBufferRef ImageBufferRef;
+
+	// We hold a weak reference to the video renderer. During destruction the video renderer could be destroyed while samples are still out there..
+	TWeakPtr<IDecoderOutputOwner, ESPMode::ThreadSafe> OwningRenderer;
+};
+
+/***************************************************************************************************************************************************/
+/***************************************************************************************************************************************************/
+/***************************************************************************************************************************************************/
+
+//-----------------------------------------------------------------------------
+/**
+ * Decoder system startup
+ *
+ * @param config
+ *
+ * @return
+ */
+bool FVideoDecoderH264::Startup(const IVideoDecoderH264::FSystemConfiguration& InConfig)
+{
+	SystemConfig = InConfig;
+	return true;
+}
+
+
+//-----------------------------------------------------------------------------
+/**
+ * Decoder system shutdown.
+ */
+void FVideoDecoderH264::Shutdown()
+{
+}
+
+
+//-----------------------------------------------------------------------------
+/**
+ * Constructor
+ */
+FVideoDecoderH264::FVideoDecoderH264()
+	: FMediaThread("ElectraPlayer::H264 decoder")
+	, bThreadStarted(false)
+	, PlayerSessionServices(nullptr)
+	, Renderer(nullptr)
+	, InputBufferListener(nullptr)
+	, ReadyBufferListener(nullptr)
+	, DecoderHandle(nullptr)
+	, MaxDecodeBufferSize(0)
+	, bError(false)
+{
+}
+
+
+//-----------------------------------------------------------------------------
+/**
+ * Destructor
+ */
+FVideoDecoderH264::~FVideoDecoderH264()
+{
+	Close();
+}
+
+
+//-----------------------------------------------------------------------------
+/**
+ * Sets an AU input buffer listener.
+ *
+ * @param Listener
+ */
+void FVideoDecoderH264::SetAUInputBufferListener(IAccessUnitBufferListener* Listener)
+{
+	FMediaCriticalSection::ScopedLock lock(ListenerMutex);
+	InputBufferListener = Listener;
+}
+
+
+//-----------------------------------------------------------------------------
+/**
+ * Sets a buffer-ready listener.
+ *
+ * @param Listener
+ */
+void FVideoDecoderH264::SetReadyBufferListener(IDecoderOutputBufferListener* Listener)
+{
+	FMediaCriticalSection::ScopedLock lock(ListenerMutex);
+	ReadyBufferListener = Listener;
+}
+
+
+//-----------------------------------------------------------------------------
+/**
+ * Sets the owning player's session service interface.
+ *
+ * @param InSessionServices
+ *
+ * @return
+ */
+void FVideoDecoderH264::SetPlayerSessionServices(IPlayerSessionServices* SessionServices)
+{
+	PlayerSessionServices = SessionServices;
+}
+
+
+//-----------------------------------------------------------------------------
+/**
+ * Opens a decoder instance
+ *
+ * @param config
+ *
+ * @return
+ */
+void FVideoDecoderH264::Open(const IVideoDecoderH264::FInstanceConfiguration& InConfig)
+{
+	Config = InConfig;
+
+	// Set a large enough size to hold a single access unit. Since we will be asking for a new AU on
+	// demand there is no need to overly restrict ourselves here. But it must be large enough to
+	// hold at least the largest expected access unit.
+	AccessUnitBuffer.CapacitySet(FAccessUnitBuffer::FConfiguration(16 << 20, 60.0));
+
+	StartThread();
+}
+
+
+//-----------------------------------------------------------------------------
+/**
+ * Closes the decoder instance.
+ */
+void FVideoDecoderH264::Close()
+{
+	StopThread();
+}
+
+
+//-----------------------------------------------------------------------------
+/**
+ * Sets a new decoder limit.
+ * As soon as a new sequence starts matching this limit the decoder will be
+ * destroyed and recreated to conserve memory.
+ * The assumption is that the data being streamed will not exceed this limit,
+ * but if it does, any access unit requiring a better decoder will take
+ * precedence and force a decoder capable of decoding it!
+ *
+ * @param MaxWidth
+ * @param MaxHeight
+ * @param MaxProfile
+ * @param MaxProfileLevel
+ * @param AdditionalOptions
+ */
+void FVideoDecoderH264::SetMaximumDecodeCapability(int32 MaxWidth, int32 MaxHeight, int32 MaxProfile, int32 MaxProfileLevel, const FParamDict& AdditionalOptions)
+{
+	// Not implemented
+}
+
+
+//-----------------------------------------------------------------------------
+/**
+ * Sets a new renderer.
+ *
+ * @param InRenderer
+ */
+void FVideoDecoderH264::SetRenderer(TSharedPtr<IMediaRenderer, ESPMode::ThreadSafe> InRenderer)
+{
+	Renderer = InRenderer;
+}
+
+void FVideoDecoderH264::SetResourceDelegate(const TSharedPtr<IVideoDecoderResourceDelegate, ESPMode::ThreadSafe>& InResourceDelegate)
+{
+    ResourceDelegate = InResourceDelegate;
+}
+
+//-----------------------------------------------------------------------------
+/**
+ * Creates and runs the decoder thread.
+ */
+void FVideoDecoderH264::StartThread()
+{
+	ThreadSetPriority(Config.ThreadConfig.Decoder.Priority);
+	ThreadSetCoreAffinity(Config.ThreadConfig.Decoder.CoreAffinity);
+	ThreadSetStackSize(Config.ThreadConfig.Decoder.StackSize);
+	ThreadStart(Electra::MakeDelegate(this, &FVideoDecoderH264::WorkerThread));
+	bThreadStarted = true;
+}
+
+
+//-----------------------------------------------------------------------------
+/**
+ * Stops the decoder thread.
+ */
+void FVideoDecoderH264::StopThread()
+{
+	if (bThreadStarted)
+	{
+		TerminateThreadSignal.Signal();
+		ThreadWaitDone();
+		bThreadStarted = false;
+	}
+}
+
+
+//-----------------------------------------------------------------------------
+/**
+ * Posts an error to the session service error listeners.
+ *
+ * @param ApiReturnValue
+ * @param Message
+ * @param Code
+ * @param Error
+ */
+void FVideoDecoderH264::PostError(int32_t ApiReturnValue, const FString& Message, uint16 Code, UEMediaError Error)
+{
+	check(PlayerSessionServices);
+	if (PlayerSessionServices)
+	{
+		FErrorDetail err;
+		err.SetError(Error != UEMEDIA_ERROR_OK ? Error : UEMEDIA_ERROR_DETAIL);
+		err.SetFacility(Facility::EFacility::H264Decoder);
+		err.SetCode(Code);
+		err.SetMessage(Message);
+		err.SetPlatformMessage(FString::Printf(TEXT("%d (0x%08x)"), ApiReturnValue, ApiReturnValue));
+		PlayerSessionServices->PostError(err);
+	}
+}
+
+
+//-----------------------------------------------------------------------------
+/**
+ * Sends a log message to the session service log.
+ *
+ * @param Level
+ * @param Message
+ */
+void FVideoDecoderH264::LogMessage(IInfoLog::ELevel Level, const FString& Message)
+{
+	if (PlayerSessionServices)
+	{
+		PlayerSessionServices->PostLog(Facility::EFacility::H264Decoder, Level, Message);
+	}
+}
+
+
+//-----------------------------------------------------------------------------
+/**
+ * Create a pool of decoded images for the decoder.
+ *
+ * @return
+ */
+bool FVideoDecoderH264::CreateDecodedImagePool()
+{
+	check(Renderer);
+	FParamDict poolOpts;
+
+	poolOpts.Set("num_buffers", FVariantValue((int64) Config.MaxDecodedFrames));
+
+	UEMediaError Error = Renderer->CreateBufferPool(poolOpts);
+	check(Error == UEMEDIA_ERROR_OK);
+
+	MaxDecodeBufferSize = (int32) Renderer->GetBufferPoolProperties().GetValue("max_buffers").GetInt64();
+
+	if (Error != UEMEDIA_ERROR_OK)
+	{
+		PostError(0, "Failed to create image pool", ERRCODE_INTERNAL_APPLE_COULD_NOT_CREATE_IMAGE_POOL, Error);
+	}
+
+	return Error == UEMEDIA_ERROR_OK;
+}
+
+
+//-----------------------------------------------------------------------------
+/**
+ * Destroys the pool of decoded images.
+ */
+void FVideoDecoderH264::DestroyDecodedImagePool()
+{
+	Renderer->ReleaseBufferPool();
+}
+
+
+//-----------------------------------------------------------------------------
+/**
+ * Called to receive a new input access unit for decoding.
+ *
+ * @param AccessUnit
+ */
+IAccessUnitBufferInterface::EAUpushResult FVideoDecoderH264::AUdataPushAU(FAccessUnit* AccessUnit)
+{
+	AccessUnit->AddRef();
+	bool bOk = AccessUnitBuffer.Push(AccessUnit);
+	if (!bOk)
+	{
+		FAccessUnit::Release(AccessUnit);
+	}
+
+	return bOk ? IAccessUnitBufferInterface::EAUpushResult::Ok : IAccessUnitBufferInterface::EAUpushResult::Full;
+}
+
+
+//-----------------------------------------------------------------------------
+/**
+ * "Pushes" an End Of Data marker indicating no further access units will be added.
+ */
+void FVideoDecoderH264::AUdataPushEOD()
+{
+	AccessUnitBuffer.PushEndOfData();
+}
+
+
+//-----------------------------------------------------------------------------
+/**
+ * Flushes the decoder and clears the input access unit buffer.
+ */
+void FVideoDecoderH264::AUdataFlushEverything()
+{
+	FlushDecoderSignal.Signal();
+	DecoderFlushedSignal.WaitAndReset();
+}
+
+
+//-----------------------------------------------------------------------------
+/**
+ * Create a decoder instance.
+ *
+ * @param InputFormatDescription
+ *
+ * @return true if successful, false on error
+ */
+bool FVideoDecoderH264::InternalDecoderCreate(CMFormatDescriptionRef InputFormatDescription)
+{
+	check(DecoderHandle == nullptr);
+	check(InputFormatDescription != nullptr);
+	DecoderHandle = new FDecoderHandle;
+	CFRetain(InputFormatDescription);
+	DecoderHandle->FormatDescription = InputFormatDescription;
+
+	VTDecompressionOutputCallbackRecord CallbackRecord;
+	CallbackRecord.decompressionOutputCallback = _DecodeCallback;
+	CallbackRecord.decompressionOutputRefCon   = this;
+
+	// Output image format configuration
+	CFMutableDictionaryRef OutputImageFormat = CFDictionaryCreateMutable(nullptr, 3, &kCFTypeDictionaryKeyCallBacks, &kCFTypeDictionaryValueCallBacks);
+	// Choice of: kCVPixelBufferOpenGLCompatibilityKey (all)  kCVPixelBufferOpenGLESCompatibilityKey (iOS only)   kCVPixelBufferMetalCompatibilityKey (all)
+#if PLATFORM_MAC
+	int pxfmt = kCVPixelFormatType_32BGRA;
+	CFNumberRef PixelFormat = CFNumberCreate(kCFAllocatorDefault, kCFNumberSInt32Type, &pxfmt);
+	CFDictionarySetValue(OutputImageFormat, kCVPixelBufferPixelFormatTypeKey, PixelFormat);
+	CFRelease(PixelFormat);
+	CFDictionarySetValue(OutputImageFormat, kCVPixelBufferMetalCompatibilityKey, kCFBooleanTrue);
+#elif PLATFORM_IOS
+	CFDictionarySetValue(OutputImageFormat, kCVPixelBufferOpenGLESCompatibilityKey, kCFBooleanFalse);
+	CFDictionarySetValue(OutputImageFormat, kCVPixelBufferMetalCompatibilityKey, kCFBooleanTrue);
+	int pxfmt = kCVPixelFormatType_32BGRA;
+	CFNumberRef PixelFormat = CFNumberCreate(kCFAllocatorDefault, kCFNumberSInt32Type, &pxfmt);
+	CFDictionarySetValue(OutputImageFormat, kCVPixelBufferPixelFormatTypeKey, PixelFormat);
+	CFRelease(PixelFormat);
+#else
+	#error "Should not get here. Check platform checks at the top of the file."
+#endif
+
+	// Session configuration
+	CFMutableDictionaryRef SessionConfiguration = CFDictionaryCreateMutable(nullptr, 1, &kCFTypeDictionaryKeyCallBacks, &kCFTypeDictionaryValueCallBacks);
+	// Ask for hardware decoding
+	//	CFDictionarySetValue(SessionConfiguration, kVTVideoDecoderSpecification_EnableHardwareAcceleratedVideoDecoder, kCFBooleanTrue);
+
+	OSStatus res = VTDecompressionSessionCreate(kCFAllocatorDefault, InputFormatDescription, SessionConfiguration, OutputImageFormat, &CallbackRecord, &DecoderHandle->DecompressionSession);
+	CFRelease(SessionConfiguration);
+	CFRelease(OutputImageFormat);
+	if (res != 0)
+	{
+		PostError(res, "Failed to create video decoder", ERRCODE_INTERNAL_APPLE_COULD_NOT_CREATE_VIDEO_DECODER);
+		return false;
+	}
+	return true;
+}
+
+
+//-----------------------------------------------------------------------------
+/**
+ * Destroys the current decoder instance.
+ */
+void FVideoDecoderH264::InternalDecoderDestroy()
+{
+	if (DecoderHandle)
+	{
+		DecoderHandle->Close();
+		delete DecoderHandle;
+		DecoderHandle = nullptr;
+	}
+}
+
+
+//-----------------------------------------------------------------------------
+/**
+ * Creates a new decoder and runs all AUs since the last IDR frame through
+ * without producing decoded images.
+ * Used to continue decoding after the application has been backgrounded
+ * which results in a decoder session loss when resumed.
+ * Since decoding cannot resume at arbitrary points in the stream everything
+ * from the last IDR frame needs to be decoded again.
+ * To speed this up AUs that have no dependencies are not added to the replay data.
+ */
+void FVideoDecoderH264::RecreateDecoderSession()
+{
+	// Destroy existing decoder. It is of no use any longer.
+	InternalDecoderDestroy();
+	// Likewise anything that was still pending we also no longer need.
+	ClearInDecoderInfos();
+
+	// Do we have replay data?
+	if (ReplayAccessUnitBuffer.Num())
+	{
+		TMediaQueueDynamicNoLock<FAccessUnit *>	ReprocessedAUs;
+
+		bool bDone = false;
+		FAccessUnit* AccessUnit = nullptr;
+		bool bFirst = true;
+		while(!bError && !bDone)
+		{
+			AccessUnit = nullptr;
+			if (ReplayAccessUnitBuffer.Pop(AccessUnit))
+			{
+				check(AccessUnit != nullptr);
+				ReprocessedAUs.Push(AccessUnit);
+				// Create the format description from the first replay AU.
+				if (bFirst)
+				{
+					bFirst = false;
+					CMFormatDescriptionRef NewFormatDescr = nullptr;
+					if (CreateFormatDescription(NewFormatDescr, AccessUnit))
+					{
+						if (InternalDecoderCreate(NewFormatDescr))
+						{
+							// Ok
+						}
+						else
+						{
+							bError = true;
+						}
+						CFRelease(NewFormatDescr);
+					}
+					else
+					{
+						bError = true;
+					}
+				}
+				// Decode
+				EDecodeResult DecRes = Decode(AccessUnit, true);
+				// On failure or yet another loss of decoder session, leave...
+				if (DecRes != EDecodeResult::Ok)
+				{
+					bDone = true;
+				}
+			}
+			else
+			{
+				bDone = true;
+			}
+		}
+		// Even in case of an error we need to get all replay AUs into our processed FIFO and
+		// from there back into the replay buffer. We may need them again and they need to
+		// stay in the original order.
+		while(ReplayAccessUnitBuffer.Pop(AccessUnit))
+		{
+			ReprocessedAUs.Push(AccessUnit);
+		}
+		while(!ReprocessedAUs.IsEmpty())
+		{
+			AccessUnit = ReprocessedAUs.Pop();
+			ReplayAccessUnitBuffer.Push(AccessUnit);
+		}
+	// Flush the decoder to get it idle and discard any accumulated source infos.
+	FlushDecoder();
+	}
+}
+
+
+//-----------------------------------------------------------------------------
+/**
+ * Notifies the buffer-ready listener that we will now be producing output data
+ * and need an output buffer.
+ * This call is not intended to create or obtain an output buffer. It is merely
+ * indicating that new output will be produced.
+ *
+ * @param bHaveOutput
+ */
+void FVideoDecoderH264::NotifyReadyBufferListener(bool bHaveOutput)
+{
+	if (ReadyBufferListener)
+	{
+		IDecoderOutputBufferListener::FDecodeReadyStats stats;
+		stats.MaxDecodedElementsReady = MaxDecodeBufferSize;
+		stats.NumElementsInDecoder    = InDecoderInfos.Num();
+		stats.bOutputStalled		  = !bHaveOutput;
+		stats.bEODreached   		  = AccessUnitBuffer.IsEndOfData() && stats.NumDecodedElementsReady == 0 && stats.NumElementsInDecoder == 0;
+		ListenerMutex.Lock();
+		if (ReadyBufferListener)
+		{
+			ReadyBufferListener->DecoderOutputReady(stats);
+		}
+		ListenerMutex.Unlock();
+	}
+}
+
+
+//-----------------------------------------------------------------------------
+/**
+ * Prepares the passed access unit to be sent to the decoder.
+ *
+ * @param AccessUnit
+ * @param bOutIsIDR
+ * @param bOutIsDiscardable
+ */
+void FVideoDecoderH264::PrepareAU(FAccessUnit* AccessUnit, bool& bOutIsIDR, bool& bOutIsDiscardable)
+{
+	bOutIsIDR = false;
+	bOutIsDiscardable = true;
+	if (!AccessUnit->bHasBeenPrepared)
+	{
+		AccessUnit->bHasBeenPrepared = true;
+
+		// Process NALUs
+		// Note: With VideoToolbox the data must be in avcc format, that is no startcodes but the
+		//       length of each NALU. We already have that so there is not a lot to do here.
+		uint32* CurrentNALU = (uint32 *)AccessUnit->AUData;
+		uint32* LastNALU    = (uint32 *)Electra::AdvancePointer(CurrentNALU, AccessUnit->AUSize);
+		while(CurrentNALU < LastNALU)
+		{
+			// Check the nal_ref_idc in the NAL unit for dependencies.
+			uint8 nal = *(const uint8 *)(CurrentNALU + 1);
+			check((nal & 0x80) == 0);
+			if ((nal >> 5) != 0)
+			{
+				bOutIsDiscardable = false;
+			}
+			// IDR frame?
+			if ((nal & 0x1f) == 5)
+			{
+				bOutIsIDR = true;
+			}
+
+			// SEI message(s)?
+			if ((nal & 0x1f) == 6)
+			{
+				// TODO: we might need to set aside any SEI messages carrying 608 or 708 caption data.
+			}
+
+	// FIXME: We may need to remove SPS and PPS from here in case the stream data is avc3 data.
+	//        Not sure if the decoder will actually deal with those or if they need to be strictly passed
+	//        on the side and not inband.
+
+			uint32 naluLen = MEDIA_FROM_BIG_ENDIAN(*CurrentNALU) + 4;
+			CurrentNALU = Electra::AdvancePointer(CurrentNALU, naluLen);
+		}
+	}
+}
+
+
+//-----------------------------------------------------------------------------
+/**
+ * Gets an output buffer from the video renderer.
+ *
+ * @param RenderOutputBuffer
+ *
+ * @return true if successful, false on error
+ */
+bool FVideoDecoderH264::AcquireOutputBuffer(IMediaRenderer::IBuffer*& RenderOutputBuffer)
+{
+	RenderOutputBuffer = nullptr;
+	FParamDict BufferAcquireOptions;
+	while(!TerminateThreadSignal.IsSignaled() && !FlushDecoderSignal.IsSignaled())
+	{
+		UEMediaError bufResult = Renderer->AcquireBuffer(RenderOutputBuffer, 0, BufferAcquireOptions);
+		check(bufResult == UEMEDIA_ERROR_OK || bufResult == UEMEDIA_ERROR_INSUFFICIENT_DATA);
+		if (bufResult != UEMEDIA_ERROR_OK && bufResult != UEMEDIA_ERROR_INSUFFICIENT_DATA)
+		{
+			PostError(0, "Failed to acquire sample buffer", ERRCODE_INTERNAL_APPLE_COULD_NOT_GET_OUTPUT_BUFFER, bufResult);
+			return false;
+		}
+		bool bHaveAvailSmpBlk = RenderOutputBuffer != nullptr;
+		NotifyReadyBufferListener(bHaveAvailSmpBlk);
+		if (bHaveAvailSmpBlk)
+		{
+			break;
+		}
+		else
+		{
+			// No available buffer. Sleep for a bit. Can't sleep on a signal since we have to check two: abort and flush
+			FMediaRunnable::SleepMilliseconds(5);
+		}
+	}
+	return true;
+}
+
+
+//-----------------------------------------------------------------------------
+/**
+ * Flushes the decoder and passes all pending frames to the renderer.
+ *
+ * @return true if successful, false otherwise.
+ */
+bool FVideoDecoderH264::FlushDecoder()
+{
+	if (DecoderHandle && DecoderHandle->DecompressionSession)
+	{
+		VTDecompressionSessionFinishDelayedFrames(DecoderHandle->DecompressionSession);
+		VTDecompressionSessionWaitForAsynchronousFrames(DecoderHandle->DecompressionSession);
+		// Push out all images we got from the decoder from our internal PTS sorting facility
+		ProcessOutput(true);
+	}
+	ClearInDecoderInfos();
+	return true;
+}
+
+
+//-----------------------------------------------------------------------------
+/**
+ * Clears out any leftover decoder source infos that were not picked by decoded frames.
+ * This must be called after having flushed the decoder to ensure no delayed decode
+ * callbacks will arrive.
+ */
+void FVideoDecoderH264::ClearInDecoderInfos()
+{
+	FMediaCriticalSection::ScopedLock lock(InDecoderInfoMutex);
+	InDecoderInfos.Empty();
+}
+
+
+//-----------------------------------------------------------------------------
+/**
+ * Creates a dummy output image that is not to be displayed and has no image data.
+ * Dummy access units are created when stream data is missing to ensure the data
+ * pipeline does not run dry and exhibits no gaps in the timeline.
+ *
+ * @param AccessUnit
+ *
+ * @return true if successful, false otherwise
+ */
+bool FVideoDecoderH264::DecodeDummy(FAccessUnit* AccessUnit)
+{
+	check(AccessUnit->Duration.IsValid());
+
+    FDecodedImage NextImage;
+    NextImage.SourceInfo = TSharedPtr<FInDecoderInfo, ESPMode::ThreadSafe>(new FInDecoderInfo);
+    NextImage.SourceInfo->DTS = AccessUnit->DTS;
+    NextImage.SourceInfo->PTS = AccessUnit->PTS;
+    NextImage.SourceInfo->Duration = AccessUnit->Duration;
+    if (AccessUnit->AUCodecData.IsValid())
+    {
+        NextImage.SourceInfo->ParsedInfo = AccessUnit->AUCodecData->ParsedInfo;
+    }
+
+	ReadyImageMutex.Lock();
+	ReadyImages.Add(NextImage);
+	ReadyImages.Sort();
+	ReadyImageMutex.Unlock();
+
+    return true;
+}
+
+//-----------------------------------------------------------------------------
+/**
+ * Flushes all images not passed to the renderer from our list.
+ */
+void FVideoDecoderH264::FlushPendingImages()
+{
+	// Clear out the map. This implicitly drops all image refcounts.
+	ReadyImageMutex.Lock();
+	ReadyImages.Empty();
+	ReadyImageMutex.Unlock();
+}
+
+//-----------------------------------------------------------------------------
+/**
+ * Checks if the codec specific data has changed.
+ *
+ * @return false if the format is still the same, true if it has changed.
+ */
+bool FVideoDecoderH264::FDecoderFormatInfo::IsDifferentFrom(const FAccessUnit* InputAccessUnit)
+{
+	if (InputAccessUnit->AUCodecData.IsValid() && InputAccessUnit->AUCodecData.Get() != CurrentCodecData.Get())
+	{
+		CurrentCodecData = InputAccessUnit->AUCodecData;
+		return true;
+	}
+	return false;
+}
+
+
+//-----------------------------------------------------------------------------
+/**
+ * Creates a source format description to be passed into the video decoder from
+ * the codec specific data (CSD) attached to the input access unit.
+ *
+ * @param OutFormatDescription
+ * @param InputAccessUnit
+ *
+ * @return true if successful, false otherwise.
+ */
+bool FVideoDecoderH264::CreateFormatDescription(CMFormatDescriptionRef& OutFormatDescription, const FAccessUnit* InputAccessUnit)
+{
+	if (InputAccessUnit && InputAccessUnit->AUCodecData.IsValid() && InputAccessUnit->AUCodecData->CodecSpecificData.Num())
+	{
+		// Get the NALUs from the CSD.
+		TArray<MPEG::FNaluInfo>	NALUs;
+		MPEG::ParseBitstreamForNALUs(NALUs, InputAccessUnit->AUCodecData->CodecSpecificData.GetData(), InputAccessUnit->AUCodecData->CodecSpecificData.Num());
+		if (NALUs.Num())
+		{
+			int32 NumRecords = NALUs.Num();
+			if (NumRecords)
+			{
+				uint8_t const* * DataPointers = new uint8_t const* [NumRecords];
+				SIZE_T*          DataSizes    = new SIZE_T [NumRecords];
+				for(int32 i=0; i<NumRecords; ++i)
+				{
+					DataPointers[i] = Electra::AdvancePointer(InputAccessUnit->AUCodecData->CodecSpecificData.GetData(), NALUs[i].Offset + NALUs[i].UnitLength);
+					DataSizes[i]    = NALUs[i].Size;
+				}
+				OSStatus res = CMVideoFormatDescriptionCreateFromH264ParameterSets(kCFAllocatorDefault, NumRecords, DataPointers, DataSizes, 4, &OutFormatDescription);
+				delete [] DataPointers;
+				delete [] DataSizes;
+				if (res == 0)
+				{
+					return true;
+				}
+				else
+				{
+					if (OutFormatDescription)
+					{
+						CFRelease(OutFormatDescription);
+						OutFormatDescription = nullptr;
+					}
+					PostError(res, "Failed to create video format description from CSD", ERRCODE_INTERNAL_APPLE_BAD_VIDEO_CSD);
+					return false;
+				}
+			}
+		}
+		PostError(0, "Failed to create video format description from CSD", ERRCODE_INTERNAL_APPLE_BAD_VIDEO_CSD);
+		return false;
+	}
+	PostError(0, "Cannot create video format description from empty CSD", ERRCODE_INTERNAL_APPLE_NO_VIDEO_CSD);
+	return false;
+}
+
+
+//-----------------------------------------------------------------------------
+/**
+ * Callback from the video decoder when a new decoded image is ready.
+ *
+ * @param pSrcRef
+ * @param status
+ * @param infoFlags
+ * @param imageBuffer
+ * @param presentationTimeStamp
+ * @param presentationDuration
+ */
+void FVideoDecoderH264::DecodeCallback(void* pSrcRef, OSStatus status, VTDecodeInfoFlags infoFlags, CVImageBufferRef imageBuffer, CMTime presentationTimeStamp, CMTime presentationDuration)
+{
+	// Remove the source info even if there ultimately was a decode error or if the frame was dropped.
+	TSharedPtr<FInDecoderInfo, ESPMode::ThreadSafe> SourceInfo;
+	InDecoderInfoMutex.Lock();
+	int32 NumCurrentDecodeInfos = InDecoderInfos.Num();
+	for(int32 i=0; i<NumCurrentDecodeInfos; ++i)
+	{
+		if (InDecoderInfos[i].Get() == pSrcRef)
+		{
+			SourceInfo = InDecoderInfos[i];
+			InDecoderInfos.RemoveSingle(SourceInfo);
+			break;
+		}
+	}
+	InDecoderInfoMutex.Unlock();
+
+	if (!SourceInfo.IsValid())
+	{
+		LogMessage(IInfoLog::ELevel::Error, FString::Printf(TEXT("FVideoDecoderH264::DecodeCallback(): No source info found for decoded srcref %p in %d pending infos (OSStatus %d, infoFlags %d)"), pSrcRef, NumCurrentDecodeInfos, (int32)status, (int32)infoFlags));
+	}
+
+	if (status == 0)
+	{
+		if (imageBuffer != nullptr && (infoFlags & kVTDecodeInfo_FrameDropped) == 0 && SourceInfo.IsValid())
+		{
+			FDecodedImage NextImage;
+			NextImage.SourceInfo = SourceInfo;
+			NextImage.SetImageBufferRef(imageBuffer);
+
+			// Recall decoded frame for later processing
+			// (we do all processing of output on the decoder thread)
+ 			ReadyImageMutex.Lock();
+			ReadyImages.Add(NextImage);
+			ReadyImages.Sort();
+			ReadyImageMutex.Unlock();
+		}
+	}
+	else
+	{
+		bError = true;
+		PostError(status, "Failed to decode video", ERRCODE_INTERNAL_APPLE_FAILED_TO_DECODE_VIDEO);
+	}
+}
+
+
+//-----------------------------------------------------------------------------
+/**
+ * Sends an access unit to the decoder for decoding.
+ *
+ * @param AccessUnit
+ * @param bRecreatingSession
+ *
+ * @return
+ */
+FVideoDecoderH264::EDecodeResult FVideoDecoderH264::Decode(FAccessUnit* AccessUnit, bool bRecreatingSession)
+{
+	if (!DecoderHandle || !DecoderHandle->DecompressionSession)
+	{
+		return EDecodeResult::Fail;
+	}
+
+	check(AccessUnit->Duration.IsValid());
+	check(AccessUnit->PTS.IsValid());
+	check(AccessUnit->DTS.IsValid());
+
+	// Create a memory block for the access unit. We pass it the memory we already have and ensure that by setting the block allocator
+	// to kCFAllocatorNull no one will attempt to deallocate the memory!
+	CMBlockBufferRef AUDataBlock = nullptr;
+	SIZE_T AUDataSize = AccessUnit->AUSize;
+	OSStatus res = CMBlockBufferCreateWithMemoryBlock(nullptr, AccessUnit->AUData, AUDataSize, kCFAllocatorNull, nullptr, 0, AUDataSize, 0, &AUDataBlock);
+	if (res)
+	{
+		PostError(res, "Failed to create video data block buffer", ERRCODE_INTERNAL_APPLE_FAILED_TO_CREATE_BLOCK_BUFFER);
+		return EDecodeResult::Fail;
+	}
+
+	// Set up the timing info with DTS, PTS and duration.
+	CMSampleTimingInfo TimingInfo;
+	const int64_t HNS_PER_S = 10000000;
+	TimingInfo.decodeTimeStamp       = CMTimeMake(AccessUnit->DTS.GetAsHNS(), HNS_PER_S);
+	TimingInfo.presentationTimeStamp = CMTimeMake(AccessUnit->PTS.GetAsHNS(), HNS_PER_S);
+	TimingInfo.duration              = CMTimeMake(AccessUnit->Duration.GetAsHNS(), HNS_PER_S);
+
+	CMSampleBufferRef SampleBufferRef = nullptr;
+	res = CMSampleBufferCreate(kCFAllocatorDefault, AUDataBlock, true, nullptr, nullptr, DecoderHandle->FormatDescription, 1, 1, &TimingInfo, 1, &AUDataSize, &SampleBufferRef);
+	// The buffer is now held by the sample, so we release our ref count.
+	CFRelease(AUDataBlock);
+	if (res)	// see CMSampleBuffer for kCMSampleBufferError_AllocationFailed and such.
+	{
+		PostError(res, "Failed to create video sample buffer", ERRCODE_INTERNAL_APPLE_FAILED_TO_CREATE_SAMPLE_BUFFER);
+		return EDecodeResult::Fail;
+	}
+
+	TSharedPtr<FInDecoderInfo, ESPMode::ThreadSafe> SourceInfo(new FInDecoderInfo);
+	SourceInfo->DTS = AccessUnit->DTS;
+	SourceInfo->PTS = AccessUnit->PTS;
+	SourceInfo->Duration = AccessUnit->Duration;
+	if (AccessUnit->AUCodecData.IsValid())
+	{
+		SourceInfo->ParsedInfo = AccessUnit->AUCodecData->ParsedInfo;
+	}
+	InDecoderInfoMutex.Lock();
+	InDecoderInfos.Push(SourceInfo);
+	InDecoderInfoMutex.Unlock();
+
+	// Decode
+/*
+	kVTDecodeFrame_EnableAsynchronousDecompression = 1<<0,
+	kVTDecodeFrame_DoNotOutputFrame = 1<<1,
+	kVTDecodeFrame_1xRealTimePlayback = 1<<2,
+	kVTDecodeFrame_EnableTemporalProcessing = 1<<3,
+*/
+// TODO: set the proper flags. This may require a larger intermediate output queue. Needs experimenting.
+	VTDecodeFrameFlags DecodeFlags = kVTDecodeFrame_EnableAsynchronousDecompression | kVTDecodeFrame_EnableTemporalProcessing;
+	if (bRecreatingSession)
+	{
+		DecodeFlags = kVTDecodeFrame_DoNotOutputFrame;
+	}
+	VTDecodeInfoFlags  InfoFlags = 0;
+/*
+	kVTDecodeInfo_Asynchronous = 1UL << 0,
+	kVTDecodeInfo_FrameDropped = 1UL << 1,
+	kVTDecodeInfo_ImageBufferModifiable = 1UL << 2,
+*/
+// FIXME: This will require access to the decoder that should be granted by some arbitrator.
+//        Especially for iOS/padOS/tvos where the app can be backgrounded and the decoder potentially becoming inaccessible
+	res = VTDecompressionSessionDecodeFrame(DecoderHandle->DecompressionSession, SampleBufferRef, DecodeFlags, SourceInfo.Get(), &InfoFlags);
+	CFRelease(SampleBufferRef);
+	if (res == 0)
+	{
+		// Ok.
+		return EDecodeResult::Ok;
+	}
+	else
+	{
+		InDecoderInfoMutex.Lock();
+		InDecoderInfos.RemoveSingle(SourceInfo);
+		InDecoderInfoMutex.Unlock();
+		if (res == kVTInvalidSessionErr)
+		{
+			// Lost the decoder session due to being backgrounded and returning to the foreground.
+			return EDecodeResult::SessionLost;
+		}
+		PostError(res, "Failed to decode video frame", ERRCODE_INTERNAL_APPLE_FAILED_TO_DECODE_VIDEO);
+		return EDecodeResult::Fail;
+	}
+}
+
+
+//-----------------------------------------------------------------------------
+/**
+ * H264 video decoder main threaded decode loop
+ */
+void FVideoDecoderH264::WorkerThread()
+{
+	bool bDone  = false;
+
+	bError = false;
+
+	// If the application is suspended and resumed we will lose out decoder session and cannot continue
+	// decoding at the point of interruption. We have to keep the access units from the previous IDR frame
+	// and run them through a new decoder session again, discarding all output, after which we can resume
+	// decoding from the point of interruption.
+	// Configure the playback AU buffer to be large enough to hold the access units to be replayed.
+	ReplayAccessUnitBuffer.CapacitySet(FAccessUnitBuffer::FConfiguration(64 << 20, 360.0));
+
+	// Create decoded image pool.
+	if (!CreateDecodedImagePool())
+	{
+		bError = true;
+	}
+
+	while(!TerminateThreadSignal.IsSignaled())
+	{
+		// Notify optional buffer listener that we will now be needing an AU for our input buffer.
+		if (!bError && InputBufferListener && AccessUnitBuffer.Num() == 0)
+		{
+			FAccessUnitBufferInfo	sin;
+			IAccessUnitBufferListener::FBufferStats	stats;
+			AccessUnitBuffer.GetStats(sin);
+			stats.NumAUsAvailable  = sin.NumCurrentAccessUnits;
+			stats.NumBytesInBuffer = sin.CurrentMemInUse;
+			stats.MaxBytesOfBuffer = sin.MaxDataSize;
+			stats.bEODSignaled     = sin.bEndOfData;
+			stats.bEODReached      = sin.bEndOfData && sin.NumCurrentAccessUnits == 0;
+			ListenerMutex.Lock();
+			if (InputBufferListener)
+			{
+				InputBufferListener->DecoderInputNeeded(stats);
+			}
+			ListenerMutex.Unlock();
+		}
+		// Wait for data.
+		bool bHaveData = AccessUnitBuffer.WaitForData(1000 * 5);
+
+		// Only send new data to the decoder if we know we got enough room (to avoid accumulating too many frames in our internal PTS-sort queue)
+		bool bTooManyImagesWaiting = false;
+		if (bHaveData)
+		{
+			// When there is data, even and especially after a previous EOD, we are no longer done and idling.
+			if (bDone)
+			{
+				bDone = false;
+			}
+
+			ReadyImageMutex.Lock();
+			bTooManyImagesWaiting = (ReadyImages.Num() > MaxImagesHoldBackForPTSOrdering);
+			ReadyImageMutex.Unlock();
+			if (bTooManyImagesWaiting)
+			{
+				// Signal blockage and wait for it to clear up before going back to normal decoder work...
+				NotifyReadyBufferListener(false);
+
+				while(!TerminateThreadSignal.IsSignaled() && !FlushDecoderSignal.IsSignaled())
+				{
+					ProcessOutput();
+
+					ReadyImageMutex.Lock();
+					bTooManyImagesWaiting = (ReadyImages.Num() > MaxImagesHoldBackForPTSOrdering);
+					ReadyImageMutex.Unlock();
+					if (!bTooManyImagesWaiting)
+					{
+						break;
+					}
+
+					FMediaRunnable::SleepMilliseconds(20);
+				}
+			}
+		}
+
+		if (bHaveData && !bTooManyImagesWaiting)
+		{
+			FAccessUnit* AccessUnit = nullptr;
+			bool bOk = AccessUnitBuffer.Pop(AccessUnit);
+			MEDIA_UNUSED_VAR(bOk);
+			check(bOk);
+
+			if (!AccessUnit->bIsDummyData)
+			{
+				bool bIsKeyframe, bIsDiscardable;
+				PrepareAU(AccessUnit, bIsKeyframe, bIsDiscardable);
+				// An IDR frame means we can start decoding there, so we can purge any accumulated replay AUs.
+				if (bIsKeyframe)
+				{
+					ReplayAccessUnitBuffer.Flush();
+				}
+
+				bool bStreamFormatChanged = CurrentStreamFormatInfo.IsDifferentFrom(AccessUnit);
+				bool bNeedNewDecoder = false;
+				if (bStreamFormatChanged || !DecoderHandle)
+				{
+					CMFormatDescriptionRef NewFormatDescr = nullptr;
+					if (CreateFormatDescription(NewFormatDescr, AccessUnit))
+					{
+						bNeedNewDecoder = DecoderHandle == nullptr || !DecoderHandle->IsCompatibleWith(NewFormatDescr);
+						if (bNeedNewDecoder)
+						{
+							FlushDecoder();
+							InternalDecoderDestroy();
+							if (InternalDecoderCreate(NewFormatDescr))
+							{
+								// Ok
+							}
+							else
+							{
+								bError = true;
+							}
+						}
+						CFRelease(NewFormatDescr);
+					}
+					else
+					{
+						bError = true;
+					}
+				}
+
+				// Decode
+				if (!bError && DecoderHandle && DecoderHandle->DecompressionSession)
+				{
+                    EDecodeResult DecRes = Decode(AccessUnit, false);
+
+					// Process any output we might have pending
+					ProcessOutput();
+
+                    // Decode ok?
+                    if (DecRes == EDecodeResult::Ok)
+                    {
+                        // Yes, add to the replay buffer if it is not a discardable access unit.
+                        if (!bIsDiscardable)
+                        {
+                            AccessUnit->AddRef();
+                            if (!ReplayAccessUnitBuffer.Push(AccessUnit))
+                            {
+                                // FIXME: Is this cause for a playback error? For now we just forget about the replay AU and take any possible decoding artefacts.
+                                FAccessUnit::Release(AccessUnit);
+                            }
+                        }
+                    }
+                    // Lost the decoder session?
+                    else if (DecRes == EDecodeResult::SessionLost)
+                    {
+                        // Did not produce output, release semaphore again.
+                        RecreateDecoderSession();
+                        // Retry this AU on the new decoder session.
+                        continue;
+                    }
+                    // Decode failed
+                    else
+                    {
+                        bError = true;
+                    }
+				}
+			}
+			else
+			{
+				if (!DecodeDummy(AccessUnit))
+				{
+					bError = true;
+				}
+			}
+
+			FAccessUnit::Release(AccessUnit);
+			AccessUnit = nullptr;
+		}
+		else
+		{
+			ProcessOutput();
+
+			// No data. Is the buffer at EOD?
+			if (AccessUnitBuffer.IsEndOfData())
+			{
+				NotifyReadyBufferListener(true);
+				// Are we done yet?
+				if (!bDone && !bError)
+				{
+					bError = !FlushDecoder();
+				}
+				bDone = true;
+				FMediaRunnable::SleepMilliseconds(10);
+				ProcessOutput(true);
+			}
+		}
+
+		// Flush?
+		if (FlushDecoderSignal.IsSignaled())
+		{
+			// Have to destroy the decoder!
+			InternalDecoderDestroy();
+			FlushPendingImages();
+			ClearInDecoderInfos();
+			AccessUnitBuffer.Flush();
+			ReplayAccessUnitBuffer.Flush();
+
+			FlushDecoderSignal.Reset();
+			DecoderFlushedSignal.Signal();
+
+			// Reset done state.
+			bDone = false;
+		}
+	}
+
+	InternalDecoderDestroy();
+	FlushPendingImages();
+	DestroyDecodedImagePool();
+	ClearInDecoderInfos();
+	AccessUnitBuffer.Flush();
+	AccessUnitBuffer.CapacitySet(0);
+	ReplayAccessUnitBuffer.Flush();
+	ReplayAccessUnitBuffer.CapacitySet(0);
+}
+
+//-----------------------------------------------------------------------------
+/**
+ */
+void FVideoDecoderH264::ProcessOutput(bool bFlush)
+{
+	// Pull any frames we can get from in-PTS-order array...
+	while(1)
+	{
+		ReadyImageMutex.Lock();
+		if (ReadyImages.Num() < (bFlush ? 1 : NumImagesHoldBackForPTSOrdering))
+		{
+			ReadyImageMutex.Unlock();
+			break;
+		}
+
+		if (!Renderer->CanReceiveOutputFrames(1))
+		{
+			NotifyReadyBufferListener(false);
+			ReadyImageMutex.Unlock();
+			break;
+		}
+
+		FDecodedImage NextImage(ReadyImages[0]);
+		ReadyImages.RemoveAt(0);
+
+		ReadyImageMutex.Unlock();
+
+		// Get an output buffer from the renderer to pass the image to.
+		IMediaRenderer::IBuffer* RenderOutputBuffer = nullptr;
+		if (AcquireOutputBuffer(RenderOutputBuffer))
+		{
+			if (RenderOutputBuffer)
+			{
+                double PixelAspectRatio = 1.0;
+
+                // Note that the ImageBuffer reference below might be null if this is a dummy frame!
+                CVImageBufferRef ImageBufferRef = NextImage.ReleaseImageBufferRef();
+
+				FParamDict* OutputBufferSampleProperties = new FParamDict();
+				if (ImageBufferRef)
+				{
+					// Start with a safe 1:1 aspect ratio assumption.
+					long ax = 1;
+					long ay = 1;
+					// Get the initial aspect ratio from the parsed codec info.
+					if (NextImage.SourceInfo->ParsedInfo.GetAspectRatio().IsSet())
+					{
+						ax = NextImage.SourceInfo->ParsedInfo.GetAspectRatio().Width;
+						ay = NextImage.SourceInfo->ParsedInfo.GetAspectRatio().Height;
+					}
+					// If there is aspect ratio information on the image itself and it's valid, use that instead.
+					NSDictionary* Dict = (NSDictionary*)CVBufferGetAttachments(ImageBufferRef, kCVAttachmentMode_ShouldPropagate);
+					if (Dict)
+					{
+						NSDictionary* AspectDict = (NSDictionary*)Dict[(__bridge NSString*)kCVImageBufferPixelAspectRatioKey];
+						if (AspectDict)
+						{
+							NSNumber* hs = (NSNumber*)AspectDict[(__bridge NSString*)kCVImageBufferPixelAspectRatioHorizontalSpacingKey];
+							NSNumber* vs = (NSNumber*)AspectDict[(__bridge NSString*)kCVImageBufferPixelAspectRatioVerticalSpacingKey];
+							if (hs && vs)
+							{
+								long parx = [hs longValue];
+								long pary = [vs longValue];
+								if (parx && pary)
+								{
+									ax = parx;
+									ay = pary;
+								}
+							}
+						}
+					}
+					PixelAspectRatio = (double)ax / (double)ay;
+
+					OutputBufferSampleProperties->Set("width",        FVariantValue((int64) NextImage.SourceInfo->ParsedInfo.GetResolution().Width));
+					OutputBufferSampleProperties->Set("height",       FVariantValue((int64) NextImage.SourceInfo->ParsedInfo.GetResolution().Height));
+					OutputBufferSampleProperties->Set("crop_left",    FVariantValue((int64) 0));
+					OutputBufferSampleProperties->Set("crop_right",   FVariantValue((int64) 0));
+					OutputBufferSampleProperties->Set("crop_top",     FVariantValue((int64) 0));
+					OutputBufferSampleProperties->Set("crop_bottom",  FVariantValue((int64) 0));
+					OutputBufferSampleProperties->Set("aspect_ratio", FVariantValue((double) PixelAspectRatio));
+					OutputBufferSampleProperties->Set("aspect_w",     FVariantValue((int64) ax));
+					OutputBufferSampleProperties->Set("aspect_h",     FVariantValue((int64) ay));
+					OutputBufferSampleProperties->Set("fps_num",      FVariantValue((int64) 0 ));
+					OutputBufferSampleProperties->Set("fps_denom",    FVariantValue((int64) 0 ));
+					OutputBufferSampleProperties->Set("pixelfmt",     FVariantValue((int64)EPixelFormat::PF_B8G8R8A8));
+				}
+				else
+				{
+					OutputBufferSampleProperties->Set("is_dummy", FVariantValue(true));
+				}
+				OutputBufferSampleProperties->Set("pts",          FVariantValue(NextImage.SourceInfo->PTS));
+				OutputBufferSampleProperties->Set("duration",     FVariantValue(NextImage.SourceInfo->Duration));
+
+				TSharedPtr<FElectraPlayerVideoDecoderOutputApple, ESPMode::ThreadSafe> DecoderOutput = RenderOutputBuffer->GetBufferProperties().GetValue("texture").GetSharedPointer<FElectraPlayerVideoDecoderOutputApple>();
+
+				DecoderOutput->Initialize(ImageBufferRef, OutputBufferSampleProperties);
+                if (ImageBufferRef)
+                {
+                    CFRelease(ImageBufferRef);
+                }
+
+				// Return the buffer to the renderer.
+				Renderer->ReturnBuffer(RenderOutputBuffer, true, *OutputBufferSampleProperties);
+			}
+		}
+	}
+}
+
+} // namespace Electra
+
+// -----------------------------------------------------------------------------------------------------------------------------
+
+FVideoDecoderOutput* FElectraPlayerPlatformVideoDecoderOutputFactory::Create()
+{
+    return new Electra::FElectraPlayerVideoDecoderOutputApple();
+}
+
+
+#endif
+
diff --git a/Engine/Plugins/Media/ElectraPlayer/Source/ElectraPlayerRuntime/Private/Runtime/Decoder/AudioDecoderAAC.h b/Engine/Plugins/Media/ElectraPlayer/Source/ElectraPlayerRuntime/Private/Runtime/Decoder/AudioDecoderAAC.h
new file mode 100644
--- /dev/null
+++ b/Engine/Plugins/Media/ElectraPlayer/Source/ElectraPlayerRuntime/Private/Runtime/Decoder/AudioDecoderAAC.h
@@ -0,0 +1,84 @@
+// Copyright Epic Games, Inc. All Rights Reserved.
+
+#pragma once
+
+#include "PlayerCore.h"
+
+#include "StreamAccessUnitBuffer.h"
+#include "Decoder/DecoderBase.h"
+
+
+
+namespace Electra
+{
+	class IPlayerSessionServices;
+
+
+	/**
+	 * LC-AAC audio decoder.
+	**/
+	class IAudioDecoderAAC : public IDecoderBase, public IAccessUnitBufferInterface, public IDecoderAUBufferDiags, public IDecoderReadyBufferDiags
+	{
+	public:
+		struct FSystemConfiguration
+		{
+			FSystemConfiguration();
+
+			struct FThreadConfig
+			{
+				FMediaRunnable::Param		Decoder;				//!< Decoder thread settings.
+				FMediaRunnable::Param		PassOn;					//!< Settings for thread passing decoded samples to the renderer.
+			};
+			FThreadConfig	ThreadConfig;
+		};
+
+		struct FInstanceConfiguration
+		{
+			FInstanceConfiguration();
+
+			FSystemConfiguration::FThreadConfig	ThreadConfig;					//!< Thread configuration (defaults to values set in SystemConfiguration)
+		};
+
+		static bool Startup(const FSystemConfiguration& Config);
+		static void Shutdown();
+
+		static IAudioDecoderAAC* Create();
+
+		virtual ~IAudioDecoderAAC() = default;
+
+		virtual void SetPlayerSessionServices(IPlayerSessionServices* SessionServices) = 0;
+
+		virtual void Open(const FInstanceConfiguration& Config) = 0;
+		virtual void Close() = 0;
+
+
+		//-------------------------------------------------------------------------
+		// Methods from IDecoderBase
+		//
+		virtual void SetRenderer(TSharedPtr<IMediaRenderer, ESPMode::ThreadSafe> InRenderer) = 0;
+
+		//-------------------------------------------------------------------------
+		// Methods from IAccessUnitBufferInterface
+		//
+		//! Attempts to push an access unit to the decoder. Ownership of the access unit is transferred if the push is successful.
+		virtual EAUpushResult AUdataPushAU(FAccessUnit* AccessUnit) = 0;
+		//! Notifies the decoder that there will be no further access units.
+		virtual void AUdataPushEOD() = 0;
+		//! Instructs the decoder to flush all pending input and all already decoded output.
+		virtual void AUdataFlushEverything() = 0;
+
+		//-------------------------------------------------------------------------
+		// Methods from IDecoderAUBufferDiags
+		//
+		//! Registers an AU input buffer listener.
+		virtual void SetAUInputBufferListener(IAccessUnitBufferListener* Listener) = 0;
+
+		//-------------------------------------------------------------------------
+		// Methods from IDecoderReadyBufferDiags
+		//
+		virtual void SetReadyBufferListener(IDecoderOutputBufferListener* Listener) = 0;
+	};
+
+
+} // namespace Electra
+
