ElectraPlayer: RELEASE. Moved from NFL to public. Note that this is done via a Delete and "Add New" operation as discussed.

diff --git a/Engine/Plugins/Media/ElectraPlayer/Source/ElectraPlayerRuntime/Private/Runtime/Decoder/DX/AudioDecoderAAC_DX.cpp b/Engine/Plugins/Media/ElectraPlayer/Source/ElectraPlayerRuntime/Private/Runtime/Decoder/DX/AudioDecoderAAC_DX.cpp
new file mode 100644
--- /dev/null
+++ b/Engine/Plugins/Media/ElectraPlayer/Source/ElectraPlayerRuntime/Private/Runtime/Decoder/DX/AudioDecoderAAC_DX.cpp
@@ -0,0 +1,1131 @@
+// Copyright Epic Games, Inc. All Rights Reserved.
+
+#ifdef ELECTRA_ENABLE_MFDECODER
+
+#include "PlayerCore.h"
+#include "ElectraPlayerPrivate.h"
+#include "ElectraPlayerPrivate_Platform.h"
+
+#include "StreamAccessUnitBuffer.h"
+#include "Decoder/AudioDecoderAAC.h"
+#include "Renderer/RendererBase.h"
+#include "Player/PlayerSessionServices.h"
+#include "Utilities/UtilsMPEG.h"
+#include "Utilities/UtilsMPEGAudio.h"
+#include "Utilities/StringHelpers.h"
+
+#include "DecoderErrors_DX.h"
+
+/*********************************************************************************************************************/
+/*********************************************************************************************************************/
+/*********************************************************************************************************************/
+
+#include "Windows/AllowWindowsPlatformTypes.h"
+#include "Windows/WindowsHWrapper.h"
+#include "HAL/LowLevelMemTracker.h"
+
+THIRD_PARTY_INCLUDES_START
+#include "mftransform.h"
+#include "mfapi.h"
+#include "mferror.h"
+#include "mfidl.h"
+//#include "wmcodecdsp.h" // for MEDIASUBTYPE_RAW_AAC1, but requires linking of additional dll for the constant.
+THIRD_PARTY_INCLUDES_END
+
+#include "Windows/HideWindowsPlatformTypes.h"
+
+namespace
+{
+	static const GUID MFTmsAACDecoder_Audio = { 0x32d186a7, 0x218f, 0x4c75, { 0x88, 0x76, 0xdd, 0x77, 0x27, 0x3a, 0x89, 0x99 } };
+	static const GUID MEDIASUBTYPE_RAW_AAC1_Audio = { 0x000000FF, 0x0000, 0x0010, { 0x80, 0x00, 0x00, 0xaa, 0x00, 0x38, 0x9b, 0x71 } };
+}
+
+
+/*********************************************************************************************************************/
+/*********************************************************************************************************************/
+/*********************************************************************************************************************/
+
+#define VERIFY_HR(FNcall, Msg, What)	\
+res = FNcall;							\
+if (FAILED(res))						\
+{										\
+	PostError(res, Msg, What);			\
+	return false;						\
+}
+
+/*********************************************************************************************************************/
+/*********************************************************************************************************************/
+/*********************************************************************************************************************/
+
+
+#define AACDEC_PCM_SAMPLE_SIZE  2048
+#define AACDEC_MAX_CHANNELS		6
+
+
+namespace Electra
+{
+
+
+	/**
+	 * AAC audio decoder class implementation.
+	**/
+	class FAudioDecoderAAC : public IAudioDecoderAAC, public FMediaThread
+	{
+	public:
+		static bool Startup(const IAudioDecoderAAC::FSystemConfiguration& InConfig);
+		static void Shutdown();
+
+		FAudioDecoderAAC();
+		virtual ~FAudioDecoderAAC();
+
+		virtual void SetPlayerSessionServices(IPlayerSessionServices* SessionServices) override;
+
+		virtual void Open(const FInstanceConfiguration& InConfig) override;
+		virtual void Close() override;
+
+		virtual void SetRenderer(TSharedPtr<IMediaRenderer, ESPMode::ThreadSafe> InRenderer) override;
+
+		virtual void SetAUInputBufferListener(IAccessUnitBufferListener* Listener) override;
+
+		virtual void SetReadyBufferListener(IDecoderOutputBufferListener* Listener) override;
+
+		virtual IAccessUnitBufferInterface::EAUpushResult AUdataPushAU(FAccessUnit* AccessUnit) override;
+		virtual void AUdataPushEOD() override;
+		virtual void AUdataFlushEverything() override;
+
+	private:
+		struct FDecoderOutputBuffer
+		{
+			FDecoderOutputBuffer()
+			{
+				FMemory::Memzero(mOutputStreamInfo);
+				FMemory::Memzero(mOutputBuffer);
+			}
+			~FDecoderOutputBuffer()
+			{
+				if (mOutputBuffer.pSample)
+					mOutputBuffer.pSample->Release();
+			}
+			TRefCountPtr<IMFSample> DetachOutputSample()
+			{
+				TRefCountPtr<IMFSample> pOutputSample;
+				if (mOutputBuffer.pSample)
+				{
+					// mOutputBuffer.pSample already holds a reference, don't need to addref here.
+					pOutputSample = TRefCountPtr<IMFSample>(mOutputBuffer.pSample, false);
+					mOutputBuffer.pSample = nullptr;
+				}
+				return(pOutputSample);
+			}
+			void PrepareForProcess()
+			{
+				mOutputBuffer.dwStatus = 0;
+				mOutputBuffer.dwStreamID = 0;
+				mOutputBuffer.pEvents = nullptr;
+			}
+			MFT_OUTPUT_STREAM_INFO	mOutputStreamInfo;
+			MFT_OUTPUT_DATA_BUFFER	mOutputBuffer;
+		};
+
+		void StartThread();
+		void StopThread();
+		void WorkerThread();
+
+		bool InternalDecoderCreate();
+		bool SetDecoderOutputType();
+
+		bool CreateDecodedSamplePool();
+		void DestroyDecodedSamplePool();
+
+		void ReturnUnusedFrame();
+
+		void NotifyReadyBufferListener(bool bHaveOutput);
+
+		void PostError(HRESULT ApiReturnValue, const FString& Message, uint16 Code, UEMediaError Error = UEMEDIA_ERROR_OK);
+
+		bool CreateDecoderOutputBuffer();
+		bool Decode(int32& OutInputSizeConsumed, const void* InData, int32 InSize, int64 PTS, int64 Duration, bool bFlushOnly, bool bRender);
+
+		bool IsDifferentFormat(const FAccessUnit* Data);
+
+		FInstanceConfiguration													Config;
+
+		FAccessUnitBuffer														AccessUnits;
+
+		FMediaEvent																TerminateThreadEvent;
+		FMediaEvent																FlushDecoderEvent;
+		FMediaEvent																DecoderFlushedEvent;
+		bool																	bThreadStarted;
+
+		TSharedPtr<IMediaRenderer, ESPMode::ThreadSafe>							Renderer;
+		int32																	MaxDecodeBufferSize;
+
+		FMediaCriticalSection													ListenerCriticalSection;
+		IAccessUnitBufferListener* InputBufferListener;
+		IDecoderOutputBufferListener* ReadyBufferListener;
+
+		IPlayerSessionServices* SessionServices;
+
+		TSharedPtr<MPEG::FAACDecoderConfigurationRecord, ESPMode::ThreadSafe>	ConfigRecord;
+		TSharedPtrTS<FAccessUnit::CodecData>									CurrentCodecData;
+		bool																	bHaveDiscontinuity;
+
+		TRefCountPtr<IMFTransform>												DecoderTransform;
+		TRefCountPtr<IMFMediaType>												CurrentOutputMediaType;
+		MFT_INPUT_STREAM_INFO													DecoderInputStreamInfo;
+		MFT_OUTPUT_STREAM_INFO													DecoderOutputStreamInfo;
+
+		FTimeValue																CurrentPTS;
+		FTimeValue																NextExpectedPTS;
+		TUniquePtr<FDecoderOutputBuffer>										CurrentDecoderOutputBuffer;
+		IMediaRenderer::IBuffer* CurrentRenderOutputBuffer;
+		FParamDict																BufferAcquireOptions;
+		FParamDict																OutputBufferSampleProperties;
+
+	public:
+		static FSystemConfiguration												SystemConfig;
+	};
+
+	IAudioDecoderAAC::FSystemConfiguration	FAudioDecoderAAC::SystemConfig;
+
+
+	/***************************************************************************************************************************************************/
+	/***************************************************************************************************************************************************/
+	/***************************************************************************************************************************************************/
+
+	IAudioDecoderAAC::FSystemConfiguration::FSystemConfiguration()
+	{
+		ThreadConfig.Decoder.Priority = TPri_Normal;
+		ThreadConfig.Decoder.StackSize = 65536;
+		ThreadConfig.Decoder.CoreAffinity = -1;
+		// Not needed, but setting meaningful values anyway.
+		ThreadConfig.PassOn.Priority = TPri_Normal;
+		ThreadConfig.PassOn.StackSize = 32768;
+		ThreadConfig.PassOn.CoreAffinity = -1;
+	}
+
+	IAudioDecoderAAC::FInstanceConfiguration::FInstanceConfiguration()
+		: ThreadConfig(FAudioDecoderAAC::SystemConfig.ThreadConfig)
+	{
+	}
+
+	bool IAudioDecoderAAC::Startup(const IAudioDecoderAAC::FSystemConfiguration& InConfig)
+	{
+		return FAudioDecoderAAC::Startup(InConfig);
+	}
+
+	void IAudioDecoderAAC::Shutdown()
+	{
+		FAudioDecoderAAC::Shutdown();
+	}
+
+	IAudioDecoderAAC* IAudioDecoderAAC::Create()
+	{
+		return new FAudioDecoderAAC;
+	}
+
+	/***************************************************************************************************************************************************/
+	/***************************************************************************************************************************************************/
+	/***************************************************************************************************************************************************/
+
+
+
+
+	//-----------------------------------------------------------------------------
+	/**
+	 * Decoder system startup
+	 *
+	 * @param config
+	 *
+	 * @return
+	 */
+	bool FAudioDecoderAAC::Startup(const IAudioDecoderAAC::FSystemConfiguration& InConfig)
+	{
+		SystemConfig = InConfig;
+		return true;
+	}
+
+
+	//-----------------------------------------------------------------------------
+	/**
+	 * Decoder system shutdown.
+	 */
+	void FAudioDecoderAAC::Shutdown()
+	{
+	}
+
+
+
+	//-----------------------------------------------------------------------------
+	/**
+	 * Constructor
+	 */
+	FAudioDecoderAAC::FAudioDecoderAAC()
+		: FMediaThread("ElectraPlayer::AAC decoder")
+		, bThreadStarted(false)
+		, Renderer(nullptr)
+		, MaxDecodeBufferSize(0)
+		, InputBufferListener(nullptr)
+		, ReadyBufferListener(nullptr)
+		, SessionServices(nullptr)
+		, CurrentRenderOutputBuffer(nullptr)
+	{
+	}
+
+
+	//-----------------------------------------------------------------------------
+	/**
+	 * Destructor
+	 */
+	FAudioDecoderAAC::~FAudioDecoderAAC()
+	{
+		Close();
+	}
+
+
+	void FAudioDecoderAAC::SetPlayerSessionServices(IPlayerSessionServices* InSessionServices)
+	{
+		SessionServices = InSessionServices;
+	}
+
+
+	//-----------------------------------------------------------------------------
+	/**
+	 * Creates the decoded sample blocks output buffers.
+	 *
+	 * @return
+	 */
+	bool FAudioDecoderAAC::CreateDecodedSamplePool()
+	{
+		check(Renderer);
+		FParamDict poolOpts;
+		// Assume LC-AAC with at most 6 channels and 16 bit decoded PCM samples. This is ok for HE-AAC as well which can at most be stereo.
+		uint32 frameSize = sizeof(int16) * AACDEC_MAX_CHANNELS * 1024;
+		static_assert(sizeof(int16) * AACDEC_MAX_CHANNELS * 1024 >= sizeof(int16) * 2 * 2048, "Sample pool size not correct!");
+
+		poolOpts.Set("max_buffer_size", FVariantValue((int64)frameSize));
+		//	poolOpts.Set("buffer_alignment", FVariantValue((int64) 32));
+		poolOpts.Set("num_buffers", FVariantValue((int64)8));
+
+		UEMediaError Error = Renderer->CreateBufferPool(poolOpts);
+		check(Error == UEMEDIA_ERROR_OK);
+
+		if (Error != UEMEDIA_ERROR_OK)
+		{
+			PostError(S_OK, "Failed to create sample pool", ERRCODE_INTERNAL_COULD_NOT_CREATE_SAMPLE_POOL, Error);
+		}
+
+		MaxDecodeBufferSize = (int32)Renderer->GetBufferPoolProperties().GetValue("max_buffers").GetInt64();
+
+		return Error == UEMEDIA_ERROR_OK;
+	}
+
+
+
+	//-----------------------------------------------------------------------------
+	/**
+	 * Destroys the pool of decoded sample blocks.
+	 */
+	void FAudioDecoderAAC::DestroyDecodedSamplePool()
+	{
+		Renderer->ReleaseBufferPool();
+	}
+
+
+
+	//-----------------------------------------------------------------------------
+	/**
+	 * Opens a decoder instance
+	 *
+	 * @param config
+	 *
+	 * @return
+	 */
+	void FAudioDecoderAAC::Open(const IAudioDecoderAAC::FInstanceConfiguration& InConfig)
+	{
+		Config = InConfig;
+
+		// Set a large enough size to hold a single access unit. Since we will be asking for a new AU on
+		// demand there is no need to overly restrict ourselves here. But it must be large enough to
+		// hold at least the largest expected access unit.
+		AccessUnits.CapacitySet(FAccessUnitBuffer::FConfiguration(2 << 20, 60.0));
+
+		StartThread();
+	}
+
+
+
+	//-----------------------------------------------------------------------------
+	/**
+	 * Closes the decoder instance.
+	 *
+	 * @param bKeepDecoder
+	 */
+	void FAudioDecoderAAC::Close()
+	{
+		StopThread();
+	}
+
+
+
+	//-----------------------------------------------------------------------------
+	/**
+	 * Sets a new renderer.
+	 *
+	 * @param InRenderer
+	 */
+	void FAudioDecoderAAC::SetRenderer(TSharedPtr<IMediaRenderer, ESPMode::ThreadSafe> InRenderer)
+	{
+		Renderer = InRenderer;
+	}
+
+
+	//-----------------------------------------------------------------------------
+	/**
+	 * Sets an AU input buffer listener.
+	 *
+	 * @param Listener
+	 */
+	void FAudioDecoderAAC::SetAUInputBufferListener(IAccessUnitBufferListener* Listener)
+	{
+		FMediaCriticalSection::ScopedLock lock(ListenerCriticalSection);
+		InputBufferListener = Listener;
+	}
+
+	void FAudioDecoderAAC::SetReadyBufferListener(IDecoderOutputBufferListener* Listener)
+	{
+		FMediaCriticalSection::ScopedLock lock(ListenerCriticalSection);
+		ReadyBufferListener = Listener;
+	}
+
+
+
+
+	//-----------------------------------------------------------------------------
+	/**
+	 * Creates and runs the decoder thread.
+	 */
+	void FAudioDecoderAAC::StartThread()
+	{
+		ThreadSetPriority(Config.ThreadConfig.Decoder.Priority);
+		ThreadSetStackSize(Config.ThreadConfig.Decoder.StackSize);
+		ThreadSetCoreAffinity(Config.ThreadConfig.Decoder.CoreAffinity);
+		ThreadStart(Electra::MakeDelegate(this, &FAudioDecoderAAC::WorkerThread));
+		bThreadStarted = true;
+	}
+
+
+
+	//-----------------------------------------------------------------------------
+	/**
+	 * Stops the decoder thread.
+	 */
+	void FAudioDecoderAAC::StopThread()
+	{
+		if (bThreadStarted)
+		{
+			TerminateThreadEvent.Signal();
+			ThreadWaitDone();
+			bThreadStarted = false;
+		}
+	}
+
+
+
+
+	IAccessUnitBufferInterface::EAUpushResult FAudioDecoderAAC::AUdataPushAU(FAccessUnit* AccessUnit)
+	{
+		bool bOk;
+		AccessUnit->AddRef();
+		bOk = AccessUnits.Push(AccessUnit);
+		if (!bOk)
+		{
+			FAccessUnit::Release(AccessUnit);
+		}
+		return bOk ? IAccessUnitBufferInterface::EAUpushResult::Ok : IAccessUnitBufferInterface::EAUpushResult::Full;
+	}
+
+	void FAudioDecoderAAC::AUdataPushEOD()
+	{
+		AccessUnits.PushEndOfData();
+	}
+
+	void FAudioDecoderAAC::AUdataFlushEverything()
+	{
+		FlushDecoderEvent.Signal();
+		DecoderFlushedEvent.WaitAndReset();
+	}
+
+
+	//-----------------------------------------------------------------------------
+	/**
+	 * Notify optional decode-ready listener that we will now be producing output data.
+	 *
+	 * @param bHaveOutput
+	 */
+	void FAudioDecoderAAC::NotifyReadyBufferListener(bool bHaveOutput)
+	{
+		if (ReadyBufferListener)
+		{
+			IDecoderOutputBufferListener::FDecodeReadyStats stats;
+			stats.MaxDecodedElementsReady = MaxDecodeBufferSize;
+			stats.NumElementsInDecoder = CurrentRenderOutputBuffer ? 1 : 0;
+			stats.bOutputStalled = !bHaveOutput;
+			stats.bEODreached = AccessUnits.IsEndOfData() && stats.NumDecodedElementsReady == 0 && stats.NumElementsInDecoder == 0;
+			ListenerCriticalSection.Lock();
+			if (ReadyBufferListener)
+			{
+				ReadyBufferListener->DecoderOutputReady(stats);
+			}
+			ListenerCriticalSection.Unlock();
+		}
+	}
+
+
+
+	void FAudioDecoderAAC::PostError(HRESULT ApiReturnValue, const FString& Message, uint16 Code, UEMediaError Error)
+	{
+		check(SessionServices);
+		if (SessionServices)
+		{
+			FErrorDetail err;
+			err.SetError(Error != UEMEDIA_ERROR_OK ? Error : UEMEDIA_ERROR_DETAIL);
+			err.SetFacility(Facility::EFacility::AACDecoder);
+			err.SetCode(Code);
+			err.SetMessage(Message);
+
+			if (ApiReturnValue != S_OK)
+			{
+				err.SetPlatformMessage(FString::Printf(TEXT("%s (0x%08lx)"), *GetComErrorDescription(ApiReturnValue), ApiReturnValue));
+			}
+			SessionServices->PostError(err);
+		}
+	}
+
+
+
+	bool FAudioDecoderAAC::InternalDecoderCreate()
+	{
+		TRefCountPtr<IMFTransform>	Decoder;
+		TRefCountPtr<IMFMediaType>	MediaType;
+		HRESULT					res;
+
+		if (!ConfigRecord.IsValid())
+		{
+			PostError(0, "No CSD to create audio decoder with", ERRCODE_INTERNAL_COULD_NOT_CREATE_DECODER);
+			return false;
+		}
+
+		// Create decoder transform
+		VERIFY_HR(CoCreateInstance(MFTmsAACDecoder_Audio, nullptr, CLSCTX_INPROC_SERVER, IID_IMFTransform, reinterpret_cast<void**>(&Decoder)), "CoCreateInstance failed", ERRCODE_INTERNAL_COULD_NOT_CREATE_DECODER);
+
+		// Create input media type
+		VERIFY_HR(MFCreateMediaType(MediaType.GetInitReference()), "MFCreateMediaType failed", ERRCODE_INTERNAL_COULD_NOT_CREATE_INPUT_MEDIA_TYPE);
+		VERIFY_HR(MediaType->SetGUID(MF_MT_MAJOR_TYPE, MFMediaType_Audio), "Failed to set input media type for audio", ERRCODE_INTERNAL_COULD_NOT_SET_INPUT_AUDIO_MAJOR_TYPE);
+		UINT32 PayloadType = 0;	// 0=raw, 1=adts, 2=adif, 3=latm
+		VERIFY_HR(MediaType->SetGUID(MF_MT_SUBTYPE, MEDIASUBTYPE_RAW_AAC1_Audio), "Failed to set input media audio type to RAW AAC", ERRCODE_INTERNAL_COULD_NOT_SET_INPUT_AUDIO_AAC_SUBTYPE);
+		VERIFY_HR(MediaType->SetUINT32(MF_MT_AAC_PAYLOAD_TYPE, PayloadType), "Failed to set input media audio payload type", ERRCODE_INTERNAL_COULD_NOT_SET_INPUT_AUDIO_AAC_SUBTYPE);
+		VERIFY_HR(MediaType->SetUINT32(MF_MT_AUDIO_SAMPLES_PER_SECOND, ConfigRecord->SamplingRate), FString::Printf(TEXT("Failed to set input audio sampling rate to %u"), ConfigRecord->SamplingRate), ERRCODE_INTERNAL_COULD_NOT_SET_INPUT_AUDIO_SAMPLING_RATE);
+		VERIFY_HR(MediaType->SetUINT32(MF_MT_AUDIO_NUM_CHANNELS, ConfigRecord->ChannelConfiguration), FString::Printf(TEXT("Failed to set input audio number of channels to %u"), ConfigRecord->ChannelConfiguration), ERRCODE_INTERNAL_COULD_NOT_SET_INPUT_AUDIO_CHANNEL_COUNT);
+		VERIFY_HR(MediaType->SetBlob(MF_MT_USER_DATA, ConfigRecord->GetCodecSpecificData().GetData(), ConfigRecord->GetCodecSpecificData().Num()), "Failed to set input audio CSD", ERRCODE_INTERNAL_COULD_NOT_SET_INPUT_AUDIO_CSD);
+		// Set input media type with decoder
+		VERIFY_HR(Decoder->SetInputType(0, MediaType, 0), "Failed to set audio decoder input type", ERRCODE_INTERNAL_COULD_NOT_SET_INPUT_TYPE);
+		DecoderTransform = Decoder;
+
+		// Set decoder output type to PCM
+		if (!SetDecoderOutputType())
+		{
+			DecoderTransform = nullptr;
+			return false;
+		}
+		// Get input and output stream information from decoder
+		VERIFY_HR(DecoderTransform->GetInputStreamInfo(0, &DecoderInputStreamInfo), "Failed to get audio decoder input stream information", ERRCODE_INTERNAL_COULD_NOT_GET_INPUT_STREAM_FORMAT_INFO);
+		VERIFY_HR(DecoderTransform->GetOutputStreamInfo(0, &DecoderOutputStreamInfo), "Failed to get audio decoder output stream information", ERRCODE_INTERNAL_COULD_NOT_GET_OUTPUT_STREAM_FORMAT_INFO);
+
+		// Start the decoder transform
+		VERIFY_HR(DecoderTransform->ProcessMessage(MFT_MESSAGE_NOTIFY_BEGIN_STREAMING, 0), "Failed to set audio decoder stream begin", ERRCODE_INTERNAL_COULD_NOT_SET_DECODER_BEGIN);
+		VERIFY_HR(DecoderTransform->ProcessMessage(MFT_MESSAGE_NOTIFY_START_OF_STREAM, 0), "Failed to start audio decoder", ERRCODE_INTERNAL_COULD_NOT_SET_DECODER_START);
+
+		return true;
+	}
+
+
+	bool FAudioDecoderAAC::SetDecoderOutputType()
+	{
+		TRefCountPtr<IMFMediaType> MediaType;
+		HRESULT res;
+
+		uint32 TypeIndex = 0;
+		while (SUCCEEDED(DecoderTransform->GetOutputAvailableType(0, TypeIndex++, MediaType.GetInitReference())))
+		{
+			GUID Subtype;
+			res = MediaType->GetGUID(MF_MT_SUBTYPE, &Subtype);
+			if (SUCCEEDED(res) && Subtype == MFAudioFormat_PCM)
+			{
+				VERIFY_HR(DecoderTransform->SetOutputType(0, MediaType, 0), "Failed to set audio decoder output type", ERRCODE_INTERNAL_COULD_NOT_SET_OUTPUT_TYPE);
+				CurrentOutputMediaType = MediaType;
+				return true;
+			}
+		}
+		PostError(S_OK, "Failed to set audio decoder output type to PCM", ERRCODE_INTERNAL_COULD_NOT_SET_OUTPUT_TYPE_TO_PCM);
+		return false;
+	}
+
+
+
+	void FAudioDecoderAAC::ReturnUnusedFrame()
+	{
+		if (CurrentRenderOutputBuffer)
+		{
+			OutputBufferSampleProperties.Clear();
+			Renderer->ReturnBuffer(CurrentRenderOutputBuffer, false, OutputBufferSampleProperties);
+			CurrentRenderOutputBuffer = nullptr;
+		}
+	}
+
+
+
+	bool FAudioDecoderAAC::CreateDecoderOutputBuffer()
+	{
+		HRESULT									res;
+		TUniquePtr<FDecoderOutputBuffer>		NewDecoderOutputBuffer(new FDecoderOutputBuffer);
+
+		VERIFY_HR(DecoderTransform->GetOutputStreamInfo(0, &NewDecoderOutputBuffer->mOutputStreamInfo), "Failed to get audio decoder output stream information", ERRCODE_INTERNAL_COULD_NOT_GET_OUTPUT_STREAM_FORMAT_INFO);
+		// Do we need to provide the sample output buffer or does the decoder create it for us?
+		if ((NewDecoderOutputBuffer->mOutputStreamInfo.dwFlags & (MFT_OUTPUT_STREAM_PROVIDES_SAMPLES | MFT_OUTPUT_STREAM_CAN_PROVIDE_SAMPLES)) == 0)
+		{
+			// We have to provide the output sample buffer.
+			TRefCountPtr<IMFSample>		OutputSample;
+			TRefCountPtr<IMFMediaBuffer>	OutputBuffer;
+			VERIFY_HR(MFCreateSample(OutputSample.GetInitReference()), "Failed to create output sample for audio decoder", ERRCODE_INTERNAL_COULD_NOT_CREATE_OUTPUT_SAMPLE);
+			if (NewDecoderOutputBuffer->mOutputStreamInfo.cbAlignment > 0)
+			{
+				VERIFY_HR(MFCreateAlignedMemoryBuffer(NewDecoderOutputBuffer->mOutputStreamInfo.cbSize, NewDecoderOutputBuffer->mOutputStreamInfo.cbAlignment, OutputBuffer.GetInitReference()), "Failed to create aligned output buffer for audio decoder", ERRCODE_INTERNAL_COULD_NOT_CREATE_ALIGNED_OUTPUTBUFFER);
+			}
+			else
+			{
+				VERIFY_HR(MFCreateMemoryBuffer(NewDecoderOutputBuffer->mOutputStreamInfo.cbSize, OutputBuffer.GetInitReference()), "Failed to create output buffer for audio decoder", ERRCODE_INTERNAL_COULD_NOT_CREATE_OUTPUTBUFFER);
+			}
+			VERIFY_HR(OutputSample->AddBuffer(OutputBuffer.GetReference()), "Failed to add sample buffer to output sample for audio decoder", ERRCODE_INTERNAL_COULD_NOT_ADD_OUTPUT_BUFFER_TO_SAMPLE);
+			(NewDecoderOutputBuffer->mOutputBuffer.pSample = OutputSample.GetReference())->AddRef();
+			OutputSample = nullptr;
+		}
+		CurrentDecoderOutputBuffer = MoveTemp(NewDecoderOutputBuffer);
+		return true;
+	}
+
+	bool FAudioDecoderAAC::Decode(int32& OutInputSizeConsumed, const void* InData, int32 InSize, int64 PTS, int64 Duration, bool bFlushOnly, bool bRender)
+	{
+		HRESULT				res;
+		TRefCountPtr<IMFSample>	InputSample;
+		bool				bFlush = InData == nullptr;
+
+		OutInputSizeConsumed = 0;
+		if (bFlush)
+		{
+			VERIFY_HR(DecoderTransform->ProcessMessage(MFT_MESSAGE_NOTIFY_END_OF_STREAM, 0), "Failed to set audio decoder end of stream notification", ERRCODE_INTERNAL_COULD_NOT_SET_DECODER_ENDOFSTREAM);
+			if (bFlushOnly)
+			{
+				VERIFY_HR(DecoderTransform->ProcessMessage(MFT_MESSAGE_COMMAND_FLUSH, 0), "Failed to issue audio decoder flush command", ERRCODE_INTERNAL_COULD_NOT_SET_DECODER_FLUSHCOMMAND);
+			}
+			else
+			{
+				VERIFY_HR(DecoderTransform->ProcessMessage(MFT_MESSAGE_COMMAND_DRAIN, 0), "Failed to issue audio decoder drain command", ERRCODE_INTERNAL_COULD_NOT_SET_DECODER_DRAINCOMMAND);
+			}
+		}
+		else
+		{
+			CSV_SCOPED_TIMING_STAT(ElectraPlayer, FAudioDecoderAAC_Worker);
+
+			// Create the input sample.
+			TRefCountPtr<IMFMediaBuffer>	InputSampleBuffer;
+			BYTE* pbNewBuffer = nullptr;
+			DWORD					dwMaxBufferSize = 0;
+			DWORD					dwSize = 0;
+			LONGLONG				llSampleTime = 0;
+
+			VERIFY_HR(MFCreateSample(InputSample.GetInitReference()), "Failed to create audio decoder input sample", ERRCODE_INTERNAL_COULD_NOT_CREATE_INPUT_SAMPLE);
+			VERIFY_HR(MFCreateMemoryBuffer((DWORD)InSize, InputSampleBuffer.GetInitReference()), "Failed to create audio decoder input sample memory buffer", ERRCODE_INTERNAL_COULD_NOT_CREATE_INPUTBUFFER);
+			VERIFY_HR(InputSample->AddBuffer(InputSampleBuffer.GetReference()), "Failed to set audio decoder input buffer with sample", ERRCODE_INTERNAL_COULD_NOT_ADD_INPUT_BUFFER_TO_SAMPLE);
+			VERIFY_HR(InputSampleBuffer->Lock(&pbNewBuffer, &dwMaxBufferSize, &dwSize), "Failed to lock audio decoder input sample buffer", ERRCODE_INTERNAL_COULD_NOT_LOCK_INPUT_BUFFER);
+			FMemory::Memcpy(pbNewBuffer, InData, InSize);
+			VERIFY_HR(InputSampleBuffer->Unlock(), "Failed to unlock audio decoder input sample buffer", ERRCODE_INTERNAL_COULD_NOT_UNLOCK_INPUT_BUFFER);
+			VERIFY_HR(InputSampleBuffer->SetCurrentLength((DWORD)InSize), "Failed to set audio decoder input sample buffer length", ERRCODE_INTERNAL_COULD_NOT_SET_BUFFER_CURRENT_LENGTH);
+			// Set sample attributes
+			llSampleTime = PTS;
+			VERIFY_HR(InputSample->SetSampleTime(llSampleTime), "Failed to set audio decoder input sample decode time", ERRCODE_INTERNAL_COULD_NOT_SET_INPUT_SAMPLE_TIME);
+			llSampleTime = Duration;
+			VERIFY_HR(InputSample->SetSampleDuration(llSampleTime), "Failed to set audio decode intput sample duration", ERRCODE_INTERNAL_COULD_NOT_SET_INPUT_SAMPLE_DURATION);
+		}
+
+		while (!TerminateThreadEvent.IsSignaled())
+		{
+			if (FlushDecoderEvent.IsSignaled() && !bFlush)
+			{
+				break;
+			}
+			if (!CurrentDecoderOutputBuffer.IsValid())
+			{
+				CSV_SCOPED_TIMING_STAT(ElectraPlayer, FAudioDecoderAAC_Worker);
+
+				if (!CreateDecoderOutputBuffer())
+				{
+					return false;
+				}
+			}
+
+			DWORD	dwStatus = 0;
+			CurrentDecoderOutputBuffer->PrepareForProcess();
+
+			{
+				CSV_SCOPED_TIMING_STAT(ElectraPlayer, FAudioDecoderAAC_Worker);
+				res = DecoderTransform->ProcessOutput(0, 1, &CurrentDecoderOutputBuffer->mOutputBuffer, &dwStatus);
+			}
+
+			if (res == MF_E_TRANSFORM_NEED_MORE_INPUT)
+			{
+				CSV_SCOPED_TIMING_STAT(ElectraPlayer, FAudioDecoderAAC_Worker);
+
+				// Flushing / draining?
+				if (bFlush)
+				{
+					// Yes. This means we have received all pending output and are done now.
+					if (!bFlushOnly)
+					{
+						// After a drain issue a flush.
+						VERIFY_HR(DecoderTransform->ProcessMessage(MFT_MESSAGE_COMMAND_FLUSH, 0), "Failed to issue audio decoder flush command", ERRCODE_INTERNAL_COULD_NOT_SET_DECODER_FLUSHCOMMAND);
+					}
+					// And start over.
+					VERIFY_HR(DecoderTransform->ProcessMessage(MFT_MESSAGE_NOTIFY_BEGIN_STREAMING, 0), "Failed to set audio decoder stream begin", ERRCODE_INTERNAL_COULD_NOT_SET_DECODER_BEGIN);
+					VERIFY_HR(DecoderTransform->ProcessMessage(MFT_MESSAGE_NOTIFY_START_OF_STREAM, 0), "Failed to start audio decoder", ERRCODE_INTERNAL_COULD_NOT_SET_DECODER_START);
+					CurrentDecoderOutputBuffer.Reset();
+					return true;
+				}
+				else if (InputSample.IsValid())
+				{
+					VERIFY_HR(DecoderTransform->ProcessInput(0, InputSample.GetReference(), 0), "Failed to process audio decoder input", ERRCODE_INTERNAL_COULD_NOT_PROCESS_INPUT);
+					// Used this sample. Have no further input data for now, but continue processing to produce output if possible.
+					InputSample = nullptr;
+					OutInputSizeConsumed = InSize;
+				}
+				else
+				{
+					// Need more input but have none right now.
+					return true;
+				}
+			}
+			else if (res == MF_E_TRANSFORM_STREAM_CHANGE)
+			{
+				// Update output type.
+				if (!SetDecoderOutputType())
+				{
+					return false;
+				}
+				// For the sake of argument lets get rid of the output buffer (might be too small)
+				CurrentDecoderOutputBuffer.Reset();
+			}
+			else if (SUCCEEDED(res))
+			{
+				TRefCountPtr<IMFSample> DecodedOutputSample = CurrentDecoderOutputBuffer->DetachOutputSample();
+				CurrentDecoderOutputBuffer.Reset();
+
+				TRefCountPtr<IMFMediaBuffer>	DecodedLinearOutputBuffer;
+				DWORD					dwBufferLen;
+				DWORD					dwMaxBufferLen;
+				BYTE* pDecompressedData = nullptr;
+				LONGLONG				llTimeStamp = 0;
+				WAVEFORMATEX			OutputWaveFormat;
+				WAVEFORMATEX* OutputWaveFormatPtr = nullptr;
+				UINT32					WaveFormatSize = 0;
+
+				VERIFY_HR(DecodedOutputSample->GetSampleTime(&llTimeStamp), "Failed to get audio decoder output sample timestamp", ERRCODE_INTERNAL_COULD_NOT_GET_OUTPUT_SAMPLE_TIME);
+				VERIFY_HR(DecodedOutputSample->ConvertToContiguousBuffer(DecodedLinearOutputBuffer.GetInitReference()), "Failed to convert audio decoder output sample to contiguous buffer", ERRCODE_INTERNAL_COULD_NOT_MAKE_CONTIGUOUS_OUTPUT_BUFFER);
+				VERIFY_HR(DecodedLinearOutputBuffer->GetCurrentLength(&dwBufferLen), "Failed to get audio decoder output buffer current length", ERRCODE_INTERNAL_COULD_NOT_GET_OUTPUT_BUFFER_LENGTH);
+				VERIFY_HR(DecodedLinearOutputBuffer->Lock(&pDecompressedData, &dwMaxBufferLen, &dwBufferLen), "Failed to lock audio decoder output buffer", ERRCODE_INTERNAL_COULD_NOT_LOCK_OUTPUT_BUFFER)
+
+					VERIFY_HR(MFCreateWaveFormatExFromMFMediaType(CurrentOutputMediaType.GetReference(), &OutputWaveFormatPtr, &WaveFormatSize, MFWaveFormatExConvertFlag_Normal), "Failed to create audio decoder output buffer format info", ERRCODE_INTERNAL_COULD_NOT_CREATE_OUTPUT_BUFFER_FORMAT_INFO);
+				FMemory::Memcpy(&OutputWaveFormat, OutputWaveFormatPtr, sizeof(OutputWaveFormat));
+				CoTaskMemFree(OutputWaveFormatPtr);
+				OutputWaveFormatPtr = nullptr;
+
+				FTimeValue pts;
+				pts.SetFromHNS((int64)llTimeStamp);
+
+				int32 nSamplesProduced = dwBufferLen / (sizeof(int16) * OutputWaveFormat.nChannels);
+				const int16* pPCMSamples = (const int16*)pDecompressedData;
+				//int32 nSamplesProduced = dwBufferLen / (sizeof(float) * OutputWaveFormat.nChannels);
+				//const float *pPCMSamples = (const float *)pDecompressedData;
+
+				// Get an sample block from the pool.
+				while (!TerminateThreadEvent.IsSignaled())
+				{
+					if (FlushDecoderEvent.IsSignaled() && !bFlush)
+					{
+						break;
+					}
+
+					if (CurrentRenderOutputBuffer == nullptr)
+					{
+						CSV_SCOPED_TIMING_STAT(ElectraPlayer, FAudioDecoderAAC_Worker);
+
+						UEMediaError bufResult = Renderer->AcquireBuffer(CurrentRenderOutputBuffer, 0, BufferAcquireOptions);
+						check(bufResult == UEMEDIA_ERROR_OK || bufResult == UEMEDIA_ERROR_INSUFFICIENT_DATA);
+						if (bufResult != UEMEDIA_ERROR_OK && bufResult != UEMEDIA_ERROR_INSUFFICIENT_DATA)
+						{
+							PostError(S_OK, "Failed to acquire sample buffer", ERRCODE_INTERNAL_COULD_NOT_GET_SAMPLE_BUFFER, bufResult);
+							return false;
+						}
+					}
+					bool bHaveAvailSmpBlk = CurrentRenderOutputBuffer != nullptr;
+					NotifyReadyBufferListener(bHaveAvailSmpBlk);
+					if (bHaveAvailSmpBlk)
+					{
+						CSV_SCOPED_TIMING_STAT(ElectraPlayer, FAudioDecoderAAC_Worker);
+
+						FTimeValue OutputSampleDuration;
+						OutputSampleDuration.SetFromND(nSamplesProduced, OutputWaveFormat.nSamplesPerSec);
+
+						DWORD CurrentRenderOutputBufferSize = (DWORD)CurrentRenderOutputBuffer->GetBufferProperties().GetValue("size").GetInt64();
+						check(dwBufferLen <= CurrentRenderOutputBufferSize);
+						if (dwBufferLen <= CurrentRenderOutputBufferSize)
+						{
+							void* CurrentRenderOutputBufferAddress = CurrentRenderOutputBuffer->GetBufferProperties().GetValue("address").GetPointer();
+							FMemory::Memcpy(CurrentRenderOutputBufferAddress, pPCMSamples, dwBufferLen);
+
+							OutputBufferSampleProperties.Clear();
+							OutputBufferSampleProperties.Set("num_channels", FVariantValue((int64)OutputWaveFormat.nChannels));
+							OutputBufferSampleProperties.Set("sample_rate", FVariantValue((int64)OutputWaveFormat.nSamplesPerSec));
+							OutputBufferSampleProperties.Set("byte_size", FVariantValue((int64)dwBufferLen));
+							OutputBufferSampleProperties.Set("duration", FVariantValue(OutputSampleDuration));
+							OutputBufferSampleProperties.Set("pts", FVariantValue(pts));
+							OutputBufferSampleProperties.Set("discontinuity", FVariantValue((bool)bHaveDiscontinuity));
+
+							Renderer->ReturnBuffer(CurrentRenderOutputBuffer, bRender, OutputBufferSampleProperties);
+							CurrentRenderOutputBuffer = nullptr;
+							bHaveDiscontinuity = false;
+
+							break;
+						}
+						else
+						{
+							PostError(S_OK, "Audio renderer buffer too small to receive decoded audio samples", ERRCODE_INTERNAL_RENDER_BUFFER_TOO_SMALL);
+							return false;
+						}
+					}
+					else
+					{
+						// No available buffer. Sleep for a bit. Can't sleep on a signal since we have to check two: abort and flush
+						// We sleep for 20ms since a sample block for LC AAC produces 1024 samples which amount to 21.3ms @48kHz (or more at lower sample rates).
+						FMediaRunnable::SleepMilliseconds(20);
+					}
+				}
+				DecodedLinearOutputBuffer->Unlock();
+			}
+			else
+			{
+				// Error!
+				VERIFY_HR(res, "Failed to process audio decoder output", ERRCODE_INTERNAL_COULD_NOT_PROCESS_OUTPUT);
+				return false;
+			}
+		}
+		return true;
+	}
+
+
+
+	//-----------------------------------------------------------------------------
+	/**
+	 * Checks if the codec specific format has changed.
+	 *
+	 * @param Data
+	 *
+	 * @return
+	 */
+	bool FAudioDecoderAAC::IsDifferentFormat(const FAccessUnit* Data)
+	{
+		if (Data->AUCodecData.IsValid() && (!CurrentCodecData.IsValid() || CurrentCodecData->CodecSpecificData != Data->AUCodecData->CodecSpecificData))
+		{
+			return true;
+		}
+		return false;
+	}
+
+
+	//-----------------------------------------------------------------------------
+	/**
+	 * AAC audio decoder main threaded decode loop
+	 */
+	void FAudioDecoderAAC::WorkerThread()
+	{
+		LLM_SCOPE(ELLMTag::ElectraPlayer);
+
+		bool bError = false;
+		bool bInDummyDecodeMode = false;
+
+		bHaveDiscontinuity = false;
+		CurrentRenderOutputBuffer = nullptr;
+
+		DecoderInputStreamInfo = {};
+		DecoderOutputStreamInfo = {};
+
+		bError = !CreateDecodedSamplePool();
+		check(!bError);
+
+		while (!TerminateThreadEvent.IsSignaled())
+		{
+			// Notify optional buffer listener that we will now be needing an AU for our input buffer.
+			if (!bError && InputBufferListener && AccessUnits.Num() == 0)
+			{
+				CSV_SCOPED_TIMING_STAT(ElectraPlayer, FAudioDecoderAAC_Worker);
+
+				FAccessUnitBufferInfo	sin;
+				IAccessUnitBufferListener::FBufferStats	stats;
+				AccessUnits.GetStats(sin);
+				stats.NumAUsAvailable = sin.NumCurrentAccessUnits;
+				stats.NumBytesInBuffer = sin.CurrentMemInUse;
+				stats.MaxBytesOfBuffer = sin.MaxDataSize;
+				stats.bEODSignaled = sin.bEndOfData;
+				stats.bEODReached = sin.bEndOfData && sin.NumCurrentAccessUnits == 0;
+				ListenerCriticalSection.Lock();
+				if (InputBufferListener)
+				{
+					InputBufferListener->DecoderInputNeeded(stats);
+				}
+				ListenerCriticalSection.Unlock();
+			}
+			// Wait for data to arrive. If no data for 10ms check if we're supposed to abort.
+			// When data arrives the wait is of course cut short and we don't spend the full 10ms waiting!
+			bool bHaveData = AccessUnits.WaitForData(1000 * 10);
+			if (bHaveData)
+			{
+				FAccessUnit* pData = nullptr;
+				AccessUnits.Pop(pData);
+
+				// Check if the format has changed such that we need to flush and re-create the decoder.
+				if (pData->bTrackChangeDiscontinuity || IsDifferentFormat(pData) || (pData->bIsDummyData && !bInDummyDecodeMode))
+				{
+					if (DecoderTransform.IsValid())
+					{
+						int32 NumBytesConsumed = 0;
+
+						// Check that the difference in expected PTS is not too great. If it is we should just discard the pending samples from the decoder and not output them.
+						FTimeValue PtsDiff = pData->PTS - NextExpectedPTS;
+						bool bRenderPending = PtsDiff.IsValid() && PtsDiff.GetAsMilliseconds() >= 0 && PtsDiff.GetAsMilliseconds() < 100;
+
+						Decode(NumBytesConsumed, nullptr, 0, 0, 0, true, bRenderPending);
+						bHaveDiscontinuity = true;
+					}
+					ReturnUnusedFrame();
+					DecoderTransform = nullptr;
+					CurrentOutputMediaType = nullptr;
+					ConfigRecord = nullptr;
+					CurrentCodecData = nullptr;
+				}
+
+				// Parse the CSD into a configuration record.
+				if (!ConfigRecord.IsValid() && !bError)
+				{
+					if (pData->AUCodecData.IsValid())
+					{
+						CurrentCodecData = pData->AUCodecData;
+						ConfigRecord = MakeShared<MPEG::FAACDecoderConfigurationRecord, ESPMode::ThreadSafe>();
+						if (!ConfigRecord->ParseFrom(CurrentCodecData->CodecSpecificData.GetData(), CurrentCodecData->CodecSpecificData.Num()))
+						{
+							ConfigRecord.Reset();
+							CurrentCodecData.Reset();
+							PostError(S_OK, "Failed to parse AAC configuration record", ERRCODE_INTERNAL_FAILED_TO_PARSE_CSD);
+							bError = true;
+						}
+					}
+				}
+
+				CurrentPTS = pData->PTS;
+				NextExpectedPTS = pData->PTS + pData->Duration;
+
+				// Is this audio packet to be dropped?
+				if (pData->DropState)
+				{
+					FAccessUnit::Release(pData);
+					pData = nullptr;
+					continue;
+				}
+
+				// Need to create a decoder instance?
+				if (!DecoderTransform.IsValid() && !bError)
+				{
+					CSV_SCOPED_TIMING_STAT(ElectraPlayer, FAudioDecoderAAC_Worker);
+
+					// Can't create a decoder based on dummy data.
+					if (!pData->bIsDummyData)
+					{
+						if (InternalDecoderCreate())
+						{
+							// Ok
+						}
+						else
+						{
+							bError = true;
+						}
+					}
+				}
+
+				// Loop until all data has been consumed.
+				bool bAllInputConsumed = false;
+				while (!bError && !bAllInputConsumed && !TerminateThreadEvent.IsSignaled() && !FlushDecoderEvent.IsSignaled())
+				{
+					if (!pData->bIsDummyData)
+					{
+						if (bInDummyDecodeMode)
+						{
+							bInDummyDecodeMode = false;
+							bHaveDiscontinuity = true;
+						}
+						int32 NumBytesConsumed = 0;
+						if (Decode(NumBytesConsumed, pData->AUData, pData->AUSize, pData->PTS.GetAsHNS(), pData->Duration.GetAsHNS(), false, true))
+						{
+							check(NumBytesConsumed == pData->AUSize);
+							bAllInputConsumed = true;
+							break;
+						}
+						else
+						{
+							bError = true;
+						}
+					}
+					else
+					{
+						if (!bInDummyDecodeMode)
+						{
+							bInDummyDecodeMode = true;
+							bHaveDiscontinuity = true;
+						}
+
+						// Need a new output buffer?
+						if (CurrentRenderOutputBuffer == nullptr)
+						{
+							CSV_SCOPED_TIMING_STAT(ElectraPlayer, FAudioDecoderAAC_Worker);
+
+							UEMediaError bufResult = Renderer->AcquireBuffer(CurrentRenderOutputBuffer, 0, BufferAcquireOptions);
+							check(bufResult == UEMEDIA_ERROR_OK || bufResult == UEMEDIA_ERROR_INSUFFICIENT_DATA);
+							if (bufResult != UEMEDIA_ERROR_OK && bufResult != UEMEDIA_ERROR_INSUFFICIENT_DATA)
+							{
+								PostError(S_OK, "Failed to acquire sample buffer", ERRCODE_INTERNAL_COULD_NOT_GET_SAMPLE_BUFFER, bufResult);
+								bError = true;
+								break;
+							}
+						}
+						bool bHaveAvailSmpBlk = CurrentRenderOutputBuffer != nullptr;
+						NotifyReadyBufferListener(bHaveAvailSmpBlk);
+						if (bHaveAvailSmpBlk)
+						{
+							CSV_SCOPED_TIMING_STAT(ElectraPlayer, FAudioDecoderAAC_Worker);
+
+							// Clear to silence
+							SIZE_T CurrentRenderOutputBufferSize = (SIZE_T)CurrentRenderOutputBuffer->GetBufferProperties().GetValue("size").GetInt64();
+							void* CurrentRenderOutputBufferAddress = CurrentRenderOutputBuffer->GetBufferProperties().GetValue("address").GetPointer();
+							FMemory::Memzero(CurrentRenderOutputBufferAddress, CurrentRenderOutputBufferSize);
+
+							// Assume sensible defaults.
+							int64 NumChannels = 2;
+							int64 SampleRate = 48000;
+							int64 SamplesPerBlock = 1024;
+
+							// With a valid configuration record we can use the actual values.
+							if (ConfigRecord.IsValid())
+							{
+								// Parameteric stereo results in stereo (2 channels).
+								NumChannels = ConfigRecord->PSSignal > 0 ? 2 : ConfigRecord->ChannelConfiguration;
+								SampleRate = ConfigRecord->ExtSamplingFrequency ? ConfigRecord->ExtSamplingFrequency : ConfigRecord->SamplingRate;
+								SamplesPerBlock = ConfigRecord->SBRSignal > 0 ? 2048 : 1024;
+							}
+
+							FTimeValue Duration;
+							Duration.SetFromND(SamplesPerBlock, (uint32)SampleRate);
+
+							// Note: The duration in the AU should be an exact multiple of the "SamplesPerBlock / SampleRate" as this is how it gets set
+							//       up in the stream reader. Should there ever be any discrepancy such that "Duration > pData->mDuration" it is simply
+							//       possible to recalculate "SamplesPerBlock" such that we return a buffer with less than the usual 1024 or 2048 samples.
+							OutputBufferSampleProperties.Clear();
+							OutputBufferSampleProperties.Set("num_channels", FVariantValue(NumChannels));
+							OutputBufferSampleProperties.Set("sample_rate", FVariantValue(SampleRate));
+							OutputBufferSampleProperties.Set("byte_size", FVariantValue((int64)(NumChannels * sizeof(int16) * SamplesPerBlock)));
+							OutputBufferSampleProperties.Set("duration", FVariantValue(Duration));
+							OutputBufferSampleProperties.Set("pts", FVariantValue(CurrentPTS));
+							OutputBufferSampleProperties.Set("eod", FVariantValue((bool)false));
+							OutputBufferSampleProperties.Set("discontinuity", FVariantValue((bool)bHaveDiscontinuity));
+
+							Renderer->ReturnBuffer(CurrentRenderOutputBuffer, true, OutputBufferSampleProperties);
+							CurrentRenderOutputBuffer = nullptr;
+							bHaveDiscontinuity = false;
+
+							CurrentPTS += Duration;
+							pData->Duration -= Duration;
+							bAllInputConsumed = pData->Duration <= FTimeValue::GetZero();
+						}
+						else
+						{
+							// No available buffer. Sleep for a bit. Can't sleep on a signal since we have to check two: abort and flush
+							// We sleep for 20ms since a sample block for LC AAC produces 1024 samples which amount to 21.3ms @48kHz (or more at lower sample rates).
+							FMediaRunnable::SleepMilliseconds(20);
+						}
+					}
+				}
+
+				// We're done with this access unit. Delete it.
+				FAccessUnit::Release(pData);
+				pData = nullptr;
+			}
+			else
+			{
+				// No data. Is the buffer at EOD?
+				if (AccessUnits.IsEndOfData())
+				{
+					NotifyReadyBufferListener(true);
+					FMediaRunnable::SleepMilliseconds(20);
+				}
+			}
+
+			// Flush?
+			if (FlushDecoderEvent.IsSignaled())
+			{
+				if (DecoderTransform.IsValid())
+				{
+					int32 NumBytesConsumed = 0;
+					Decode(NumBytesConsumed, nullptr, 0, 0, 0, true, true);
+				}
+
+				ReturnUnusedFrame();
+				AccessUnits.Flush();
+
+				CurrentPTS.SetToInvalid();
+				NextExpectedPTS.SetToInvalid();
+
+				FlushDecoderEvent.Reset();
+				DecoderFlushedEvent.Signal();
+			}
+		}
+
+		ReturnUnusedFrame();
+		// Close the decoder.
+		DecoderTransform = nullptr;
+		CurrentOutputMediaType = nullptr;
+
+		DestroyDecodedSamplePool();
+
+		CurrentCodecData.Reset();
+		ConfigRecord.Reset();
+
+		// Flush any remaining input data.
+		AccessUnits.Flush();
+		AccessUnits.CapacitySet(0);
+
+		CurrentPTS.SetToInvalid();
+		NextExpectedPTS.SetToInvalid();
+	}
+
+
+} // namespace Electra
+
+
+#endif
+
diff --git a/Engine/Plugins/Media/ElectraPlayer/Source/ElectraPlayerRuntime/Private/Runtime/Decoder/DX/DecoderErrors_DX.h b/Engine/Plugins/Media/ElectraPlayer/Source/ElectraPlayerRuntime/Private/Runtime/Decoder/DX/DecoderErrors_DX.h
new file mode 100644
--- /dev/null
+++ b/Engine/Plugins/Media/ElectraPlayer/Source/ElectraPlayerRuntime/Private/Runtime/Decoder/DX/DecoderErrors_DX.h
@@ -0,0 +1,75 @@
+// Copyright Epic Games, Inc. All Rights Reserved.
+
+#pragma once
+
+#define ERRCODE_INTERNAL_OUT_OF_MEMORY							1
+#define ERRCODE_INTERNAL_COULD_NOT_CREATE_DECODER				2
+#define ERRCODE_INTERNAL_COULD_NOT_CREATE_IMAGE_POOL			3
+#define ERRCODE_INTERNAL_COULD_NOT_GET_SAMPLE_BUFFER			4
+#define ERRCODE_INTERNAL_FAILED_TO_DECODE						5
+#define ERRCODE_INTERNAL_FAILED_TO_FLUSH						6
+#define ERRCODE_INTERNAL_RENDER_BUFFER_TOO_SMALL				7
+#define ERRCODE_INTERNAL_COULD_NOT_CREATE_SAMPLE_POOL			8
+#define ERRCODE_INTERNAL_FAILED_TO_PARSE_CSD					9
+
+
+#define ERRCODE_INTERNAL_COULD_GET_TRANSFORM_ATTRIBUTES					10
+#define ERRCODE_INTERNAL_UNSUPPORTED_D3D_TYPE							11
+#define ERRCODE_INTERNAL_FAILED_TO_SWITCH_TO_SOFTWARE_MODE				12
+#define ERRCODE_INTERNAL_COULD_NOT_CREATE_INPUT_MEDIA_TYPE				20
+#define ERRCODE_INTERNAL_COULD_NOT_CREATE_INPUT_MEDIA_TYPE_MAJOR		21
+#define ERRCODE_INTERNAL_COULD_NOT_CREATE_INPUT_MEDIA_TYPE_SUBTYPE		22
+#define ERRCODE_INTERNAL_COULD_NOT_CREATE_INPUT_MEDIA_TYPE_RESOLUTION	23
+#define ERRCODE_INTERNAL_COULD_NOT_SET_INPUT_MEDIA_TYPE					24
+#define ERRCODE_INTERNAL_COULD_NOT_SET_INPUT_MEDIA_PROGRESSIVE_MODE		25
+#define ERRCODE_INTERNAL_COULD_NOT_SET_INPUT_MEDIA_TYPE_CODEC_BLOB		26
+#define ERRCODE_INTERNAL_COULD_NOT_GET_OUTPUT_MEDIA_TYPE				30
+#define ERRCODE_INTERNAL_COULD_NOT_GET_OUTPUT_MEDIA_TYPE_MAJOR			31
+#define ERRCODE_INTERNAL_COULD_NOT_GET_OUTPUT_MEDIA_TYPE_SUBTYPE		32
+#define ERRCODE_INTERNAL_COULD_NOT_SET_OUTPUT_MEDIA_TYPE				33
+#define ERRCODE_INTERNAL_COULD_NOT_SET_OUTPUT_DESIRED_MEDIA_TYPE		34
+#define ERRCODE_INTERNAL_COULD_NOT_GET_STREAM_COUNT						40
+#define ERRCODE_INTERNAL_UNEXPECTED_NUMBER_OF_STREAMS					41
+#define ERRCODE_INTERNAL_COULD_NOT_GET_INPUT_STATUS						42
+#define ERRCODE_INTERNAL_DECODER_NOT_ACCEPTING_DATA						43
+#define ERRCODE_INTERNAL_COULD_NOT_GET_OUTPUT_STREAM_INFO				44
+#define ERRCODE_INTERNAL_INCOMPATIBLE_DECODER							45
+#define ERRCODE_INTERNAL_COULD_NOT_SET_DECODER_BEGIN					50
+#define ERRCODE_INTERNAL_COULD_NOT_SET_DECODER_START					51
+#define ERRCODE_INTERNAL_COULD_NOT_SET_DECODER_ENDOFSTREAM				60
+#define ERRCODE_INTERNAL_COULD_NOT_SET_DECODER_FLUSHCOMMAND				61
+#define ERRCODE_INTERNAL_COULD_NOT_SET_DECODER_DRAINCOMMAND				62
+#define ERRCODE_INTERNAL_COULD_NOT_CREATE_OUTPUT_SAMPLE					70
+#define ERRCODE_INTERNAL_COULD_NOT_CREATE_ALIGNED_OUTPUTBUFFER			71
+#define ERRCODE_INTERNAL_COULD_NOT_CREATE_OUTPUTBUFFER					72
+#define ERRCODE_INTERNAL_COULD_NOT_ADD_OUTPUT_BUFFER_TO_SAMPLE			73
+#define ERRCODE_INTERNAL_COULD_NOT_CREATE_INPUT_SAMPLE					80
+#define ERRCODE_INTERNAL_COULD_NOT_CREATE_INPUTBUFFER					81
+#define ERRCODE_INTERNAL_COULD_NOT_ADD_INPUT_BUFFER_TO_SAMPLE			82
+#define ERRCODE_INTERNAL_COULD_NOT_LOCK_INPUT_BUFFER					83
+#define ERRCODE_INTERNAL_COULD_NOT_UNLOCK_INPUT_BUFFER					84
+#define ERRCODE_INTERNAL_COULD_NOT_SET_BUFFER_CURRENT_LENGTH			85
+#define ERRCODE_INTERNAL_COULD_NOT_SET_INPUT_SAMPLE_TIME				86
+#define ERRCODE_INTERNAL_COULD_NOT_SET_INPUT_SAMPLE_DURATION			87
+#define ERRCODE_INTERNAL_COULD_NOT_SET_INPUT_SAMPLE_DECODETIME			88
+#define ERRCODE_INTERNAL_COULD_NOT_SET_INPUT_SAMPLE_KEYFRAME			89
+#define ERRCODE_INTERNAL_COULD_NOT_PROCESS_INPUT						100
+#define ERRCODE_INTERNAL_COULD_NOT_PROCESS_OUTPUT						110
+#define ERRCODE_INTERNAL_COULD_NOT_GET_OUTPUT_SAMPLE_TIME				120
+#define ERRCODE_INTERNAL_COULD_NOT_GET_OUTPUT_SAMPLE_DURATION			121
+#define ERRCODE_INTERNAL_COULD_NOT_SET_INPUT_AUDIO_MAJOR_TYPE			122
+#define ERRCODE_INTERNAL_COULD_NOT_SET_INPUT_AUDIO_AAC_SUBTYPE			123
+#define ERRCODE_INTERNAL_COULD_NOT_SET_INPUT_AUDIO_SAMPLING_RATE		124
+#define ERRCODE_INTERNAL_COULD_NOT_SET_INPUT_AUDIO_CHANNEL_COUNT		125
+#define ERRCODE_INTERNAL_COULD_NOT_SET_INPUT_AUDIO_CSD					126
+#define ERRCODE_INTERNAL_COULD_NOT_SET_INPUT_TYPE						127
+#define ERRCODE_INTERNAL_COULD_NOT_SET_OUTPUT_TYPE						128
+#define ERRCODE_INTERNAL_COULD_NOT_SET_OUTPUT_TYPE_TO_PCM				128
+#define ERRCODE_INTERNAL_COULD_NOT_GET_INPUT_STREAM_FORMAT_INFO			129
+#define ERRCODE_INTERNAL_COULD_NOT_GET_OUTPUT_STREAM_FORMAT_INFO		139
+#define ERRCODE_INTERNAL_COULD_NOT_MAKE_CONTIGUOUS_OUTPUT_BUFFER		140
+#define ERRCODE_INTERNAL_COULD_NOT_GET_OUTPUT_BUFFER_LENGTH				141
+#define ERRCODE_INTERNAL_COULD_NOT_LOCK_OUTPUT_BUFFER					142
+#define ERRCODE_INTERNAL_COULD_NOT_UNLOCK_OUTPUT_BUFFER					143
+#define ERRCODE_INTERNAL_COULD_NOT_CREATE_OUTPUT_BUFFER_FORMAT_INFO		144
+#define ERRCODE_INTERNAL_COULD_DEVICE_LOST_DETECTED						145
diff --git a/Engine/Plugins/Media/ElectraPlayer/Source/ElectraPlayerRuntime/Private/Runtime/Decoder/DX/VideoDecoderH264_DX.cpp b/Engine/Plugins/Media/ElectraPlayer/Source/ElectraPlayerRuntime/Private/Runtime/Decoder/DX/VideoDecoderH264_DX.cpp
new file mode 100644
--- /dev/null
+++ b/Engine/Plugins/Media/ElectraPlayer/Source/ElectraPlayerRuntime/Private/Runtime/Decoder/DX/VideoDecoderH264_DX.cpp
@@ -0,0 +1,1468 @@
+// Copyright Epic Games, Inc. All Rights Reserved.
+
+#ifdef ELECTRA_ENABLE_MFDECODER
+
+#include "PlayerCore.h"
+#include "ElectraPlayerPrivate.h"
+
+#include "ElectraPlayerPrivate_Platform.h"
+
+#include "StreamAccessUnitBuffer.h"
+#include "Decoder/VideoDecoderH264.h"
+#include "Renderer/RendererBase.h"
+#include "Renderer/RendererVideo.h"
+#include "Player/PlayerSessionServices.h"
+#include "Utilities/StringHelpers.h"
+
+/*********************************************************************************************************************/
+/*********************************************************************************************************************/
+/*********************************************************************************************************************/
+
+#include "Windows/AllowWindowsPlatformTypes.h"
+#include "Windows/WindowsHWrapper.h"
+#include "HAL/LowLevelMemTracker.h"
+
+THIRD_PARTY_INCLUDES_START
+#include "mftransform.h"
+#include "mfapi.h"
+#include "mferror.h"
+#include "mfidl.h"
+THIRD_PARTY_INCLUDES_END
+
+#include "Windows/HideWindowsPlatformTypes.h"
+
+#include "DecoderErrors_DX.h"
+#include "VideoDecoderH264_DX.h"
+
+namespace
+{
+
+static const int32 MaxHWDecodeW = 1920;
+static const int32 MaxHWDecodeH = 1088;
+
+//DEFINE_GUID(MF_MT_DEFAULT_STRIDE,    0x644b4e48, 0x1e02, 0x4516, 0xb0, 0xeb, 0xc0, 0x1c, 0xa9, 0xd4, 0x9a, 0xc6);
+//static const GUID AVDecVideoMaxCodedWidth   = { 0x5ae557b8, 0x77af, 0x41f5, { 0x9f, 0xa6, 0x4d, 0xb2, 0xfe, 0x1d, 0x4b, 0xca } };
+//static const GUID AVDecVideoMaxCodedHeight  = { 0x7262a16a, 0xd2dc, 0x4e75, { 0x9b, 0xa8, 0x65, 0xc0, 0xc6, 0xd3, 0x2b, 0x13 } };
+
+//#define MFSampleExtension_ClosedCaption_CEA708_MAX_SIZE	256
+//DEFINE_GUID(MFSampleExtension_ClosedCaption_CEA708, 0x26f09068, 0xe744, 0x47dc, 0xaa,0x03,0xdb,0xf2,0x04,0x03,0xbd,0xe6);
+//DEFINE_GUID(MF_USER_DATA_PAYLOAD,	0xd1d4985d, 0xdc92, 0x457a, 0xb3,0xa0,0x65,0x1a,0x33,0xa3,0x10,0x47);
+
+// `MF_LOW_LATENCY` is defined in "mfapi.h" for >= WIN8
+// UE4 supports lower Windows versions at the moment and so `WINVER` is < `_WIN32_WINNT_WIN8`
+// to be able to use `MF_LOW_LATENCY` with default UE4 build we define it ourselves and check actual
+// Windows version in runtime
+#if (WINVER < _WIN32_WINNT_WIN8)
+	const GUID MF_LOW_LATENCY = { 0x9c27891a, 0xed7a, 0x40e1,{ 0x88, 0xe8, 0xb2, 0x27, 0x27, 0xa0, 0x24, 0xee } };
+#endif
+
+}
+
+namespace Electra {
+
+IVideoDecoderH264::FSystemConfiguration			FVideoDecoderH264::SystemConfig;
+
+#if ELECTRA_ENABLE_SWDECODE
+bool FVideoDecoderH264::bDidCheckHWSupport = false;
+bool FVideoDecoderH264::bIsHWSupported = false;
+#endif
+
+/*********************************************************************************************************************/
+/*********************************************************************************************************************/
+/*********************************************************************************************************************/
+
+bool IVideoDecoderH264::Startup(const IVideoDecoderH264::FSystemConfiguration& InConfig)
+{
+	return FVideoDecoderH264::Startup(InConfig);
+}
+
+void IVideoDecoderH264::Shutdown()
+{
+	FVideoDecoderH264::Shutdown();
+}
+
+bool IVideoDecoderH264::GetStreamDecodeCapability(FStreamDecodeCapability& OutResult, const FStreamDecodeCapability& InStreamParameter)
+{
+	return FVideoDecoderH264::GetStreamDecodeCapability(OutResult, InStreamParameter);
+}
+
+IVideoDecoderH264::FSystemConfiguration::FSystemConfiguration()
+{
+	ThreadConfig.Decoder.Priority 	= TPri_Normal;
+	ThreadConfig.Decoder.StackSize	= 256 << 10;	// how much do we need?
+	ThreadConfig.Decoder.CoreAffinity = -1;
+	// Not needed, set to useful values anyway
+	ThreadConfig.PassOn.Priority  	= TPri_Normal;
+	ThreadConfig.PassOn.StackSize 	= 65536;
+	ThreadConfig.PassOn.CoreAffinity  = -1;
+}
+
+IVideoDecoderH264::FInstanceConfiguration::FInstanceConfiguration()
+	: MaxDecodedFrames(8)														// FIXME: how many do we need exactly?
+	, ThreadConfig(FVideoDecoderH264::SystemConfig.ThreadConfig)
+{
+}
+
+/***************************************************************************************************************************************************/
+/***************************************************************************************************************************************************/
+/***************************************************************************************************************************************************/
+
+
+//-----------------------------------------------------------------------------
+/**
+ * Decoder system startup
+ *
+ * @param config
+ *
+ * @return
+ */
+bool FVideoDecoderH264::Startup(const IVideoDecoderH264::FSystemConfiguration& config)
+{
+	SystemConfig = config;
+#ifdef ELECTRA_ENABLE_SWDECODE
+	bDidCheckHWSupport = false;
+	bIsHWSupported = false;
+#endif
+	return true;
+}
+
+
+//-----------------------------------------------------------------------------
+/**
+ * Decoder system shutdown.
+ */
+void FVideoDecoderH264::Shutdown()
+{
+}
+
+
+//-----------------------------------------------------------------------------
+/**
+ * Constructor
+ */
+FVideoDecoderH264::FVideoDecoderH264()
+	: FMediaThread("Electra::H264 decoder")
+	, bThreadStarted(false)
+	, PlayerSessionServices(nullptr)
+	, Renderer(nullptr)
+	, CurrentAccessUnit(nullptr)
+	, InputBufferListener(nullptr)
+	, ReadyBufferListener(nullptr)
+	, bIsHardwareAccelerated(true)
+	, bRequiresReconfigurationForSW(false)
+	, NumFramesInDecoder(0)
+	, bDecoderFlushPending(false)
+	, bError(false)
+	, CurrentRenderOutputBuffer(nullptr)
+	, bHaveDecoder(false)
+	, MaxDecodeBufferSize(0)
+{
+}
+
+
+//-----------------------------------------------------------------------------
+/**
+ * Destructor
+ */
+FVideoDecoderH264::~FVideoDecoderH264()
+{
+	Close();
+}
+
+
+//-----------------------------------------------------------------------------
+/**
+ * Sets an AU input buffer listener.
+ *
+ * @param InListener
+ */
+void FVideoDecoderH264::SetAUInputBufferListener(IAccessUnitBufferListener* InListener)
+{
+	FMediaCriticalSection::ScopedLock lock(ListenerMutex);
+	InputBufferListener = InListener;
+}
+
+
+//-----------------------------------------------------------------------------
+/**
+ * Sets a buffer-ready listener.
+ *
+ * @param InListener
+ */
+void FVideoDecoderH264::SetReadyBufferListener(IDecoderOutputBufferListener* InListener)
+{
+	FMediaCriticalSection::ScopedLock lock(ListenerMutex);
+	ReadyBufferListener = InListener;
+}
+
+
+//-----------------------------------------------------------------------------
+/**
+ * Sets the owning player's session service interface.
+ *
+ * @param InSessionServices
+ *
+ * @return
+ */
+void FVideoDecoderH264::SetPlayerSessionServices(IPlayerSessionServices* SessionServices)
+{
+	PlayerSessionServices = SessionServices;
+}
+
+
+//-----------------------------------------------------------------------------
+/**
+ * Opens a decoder instance
+ *
+ * @param InConfig
+ *
+ * @return
+ */
+void FVideoDecoderH264::Open(const IVideoDecoderH264::FInstanceConfiguration& InConfig)
+{
+	Config = InConfig;
+
+	MaxDecodeDim = FIntPoint(Config.MaxFrameWidth, Config.MaxFrameHeight);
+
+	// Set a large enough size to hold a single access unit. Since we will be asking for a new AU on
+	// demand there is no need to overly restrict ourselves here. But it must be large enough to
+	// hold at least the largest expected access unit.
+	AccessUnits.CapacitySet(FAccessUnitBuffer::FConfiguration(16 << 20, 60.0));
+
+	StartThread();
+}
+
+
+//-----------------------------------------------------------------------------
+/**
+ * Closes the decoder instance.
+ */
+void FVideoDecoderH264::Close()
+{
+	StopThread();
+}
+
+
+//-----------------------------------------------------------------------------
+/**
+ * Sets a new decoder limit.
+ * As soon as a new sequence starts matching this limit the decoder will be
+ * destroyed and recreated to conserve memory.
+ * The assumption is that the data being streamed will not exceed this limit,
+ * but if it does, any access unit requiring a better decoder will take
+ * precendence and force a decoder capable of decoding it!
+ *
+ * @param MaxWidth
+ * @param MaxHeight
+ * @param MaxProfile
+ * @param MaxProfileLevel
+ * @param AdditionalOptions
+ */
+void FVideoDecoderH264::SetMaximumDecodeCapability(int32 MaxWidth, int32 MaxHeight, int32 MaxProfile, int32 MaxProfileLevel, const FParamDict& AdditionalOptions)
+{
+	// Recall max resolution so we can allocate properly sized sample buffers
+	MaxDecodeDim = FIntPoint(MaxWidth, MaxHeight);
+}
+
+
+//-----------------------------------------------------------------------------
+/**
+ * Sets a new renderer.
+ *
+ * @param InRenderer
+ */
+void FVideoDecoderH264::SetRenderer(TSharedPtr<IMediaRenderer, ESPMode::ThreadSafe> InRenderer)
+{
+	Renderer = InRenderer;
+}
+
+
+void FVideoDecoderH264::SetResourceDelegate(const TSharedPtr<IVideoDecoderResourceDelegate, ESPMode::ThreadSafe>& InResourceDelegate)
+{
+	ResourceDelegate = InResourceDelegate;
+}
+
+//-----------------------------------------------------------------------------
+/**
+ * Creates and runs the decoder thread.
+ */
+void FVideoDecoderH264::StartThread()
+{
+	ThreadSetPriority(Config.ThreadConfig.Decoder.Priority);
+	ThreadSetCoreAffinity(Config.ThreadConfig.Decoder.CoreAffinity);
+	ThreadSetStackSize(Config.ThreadConfig.Decoder.StackSize);
+	ThreadStart(Electra::MakeDelegate(this, &FVideoDecoderH264::WorkerThread));
+	bThreadStarted = true;
+}
+
+
+//-----------------------------------------------------------------------------
+/**
+ * Stops the decoder thread.
+ */
+void FVideoDecoderH264::StopThread()
+{
+	if (bThreadStarted)
+	{
+		TerminateThreadSignal.Signal();
+		ThreadWaitDone();
+		bThreadStarted = false;
+	}
+}
+
+
+//-----------------------------------------------------------------------------
+/**
+ * Posts an error to the session service error listeners.
+ *
+ * @param ApiReturnValue
+ * @param Message
+ * @param Code
+ * @param Error
+ */
+void FVideoDecoderH264::PostError(int32 ApiReturnValue, const FString& Message, uint16 Code, UEMediaError Error)
+{
+	bError = true;
+	check(PlayerSessionServices);
+	if (PlayerSessionServices)
+	{
+		FErrorDetail err;
+		err.SetError(Error != UEMEDIA_ERROR_OK ? Error : UEMEDIA_ERROR_DETAIL);
+		err.SetFacility(Facility::EFacility::H264Decoder);
+		err.SetCode(Code);
+		err.SetMessage(Message);
+		err.SetPlatformMessage(FString::Printf(TEXT("%d (0x%08x)"), ApiReturnValue, ApiReturnValue));
+		PlayerSessionServices->PostError(err);
+	}
+}
+
+
+//-----------------------------------------------------------------------------
+/**
+ * Sends a log message to the session service log.
+ *
+ * @param Level
+ * @param Message
+ */
+void FVideoDecoderH264::LogMessage(IInfoLog::ELevel Level, const FString& Message)
+{
+	if (PlayerSessionServices)
+	{
+		PlayerSessionServices->PostLog(Facility::EFacility::H264Decoder, Level, Message);
+	}
+}
+
+
+//-----------------------------------------------------------------------------
+/**
+ * Create a pool of decoded images for the decoder.
+ *
+ * @return
+ */
+bool FVideoDecoderH264::CreateDecodedImagePool()
+{
+	check(Renderer);
+	FParamDict poolOpts;
+
+	poolOpts.Set("num_buffers", FVariantValue((int64) Config.MaxDecodedFrames));
+	if (!bIsHardwareAccelerated)
+	{
+		poolOpts.Set("sw_texture", FVariantValue(true));
+	}
+	UEMediaError Error = Renderer->CreateBufferPool(poolOpts);
+	check(Error == UEMEDIA_ERROR_OK);
+
+	MaxDecodeBufferSize = (int32) Renderer->GetBufferPoolProperties().GetValue("max_buffers").GetInt64();
+
+	if (Error != UEMEDIA_ERROR_OK)
+	{
+		PostError(0, "Failed to create image pool", ERRCODE_INTERNAL_COULD_NOT_CREATE_IMAGE_POOL, Error);
+	}
+
+	return Error == UEMEDIA_ERROR_OK;
+}
+
+
+//-----------------------------------------------------------------------------
+/**
+ * Destroys the pool of decoded images.
+ */
+void FVideoDecoderH264::DestroyDecodedImagePool()
+{
+	Renderer->ReleaseBufferPool();
+}
+
+
+//-----------------------------------------------------------------------------
+/**
+ * Called to receive a new input access unit for decoding.
+ *
+ * @param InAccessUnit
+ */
+IAccessUnitBufferInterface::EAUpushResult FVideoDecoderH264::AUdataPushAU(FAccessUnit* InAccessUnit)
+{
+	InAccessUnit->AddRef();
+	bool bOk = AccessUnits.Push(InAccessUnit);
+	if (!bOk)
+	{
+		FAccessUnit::Release(InAccessUnit);
+	}
+
+	return bOk ? IAccessUnitBufferInterface::EAUpushResult::Ok : IAccessUnitBufferInterface::EAUpushResult::Full;
+}
+
+
+//-----------------------------------------------------------------------------
+/**
+ * "Pushes" an End Of Data marker indicating no further access units will be added.
+ */
+void FVideoDecoderH264::AUdataPushEOD()
+{
+	AccessUnits.PushEndOfData();
+}
+
+
+//-----------------------------------------------------------------------------
+/**
+ * Flushes the decoder and clears the input access unit buffer.
+ */
+void FVideoDecoderH264::AUdataFlushEverything()
+{
+	FlushDecoderSignal.Signal();
+	DecoderFlushedSignal.WaitAndReset();
+}
+
+
+//-----------------------------------------------------------------------------
+/**
+ *	Switches Windows 8+ to multithreaded DX for software decoding.
+ */
+#ifdef ELECTRA_ENABLE_SWDECODE
+bool FVideoDecoderH264::FallbackToSwDecoding(FString Reason)
+{
+	if (!bIsHardwareAccelerated)
+	{
+		return false;
+	}
+
+	UE_LOG(LogElectraPlayer, Log, TEXT("FVideoDecoderH264::FallbackToSwDecoding: %s"), *Reason);
+
+	bIsHardwareAccelerated = false;
+
+	if (Electra::IsWindows8Plus())
+	{
+		// NOTE: the following doesn't apply to Windows 7 as it doesn't use DX11 device in decoding thread
+		// as we don't use a dedicated DirextX device for s/w decoding, UE4's rendering device will be used from inside the decoder
+		// to produce output samples, which means access from render and decoding threads. We need to enable multithread protection
+		// for the device. Multithread protection can have performance impact, though its affect is expected to be negligible in most cases.
+		// WARNING:
+		// Once multithread protection is enabled we don't disable it, so UE4's rendering device stays protected for the rest of its lifetime.
+		// Some other system could enable multithread protection after we did it, we have no means to know about this, and so disabling it
+		// at the end of playback can cause GPU driver crash
+		if (Electra::FDXDeviceInfo::s_DXDeviceInfo->RenderingDx11Device)
+		{
+			HRESULT res;
+			TRefCountPtr<ID3D10Multithread> DxMultithread;
+			VERIFY_HR(Electra::FDXDeviceInfo::s_DXDeviceInfo->RenderingDx11Device->QueryInterface(__uuidof(ID3D10Multithread), (void**)DxMultithread.GetInitReference()), "Failed to set video decoder into software mode (multithread protected)", ERRCODE_INTERNAL_FAILED_TO_SWITCH_TO_SOFTWARE_MODE);
+			DxMultithread->SetMultithreadProtected(1);
+		}
+	}
+	return true;
+}
+#else
+bool FVideoDecoderH264::FallbackToSwDecoding(FString Reason)
+{
+	// Software decoding not supported. It's not even being called but we need to return something.
+	return false;
+}
+#endif
+
+
+//-----------------------------------------------------------------------------
+/**
+ *	Reconfigures decoder for software mode.
+ */
+bool FVideoDecoderH264::ReconfigureForSwDecoding(FString Reason)
+{
+	bRequiresReconfigurationForSW = true;
+	if (!FallbackToSwDecoding(MoveTemp(Reason)))
+	{
+		return false;
+	}
+
+	// Nullify D3D Manager to switch decoder to software mode.
+	HRESULT res;
+	check(DecoderTransform.GetReference());
+	VERIFY_HR(DecoderTransform->ProcessMessage(MFT_MESSAGE_SET_D3D_MANAGER, 0), "Failed to set video decoder into software mode (unset D3D manager)", ERRCODE_INTERNAL_FAILED_TO_SWITCH_TO_SOFTWARE_MODE);
+	return true;
+}
+
+
+//-----------------------------------------------------------------------------
+/**
+ * Configures the decoder input and output media types.
+ *
+ * @return true if successful, false on error
+ */
+bool FVideoDecoderH264::Configure()
+{
+	// Setup media input type
+	if (!DecoderSetInputType())
+	{
+		return false;
+	}
+	if (bRequiresReconfigurationForSW)
+	{
+		return false;
+	}
+
+	// Setup media output type
+	if (!DecoderSetOutputType())
+	{
+		return false;
+	}
+	if (bRequiresReconfigurationForSW)
+	{
+		return false;
+	}
+
+	// Verify status
+	if (!DecoderVerifyStatus())
+	{
+		return false;
+	}
+	if (bRequiresReconfigurationForSW)
+	{
+		return false;
+	}
+	return true;
+}
+
+
+//-----------------------------------------------------------------------------
+/**
+ * Starts decoder processing.
+ *
+ * @return true if successful, false on error
+ */
+bool FVideoDecoderH264::StartStreaming()
+{
+	HRESULT res;
+	// Start it!
+	VERIFY_HR(DecoderTransform->ProcessMessage(MFT_MESSAGE_NOTIFY_BEGIN_STREAMING, 0), "Failed to notify video decoder of stream begin", ERRCODE_INTERNAL_COULD_NOT_SET_DECODER_BEGIN);
+	VERIFY_HR(DecoderTransform->ProcessMessage(MFT_MESSAGE_NOTIFY_START_OF_STREAM, 0), "Failed to notify video decoder of start start", ERRCODE_INTERNAL_COULD_NOT_SET_DECODER_START);
+	return true;
+}
+
+
+//-----------------------------------------------------------------------------
+/**
+ * Configures the decoder input media type.
+ *
+ * @return true if successful, false on error
+ */
+bool FVideoDecoderH264::DecoderSetInputType()
+{
+	TRefCountPtr<IMFMediaType>	InputMediaType;
+	HRESULT					res;
+
+	// https://docs.microsoft.com/en-us/windows/desktop/medfound/h-264-video-decoder
+	VERIFY_HR(MFCreateMediaType(InputMediaType.GetInitReference()), "Failed to create video decoder input media type", ERRCODE_INTERNAL_COULD_NOT_CREATE_INPUT_MEDIA_TYPE);
+	VERIFY_HR(InputMediaType->SetGUID(MF_MT_MAJOR_TYPE, MFMediaType_Video), "Failed to set video decoder input media type major", ERRCODE_INTERNAL_COULD_NOT_CREATE_INPUT_MEDIA_TYPE_MAJOR);
+	VERIFY_HR(InputMediaType->SetGUID(MF_MT_SUBTYPE, MFVideoFormat_H264), "Failed to set video decoder input media sub type", ERRCODE_INTERNAL_COULD_NOT_CREATE_INPUT_MEDIA_TYPE_SUBTYPE);
+
+	int32 configW, configH;
+	if (bIsHardwareAccelerated)
+	{
+		configW = MaxHWDecodeW;
+		configH = MaxHWDecodeH;
+	}
+	else
+	{
+		configW = CurrentSampleInfo.GetResolution().Width;
+		configH = CurrentSampleInfo.GetResolution().Height;
+	}
+
+	if (CurrentAccessUnit && CurrentAccessUnit->AUCodecData.IsValid() && CurrentAccessUnit->AUCodecData->CodecSpecificData.Num())
+	{
+		VERIFY_HR(InputMediaType->SetBlob(MF_MT_USER_DATA, CurrentAccessUnit->AUCodecData->CodecSpecificData.GetData(), (UINT32) CurrentAccessUnit->AUCodecData->CodecSpecificData.Num()), "Failed to set video decoder input media type codec blob", ERRCODE_INTERNAL_COULD_NOT_SET_INPUT_MEDIA_TYPE_CODEC_BLOB);
+	}
+
+	VERIFY_HR(MFSetAttributeSize(InputMediaType, MF_MT_FRAME_SIZE, configW, configH), "Failed to set video decoder input media type resolution", ERRCODE_INTERNAL_COULD_NOT_CREATE_INPUT_MEDIA_TYPE_RESOLUTION);
+	VERIFY_HR(InputMediaType->SetUINT32(MF_MT_INTERLACE_MODE, MFVideoInterlace_MixedInterlaceOrProgressive), "Failed to set video decoder interlace mode", ERRCODE_INTERNAL_COULD_NOT_SET_INPUT_MEDIA_PROGRESSIVE_MODE);
+
+	res = DecoderTransform->SetInputType(0, InputMediaType, 0);
+
+#ifdef ELECTRA_ENABLE_SWDECODE
+	if (bIsHardwareAccelerated && res == MF_E_UNSUPPORTED_D3D_TYPE)
+		// h/w acceleration is not supported, e.g. unsupported resolution (4K), fall back to s/w decoding
+	{
+		return ReconfigureForSwDecoding(TEXT("MF_E_UNSUPPORTED_D3D_TYPE"));
+	}
+	else
+#endif
+
+	if (FAILED(res))
+	{
+		VERIFY_HR(res, "Failed to set video decoder input type", ERRCODE_INTERNAL_COULD_NOT_SET_INPUT_MEDIA_TYPE)
+		return false;
+	}
+
+	return true;
+}
+
+
+//-----------------------------------------------------------------------------
+/**
+ * Configures the decoder output media type.
+ *
+ * @return true if successful, false on error
+ */
+bool FVideoDecoderH264::DecoderSetOutputType()
+{
+	TRefCountPtr<IMFMediaType>	OutputMediaType;
+	GUID					OutputMediaMajorType;
+	GUID					OutputMediaSubtype;
+	HRESULT					res;
+
+	// Supposedly calling GetOutputAvailableType() returns following output media subtypes:
+	// MFVideoFormat_NV12, MFVideoFormat_YV12, MFVideoFormat_IYUV, MFVideoFormat_I420, MFVideoFormat_YUY2
+	for(int32 TypeIndex=0; ; ++TypeIndex)
+	{
+		VERIFY_HR(DecoderTransform->GetOutputAvailableType(0, TypeIndex, OutputMediaType.GetInitReference()), "Failed to get video decoder available output type", ERRCODE_INTERNAL_COULD_NOT_GET_OUTPUT_MEDIA_TYPE);
+		VERIFY_HR(OutputMediaType->GetGUID(MF_MT_MAJOR_TYPE, &OutputMediaMajorType), "Failed to get video decoder available output media type", ERRCODE_INTERNAL_COULD_NOT_GET_OUTPUT_MEDIA_TYPE_MAJOR);
+		VERIFY_HR(OutputMediaType->GetGUID(MF_MT_SUBTYPE, &OutputMediaSubtype), "Failed to get video decoder available output media subtype", ERRCODE_INTERNAL_COULD_NOT_GET_OUTPUT_MEDIA_TYPE_SUBTYPE);
+		if (OutputMediaMajorType == MFMediaType_Video && OutputMediaSubtype == MFVideoFormat_NV12)
+		{
+			VERIFY_HR(DecoderTransform->SetOutputType(0, OutputMediaType, 0), "Failed to set video decoder output type", ERRCODE_INTERNAL_COULD_NOT_SET_OUTPUT_MEDIA_TYPE);
+			CurrentOutputMediaType = OutputMediaType;
+			return true;
+		}
+	}
+	PostError(S_OK, "Failed to set video decoder output type to desired format", ERRCODE_INTERNAL_COULD_NOT_SET_OUTPUT_DESIRED_MEDIA_TYPE);
+	return false;
+}
+
+
+//-----------------------------------------------------------------------------
+/**
+ * Checks that the decoder is in the state we think it should be and good to go.
+ *
+ * @return true if successful, false on error
+ */
+bool FVideoDecoderH264::DecoderVerifyStatus()
+{
+	HRESULT		res;
+	DWORD		NumInputStreams;
+	DWORD		NumOutputStreams;
+
+	VERIFY_HR(DecoderTransform->GetStreamCount(&NumInputStreams, &NumOutputStreams), "Failed to get video decoder stream count", ERRCODE_INTERNAL_COULD_NOT_GET_STREAM_COUNT);
+	if (NumInputStreams != 1 || NumOutputStreams != 1)
+	{
+		PostError(S_OK, FString::Printf(TEXT("Unexpected number of streams: input %lu, output %lu"), NumInputStreams, NumOutputStreams), ERRCODE_INTERNAL_UNEXPECTED_NUMBER_OF_STREAMS);
+		return false;
+	}
+
+	DWORD DecoderStatus = 0;
+	VERIFY_HR(DecoderTransform->GetInputStatus(0, &DecoderStatus), "Failed to get video decoder input status", ERRCODE_INTERNAL_COULD_NOT_GET_INPUT_STATUS);
+	if (MFT_INPUT_STATUS_ACCEPT_DATA != DecoderStatus)
+	{
+		PostError(S_OK, FString::Printf(TEXT("Decoder doesn't accept data, status %lu"), DecoderStatus), ERRCODE_INTERNAL_DECODER_NOT_ACCEPTING_DATA);
+		return false;
+	}
+
+	VERIFY_HR(DecoderTransform->GetOutputStreamInfo(0, &DecoderOutputStreamInfo), "Failed to get video decoder output stream info", ERRCODE_INTERNAL_COULD_NOT_GET_OUTPUT_STREAM_INFO);
+	if (!(DecoderOutputStreamInfo.dwFlags & MFT_OUTPUT_STREAM_FIXED_SAMPLE_SIZE))
+	{
+		return ReconfigureForSwDecoding(TEXT("Incompatible H.264 decoder: fixed sample size expected"));
+	}
+	if (!(DecoderOutputStreamInfo.dwFlags & MFT_OUTPUT_STREAM_WHOLE_SAMPLES))
+	{
+		return ReconfigureForSwDecoding(TEXT("Incompatible H.264 decoder: whole samples expected"));
+	}
+	if (bIsHardwareAccelerated && !(DecoderOutputStreamInfo.dwFlags & MFT_OUTPUT_STREAM_PROVIDES_SAMPLES))
+	{
+		// theoretically we can handle this situation with H/W decoder, but we can't reproduce it locally for testing so we aren't sure if H/W
+		// decoder would work in this case
+		return ReconfigureForSwDecoding(TEXT("Incompatible H.264 decoder: h/w accelerated decoder is expected to provide output samples"));
+	}
+	if (!bIsHardwareAccelerated && (DecoderOutputStreamInfo.dwFlags & MFT_OUTPUT_STREAM_PROVIDES_SAMPLES))
+	{
+		PostError(S_OK, "Incompatible H.264 decoder: s/w decoder is expected to require preallocated output samples", ERRCODE_INTERNAL_INCOMPATIBLE_DECODER);
+		return false;
+	}
+	return true;
+}
+
+
+//-----------------------------------------------------------------------------
+/**
+ * Destroys the current decoder instance.
+ */
+void FVideoDecoderH264::InternalDecoderDestroy()
+{
+	DecoderTransform = nullptr;
+	CurrentOutputMediaType = nullptr;
+	bHaveDecoder = false;
+}
+
+
+//-----------------------------------------------------------------------------
+/**
+ * Notifies the buffer-ready listener that we will now be producing output data
+ * and need an output buffer.
+ * This call is not intended to create or obtain an output buffer. It is merely
+ * indicating that new output will be produced.
+ *
+ * @param bHaveOutput
+ */
+void FVideoDecoderH264::NotifyReadyBufferListener(bool bHaveOutput)
+{
+	if (ReadyBufferListener)
+	{
+		IDecoderOutputBufferListener::FDecodeReadyStats stats;
+		stats.MaxDecodedElementsReady = MaxDecodeBufferSize;
+		stats.NumElementsInDecoder    = NumFramesInDecoder;
+		stats.bOutputStalled		  = !bHaveOutput;
+		stats.bEODreached   		  = AccessUnits.IsEndOfData() && stats.NumDecodedElementsReady == 0 && stats.NumElementsInDecoder == 0;
+		ListenerMutex.Lock();
+		if (ReadyBufferListener)
+		{
+			ReadyBufferListener->DecoderOutputReady(stats);
+		}
+		ListenerMutex.Unlock();
+	}
+}
+
+
+
+void FVideoDecoderH264::PrepareAU(FAccessUnit* AccessUnit)
+{
+	if (!AccessUnit->bHasBeenPrepared)
+	{
+		AccessUnit->bHasBeenPrepared = true;
+
+		// Process NALUs
+		bool bIsDiscardable = true;
+		bool bIsKeyframe = false;
+		//AccessUnit->mIsDiscardable = 1;
+		// Replace the NALU lengths with the startcode.
+		uint32* pNALU = (uint32 *)AccessUnit->AUData;
+		uint32* pEnd  = (uint32 *)Electra::AdvancePointer(pNALU, AccessUnit->AUSize);
+		while(pNALU < pEnd)
+		{
+			// Check the nal_ref_idc in the NAL unit for dependencies.
+			uint8 nal = *(const uint8 *)(pNALU + 1);
+			check((nal & 0x80) == 0);
+			if ((nal >> 5) != 0)
+			{
+				bIsDiscardable = false;		//AccessUnit->mIsDiscardable = 0;
+			}
+			// IDR frame?
+			if ((nal & 0x1f) == 5)
+			{
+				bIsKeyframe = true;			//AccessUnit->mIsKeyframe = 1;
+			}
+
+			// SEI message(s)?
+			if ((nal & 0x1f) == 6)
+			{
+				// ...
+			}
+
+			uint32 naluLen = MEDIA_FROM_BIG_ENDIAN(*pNALU) + 4;
+			*pNALU = MEDIA_TO_BIG_ENDIAN(0x00000001U);
+			pNALU = Electra::AdvancePointer(pNALU, naluLen);
+		}
+
+		// Is there codec specific data?
+		if (AccessUnit->AUCodecData.IsValid())
+		{
+			// Yes.
+			if (bIsKeyframe || AccessUnit->bIsSyncSample || AccessUnit->bIsFirstInSequence)
+			{
+				NewSampleInfo = AccessUnit->AUCodecData->ParsedInfo;
+
+				// Have to re-allocate the AU memory to preprend the codec data
+				uint64 nb = AccessUnit->AUSize + AccessUnit->AUCodecData->CodecSpecificData.Num();
+				void* pD = AccessUnit->AllocatePayloadBuffer(nb);
+				check(pD);
+				void* pP = pD;
+				FMemory::Memcpy(pP, AccessUnit->AUCodecData->CodecSpecificData.GetData(), AccessUnit->AUCodecData->CodecSpecificData.Num());
+				pP = Electra::AdvancePointer(pP, AccessUnit->AUCodecData->CodecSpecificData.Num());
+				FMemory::Memcpy(pP, AccessUnit->AUData, AccessUnit->AUSize);
+				AccessUnit->AdoptNewPayloadBuffer(pD, nb);
+				// New codec data makes this AU non-discardable.
+				bIsDiscardable = false;
+			}
+		}
+	}
+}
+
+
+void FVideoDecoderH264::SetupBufferAcquisitionProperties()
+{
+	BufferAcquireOptions.Clear();
+	int32 w = CurrentSampleInfo.GetResolution().Width;
+	int32 h = CurrentSampleInfo.GetResolution().Height;
+	BufferAcquireOptions.Set("width",  FVariantValue((int64) w));
+	BufferAcquireOptions.Set("height", FVariantValue((int64) h));
+}
+
+
+
+void FVideoDecoderH264::ReturnUnusedFrame()
+{
+	CurrentDecoderOutputBuffer.Reset();
+	if (CurrentRenderOutputBuffer)
+	{
+		Electra::FParamDict DummySampleProperties;
+		Renderer->ReturnBuffer(CurrentRenderOutputBuffer, false, DummySampleProperties);
+		CurrentRenderOutputBuffer = nullptr;
+	}
+}
+
+
+bool FVideoDecoderH264::AcquireOutputBuffer()
+{
+	// If we still have one we're done.
+	if (CurrentRenderOutputBuffer)
+	{
+		return true;
+	}
+	while(!TerminateThreadSignal.IsSignaled())
+	{
+		if (FlushDecoderSignal.IsSignaled())
+		{
+			break;
+		}
+
+		UEMediaError bufResult = Renderer->AcquireBuffer(CurrentRenderOutputBuffer, 0, BufferAcquireOptions);
+		check(bufResult == UEMEDIA_ERROR_OK || bufResult == UEMEDIA_ERROR_INSUFFICIENT_DATA);
+		if (bufResult != UEMEDIA_ERROR_OK && bufResult != UEMEDIA_ERROR_INSUFFICIENT_DATA)
+		{
+			PostError(S_OK, "Failed to acquire sample buffer", ERRCODE_INTERNAL_COULD_NOT_GET_SAMPLE_BUFFER, bufResult);
+			return false;
+		}
+
+		bool bHaveOutputBuffer = CurrentRenderOutputBuffer != nullptr;
+		NotifyReadyBufferListener(bHaveOutputBuffer);
+		if (bHaveOutputBuffer)
+		{
+			break;
+		}
+		else
+		{
+			// No available buffer. Sleep for a bit on the flush signal. This gets set before the termination signal.
+			FlushDecoderSignal.WaitTimeout(1000 * 5);
+		}
+	}
+	return true;
+
+}
+
+
+bool FVideoDecoderH264::ConvertDecodedImage(const TRefCountPtr<IMFSample>& DecodedOutputSample)
+{
+	FParamDict* OutputBufferSampleProperties = new FParamDict();
+	if (!OutputBufferSampleProperties)
+	{
+		return false;
+	}
+
+	HRESULT					res;
+	LONGLONG				llTimeStamp = 0;
+	MFVideoArea				videoArea = {};
+	UINT32					uiPanScanEnabled = 0;
+	FTimeValue				TimeStamp, Duration;
+
+	//PVS reports warning 773 (possible memory leak) for OutputBufferSampleProperties in the line below. It's a false positive as passing it into SetupDecodeOutputData() later on will pass on ownership implicitly.
+	VERIFY_HR(DecodedOutputSample->GetSampleTime(&llTimeStamp), "Failed to get video decoder output sample timestamp", ERRCODE_INTERNAL_COULD_NOT_GET_OUTPUT_SAMPLE_TIME); //-V773
+	TimeStamp.SetFromHNS((int64) llTimeStamp);
+	OutputBufferSampleProperties->Set("pts", FVariantValue(TimeStamp));
+
+	VERIFY_HR(DecodedOutputSample->GetSampleDuration(&llTimeStamp), "Failed to get video decoder output sample duration", ERRCODE_INTERNAL_COULD_NOT_GET_OUTPUT_SAMPLE_DURATION);
+	Duration.SetFromHNS((int64) llTimeStamp);
+	OutputBufferSampleProperties->Set("duration", FVariantValue(Duration));
+
+	if (SUCCEEDED(res = CurrentOutputMediaType->GetUINT32(MF_MT_PAN_SCAN_ENABLED, &uiPanScanEnabled)) && uiPanScanEnabled)
+	{
+		res = CurrentOutputMediaType->GetBlob(MF_MT_PAN_SCAN_APERTURE, (UINT8*)&videoArea, sizeof(MFVideoArea), nullptr);
+	}
+
+	UINT32	dwInputWidth = 0;
+	UINT32	dwInputHeight = 0;
+	res = MFGetAttributeSize(CurrentOutputMediaType.GetReference(), MF_MT_FRAME_SIZE, &dwInputWidth, &dwInputHeight);
+	check(SUCCEEDED(res));
+	res = CurrentOutputMediaType->GetBlob(MF_MT_MINIMUM_DISPLAY_APERTURE, (UINT8*)&videoArea, sizeof(MFVideoArea), nullptr);
+	if (FAILED(res))
+	{
+		res = CurrentOutputMediaType->GetBlob(MF_MT_GEOMETRIC_APERTURE, (UINT8*)&videoArea, sizeof(MFVideoArea), nullptr);
+		if (FAILED(res))
+		{
+			videoArea.OffsetX.fract = 0;
+			videoArea.OffsetX.value = 0;
+			videoArea.OffsetY.fract = 0;
+			videoArea.OffsetY.value = 0;
+			videoArea.Area.cx = dwInputWidth;
+			videoArea.Area.cy = dwInputHeight;
+		}
+	}
+
+	// Width and height is the cropped area!
+	OutputBufferSampleProperties->Set("width",  FVariantValue((int64) videoArea.Area.cx));
+	OutputBufferSampleProperties->Set("height", FVariantValue((int64) videoArea.Area.cy));
+
+	if (!bIsHardwareAccelerated || (FDXDeviceInfo::s_DXDeviceInfo->DxVersion == FDXDeviceInfo::ED3DVersion::Version11Win8))
+	{
+		// Cropping is applied to immediate decoder output and never visible to outside code
+		OutputBufferSampleProperties->Set("crop_left", FVariantValue((int64)0));
+		OutputBufferSampleProperties->Set("crop_right", FVariantValue((int64)0));
+		OutputBufferSampleProperties->Set("crop_top", FVariantValue((int64)0));
+		OutputBufferSampleProperties->Set("crop_bottom", FVariantValue((int64)0));
+	}
+	else
+	{
+		// Note: WMF decoder will always crop only at the lower, right borders
+		OutputBufferSampleProperties->Set("crop_left", FVariantValue((int64)0));
+		OutputBufferSampleProperties->Set("crop_right", FVariantValue((int64)(dwInputWidth - videoArea.Area.cx)));
+		OutputBufferSampleProperties->Set("crop_top", FVariantValue((int64)0));
+		OutputBufferSampleProperties->Set("crop_bottom", FVariantValue((int64)(dwInputHeight - videoArea.Area.cy)));
+	}
+
+	// Try to get the stride. Defaults to 0 should it not be obtainable.
+	UINT32 stride = MFGetAttributeUINT32(CurrentOutputMediaType, MF_MT_DEFAULT_STRIDE, 0);
+	OutputBufferSampleProperties->Set("pitch",  FVariantValue((int64) stride));
+
+	// Try to get the frame rate ratio
+	UINT32 num=0, denom=0;
+	MFGetAttributeRatio(CurrentOutputMediaType, MF_MT_FRAME_RATE, &num, &denom);
+	OutputBufferSampleProperties->Set("fps_num",   FVariantValue((int64) num ));
+	OutputBufferSampleProperties->Set("fps_denom", FVariantValue((int64) denom ));
+
+	// Try to get the pixel aspect ratio
+	num=0, denom=0;
+	MFGetAttributeRatio(CurrentOutputMediaType, MF_MT_PIXEL_ASPECT_RATIO, &num, &denom);
+	if (!num || !denom)
+	{
+		num   = 1;
+		denom = 1;
+	}
+	OutputBufferSampleProperties->Set("aspect_ratio", FVariantValue((double) num / (double)denom));
+	OutputBufferSampleProperties->Set("aspect_w",     FVariantValue((int64) num));
+	OutputBufferSampleProperties->Set("aspect_h",     FVariantValue((int64) denom));
+
+	OutputBufferSampleProperties->Set("pixelfmt",	  FVariantValue((int64)EPixelFormat::PF_NV12));
+
+	if (CurrentRenderOutputBuffer != nullptr)
+	{
+		if (!SetupDecodeOutputData(FIntPoint(videoArea.Area.cx, videoArea.Area.cy), DecodedOutputSample, OutputBufferSampleProperties))
+		{
+			return false;
+		}
+
+		Renderer->ReturnBuffer(CurrentRenderOutputBuffer, true, *OutputBufferSampleProperties);
+		CurrentRenderOutputBuffer = nullptr;
+
+	}
+	return true;
+}
+
+
+bool FVideoDecoderH264::Decode(FAccessUnit* AccessUnit, bool bResolutionChanged)
+{
+	// No decoder, nothing to do.
+	if (DecoderTransform.GetReference() == nullptr)
+	{
+		return true;
+	}
+	// Draining the transform for resolution changes is only required for software transforms
+	if (bIsHardwareAccelerated && bResolutionChanged)
+	{
+		return true;
+	}
+
+	HRESULT res;
+	TRefCountPtr<IMFSample>	InputSample;
+	bool bFlush = AccessUnit == nullptr;
+
+	if (bFlush)
+	{
+		if (bResolutionChanged)
+		{
+			VERIFY_HR(DecoderTransform->ProcessMessage(MFT_MESSAGE_COMMAND_DRAIN, 0), "Failed to issue video decoder drain command for resolution change", ERRCODE_INTERNAL_COULD_NOT_SET_DECODER_DRAINCOMMAND);
+		}
+		else
+		{
+			VERIFY_HR(DecoderTransform->ProcessMessage(MFT_MESSAGE_NOTIFY_END_OF_STREAM, 0), "Failed to set video decoder end of stream notification", ERRCODE_INTERNAL_COULD_NOT_SET_DECODER_ENDOFSTREAM);
+			VERIFY_HR(DecoderTransform->ProcessMessage(MFT_MESSAGE_COMMAND_DRAIN, 0), "Failed to issue video decoder drain command", ERRCODE_INTERNAL_COULD_NOT_SET_DECODER_DRAINCOMMAND);
+		}
+	}
+	else
+	{
+		// Create the input sample.
+		TRefCountPtr<IMFMediaBuffer> InputSampleBuffer;
+		BYTE*					pbNewBuffer = nullptr;
+		DWORD					dwMaxBufferSize = 0;
+		DWORD					dwSize = 0;
+		LONGLONG				llSampleTime = 0;
+
+		CSV_SCOPED_TIMING_STAT(ElectraPlayer, FVideoDecoderH264_Worker);
+
+		VERIFY_HR(MFCreateSample(InputSample.GetInitReference()), "Failed to create video decoder input sample", ERRCODE_INTERNAL_COULD_NOT_CREATE_INPUT_SAMPLE);
+		VERIFY_HR(MFCreateMemoryBuffer((DWORD) AccessUnit->AUSize, InputSampleBuffer.GetInitReference()), "Failed to create video decoder input sample memory buffer", ERRCODE_INTERNAL_COULD_NOT_CREATE_INPUTBUFFER);
+		VERIFY_HR(InputSample->AddBuffer(InputSampleBuffer.GetReference()), "Failed to set video decoder input buffer with sample", ERRCODE_INTERNAL_COULD_NOT_ADD_INPUT_BUFFER_TO_SAMPLE);
+		VERIFY_HR(InputSampleBuffer->Lock(&pbNewBuffer, &dwMaxBufferSize, &dwSize), "Failed to lock video decoder input sample buffer", ERRCODE_INTERNAL_COULD_NOT_LOCK_INPUT_BUFFER);
+		FMemory::Memcpy(pbNewBuffer, AccessUnit->AUData, AccessUnit->AUSize);
+		VERIFY_HR(InputSampleBuffer->Unlock(), "Failed to unlock video decoder input sample buffer", ERRCODE_INTERNAL_COULD_NOT_UNLOCK_INPUT_BUFFER);
+		VERIFY_HR(InputSampleBuffer->SetCurrentLength((DWORD) AccessUnit->AUSize), "Failed to set video decoder input sample buffer length", ERRCODE_INTERNAL_COULD_NOT_SET_BUFFER_CURRENT_LENGTH);
+		// Set sample attributes
+		llSampleTime = AccessUnit->PTS.GetAsHNS();
+		VERIFY_HR(InputSample->SetSampleTime(llSampleTime), "Failed to set video decoder input sample presentation time", ERRCODE_INTERNAL_COULD_NOT_SET_INPUT_SAMPLE_TIME);
+		llSampleTime = AccessUnit->Duration.GetAsHNS();
+		VERIFY_HR(InputSample->SetSampleDuration(llSampleTime), "Failed to set video decode intput sample duration", ERRCODE_INTERNAL_COULD_NOT_SET_INPUT_SAMPLE_DURATION);
+		llSampleTime = AccessUnit->DTS.GetAsHNS();
+		VERIFY_HR(InputSample->SetUINT64(MFSampleExtension_DecodeTimestamp, llSampleTime), "Failed to set video decoder input sample decode time", ERRCODE_INTERNAL_COULD_NOT_SET_INPUT_SAMPLE_DECODETIME);
+		VERIFY_HR(InputSample->SetUINT32(MFSampleExtension_CleanPoint, AccessUnit->bIsSyncSample ? 1 : 0), "Failed to set video decoder input sample clean point", ERRCODE_INTERNAL_COULD_NOT_SET_INPUT_SAMPLE_KEYFRAME);
+	}
+
+	while(!TerminateThreadSignal.IsSignaled())
+	{
+#ifdef ELECTRA_ENABLE_SWDECODE
+		// For SW decoding we need to get a buffer from the renderer right now.
+		if (!bIsHardwareAccelerated)
+		{
+			if (!AcquireOutputBuffer())
+			{
+				return false;
+			}
+
+			// Check if we got an output buffer. When flushing or terminating AcquireOutputBuffer() may return with no buffer!
+			if (CurrentRenderOutputBuffer)
+			{
+				// Setup texture sample to receive data during decode
+				PreInitDecodeOutputForSW(FIntPoint(CurrentSampleInfo.GetResolution().Width, CurrentSampleInfo.GetResolution().Height));
+			}
+			else
+			{
+				break;
+			}
+		}
+#endif
+
+		// Check if to be flushed now. This may have resulted in AcquireOutputBuffer() to return with no buffer!
+		if (FlushDecoderSignal.IsSignaled())
+		{
+			break;
+		}
+
+		if (!CurrentDecoderOutputBuffer.IsValid())
+		{
+			CSV_SCOPED_TIMING_STAT(ElectraPlayer, FVideoDecoderH264_Worker);
+
+			if (!CreateDecoderOutputBuffer())
+			{
+				return false;
+			}
+		}
+
+		CurrentDecoderOutputBuffer->PrepareForProcess();
+		DWORD	dwStatus = 0;
+		res = DecoderTransform->ProcessOutput(0, 1, &CurrentDecoderOutputBuffer->mOutputBuffer, &dwStatus);
+		CurrentDecoderOutputBuffer->UnprepareAfterProcess();
+
+		if (res == MF_E_TRANSFORM_NEED_MORE_INPUT)
+		{
+			// Draining for resolution change?
+			if (bResolutionChanged || bDecoderFlushPending)
+			{
+				bDecoderFlushPending = false;
+				CurrentSampleInfo = NewSampleInfo;
+				SetupBufferAcquisitionProperties();
+				ReturnUnusedFrame();
+				NumFramesInDecoder = 0;
+				if (!DecoderSetInputType())
+				{
+					return false;
+				}
+				return DecoderSetOutputType();
+			}
+			else if (bFlush)
+			{
+				// Yes. This means we have received all pending output and are done now.
+				// Issue a flush for completeness sake.
+				VERIFY_HR(DecoderTransform->ProcessMessage(MFT_MESSAGE_COMMAND_FLUSH, 0), "Failed to issue video decoder flush command", ERRCODE_INTERNAL_COULD_NOT_SET_DECODER_FLUSHCOMMAND);
+				// And start over.
+				VERIFY_HR(DecoderTransform->ProcessMessage(MFT_MESSAGE_NOTIFY_BEGIN_STREAMING, 0), "Failed to set video decoder stream begin", ERRCODE_INTERNAL_COULD_NOT_SET_DECODER_BEGIN);
+				VERIFY_HR(DecoderTransform->ProcessMessage(MFT_MESSAGE_NOTIFY_START_OF_STREAM, 0), "Failed to start video decoder", ERRCODE_INTERNAL_COULD_NOT_SET_DECODER_START);
+				CurrentDecoderOutputBuffer.Reset();
+				NumFramesInDecoder = 0;
+				return true;
+			}
+			else if (InputSample.IsValid())
+			{
+				CSV_SCOPED_TIMING_STAT(ElectraPlayer, FVideoDecoderH264_Worker);
+
+				VERIFY_HR(DecoderTransform->ProcessInput(0, InputSample.GetReference(), 0), "Failed to process video decoder input", ERRCODE_INTERNAL_COULD_NOT_PROCESS_INPUT);
+				// Used this sample. Have no further input data for now, but continue processing to produce output if possible.
+				InputSample = nullptr;
+				++NumFramesInDecoder;
+			}
+			else
+			{
+				// Need more input but have none right now.
+				return true;
+			}
+		}
+		else if (res == MF_E_TRANSFORM_STREAM_CHANGE)
+		{
+			//if((CurrentDecoderOutputBuffer->mOutputBuffer.dwFlags & MFT_OUTPUT_DATA_BUFFER_FORMAT_CHANGE) != 0)
+			{
+				// Update output type.
+				if(!DecoderSetOutputType())
+				{
+					return false;
+				}
+			}
+		}
+		else if (SUCCEEDED(res))
+		{
+			TRefCountPtr<IMFSample> DecodedOutputSample = CurrentDecoderOutputBuffer->DetachOutputSample();
+			//LogMessage(IInfoLog::ELevel::Info, StringHelpers::SPrintf(TEXT("OUT: %p"), DecodedOutputSample.Get()));
+			CurrentDecoderOutputBuffer.Reset();
+
+			if (DecodedOutputSample)
+			{
+				// HW accelerated decoding needs to get a buffer from the renderer only now after decoding.
+				// SW decoding required it earlier for the call to ProcessOutput().
+				if (bIsHardwareAccelerated)
+				{
+					if (!AcquireOutputBuffer())
+					{
+						return false;
+					}
+				}
+
+				CSV_SCOPED_TIMING_STAT(ElectraPlayer, FVideoDecoderH264_Worker);
+
+				// Check if to be flushed now. This may have resulted in AcquireOutputBuffer() to return with no buffer!
+				if (FlushDecoderSignal.IsSignaled())
+				{
+					break;
+				}
+				if (!ConvertDecodedImage(DecodedOutputSample))
+				{
+					return false;
+				}
+			}
+			--NumFramesInDecoder;
+		}
+		else
+		{
+			// Error!
+			VERIFY_HR(res, "Failed to process video decoder output", ERRCODE_INTERNAL_COULD_NOT_PROCESS_OUTPUT);
+			return false;
+		}
+	}
+	return true;
+}
+
+
+bool FVideoDecoderH264::PerformFlush()
+{
+	// No decoder, nothing to do.
+	if (DecoderTransform.GetReference() == nullptr)
+	{
+		return true;
+	}
+
+	HRESULT				res;
+	bDecoderFlushPending = true;
+	VERIFY_HR(DecoderTransform->ProcessMessage(MFT_MESSAGE_COMMAND_FLUSH, 0), "Failed to issue video decoder flush command", ERRCODE_INTERNAL_COULD_NOT_SET_DECODER_FLUSHCOMMAND);
+
+	return true;
+}
+
+
+bool FVideoDecoderH264::DecodeDummy(FAccessUnit* AccessUnit)
+{
+	check(AccessUnit->Duration.IsValid());
+
+	FParamDict OutputBufferSampleProperties;
+
+	OutputBufferSampleProperties.Set("pts", FVariantValue(AccessUnit->PTS));
+	OutputBufferSampleProperties.Set("duration", FVariantValue(AccessUnit->Duration));
+	OutputBufferSampleProperties.Set("is_dummy", FVariantValue(true));
+
+	if (!AcquireOutputBuffer())
+	{
+		return false;
+	}
+
+	if (CurrentRenderOutputBuffer != nullptr)
+	{
+//TODO: THIS BEING A DUMMY - THE DICT IN ANY DECODER OUTPUT BUFFER (is there any?) WILL BE OUTDATED THIS WAY!!
+		Renderer->ReturnBuffer(CurrentRenderOutputBuffer, true, OutputBufferSampleProperties);
+		CurrentRenderOutputBuffer = nullptr;
+	}
+	return true;
+}
+
+
+void FVideoDecoderH264::DecoderCreate()
+{
+	if (InternalDecoderCreate())
+	{
+		// Configure the decoder with our default values.
+		bRequiresReconfigurationForSW = false;
+		Configure();
+		if (!bError && bRequiresReconfigurationForSW)
+		{
+			// No failure yet, but a switch to software decoding is required. We do the configuration over one more time.
+			check(!bIsHardwareAccelerated);	// must have been reset already!
+			// Clear this out, we can't get another request for reconfiguration since we already did.
+			bRequiresReconfigurationForSW = false;
+			Configure();
+			if (!bError && bRequiresReconfigurationForSW)
+			{
+				PostError(S_OK, "Failed to switch H.264 decoder transform to software decoding", ERRCODE_INTERNAL_FAILED_TO_SWITCH_TO_SOFTWARE_MODE);
+			}
+		}
+		if (!bError)
+		{
+			StartStreaming();
+			bHaveDecoder = true;
+		}
+		else
+		{
+			InternalDecoderDestroy();
+		}
+	}
+	else
+	{
+		bError = true;
+		PostError(S_OK, "Failed to create H.264 decoder transform", ERRCODE_INTERNAL_COULD_NOT_CREATE_DECODER);
+	}
+}
+
+
+//-----------------------------------------------------------------------------
+/**
+ * Creates a decoder and checks if it will use hardware or software decoding.
+ */
+void FVideoDecoderH264::TestHardwareDecoding()
+{
+	bRequiresReconfigurationForSW = false;
+	CurrentSampleInfo.SetResolution(FStreamCodecInformation::FResolution(Align(Config.MaxFrameWidth, 16), Align(Config.MaxFrameHeight, 16)));
+	NewSampleInfo = CurrentSampleInfo;
+	bIsHardwareAccelerated = true;
+	DecoderCreate();
+	InternalDecoderDestroy();
+}
+
+
+//-----------------------------------------------------------------------------
+/**
+ * H264 video decoder main threaded decode loop
+ */
+void FVideoDecoderH264::WorkerThread()
+{
+	LLM_SCOPE(ELLMTag::ElectraPlayer);
+
+	CurrentRenderOutputBuffer     = nullptr;
+	CurrentAccessUnit   		  = nullptr;
+	bHaveDecoder				  = false;
+	NumFramesInDecoder  		  = 0;
+	bDecoderFlushPending		  = false;
+	bRequiresReconfigurationForSW = false;
+	bError  					  = false;
+
+	// Set the configured maximum resolution as the current resolution which is used to create the decoder transform.
+	CurrentSampleInfo.SetResolution(FStreamCodecInformation::FResolution(Align(Config.MaxFrameWidth, 16), Align(Config.MaxFrameHeight, 16)));
+	NewSampleInfo = CurrentSampleInfo;
+	SetupBufferAcquisitionProperties();
+	// Create a decoder transform. This will determine if we will be using a hardware or software decoder.
+	bIsHardwareAccelerated = true;
+	DecoderCreate();
+
+	// Now that we know whether we will be running a hardware or software decoder we can create the decoded image pool.
+	if (!CreateDecodedImagePool())
+	{
+		bError = true;
+	}
+
+	bool bDone = false;
+	bool bInDummyDecodeMode = false;
+
+	// Require a new media input type based on the actual first access unit.
+	bool bNeedInitialReconfig = true;
+	while(!TerminateThreadSignal.IsSignaled())
+	{
+		// Notify optional buffer listener that we will now be needing an AU for our input buffer.
+		if (!bError && InputBufferListener && AccessUnits.Num() == 0)
+		{
+			CSV_SCOPED_TIMING_STAT(ElectraPlayer, FVideoDecoderH264_Worker);
+
+			FAccessUnitBufferInfo	sin;
+			IAccessUnitBufferListener::FBufferStats	stats;
+			AccessUnits.GetStats(sin);
+			stats.NumAUsAvailable  = sin.NumCurrentAccessUnits;
+			stats.NumBytesInBuffer = sin.CurrentMemInUse;
+			stats.MaxBytesOfBuffer = sin.MaxDataSize;
+			stats.bEODSignaled     = sin.bEndOfData;
+			stats.bEODReached      = sin.bEndOfData && sin.NumCurrentAccessUnits == 0;
+			ListenerMutex.Lock();
+			if (InputBufferListener)
+			{
+				InputBufferListener->DecoderInputNeeded(stats);
+			}
+			ListenerMutex.Unlock();
+		}
+		// Wait for data.
+		bool bHaveData = AccessUnits.WaitForData(1000 * 5);
+		// When there is data, even and especially after a previous EOD, we are no longer done and idling.
+		if (bHaveData && bDone)
+		{
+			bDone = false;
+		}
+
+		// Tick output buffer pool to handle state progression of samples in flight
+		Renderer->TickOutputBufferPool();
+		// Tick platform code in case something needs doing
+		PlatformTick();
+
+		if (bHaveData && Renderer->CanReceiveOutputFrames(1)) // note: the check for full output buffers ignores the NumFramesInDecoder value as we don't seem to receive new data if we don't feed new input
+		{
+			CurrentAccessUnit = nullptr;
+			bool bOk = AccessUnits.Pop(CurrentAccessUnit);
+			MEDIA_UNUSED_VAR(bOk);
+			check(bOk);
+
+			if (!bError)
+			{
+				if (!CurrentAccessUnit->bIsDummyData)
+				{
+					{
+					CSV_SCOPED_TIMING_STAT(ElectraPlayer, FVideoDecoderH264_Worker);
+					PrepareAU(CurrentAccessUnit);
+					}
+
+					// Decode
+					if (!bError)
+					{
+						// Update the buffer acquisition after a change in streams?
+						bool bFormatChangedJustNow = CurrentSampleInfo.IsDifferentFromOtherVideo(NewSampleInfo);
+
+						if (bInDummyDecodeMode)
+						{
+							bFormatChangedJustNow = true;
+							bInDummyDecodeMode = false;
+						}
+
+						if (!bHaveDecoder)
+						{
+							CSV_SCOPED_TIMING_STAT(ElectraPlayer, FVideoDecoderH264_Worker);
+							DecoderCreate();
+							bNeedInitialReconfig = true;
+						}
+
+						if (bFormatChangedJustNow || bNeedInitialReconfig)
+						{
+							if (bNeedInitialReconfig)
+							{
+								bDecoderFlushPending = true;
+							}
+							if (!Decode(nullptr, bNeedInitialReconfig ? false : true))
+							{
+								// FIXME: Is this error fatal or safe to ignore?
+							}
+							CurrentSampleInfo = NewSampleInfo;
+							SetupBufferAcquisitionProperties();
+						}
+						bNeedInitialReconfig = false;
+
+						if (!Decode(CurrentAccessUnit, false))
+						{
+							bError = true;
+						}
+					}
+				}
+				else
+				{
+					if (!bInDummyDecodeMode)
+					{
+						bInDummyDecodeMode = true;
+						Decode(nullptr, true);
+					}
+
+					if (!DecodeDummy(CurrentAccessUnit))
+					{
+						bError = true;
+					}
+				}
+			}
+
+			// If decoding was aborted and the AU not handed off, delete it now.
+			FAccessUnit::Release(CurrentAccessUnit);
+			CurrentAccessUnit = nullptr;
+		}
+		else
+		{
+			if (!bHaveData)
+			{
+				// No data. Is the buffer at EOD?
+				if (AccessUnits.IsEndOfData())
+				{
+					NotifyReadyBufferListener(true);
+					// Are we done yet?
+					if (!bDone && !bError)
+					{
+						bError = !Decode(nullptr, false);
+						NumFramesInDecoder = 0;
+					}
+					bDone = true;
+					FMediaRunnable::SleepMilliseconds(10);
+				}
+			}
+			else
+			{
+				// We could not decode new data as the output buffers are full... notify listener & wait a bit...
+				NotifyReadyBufferListener(false);
+				FMediaRunnable::SleepMilliseconds(5);
+			}
+		}
+
+		// Flush?
+		if (FlushDecoderSignal.IsSignaled())
+		{
+			CSV_SCOPED_TIMING_STAT(ElectraPlayer, FVideoDecoderH264_Worker);
+
+			// Destroy and re-create.
+			InternalDecoderDestroy();
+
+			ReturnUnusedFrame();
+			AccessUnits.Flush();
+
+			FlushDecoderSignal.Reset();
+			DecoderFlushedSignal.Signal();
+			NumFramesInDecoder = 0;
+			CurrentSampleInfo.SetResolution(FStreamCodecInformation::FResolution(Align(Config.MaxFrameWidth, 16), Align(Config.MaxFrameHeight, 16)));
+			NewSampleInfo = CurrentSampleInfo;
+
+			// Reset done state.
+			bDone = false;
+		}
+	}
+
+	InternalDecoderDestroy();
+	ReturnUnusedFrame();
+	DestroyDecodedImagePool();
+	AccessUnits.Flush();
+	AccessUnits.CapacitySet(0);
+}
+
+} // namespace Electra
+
+#endif // ELECTRA_ENABLE_MFDECODER
+
+
