ElectraPlayer: RELEASE. Moved from NFL to public. Note that this is done via a Delete and "Add New" operation as discussed.

diff --git a/Engine/Plugins/Media/ElectraPlayer/Source/ElectraPlayerRuntime/Private/Runtime/Decoder/DX/VideoDecoderH264_DX.h b/Engine/Plugins/Media/ElectraPlayer/Source/ElectraPlayerRuntime/Private/Runtime/Decoder/DX/VideoDecoderH264_DX.h
new file mode 100644
--- /dev/null
+++ b/Engine/Plugins/Media/ElectraPlayer/Source/ElectraPlayerRuntime/Private/Runtime/Decoder/DX/VideoDecoderH264_DX.h
@@ -0,0 +1,181 @@
+// Copyright Epic Games, Inc. All Rights Reserved.
+
+#pragma once
+
+namespace Electra {
+
+	// Define the H.264 class GUID ourselves. Using the CLSID_CMSH264DecoderMFT from wmcodecdsp.h requires linking with some lib again...
+	static const GUID MFTmsH264Decoder = { 0x62CE7E72, 0x4C71, 0x4D20, { 0xB1, 0x5D, 0x45, 0x28, 0x31, 0xA8, 0x7D, 0x9D } };
+
+	static const GUID CODECAPI_AVLowLatencyMode = { 0x9c27891a, 0xed7a, 0x40e1, { 0x88, 0xe8, 0xb2, 0x27, 0x27, 0xa0, 0x24, 0xee } };
+
+	/**
+	 * H264 video decoder class implementation.
+	**/
+	class FVideoDecoderH264 : public IVideoDecoderH264, public FMediaThread
+	{
+	public:
+		static bool Startup(const IVideoDecoderH264::FSystemConfiguration& InConfig);
+		static void Shutdown();
+
+		static bool GetStreamDecodeCapability(FStreamDecodeCapability& OutResult, const FStreamDecodeCapability& InStreamParameter);
+
+		FVideoDecoderH264();
+		virtual ~FVideoDecoderH264();
+
+		void TestHardwareDecoding();
+
+		virtual void SetPlayerSessionServices(IPlayerSessionServices* SessionServices) override;
+
+		virtual void Open(const FInstanceConfiguration& InConfig) override;
+		virtual void Close() override;
+
+		virtual void SetMaximumDecodeCapability(int32 MaxWidth, int32 MaxHeight, int32 MaxProfile, int32 MaxProfileLevel, const FParamDict& AdditionalOptions) override;
+
+		virtual void SetAUInputBufferListener(IAccessUnitBufferListener* InListener) override;
+
+		virtual void SetReadyBufferListener(IDecoderOutputBufferListener* InListener) override;
+
+		virtual void SetRenderer(TSharedPtr<IMediaRenderer, ESPMode::ThreadSafe> InRenderer) override;
+
+		virtual void SetResourceDelegate(const TSharedPtr<IVideoDecoderResourceDelegate, ESPMode::ThreadSafe>& ResourceDelegate) override;
+
+		virtual IAccessUnitBufferInterface::EAUpushResult AUdataPushAU(FAccessUnit* InAccessUnit) override;
+		virtual void AUdataPushEOD() override;
+		virtual void AUdataFlushEverything() override;
+
+	protected:
+		struct FDecoderOutputBuffer
+		{
+			FDecoderOutputBuffer()
+			{
+				FMemory::Memzero(mOutputStreamInfo);
+				FMemory::Memzero(mOutputBuffer);
+			}
+			~FDecoderOutputBuffer()
+			{
+				if (mOutputBuffer.pSample)
+				{
+					mOutputBuffer.pSample->Release();
+				}
+			}
+			TRefCountPtr<IMFSample> DetachOutputSample()
+			{
+				TRefCountPtr<IMFSample> pOutputSample;
+				if (mOutputBuffer.pSample)
+				{
+					pOutputSample = TRefCountPtr<IMFSample>(mOutputBuffer.pSample, false);
+					mOutputBuffer.pSample = nullptr;
+				}
+				return pOutputSample;
+			}
+			void PrepareForProcess()
+			{
+				mOutputBuffer.dwStatus = 0;
+				mOutputBuffer.dwStreamID = 0;
+				mOutputBuffer.pEvents = nullptr;
+			}
+			void UnprepareAfterProcess()
+			{
+				if (mOutputBuffer.pEvents)
+				{
+					// https://docs.microsoft.com/en-us/windows/desktop/api/mftransform/nf-mftransform-imftransform-processoutput
+					// The caller is responsible for releasing any events that the MFT allocates.
+					mOutputBuffer.pEvents->Release();
+					mOutputBuffer.pEvents = nullptr;
+				}
+			}
+			MFT_OUTPUT_STREAM_INFO	mOutputStreamInfo;
+			MFT_OUTPUT_DATA_BUFFER	mOutputBuffer;
+		};
+
+		void InternalDecoderDestroy();
+		void DecoderCreate();
+		bool DecoderSetInputType();
+		bool DecoderSetOutputType();
+		bool DecoderVerifyStatus();
+		void StartThread();
+		void StopThread();
+		void WorkerThread();
+
+		bool CreateDecodedImagePool();
+		void DestroyDecodedImagePool();
+
+		void NotifyReadyBufferListener(bool bHaveOutput);
+
+		void SetupBufferAcquisitionProperties();
+
+		bool AcquireOutputBuffer();
+		bool ConvertDecodedImage(const TRefCountPtr<IMFSample>& DecodedSample);
+
+		void PrepareAU(FAccessUnit* InAccessUnit);
+
+		bool Decode(FAccessUnit* InAccessUnit, bool bResolutionChanged);
+		bool PerformFlush();
+		bool DecodeDummy(FAccessUnit* InAccessUnit);
+
+		void ReturnUnusedFrame();
+
+		void PostError(int32 ApiReturnValue, const FString& Message, uint16 Code, UEMediaError Error = UEMEDIA_ERROR_OK);
+		void LogMessage(IInfoLog::ELevel Level, const FString& Message);
+		bool FallbackToSwDecoding(FString Reason);
+		bool ReconfigureForSwDecoding(FString Reason);
+		bool Configure();
+		bool StartStreaming();
+
+		// Per platform specialization
+		virtual bool InternalDecoderCreate() = 0;
+		virtual bool CreateDecoderOutputBuffer() = 0;
+		virtual void PreInitDecodeOutputForSW(const FIntPoint& Dim) = 0;
+		virtual bool SetupDecodeOutputData(const FIntPoint& ImageDim, const TRefCountPtr<IMFSample>& DecodedOutputSample, FParamDict* OutputBufferSampleProperties) = 0;
+		virtual void PlatformTick() {}
+
+		FInstanceConfiguration								Config;
+
+		FMediaEvent											TerminateThreadSignal;
+		FMediaEvent											FlushDecoderSignal;
+		FMediaEvent											DecoderFlushedSignal;
+		bool												bThreadStarted;
+
+		IPlayerSessionServices*								PlayerSessionServices;
+
+		TSharedPtr<IMediaRenderer, ESPMode::ThreadSafe>		Renderer;
+
+		TWeakPtr<IVideoDecoderResourceDelegate, ESPMode::ThreadSafe> ResourceDelegate;
+
+		FAccessUnitBuffer									AccessUnits;
+		FStreamCodecInformation								NewSampleInfo;
+		FStreamCodecInformation								CurrentSampleInfo;
+		FAccessUnit*										CurrentAccessUnit;
+
+		FMediaCriticalSection								ListenerMutex;
+		IAccessUnitBufferListener*							InputBufferListener;
+		IDecoderOutputBufferListener*						ReadyBufferListener;
+
+		TRefCountPtr<IMFTransform>							DecoderTransform;
+		TRefCountPtr<IMFMediaType>							CurrentOutputMediaType;
+		MFT_OUTPUT_STREAM_INFO								DecoderOutputStreamInfo;
+		bool												bIsHardwareAccelerated;
+		bool												bRequiresReconfigurationForSW;
+		int32												NumFramesInDecoder;
+		bool												bDecoderFlushPending;
+		bool												bError;
+
+		TUniquePtr<FDecoderOutputBuffer>					CurrentDecoderOutputBuffer;
+		IMediaRenderer::IBuffer*							CurrentRenderOutputBuffer;
+		FParamDict											BufferAcquireOptions;
+		bool												bHaveDecoder;
+		int32												MaxDecodeBufferSize;
+
+		FIntPoint											MaxDecodeDim;
+
+	public:
+		static FSystemConfiguration							SystemConfig;
+
+#ifdef ELECTRA_ENABLE_SWDECODE
+		static bool bDidCheckHWSupport;
+		static bool bIsHWSupported;
+#endif
+	};
+
+} // namespace
diff --git a/Engine/Plugins/Media/ElectraPlayer/Source/ElectraPlayerRuntime/Private/Runtime/Decoder/GStreamer/AudioDecoderAAC_GStreamer.cpp b/Engine/Plugins/Media/ElectraPlayer/Source/ElectraPlayerRuntime/Private/Runtime/Decoder/GStreamer/AudioDecoderAAC_GStreamer.cpp
new file mode 100644
--- /dev/null
+++ b/Engine/Plugins/Media/ElectraPlayer/Source/ElectraPlayerRuntime/Private/Runtime/Decoder/GStreamer/AudioDecoderAAC_GStreamer.cpp
@@ -0,0 +1,46 @@
+// Copyright Epic Games, Inc. All Rights Reserved.
+
+#include "CoreMinimal.h"
+#if PLATFORM_UNIX
+
+#include "PlayerCore.h"
+#include "StreamAccessUnitBuffer.h"
+#include "Decoder/AudioDecoderAAC.h"
+#include "Renderer/RendererBase.h"
+#include "Player/PlayerSessionServices.h"
+#include "Utilities/UtilsMPEGAudio.h"
+#include "Utilities/StringHelpers.h"
+#include "HAL/LowLevelMemTracker.h"
+#include "ElectraPlayerPrivate.h"
+
+namespace Electra
+{
+
+IAudioDecoderAAC::FSystemConfiguration::FSystemConfiguration()
+{
+}
+
+IAudioDecoderAAC::FInstanceConfiguration::FInstanceConfiguration()
+{
+}
+
+bool IAudioDecoderAAC::Startup(const IAudioDecoderAAC::FSystemConfiguration& InConfig)
+{
+	return false;
+}
+
+void IAudioDecoderAAC::Shutdown()
+{
+}
+
+IAudioDecoderAAC* IAudioDecoderAAC::Create()
+{
+	return nullptr;
+}
+
+} // namespace Electra
+
+
+#endif // PLATFORM_UNIX
+
+
diff --git a/Engine/Plugins/Media/ElectraPlayer/Source/ElectraPlayerRuntime/Private/Runtime/Decoder/GStreamer/VideoDecoderH264_GStreamer.cpp b/Engine/Plugins/Media/ElectraPlayer/Source/ElectraPlayerRuntime/Private/Runtime/Decoder/GStreamer/VideoDecoderH264_GStreamer.cpp
new file mode 100644
--- /dev/null
+++ b/Engine/Plugins/Media/ElectraPlayer/Source/ElectraPlayerRuntime/Private/Runtime/Decoder/GStreamer/VideoDecoderH264_GStreamer.cpp
@@ -0,0 +1,50 @@
+// Copyright Epic Games, Inc. All Rights Reserved.
+
+#include "CoreMinimal.h"
+#if PLATFORM_UNIX
+
+#include "PlayerCore.h"
+#include "StreamAccessUnitBuffer.h"
+#include "Decoder/VideoDecoderH264.h"
+#include "Renderer/RendererBase.h"
+#include "Player/PlayerSessionServices.h"
+#include "Utilities/StringHelpers.h"
+#include "HAL/LowLevelMemTracker.h"
+#include "ElectraPlayerPrivate.h"
+
+namespace Electra
+{
+
+
+bool IVideoDecoderH264::Startup(const IVideoDecoderH264::FSystemConfiguration& InConfig)
+{
+	return false;
+}
+
+void IVideoDecoderH264::Shutdown()
+{
+}
+
+bool IVideoDecoderH264::GetStreamDecodeCapability(FStreamDecodeCapability& OutResult, const FStreamDecodeCapability& InStreamParameter)
+{
+	return false;
+}
+
+IVideoDecoderH264::FSystemConfiguration::FSystemConfiguration()
+{
+}
+
+IVideoDecoderH264::FInstanceConfiguration::FInstanceConfiguration()
+{
+}
+
+IVideoDecoderH264* IVideoDecoderH264::Create()
+{
+	return nullptr;
+}
+
+
+} // namespace Electra
+
+
+#endif // PLATFORM_UNIX
diff --git a/Engine/Plugins/Media/ElectraPlayer/Source/ElectraPlayerRuntime/Private/Runtime/Decoder/VideoDecoderH264.h b/Engine/Plugins/Media/ElectraPlayer/Source/ElectraPlayerRuntime/Private/Runtime/Decoder/VideoDecoderH264.h
new file mode 100644
--- /dev/null
+++ b/Engine/Plugins/Media/ElectraPlayer/Source/ElectraPlayerRuntime/Private/Runtime/Decoder/VideoDecoderH264.h
@@ -0,0 +1,127 @@
+// Copyright Epic Games, Inc. All Rights Reserved.
+
+#pragma once
+
+#include "PlayerCore.h"
+#include "StreamAccessUnitBuffer.h"
+#include "Decoder/DecoderBase.h"
+#include "VideoDecoderResourceDelegate.h"
+
+class IOptionPointerValueContainer;
+
+namespace Electra
+{
+	class IPlayerSessionServices;
+
+
+	/**
+	 * H264 video decoder.
+	**/
+	class IVideoDecoderH264 : public IDecoderBase, public IAccessUnitBufferInterface, public IDecoderAUBufferDiags, public IDecoderReadyBufferDiags
+	{
+	public:
+		struct FSystemConfiguration
+		{
+			FSystemConfiguration();
+			struct FThreadConfig
+			{
+				FMediaRunnable::Param		Decoder;							//!< Decoder thread settings.
+				FMediaRunnable::Param		PassOn;							//!< Settings for thread passing decoded images to the renderer.
+			};
+			FThreadConfig					ThreadConfig;
+		};
+
+		struct FInstanceConfiguration
+		{
+			FInstanceConfiguration();
+			int32								MaxFrameWidth;				//!< Maximum width of any image to be decoded
+			int32								MaxFrameHeight;			//!< Maximum height of any image to be decoded
+			int32								ProfileIdc;				//!< IDC profile (baseline (66), main (77) or high (100))
+			int32								LevelIdc;					//!< IDC profile level (eg, 30, 31, 51)
+			uint32								MaxDecodedFrames;			//!<
+			FSystemConfiguration::FThreadConfig	ThreadConfig;					//!< Thread configuration (defaults to values set in SystemConfiguration)
+
+			FParamDict							AdditionalOptions;
+		};
+
+		static bool Startup(const FSystemConfiguration& InConfig);
+		static void Shutdown();
+
+		struct FStreamDecodeCapability
+		{
+			enum class ESupported
+			{
+				NotSupported,
+				SoftwareOnly,
+				HardwareOnly,
+				HardAndSoftware
+			};
+			ESupported	DecoderSupportType = ESupported::NotSupported;
+			int32		Width = 0;
+			int32		Height = 0;
+			int32		Profile = 0;
+			int32		Level = 0;
+			double		FPS = 0.0;
+			FParamDict	AdditionalOptions;
+		};
+		/*
+			Queries decoder support/capability for a stream with given properties.
+			Can be called after Startup() but should be called only shortly before playback to ensure all relevant
+			libraries are initialized.
+		*/
+		static bool GetStreamDecodeCapability(FStreamDecodeCapability& OutResult, const FStreamDecodeCapability& InStreamParameter);
+
+		static IVideoDecoderH264* Create();
+
+		virtual ~IVideoDecoderH264() = default;
+
+		virtual void SetPlayerSessionServices(IPlayerSessionServices* SessionServices) = 0;
+
+		virtual void Open(const FInstanceConfiguration& InConfig) = 0;
+		virtual void Close() = 0;
+
+
+		virtual void SetMaximumDecodeCapability(int32 MaxWidth, int32 MaxHeight, int32 MaxProfile, int32 MaxProfileLevel, const FParamDict& AdditionalOptions) = 0;
+
+		virtual void SetResourceDelegate(const TSharedPtr<IVideoDecoderResourceDelegate, ESPMode::ThreadSafe>& ResourceDelegate) = 0;
+
+		//-------------------------------------------------------------------------
+		// Methods from IDecoderBase
+		//
+		virtual void SetRenderer(TSharedPtr<IMediaRenderer, ESPMode::ThreadSafe> InRenderer) = 0;
+
+
+		//-------------------------------------------------------------------------
+		// Methods from IAccessUnitBufferInterface
+		//
+		//! Attempts to push an access unit to the decoder. Ownership of the access unit is transferred if the push is successful.
+		virtual EAUpushResult AUdataPushAU(FAccessUnit* AccessUnit) = 0;
+		//! Notifies the decoder that there will be no further access units.
+		virtual void AUdataPushEOD() = 0;
+		//! Instructs the decoder to flush all pending input and all already decoded output.
+		virtual void AUdataFlushEverything() = 0;
+
+
+		//-------------------------------------------------------------------------
+		// Methods from IDecoderAUBufferDiags
+		//
+		//! Registers an AU input buffer listener.
+		virtual void SetAUInputBufferListener(IAccessUnitBufferListener* Listener) = 0;
+
+
+		//-------------------------------------------------------------------------
+		// Methods from IDecoderReadyBufferDiags
+		//
+		virtual void SetReadyBufferListener(IDecoderOutputBufferListener* Listener) = 0;
+
+		//-------------------------------------------------------------------------
+		// Platform specifics
+		//
+#if PLATFORM_ANDROID
+		virtual void Android_UpdateSurface(const TSharedPtr<IOptionPointerValueContainer>& Surface) = 0;
+#endif
+	};
+
+
+} // namespace Electra
+
diff --git a/Engine/Plugins/Media/ElectraPlayer/Source/ElectraPlayerRuntime/Private/Runtime/Decoder/Windows/VideoDecoderH264_PC.cpp b/Engine/Plugins/Media/ElectraPlayer/Source/ElectraPlayerRuntime/Private/Runtime/Decoder/Windows/VideoDecoderH264_PC.cpp
new file mode 100644
--- /dev/null
+++ b/Engine/Plugins/Media/ElectraPlayer/Source/ElectraPlayerRuntime/Private/Runtime/Decoder/Windows/VideoDecoderH264_PC.cpp
@@ -0,0 +1,877 @@
+// Copyright Epic Games, Inc. All Rights Reserved.
+
+#include "PlayerCore.h"
+#include "ElectraPlayerPrivate.h"
+#include "ElectraPlayerPrivate_Platform.h"
+#include "StreamAccessUnitBuffer.h"
+#include "Decoder/VideoDecoderH264.h"
+#include "Renderer/RendererBase.h"
+#include "Renderer/RendererVideo.h"
+#include "Player/PlayerSessionServices.h"
+#include "Utilities/StringHelpers.h"
+
+#include "Decoder/DX/DecoderErrors_DX.h"
+#include "MediaVideoDecoderOutputPC.h"
+
+/*********************************************************************************************************************/
+/*********************************************************************************************************************/
+/*********************************************************************************************************************/
+
+#include "Windows/AllowWindowsPlatformTypes.h"
+#include "Windows/WindowsHWrapper.h"
+#include "HAL/LowLevelMemTracker.h"
+
+THIRD_PARTY_INCLUDES_START
+#include "mftransform.h"
+#include "mfapi.h"
+#include "mferror.h"
+#include "mfidl.h"
+THIRD_PARTY_INCLUDES_END
+
+#include "Windows/HideWindowsPlatformTypes.h"
+
+#include "Decoder/DX/VideoDecoderH264_DX.h"
+
+
+#define ELECTRA_USE_IMF2DBUFFER	1		// set to 1 to use IMF2DBuffer interface rather then plain media buffer when retrieving HW decoded data for CPU use
+
+
+const D3DFORMAT DX9_NV12_FORMAT = (D3DFORMAT)MAKEFOURCC('N', 'V', '1', '2');
+
+#if (WINVER < _WIN32_WINNT_WIN8)
+const GUID MF_SA_D3D11_AWARE = { 0x206b4fc8, 0xfcf9, 0x4c51, { 0xaf, 0xe3, 0x97, 0x64, 0x36, 0x9e, 0x33, 0xa0 } };
+#endif
+
+/*********************************************************************************************************************/
+/*********************************************************************************************************************/
+/*********************************************************************************************************************/
+
+#ifdef ELECTRA_ENABLE_SWDECODE
+static TAutoConsoleVariable<int32> CVarElectraPCUseSoftwareDecoding(
+	TEXT("Electra.PC.UseSoftwareDecoding"),
+	0,
+	TEXT("Use software decoding on PC even if hardware decoding is supported.\n")
+	TEXT(" 0: use hardware decoding if supported (default); 1: use software decoding."),
+	ECVF_Default);
+#endif
+
+/*********************************************************************************************************************/
+/*********************************************************************************************************************/
+/*********************************************************************************************************************/
+
+class FElectraPlayerVideoDecoderOutputPC : public FVideoDecoderOutputPC
+{
+public:
+	FElectraPlayerVideoDecoderOutputPC()
+		: OutputType(EOutputType::Unknown)
+		, Stride(0)
+		, SampleDim(0,0)
+	{
+	}
+
+	// Hardware decode to buffer (Win7/DX12)
+	void InitializeWithBuffer(const void* InBuffer, uint32 InSize, uint32 InStride, FIntPoint Dim, Electra::FParamDict* InParamDict);
+
+	// Hardware decode to shared DX11 texture (Win8+)
+	void InitializeWithSharedTexture(const TRefCountPtr<ID3D11Device>& InD3D11Device, const TRefCountPtr<ID3D11DeviceContext> InDeviceContext, const TRefCountPtr<IMFSample>& MFSample, const FIntPoint& OutputDim, Electra::FParamDict* InParamDict);
+
+	// Software decode (into texture if DX11 device specified - available only Win8+)
+	void PreInitForDecode(FIntPoint OutputDim, const TFunction<void(int32 /*ApiReturnValue*/, const FString& /*Message*/, uint16 /*Code*/, UEMediaError /*Error*/)>& PostError);
+	void ProcessDecodeOutput(FIntPoint OutputDim, Electra::FParamDict* InParamDict);
+
+	void SetOwner(const TSharedPtr<IDecoderOutputOwner, ESPMode::ThreadSafe>& InOwningRenderer) override
+	{
+		OwningRenderer = InOwningRenderer;
+	}
+
+	void ShutdownPoolable() override;
+
+	virtual EOutputType GetOutputType() const override
+	{
+		return OutputType;
+	}
+
+	virtual TRefCountPtr<IMFSample> GetMFSample() const override
+	{
+		check(OutputType == EOutputType::SoftwareWin8Plus || OutputType == EOutputType::SoftwareWin7);
+		return MFSample;
+	}
+
+	virtual const TArray<uint8>& GetBuffer() const override
+	{
+		check(OutputType == EOutputType::SoftwareWin8Plus || OutputType == EOutputType::SoftwareWin7 || OutputType == EOutputType::HardwareDX9_DX12);
+		return Buffer;
+	}
+
+	virtual uint32 GetStride() const override
+	{
+		return Stride;
+	}
+
+	virtual TRefCountPtr<ID3D11Texture2D> GetTexture() const override
+	{
+		return Texture;
+	}
+
+	virtual TRefCountPtr<ID3D11Device> GetDevice() const override
+	{
+		check(OutputType == EOutputType::SoftwareWin8Plus || OutputType == EOutputType::HardwareWin8Plus);
+		return D3D11Device;
+	}
+
+	virtual FIntPoint GetDim() const override
+	{
+		return SampleDim;
+	}
+
+private:
+	// Decoder output type
+	EOutputType OutputType;
+
+	// Output texture (with device that created it) for SW decoder output Win8+
+	TRefCountPtr<ID3D11Texture2D> Texture;
+	TRefCountPtr<ID3D11Texture2D> SharedTexture;
+	TRefCountPtr<ID3D11Device> D3D11Device;
+
+	// CPU-side buffer to 
+	TArray<uint8> Buffer;
+	uint32 Stride;
+
+	// WMF sample (owned by this class if SW decoder is used)
+	TRefCountPtr<IMFSample> MFSample;
+
+	// Dimension of any internally allocated buffer - stored explicitly to cover various special cases for DX
+	FIntPoint SampleDim;
+
+	// We hold a weak reference to the video renderer. During destruction the video renderer could be destroyed while samples are still out there..
+	TWeakPtr<IDecoderOutputOwner, ESPMode::ThreadSafe> OwningRenderer;
+};
+
+
+void FElectraPlayerVideoDecoderOutputPC::PreInitForDecode(FIntPoint OutputDim, const TFunction<void(int32 /*ApiReturnValue*/, const FString& /*Message*/, uint16 /*Code*/, UEMediaError /*Error*/ )>& PostError)
+{
+	FIntPoint Dim;
+	Dim.X = OutputDim.X;
+	Dim.Y = OutputDim.Y * 3 / 2;	// adjust height to encompass Y and UV planes
+
+	bool bNeedNew = !MFSample.IsValid() || SampleDim != Dim;
+
+	if (bNeedNew)
+	{
+		HRESULT Result;
+		MFSample = nullptr;
+		Texture = nullptr;
+		TRefCountPtr<IMFMediaBuffer> MediaBuffer;
+
+		SampleDim = Dim;
+
+		if (Electra::IsWindows8Plus() && (Electra::FDXDeviceInfo::s_DXDeviceInfo->DxVersion != Electra::FDXDeviceInfo::ED3DVersion::Version12Win10))
+		{
+			// Software decode into DX11 texture
+
+			// We SW decode right into a DX11 texture (via the media buffer created from it below)
+			/*
+				As we cannot use RHI at this point, we cannot create a texture RHI can readily use,
+				but can just generate a plain D3D one. We cannot easily use a delegate to ask UE to do
+				it for use as we want to keep the "how" of the decoder hidden behind the abstract interfaces
+				so we can also encapsulate non-UE-aware systems easily.
+
+				So: we will later need to copy this once. (TODO: see if the render team can provide a way RHI could use
+				a native texture like this directly)
+			*/
+			D3D11_TEXTURE2D_DESC TextureDesc;
+			TextureDesc.Width = Dim.X;
+			TextureDesc.Height = Dim.Y;
+			TextureDesc.MipLevels = 1;
+			TextureDesc.ArraySize = 1;
+			TextureDesc.Format = DXGI_FORMAT_R8_UINT;
+			TextureDesc.SampleDesc.Count = 1;
+			TextureDesc.SampleDesc.Quality = 0;
+			TextureDesc.Usage = D3D11_USAGE_DEFAULT;
+			TextureDesc.BindFlags = 0;
+			TextureDesc.CPUAccessFlags = 0;
+			TextureDesc.MiscFlags = 0;
+
+			if (FAILED(Result = Electra::FDXDeviceInfo::s_DXDeviceInfo->RenderingDx11Device->CreateTexture2D(&TextureDesc, nullptr, Texture.GetInitReference())))
+			{
+				UE_LOG(LogElectraPlayer, Error, TEXT("CreateTexture2D() failed with 0x%X %s"), Result, *GetComErrorDescription(Result));
+				PostError(Result, "Failed to create software decode output texture", ERRCODE_INTERNAL_COULD_NOT_CREATE_OUTPUTBUFFER, UEMEDIA_ERROR_OK);
+				return;
+			}
+			if (FAILED(Result = MFCreateDXGISurfaceBuffer(__uuidof(ID3D11Texture2D), Texture, 0, false, MediaBuffer.GetInitReference())))
+			{
+				UE_LOG(LogElectraPlayer, Error, TEXT("MFCreateDXGISurfaceBuffer() failed with 0x%X %s"), Result, *GetComErrorDescription(Result));
+				PostError(Result, "Failed to create software decode output texture", ERRCODE_INTERNAL_COULD_NOT_CREATE_OUTPUTBUFFER, UEMEDIA_ERROR_OK);
+				return;
+			}
+
+			OutputType = EOutputType::SoftwareWin8Plus;
+		}
+		else // Software decode into CPU buffer
+		{
+			// SW decode results are just delivered in a simple CPU-side buffer. Create the decoder side version of this...
+			if (FAILED(Result = MFCreateMemoryBuffer(Dim.X * Dim.Y * sizeof(uint8), MediaBuffer.GetInitReference())))
+			{
+				UE_LOG(LogElectraPlayer, Error, TEXT("MFCreateMemoryBuffer() failed with 0x%X %s"), Result, *GetComErrorDescription(Result));
+				PostError(Result, "Failed to create software decode output buffer", ERRCODE_INTERNAL_COULD_NOT_CREATE_OUTPUTBUFFER, UEMEDIA_ERROR_OK);
+				return;
+			}
+			OutputType = EOutputType::SoftwareWin7;
+		}
+
+		if (SUCCEEDED(Result = MFCreateSample(MFSample.GetInitReference())))
+		{
+			if (FAILED(Result = MFSample->AddBuffer(MediaBuffer)))
+			{
+				UE_LOG(LogElectraPlayer, Error, TEXT("IMFSample::AddBuffer() failed with 0x%X %s"), Result, *GetComErrorDescription(Result));
+				PostError(Result, "Failed to add output buffer to software decoder sample", ERRCODE_INTERNAL_COULD_NOT_ADD_OUTPUT_BUFFER_TO_SAMPLE, UEMEDIA_ERROR_OK);
+				return;
+			}
+		}
+		else
+		{
+			UE_LOG(LogElectraPlayer, Error, TEXT("MFCreateSample() failed with 0x%X %s"), Result, *GetComErrorDescription(Result));
+			PostError(Result, "Failed to create software decode output sample", ERRCODE_INTERNAL_COULD_NOT_CREATE_OUTPUT_SAMPLE, UEMEDIA_ERROR_OK);
+			return;
+		}
+	}
+}
+
+
+void FElectraPlayerVideoDecoderOutputPC::ProcessDecodeOutput(FIntPoint OutputDim, Electra::FParamDict* InParamDict)
+{
+	check(OutputType == EOutputType::SoftwareWin7 || OutputType == EOutputType::SoftwareWin8Plus);
+
+	FVideoDecoderOutput::Initialize(InParamDict);
+
+	SampleDim.X = OutputDim.X;
+	SampleDim.Y = OutputDim.Y * 3 / 2;
+
+	Stride = SampleDim.X * sizeof(uint8);
+
+	if (OutputType == EOutputType::SoftwareWin7) // Win7 & DX12 (SW)
+	{
+		// Retrieve frame data and store it in Buffer for rendering later
+		TRefCountPtr<IMFMediaBuffer> MediaBuffer;
+		if (MFSample->GetBufferByIndex(0, MediaBuffer.GetInitReference()) != S_OK)
+		{
+			return;
+		}
+		DWORD BufferSize = 0;
+		if (MediaBuffer->GetCurrentLength(&BufferSize) != S_OK)
+		{
+			return;
+		}
+		check(SampleDim.X * SampleDim.Y <= (int32)BufferSize);
+
+		// MediaBuffer reports incorrect buffer sizes (too large) for some resolutions: we use our computed values!
+		BufferSize = SampleDim.X * SampleDim.Y;
+
+		uint8* Data = nullptr;
+		Buffer.Reset(BufferSize);
+		if (MediaBuffer->Lock(&Data, NULL, NULL) == S_OK)
+		{
+			Buffer.Append(Data, BufferSize);
+			MediaBuffer->Unlock();
+		}
+
+		// The output IMFSample needs to be released (and recreated for the next frame) for unknown reason.
+		// If not destroyed decoder throws an unknown error later on.
+		MFSample = nullptr;
+	}
+	else
+	{
+		// SW decode into texture need Win8+ (but no additional processing)
+		check(Electra::IsWindows8Plus());
+	}
+}
+
+
+void FElectraPlayerVideoDecoderOutputPC::InitializeWithBuffer(const void* InBuffer, uint32 InSize, uint32 InStride, FIntPoint Dim, Electra::FParamDict* InParamDict)
+{
+	FVideoDecoderOutput::Initialize(InParamDict);
+
+	OutputType = EOutputType::HardwareDX9_DX12;
+
+	Buffer.Reset(InSize);
+	Buffer.Append((uint8*)InBuffer, InSize);
+
+	SampleDim = Dim;
+	Stride = InStride;
+}
+
+
+void FElectraPlayerVideoDecoderOutputPC::InitializeWithSharedTexture(const TRefCountPtr<ID3D11Device>& InD3D11Device, const TRefCountPtr<ID3D11DeviceContext> InDeviceContext, const TRefCountPtr<IMFSample>& InMFSample, const FIntPoint& OutputDim, Electra::FParamDict* InParamDict)
+{
+	FVideoDecoderOutput::Initialize(InParamDict);
+
+	OutputType = EOutputType::HardwareWin8Plus;
+
+	bool bNeedsNew = !Texture.IsValid() || (SampleDim.X != OutputDim.X || SampleDim.Y != OutputDim.Y);
+
+	if (bNeedsNew)
+	{
+		SampleDim = OutputDim;
+
+		// Make a texture we pass on as output
+		D3D11_TEXTURE2D_DESC TextureDesc;
+		TextureDesc.Width = SampleDim.X;
+		TextureDesc.Height = SampleDim.Y;
+		TextureDesc.MipLevels = 1;
+		TextureDesc.ArraySize = 1;
+		TextureDesc.Format = DXGI_FORMAT_NV12;
+		TextureDesc.SampleDesc.Count = 1;
+		TextureDesc.SampleDesc.Quality = 0;
+		TextureDesc.Usage = D3D11_USAGE_DEFAULT;
+		TextureDesc.BindFlags = 0;
+		TextureDesc.CPUAccessFlags = 0;
+		TextureDesc.MiscFlags = D3D11_RESOURCE_MISC_SHARED_KEYEDMUTEX;
+
+		if (InD3D11Device->CreateTexture2D(&TextureDesc, nullptr, Texture.GetInitReference()) != S_OK)
+		{
+			return;
+		}
+
+		D3D11Device = InD3D11Device;
+	}
+
+	// If we got a texture, copy the data from the decoder into it...
+	if (Texture)
+	{
+		// Get output texture from decoder...
+		TRefCountPtr<IMFMediaBuffer> MediaBuffer;
+		if (InMFSample->GetBufferByIndex(0, MediaBuffer.GetInitReference()) != S_OK)
+		{
+			return;
+		}
+		TRefCountPtr<IMFDXGIBuffer> DXGIBuffer;
+		if (MediaBuffer->QueryInterface(__uuidof(IMFDXGIBuffer), (void**)DXGIBuffer.GetInitReference()) != S_OK)
+		{
+			return;
+		}
+		TRefCountPtr<ID3D11Texture2D> DecoderTexture;
+		if (DXGIBuffer->GetResource(IID_PPV_ARGS(DecoderTexture.GetInitReference())) != S_OK)
+		{
+			return;
+		}
+		uint32 ViewIndex = 0;
+		if (DXGIBuffer->GetSubresourceIndex(&ViewIndex) != S_OK)
+		{
+			return;
+		}
+
+		// Source data may be larger than desired output, but crop area will be aligned such that we can always copy from 0,0
+		D3D11_BOX SrcBox;
+		SrcBox.left = 0;
+		SrcBox.top = 0;
+		SrcBox.front = 0;
+		SrcBox.right = OutputDim.X;
+		SrcBox.bottom = OutputDim.Y;
+		SrcBox.back = 1;
+
+		TRefCountPtr<IDXGIKeyedMutex> KeyedMutex;
+		HRESULT res = Texture->QueryInterface(_uuidof(IDXGIKeyedMutex), (void**)&KeyedMutex);
+
+		if (KeyedMutex)
+		{
+			// No wait on acquire since sample is new and key is 0.
+			res = KeyedMutex->AcquireSync(0, 0);
+			if (res == S_OK)
+			{
+				// Copy texture using the decoder DX11 device... (and apply any cropping - see above note)
+				InDeviceContext->CopySubresourceRegion(Texture, 0, 0, 0, 0, DecoderTexture, ViewIndex, &SrcBox);
+
+				// Mark texture as updated with key of 1
+				// Sample will be read in Convert() method of texture sample
+				KeyedMutex->ReleaseSync(1);
+			}
+		}
+
+		// Make sure texture is updated before giving access to the sample in the rendering thread.
+		InDeviceContext->Flush();
+	}
+}
+
+void FElectraPlayerVideoDecoderOutputPC::ShutdownPoolable()
+{
+	TSharedPtr<IDecoderOutputOwner, ESPMode::ThreadSafe> lockedVideoRenderer = OwningRenderer.Pin();
+	if (lockedVideoRenderer.IsValid())
+	{
+		lockedVideoRenderer->SampleReleasedToPool(GetDuration());
+	}
+
+	if (OutputType != EOutputType::HardwareWin8Plus)
+	{
+		return;
+	}
+
+	// Correctly release the keyed mutex when the sample is returned to the pool
+	TRefCountPtr<IDXGIResource> OtherResource(nullptr);
+	if (Texture)
+	{
+		Texture->QueryInterface(__uuidof(IDXGIResource), (void**)&OtherResource);
+	}
+
+	if (OtherResource)
+	{
+		HANDLE SharedHandle = nullptr;
+		if (OtherResource->GetSharedHandle(&SharedHandle) == S_OK)
+		{
+			TRefCountPtr<ID3D11Resource> SharedResource;
+			D3D11Device->OpenSharedResource(SharedHandle, __uuidof(ID3D11Texture2D), (void**)&SharedResource);
+			if (SharedResource)
+			{
+				TRefCountPtr<IDXGIKeyedMutex> KeyedMutex;
+				OtherResource->QueryInterface(_uuidof(IDXGIKeyedMutex), (void**)&KeyedMutex);
+
+				if (KeyedMutex)
+				{
+					// Reset keyed mutex
+					if (KeyedMutex->AcquireSync(1, 0) == S_OK)
+					{
+						// Texture was never read
+						KeyedMutex->ReleaseSync(0);
+					}
+					else if (KeyedMutex->AcquireSync(2, 0) == S_OK)
+					{
+						// Texture was read at least once
+						KeyedMutex->ReleaseSync(0);
+					}
+				}
+			}
+		}
+	}
+
+}
+
+// -----------------------------------------------------------------------------------------------------------------------------
+
+FVideoDecoderOutput* FElectraPlayerPlatformVideoDecoderOutputFactory::Create()
+{
+	return new FElectraPlayerVideoDecoderOutputPC();
+}
+
+
+namespace Electra
+{
+
+/*********************************************************************************************************************/
+/*********************************************************************************************************************/
+/*********************************************************************************************************************/
+
+class FVideoDecoderH264_PC : public FVideoDecoderH264
+{
+public:
+private:
+	virtual bool InternalDecoderCreate() override;
+	virtual bool CreateDecoderOutputBuffer() override;
+	virtual void PreInitDecodeOutputForSW(const FIntPoint& Dim);
+	virtual bool SetupDecodeOutputData(const FIntPoint& ImageDim, const TRefCountPtr<IMFSample>& DecodedOutputSample, FParamDict* OutputBufferSampleProperties) override;
+
+	bool CopyTexture(const TRefCountPtr<IMFSample>& DecodedSample, Electra::FParamDict* ParamDict, FIntPoint OutputDim);
+};
+
+static const int32 MaxHWDecodeW = 1920;
+static const int32 MaxHWDecodeH = 1088;
+
+/***************************************************************************************************************************************************/
+/***************************************************************************************************************************************************/
+/***************************************************************************************************************************************************/
+
+IVideoDecoderH264* IVideoDecoderH264::Create()
+{
+	return new FVideoDecoderH264_PC();
+}
+
+//-----------------------------------------------------------------------------
+/**
+ *	Queries decoder support/capability for a stream with given properties.
+ *	Can be called after Startup() but should be called only shortly before playback to ensure all relevant
+ *	libraries are initialized.
+ */
+bool FVideoDecoderH264::GetStreamDecodeCapability(FStreamDecodeCapability& OutResult, const FStreamDecodeCapability& InStreamParameter)
+{
+	// Set the output same as the input. Modifications are made below.
+	OutResult = InStreamParameter;
+	OutResult.DecoderSupportType = FStreamDecodeCapability::ESupported::NotSupported;
+
+#ifdef ELECTRA_ENABLE_SWDECODE
+	// Software forced through CVar?
+	if (CVarElectraPCUseSoftwareDecoding.GetValueOnAnyThread() > 0)
+	{
+		bIsHWSupported = false;
+	}
+	else
+#endif
+	{
+		// Did we check the actual decoder for being HW supported yet?
+		if (!bDidCheckHWSupport)
+		{
+			// Check for HW support by creating and configuring a decoder.
+			FVideoDecoderH264* TestDecoder = new FVideoDecoderH264_PC;
+			// Set width and height in the decoder configuration. These are used in creating the decoder.
+			TestDecoder->Config.MaxFrameWidth = InStreamParameter.Width ? InStreamParameter.Width : MaxHWDecodeW;
+			TestDecoder->Config.MaxFrameHeight = InStreamParameter.Height ? InStreamParameter.Height : MaxHWDecodeH;
+			TestDecoder->TestHardwareDecoding();
+			bIsHWSupported = TestDecoder->bIsHardwareAccelerated;
+			delete TestDecoder;
+			bDidCheckHWSupport = true;
+		}
+	}
+	if (bIsHWSupported)
+	{
+		// FIXME: Should we check the resolution or profile level here?
+		//        Right now we won't be going higher than 1080 so we should be fine.
+		OutResult.DecoderSupportType = FStreamDecodeCapability::ESupported::HardwareOnly;
+	}
+	else
+	{
+		// Global query?
+		if (InStreamParameter.Width <= 0 && InStreamParameter.Height <= 0)
+		{
+			// For the sake of argument we limit SW decoding to 720p@60.
+			OutResult.Width = 1280;
+			OutResult.Height = 720;
+			OutResult.Profile = 100;
+			OutResult.Level = 42;
+			OutResult.FPS = 60.0;
+			OutResult.DecoderSupportType = FStreamDecodeCapability::ESupported::SoftwareOnly;
+		}
+		else
+		{
+			// Calculate via number of macroblocks per second.
+			int32 maxMB = (1280 / 16) * (720 / 16) * 60;
+			int32 numMB = (int32) (((InStreamParameter.Width + 15) / 16) * ((InStreamParameter.Height + 15) / 16) * (InStreamParameter.FPS > 0.0 ? InStreamParameter.FPS : 30.0));
+			OutResult.DecoderSupportType = numMB <= maxMB ? FStreamDecodeCapability::ESupported::SoftwareOnly : FStreamDecodeCapability::ESupported::NotSupported;
+		}
+	}
+
+	return true;
+}
+
+//-----------------------------------------------------------------------------
+/**
+ * Create a decoder instance.
+ *
+ * @return true if successful, false on error
+ */
+bool FVideoDecoderH264_PC::InternalDecoderCreate()
+{
+	TRefCountPtr<IMFAttributes>	Attributes;
+	TRefCountPtr<IMFTransform>	Decoder;
+	HRESULT					res;
+
+	if (Electra::IsWindows8Plus())
+	{
+		// Check if there is any reason for a "device lost" - if not we know all is stil well; otherwise we bail without creating a decoder
+		if (Electra::FDXDeviceInfo::s_DXDeviceInfo->DxDevice && (res = Electra::FDXDeviceInfo::s_DXDeviceInfo->DxDevice->GetDeviceRemovedReason()) != S_OK)
+		{
+			PostError(res, "Device lost detected.", ERRCODE_INTERNAL_COULD_NOT_SET_OUTPUT_DESIRED_MEDIA_TYPE);
+			return false;
+		}
+	}
+
+	VERIFY_HR(CoCreateInstance(MFTmsH264Decoder, nullptr, CLSCTX_INPROC_SERVER, IID_IMFTransform, reinterpret_cast<void**>(&Decoder)), "CoCreateInstance failed", ERRCODE_INTERNAL_COULD_NOT_CREATE_DECODER);
+	VERIFY_HR(Decoder->GetAttributes(Attributes.GetInitReference()), "Failed to get video decoder transform attributes", ERRCODE_INTERNAL_COULD_GET_TRANSFORM_ATTRIBUTES);
+
+	// Force SW decoding?
+#ifdef ELECTRA_ENABLE_SWDECODE
+	if (CVarElectraPCUseSoftwareDecoding.GetValueOnAnyThread() > 0)
+	{
+		FallbackToSwDecoding(FString::Printf(TEXT("Windows %s"), *FWindowsPlatformMisc::GetOSVersion()));
+	}
+	else
+#endif
+	{
+		// Check if the transform is D3D aware
+		if (Electra::IsWindows8Plus())
+		{
+			if (!Electra::FDXDeviceInfo::s_DXDeviceInfo->DxDeviceManager)
+			{
+				// Win8+ with no DX11. No need to check for any DX11 awareness...
+				return true;
+			}
+
+			uint32 IsDX11Aware = 0;
+			res = Attributes->GetUINT32(MF_SA_D3D11_AWARE, &IsDX11Aware);
+			if (FAILED(res))
+			{
+				FallbackToSwDecoding(TEXT("Failed to get MF_SA_D3D11_AWARE"));
+			}
+			else if (IsDX11Aware == 0)
+			{
+				FallbackToSwDecoding(TEXT("Not MF_SA_D3D11_AWARE"));
+			}
+			else if (FAILED(res = Decoder->ProcessMessage(MFT_MESSAGE_SET_D3D_MANAGER, reinterpret_cast<ULONG_PTR>(Electra::FDXDeviceInfo::s_DXDeviceInfo->DxDeviceManager.GetReference()))))
+			{
+				FallbackToSwDecoding(FString::Printf(TEXT("Failed to set MFT_MESSAGE_SET_D3D_MANAGER: 0x%X %s"), res, *GetComErrorDescription(res)));
+			}
+		}
+		else // Windows 7
+		{
+#ifdef ELECTRA_HAVE_DX9
+			if (!Electra::FDXDeviceInfo::s_DXDeviceInfo->Dx9Device || !Electra::FDXDeviceInfo::s_DXDeviceInfo->Dx9DeviceManager)
+			{
+				FallbackToSwDecoding(TEXT("Failed to create DirectX 9 device / device manager"));
+			}
+
+			uint32 IsD3DAware = 0;
+			HRESULT HRes = Attributes->GetUINT32(MF_SA_D3D_AWARE, &IsD3DAware);
+			if (FAILED(HRes))
+			{
+				FallbackToSwDecoding(TEXT("Failed to get MF_SA_D3D_AWARE"));
+			}
+			else if (IsD3DAware == 0)
+			{
+				FallbackToSwDecoding(TEXT("Not MF_SA_D3D_AWARE"));
+			}
+			else if (FAILED(HRes = Decoder->ProcessMessage(MFT_MESSAGE_SET_D3D_MANAGER, reinterpret_cast<ULONG_PTR>(Electra::FDXDeviceInfo::s_DXDeviceInfo->Dx9DeviceManager.GetReference()))))
+			{
+				FallbackToSwDecoding(FString::Printf(TEXT("Failed to set MFT_MESSAGE_SET_D3D_MANAGER: 0x%X %s"), HRes, *GetComErrorDescription(HRes)));
+			}
+#else
+			return false;
+#endif
+		}
+	}
+
+	// Sugar & spice
+	if (FAILED(res = Attributes->SetUINT32(CODECAPI_AVLowLatencyMode, 1)))
+	{
+		// Not an error. If it doesn't work it just doesn't.
+	}
+
+	/*
+	if (FAILED(res = Attributes->SetUINT32(AVDecVideoMaxCodedWidth, 1280 + 64)))
+		return false;
+	if (FAILED(res = Attributes->SetUINT32(AVDecVideoMaxCodedHeight, 720 + 64)))
+		return false;
+	*/
+
+	// Create successful, take on the decoder.
+	DecoderTransform = Decoder;
+
+	return true;
+}
+
+
+bool FVideoDecoderH264_PC::CreateDecoderOutputBuffer()
+{
+	HRESULT									res;
+	TUniquePtr<FDecoderOutputBuffer>		NewDecoderOutputBuffer(new FDecoderOutputBuffer());
+
+	VERIFY_HR(DecoderTransform->GetOutputStreamInfo(0, &NewDecoderOutputBuffer->mOutputStreamInfo), "Failed to get video decoder output stream information", ERRCODE_INTERNAL_COULD_NOT_GET_OUTPUT_STREAM_INFO);
+	// Do we need to provide the sample output buffer or does the decoder create it for us?
+	if ((NewDecoderOutputBuffer->mOutputStreamInfo.dwFlags & (MFT_OUTPUT_STREAM_PROVIDES_SAMPLES | MFT_OUTPUT_STREAM_CAN_PROVIDE_SAMPLES)) == 0)
+	{
+#if ELECTRA_ENABLE_SWDECODE
+		check(bIsHardwareAccelerated == false);
+		if (CurrentRenderOutputBuffer)
+		{
+			TSharedPtr<FElectraPlayerVideoDecoderOutputPC, ESPMode::ThreadSafe> DecoderOutput = CurrentRenderOutputBuffer->GetBufferProperties().GetValue("texture").GetSharedPointer<FElectraPlayerVideoDecoderOutputPC>();
+			check(DecoderOutput);
+			// Get the sample which was created earlier
+			NewDecoderOutputBuffer->mOutputBuffer.pSample = DecoderOutput->GetMFSample();
+			// Destination is a plain old pointer, we need to ref manually
+			if (NewDecoderOutputBuffer->mOutputBuffer.pSample)
+			{
+				NewDecoderOutputBuffer->mOutputBuffer.pSample->AddRef();
+			}
+			else
+			{
+				return false;
+			}
+		}
+		else
+#endif
+		{
+			return false;
+		}
+	}
+	CurrentDecoderOutputBuffer = MoveTemp(NewDecoderOutputBuffer);
+	return true;
+}
+
+
+void FVideoDecoderH264_PC::PreInitDecodeOutputForSW(const FIntPoint& Dim)
+{
+	TSharedPtr<FElectraPlayerVideoDecoderOutputPC, ESPMode::ThreadSafe> DecoderOutput = CurrentRenderOutputBuffer->GetBufferProperties().GetValue("texture").GetSharedPointer<FElectraPlayerVideoDecoderOutputPC>();
+	check(DecoderOutput);
+	DecoderOutput->PreInitForDecode(Dim, [this](int32 ApiReturnValue, const FString& Message, uint16 Code, UEMediaError Error) { PostError(ApiReturnValue, Message, Code, Error); });
+}
+
+
+bool FVideoDecoderH264_PC::SetupDecodeOutputData(const FIntPoint & ImageDim, const TRefCountPtr<IMFSample>& DecodedOutputSample, FParamDict* OutputBufferSampleProperties)
+{
+#ifdef ELECTRA_ENABLE_SWDECODE
+	if (!bIsHardwareAccelerated)
+	{
+		TSharedPtr<FElectraPlayerVideoDecoderOutputPC, ESPMode::ThreadSafe> DecoderOutput = CurrentRenderOutputBuffer->GetBufferProperties().GetValue("texture").GetSharedPointer<FElectraPlayerVideoDecoderOutputPC>();
+		check(DecoderOutput);
+		DecoderOutput->ProcessDecodeOutput(ImageDim, OutputBufferSampleProperties);
+	}
+	else
+#endif
+	{
+		bool bCopyResult = CopyTexture(DecodedOutputSample, OutputBufferSampleProperties, ImageDim);
+		if (!bCopyResult)
+		{
+			// Failed for some reason. Let's best not render this then, but put in the buffers nonetheless to maintain the normal data flow.
+			OutputBufferSampleProperties->Set("is_dummy", FVariantValue(true));
+		}
+	}
+	return true;
+}
+
+
+bool FVideoDecoderH264_PC::CopyTexture(const TRefCountPtr<IMFSample>& Sample, Electra::FParamDict* ParamDict, FIntPoint OutputDim)
+{
+#if !UE_SERVER
+	TSharedPtr<FElectraPlayerVideoDecoderOutputPC, ESPMode::ThreadSafe> DecoderOutput = CurrentRenderOutputBuffer->GetBufferProperties().GetValue("texture").GetSharedPointer<FElectraPlayerVideoDecoderOutputPC>();
+	check(DecoderOutput);
+
+	DWORD BuffersNum = 0;
+	CHECK_HR(Sample->GetBufferCount(&BuffersNum));
+
+	if (BuffersNum != 1)
+	{
+		return false;
+	}
+
+	LONGLONG SampleTime = 0;
+	CHECK_HR(Sample->GetSampleTime(&SampleTime));
+	LONGLONG SampleDuration = 0;
+	CHECK_HR(Sample->GetSampleDuration(&SampleDuration));
+
+	TRefCountPtr<IMFMediaBuffer> Buffer;
+	CHECK_HR(Sample->GetBufferByIndex(0, Buffer.GetInitReference()));
+
+	if (FDXDeviceInfo::s_DXDeviceInfo->DxVersion == FDXDeviceInfo::ED3DVersion::Version11Win8)
+	{
+		check(FDXDeviceInfo::s_DXDeviceInfo->DxDevice);
+
+		DecoderOutput->InitializeWithSharedTexture(
+			FDXDeviceInfo::s_DXDeviceInfo->DxDevice,
+			FDXDeviceInfo::s_DXDeviceInfo->DxDeviceContext,
+			Sample,
+			OutputDim,
+			ParamDict);
+
+		if (!DecoderOutput->GetTexture())
+		{
+			UE_LOG(LogElectraPlayer, Error, TEXT("FVideoDecoderH264: InD3D11Device->CreateTexture2D() failed!"));
+			return false;
+		}
+	}
+	else if (FDXDeviceInfo::s_DXDeviceInfo->DxVersion == FDXDeviceInfo::ED3DVersion::Version12Win10)
+	{
+		// Gather some info from the DX11 texture in the buffer...
+		TRefCountPtr<IMFDXGIBuffer> DXGIBuffer;
+		CHECK_HR(Buffer->QueryInterface(__uuidof(IMFDXGIBuffer), (void**)DXGIBuffer.GetInitReference()));
+
+		TRefCountPtr<ID3D11Texture2D> Texture2D;
+		CHECK_HR(DXGIBuffer->GetResource(IID_PPV_ARGS(Texture2D.GetInitReference())));
+		D3D11_TEXTURE2D_DESC TextureDesc;
+		Texture2D->GetDesc(&TextureDesc);
+		if (TextureDesc.Format != DXGI_FORMAT_NV12)
+		{
+			LogMessage(IInfoLog::ELevel::Error, "FVideoDecoderH264::CopyTexture(): Decoded texture is not in NV12 format");
+			return false;
+		}
+
+		// Retrieve frame data and store it in Buffer for rendering later
+#if ELECTRA_USE_IMF2DBUFFER
+		TRefCountPtr<IMF2DBuffer> Buffer2D;
+		CHECK_HR(Buffer->QueryInterface(__uuidof(IMF2DBuffer), (void**)Buffer2D.GetInitReference()));
+
+		uint8* Data = nullptr;
+		LONG Pitch;
+		CHECK_HR(Buffer2D->Lock2D(&Data, &Pitch));
+
+		DecoderOutput->InitializeWithBuffer(Data, Pitch * (TextureDesc.Height * 3 / 2),
+			Pitch,															// Buffer stride
+			FIntPoint(TextureDesc.Width, TextureDesc.Height * 3 / 2),		// Buffer dimensions
+			ParamDict);
+
+		CHECK_HR(Buffer2D->Unlock2D());
+#else
+		DWORD BufferSize = 0;
+		CHECK_HR(Buffer->GetCurrentLength(&BufferSize));
+		// Not really trusting the size here. Seen crash reports where the returned buffer size is larger than what is contained,
+		// as if the texture was wider (a multiple of 256 actually). The crash then happens on the memcpy() trying to _read_ more
+		// data than was actually mapped to the buffer (access violation on exactly the start of the next page (4096 byte aligned)).
+		// https://docs.microsoft.com/en-us/windows/win32/api/dxgiformat/ne-dxgiformat-dxgi_format stipulates the format to be aligned
+		// to 'rowPitch' (whatever that is) and 'SysMemPitch' (https://docs.microsoft.com/en-us/windows/win32/api/d3d11/ns-d3d11-d3d11_subresource_data)
+		// which as far as I can tell is a property set when the texture was created. Which we have NO control over since the video decoder creates this
+		// texture for us. There does not seem to be any API to query this pitch at all.
+		// NV12 format as described here https://docs.microsoft.com/en-us/windows/win32/medfound/recommended-8-bit-yuv-formats-for-video-rendering
+		// does not suggest there to be any padding/alignment/pitch.
+		// What makes the crashes so weird is that they do not happen all the time right on the first frame. It takes a bit before wrong values are returned.
+		// For instance, for a 1152*656 texture 1259520 (1280*656) is returned when really it should be only 1133568,
+		// or for a 864*480 texture instead of 622080 bytes we are told there'd be 737280 bytes (1024*480).
+		// Since we are not using any "actual pitch" in the Initialize() call but the actual width of the texture we must only copy as much data anyway.
+		// Even if there was padding in the texture that we would be copying over into our buffer it would not be skipped on rendering later on account
+		// of us specifying the width for pitch.
+		// It seems prudent at this point not to trust the returned buffer size but to calculate it from the actual dimensions ourselves.
+		// NOTE: this value is probably dependent on the GPU and driver implementing the decoder and may therefore not be a general issue!
+		BufferSize = TextureDesc.Width * (TextureDesc.Height * 3 / 2);
+
+		uint8* Data = nullptr;
+		CHECK_HR(Buffer->Lock(&Data, NULL, NULL));
+
+		DecoderOutput->InitializeWithBuffer(Data, BufferSize,
+			TextureDesc.Width,												// Buffer stride
+			FIntPoint(TextureDesc.Width, TextureDesc.Height * 3 / 2),		// Buffer dimensions
+			ParamDict);
+
+		CHECK_HR(Buffer->Unlock());
+#endif
+
+		return true;
+	}
+#ifdef ELECTRA_HAVE_DX9
+	else if (FDXDeviceInfo::s_DXDeviceInfo->DxVersion == FDXDeviceInfo::ED3DVersion::Version9Win7)
+	{
+		TRefCountPtr<IDirect3DSurface9> Dx9DecoderSurface;
+		TRefCountPtr<IMFGetService> BufferService;
+		CHECK_HR(Buffer->QueryInterface(IID_PPV_ARGS(BufferService.GetInitReference())));
+		CHECK_HR(BufferService->GetService(MR_BUFFER_SERVICE, IID_PPV_ARGS(Dx9DecoderSurface.GetInitReference())));
+
+		D3DSURFACE_DESC Dx9SurfaceDesc;
+		CHECK_HR_DX9(Dx9DecoderSurface->GetDesc(&Dx9SurfaceDesc));
+		check(Dx9SurfaceDesc.Format == DX9_NV12_FORMAT);
+
+		if (Dx9SurfaceDesc.Format != DX9_NV12_FORMAT)
+		{
+			LogMessage(IInfoLog::ELevel::Error, "FVideoDecoderH264::CopyTexture(): Decoded DX9 surface is not in NV12 format");
+		}
+
+		// Read back DX9 surface data and pass it onto the texture sample
+		D3DLOCKED_RECT Dx9DecoderTexLockedRect;
+		CHECK_HR_DX9(Dx9DecoderSurface->LockRect(&Dx9DecoderTexLockedRect, nullptr, D3DLOCK_READONLY));
+		check(Dx9DecoderTexLockedRect.pBits && Dx9DecoderTexLockedRect.Pitch > 0);
+
+		// note: 3/2 scale is to account for NV12 data format (Y+UV section)
+		DecoderOutput->InitializeWithBuffer(
+			Dx9DecoderTexLockedRect.pBits,									// Buffer ptr
+			Dx9DecoderTexLockedRect.Pitch * Dx9SurfaceDesc.Height * 3 / 2,	// Buffer size
+			Dx9DecoderTexLockedRect.Pitch,									// Buffer stride
+			FIntPoint(Dx9DecoderTexLockedRect.Pitch, Dx9SurfaceDesc.Height * 3 / 2), // Buffer dimensions
+			ParamDict);
+
+		CHECK_HR_DX9(Dx9DecoderSurface->UnlockRect());
+	}
+#endif
+	else
+	{
+		LogMessage(IInfoLog::ELevel::Error, "FVideoDecoderH264::CopyTexture(): Unhandled D3D version");
+		return false;
+	}
+	return true;
+#else
+	return false;
+#endif
+}
+
+} // namespace Electra
+
diff --git a/Engine/Plugins/Media/ElectraPlayer/Source/ElectraPlayerRuntime/Private/Runtime/Demuxer/ParserISO14496-12.h b/Engine/Plugins/Media/ElectraPlayer/Source/ElectraPlayerRuntime/Private/Runtime/Demuxer/ParserISO14496-12.h
new file mode 100644
--- /dev/null
+++ b/Engine/Plugins/Media/ElectraPlayer/Source/ElectraPlayerRuntime/Private/Runtime/Demuxer/ParserISO14496-12.h
@@ -0,0 +1,237 @@
+// Copyright Epic Games, Inc. All Rights Reserved.
+
+#pragma once
+
+
+#include "PlayerCore.h"
+
+#include "PlayerTime.h"
+#include "OptionalValue.h"
+#include "StreamTypes.h"
+#include "ParameterDictionary.h"
+
+namespace Electra
+{
+	//
+	// Forward declarations
+	//
+	class IPlayerSessionServices;
+
+
+	/**
+	 * Interface for parsing an ISO/IEC 14496-12 file, commonly referred to as an "mp4" file.
+	 */
+	class IParserISO14496_12
+	{
+	public:
+		virtual ~IParserISO14496_12() = default;
+
+		/**
+		 * Interface for reading data from a source.
+		 */
+		class IReader
+		{
+		public:
+			virtual ~IReader() = default;
+			/**
+			 * Read n bytes of data into the provided buffer.
+			 *
+			 * Reading must return the number of bytes asked to get, if necessary by blocking.
+			 * If a read error prevents reading the number of bytes -1 must be returned.
+			 *
+			 * @param IntoBuffer Buffer into which to store the data bytes. If nullptr is passed the data must be skipped over.
+			 * @param NumBytesToRead The number of bytes to read. Must not read more bytes and no less than requested.
+			 * @return The number of bytes read or -1 on a read error.
+			 */
+			virtual int64 ReadData(void* IntoBuffer, int64 NumBytesToRead) = 0;
+
+			/**
+			 * Checks if the data source has reached the End Of File (EOF) and cannot provide any additional data.
+			 *
+			 * @return If EOF has been reached returns true, otherwise false.
+			 */
+			virtual bool HasReachedEOF() const = 0;
+
+			/**
+			 * Checks if reading of the file and therefor parsing has been aborted.
+			 *
+			 * @return true if reading/parsing has been aborted, false otherwise.
+			 */
+			virtual bool HasReadBeenAborted() const = 0;
+
+			/**
+			 * Returns the current read offset.
+			 *
+			 * The first read offset is not necessarily zero. It could be anywhere inside the source.
+			 *
+			 * @return The current byte offset in the source.
+			 */
+			virtual int64 GetCurrentOffset() const = 0;
+		};
+
+		/** Box type is a 32 bit value in an mp4 file. */
+		typedef uint32  FBoxType;
+
+		/**
+		 * Commonly referenced boxes
+		 */
+		 // This whole constexpr thing doesn't work...
+			 /*
+				 static inline constexpr FBoxType MakeBoxType(unsigned char v1, unsigned char v2, unsigned char v3, unsigned char v4)
+				 {
+					 return ((uint32)v1 << 24) | ((uint32)v2 << 16) | ((uint32)v3 << 8) | (uint32)v4;
+				 }
+			 */
+#define MAKE_BOX_ATOM(a,b,c,d) (IParserISO14496_12::FBoxType)((uint32)a << 24) | ((uint32)b << 16) | ((uint32)c << 8) | ((uint32)d)
+		static const FBoxType BoxType_ftyp = MAKE_BOX_ATOM('f', 't', 'y', 'p');
+		static const FBoxType BoxType_moov = MAKE_BOX_ATOM('m', 'o', 'o', 'v');
+		static const FBoxType BoxType_sidx = MAKE_BOX_ATOM('s', 'i', 'd', 'x');
+		static const FBoxType BoxType_moof = MAKE_BOX_ATOM('m', 'o', 'o', 'f');
+		static const FBoxType BoxType_mdat = MAKE_BOX_ATOM('m', 'd', 'a', 't');
+
+
+
+		/**
+		 * Interface for receiving parse notifications
+		 */
+		class IBoxCallback
+		{
+		public:
+			virtual ~IBoxCallback() = default;
+
+			/**
+			 * Choices for continued parsing.
+			 */
+			enum class EParseContinuation
+			{
+				/** Continue parsing */
+				Continue,
+				/** Stop parsing */
+				Stop
+			};
+
+			/**
+			 * Notifies the caller of the start of a new box.
+			 *
+			 * The caller may continue to parse this box or ask to stop parsing.
+			 *
+			 * @param Box Type of the box that will be parsed next.
+			 * @param BoxSizeInBytes Size of the box in bytes.
+			 * @param FileDataOffset Offset of the box within the file. This points to the first byte of the box length field.
+			 * @param BoxDataOffset Offset of the box' data. This points to the first payload byte of the box.
+			 * @return Whether or not to continue parsing thix box.
+			 */
+			virtual EParseContinuation OnFoundBox(FBoxType Box, int64 BoxSizeInBytes, int64 FileDataOffset, int64 BoxDataOffset) = 0;
+		};
+
+
+		static TSharedPtrTS<IParserISO14496_12> CreateParser();
+
+		static const FString 	OptionKeyParseFirstFreeBox;			// Option to not skip 'free' box but instead attempt to parse their content
+
+
+		/**
+		 * Parses the header boxes (all non-MDAT boxes).
+		 */
+		virtual UEMediaError ParseHeader(IReader* DataReader, IBoxCallback* BoxParseCallback, const FParamDict& Options, IPlayerSessionServices* PlayerSession) = 0;
+
+
+
+
+		/*******************************************************************************************************************/
+		class ITrack;
+
+		virtual UEMediaError PrepareTracks(TSharedPtrTS<const IParserISO14496_12>	OptionalMP4InitSegment) = 0;
+
+		virtual TMediaOptionalValue<FTimeFraction> GetMovieDuration() const = 0;
+
+		virtual int32 GetNumberOfTracks() const = 0;
+
+
+		class ITrackIterator
+		{
+		public:
+			virtual ~ITrackIterator() = default;
+
+			enum class ESearchMode
+			{
+				Before,
+				After,
+				Closest
+			};
+
+			virtual UEMediaError StartAtTime(const FTimeValue& AtTime, ESearchMode SearchMode) = 0;
+			virtual UEMediaError Next() = 0;
+
+			virtual bool IsAtEOS() const = 0;
+			virtual const ITrack* GetTrack() const = 0;
+			virtual int64 GetBaseMediaDecodeTime() const = 0;
+			virtual uint32 GetSampleNumber() const = 0;
+			virtual int64 GetDTS() const = 0;
+			virtual int64 GetPTS() const = 0;
+			virtual int64 GetDuration() const = 0;
+			virtual uint32 GetTimescale() const = 0;
+			virtual bool IsSyncSample() const = 0;
+			virtual int64 GetSampleSize() const = 0;
+			virtual int64 GetSampleFileOffset() const = 0;
+		};
+
+		class ITrack
+		{
+		public:
+			// Bitrate box from either 'btrt' or the 'esds' DecoderConfigDescriptor
+			struct FBitrateInfo
+			{
+				FBitrateInfo() : BufferSizeDB(0), MaxBitrate(0), AvgBitrate(0)
+				{
+				}
+				uint32	BufferSizeDB;
+				uint32	MaxBitrate;
+				uint32	AvgBitrate;
+			};
+
+			virtual ~ITrack() = default;
+
+			virtual uint32 GetID() const = 0;
+			virtual FTimeFraction GetDuration() const = 0;
+
+			virtual ITrackIterator* CreateIterator() const = 0;
+			virtual ITrackIterator* CreateIterator(const FParamDict& InOptions) const = 0;
+
+			virtual const TArray<uint8>& GetCodecSpecificData() const = 0;
+			virtual const TArray<uint8>& GetCodecSpecificDataRAW() const = 0;
+			virtual const FStreamCodecInformation& GetCodecInformation() const = 0;
+			virtual const FBitrateInfo& GetBitrateInfo() const = 0;
+			virtual const FString GetLanguage() const = 0;
+		};
+
+
+		class IAllTrackIterator
+		{
+		public:
+			virtual ~IAllTrackIterator() = default;
+			//! Returns the iterator at the current file position.
+			virtual const ITrackIterator* Current() const = 0;
+			//! Advance iterator to point to the next sample in sequence. Returns false if there are no more samples.
+			virtual bool Next() = 0;
+			//! Returns a list of all tracks iterators that reached EOS while iterating since the most recent call to ClearNewEOSTracks().
+			virtual void GetNewEOSTracks(TArray<const ITrackIterator*>& OutTracksThatNewlyReachedEOS) const = 0;
+			//! Clears the list of track iterators that have reached EOS.
+			virtual void ClearNewEOSTracks() = 0;
+			//! Returns list of all iterators.
+			virtual void GetAllIterators(TArray<const ITrackIterator*>& OutIterators) const = 0;
+		};
+
+		virtual TSharedPtr<IAllTrackIterator, ESPMode::ThreadSafe> CreateAllTrackIteratorByFilePos(int64 InFromFilePos) const = 0;
+
+		virtual const ITrack* GetTrackByIndex(int32 Index) const = 0;
+		virtual const ITrack* GetTrackByTrackID(int32 TrackID) const = 0;
+
+
+	};
+
+
+
+} // namespace Electra
+
+
