ElectraPlayer: RELEASE. Moved from NFL to public. Note that this is done via a Delete and "Add New" operation as discussed.

diff --git a/Engine/Plugins/Media/ElectraPlayer/Source/ElectraPlayerRuntime/Private/Runtime/HTTP/HTTPManager.cpp b/Engine/Plugins/Media/ElectraPlayer/Source/ElectraPlayerRuntime/Private/Runtime/HTTP/HTTPManager.cpp
new file mode 100644
--- /dev/null
+++ b/Engine/Plugins/Media/ElectraPlayer/Source/ElectraPlayerRuntime/Private/Runtime/HTTP/HTTPManager.cpp
@@ -0,0 +1,1405 @@
+// Copyright Epic Games, Inc. All Rights Reserved.
+
+#include "HTTP/HTTPManager.h"
+#include "Utilities/StringHelpers.h"
+#include "PlayerTime.h"
+#include "Player/PlayerSessionServices.h"
+#include "HAL/LowLevelMemTracker.h"
+
+// For file:// scheme archive deserialization
+#include "HAL/FileManager.h"
+#include "Misc/FileHelper.h"
+#include "Serialization/Archive.h"
+#include "Serialization/ArrayReader.h"
+
+// For https:// scheme
+#include "Modules/ModuleManager.h"
+#include "HttpModule.h"
+#include "PlatformHttp.h"
+#include "Interfaces/IHttpResponse.h"
+
+
+#include "ElectraPlayerPrivate.h"
+
+#define ERRCODE_OK								0
+#define ERRCODE_HTTP_CONNECTION_TIMEOUT			1
+#define ERRCODE_HTTP_RETURNED_ERROR		 		2
+#define ERRCODE_WRITE_ERROR						3
+#define ERRCODE_OPERATION_TIMEDOUT				4
+#define ERRCODE_HTTP_RANGE_ERROR				5
+#define ERRCODE_HTTP_FILE_COULDNT_READ_FILE		6
+#define ERRCODE_HTTPMODULE_FAILURE				100
+
+DECLARE_CYCLE_STAT(TEXT("FElectraHttpManager::WorkerThread"), STAT_ElectraPlayer_FElectraHttpManager_Worker, STATGROUP_ElectraPlayer);
+
+DEFINE_LOG_CATEGORY(LogElectraHTTPManager);
+
+
+namespace Electra
+{
+
+	class FElectraHttpManager : public IElectraHttpManager, public FMediaThread, public TSharedFromThis<FElectraHttpManager, ESPMode::ThreadSafe>
+	{
+	public:
+		static TSharedPtrTS<FElectraHttpManager> Create();
+
+		virtual ~FElectraHttpManager();
+		void Initialize();
+
+		virtual void AddRequest(TSharedPtrTS<FRequest> Request) override;
+		virtual void RemoveRequest(TSharedPtrTS<FRequest> Request) override;
+
+	private:
+		struct FLocalByteStream
+		{
+			void SetConnected(TSharedPtrTS<FRequest> Request);
+			int32 Read(TSharedPtrTS<FReceiveBuffer> RcvBuffer, TSharedPtrTS<FRequest> Request);
+
+			TSharedPtr<FArchive, ESPMode::ThreadSafe>	Archive;
+			FString										Filename;
+			bool										bIsConnected = false;
+			int64										FileStartOffset = 0;		//!< The base offset into the file data is requested at.
+			int64										FileSize = 0;				//!< The size of the file requested
+			int64										FileSizeToGo = 0;			//!< Number of bytes still to read from the file.
+		};
+
+
+		class FHTTPCallbackWrapper
+		{
+		public:
+			FHTTPCallbackWrapper(TSharedPtrTS<FElectraHttpManager> InOwner) : Owner(InOwner) {}
+			~FHTTPCallbackWrapper()
+			{
+				Unbind();
+			}
+			void ReportRequestComplete(FHttpRequestPtr Request, FHttpResponsePtr Response, bool bConnectedSuccessfully)
+			{
+				FScopeLock lock(&Lock);
+				TSharedPtrTS<FElectraHttpManager> That = Owner.Pin();
+				if (That.IsValid())
+				{
+					That->OnProcessRequestCompleteInternal(Request, Response, bConnectedSuccessfully);
+				}
+			}
+			void ReportRequestProgress(FHttpRequestPtr Request, int32 BytesSent, int32 BytesReceived)
+			{
+				FScopeLock lock(&Lock);
+				TSharedPtrTS<FElectraHttpManager> That = Owner.Pin();
+				if (That.IsValid())
+				{
+					That->OnRequestProgressInternal(Request, BytesSent, BytesReceived);
+				}
+			}
+			void ReportRequestHeaderReceived(FHttpRequestPtr Request, const FString& HeaderName, const FString& NewHeaderValue)
+			{
+				FScopeLock lock(&Lock);
+				TSharedPtrTS<FElectraHttpManager> That = Owner.Pin();
+				if (That.IsValid())
+				{
+					That->OnRequestHeaderReceivedInternal(Request, HeaderName, NewHeaderValue);
+				}
+			}
+
+			void Unbind()
+			{
+				FScopeLock lock(&Lock);
+				Owner.Reset();
+			}
+		private:
+			FCriticalSection Lock;
+			TWeakPtrTS<FElectraHttpManager> Owner;
+		};
+
+		struct FHandle
+		{
+			~FHandle()
+			{
+				Cleanup();
+			}
+
+			void Cleanup()
+			{
+				if (HandleType == EHandleType::FileHandle)
+				{
+					delete LocalByteStream;
+					LocalByteStream = nullptr;
+				}
+				else if (HandleType == EHandleType::HTTPHandle)
+				{
+					if (HttpsRequestCallbackWrapper.IsValid())
+					{
+						HttpsRequestCallbackWrapper->Unbind();
+					}
+					HttpsRequestCallbackWrapper.Reset();
+					if (HttpRequest.IsValid() && !EHttpRequestStatus::IsFinished(HttpRequest->GetStatus()))
+					{
+						HttpRequest->CancelRequest();
+					}
+					HttpRequest.Reset();
+				}
+			}
+
+			struct FRequestResponse
+			{
+				// Current range requested from server.
+				FParams::FRange Range;
+				// Current response received from server
+				FHttpResponsePtr Response;
+				// Number of bytes passed along already.
+				int64 NumBytesPassedOut = 0;
+				// Original range requested for easier access and comparison.
+				FParams::FRange OriginalRange;
+				FParams::FRange ReceivedContentRange;
+				//
+				bool bIsSubRangeRequest = false;
+				int32 NumSubRangeRequest = 0;
+
+				//
+				int64 SizeRemaining() const
+				{
+					// Did we break this into sub range requests?
+					if (bIsSubRangeRequest)
+					{
+						// The expected end position is either the end of the original requested range or the end of the document.
+						int64 ExpectedEndPos = OriginalRange.GetEndIncluding() >=0 ? OriginalRange.GetEndIncluding()+1 : ReceivedContentRange.GetDocumentSize();
+						// It should not be negative. The total document size should be available by now.
+						// If it is not we may be faced with a chunked transfer of a document with unknown/infinite size, which is bad.
+						check(ExpectedEndPos >= 0);
+						if (ExpectedEndPos >= 0)
+						{
+							check(ReceivedContentRange.IsSet());
+							int64 SizeToGo = ExpectedEndPos - (ReceivedContentRange.GetEndIncluding() + 1);
+							check(SizeToGo >= 0);
+							return SizeToGo;
+						}
+						else
+						{
+							UE_LOG(LogElectraHTTPManager, Error, TEXT("Unknown document size in sub ranged download"));
+						}
+					}
+					return 0;
+				}
+
+			};
+
+			enum class EHandleType
+			{
+				Undefined,
+				FileHandle,
+				HTTPHandle
+			};
+
+			FElectraHttpManager* 	Owner = nullptr;
+
+			EHandleType				HandleType = EHandleType::Undefined;
+
+			// Local file handle (for file://)
+			FLocalByteStream*		LocalByteStream = nullptr;
+
+			// HTTP module handle (for http:// and https://)
+			FHttpRequestPtr			HttpRequest;
+			TSharedPtr<FHTTPCallbackWrapper, ESPMode::ThreadSafe>	HttpsRequestCallbackWrapper;
+			bool					bHttpRequestFirstEvent = true;
+
+			FTimeValue				LastTimeDataReceived;
+			FTimeValue				TimeAtNextProgressCallback;
+			FTimeValue				TimeAtConnectionTimeoutCheck;
+
+			FRequestResponse		ActiveResponse;
+			bool					bResponseReceived = false;
+
+			int64					BytesReadSoFar = 0;
+
+			// Internal work variables mirroring HTTP::FConnectionInfo values that may change with sub range requests and
+			// should not change for the original request.
+			FTimeValue				RequestStartTime;
+			bool					bIsConnected = false;
+			bool					bHaveResponseHeaders = false;
+
+
+			void ClearForNextSubRange()
+			{
+				LastTimeDataReceived.SetToInvalid();
+				TimeAtConnectionTimeoutCheck.SetToInvalid();
+				RequestStartTime.SetToInvalid();
+				bHttpRequestFirstEvent = true;
+				bResponseReceived = false;
+				bIsConnected = false;
+				bHaveResponseHeaders = false;
+			}
+		};
+
+		struct FTransportError
+		{
+			FTransportError()
+			{
+				Clear();
+			}
+			void Clear()
+			{
+				Message.Empty();
+				ErrorCode = ERRCODE_OK;
+			}
+			void Set(int32 InCode, const char* InMessage)
+			{
+				ErrorCode = InCode;
+				Message = FString(InMessage);
+			}
+			void Set(int32 InCode, const FString& InMessage)
+			{
+				ErrorCode = InCode;
+				Message = InMessage;
+			}
+
+			FString			Message;
+			int32			ErrorCode;
+		};
+
+		FHandle* CreateLocalFileHandle(const FTimeValue& Now, FTransportError& OutError, const TSharedPtrTS<FRequest>& Request);
+		FHandle* CreateHTTPModuleHandle(const FTimeValue& Now, FTransportError& OutError, const TSharedPtrTS<FRequest>& Request);
+		bool PrepareHTTPModuleHandle(const FTimeValue& Now, FHandle* InHandle, const TSharedPtrTS<FRequest>& Request, bool bIsFirstSetup);
+		void AddPendingRequests(const FTimeValue& Now);
+		void RemovePendingRequests(const FTimeValue& Now);
+		void HandleCompletedRequests();
+		void HandlePeriodicCallbacks(const FTimeValue& Now);
+		void HandleTimeouts(const FTimeValue& Now);
+		void HandleLocalFileRequests();
+		void HandleHTTPRequests(const FTimeValue& Now);
+		void HandleHTTPResponses(const FTimeValue& Now);
+		void RemoveAllRequests();
+
+		void WorkerThread();
+
+	public:
+		// Callbacks from the HTTP module
+		void OnProcessRequestCompleteInternal(FHttpRequestPtr Request, FHttpResponsePtr Response, bool bConnectedSuccessfully);
+		void OnRequestProgressInternal(FHttpRequestPtr Request, int32 BytesSent, int32 BytesReceived);
+		void OnRequestHeaderReceivedInternal(FHttpRequestPtr Request, const FString& HeaderName, const FString& NewHeaderValue);
+	private:
+		struct FHttpEventData
+		{
+			TArray<HTTP::FHTTPHeader>	Headers;
+			FHttpResponsePtr			Response;
+			EHttpRequestStatus::Type	Status = EHttpRequestStatus::Type::NotStarted;
+			int32						BytesSent = 0;
+			int32						BytesReceived = 0;
+			bool						bComplete = false;
+			bool						bConnectedSuccessfully = false;
+		};
+
+		struct FRemoveRequest
+		{
+			FRemoveRequest()
+				: WaitingEvent(nullptr)
+			{
+			}
+			~FRemoveRequest()
+			{
+				if (WaitingEvent)
+				{
+					WaitingEvent->Signal();
+				}
+			}
+			FRemoveRequest(const FRemoveRequest& other)
+				: Request(other.Request), WaitingEvent(other.WaitingEvent)
+			{
+			}
+
+			FRemoveRequest(FRemoveRequest&& other)
+				: Request(other.Request), WaitingEvent(other.WaitingEvent)
+			{
+				other.Request.Reset();
+				other.WaitingEvent = nullptr;
+			}
+			FRemoveRequest& operator=(FRemoveRequest&& other)
+			{
+				if (this != &other)
+				{
+					if (WaitingEvent)
+					{
+						WaitingEvent->Signal();
+					}
+					Request = other.Request;
+					WaitingEvent = other.WaitingEvent;
+
+					other.Request.Reset();
+					other.WaitingEvent = nullptr;
+				}
+				return *this;
+			}
+			FRemoveRequest& operator=(const FRemoveRequest& other)
+			{
+				if (this != &other)
+				{
+					if (WaitingEvent)
+					{
+						WaitingEvent->Signal();
+					}
+					Request = other.Request;
+					WaitingEvent = other.WaitingEvent;
+				}
+				return *this;
+			}
+
+			TSharedPtrTS<FRequest>			Request;
+			FMediaEvent* WaitingEvent;
+		};
+
+		FCriticalSection												Lock;
+
+		FMediaEvent														RequestChangesEvent;
+		TQueue<TSharedPtrTS<FRequest>, EQueueMode::Mpsc>				RequestsToAdd;
+		TQueue<FRemoveRequest, EQueueMode::Mpsc>						RequestsToRemove;
+		TQueue<TSharedPtrTS<FRequest>, EQueueMode::Mpsc>				RequestsCompleted;
+
+		TMap<FHandle*, TSharedPtrTS<FRequest>>							ActiveRequests;
+		FTimeValue														ProgressInterval;
+
+		bool															bThreadStarted = false;
+		bool															bTerminate = false;
+
+		FCriticalSection												HttpModuleEventLock;
+		TMap<FHttpRequestPtr, FHttpEventData>							HttpModuleEventMap;
+
+		static TWeakPtrTS<FElectraHttpManager>							SingletonSelf;
+		static FCriticalSection											SingletonLock;
+	};
+
+	TWeakPtrTS<FElectraHttpManager>		FElectraHttpManager::SingletonSelf;
+	FCriticalSection					FElectraHttpManager::SingletonLock;
+
+	/***************************************************************************************************************************************************/
+	/***************************************************************************************************************************************************/
+	/***************************************************************************************************************************************************/
+
+	TSharedPtrTS<IElectraHttpManager> IElectraHttpManager::Create()
+	{
+		return FElectraHttpManager::Create();
+	}
+
+	TSharedPtrTS<FElectraHttpManager> FElectraHttpManager::Create()
+	{
+		FScopeLock lock(&SingletonLock);
+		TSharedPtrTS<FElectraHttpManager> Self = SingletonSelf.Pin();
+		if (!Self.IsValid())
+		{
+			FElectraHttpManager* Manager = new FElectraHttpManager;
+			if (Manager)
+			{
+				Manager->Initialize();
+			}
+			Self = MakeShareable(Manager);
+			SingletonSelf = Self;
+		}
+		return Self;
+	}
+
+
+	FElectraHttpManager::~FElectraHttpManager()
+	{
+		bTerminate = true;
+		if (bThreadStarted)
+		{
+			ThreadWaitDone();
+			ThreadReset();
+		}
+	}
+
+	void FElectraHttpManager::Initialize()
+	{
+		ProgressInterval.SetFromMilliseconds(100);
+
+		ThreadSetName("ElectraPlayer::HTTPManager");
+		ThreadStart(Electra::MakeDelegate(this, &FElectraHttpManager::WorkerThread));
+
+		bThreadStarted = true;
+	}
+
+	void FElectraHttpManager::AddRequest(TSharedPtrTS<FRequest> Request)
+	{
+		FScopeLock lock(&Lock);
+		if (!bTerminate)
+		{
+			RequestsToAdd.Enqueue(Request);
+			RequestChangesEvent.Signal();
+		}
+	}
+
+	void FElectraHttpManager::RemoveRequest(TSharedPtrTS<FRequest> Request)
+	{
+		FMediaEvent Sig;
+		FRemoveRequest Remove;
+		Remove.Request = Request;
+		Remove.WaitingEvent = &Sig;
+		Lock.Lock();
+		RequestsToRemove.Enqueue(MoveTemp(Remove));
+		RequestChangesEvent.Signal();
+		Lock.Unlock();
+		Sig.Wait();
+	}
+
+	FElectraHttpManager::FHandle* FElectraHttpManager::CreateLocalFileHandle(const FTimeValue& Now, FTransportError& OutError, const TSharedPtrTS<IElectraHttpManager::FRequest>& Request)
+	{
+		TUniquePtr<FHandle> Handle(new FHandle);
+		Handle->Owner = this;
+		Handle->HandleType = FHandle::EHandleType::FileHandle;
+		Handle->LocalByteStream = new FLocalByteStream;
+		Handle->LocalByteStream->Filename = Request->Parameters.URL.Mid(7); /* file:// */
+		Handle->LocalByteStream->Archive = MakeShareable(IFileManager::Get().CreateFileReader(*Handle->LocalByteStream->Filename));
+		if (!Handle->LocalByteStream->Archive.IsValid())
+		{
+			OutError.Set(ERRCODE_HTTP_FILE_COULDNT_READ_FILE, FString::Printf(TEXT("Failed to open media file \"%s\""), *Request->Parameters.URL));
+			return nullptr;
+		}
+		return Handle.Release();
+	}
+
+	FElectraHttpManager::FHandle* FElectraHttpManager::CreateHTTPModuleHandle(const FTimeValue& Now, FTransportError& OutError, const TSharedPtrTS<IElectraHttpManager::FRequest>& Request)
+	{
+		TUniquePtr<FHandle> Handle(new FHandle);
+		Handle->Owner = this;
+		Handle->HandleType = FHandle::EHandleType::HTTPHandle;
+		Handle->bHttpRequestFirstEvent = true;
+		Handle->ActiveResponse.NumSubRangeRequest = 0;
+		Handle->ActiveResponse.OriginalRange = Request->Parameters.Range;
+
+		PrepareHTTPModuleHandle(Now, Handle.Get(), Request, true);
+		return Handle.Release();
+	}
+
+
+	bool FElectraHttpManager::PrepareHTTPModuleHandle(const FTimeValue& Now, FHandle* Handle, const TSharedPtrTS<FRequest>& Request, bool bIsFirstSetup)
+	{
+		FHttpModule& HttpModule = FModuleManager::LoadModuleChecked<FHttpModule>("HTTP");
+
+		Handle->HttpsRequestCallbackWrapper = MakeShared<FHTTPCallbackWrapper, ESPMode::ThreadSafe>(AsShared());
+		Handle->HttpRequest = HttpModule.Get().CreateRequest();
+		Handle->HttpRequest->OnProcessRequestComplete().BindThreadSafeSP(Handle->HttpsRequestCallbackWrapper.ToSharedRef(), &FHTTPCallbackWrapper::ReportRequestComplete);
+		Handle->HttpRequest->OnRequestProgress().BindThreadSafeSP(Handle->HttpsRequestCallbackWrapper.ToSharedRef(), &FHTTPCallbackWrapper::ReportRequestProgress);
+		Handle->HttpRequest->OnHeaderReceived().BindThreadSafeSP(Handle->HttpsRequestCallbackWrapper.ToSharedRef(), &FHTTPCallbackWrapper::ReportRequestHeaderReceived);
+		Handle->HttpRequest->SetURL(Request->Parameters.URL);
+		Handle->HttpRequest->SetVerb(TEXT("GET"));
+		Handle->HttpRequest->SetHeader(TEXT("User-Agent"), "X-UnrealEngine-Agent");
+		for(int32 i=0; i<Request->Parameters.RequestHeaders.Num(); ++i)
+		{
+			Handle->HttpRequest->SetHeader(Request->Parameters.RequestHeaders[i].Header, Request->Parameters.RequestHeaders[i].Value);
+		}
+
+		Handle->RequestStartTime = Now;
+		Handle->TimeAtConnectionTimeoutCheck.SetToPositiveInfinity();
+		// Is this the first sub range request or a continuation?
+		if (bIsFirstSetup)
+		{
+			Request->ConnectionInfo.RequestStartTime = Now;
+			Handle->TimeAtNextProgressCallback = Now + ProgressInterval;
+			Handle->ActiveResponse.Range = Handle->ActiveResponse.OriginalRange;
+			Handle->ActiveResponse.NumSubRangeRequest = 0;
+			Handle->BytesReadSoFar = 0;
+		}
+		else
+		{
+			// Set up the range to follow the previous.
+			Handle->ActiveResponse.Range.SetStart(Handle->ActiveResponse.ReceivedContentRange.GetEndIncluding() + 1);
+			if (Handle->ActiveResponse.OriginalRange.GetEndIncluding() >= 0)
+			{
+				Handle->ActiveResponse.Range.SetEndIncluding(Handle->ActiveResponse.OriginalRange.GetEndIncluding());
+			}
+			else if (Handle->ActiveResponse.ReceivedContentRange.GetDocumentSize() > 0)
+			{
+				Handle->ActiveResponse.Range.SetEndIncluding(Handle->ActiveResponse.ReceivedContentRange.GetDocumentSize() - 1);
+			}
+			else
+			{
+				Handle->ActiveResponse.Range.SetEndIncluding(-1);
+			}
+			++Handle->ActiveResponse.NumSubRangeRequest;
+			Handle->BytesReadSoFar += Handle->ActiveResponse.NumBytesPassedOut;
+			Handle->ClearForNextSubRange();
+		}
+		Handle->ActiveResponse.NumBytesPassedOut = 0;
+		Handle->ActiveResponse.ReceivedContentRange.Reset();
+		Handle->ActiveResponse.Response.Reset();
+
+		// Response not yet received.
+		Handle->bResponseReceived = false;
+
+		// Do a quick check that if the request is reading into a ring buffer that there is also a sub-range request size set up.
+		TSharedPtrTS<FReceiveBuffer> ReceiveBuffer = Request->ReceiveBuffer.Pin();
+		if (ReceiveBuffer.IsValid() && ReceiveBuffer->bEnableRingbuffer)
+		{
+			if (Request->Parameters.SubRangeRequestSize > 0)
+			{
+				// The requested sub range size may be smaller than the range we are supposed to get.
+				// Adjust the first range request accordingly.
+				if (Handle->ActiveResponse.Range.IsSet())
+				{
+					// If the original range is larger than the sub range request, including an open ended range, we need to adjust the range.
+					int64 NumOrgBytes = Handle->ActiveResponse.Range.GetNumberOfBytes();
+					if (NumOrgBytes > Request->Parameters.SubRangeRequestSize || Handle->ActiveResponse.Range.GetEndIncluding() < 0)
+					{
+						Handle->ActiveResponse.Range.SetEndIncluding(Handle->ActiveResponse.Range.GetStart() + Request->Parameters.SubRangeRequestSize - 1);
+						Handle->ActiveResponse.bIsSubRangeRequest = true;
+					}
+				}
+				else
+				{
+					// No range requested. We still need to perform sub range requests.
+					Handle->ActiveResponse.Range.SetStart(0);
+					Handle->ActiveResponse.Range.SetEndIncluding(Request->Parameters.SubRangeRequestSize - 1);
+					Handle->ActiveResponse.bIsSubRangeRequest = true;
+				}
+			}
+			else
+			{
+				// Give a warning to the console about this circumstance.
+				UE_LOG(LogElectraHTTPManager, Warning, TEXT("Receive buffer set to ring buffer mode but no sub range request size specified!"));
+			}
+		}
+		// Set the possibly adjusted range.
+		if (Handle->ActiveResponse.Range.IsSet())
+		{
+			FString Range = FString::Printf(TEXT("bytes=%s"), *Handle->ActiveResponse.Range.GetString());
+			Handle->HttpRequest->SetHeader(TEXT("Range"), Range);
+		}
+		return true;
+	}
+
+
+	void FElectraHttpManager::AddPendingRequests(const FTimeValue& Now)
+	{
+		while(!RequestsToAdd.IsEmpty())
+		{
+			TSharedPtrTS<FRequest> Request;
+			RequestsToAdd.Dequeue(Request);
+
+			FTransportError HttpError;
+			Request->ConnectionInfo.EffectiveURL = Request->Parameters.URL;
+
+			// Is this a local file?
+			if (Request->Parameters.URL.Len() > 7 && Request->Parameters.URL.Left(7) == TEXT("file://"))
+			{
+				FHandle* Handle = CreateLocalFileHandle(Now, HttpError, Request);
+				if (Handle)
+				{
+					ActiveRequests.Add(Handle, Request);
+
+					Request->ConnectionInfo.RequestStartTime = Now;
+					Handle->RequestStartTime = Now;
+					Handle->TimeAtNextProgressCallback = Now + ProgressInterval;
+					Handle->TimeAtConnectionTimeoutCheck.SetToPositiveInfinity();
+				}
+				else
+				{
+					Request->ConnectionInfo.StatusInfo.ErrorDetail.SetError(UEMEDIA_ERROR_INTERNAL).SetFacility(Facility::EFacility::HTTPReader).SetMessage(HttpError.Message).SetCode((uint16)HttpError.ErrorCode);
+					RequestsCompleted.Enqueue(Request);
+				}
+			}
+			else
+			{
+				FHandle* Handle = CreateHTTPModuleHandle(Now, HttpError, Request);
+				if (Handle)
+				{
+					ActiveRequests.Add(Handle, Request);
+					if (!Handle->HttpRequest->ProcessRequest())
+					{
+						Request->ConnectionInfo.StatusInfo.ErrorDetail.SetError(UEMEDIA_ERROR_INTERNAL).SetFacility(Facility::EFacility::HTTPReader).SetMessage("HTTP request failed on ProcessRequest()");
+						RequestsCompleted.Enqueue(Request);
+					}
+				}
+				else
+				{
+					Request->ConnectionInfo.StatusInfo.ErrorDetail.SetError(UEMEDIA_ERROR_INTERNAL).SetFacility(Facility::EFacility::HTTPReader).SetMessage(HttpError.Message).SetCode((uint16)HttpError.ErrorCode);
+					RequestsCompleted.Enqueue(Request);
+				}
+			}
+		}
+	}
+
+	void FElectraHttpManager::RemovePendingRequests(const FTimeValue& Now)
+	{
+		// Remove pending requests
+		while(!RequestsToRemove.IsEmpty())
+		{
+			TSharedPtrTS<FRequest>	Request = RequestsToRemove.Peek()->Request;
+
+			// Is this an active request?
+			for(TMap<FHandle*, TSharedPtrTS<FRequest>>::TIterator It = ActiveRequests.CreateIterator(); It; ++It)
+			{
+				if (It.Value() == Request)
+				{
+					FHandle* Handle = It.Key();
+					It.RemoveCurrent();
+					delete Handle;
+					break;
+				}
+			}
+
+			// Removing an unfinished transfer that has not errored means it was aborted.
+			if (!Request->ConnectionInfo.bHasFinished && !Request->ConnectionInfo.StatusInfo.ErrorDetail.IsError())
+			{
+				Request->ConnectionInfo.bWasAborted = true;
+			}
+			if (!Request->ConnectionInfo.RequestEndTime.IsValid())
+			{
+				Request->ConnectionInfo.RequestEndTime = Now;
+			}
+
+			// Pop the remove request. This will release the waiting caller.
+			RequestsToRemove.Pop();
+		}
+	}
+
+	void FElectraHttpManager::HandleCompletedRequests()
+	{
+		// Remove pending requests
+		while(!RequestsCompleted.IsEmpty())
+		{
+			TSharedPtrTS<FRequest> Request;
+			RequestsCompleted.Dequeue(Request);
+			// Remove from active requests. It may not be in there if it had an error upon creating.
+			for(TMap<FHandle*, TSharedPtrTS<FRequest>>::TIterator It = ActiveRequests.CreateIterator(); It; ++It)
+			{
+				if (It.Value() == Request)
+				{
+					FHandle* Handle = It.Key();
+					It.RemoveCurrent();
+					delete Handle;
+
+					Request->ConnectionInfo.bHasFinished = true;
+					TSharedPtrTS<FReceiveBuffer>	ReceiveBuffer = Request->ReceiveBuffer.Pin();
+					if (ReceiveBuffer.IsValid())
+					{
+						ReceiveBuffer->Buffer.SetEOD();
+					}
+
+					// Call completion delegate.
+					TSharedPtrTS<FProgressListener> ProgressListener = Request->ProgressListener.Pin();
+					if (ProgressListener.IsValid())
+					{
+						if (!ProgressListener->CompletionDelegate.empty())
+						{
+							ProgressListener->CallCompletionDelegate(Request.Get());
+						}
+					}
+					break;
+				}
+			}
+		}
+	}
+
+
+	void FElectraHttpManager::HandlePeriodicCallbacks(const FTimeValue& Now)
+	{
+		SCOPE_CYCLE_COUNTER(STAT_ElectraPlayer_FElectraHttpManager_Worker);
+		CSV_SCOPED_TIMING_STAT(ElectraPlayer, ElectraHttpManager_Worker);
+
+		for(TMap<FHandle*, TSharedPtrTS<FRequest>>::TIterator It = ActiveRequests.CreateIterator(); It; ++It)
+		{
+			FHandle* Handle = It.Key();
+
+			// Update throughput stats when something was read.
+			HTTP::FConnectionInfo& ci = It.Value()->ConnectionInfo;
+			if (ci.Throughput.LastReadCount != ci.BytesReadSoFar)
+			{
+				FTimeValue DiffTime = Now - ci.Throughput.LastCheckTime;
+				if (DiffTime > FTimeValue::GetZero())
+				{
+					ci.Throughput.ActiveReadTime += DiffTime;
+					int64 BytesPerSecond = 1000 * ci.Throughput.LastReadCount / ci.Throughput.ActiveReadTime.GetAsMilliseconds();
+					ci.Throughput.AccumulatedBytesPerSec += BytesPerSecond;
+					if (!ci.Throughput.bIsFirst)
+					{
+						ci.Throughput.AccumulatedBytesPerSecWithoutFirst += BytesPerSecond;
+					}
+					ci.Throughput.bIsFirst = false;
+					++ci.Throughput.NumAccumulatedBytesPerSec;
+				}
+				ci.Throughput.LastReadCount = ci.BytesReadSoFar;
+			}
+			ci.Throughput.LastCheckTime = Now;
+
+			// Fire periodic progress callback?
+			if (Now >= Handle->TimeAtNextProgressCallback)
+			{
+				Handle->TimeAtNextProgressCallback += ProgressInterval - (Now - Handle->TimeAtNextProgressCallback);
+				TSharedPtrTS<FProgressListener> ProgressListener = It.Value()->ProgressListener.Pin();
+				if (ProgressListener.IsValid())
+				{
+					if (!ProgressListener->ProgressDelegate.empty())
+					{
+						int32 Result = ProgressListener->CallProgressDelegate(It.Value().Get());
+						// Did the progress callback ask to abort the download?
+						if (Result)
+						{
+							It.Value()->ConnectionInfo.bWasAborted = true;
+							It.Value()->ConnectionInfo.RequestEndTime = MEDIAutcTime::Current();
+							if (Handle->HttpsRequestCallbackWrapper.IsValid())
+							{
+								Handle->HttpsRequestCallbackWrapper->Unbind();
+							}
+							RequestsCompleted.Enqueue(It.Value());
+						}
+					}
+				}
+			}
+		}
+	}
+
+
+	void FElectraHttpManager::HandleTimeouts(const FTimeValue& Now)
+	{
+		for(TMap<FHandle*, TSharedPtrTS<FRequest>>::TIterator It = ActiveRequests.CreateIterator(); It; ++It)
+		{
+			FHandle* Handle = It.Key();
+
+			// Timeout handling for file handles is not desired. Skip to the next handle.
+			if (Handle->HandleType == FHandle::EHandleType::FileHandle)
+			{
+				continue;
+			}
+			// HTTP transfers that are not processing also do not need to be checked.
+			else if (Handle->HandleType == FHandle::EHandleType::HTTPHandle && Handle->HttpRequest.IsValid() && EHttpRequestStatus::IsFinished(Handle->HttpRequest->GetStatus()))
+			{
+				continue;
+			}
+
+			TSharedPtrTS<FRequest> Request = It.Value();
+
+			// Time to check for a connection timeout?
+			if (Now >= Handle->TimeAtConnectionTimeoutCheck)
+			{
+				// For our purposes we timeout when we are not connected / have not received any response header.
+				if (!Handle->bIsConnected || !Handle->bHaveResponseHeaders)
+				{
+					Request->ConnectionInfo.StatusInfo.ConnectionTimeoutAfterMilliseconds = (Now - Handle->RequestStartTime).GetAsMilliseconds();
+					Request->ConnectionInfo.StatusInfo.ErrorCode = ERRCODE_OPERATION_TIMEDOUT;
+					Request->ConnectionInfo.StatusInfo.ErrorDetail.
+						SetError(UEMEDIA_ERROR_READ_ERROR).
+						SetFacility(Facility::EFacility::HTTPReader).
+						SetCode(ERRCODE_HTTP_CONNECTION_TIMEOUT).
+						SetMessage(FString::Printf(TEXT("Connection timeout after %d milliseconds, limit was %d"), Request->ConnectionInfo.StatusInfo.ConnectionTimeoutAfterMilliseconds, (int32)Request->Parameters.ConnectTimeout.GetAsMilliseconds()));
+					Request->ConnectionInfo.RequestEndTime = Now;
+					if (Handle->HttpsRequestCallbackWrapper.IsValid())
+					{
+						Handle->HttpsRequestCallbackWrapper->Unbind();
+					}
+					RequestsCompleted.Enqueue(Request);
+				}
+				Handle->TimeAtConnectionTimeoutCheck.SetToPositiveInfinity();
+			}
+
+			// Data timeout? This requires to be connected to the server and to have received at least one response header.
+			if (Handle->LastTimeDataReceived.IsValid() && Request->Parameters.NoDataTimeout.IsValid())
+			{
+				FTimeValue DeltaTime = Now - Handle->LastTimeDataReceived;
+				if (DeltaTime >= Request->Parameters.NoDataTimeout)
+				{
+					Request->ConnectionInfo.StatusInfo.NoDataTimeoutAfterMilliseconds = DeltaTime.GetAsMilliseconds();
+					Request->ConnectionInfo.StatusInfo.ErrorCode = ERRCODE_OPERATION_TIMEDOUT;
+					Request->ConnectionInfo.StatusInfo.ErrorDetail.
+						SetError(UEMEDIA_ERROR_READ_ERROR).
+						SetFacility(Facility::EFacility::HTTPReader).
+						SetCode(ERRCODE_HTTP_CONNECTION_TIMEOUT).
+						SetMessage(FString::Printf(TEXT("No data timeout after %d milliseconds, limit was %d. Received %lld of %lld bytes"), Request->ConnectionInfo.StatusInfo.NoDataTimeoutAfterMilliseconds, (int32)Request->Parameters.NoDataTimeout.GetAsMilliseconds()
+							, (long long int)Request->ConnectionInfo.BytesReadSoFar, (long long int)Request->ConnectionInfo.ContentLength));
+					Request->ConnectionInfo.RequestEndTime = Now;
+					if (Handle->HttpsRequestCallbackWrapper.IsValid())
+					{
+						Handle->HttpsRequestCallbackWrapper->Unbind();
+					}
+					RequestsCompleted.Enqueue(Request);
+				}
+			}
+		}
+	}
+
+
+	void FElectraHttpManager::RemoveAllRequests()
+	{
+		RequestsToAdd.Empty();
+		RequestsToRemove.Empty();
+		RequestsCompleted.Empty();
+		while(!ActiveRequests.IsEmpty())
+		{
+			TMap<FHandle*, TSharedPtrTS<FRequest>>::TIterator It = ActiveRequests.CreateIterator();
+			FHandle* Handle = It.Key();
+			It.RemoveCurrent();
+			delete Handle;
+		}
+	}
+
+
+	void FElectraHttpManager::WorkerThread()
+	{
+		LLM_SCOPE(ELLMTag::ElectraPlayer);
+
+		FTimeValue Now = MEDIAutcTime::Current();
+		while(!bTerminate)
+		{
+			const long kSelectTimeoutMilliseconds = 15;
+
+			{
+				SCOPE_CYCLE_COUNTER(STAT_ElectraPlayer_FElectraHttpManager_Worker);
+				CSV_SCOPED_TIMING_STAT(ElectraPlayer, ElectraHttpManager_Worker);
+
+				// Add and remove pending requests
+				Lock.Lock();
+				AddPendingRequests(Now);
+				RemovePendingRequests(Now);
+				Lock.Unlock();
+			}
+
+			// Throttle ourselves here a bit.
+			RequestChangesEvent.WaitTimeoutAndReset(kSelectTimeoutMilliseconds * 1000);
+
+			{
+				SCOPE_CYCLE_COUNTER(STAT_ElectraPlayer_FElectraHttpManager_Worker);
+				CSV_SCOPED_TIMING_STAT(ElectraPlayer, ElectraHttpManager_Worker);
+
+				// Handle local file requests.
+				HandleLocalFileRequests();
+
+				Now = MEDIAutcTime::Current();
+
+				// Handle requests to the HTTP module
+				HandleHTTPRequests(Now);
+				// Handle the responses.
+				HandleHTTPResponses(Now);
+
+				// Handle periodic progress callbacks. Do this before handling completed requests in case a callback asks to abort.
+				HandlePeriodicCallbacks(Now);
+				// Handle timeouts after the progress callbacks.
+				HandleTimeouts(Now);
+				// Handle all finished requests.
+				HandleCompletedRequests();
+			}
+		}
+		RemoveAllRequests();
+		HttpModuleEventMap.Reset();
+	}
+
+
+	void FElectraHttpManager::HandleLocalFileRequests()
+	{
+		for(TMap<FHandle*, TSharedPtrTS<FRequest>>::TIterator It = ActiveRequests.CreateIterator(); It; ++It)
+		{
+			FHandle* Handle = It.Key();
+			if (Handle->HandleType == FHandle::EHandleType::FileHandle && Handle->LocalByteStream)
+			{
+				TSharedPtrTS<FRequest> Request = It.Value();
+				// Establish the file handle as "connected"
+				Handle->LocalByteStream->SetConnected(Request);
+				// Read from local file
+				TSharedPtrTS<FReceiveBuffer> ReceiveBuffer = Request->ReceiveBuffer.Pin();
+				if (ReceiveBuffer.IsValid())
+				{
+					int32 NumBytesRead = Handle->LocalByteStream->Read(ReceiveBuffer, Request);
+					if (NumBytesRead > 0)
+					{
+						Handle->LastTimeDataReceived = MEDIAutcTime::Current();
+					}
+					// Reading done?
+					if (Handle->LocalByteStream->FileSizeToGo <= 0)
+					{
+						Request->ConnectionInfo.RequestEndTime = Request->ConnectionInfo.StatusInfo.OccurredAtUTC = MEDIAutcTime::Current();
+						RequestsCompleted.Enqueue(Request);
+					}
+				}
+				else
+				{
+					// With the receive buffer having been released we can abort the transfer.
+					Request->ConnectionInfo.bWasAborted = true;
+					Request->ConnectionInfo.RequestEndTime = Request->ConnectionInfo.StatusInfo.OccurredAtUTC = MEDIAutcTime::Current();
+					RequestsCompleted.Enqueue(Request);
+				}
+			}
+		}
+	}
+
+	void FElectraHttpManager::HandleHTTPRequests(const FTimeValue& Now)
+	{
+		// Get the events that have fired so far into a local map and clear out the original.
+		HttpModuleEventLock.Lock();
+		TMap<FHttpRequestPtr, FHttpEventData> Events = MoveTemp(HttpModuleEventMap);
+		HttpModuleEventLock.Unlock();
+
+		for(auto& Evt : Events)
+		{
+			// Find the request.
+			for(TMap<FHandle*, TSharedPtrTS<FRequest>>::TIterator It = ActiveRequests.CreateIterator(); It; ++It)
+			{
+				FHandle* Handle = It.Key();
+				if (Handle->HandleType == FHandle::EHandleType::HTTPHandle && Handle->HttpRequest == Evt.Key)
+				{
+					TSharedPtrTS<FRequest> Request = It.Value();
+					HTTP::FConnectionInfo& ci = Request->ConnectionInfo;
+					const FHttpEventData& EvData = Evt.Value;
+
+					FHttpResponsePtr Response = EvData.Response;
+					const int32 /*EHttpResponseCodes::Type*/ ResponseCode = Response ? Response->GetResponseCode() : EHttpResponseCodes::Unknown;
+					const bool bHTTPResponseSuccess = ResponseCode >= 200 && ResponseCode <= 299;
+
+					//FPlatformMisc::LocalPrint(*FString::Printf(TEXT("(%d, %d, %f; %d, %d, %d, %d)\n"), Status, ResponseCode, ElapsedTime, EvData.bComplete, EvData.bConnectedSuccessfully, EvData.BytesReceived, EvData.BytesSent));
+
+					// Set connect state on the first event indicating any kind of progress.
+					// Note: The HTTP module does not allow access to the response headers or HTTP status code during processing,
+					//       even when those are already available.
+					if (Handle->bHttpRequestFirstEvent)
+					{
+						if (EvData.bComplete || EvData.BytesSent || EvData.BytesReceived || EvData.Headers.Num())
+						{
+							// Any data flow means that the connection has been established.
+							if (Handle->ActiveResponse.NumSubRangeRequest == 0)
+							{
+								ci.bIsConnected = true;
+								ci.TimeUntilConnected = (Now - Handle->RequestStartTime).GetAsSeconds();
+							}
+							Handle->bIsConnected = true;
+							// That's it for the first event.
+							Handle->bHttpRequestFirstEvent = false;
+						}
+					}
+
+					// Process response headers.
+					if (EvData.Headers.Num())
+					{
+						// It may not be all of them yet, but headers were received.
+						// Unfortunately we do not get the "HTTP/1.0" / "HTTP/1.1" / "HTTP/2" header here to see which protocol version we are dealing with.
+						ci.bHaveResponseHeaders = true;
+						Handle->bHaveResponseHeaders = true;
+						for(int32 i=0, iMax=EvData.Headers.Num(); i<iMax; ++i)
+						{
+							if (!EvData.Headers[i].Header.IsEmpty())
+							{
+								ci.ResponseHeaders.Add(EvData.Headers[i]);
+								const FString& HeaderKey = EvData.Headers[i].Header;
+								const FString& HeaderValue = EvData.Headers[i].Value;
+								if (HeaderKey == TEXT("Content-Length"))
+								{
+									ci.ContentLengthHeader = TEXT("Content-Length:") + HeaderValue;
+								}
+								else if (HeaderKey == TEXT("Content-Range"))
+								{
+									ci.ContentRangeHeader = TEXT("Content-Range: ") + HeaderValue;
+								}
+								else if (HeaderKey == TEXT("Content-Type"))
+								{
+									ci.ContentType = HeaderValue;
+								}
+								else if (HeaderKey == TEXT("Accept-Ranges"))
+								{
+									if (ci.HTTPVersionReceived == 0)
+									{
+										if (HeaderValue.Find(TEXT("none"), ESearchCase::IgnoreCase, ESearchDir::FromStart, INDEX_NONE) == INDEX_NONE)
+										{
+											ci.HTTPVersionReceived = 11;
+										}
+										else
+										{
+											ci.HTTPVersionReceived = 10;
+										}
+									}
+								}
+								else if (HeaderKey == TEXT("Transfer-Encoding"))
+								{
+									if (HeaderValue.Find(TEXT("chunked"), ESearchCase::IgnoreCase, ESearchDir::FromStart, INDEX_NONE) != INDEX_NONE)
+									{
+										ci.bIsChunked = true;
+									}
+								}
+							}
+						}
+
+						// Content length needs a bit of special handling.
+						if (Handle->ActiveResponse.NumSubRangeRequest == 0)
+						{
+							bool bHaveSize = false;
+							// Is there a document size from a Content-Range header?
+							if (!ci.ContentRangeHeader.IsEmpty())
+							{
+								FParams::FRange ContentRange;
+								if (ContentRange.ParseFromContentRangeResponse(ci.ContentRangeHeader))
+								{
+									int64 ds = ContentRange.GetDocumentSize();
+									// Was the request for a range or the entire document?
+									if (Handle->ActiveResponse.OriginalRange.IsEverything())
+									{
+										ci.ContentLength = ds;
+									}
+									else
+									{
+										// A range was requested. Was it an open ended range?
+										if (Handle->ActiveResponse.OriginalRange.IsOpenEnded())
+										{
+											// Content size is the document size minus the start.
+											ci.ContentLength = ds >= 0 ? ds - Handle->ActiveResponse.OriginalRange.GetStart() : -1;
+										}
+										else
+										{
+											// Request was for an actual range.
+											int64 end = Handle->ActiveResponse.OriginalRange.GetEndIncluding() + 1;
+											if (ds >= 0 && end > ds)
+											{
+												end = ds;
+											}
+											ci.ContentLength = end - Handle->ActiveResponse.OriginalRange.GetStart();
+										}
+									}
+									bHaveSize = true;
+								}
+							}
+							if (!bHaveSize)
+							{
+								if (!ci.ContentLengthHeader.IsEmpty())
+								{
+									// Parse from "Content-Length: " header
+									LexFromString(ci.ContentLength, *ci.ContentLengthHeader.Mid(15));
+								}
+							}
+						}
+					}
+
+					if (EvData.bComplete && Response.IsValid())
+					{
+						/*
+											// Get the effective URL. This may however be the original URL that was requested.
+											ci.EffectiveURL.assign(TCHAR_TO_ANSI(*Response->GetURL()));
+
+									Taken out because this does not work on iOS and Mac. The engine HTTP module implementation there return an *empty* URL
+									for some reason which is very bad since user code makes use of the effective URL.
+									In absence of the actual effective URL (after redirection) we set the URL we made the request with.
+						*/
+						ci.EffectiveURL = Request->Parameters.URL;
+						ci.StatusInfo.HTTPStatus = ResponseCode;
+						// Update the time to first byte only on the first sub range request. This is the anchor for the request as a whole.
+						if (Handle->ActiveResponse.NumSubRangeRequest == 0)
+						{
+							ci.TimeUntilFirstByte = ci.TimeUntilConnected;
+							ci.TimeForDNSResolve = 0.001;			// something not zero for the lack of an actual value.
+						}
+						if (ci.HTTPVersionReceived == 0)
+						{
+							ci.HTTPVersionReceived = 11;			// for lack of an actual value let's assume it is HTTP/1.1
+						}
+						ci.NumberOfRedirections = 0;				// no official value to be had
+
+						// Check that the data we are receiving from a range request is actually that range!
+						if (!Handle->bResponseReceived && Request->Parameters.Range.IsSet() && !Request->Parameters.Range.IsEverything())
+						{
+							// Result needs to be "206 - partial content" and a content range header must be present.
+							if (ci.StatusInfo.HTTPStatus != 206 || ci.ContentRangeHeader.IsEmpty())
+							{
+								// Note that we did not receive what we thought to get. This gets checked for below.
+								ci.bResponseNotRanged = true;
+							}
+						}
+					}
+
+					// Update receive stats.
+					Handle->LastTimeDataReceived = Now;
+					ci.BytesReadSoFar = Handle->BytesReadSoFar + EvData.BytesReceived;
+
+					// Is the transfer done?
+					if (EvData.bComplete)
+					{
+						bool bFailure = false;
+						bool bContentRangeOk = true;
+						// If there is a Content-Range header we must be able to parse it.
+						if (!ci.ContentRangeHeader.IsEmpty())
+						{
+							bContentRangeOk = Handle->ActiveResponse.ReceivedContentRange.ParseFromContentRangeResponse(ci.ContentRangeHeader);
+						}
+
+						// Success or failure?
+						if (EvData.Status == EHttpRequestStatus::Succeeded && bHTTPResponseSuccess && ci.bResponseNotRanged == false && bContentRangeOk)
+						{
+							// Active responses are handed out in HandleHTTPResponses since they may not fit into the target receive ring buffer
+							// at once and have to be passed out in pieces.
+							if (Response.IsValid())
+							{
+								Handle->ActiveResponse.Response = Response;
+							}
+							else
+							{
+								// Response not valid?
+								bFailure = true;
+							}
+						}
+						else
+						{
+							bFailure = true;
+							ci.StatusInfo.ErrorDetail.SetError(UEMEDIA_ERROR_READ_ERROR).SetFacility(Facility::EFacility::HTTPReader).SetCode(ERRCODE_HTTPMODULE_FAILURE);
+							// Set up a bit more verbose error message
+							if (ci.StatusInfo.HTTPStatus >= 400)
+							{
+								ci.StatusInfo.ErrorCode = ERRCODE_HTTP_RETURNED_ERROR;
+								ci.StatusInfo.ErrorDetail.SetMessage(FString::Printf(TEXT("HTTP returned status %d"), ci.StatusInfo.HTTPStatus));
+							}
+							else if (ci.bResponseNotRanged)
+							{
+								ci.StatusInfo.ErrorCode = ERRCODE_HTTP_RANGE_ERROR;
+								ci.StatusInfo.ErrorDetail.SetMessage(FString::Printf(TEXT("Did not receive HTTP 206 for range request")));
+							}
+							else if (!bContentRangeOk)
+							{
+								ci.StatusInfo.ErrorCode = ERRCODE_HTTP_RANGE_ERROR;
+								ci.StatusInfo.ErrorDetail.SetMessage(FString::Printf(TEXT("Failed to parse HTTP response header \"%s\""), *ci.ContentRangeHeader));
+							}
+							else
+							{
+								ci.StatusInfo.ErrorCode = ERRCODE_WRITE_ERROR;
+								ci.StatusInfo.bReadError = true;
+								ci.StatusInfo.ErrorDetail.SetMessage(FString::Printf(TEXT("Read error after receiving %lld of %lld bytes"), (long long int)ci.BytesReadSoFar, (long long int)ci.ContentLength));
+							}
+						}
+						Handle->bResponseReceived = true;
+						ci.RequestEndTime = ci.StatusInfo.OccurredAtUTC = Now;
+						// Successful responses are processed outside this method. See above.
+						if (bFailure)
+						{
+							RequestsCompleted.Enqueue(Request);
+						}
+					}
+				}
+			}
+		}
+	}
+
+	void FElectraHttpManager::HandleHTTPResponses(const FTimeValue& Now)
+	{
+		for(TMap<FHandle*, TSharedPtrTS<FRequest>>::TIterator It = ActiveRequests.CreateIterator(); It; ++It)
+		{
+			FHandle* Handle = It.Key();
+			if (Handle->HandleType == FHandle::EHandleType::HTTPHandle)
+			{
+				TSharedPtrTS<FRequest> Request = It.Value();
+				HTTP::FConnectionInfo& ci = Request->ConnectionInfo;
+
+				// Active response?
+				if (Handle->ActiveResponse.Response.IsValid())
+				{
+					bool bHasFinished = false;
+
+					// Receive buffer still there?
+					TSharedPtrTS<FReceiveBuffer> ReceiveBuffer = Request->ReceiveBuffer.Pin();
+					if (ReceiveBuffer.IsValid())
+					{
+						// Is it to be used as a linear buffer?
+						if (!ReceiveBuffer->bEnableRingbuffer)
+						{
+							const TArray<uint8>& Content = Handle->ActiveResponse.Response->GetContent();
+							// Set the actual content length now. What we got in the header before may have been incorrect if the content was zipped
+							// or was transferred using chunked content encoding.
+							ci.ContentLength = Content.Num();
+							if (ReceiveBuffer->Buffer.EnlargeTo(ci.ContentLength))
+							{
+								if (ReceiveBuffer->Buffer.PushData(Content.GetData(), Content.Num()))
+								{
+									ci.BytesReadSoFar = ci.ContentLength;
+								}
+								else
+								{
+									ci.bWasAborted = true;
+								}
+							}
+							else
+							{
+								// When OOM we just claim we've been aborted.
+								ci.bWasAborted = true;
+							}
+							bHasFinished = true;
+						}
+						else
+						{
+							// The ring buffer must have been set up to a size the client wants it to have.
+							check(ReceiveBuffer->Buffer.Capacity() > 0);
+
+							// How much room is there available?
+							int32 NumAvail = ReceiveBuffer->Buffer.Avail();
+							if (NumAvail)
+							{
+								const TArray<uint8>& Content = Handle->ActiveResponse.Response->GetContent();
+								int32 NumToGo = Content.Num() - Handle->ActiveResponse.NumBytesPassedOut;
+								int32 NumToCopy = NumAvail < NumToGo ? NumAvail : NumToGo;
+								ReceiveBuffer->Buffer.PushData(Content.GetData() + Handle->ActiveResponse.NumBytesPassedOut, NumToCopy);
+								Handle->ActiveResponse.NumBytesPassedOut += NumToCopy;
+								// All copied out now?
+								if (Handle->ActiveResponse.NumBytesPassedOut >= Content.Num())
+								{
+									// Check if this was a sub ranged request and if there is still data to go for the original request.
+									if (Handle->ActiveResponse.SizeRemaining() == 0)
+									{
+										// All done now.
+										bHasFinished = true;
+									}
+									else
+									{
+										// Still another sub range to go.
+										if (PrepareHTTPModuleHandle(Now, Handle, Request, false))
+										{
+											if (!Handle->HttpRequest->ProcessRequest())
+											{
+												ci.StatusInfo.ErrorDetail.SetError(UEMEDIA_ERROR_INTERNAL).SetFacility(Facility::EFacility::HTTPReader).SetMessage("HTTP request failed on ProcessRequest()");
+												bHasFinished = true;
+											}
+										}
+										else
+										{
+											// 
+											ci.StatusInfo.ErrorCode = ERRCODE_WRITE_ERROR;
+											ci.StatusInfo.bReadError = true;
+											ci.StatusInfo.ErrorDetail.SetMessage(FString::Printf(TEXT("Error setting up the next sub range request")));
+											bHasFinished = true;
+										}
+									}
+								}
+							}
+						}
+					}
+					else
+					{
+						// With the receive buffer having been released we can abort the transfer.
+						ci.bWasAborted = true;
+						if (Handle->HttpsRequestCallbackWrapper.IsValid())
+						{
+							Handle->HttpsRequestCallbackWrapper->Unbind();
+						}
+						bHasFinished = true;
+					}
+					if (bHasFinished)
+					{
+						ci.RequestEndTime = ci.StatusInfo.OccurredAtUTC = Now;
+						Handle->ActiveResponse.Response.Reset();
+						RequestsCompleted.Enqueue(Request);
+					}
+				}
+			}
+		}
+	}
+
+
+
+	void FElectraHttpManager::OnProcessRequestCompleteInternal(FHttpRequestPtr Request, FHttpResponsePtr Response, bool bConnectedSuccessfully)
+	{
+		FScopeLock lock(&HttpModuleEventLock);
+		FHttpEventData& Evt = HttpModuleEventMap.FindOrAdd(Request);
+		Evt.bComplete = true;
+		Evt.bConnectedSuccessfully = bConnectedSuccessfully;
+		Evt.Status = Request->GetStatus();
+		Evt.Response = Response;
+	}
+
+	void FElectraHttpManager::OnRequestProgressInternal(FHttpRequestPtr Request, int32 BytesSent, int32 BytesReceived)
+	{
+		FScopeLock lock(&HttpModuleEventLock);
+		FHttpEventData& Evt = HttpModuleEventMap.FindOrAdd(Request);
+		Evt.BytesReceived = BytesReceived;
+		Evt.BytesSent = BytesSent;
+		Evt.Status = Request->GetStatus();
+	}
+
+	void FElectraHttpManager::OnRequestHeaderReceivedInternal(FHttpRequestPtr Request, const FString& HeaderName, const FString& NewHeaderValue)
+	{
+		FScopeLock lock(&HttpModuleEventLock);
+		FHttpEventData& Evt = HttpModuleEventMap.FindOrAdd(Request);
+		Evt.Headers.Add({HeaderName, NewHeaderValue});
+	}
+
+
+
+
+
+
+	void FElectraHttpManager::FLocalByteStream::SetConnected(TSharedPtrTS<FRequest> Request)
+	{
+		if (!bIsConnected)
+		{
+			bIsConnected = true;
+			// Go through the notions of this being a network request.
+			Request->ConnectionInfo.bIsConnected = true;
+			Request->ConnectionInfo.bHaveResponseHeaders = true;
+			Request->ConnectionInfo.ContentType = "application/octet-stream";
+			Request->ConnectionInfo.EffectiveURL = Request->Parameters.URL;
+			Request->ConnectionInfo.HTTPVersionReceived = 11;
+			Request->ConnectionInfo.bIsChunked = false;
+
+			// Range request?
+			if (!Request->Parameters.Range.IsSet())
+			{
+				Request->ConnectionInfo.ContentLength = Archive->TotalSize();
+				Request->ConnectionInfo.StatusInfo.HTTPStatus = 200;
+				FileStartOffset = 0;
+				FileSize = Archive->TotalSize();
+				FileSizeToGo = FileSize;
+				Request->ConnectionInfo.ContentLengthHeader = FString::Printf(TEXT("Content-Length: %lld"), (long long int)FileSize);
+			}
+			else
+			{
+				int64 fs = Archive->TotalSize();
+				int64 off = 0;
+				int64 end = fs - 1;
+				// For now we support partial data only from the beginning of the file, not the end (aka, seek_set and not seek_end)
+				check(Request->Parameters.Range.Start >= 0);
+				if (Request->Parameters.Range.Start >= 0)
+				{
+					off = Request->Parameters.Range.Start;
+					if (off < fs)
+					{
+						end = Request->Parameters.Range.EndIncluding;
+						if (end < 0 || end >= fs)
+						{
+							end = fs - 1;
+						}
+						int64 numBytes = end - off + 1;
+
+						Request->ConnectionInfo.ContentLength = numBytes;
+						Request->ConnectionInfo.StatusInfo.HTTPStatus = 206;
+						FileStartOffset = off;
+						FileSize = Archive->TotalSize();
+						FileSizeToGo = numBytes;
+						Request->ConnectionInfo.ContentLengthHeader = FString::Printf(TEXT("Content-Length: %lld"), (long long int)numBytes);
+						Request->ConnectionInfo.ContentRangeHeader = FString::Printf(TEXT("Content-Range: bytes %lld-%lld/%lld"), (long long int)off, (long long int)end, (long long int)fs);
+						Archive->Seek(off);
+					}
+					else
+					{
+						Request->ConnectionInfo.StatusInfo.HTTPStatus = 416;		// Range not satisfiable
+						Request->ConnectionInfo.ContentRangeHeader = FString::Printf(TEXT("Content-Range: bytes */%lld"), (long long int)fs);
+					}
+				}
+			}
+		}
+	}
+
+	int32 FElectraHttpManager::FLocalByteStream::Read(TSharedPtrTS<FReceiveBuffer> RcvBuffer, TSharedPtrTS<FRequest> Request)
+	{
+		// How much room do we have in the buffer?
+		int32 blkAv1, blkAv2;
+		uint8* blkData1;
+		uint8* blkData2;
+		int32 MaxAvail = RcvBuffer->Buffer.PushBlockOpen(blkData1, blkAv1, blkData2, blkAv2);
+		// How much can we read from the file?
+		int32 NumToRead = FileSizeToGo < MaxAvail ? FileSizeToGo : MaxAvail;
+		if (NumToRead)
+		{
+			if (NumToRead <= blkAv1)
+			{
+				Archive->Serialize(blkData1, NumToRead);
+			}
+			else
+			{
+				Archive->Serialize(blkData1, blkAv1);
+				Archive->Serialize(blkData2, NumToRead - blkAv1);
+			}
+			Request->ConnectionInfo.BytesReadSoFar += NumToRead;
+		}
+		RcvBuffer->Buffer.PushBlockClose(NumToRead);
+		FileSizeToGo -= NumToRead;
+		return NumToRead;
+	}
+
+
+
+} // namespace Electra
+
+
diff --git a/Engine/Plugins/Media/ElectraPlayer/Source/ElectraPlayerRuntime/Private/Runtime/HTTP/HTTPManager.h b/Engine/Plugins/Media/ElectraPlayer/Source/ElectraPlayerRuntime/Private/Runtime/HTTP/HTTPManager.h
new file mode 100644
--- /dev/null
+++ b/Engine/Plugins/Media/ElectraPlayer/Source/ElectraPlayerRuntime/Private/Runtime/HTTP/HTTPManager.h
@@ -0,0 +1,388 @@
+// Copyright Epic Games, Inc. All Rights Reserved.
+
+#pragma once
+
+#include "PlayerCore.h"
+
+#include "StreamDataBuffer.h"
+#include "OptionalValue.h"
+
+#include "ParameterDictionary.h"
+#include "ErrorDetail.h"
+
+
+DECLARE_LOG_CATEGORY_EXTERN(LogElectraHTTPManager, Log, All);
+
+namespace Electra
+{
+	//
+	class IPlayerSessionServices;
+
+
+	namespace HTTP
+	{
+
+		struct FStatusInfo
+		{
+			FStatusInfo()
+				: HTTPStatus(0)
+				, ErrorCode(0)
+				, ConnectionTimeoutAfterMilliseconds(0)
+				, NoDataTimeoutAfterMilliseconds(0)
+				, bRedirectionError(false)
+				, bReadError(false)
+			{
+			}
+			FErrorDetail	ErrorDetail;
+			FTimeValue		OccurredAtUTC;
+			int32			HTTPStatus;
+			int32			ErrorCode;
+			int32			ConnectionTimeoutAfterMilliseconds;
+			int32			NoDataTimeoutAfterMilliseconds;
+			bool			bRedirectionError;
+			bool			bReadError;
+		};
+
+		struct FRetryInfo
+		{
+			FRetryInfo()
+			{
+				AttemptNumber = 0;
+				MaxAttempts = 0;
+			}
+			int32						AttemptNumber;
+			int32						MaxAttempts;
+			TArray<FStatusInfo>			PreviousFailureStates;
+		};
+
+
+		struct FHTTPHeader
+		{
+			FString	Header;
+			FString	Value;
+		};
+
+		struct FConnectionInfo
+		{
+			struct FThroughput
+			{
+				FThroughput()
+				{
+					Clear();
+				}
+				void Clear()
+				{
+					ActiveReadTime.SetToZero();
+					LastCheckTime.SetToInvalid();
+					LastReadCount = 0;
+					AccumulatedBytesPerSec = 0;
+					AccumulatedBytesPerSecWithoutFirst = 0;
+					NumAccumulatedBytesPerSec = 0;
+					bIsFirst = true;
+				}
+				int64 GetThroughput() const
+				{
+					if (NumAccumulatedBytesPerSec > 1)
+					{
+						return AccumulatedBytesPerSecWithoutFirst / (NumAccumulatedBytesPerSec - 1);
+					}
+					else if (NumAccumulatedBytesPerSec)
+					{
+						return AccumulatedBytesPerSec / NumAccumulatedBytesPerSec;
+					}
+					else
+					{
+						// Did not get to collect a throughput sample.
+						return 0;
+					}
+				}
+				FTimeValue				ActiveReadTime;
+				FTimeValue				LastCheckTime;
+				int64					LastReadCount;
+				int64					AccumulatedBytesPerSec;
+				int64					AccumulatedBytesPerSecWithoutFirst;
+				int32					NumAccumulatedBytesPerSec;
+				bool					bIsFirst;
+			};
+
+			TArray<FHTTPHeader>					ResponseHeaders;					//!< Response headers
+			FString								EffectiveURL;						//!< Effective URL after all redirections
+			FString								ContentType;						//!< Content-Type header
+			FString  							ContentRangeHeader;
+			FString  							ContentLengthHeader;
+			FTimeValue							RequestStartTime;					//!< Time at which the request was started
+			FTimeValue							RequestEndTime;						//!< Time at which the request ended
+			double								TimeForDNSResolve;					//!< Time it took to resolve DNS
+			double								TimeUntilConnected;					//!< Time it took until connected to the server
+			double								TimeUntilFirstByte;					//!< Time it took until received the first response data byte
+			int64								ContentLength;						//!< Content length, if known. Chunked transfers may have no length (set to -1). Compressed file will store compressed size here!
+			int64								BytesReadSoFar;						//!< Number of bytes read so far.
+			int32								NumberOfRedirections;				//!< Number of redirections performed
+			int32								HTTPVersionReceived;				//!< Version of HTTP header received (10 = 1.0, 11=1.1, 20=2.0)
+			bool								bIsConnected;						//!< true once connected to the server
+			bool								bHaveResponseHeaders;				//!< true when response headers have been received
+			bool								bIsChunked;							//!< true if response is received in chunks
+			bool								bWasAborted;						//!< true if transfer was aborted.
+			bool								bHasFinished;						//!< true once the connection is closed regardless of state.
+			bool								bResponseNotRanged;					//!< true if the response is not a range as was requested.
+			FStatusInfo							StatusInfo;
+			TSharedPtrTS<FRetryInfo>			RetryInfo;
+
+			FThroughput							Throughput;
+
+			FConnectionInfo()
+			{
+				TimeForDNSResolve = 0.0;
+				TimeUntilConnected = 0.0;
+				TimeUntilFirstByte = 0.0;
+				NumberOfRedirections = 0;
+				HTTPVersionReceived = 0;
+				ContentLength = -1;
+				BytesReadSoFar = 0;
+				bIsConnected = false;
+				bHaveResponseHeaders = false;
+				bIsChunked = false;
+				bWasAborted = false;
+				bHasFinished = false;
+				bResponseNotRanged = false;
+			}
+
+			FConnectionInfo& CopyFrom(const FConnectionInfo& rhs)
+			{
+				ResponseHeaders = rhs.ResponseHeaders;
+				EffectiveURL = rhs.EffectiveURL;
+				ContentType = rhs.ContentType;
+				ContentRangeHeader = rhs.ContentRangeHeader;
+				ContentLengthHeader = rhs.ContentLengthHeader;
+				RequestStartTime = rhs.RequestStartTime;
+				RequestEndTime = rhs.RequestEndTime;
+				TimeForDNSResolve = rhs.TimeForDNSResolve;
+				TimeUntilConnected = rhs.TimeUntilConnected;
+				TimeUntilFirstByte = rhs.TimeUntilFirstByte;
+				ContentLength = rhs.ContentLength;
+				BytesReadSoFar = rhs.BytesReadSoFar;
+				NumberOfRedirections = rhs.NumberOfRedirections;
+				HTTPVersionReceived = rhs.HTTPVersionReceived;
+				bIsConnected = rhs.bIsConnected;
+				bHaveResponseHeaders = rhs.bHaveResponseHeaders;
+				bIsChunked = rhs.bIsChunked;
+				bWasAborted = rhs.bWasAborted;
+				bHasFinished = rhs.bHasFinished;
+				bResponseNotRanged = rhs.bResponseNotRanged;
+				StatusInfo = rhs.StatusInfo;
+				Throughput = rhs.Throughput;
+				if (rhs.RetryInfo.IsValid())
+				{
+					RetryInfo = MakeSharedTS<FRetryInfo>(*rhs.RetryInfo);
+				}
+				return *this;
+			}
+		};
+
+	} // namespace HTTP
+
+
+	class IElectraHttpManager
+	{
+	public:
+		static TSharedPtrTS<IElectraHttpManager> Create();
+
+		virtual ~IElectraHttpManager() = default;
+
+		struct FRequest;
+
+		typedef Electra::FastDelegate1<const FRequest*, int32> FProgressDelegate;
+		typedef Electra::FastDelegate1<const FRequest*> FCompletionDelegate;
+		struct FProgressListener
+		{
+			~FProgressListener()
+			{
+				Clear();
+			}
+			void Clear()
+			{
+				FMediaCriticalSection::ScopedLock lock(Lock);
+				ProgressDelegate.clear();
+				CompletionDelegate.clear();
+			}
+			int32 CallProgressDelegate(const FRequest* Request)
+			{
+				FMediaCriticalSection::ScopedLock lock(Lock);
+				if (!ProgressDelegate.empty())
+				{
+					return ProgressDelegate(Request);
+				}
+				return 0;
+			}
+			void CallCompletionDelegate(const FRequest* Request)
+			{
+				FMediaCriticalSection::ScopedLock lock(Lock);
+				if (!CompletionDelegate.empty())
+				{
+					CompletionDelegate(Request);
+				}
+			}
+			FMediaCriticalSection	Lock;
+			FProgressDelegate		ProgressDelegate;
+			FCompletionDelegate		CompletionDelegate;
+		};
+
+		struct FReceiveBuffer
+		{
+			FPODRingbuffer		Buffer;
+			bool				bEnableRingbuffer = false;
+		};
+
+		struct FParams
+		{
+			struct FRange
+			{
+				void Reset()
+				{
+					Start = -1;
+					EndIncluding = -1;
+					DocumentSize = -1;
+				}
+				bool IsSet() const
+				{
+					return Start != -1 || EndIncluding != -1;
+				}
+				//! Check if the range would result in "0-" for the entire file in which case we don't need to use range request.
+				bool IsEverything() const
+				{
+					return Start <= 0 && EndIncluding < 0;
+				}
+				FString GetString() const
+				{
+					FString s, e, d(TEXT("-"));
+					if (Start >= 0)
+					{
+						// An explicit range?
+						if (EndIncluding >= Start)
+						{
+							return FString::Printf(TEXT("%lld-%lld"), (long long int)Start, (long long int)EndIncluding);
+						}
+						// Start to end, whereever that is.
+						else
+						{
+							return FString::Printf(TEXT("%lld-"), (long long int)Start);
+						}
+					}
+					// Everything
+					return FString(TEXT("0-"));
+				}
+				// Returns the number of bytes in the range, which must be fully specified. An unset or partially open range will return -1.
+				int64 GetNumberOfBytes() const
+				{
+					if (Start >= 0 && EndIncluding >= 0)
+					{
+						return EndIncluding - Start + 1;
+					}
+					return -1;
+				}
+				int64 GetStart() const
+				{
+					return Start;
+				}
+				int64 GetEndIncluding() const
+				{
+					return EndIncluding;
+				}
+				bool IsOpenEnded() const
+				{
+					return GetEndIncluding() < 0;
+				}
+				void SetStart(int64 InStart)
+				{
+					Start = InStart;
+				}
+				void SetEndIncluding(int64 InEndIncluding)
+				{
+					EndIncluding = InEndIncluding;
+				}
+				int64 GetDocumentSize() const
+				{
+					return DocumentSize;
+				}
+				bool ParseFromContentRangeResponse(const FString& ContentRangeHeader)
+				{
+					// Examples: <unit> <range-start>-<range-end>/<size>
+					//   Content-Range: bytes 26151-157222/7594984
+					//   Content-Range: bytes 26151-157222/*
+					//   Content-Range: bytes */7594984
+					//
+					Start = -1;
+					EndIncluding = -1;
+					DocumentSize = -1;
+					FString rh = ContentRangeHeader;
+					// In case the entire header is given, remove the header including the separating colon and space.
+					rh.RemoveFromStart(TEXT("Content-Range: "), ESearchCase::CaseSensitive);
+					// Split into parts
+					TArray<FString> Parts;
+					const TCHAR* const Delims[3] = {TEXT(" "),TEXT("-"),TEXT("/")};
+					int32 NumParts = rh.ParseIntoArray(Parts, Delims, 3);
+					if (NumParts)
+					{
+						if (Parts[0] == TEXT("bytes"))
+						{
+							Parts.RemoveAt(0);
+						}
+						// We should now be left with 3 remaining results, the start, end and document size.
+						// The case where we get "*/<size>" we treat as invalid.
+						if (Parts.Num() == 3)
+						{
+							if (Parts[0].IsNumeric())
+							{
+								LexFromString(Start, *Parts[0]);
+								if (Parts[1].IsNumeric())
+								{
+									LexFromString(EndIncluding, *Parts[1]);
+									if (Parts[2].IsNumeric())
+									{
+										LexFromString(DocumentSize, *Parts[2]);
+										return true;
+									}
+									else if (Parts[2] == TEXT("*"))
+									{
+										DocumentSize = -1;
+										return true;
+									}
+								}
+							}
+						}
+					}
+					UE_LOG(LogElectraHTTPManager, Error, TEXT("Failed to parse Content-Range HTTP response header \"%s\""), *ContentRangeHeader);
+					return false;
+				}
+				int64			Start = -1;
+				int64			EndIncluding = -1;
+				int64			DocumentSize = -1;
+			};
+			FString								URL;							//!< URL
+			FRange								Range;							//!< Optional request range
+			int32								SubRangeRequestSize = 0;		//!< If not 0 the size to break the request into smaller range requests into.
+			TArray<HTTP::FHTTPHeader>			RequestHeaders;					//!< Request headers
+			TMediaOptionalValue<FString>		AcceptEncoding;					//!< Optional accepted encoding
+			FTimeValue							ConnectTimeout;					//!< Optional timeout for connecting to the server
+			FTimeValue							NoDataTimeout;					//!< Optional timeout when no data is being received
+		};
+
+
+		struct FRequest
+		{
+			FParams								Parameters;
+			FParamDict							Options;
+			HTTP::FConnectionInfo				ConnectionInfo;
+			TWeakPtrTS<FReceiveBuffer>			ReceiveBuffer;
+			TWeakPtrTS<FProgressListener>		ProgressListener;
+		};
+
+		virtual void AddRequest(TSharedPtrTS<FRequest> Request) = 0;
+		virtual void RemoveRequest(TSharedPtrTS<FRequest> Request) = 0;
+
+	};
+
+
+} // namespace Electra
+
+
diff --git a/Engine/Plugins/Media/ElectraPlayer/Source/ElectraPlayerRuntime/Private/Runtime/OptionalValue.h b/Engine/Plugins/Media/ElectraPlayer/Source/ElectraPlayerRuntime/Private/Runtime/OptionalValue.h
new file mode 100644
--- /dev/null
+++ b/Engine/Plugins/Media/ElectraPlayer/Source/ElectraPlayerRuntime/Private/Runtime/OptionalValue.h
@@ -0,0 +1,28 @@
+// Copyright Epic Games, Inc. All Rights Reserved.
+
+#pragma once
+
+#include "PlayerCore.h"
+
+
+/**
+ * "Optional" value. Includes a boolean to indicate whether or not it is set.
+**/
+template<typename T>
+struct TMediaOptionalValue
+{
+	TMediaOptionalValue() : ValueSet(false) {}
+	TMediaOptionalValue(const T& v) : OptionalValue(v), ValueSet(true) {}
+	TMediaOptionalValue(const TMediaOptionalValue& rhs) : OptionalValue(rhs.OptionalValue), ValueSet(rhs.ValueSet) {}
+	void Set(const T& v) { OptionalValue = v; ValueSet = true; }
+	void SetIfNot(const T& v) { if (!IsSet()) Set(v); }
+	bool IsSet(void) const { return(ValueSet); }
+	const T& Value(void) const { return(OptionalValue); }
+	T GetWithDefault(const T& Default) const { return ValueSet ? OptionalValue : Default; }
+	void Reset(void) { ValueSet = false; }
+private:
+	// Hide assignment from public view
+	TMediaOptionalValue& operator = (const T& v) { Set(v); return(*this); }
+	T		OptionalValue;
+	bool	ValueSet;
+};
