ElectraPlayer: RELEASE. Moved from NFL to public. Note that this is done via a Delete and "Add New" operation as discussed.

diff --git a/Engine/Plugins/Media/ElectraPlayer/Source/ElectraPlayerRuntime/Private/Runtime/Player/AdaptiveStreamingPlayer.cpp b/Engine/Plugins/Media/ElectraPlayer/Source/ElectraPlayerRuntime/Private/Runtime/Player/AdaptiveStreamingPlayer.cpp
new file mode 100644
--- /dev/null
+++ b/Engine/Plugins/Media/ElectraPlayer/Source/ElectraPlayerRuntime/Private/Runtime/Player/AdaptiveStreamingPlayer.cpp
@@ -0,0 +1,3226 @@
+// Copyright Epic Games, Inc. All Rights Reserved.
+
+#include "Player/AdaptiveStreamingPlayerInternal.h"
+#include "ElectraPlayerPrivate.h"
+#include "PlayerRuntimeGlobal.h"
+
+#include "StreamAccessUnitBuffer.h"
+#include "Player/AdaptiveStreamingPlayerABR.h"
+#include "Player/PlayerLicenseKey.h"
+
+#include "HAL/LowLevelMemTracker.h"
+
+
+
+namespace Electra
+{
+FAdaptiveStreamingPlayer *FAdaptiveStreamingPlayer::PointerToLatestPlayer;
+
+TSharedPtr<IAdaptiveStreamingPlayer, ESPMode::ThreadSafe> IAdaptiveStreamingPlayer::Create(const IAdaptiveStreamingPlayer::FCreateParam& InCreateParameters)
+{
+	return MakeShared<FAdaptiveStreamingPlayer, ESPMode::ThreadSafe>(InCreateParameters);
+}
+
+void IAdaptiveStreamingPlayer::DebugHandle(void* pPlayer, void (*debugDrawPrintf)(void* pPlayer, const char *pFmt, ...))
+{
+	FAdaptiveStreamingPlayer::DebugHandle(pPlayer, debugDrawPrintf);
+}
+
+
+//---------------------------------------------------------------------------------------------------------------------------------
+//---------------------------------------------------------------------------------------------------------------------------------
+//---------------------------------------------------------------------------------------------------------------------------------
+
+
+//-----------------------------------------------------------------------------
+/**
+ * CTOR
+ */
+FAdaptiveStreamingPlayer::FAdaptiveStreamingPlayer(const IAdaptiveStreamingPlayer::FCreateParam& InCreateParameters)
+{
+	Electra::AddActivePlayerInstance();
+
+	ExternalPlayerGUID		= InCreateParameters.ExternalPlayerGUID;
+
+	VideoRender.Renderer    = InCreateParameters.VideoRenderer;
+	AudioRender.Renderer    = InCreateParameters.AudioRenderer;
+	CurrentState   		    = EPlayerState::eState_Idle;
+	PipelineState  		    = EPipelineState::ePipeline_Stopped;
+	DecoderState   		    = EDecoderState::eDecoder_Paused;
+	StreamState			    = EStreamState::eStream_Running;
+	PlaybackRate   		    = 0.0;
+	StreamReaderHandler     = nullptr;
+	bRebufferPending   	    = false;
+	Manifest  				= nullptr;
+	ManifestReader			= nullptr;
+	LastBufferingState 		= EPlayerState::eState_Buffering;
+	bIsPlayStart   			= true;
+	bIsClosing 				= false;
+	StartAtTime.Reset();
+	bHaveVideoReader.Reset();
+	bHaveAudioReader.Reset();
+	bHaveTextReader.Reset();
+	RebufferDetectedAtPlayPos.SetToInvalid();
+
+	CurrentPlaybackSequenceID = 0;
+
+	DataAvailabilityStateVid.StreamType = EStreamType::Video;
+	DataAvailabilityStateVid.Availability = Metrics::FDataAvailabilityChange::EAvailability::DataNotAvailable;
+	DataAvailabilityStateAud.StreamType = EStreamType::Audio;
+	DataAvailabilityStateAud.Availability = Metrics::FDataAvailabilityChange::EAvailability::DataNotAvailable;
+	DataAvailabilityStateTxt.StreamType = EStreamType::Subtitle;
+	DataAvailabilityStateTxt.Availability = Metrics::FDataAvailabilityChange::EAvailability::DataNotAvailable;
+
+	// Create the UTC time handler for this player session
+	SynchronizedUTCTime = ISynchronizedUTCTime::Create();
+	SynchronizedUTCTime->SetTime(MEDIAutcTime::Current());
+
+	// Create the render clock
+	RenderClock = MakeSharedTS<FMediaRenderClock>();
+
+	bShouldBePaused  = false;
+	bShouldBePlaying = false;
+	bSeekPending	 = false;
+
+	CurrentLoopState.Reset();
+	NextLoopStates.Clear();
+
+	BitrateCeiling = 0;
+	VideoResolutionLimitWidth = 0;
+	VideoResolutionLimitHeight = 0;
+	MaxVideoTextureResolutionLimitHeight = 0;
+
+	TMediaInterlockedExchangePointer(PointerToLatestPlayer, this);
+
+	// Create the worker thread.
+	StartWorkerThread();
+}
+
+
+
+//-----------------------------------------------------------------------------
+/**
+ * DTOR
+ */
+FAdaptiveStreamingPlayer::~FAdaptiveStreamingPlayer()
+{
+	TMediaInterlockedExchangePointer(PointerToLatestPlayer, (FAdaptiveStreamingPlayer*)0);
+
+	FMediaEvent doneSig;
+	WorkerThread.SendCloseMessage(&doneSig);
+	doneSig.Wait();
+	StopWorkerThread();
+
+	MetricListenerCriticalSection.Lock();
+	MetricListeners.Empty();
+	MetricListenerCriticalSection.Unlock();
+
+	RenderClock.Reset();
+	delete SynchronizedUTCTime;
+
+	Electra::RemoveActivePlayerInstance();
+}
+
+
+
+//-----------------------------------------------------------------------------
+/**
+ * Returns the error event. If there is no error an empty pointer will be returned.
+ *
+ * @return
+ */
+FErrorDetail FAdaptiveStreamingPlayer::GetError() const
+{
+	FMediaCriticalSection::ScopedLock lock(DiagnosticsCriticalSection);
+	return LastErrorDetail;
+}
+
+
+
+
+
+
+
+
+
+
+
+//-----------------------------------------------------------------------------
+/**
+ * Updates the internally kept diagnostics value with the current
+ * decoder live values.
+ */
+void FAdaptiveStreamingPlayer::UpdateDiagnostics()
+{
+	DiagnosticsCriticalSection.Lock();
+	MultiStreamBufferVid.GetStats(VideoBufferStats.StreamBuffer);
+	MultiStreamBufferAud.GetStats(AudioBufferStats.StreamBuffer);
+	MultiStreamBufferTxt.GetStats(TextBufferStats.StreamBuffer);
+	DiagnosticsCriticalSection.Unlock();
+}
+
+
+
+
+//-----------------------------------------------------------------------------
+/**
+ * Initializes the player instance.
+ *
+ * @param Options
+ *
+ * @return
+ */
+bool FAdaptiveStreamingPlayer::Initialize(const FParamDict& Options)
+{
+	PlayerOptions = Options;
+
+	// Get the HTTP manager. This is a shared instance for all players.
+	HttpManager = IElectraHttpManager::Create();
+
+	// Create the ABR stream selector.
+	StreamSelector = IAdaptiveStreamSelector::Create(this, PlayerConfig.StreamSelectorConfig);
+	AddMetricsReceiver(StreamSelector.Get());
+
+	return true;
+}
+
+
+void FAdaptiveStreamingPlayer::SetStaticResourceProviderCallback(const TSharedPtr<IAdaptiveStreamingPlayerResourceProvider, ESPMode::ThreadSafe>& InStaticResourceProvider)
+{
+	StaticResourceProvider = InStaticResourceProvider;
+}
+
+void FAdaptiveStreamingPlayer::SetVideoDecoderResourceDelegate(const TSharedPtr<IVideoDecoderResourceDelegate, ESPMode::ThreadSafe>& InResourceDelegate)
+{
+	VideoDecoderResourceDelegate = InResourceDelegate;
+}
+
+
+void FAdaptiveStreamingPlayer::AddMetricsReceiver(IAdaptiveStreamingPlayerMetrics* InMetricsReceiver)
+{
+	if (InMetricsReceiver)
+	{  
+		FMediaCriticalSection::ScopedLock lock(MetricListenerCriticalSection);
+		if (MetricListeners.Find(InMetricsReceiver) == INDEX_NONE)
+		{
+			MetricListeners.Push(InMetricsReceiver);
+		}
+	}
+}
+
+void FAdaptiveStreamingPlayer::RemoveMetricsReceiver(IAdaptiveStreamingPlayerMetrics* InMetricsReceiver)
+{
+	if (InMetricsReceiver)
+	{
+		FMediaCriticalSection::ScopedLock lock(MetricListenerCriticalSection);
+		/*bool bRemoved =*/ MetricListeners.Remove(InMetricsReceiver);
+	}
+}
+
+
+
+//-----------------------------------------------------------------------------
+/**
+ * Dispatches an event to the listeners.
+ *
+ * @param pEvent
+ */
+void FAdaptiveStreamingPlayer::DispatchEvent(TSharedPtrTS<FMetricEvent> Event)
+{
+	if (EventDispatcher.IsValid())
+	{
+		Event->Player = SharedThis(this);
+		EventDispatcher->DispatchEvent(Event);
+	}
+}
+
+void FAdaptiveStreamingPlayer::DispatchEventAndWait(TSharedPtrTS<FMetricEvent> Event)
+{
+	if (EventDispatcher.IsValid())
+	{
+		Event->Player = SharedThis(this);
+		EventDispatcher->DispatchEventAndWait(Event);
+	}
+}
+
+
+
+
+//-----------------------------------------------------------------------------
+/**
+ * Starts the worker thread.
+ */
+void FAdaptiveStreamingPlayer::StartWorkerThread()
+{
+	if (!WorkerThread.bStarted)
+	{
+		WorkerThread.MediaThread.ThreadSetPriority(PlayerConfig.WorkerThread.Priority);
+		WorkerThread.MediaThread.ThreadSetCoreAffinity(PlayerConfig.WorkerThread.CoreAffinity);
+		WorkerThread.MediaThread.ThreadSetStackSize(PlayerConfig.WorkerThread.StackSize);
+		WorkerThread.MediaThread.ThreadSetName("ElectraPlayer::Worker");
+		WorkerThread.MediaThread.ThreadStart(Electra::MakeDelegate(this, &FAdaptiveStreamingPlayer::WorkerThreadFN));
+		WorkerThread.bStarted = true;
+	}
+}
+
+
+//-----------------------------------------------------------------------------
+/**
+ * Stops the worker thread.
+ */
+void FAdaptiveStreamingPlayer::StopWorkerThread()
+{
+	if (WorkerThread.bStarted)
+	{
+		WorkerThread.SendMessage(FWorkerThread::FMessage::EType::Quit);
+		WorkerThread.MediaThread.ThreadWaitDone();
+		WorkerThread.MediaThread.ThreadReset();
+		WorkerThread.bStarted = false;
+	}
+}
+
+
+
+
+
+void FAdaptiveStreamingPlayer::LoadManifest(const FString& ManifestURL)
+{
+	WorkerThread.SendLoadManifestMessage(ManifestURL, FString());
+}
+
+
+//-----------------------------------------------------------------------------
+/**
+ * Pauses playback
+ */
+void FAdaptiveStreamingPlayer::Pause()
+{
+	WorkerThread.SendPauseMessage();
+}
+
+
+
+//-----------------------------------------------------------------------------
+/**
+ * Resumes playback
+ */
+void FAdaptiveStreamingPlayer::Resume()
+{
+	WorkerThread.SendResumeMessage();
+}
+
+
+
+//-----------------------------------------------------------------------------
+/**
+ * Seek to a new position and play from there. This includes first playstart.
+ *
+ * @param NewPosition
+ */
+void FAdaptiveStreamingPlayer::SeekTo(const FSeekParam& NewPosition)
+{
+	TSharedPtrTS<IManifest> CurrentManifest = Manifest;
+	if (CurrentManifest.IsValid())
+	{
+		WorkerThread.SendSeekMessage(NewPosition);
+	}
+	else
+	{
+		PostLog(Facility::EFacility::Player, IInfoLog::ELevel::Error, FString::Printf(TEXT("API error: Cannot perform seek. No playlist loaded yet. Seeking is possible only when a valid playlist is available!")));
+	}
+}
+
+
+
+//-----------------------------------------------------------------------------
+/**
+ * Stops playback. Playback cannot be resumed. Final player events will be sent to registered listeners.
+ */
+void FAdaptiveStreamingPlayer::Stop()
+{
+	FMediaEvent doneSig;
+	WorkerThread.SendCloseMessage(&doneSig);
+	doneSig.Wait();
+}
+
+
+//-----------------------------------------------------------------------------
+/**
+ * Puts playback into loop mode if possible. Live streams cannot be made to loop as they have infinite duration.
+ *
+ * @param InLoopParams
+ */
+void FAdaptiveStreamingPlayer::SetLooping(const FLoopParam& InLoopParams)
+{
+	WorkerThread.SendLoopMessage(InLoopParams);
+}
+
+
+//-----------------------------------------------------------------------------
+/**
+ * Returns whether or not a manifest has been loaded and assigned yet.
+ *
+ * @return
+ */
+bool FAdaptiveStreamingPlayer::HaveMetadata() const
+{
+	return PlaybackState.GetHaveMetadata();
+}
+
+
+//-----------------------------------------------------------------------------
+/**
+ * Returns the duration of the video. Returns invalid time when there is nothing to play. Returns positive infinite for live streams.
+ *
+ * @return
+ */
+FTimeValue FAdaptiveStreamingPlayer::GetDuration() const
+{
+	return PlaybackState.GetDuration();
+}
+
+//-----------------------------------------------------------------------------
+/**
+ * Returns the current play position. Returns -1.0 when there is nothing to play.
+ *
+ * @return
+ */
+FTimeValue FAdaptiveStreamingPlayer::GetPlayPosition() const
+{
+	return PlaybackState.GetPlayPosition();
+}
+
+
+//-----------------------------------------------------------------------------
+/**
+ * Returns the seekable range.
+ *
+ * @param OutRange
+ */
+void FAdaptiveStreamingPlayer::GetSeekableRange(FTimeRange& OutRange) const
+{
+	PlaybackState.GetSeekableRange(OutRange);
+}
+
+//-----------------------------------------------------------------------------
+/**
+ * Fills the provided array with time values that can be seeked to. These are segment start times from the video (or audio if there is no video) track.
+ *
+ * @param OutPositions
+ */
+void FAdaptiveStreamingPlayer::GetSeekablePositions(TArray<FTimespan>& OutPositions) const
+{
+	PlaybackState.GetSeekablePositions(OutPositions);
+}
+
+//-----------------------------------------------------------------------------
+/**
+ * Returns the timeline range.
+ *
+ * @param OutRange
+ */
+void FAdaptiveStreamingPlayer::GetTimelineRange(FTimeRange& OutRange) const
+{
+	PlaybackState.GetTimelineRange(OutRange);
+}
+
+//-----------------------------------------------------------------------------
+/**
+ * Returns true when playback has finished.
+ *
+ * @return
+ */
+bool FAdaptiveStreamingPlayer::HasEnded() const
+{
+	return PlaybackState.GetHasEnded();
+}
+
+//-----------------------------------------------------------------------------
+/**
+ * Returns true when seeking is in progress. False if not.
+ *
+ * @return
+ */
+bool FAdaptiveStreamingPlayer::IsSeeking() const
+{
+	return PlaybackState.GetIsSeeking();
+}
+
+//-----------------------------------------------------------------------------
+/**
+ * Returns true when data is being buffered/rebuffered, false otherwise.
+ *
+ * @return
+ */
+bool FAdaptiveStreamingPlayer::IsBuffering() const
+{
+	return PlaybackState.GetIsBuffering();
+}
+
+//-----------------------------------------------------------------------------
+/**
+ * Returns true when playing back, false if not.
+ *
+ * @return
+ */
+bool FAdaptiveStreamingPlayer::IsPlaying() const
+{
+	return PlaybackState.GetIsPlaying();
+}
+
+//-----------------------------------------------------------------------------
+/**
+ * Returns true when paused, false if not.
+ *
+ * @return
+ */
+bool FAdaptiveStreamingPlayer::IsPaused() const
+{
+	return PlaybackState.GetIsPaused();
+}
+
+
+//-----------------------------------------------------------------------------
+/**
+ * Returns the current loop state.
+ *
+ * @param OutLoopState
+ */
+void FAdaptiveStreamingPlayer::GetLoopState(FPlayerLoopState& OutLoopState) const
+{
+	PlaybackState.GetLoopState(OutLoopState);
+}
+
+//-----------------------------------------------------------------------------
+/**
+ * Returns stream metadata of the currently active play period.
+ *
+ * @param OutStreamMetadata
+ * @param StreamType
+ */
+void FAdaptiveStreamingPlayer::GetStreamMetadata(TArray<FStreamMetadata>& OutStreamMetadata, EStreamType StreamType) const
+{
+	TArray<FStreamMetadata> dummy;
+	if (StreamType == EStreamType::Video)
+	{
+		PlaybackState.GetStreamMetadata(OutStreamMetadata, dummy);
+	}
+	else if (StreamType == EStreamType::Audio)
+	{
+		PlaybackState.GetStreamMetadata(dummy, OutStreamMetadata);
+	}
+	else
+	{
+		OutStreamMetadata.Empty();
+	}
+}
+
+
+//-----------------------------------------------------------------------------
+/**
+ * Sets the highest bitrate when selecting a candidate stream.
+ *
+ * @param highestSelectableBitrate
+ */
+void FAdaptiveStreamingPlayer::SetBitrateCeiling(int32 highestSelectableBitrate)
+{
+	WorkerThread.SendBitrateMessage(EStreamType::Video, highestSelectableBitrate, 1);
+}
+
+
+//-----------------------------------------------------------------------------
+/**
+ * Sets the maximum resolution to use. Set both to 0 to disable, set only one to limit widht or height only.
+ * Setting both will limit on either width or height, whichever limits first.
+ *
+ * @param MaxWidth
+ * @param MaxHeight
+ */
+void FAdaptiveStreamingPlayer::SetMaxResolution(int32 MaxWidth, int32 MaxHeight)
+{
+	WorkerThread.SendResolutionMessage(MaxWidth, MaxHeight);
+}
+
+
+//-----------------------------------------------------------------------------
+/**
+ * Selects a track based from one of the array members returned by GetStreamMetadata().
+ *
+ * @param StreamType
+ * @param StreamMetadata
+ */
+void FAdaptiveStreamingPlayer::SelectTrackByMetadata(EStreamType StreamType, const FStreamMetadata& StreamMetadata)
+{
+	WorkerThread.SendTrackSelectByMetadataMessage(StreamType, StreamMetadata);
+}
+
+//-----------------------------------------------------------------------------
+/**
+ * Deselect track. The stream will continue to stream to allow for immediate selection/activation but no data will be fed to the decoder.
+ *
+ * @param StreamType
+ */
+void FAdaptiveStreamingPlayer::DeselectTrack(EStreamType StreamType)
+{
+	WorkerThread.SendTrackDeselectMessage(StreamType);
+}
+
+//-----------------------------------------------------------------------------
+/**
+ * Starts the renderers.
+ *
+ * Call this once enough renderable data (both audio and video) is present.
+ */
+void FAdaptiveStreamingPlayer::StartRendering()
+{
+	RenderClock->Start();
+
+	if (VideoRender.Renderer)
+	{
+		FParamDict noOptions;
+		VideoRender.Renderer->StartRendering(noOptions);
+	}
+
+	if (AudioRender.Renderer)
+	{
+		FParamDict noOptions;
+		AudioRender.Renderer->StartRendering(noOptions);
+	}
+}
+
+//-----------------------------------------------------------------------------
+/**
+ * Stops renderers
+ */
+void FAdaptiveStreamingPlayer::StopRendering()
+{
+	RenderClock->Stop();
+
+	if (VideoRender.Renderer)
+	{
+		FParamDict noOptions;
+		VideoRender.Renderer->StopRendering(noOptions);
+	}
+	if (AudioRender.Renderer)
+	{
+		FParamDict noOptions;
+		AudioRender.Renderer->StopRendering(noOptions);
+	}
+}
+
+//-----------------------------------------------------------------------------
+/**
+ * Create the necessary AV renderers.
+ *
+ * @return
+ */
+int32 FAdaptiveStreamingPlayer::CreateRenderers()
+{
+	// Set the render clock with the renderes.
+	VideoRender.Renderer->SetRenderClock(RenderClock);
+	AudioRender.Renderer->SetRenderClock(RenderClock);
+	return 0;
+}
+
+
+//-----------------------------------------------------------------------------
+/**
+ * Destroys the renderers and decoders.
+ */
+void FAdaptiveStreamingPlayer::DestroyRenderers()
+{
+	// Decoders must be gone already
+	check(!VideoDecoder.Decoder);
+	check(!AudioDecoder.Decoder);
+	AudioRender.Close();
+	VideoRender.Close();
+}
+
+
+
+//-----------------------------------------------------------------------------
+/**
+ * Create the necessary AV decoders.
+ *
+ * @return
+ */
+int32 FAdaptiveStreamingPlayer::CreateInitialDecoder(EStreamType type)
+{
+	if (type == EStreamType::Video)
+	{
+		if (VideoDecoder.Decoder == nullptr)
+		{
+			// Do we have a user defined fixed limited for vertical stream resolution?
+			MaxVideoTextureResolutionLimitHeight = 0;
+			if (PlayerOptions.HaveKey("max_resoY"))
+			{
+				MaxVideoTextureResolutionLimitHeight = (int32) PlayerOptions.GetValue("max_resoY").GetInt64();
+			}
+
+			// Get the largest stream resolution. We pass this to the decoder's Open() call to let it know what
+			// the largest resolution is it may need to decode at some point.
+		// FIXME: This is only correct for HLS and simple mp4 playback. For MPEG DASH multi period playback we need to look at all
+		//        periods and get the largest resolution across all of them. And even then, for a Live presentation where a new
+		//        period can be added at any point there is the possibility that this may be of a higher resolution than any period
+		//        encountered so far/
+			FStreamCodecInformation HighestStream;
+			FindMatchingStreamInfo(HighestStream, 0, MaxVideoTextureResolutionLimitHeight);
+
+			// Create H.264 video decoder
+			VideoDecoder.Decoder = IVideoDecoderH264::Create();
+			VideoDecoder.Decoder->SetPlayerSessionServices(this);
+			VideoDecoder.Parent = this;
+			IVideoDecoderH264::FInstanceConfiguration h264Cfg = PlayerConfig.DecoderCfg264;
+			h264Cfg.ProfileIdc = HighestStream.GetProfile();
+			h264Cfg.LevelIdc = HighestStream.GetProfileLevel();
+			h264Cfg.MaxFrameWidth = HighestStream.GetResolution().Width;
+			h264Cfg.MaxFrameHeight = HighestStream.GetResolution().Height;
+			h264Cfg.AdditionalOptions = HighestStream.GetExtras();
+
+			// Add in any player options that are for decoder use
+			TArray<FString> DecoderOptionKeys;
+			PlayerOptions.GetKeysStartingWith("videoDecoder", DecoderOptionKeys);
+			for (const FString & Key : DecoderOptionKeys)
+			{
+				h264Cfg.AdditionalOptions.Set(Key, PlayerOptions.GetValue(Key));
+			}
+
+			// Attach video decoder buffer monitor.
+			VideoDecoder.Decoder->SetAUInputBufferListener(&VideoDecoder);
+			VideoDecoder.Decoder->SetReadyBufferListener(&VideoDecoder);
+			// Have the video decoder send its output to the video renderer
+			VideoDecoder.Decoder->SetRenderer(VideoRender.Renderer);
+			// Hand it (may be nullptr) a delegate for platform for resource queries
+			VideoDecoder.Decoder->SetResourceDelegate(VideoDecoderResourceDelegate.Pin());
+			// Open the decoder after having set all listeners.
+			VideoDecoder.Decoder->Open(h264Cfg);
+
+			// Now we get the currently limited stream resolution and let the decoder now what we will be using
+			// at most right now. This allows the decoder to be created with a smaller memory footprint at first.
+			UpdateStreamResolutionLimit();
+		}
+		return 0;
+	}
+	else if (type == EStreamType::Audio)
+	{
+		if (AudioDecoder.Decoder == nullptr)
+		{
+			// Create an AAC audio decoder
+			AudioDecoder.Decoder = IAudioDecoderAAC::Create();
+			AudioDecoder.Decoder->SetPlayerSessionServices(this);
+			AudioDecoder.Parent = this;
+			// Attach buffer monitors.
+			AudioDecoder.Decoder->SetAUInputBufferListener(&AudioDecoder);
+			AudioDecoder.Decoder->SetReadyBufferListener(&AudioDecoder);
+			// Have to audio decoder send its output to the audio renderer
+			AudioDecoder.Decoder->SetRenderer(AudioRender.Renderer);
+			// Open the decoder after having set all listeners.
+			AudioDecoder.Decoder->Open(PlayerConfig.DecoderCfgAAC);
+		}
+		return 0;
+	}
+	return -1;
+}
+
+
+//-----------------------------------------------------------------------------
+/**
+ * Destroys the decoders.
+ */
+void FAdaptiveStreamingPlayer::DestroyDecoders()
+{
+/*
+NOTE: We do not clear out the renderers from the decoder. On their way down the decoders should still be able
+      to access the renderer without harm and dispatch their last remaining data.
+
+	if (mVideoDecoder.Decoder)
+		mVideoDecoder.Decoder->SetRenderer(nullptr);
+	if (mAudioDecoder.Decoder)
+		mAudioDecoder.Decoder->SetRenderer(nullptr);
+*/
+	AudioDecoder.Close();
+	VideoDecoder.Close();
+}
+
+
+
+
+
+
+void FAdaptiveStreamingPlayer::PostLog(Facility::EFacility FromFacility, IInfoLog::ELevel inLogLevel, const FString &Message)
+{
+	int64 millisNow = SynchronizedUTCTime->GetTime().GetAsMilliseconds();
+	//FString msg = FString::Printf("[%s] %s: %s", IInfoLog::GetLevelName(inLogLevel), Facility::GetName(FromFacility), Message.c_str());
+	FString msg = FString::Printf(TEXT("%s: %s"), Facility::GetName(FromFacility), *Message);
+	DispatchEvent(FMetricEvent::ReportLogMessage(inLogLevel, msg, millisNow));
+}
+
+void FAdaptiveStreamingPlayer::PostError(const FErrorDetail& InError)
+{
+	TSharedPtrTS<FErrorDetail> Error(new FErrorDetail(InError));
+	ErrorQueue.Push(Error);
+}
+
+void FAdaptiveStreamingPlayer::SendMessageToPlayer(TSharedPtrTS<IPlayerMessage> PlayerMessage)
+{
+	WorkerThread.SendPlayerSessionMessage(PlayerMessage);
+}
+
+//-----------------------------------------------------------------------------
+/**
+ * Returns the external GUID identifying this player and its associated external texture.
+ */
+void FAdaptiveStreamingPlayer::GetExternalGuid(FGuid& OutExternalGuid)
+{
+	OutExternalGuid = ExternalPlayerGUID;
+}
+
+ISynchronizedUTCTime* FAdaptiveStreamingPlayer::GetSynchronizedUTCTime()
+{
+	return SynchronizedUTCTime;
+}
+
+TSharedPtr<IAdaptiveStreamingPlayerResourceProvider, ESPMode::ThreadSafe> FAdaptiveStreamingPlayer::GetStaticResourceProvider()
+{
+	return StaticResourceProvider.Pin();
+}
+
+IElectraHttpManager* FAdaptiveStreamingPlayer::GetHTTPManager()
+{
+	return HttpManager.Get();
+}
+
+TSharedPtrTS<IAdaptiveStreamSelector> FAdaptiveStreamingPlayer::GetStreamSelector()
+{
+	return StreamSelector;
+}
+
+void FAdaptiveStreamingPlayer::GetStreamBufferStats(FAccessUnitBufferInfo& OutBufferStats, EStreamType ForStream)
+{
+	FMediaCriticalSection::ScopedLock lock(DiagnosticsCriticalSection);
+	switch(ForStream)
+	{
+		case EStreamType::Video:
+			OutBufferStats = VideoBufferStats.StreamBuffer;
+			break;
+		case EStreamType::Audio:
+			OutBufferStats = AudioBufferStats.StreamBuffer;
+			break;
+		default:
+			OutBufferStats.Clear();
+			break;
+	}
+}
+
+IPlayerStreamFilter* FAdaptiveStreamingPlayer::GetStreamFilter()
+{
+	return this;
+}
+
+bool FAdaptiveStreamingPlayer::CanDecodeStream(const FStreamCodecInformation& InStreamCodecInfo) const
+{
+	if (InStreamCodecInfo.IsVideoCodec())
+	{
+		if (InStreamCodecInfo.GetFrameRate().IsValid())
+		{
+			// Check framerate related configuration options. For this the stream needs to have the framerate specified.
+			FTimeValue StreamFPS;
+			StreamFPS.SetFromTimeFraction(InStreamCodecInfo.GetFrameRate());
+			if (StreamFPS.IsValid())
+			{
+				// Required minimum fps?
+				if (PlayerOptions.HaveKey("min_stream_fps"))
+				{
+					FTimeValue MinFPS(FTimeValue().SetFromND(PlayerOptions.GetValue("min_stream_fps").GetInt64(), 1));
+					if (StreamFPS < MinFPS)
+					{
+						return false;
+					}
+				}
+				// Required maximum fps?
+				if (PlayerOptions.HaveKey("max_stream_fps"))
+				{
+					FTimeValue MaxFPS(FTimeValue().SetFromND(PlayerOptions.GetValue("max_stream_fps").GetInt64(), 1));
+					if (StreamFPS > MaxFPS)
+					{
+						return false;
+					}
+				}
+				// Stream fps faster than 30fps and a vertical resolution limit? (check against 31 fps in case of slight rounding errors)
+				if (PlayerOptions.HaveKey("max_resoY_above_30fps") && StreamFPS > FTimeValue().SetFromND(31,1))
+				{
+					int64 MaxVertReso = PlayerOptions.GetValue("max_resoY_above_30fps").GetInt64();
+
+					// Check if a maximum texture dimension limit is in place. If it is then it must never be exceeded.
+					if (MaxVideoTextureResolutionLimitHeight > 0 && MaxVertReso > MaxVideoTextureResolutionLimitHeight)
+					{
+						MaxVertReso = MaxVideoTextureResolutionLimitHeight;
+					}
+
+					if (InStreamCodecInfo.GetResolution().Height > MaxVertReso)
+					{
+						return false;
+					}
+				}
+			}
+		}
+
+		IVideoDecoderH264::FStreamDecodeCapability StreamParam, Capability;
+		// Do a one-time global capability check with a default-empty stream param structure.
+		// This then gets used in the individual stream capability checks.
+		if (IVideoDecoderH264::GetStreamDecodeCapability(Capability, StreamParam))
+		{
+			// Manual vertical resolution selection?
+			if (PlayerOptions.HaveKey("max_resoY_windows_software"))
+			{
+				int64 MaxWinSWHeight = PlayerOptions.GetValue("max_resoY_windows_software").GetInt64();
+				// If the global capability indicates software decoding, check if this stream exceeds the vertical limit. If so, reject it.
+				if (MaxWinSWHeight > 0 && Capability.DecoderSupportType == IVideoDecoderH264::FStreamDecodeCapability::ESupported::SoftwareOnly && InStreamCodecInfo.GetResolution().Height > MaxWinSWHeight)
+				{
+					return false;
+				}
+			}
+			else
+			{
+				StreamParam.Width = InStreamCodecInfo.GetResolution().Width;
+				StreamParam.Height = InStreamCodecInfo.GetResolution().Height;
+				StreamParam.Profile = InStreamCodecInfo.GetProfile();
+				StreamParam.Level = InStreamCodecInfo.GetProfileLevel();
+				StreamParam.FPS = InStreamCodecInfo.GetFrameRate().IsValid() ? InStreamCodecInfo.GetFrameRate().GetAsDouble() : 30.0;
+				if (IVideoDecoderH264::GetStreamDecodeCapability(Capability, StreamParam))
+				{
+					if (Capability.DecoderSupportType == IVideoDecoderH264::FStreamDecodeCapability::ESupported::NotSupported)
+					{
+						return false;
+					}
+				}
+			}
+		}
+	}
+
+	return true;
+}
+
+
+
+
+// AU memory
+void* FAdaptiveStreamingPlayer::AUAllocate(IAccessUnitMemoryProvider::EDataType type, SIZE_T NumBytes, SIZE_T Alignment)
+{
+	if (Alignment)
+	{
+		return FMemory::Malloc(NumBytes, Alignment);
+	}
+	return FMemory::Malloc(NumBytes);
+}
+
+void FAdaptiveStreamingPlayer::AUDeallocate(IAccessUnitMemoryProvider::EDataType type, void* Address)
+{
+	FMemory::Free(Address);
+}
+
+
+//-----------------------------------------------------------------------------
+/**
+ * Notified when the given fragment will be opened.
+ *
+ * @param pRequest
+ */
+void FAdaptiveStreamingPlayer::OnFragmentOpen(TSharedPtrTS<IStreamSegment> pRequest)
+{
+	WorkerThread.SendMessage(FWorkerThread::FMessage::EType::FragmentOpen, pRequest);
+}
+
+
+//-----------------------------------------------------------------------------
+/**
+ * Fragment access unit received callback.
+ *
+ * NOTE: This must be done in place and cannot be deferred to a worker thread
+ *       since we must return whether or not we could store the AU in our buffer right away.
+ *
+ * @param pAccessUnit
+ *
+ * @return
+ */
+bool FAdaptiveStreamingPlayer::OnFragmentAccessUnitReceived(FAccessUnit* pAccessUnit)
+{
+	// When shutting down we don't accept anything new, but we return we did so the stream reader won't keep pummeling us with new data.
+	if (bIsClosing)
+	{
+		return true;
+	}
+	FMultiTrackAccessUnitBuffer* pBuffer;
+	switch(pAccessUnit->ESType)
+	{
+		case EStreamType::Video:
+			pBuffer = &MultiStreamBufferVid;
+			break;
+		case EStreamType::Audio:
+			pBuffer = &MultiStreamBufferAud;
+			break;
+		case EStreamType::Subtitle:
+			pBuffer = &MultiStreamBufferTxt;
+			break;
+		default:
+			checkf(0, TEXT("Bad ES stream type %d"), pAccessUnit->ESType);
+			return true;
+	}
+
+	// Try to push the data into the receiving buffer
+	return pBuffer->Push(pAccessUnit);
+}
+
+//-----------------------------------------------------------------------------
+/**
+ * Fragment reached end-of-stream callback.
+ * No additional access units will be received for this fragment.
+ *
+ * @param InStreamType
+ * @param InStreamSourceInfo
+ */
+void FAdaptiveStreamingPlayer::OnFragmentReachedEOS(EStreamType InStreamType, TSharedPtr<const FStreamSourceInfo, ESPMode::ThreadSafe> InStreamSourceInfo)
+{
+	if (!bIsClosing)
+	{
+		FMultiTrackAccessUnitBuffer* pBuffer = nullptr;
+		switch(InStreamType)
+		{
+			case EStreamType::Video:
+				pBuffer = &MultiStreamBufferVid;
+				break;
+			case EStreamType::Audio:
+				pBuffer = &MultiStreamBufferAud;
+				break;
+			case EStreamType::Subtitle:
+				pBuffer = &MultiStreamBufferTxt;
+				break;
+			default:
+				checkf(0, TEXT("Bad ES stream type %d"), InStreamType);
+				break;
+		}
+		if (pBuffer)
+		{
+			pBuffer->PushEndOfDataFor(InStreamSourceInfo);
+		}
+	}
+}
+
+
+//-----------------------------------------------------------------------------
+/**
+ * Notified when the current fragment is closed.
+ * This is called regardless of the error state and can always
+ * be used as an indication that the fragment has finished processing.
+ *
+ * @param pRequest
+ */
+void FAdaptiveStreamingPlayer::OnFragmentClose(TSharedPtrTS<IStreamSegment> pRequest)
+{
+	WorkerThread.SendMessage(FWorkerThread::FMessage::EType::FragmentClose, pRequest);
+}
+
+
+
+//-----------------------------------------------------------------------------
+/**
+ * Pauses the stream readers.
+ */
+void FAdaptiveStreamingPlayer::PauseStreamReaders()
+{
+	if (StreamReaderHandler)
+	{
+		StreamReaderHandler->PauseDownload();
+	}
+}
+
+//-----------------------------------------------------------------------------
+/**
+ * Resumes the stream readers
+ */
+void FAdaptiveStreamingPlayer::ResumeStreamReaders()
+{
+	if (StreamReaderHandler)
+	{
+		StreamReaderHandler->ResumeDownload();
+	}
+}
+
+
+
+
+
+
+
+//-----------------------------------------------------------------------------
+/**
+ * Worker thread.
+ */
+void FAdaptiveStreamingPlayer::WorkerThreadFN()
+{
+	LLM_SCOPE(ELLMTag::ElectraPlayer);
+
+	bool	bQuit = false;
+
+	// Get the event dispatcher.
+	EventDispatcher = FAdaptiveStreamingPlayerEventHandler::Create();
+
+
+	// Create renderers and decoders.
+	// FIXME: At some later point in time we should do this on-demand only. We do not know the codecs
+	//        here yet and the codecs could also change dynamically as streams get switched.
+	//        For now this is ok and also prevents spikes from loading codec DLL/PRX later.
+	/*int32 CreateResult =*/ CreateRenderers();
+
+	while(!bQuit)
+	{
+		FWorkerThread::FMessage msg;
+		if (WorkerThread.WorkMessages.ReceiveMessage(msg, 1000 * 20))
+		{
+			CSV_SCOPED_TIMING_STAT(ElectraPlayer, AdaptiveStreamingPlayer_Worker);
+
+			// While closed ignore all but the quit message.
+			if (bIsClosing && msg.Type != FWorkerThread::FMessage::EType::Quit)
+			{
+				if (msg.Data.MediaEvent.Event)
+				{
+					msg.Data.MediaEvent.Event->Signal();
+				}
+				continue;
+			}
+
+			// What is it?
+			switch(msg.Type)
+			{
+				case FWorkerThread::FMessage::EType::Quit:
+				{
+					bQuit = true;
+					break;
+				}
+				case FWorkerThread::FMessage::EType::LoadManifest:
+				{
+					InternalLoadManifest(msg.Data.ManifestToLoad.URL, msg.Data.ManifestToLoad.MimeType);
+					break;
+				}
+				case FWorkerThread::FMessage::EType::Pause:
+				{
+					bShouldBePaused = true;
+					bShouldBePlaying = false;
+					break;
+				}
+				case FWorkerThread::FMessage::EType::Resume:
+				{
+					bShouldBePaused = false;
+					bShouldBePlaying = true;
+					break;
+				}
+				case FWorkerThread::FMessage::EType::Seek:
+				{
+					StartAtTime = msg.Data.StartPlay.Position;
+
+					if (bIsPlayStart)
+					{
+						CurrentState = EPlayerState::eState_Buffering;
+						bool bOk = InternalStartAt(StartAtTime);
+					}
+					else
+					{
+						// Stop everything.
+						InternalStop(PlayerConfig.bHoldLastFrameDuringSeek);
+						CurrentState = EPlayerState::eState_Seeking;
+						bool bOk = InternalStartAt(StartAtTime);
+					}
+					break;
+				}
+				case FWorkerThread::FMessage::EType::Loop:
+				{
+					InternalSetLoop(msg.Data.Looping.Loop);
+					if (msg.Data.Looping.Signal)
+					{
+						msg.Data.Looping.Signal->Signal();
+					}
+					break;
+				}
+				case FWorkerThread::FMessage::EType::Close:
+				{
+					if (!bIsClosing)
+					{
+						bIsClosing = true;
+						InternalStop(false);
+						InternalClose();
+						// This is the end of the rope anyway so let's terminate the event dispatcher.
+						// This ensures that the client will not get out of the Stop() call without having
+						// received the stop message.
+						DispatchEventAndWait(FMetricEvent::ReportPlaybackStopped());
+						EventDispatcher.Reset();
+					}
+					if (msg.Data.MediaEvent.Event)
+					{
+						msg.Data.MediaEvent.Event->Signal();
+					}
+					break;
+					}
+
+				case FWorkerThread::FMessage::EType::ChangeBitrate:
+				{
+					BitrateCeiling = msg.Data.Bitrate.Value;
+					StreamSelector->SetBandwidthCeiling(BitrateCeiling);
+					break;
+				}
+
+				case FWorkerThread::FMessage::EType::LimitResolution:
+				{
+					VideoResolutionLimitWidth  = msg.Data.Resolution.Width;
+					VideoResolutionLimitHeight = msg.Data.Resolution.Height;
+					UpdateStreamResolutionLimit();
+					break;
+				}
+
+				case FWorkerThread::FMessage::EType::SelectTrackByMetadata:
+				{
+					switch(msg.Data.TrackSelection.StreamType)
+					{
+						case EStreamType::Video:
+							InitialStreamSelectionVid = MakeShared<FStreamMetadata, ESPMode::ThreadSafe>(msg.Data.TrackSelection.StreamMetadata);
+							break;
+						case EStreamType::Audio:
+							InitialStreamSelectionAud = MakeShared<FStreamMetadata, ESPMode::ThreadSafe>(msg.Data.TrackSelection.StreamMetadata);
+							// For media with multiplexed tracks we can switch to a different track pretty much immediately.
+							if (ManifestReader && ManifestReader->GetPlaylistType() == "mp4")
+							{
+								// Only select audio for now.
+								if (InitialStreamSelectionAud.IsValid())
+								{
+									MultiStreamBufferAud.SelectTrackByID(InitialStreamSelectionAud->StreamUniqueID);
+								}
+							}
+							break;
+						case EStreamType::Subtitle:
+							InitialStreamSelectionTxt = MakeShared<FStreamMetadata, ESPMode::ThreadSafe>(msg.Data.TrackSelection.StreamMetadata);
+							break;
+					}
+					break;
+				}
+
+				case FWorkerThread::FMessage::EType::DeselectTrack:
+				{
+					switch(msg.Data.TrackSelection.StreamType)
+					{
+						case EStreamType::Video:
+							InitialStreamSelectionVid.Reset();
+							break;
+						case EStreamType::Audio:
+							InitialStreamSelectionAud.Reset();
+							// For media with multiplexed tracks we can switch to a different track pretty much immediately.
+							if (ManifestReader && ManifestReader->GetPlaylistType() == "mp4")
+							{
+								MultiStreamBufferAud.Deselect();
+							}
+							break;
+						case EStreamType::Subtitle:
+							InitialStreamSelectionTxt.Reset();
+							break;
+					}
+					break;
+				}
+
+				case FWorkerThread::FMessage::EType::FragmentOpen:
+				{
+					TSharedPtrTS<IStreamSegment> pRequest = msg.Data.StreamReader.Request;
+					// Check that the request is for this current playback sequence and not an outdated one.
+					if (pRequest.IsValid() && pRequest->GetPlaybackSequenceID() == CurrentPlaybackSequenceID)
+					{
+						DispatchBufferUtilizationEvent(msg.Data.StreamReader.Request->GetType() == EStreamType::Video ? EStreamType::Video : EStreamType::Audio);
+
+						// Video bitrate change?
+						if (msg.Data.StreamReader.Request->GetType() == EStreamType::Video)
+						{
+							int32 SegmentBitrate = pRequest->GetBitrate();
+							int32 SegmentQualityLevel = pRequest->GetQualityIndex();
+							if (SegmentBitrate != CurrentVideoStreamBitrate.Bitrate)
+							{
+								bool bDrastic = CurrentVideoStreamBitrate.Bitrate && SegmentQualityLevel < CurrentVideoStreamBitrate.QualityLevel-1;
+								DispatchEvent(FMetricEvent::ReportVideoQualityChange(SegmentBitrate, CurrentVideoStreamBitrate.Bitrate, bDrastic));
+								CurrentVideoStreamBitrate.Bitrate      = SegmentBitrate;
+								CurrentVideoStreamBitrate.QualityLevel = SegmentQualityLevel;
+							}
+						}
+					}
+					break;
+				}
+
+				case FWorkerThread::FMessage::EType::FragmentClose:
+				{
+					TSharedPtrTS<IStreamSegment> pRequest = msg.Data.StreamReader.Request;
+					// Check that the request is for this current playback sequence and not an outdated one.
+					if (pRequest.IsValid() && pRequest->GetPlaybackSequenceID() == CurrentPlaybackSequenceID)
+					{
+						EStreamType reqType = msg.Data.StreamReader.Request->GetType();
+
+						// Dispatch download event
+						DispatchSegmentDownloadedEvent(pRequest);
+
+						// Dispatch buffer utilization
+						DispatchBufferUtilizationEvent(reqType == EStreamType::Video ? EStreamType::Video : EStreamType::Audio);
+
+						// Dispatch avergae throughput, bandwidth and latency event for video segments only.
+						if (reqType == EStreamType::Video)
+						{
+							DispatchEvent(FMetricEvent::ReportBandwidth(StreamSelector->GetAverageBandwidth(), StreamSelector->GetAverageThroughput(), StreamSelector->GetAverageLatency()));
+						}
+
+						// Add to the list of completed requests for which we need to find the next or retry segment to fetch.
+						FPendingSegmentRequest NextReq;
+						NextReq.Request = pRequest;
+						NextPendingSegmentRequests.Enqueue(MoveTemp(NextReq));
+					}
+					break;
+				}
+
+				case FWorkerThread::FMessage::EType::BufferUnderrun:
+				{
+					InternalRebuffer();
+					break;
+				}
+
+				case FWorkerThread::FMessage::EType::PlayerSession:
+				{
+					HandleSessionMessage(msg.Data.Session.PlayerMessage);
+					break;
+				}
+
+				default:
+				{
+					checkNoEntry();
+					break;
+				}
+			}
+		}
+
+		if (!bIsClosing)
+		{
+			CSV_SCOPED_TIMING_STAT(ElectraPlayer, AdaptiveStreamingPlayer_Worker);
+
+			if (RenderClock->IsRunning())
+			{
+				FTimeValue playPos = RenderClock->GetInterpolatedRenderTime(IMediaRenderClock::ERendererType::Video);
+				PlaybackState.SetPlayPosition(playPos);
+				// When we reach the time at which the next loop occurred pop it off the queue and update the playback state loop state.
+				while(NextLoopStates.Num())
+				{
+					FTimeValue loopAtTime = NextLoopStates.FrontRef().LoopBasetime;
+					if (playPos >= loopAtTime)
+					{
+						FPlayerLoopState newLoopState = NextLoopStates.Pop();
+						PlaybackState.SetLoopState(newLoopState);
+						DispatchEvent(FMetricEvent::ReportJumpInPlayPosition(playPos, playPos, Metrics::ETimeJumpReason::Looping));
+					}
+					else
+					{
+						break;
+					}
+				}
+			}
+
+			// Handle state changes to match the user request.
+			HandlePlayStateChanges();
+
+			// Update diag buffers
+			UpdateDiagnostics();
+			// Check for end of stream.
+			CheckForStreamEnd();
+			// Check the error queue for new arrivals.
+			CheckForErrors();
+			// Handle pending media segment requests.
+			HandlePendingMediaSegmentRequests();
+			// Handle buffer level changes
+			HandleNewBufferedData();
+			// Handle new output data (finish prerolling)
+			HandleNewOutputData();
+			// Handle data buffers from deselected tracks to align with selected ones.
+			HandleDeselectedBuffers();
+			// Handle decoder changes
+			HandleDecoderChanges();
+		}
+	}
+
+	if (EventDispatcher.IsValid())
+	{
+		EventDispatcher->DispatchEventAndWait(TSharedPtrTS<FMetricEvent>());
+		EventDispatcher.Reset();
+	}
+
+	while(WorkerThread.WorkMessages.HaveMessage())
+	{
+		WorkerThread.WorkMessages.ReceiveMessage();
+	}
+}
+
+
+
+void FAdaptiveStreamingPlayer::HandlePlayStateChanges()
+{
+	if (CurrentState == EPlayerState::eState_Error)
+	{
+		return;
+	}
+
+	// Should the player be paused?
+	if (bShouldBePaused)
+	{
+		// Are we not paused and in a state in which we can pause?
+		if (CurrentState == EPlayerState::eState_Playing)
+		{
+			InternalPause();
+		}
+	}
+
+	// Should the player be playing?
+	if (bShouldBePlaying)
+	{
+		// Are we paused and in a state in which we can play?
+		if (CurrentState == EPlayerState::eState_Paused)
+		{
+			InternalResume();
+		}
+	}
+}
+
+
+//-----------------------------------------------------------------------------
+/**
+ * Handles a player session message sent by one of the player sub-systems.
+ *
+ * @param SessionMessage
+ */
+void FAdaptiveStreamingPlayer::HandleSessionMessage(TSharedPtrTS<IPlayerMessage> SessionMessage)
+{
+	// Playlist downloaded (successful or not)?
+	if (SessionMessage->GetType() == IPlaylistReader::PlaylistDownloadMessage::Type())
+	{
+		IPlaylistReader::PlaylistDownloadMessage* pMsg = static_cast<IPlaylistReader::PlaylistDownloadMessage *>(SessionMessage.Get());
+		Metrics::FPlaylistDownloadStats stats;
+		stats.bWasSuccessful = !pMsg->GetConnectionInfo().StatusInfo.ErrorDetail.IsSet();
+		stats.ListType  	 = pMsg->GetListType();
+		stats.LoadType  	 = pMsg->GetLoadType();
+		stats.URL   		 = pMsg->GetConnectionInfo().EffectiveURL;
+		stats.FailureReason  = pMsg->GetConnectionInfo().StatusInfo.ErrorDetail.GetMessage();
+		stats.HTTPStatusCode = pMsg->GetConnectionInfo().StatusInfo.HTTPStatus;
+		stats.RetryNumber    = pMsg->GetConnectionInfo().RetryInfo.IsValid() ? pMsg->GetConnectionInfo().RetryInfo->AttemptNumber : 0;
+		DispatchEvent(FMetricEvent::ReportPlaylistDownload(stats));
+	}
+	// Playlist fetched & parsed?
+	else if (SessionMessage->GetType() == IPlaylistReader::PlaylistLoadedMessage::Type())
+	{
+		if (ManifestReader)
+		{
+			IPlaylistReader::PlaylistLoadedMessage* pMsg = static_cast<IPlaylistReader::PlaylistLoadedMessage *>(SessionMessage.Get());
+			const FErrorDetail& Result = pMsg->GetResult();
+			// Note: Right now only successful messages should get here. We check for failure anyway in case this changes in the future.
+			if (!Result.IsError())
+			{
+				if (!Result.WasAborted())
+				{
+					if (pMsg->GetListType() == Playlist::EListType::Master)
+					{
+						DispatchEvent(FMetricEvent::ReportReceivedMasterPlaylist(pMsg->GetConnectionInfo().EffectiveURL));
+					}
+					else if (pMsg->GetLoadType() == Playlist::ELoadType::Initial)
+					{
+						SelectManifest();
+						DispatchEvent(FMetricEvent::ReportReceivedPlaylists());
+					}
+					else
+					{
+						UpdateManifest();
+					}
+				}
+			}
+			else
+			{
+				PostError(Result);
+			}
+		}
+	}
+	// License key?
+	else if (SessionMessage->GetType() == FLicenseKeyMessage::Type())
+	{
+		FLicenseKeyMessage* pMsg = static_cast<FLicenseKeyMessage*>(SessionMessage.Get());
+		Metrics::FLicenseKeyStats stats;
+		stats.bWasSuccessful = !pMsg->GetResult().IsError();
+		stats.URL   		 = pMsg->GetConnectionInfo().EffectiveURL;
+		stats.FailureReason  = pMsg->GetResult().IsError() ? pMsg->GetResult().GetPrintable() : FString(); //pMsg->GetConnectionInfo().StatusInfo.ErrorDetail.GetMessage();
+		DispatchEvent(FMetricEvent::ReportLicenseKey(stats));
+	}
+	else
+	{
+		checkNoEntry();
+	}
+}
+
+
+
+//-----------------------------------------------------------------------------
+/**
+ * Checks if buffers have enough data to advance the play state.
+ */
+void FAdaptiveStreamingPlayer::HandleNewBufferedData()
+{
+	// Are we in buffering state? (this includes rebuffering)
+	if (CurrentState == EPlayerState::eState_Buffering)
+	{
+		if (DecoderState == EDecoderState::eDecoder_Paused)
+		{
+			check(LastBufferingState == EPlayerState::eState_Buffering || LastBufferingState == EPlayerState::eState_Rebuffering || LastBufferingState == EPlayerState::eState_Seeking);
+			// Yes. Check if we have enough data buffered up to begin handing off data to the decoders.
+			const double kMinBufferBeforePlayback = LastBufferingState == EPlayerState::eState_Seeking	  ? PlayerConfig.SeekBufferMinTimeAvailBeforePlayback :
+												    LastBufferingState == EPlayerState::eState_Rebuffering ? PlayerConfig.RebufferMinTimeAvailBeforePlayback :
+																							    PlayerConfig.InitialBufferMinTimeAvailBeforePlayback;
+
+			DiagnosticsCriticalSection.Lock();
+
+			bool bHaveEnoughVideo = false;
+			if (bHaveVideoReader.IsSet())
+			{
+				if (StreamReaderHandler && bHaveVideoReader.Value())
+				{
+					// If the video stream has reached the end then there won't be any new data and whatever we have is all there is.
+					if (VideoBufferStats.StreamBuffer.bEndOfData)
+					{
+						bHaveEnoughVideo = true;
+					}
+					else
+					{
+						// Does the buffer have the requested amount of data?
+						if (VideoBufferStats.StreamBuffer.PlayableDuration.GetAsSeconds() >= kMinBufferBeforePlayback)
+						{
+							bHaveEnoughVideo = true;
+						}
+						else
+						{
+							// There is a chance that the input stream buffer is too small to hold the requested amount of material.
+							if (VideoBufferStats.StreamBuffer.bLastPushWasBlocked)
+							{
+								// We won't be able to receive any more, so we have enough as it is.
+								bHaveEnoughVideo = true;
+							}
+						}
+					}
+				}
+				else
+				{
+					bHaveEnoughVideo = true;
+				}
+			}
+
+
+			bool bHaveEnoughAudio = false;
+			if (bHaveAudioReader.IsSet())
+			{
+				if (StreamReaderHandler && bHaveAudioReader.Value())
+				{
+					// If the audio stream has reached the end then there won't be any new data and whatever we have is all there is.
+					if (AudioBufferStats.StreamBuffer.bEndOfData)
+					{
+						bHaveEnoughAudio = true;
+					}
+					else
+					{
+						// Does the buffer have the requested amount of data?
+						if (AudioBufferStats.StreamBuffer.PlayableDuration.GetAsSeconds() >= kMinBufferBeforePlayback)
+						{
+							bHaveEnoughAudio = true;
+						}
+						else
+						{
+							// There is a chance that the input stream buffer is too small to hold the requested amount of material.
+							if (AudioBufferStats.StreamBuffer.bLastPushWasBlocked)
+							{
+								// We won't be able to receive any more, so we have enough as it is.
+								bHaveEnoughAudio = true;
+							}
+						}
+					}
+				}
+				else
+				{
+					bHaveEnoughAudio = true;
+				}
+			}
+
+			// When we are dealing with a single multiplexed stream and one buffer is blocked then essentially all buffers
+			// must be considered blocked since demuxing cannot continue.
+			// FIXME: do this more elegant somehow
+			if (ManifestReader && ManifestReader->GetPlaylistType() == "mp4")
+			{
+				if (VideoBufferStats.StreamBuffer.bLastPushWasBlocked ||
+					AudioBufferStats.StreamBuffer.bLastPushWasBlocked)
+				{
+					bHaveEnoughVideo = bHaveEnoughAudio = true;
+				}
+			}
+
+			DiagnosticsCriticalSection.Unlock();
+
+			int64 tNow = MEDIAutcTime::CurrentMSec();
+			if (bHaveEnoughVideo && bHaveEnoughAudio)
+			{
+				PlaybackState.SetIsBuffering(false);
+				PlaybackState.SetIsSeeking(false);
+
+				DecoderState   		= EDecoderState::eDecoder_Running;
+				PipelineState  		= EPipelineState::ePipeline_Prerolling;
+				PrerollVars.StartTime = tNow;
+				PostrollVars.Clear();
+
+				// Send buffering end event
+				DispatchBufferingEvent(false, LastBufferingState);
+
+				// Send pre-roll begin event.
+				DispatchEvent(FMetricEvent::ReportPrerollStart());
+
+				// Create the first decoders here.
+				if (bHaveVideoReader.GetWithDefault(false))
+				{
+					CreateInitialDecoder(EStreamType::Video);
+				}
+				if (bHaveAudioReader.GetWithDefault(false))
+				{
+					CreateInitialDecoder(EStreamType::Audio);
+				}
+			}
+		}
+	}
+}
+
+//-----------------------------------------------------------------------------
+/**
+ * Check if the decoders need to change.
+ */
+void FAdaptiveStreamingPlayer::HandleDecoderChanges()
+{
+	// Nothing to do for now.
+}
+
+//-----------------------------------------------------------------------------
+/**
+ * Handle data in deselected buffers.
+ * Since deselected buffers do not feed a decoder the data must be discarded as playback progresses
+ * to avoid buffer overflows.
+ */
+void FAdaptiveStreamingPlayer::HandleDeselectedBuffers()
+{
+	if (MultiStreamBufferVid.IsDeselected() || MultiStreamBufferAud.IsDeselected() || MultiStreamBufferTxt.IsDeselected())
+	{
+		// Get the current playback position in case nothing is selected.
+		FTimeValue DiscardUntilTime = PlaybackState.GetPlayPosition();
+
+		// Get last popped video PTS if video track is selected.
+		if (!MultiStreamBufferVid.IsDeselected())
+		{
+			FTimeValue v = MultiStreamBufferVid.GetLastPoppedPTS();
+			if (v.IsValid())
+			{
+				DiscardUntilTime = v;
+			}
+		}
+		// Get last popped audio PTS if audio track is selected.
+		if (!MultiStreamBufferAud.IsDeselected())
+		{
+			FTimeValue v = MultiStreamBufferAud.GetLastPoppedPTS();
+			if (v.IsValid())
+			{
+				DiscardUntilTime = v;
+			}
+		}
+		// FIXME: do we need to consider subtitle tracks as being the only selected ones??
+
+		if (DiscardUntilTime.IsValid())
+		{
+			if (MultiStreamBufferVid.IsDeselected())
+			{
+				FMultiTrackAccessUnitBuffer::FScopedLock lock(MultiStreamBufferVid);
+				MultiStreamBufferVid.PopDiscardUntil(DiscardUntilTime);
+			}
+			if (MultiStreamBufferAud.IsDeselected())
+			{
+				FMultiTrackAccessUnitBuffer::FScopedLock lock(MultiStreamBufferAud);
+				MultiStreamBufferAud.PopDiscardUntil(DiscardUntilTime);
+			}
+			if (MultiStreamBufferTxt.IsDeselected())
+			{
+				FMultiTrackAccessUnitBuffer::FScopedLock lock(MultiStreamBufferTxt);
+				MultiStreamBufferTxt.PopDiscardUntil(DiscardUntilTime);
+			}
+		}
+	}
+}
+
+
+//-----------------------------------------------------------------------------
+/**
+ * Checks if the player state can transition to the next state with available decoder output.
+ */
+void FAdaptiveStreamingPlayer::HandleNewOutputData()
+{
+	int64 timeNow = MEDIAutcTime::CurrentMSec();
+
+	// Are we currently pre-rolling the media pipeline?
+	if (PipelineState == EPipelineState::ePipeline_Prerolling)
+	{
+		// Yes. Is it "warm" enough to start rendering?
+		DiagnosticsCriticalSection.Lock();
+
+		// Check video
+		bool bHaveEnoughVideo = PrerollVars.bHaveEnoughVideo;
+		if (bHaveVideoReader.IsSet())
+		{
+			if (StreamReaderHandler && bHaveVideoReader.Value())
+			{
+				// If the video decoder output buffer is stalled then we have enough video output. There won't be any more coming in.
+				// Also when the stream has reached the end and that has propagated through the buffers.
+				if (VideoBufferStats.DecoderOutputBuffer.bOutputStalled ||
+					(VideoBufferStats.StreamBuffer.bEndOfData && VideoBufferStats.DecoderInputBuffer.bEODSignaled && VideoBufferStats.DecoderOutputBuffer.bEODreached) ||
+					MultiStreamBufferVid.IsDeselected())
+				{
+					bHaveEnoughVideo = true;
+				}
+			}
+			else
+			{
+				bHaveEnoughVideo = true;
+			}
+		}
+
+		// Check audio
+		bool bHaveEnoughAudio = PrerollVars.bHaveEnoughAudio;
+		if (bHaveAudioReader.IsSet())
+		{
+			if (StreamReaderHandler && bHaveAudioReader.Value())
+			{
+				// If the audio decoder output buffer is stalled then we have enough video output. There won't be any more coming in.
+				// Also when the stream has reached the end and that has propagated through the buffers.
+				if (AudioBufferStats.DecoderOutputBuffer.bOutputStalled ||
+					(AudioBufferStats.StreamBuffer.bEndOfData && AudioBufferStats.DecoderInputBuffer.bEODSignaled && AudioBufferStats.DecoderOutputBuffer.bEODreached) ||
+					MultiStreamBufferAud.IsDeselected())
+				{
+					bHaveEnoughAudio = true;
+				}
+			}
+			else
+			{
+				bHaveEnoughAudio = true;
+			}
+		}
+
+		DiagnosticsCriticalSection.Unlock();
+
+		/*
+		// In case a decoder doesn't provide suitable information we check for a timeout.
+		if (mConfig.mMaxPrerollDurationUntilForcedStart > 0.0)
+		{
+			if ((timeNow - mPrerollVars.mStartTime).fsec() >= mConfig.mMaxPrerollDurationUntilForcedStart)
+			{
+				bHaveEnoughVideo = true;
+				bHaveEnoughAudio = true;
+			}
+		}
+		*/
+
+		// Keep current decision around to prime the next check.
+		if (bHaveEnoughVideo)
+		{
+			PrerollVars.bHaveEnoughVideo = true;
+		}
+		if (bHaveEnoughAudio)
+		{
+			PrerollVars.bHaveEnoughAudio = true;
+		}
+
+		// Ready to go?
+		if (bHaveEnoughVideo && bHaveEnoughAudio)
+		{
+			PipelineState = EPipelineState::ePipeline_Stopped;
+			CurrentState = EPlayerState::eState_Paused;
+			PlaybackRate = 0.0;
+
+			PrerollVars.Clear();
+			DispatchEvent(FMetricEvent::ReportPrerollEnd());
+		}
+	}
+}
+
+
+
+//-----------------------------------------------------------------------------
+/**
+ * Handles pending media segment requests.
+ */
+void FAdaptiveStreamingPlayer::HandlePendingMediaSegmentRequests()
+{
+	if (CurrentState == EPlayerState::eState_Paused || CurrentState == EPlayerState::eState_Error)
+	{
+		return;
+	}
+
+	FTimeValue	CurrentTime 			= MEDIAutcTime::Current();
+
+	// Is there a start request?
+	if (PendingStartRequest.IsValid())
+	{
+		// Need to have a manifest.
+		if (Manifest.IsValid())
+		{
+			// For sanities sake disable any loop flag that might have been set before loading the playlist if the presentation isn't VoD.
+			if (Manifest->GetPresentationType() != IManifest::EType::OnDemand)
+			{
+				CurrentLoopParam.bEnableLooping = false;
+			}
+
+			// Do we have a play period yet?
+			if (CurrentPlayPeriod.Get() == nullptr)
+			{
+				// No. Check if we should ask for one now.
+				if (!PendingStartRequest->RetryAtTime.IsValid() || CurrentTime >= PendingStartRequest->RetryAtTime)
+				{
+					// If the starting time has not been set we now check if we are dealing with a VoD or a Live stream and choose a starting point.
+					if (!PendingStartRequest->StartAt.Time.IsValid())
+					{
+						FTimeRange Seekable = CurrentTimeline->GetSeekableTimeRange();
+						if (Manifest->GetPresentationType() == IManifest::EType::OnDemand)
+						{
+							check(Seekable.Start.IsValid());
+							PendingStartRequest->StartAt.Time = Seekable.Start;
+						}
+						else
+						{
+							check(Seekable.End.IsValid());
+							PendingStartRequest->StartAt.Time = Seekable.End;
+						}
+					}
+
+					TSharedPtrTS<IManifest::IPlayPeriod> StartingPeriod;
+					IManifest::FResult Result = Manifest->FindPlayPeriod(StartingPeriod, PendingStartRequest->StartAt, PendingStartRequest->SearchType);
+					switch(Result.GetType())
+					{
+						case IManifest::FResult::EType::TryAgainLater:
+						{
+							PendingStartRequest->RetryAtTime = Result.GetRetryAgainAtTime();
+							break;
+						}
+						case IManifest::FResult::EType::Found:
+						{
+							check(StartingPeriod.IsValid());
+							if (StartingPeriod.IsValid())
+							{
+								StartingPeriod->SetStreamPreferences(StreamPreferences);
+							}
+							CurrentPlayPeriod = StartingPeriod;
+							// NOTE: Do *not* reset the pending start request yet. It is still needed.
+							break;
+						}
+						case IManifest::FResult::EType::NotFound:
+						case IManifest::FResult::EType::BeforeStart:
+						case IManifest::FResult::EType::PastEOS:
+						case IManifest::FResult::EType::NotLoaded:
+						{
+							// Throw a playback error for now.
+							check(!"TODO");
+							PendingStartRequest.Reset();
+							break;
+						}
+					}
+				}
+			}
+			// Do we have a starting play period now?
+			if (CurrentPlayPeriod.IsValid())
+			{
+				// Is it ready yet?
+				switch(CurrentPlayPeriod->GetReadyState())
+				{
+					case IManifest::IPlayPeriod::EReadyState::NotReady:
+					{
+						// Prepare the play period. This must immediately change the state away from NotReady
+						FParamDict Options;
+						CurrentPlayPeriod->PrepareForPlay(Options);
+						check(CurrentPlayPeriod->GetReadyState() != IManifest::IPlayPeriod::EReadyState::NotReady);
+						break;
+					}
+					case IManifest::IPlayPeriod::EReadyState::IsReady:
+					{
+						// With the period being ready we can now get the initial media segment request.
+						check(PendingStartRequest.IsValid());	// must not have been released yet
+
+					// TODO: this should update the stream metadata.
+
+						// Tell the ABR the current playback period.
+						StreamSelector->SetCurrentPlaybackPeriod(CurrentPlayPeriod);
+
+						// Single stream mp4 files cannot switch streams. Failures need to retry the same file.
+						if (ManifestReader && ManifestReader->GetPlaylistType() == "mp4")
+						{
+							StreamSelector->SetCanSwitchToAlternateStreams(false);
+						}
+
+						// Apply any set resolution limit at the start of the new play period.
+						UpdateStreamResolutionLimit();
+
+						//
+						switch(Manifest->GetPresentationType())
+						{
+							case IManifest::EType::OnDemand:
+								StreamSelector->SetPresentationType(IAdaptiveStreamSelector::EMediaPresentationType::OnDemand);
+								break;
+// FIXME: Need to differentiate between "Live" and "Realtime" here!
+							default:
+							case IManifest::EType::Live:
+								StreamSelector->SetPresentationType(IAdaptiveStreamSelector::EMediaPresentationType::Realtime);
+								break;
+						}
+
+						// Have the ABR select the initial bandwidth. Pass it an empty segment request to indicate this.
+						FTimeValue ActionDelay;
+						IAdaptiveStreamSelector::ESegmentAction Action = StreamSelector->SelectSuitableStreams(ActionDelay, TSharedPtrTS<IStreamSegment>());
+// FIXME: Can the starting choice with no previous segment even be anything besides failure or fetch?
+						check(Action == IAdaptiveStreamSelector::ESegmentAction::FetchNext || Action == IAdaptiveStreamSelector::ESegmentAction::Fail);
+						//if (Action == IAdaptiveStreamSelector::ESegmentAction::Fail)
+						if (Action != IAdaptiveStreamSelector::ESegmentAction::FetchNext)
+						{
+							FErrorDetail err;
+							err.SetFacility(Facility::EFacility::Player);
+							err.SetMessage("All streams have failed. There is nothing to play any more.");
+							err.SetCode(INTERR_ALL_STREAMS_HAVE_FAILED);
+							PostError(err);
+							PendingStartRequest.Reset();
+							break;
+						}
+
+						if (!PendingStartRequest->RetryAtTime.IsValid() || CurrentTime >= PendingStartRequest->RetryAtTime)
+						{
+							TSharedPtrTS<IStreamSegment> FirstSegmentRequest;
+							IManifest::FResult Result = CurrentPlayPeriod->GetStartingSegment(FirstSegmentRequest, PendingStartRequest->StartAt, PendingStartRequest->SearchType);
+							switch(Result.GetType())
+							{
+								case IManifest::FResult::EType::TryAgainLater:
+								{
+									PendingStartRequest->RetryAtTime = Result.GetRetryAgainAtTime();
+									break;
+								}
+								case IManifest::FResult::EType::Found:
+								{
+									PendingFirstSegmentRequest = FirstSegmentRequest;
+									// No longer need the start request.
+									PendingStartRequest.Reset();
+									break;
+								}
+								case IManifest::FResult::EType::NotFound:
+								case IManifest::FResult::EType::BeforeStart:
+								case IManifest::FResult::EType::PastEOS:
+								case IManifest::FResult::EType::NotLoaded:
+								{
+									FErrorDetail err;
+									err.SetFacility(Facility::EFacility::Player);
+									err.SetMessage("Could not locate the stream segment to begin playback at.");
+									err.SetCode(INTERR_COULD_NOT_LOCATE_START_SEGMENT);
+									PostError(err);
+									PendingStartRequest.Reset();
+									break;
+								}
+							}
+						}
+						break;
+					}
+					case IManifest::IPlayPeriod::EReadyState::Preparing:
+					{
+						// Period is not yet ready. Check again on the next iteration.
+					// TODO: get a try-again-at value instead of checking each loop?
+						break;
+					}
+					case IManifest::IPlayPeriod::EReadyState::MustReselect:
+					{
+						// The period has asked to be discarded.
+						CurrentPlayPeriod.Reset();
+						break;
+					}
+				}
+			}
+		}
+		else
+		{
+			// No manifest, no start request. Those are the rules.
+			PendingStartRequest.Reset();
+		}
+	}
+
+	// Is there a play start initial segment download request?
+	if (PendingFirstSegmentRequest.IsValid())
+	{
+		check(StreamReaderHandler);
+		if (StreamReaderHandler)
+		{
+			FTimeValue CurrentPlayPos = PlaybackState.GetPlayPosition();
+			FTimeValue NewPlayPos = PendingFirstSegmentRequest->GetFirstPTS();
+			PlaybackState.SetPlayPosition(NewPlayPos);
+			PlaybackState.SetLoopState(CurrentLoopState);
+
+			check(CurrentState == EPlayerState::eState_Buffering || CurrentState == EPlayerState::eState_Rebuffering || CurrentState == EPlayerState::eState_Seeking);
+			if (CurrentState == EPlayerState::eState_Seeking)
+			{
+				DispatchEvent(FMetricEvent::ReportJumpInPlayPosition(NewPlayPos, CurrentPlayPos, Metrics::ETimeJumpReason::UserSeek));
+			}
+
+			// Get the expected stream types from the starting segment request.
+		// FIXME: get rid of the "mbHave...XYZ...Reader" states to become more flexible with dynamic stream selection (enabling/disabling) during playback.
+			bHaveVideoReader.Set(false);
+			bHaveAudioReader.Set(false);
+			bHaveTextReader.Set(false);
+			switch(PendingFirstSegmentRequest->GetType())
+			{
+				case EStreamType::Video:	bHaveVideoReader.Set(true);	break;
+				case EStreamType::Audio:	bHaveAudioReader.Set(true);	break;
+				case EStreamType::Subtitle:	bHaveTextReader.Set(true);		break;
+				default:													break;
+			}
+			TArray<IStreamSegment::FDependentStreams> DependentStreams;
+			PendingFirstSegmentRequest->GetDependentStreams(DependentStreams);
+			for(int32 i=0; i<DependentStreams.Num(); ++i)
+			{
+				switch(DependentStreams[i].StreamType)
+				{
+					case EStreamType::Video:	bHaveVideoReader.Set(true);	break;
+					case EStreamType::Audio:	bHaveAudioReader.Set(true);	break;
+					case EStreamType::Subtitle:	bHaveTextReader.Set(true);		break;
+					default:													break;
+				}
+			}
+			// Get the stream types that have already reached EOS in this start request because the streams are of different duration.
+			// The streams are there in principle so we need to set them up (in case we loop back into them) but we won't see any data arriving.
+			TArray<TSharedPtrTS<IStreamSegment>> AlreadyEndedStreams;
+			PendingFirstSegmentRequest->GetEndedStreams(AlreadyEndedStreams);
+			for(int32 i=0; i<AlreadyEndedStreams.Num(); ++i)
+			{
+				CompletedSegmentRequests.Add(AlreadyEndedStreams[i]->GetType(), AlreadyEndedStreams[i]);
+			}
+
+			if (AudioRender.Renderer)
+			{
+				RenderClock->SetCurrentTime(IMediaRenderClock::ERendererType::Audio, NewPlayPos);
+				AudioRender.Renderer->SetNextApproximatePresentationTime(NewPlayPos);
+			}
+			if (VideoRender.Renderer)
+			{
+				RenderClock->SetCurrentTime(IMediaRenderClock::ERendererType::Video, NewPlayPos);
+				VideoRender.Renderer->SetNextApproximatePresentationTime(NewPlayPos);
+			}
+
+			PlaybackState.SetIsBuffering(true);
+			PlaybackState.SetIsSeeking(true);
+
+			// Send QoS buffering event
+			DispatchBufferingEvent(true, CurrentState);
+
+			// Remember why we are buffering so we can send the proper QoS event when buffering is done.
+			LastBufferingState = CurrentState;
+
+			// Whether we were seeking or rebuffering, we're now just buffering.
+			CurrentState = EPlayerState::eState_Buffering;
+
+			// Kick off the request
+//// TODO: check return code!!!
+			StreamReaderHandler->AddRequest(CurrentPlaybackSequenceID, PendingFirstSegmentRequest);
+		}
+		PendingFirstSegmentRequest.Reset();
+	}
+
+
+	// Are there completed downloads for which we need to find the next segment?
+	TQueue<FPendingSegmentRequest> pendingRequests;
+	Swap(pendingRequests, NextPendingSegmentRequests);
+	while(!pendingRequests.IsEmpty())
+	{
+		FPendingSegmentRequest FinishedReq;
+		pendingRequests.Dequeue(FinishedReq);
+
+		// If not yet ready to check again put the request back into the list.
+		if (FinishedReq.AtTime.IsValid() && FinishedReq.AtTime > CurrentTime)
+		{
+			NextPendingSegmentRequests.Enqueue(MoveTemp(FinishedReq));
+		}
+		else if (CurrentPlayPeriod.IsValid())
+		{
+			// Evaluate ABR to select the next stream quality.
+			FTimeValue ActionDelay;
+			IAdaptiveStreamSelector::ESegmentAction Action = StreamSelector->SelectSuitableStreams(ActionDelay, FinishedReq.Request);
+			if (Action == IAdaptiveStreamSelector::ESegmentAction::Fail)
+			{
+				FErrorDetail err;
+				err.SetFacility(Facility::EFacility::Player);
+				err.SetMessage("All streams have failed. There is nothing to play any more.");
+				err.SetCode(INTERR_ALL_STREAMS_HAVE_FAILED);
+				PostError(err);
+				break;
+			}
+
+			FParamDict Options;
+			TSharedPtrTS<IStreamSegment> pNext;
+			IManifest::FResult Result;
+			if (Action == IAdaptiveStreamSelector::ESegmentAction::FetchNext)
+			{
+				Result = CurrentPlayPeriod->GetNextSegment(pNext, FinishedReq.Request, Options);
+			}
+			else if (Action == IAdaptiveStreamSelector::ESegmentAction::Retry || Action == IAdaptiveStreamSelector::ESegmentAction::Fill)
+			{
+				if (Action == IAdaptiveStreamSelector::ESegmentAction::Fill)
+				{
+					Options.Set("insertFiller", FVariantValue(true));
+				}
+				Result = CurrentPlayPeriod->GetRetrySegment(pNext, FinishedReq.Request, Options);
+			}
+			switch(Result.GetType())
+			{
+				case IManifest::FResult::EType::TryAgainLater:
+				{
+					FinishedReq.AtTime = Result.GetRetryAgainAtTime();
+					NextPendingSegmentRequests.Enqueue(MoveTemp(FinishedReq));
+					break;
+				}
+				case IManifest::FResult::EType::Found:
+				{
+//// TODO: check return code!!!
+					StreamReaderHandler->AddRequest(CurrentPlaybackSequenceID, pNext);
+					break;
+				}
+				case IManifest::FResult::EType::PastEOS:
+				{
+					// Get the dependent stream types, if any.
+					// This is mainly for multiplexed streams like an .mp4 at this point.
+					// All stream types have reached EOS now.
+					TArray<IStreamSegment::FDependentStreams> DependentStreams;
+					FinishedReq.Request->GetDependentStreams(DependentStreams);
+					// Add the primary stream type to the list as well.
+					DependentStreams.AddDefaulted_GetRef().StreamType = FinishedReq.Request->GetType();
+					for(int32 i=0; i<DependentStreams.Num(); ++i)
+					{
+						CompletedSegmentRequests.Add(DependentStreams[i].StreamType, FinishedReq.Request);
+					}
+					break;
+				}
+				case IManifest::FResult::EType::NotFound:
+				case IManifest::FResult::EType::BeforeStart:
+				case IManifest::FResult::EType::NotLoaded:
+				{
+					// Throw a playback error for now.
+					FErrorDetail err;
+					err.SetFacility(Facility::EFacility::Player);
+					err.SetMessage("Could not locate next segment");
+					err.SetCode(INTERR_FRAGMENT_NOT_AVAILABLE);
+					PostError(err);
+					break;
+				}
+			}
+		}
+	}
+
+	// Check the streams that reached EOS. If all of them are done decide what to do, end playback or loop.
+	if (CompletedSegmentRequests.Num())
+	{
+		bool bAllAtEOS = true;
+		if (bHaveVideoReader.GetWithDefault(false))
+		{
+			if (!CompletedSegmentRequests.Contains(EStreamType::Video))
+			{
+				bAllAtEOS = false;
+			}
+			else
+			{
+				MultiStreamBufferVid.PushEndOfDataAll();
+			}
+		}
+		if (bHaveAudioReader.GetWithDefault(false))
+		{
+			if (!CompletedSegmentRequests.Contains(EStreamType::Audio))
+			{
+				bAllAtEOS = false;
+			}
+			else
+			{
+				MultiStreamBufferAud.PushEndOfDataAll();
+			}
+		}
+		if (bHaveTextReader.GetWithDefault(false))
+		{
+			if (!CompletedSegmentRequests.Contains(EStreamType::Subtitle))
+			{
+				bAllAtEOS = false;
+			}
+			else
+			{
+				MultiStreamBufferTxt.PushEndOfDataAll();
+			}
+		}
+		if (bAllAtEOS)
+		{
+			// Move the finished requests into a local work map, emptying the original.
+			TMultiMap<EStreamType, TSharedPtrTS<IStreamSegment>> LocalFinishedRequests(MoveTemp(CompletedSegmentRequests));
+
+			// Looping enabled?
+			if (CurrentLoopParam.bEnableLooping)
+			{
+				if (CurrentTimeline.IsValid() && CurrentPlayPeriod.IsValid())
+				{
+					FPlayStartPosition	StartAt;
+					FTimeRange			Seekable = CurrentTimeline->GetSeekableTimeRange();
+					StartAt.Time = Seekable.Start;
+
+					TSharedPtrTS<IStreamSegment> FirstSegmentRequest;
+					IManifest::FResult Result = CurrentPlayPeriod->GetLoopingSegment(FirstSegmentRequest, CurrentLoopState, LocalFinishedRequests, StartAt, IManifest::ESearchType::Closest);
+					if (Result.GetType() == IManifest::FResult::EType::Found)
+					{
+						// Get the stream types that have already reached EOS in this start request because the streams are of different duration.
+						// The streams are there in principle so we need to set them up (in case we loop back into them) but we won't see any data arriving.
+						TArray<TSharedPtrTS<IStreamSegment>> AlreadyEndedStreams;
+						FirstSegmentRequest->GetEndedStreams(AlreadyEndedStreams);
+						for(int32 i=0; i<AlreadyEndedStreams.Num(); ++i)
+						{
+							CompletedSegmentRequests.Add(AlreadyEndedStreams[i]->GetType(), AlreadyEndedStreams[i]);
+						}
+
+						PostLog(Facility::EFacility::Player, IInfoLog::ELevel::Verbose, FString::Printf(TEXT("Enqueuing stream loop start request.")));
+						// Add the updated loop state to the upcoming loop FIFO. When the play position reaches the next enqueued loop state it gets popped.
+						NextLoopStates.Push(CurrentLoopState);
+						StreamReaderHandler->AddRequest(CurrentPlaybackSequenceID, FirstSegmentRequest);
+					}
+					else
+					{
+						// How could this have failed? Switch off looping to end the stream.
+						PostLog(Facility::EFacility::Player, IInfoLog::ELevel::Error, FString::Printf(TEXT("Could not locate segment to loop back to. Disabling looping and ending stream.")));
+						CurrentLoopParam.bEnableLooping = false;
+					}
+				}
+			}
+		}
+	}
+
+}
+
+
+
+//-----------------------------------------------------------------------------
+/**
+ * Checks the error collectors for any errors thrown by the decoders.
+ */
+void FAdaptiveStreamingPlayer::CheckForErrors()
+{
+	while(!ErrorQueue.IsEmpty())
+	{
+		TSharedPtrTS<FErrorDetail> Error = ErrorQueue.Pop();
+		// Do this only once in case there will be several decoder errors thrown.
+		if (CurrentState != EPlayerState::eState_Error)
+		{
+			// Pause before setting up error.
+			InternalPause();
+
+			CurrentState = EPlayerState::eState_Error;
+			// Only keep the first error, not any errors after that which may just be the avalanche and not the cause.
+			if (!LastErrorDetail.IsSet())
+			{
+				DiagnosticsCriticalSection.Lock();
+				LastErrorDetail = *Error;
+				DiagnosticsCriticalSection.Unlock();
+				DispatchEvent(FMetricEvent::ReportError(LastErrorDetail));
+			}
+			// In error state we do not need any periodic manifest refetches any more.
+			if (ManifestReader)
+			{
+				delete ManifestReader;
+				ManifestReader = nullptr;
+			}
+		}
+	}
+}
+
+
+
+//-----------------------------------------------------------------------------
+/**
+ * Updates the data availability state if it has changed and dispatches the
+ * corresponding metrics event.
+ *
+ * @param DataAvailabilityState
+ * @param NewAvailability
+ */
+void FAdaptiveStreamingPlayer::UpdateDataAvailabilityState(Metrics::FDataAvailabilityChange& DataAvailabilityState, Metrics::FDataAvailabilityChange::EAvailability NewAvailability)
+{
+	if (DataAvailabilityState.Availability != NewAvailability)
+	{
+		DataAvailabilityState.Availability = NewAvailability;
+		DispatchEvent(FMetricEvent::ReportDataAvailabilityChange(DataAvailabilityState));
+	}
+}
+
+
+
+//-----------------------------------------------------------------------------
+/**
+ * Sends an available AU to a decoder.
+ * If the current buffer level is below the underrun threshold an underrun
+ * message is sent to the worker thread.
+ *
+ * @param Type
+ * @param FromMultistreamBuffer
+ * @param Decoder
+ */
+void FAdaptiveStreamingPlayer::FeedDecoder(EStreamType Type, FMultiTrackAccessUnitBuffer& FromMultistreamBuffer, IAccessUnitBufferInterface* Decoder)
+{
+	// Lock the AU buffer for the duration of this function to ensure this can never clash with a Flush() call
+	// since we are checking size, eod state and subsequently popping an AU, for which the buffer must stay consistent inbetween!
+	// Also to ensure the active buffer doesn't get changed from one track to another.
+	FMultiTrackAccessUnitBuffer::FScopedLock lock(FromMultistreamBuffer);
+
+	FBufferStats* pStats = nullptr;
+	Metrics::FDataAvailabilityChange* pAvailability = nullptr;
+	switch(Type)
+	{
+		case EStreamType::Video:
+			pStats = &VideoBufferStats;
+			pAvailability = &DataAvailabilityStateVid;
+			break;
+		case EStreamType::Audio:
+			pStats = &AudioBufferStats;
+			pAvailability = &DataAvailabilityStateAud;
+			break;
+		case EStreamType::Subtitle:
+			pStats = &TextBufferStats;
+			pAvailability = &DataAvailabilityStateTxt;
+			break;
+		default:
+			checkNoEntry();
+			return;
+	}
+
+	// Is the buffer (the Type of elementary stream actually) active/selected?
+	if (!FromMultistreamBuffer.IsDeselected())
+	{
+		// Check for buffer underrun.
+		if (!bRebufferPending && CurrentState == EPlayerState::eState_Playing && StreamState == EStreamState::eStream_Running && PipelineState == EPipelineState::ePipeline_Running)
+		{
+			bool bEODSet = FromMultistreamBuffer.IsEODFlagSet();
+			if (!bEODSet && FromMultistreamBuffer.Num() == 0)
+			{
+				// Buffer underrun.
+				bRebufferPending = true;
+				FTimeValue LastKnownPTS = FromMultistreamBuffer.GetLastPoppedPTS();
+				// Only set the 'rebuffer at' time if we have a valid last known PTS. If we don't
+				// then maybe this is a cascade failure from a previous rebuffer attempt for which
+				// we then try that time once more.
+				if (LastKnownPTS.IsValid())
+				{
+					RebufferDetectedAtPlayPos = LastKnownPTS;
+				}
+				WorkerThread.SendMessage(FWorkerThread::FMessage::EType::BufferUnderrun);
+			}
+		}
+
+		FAccessUnit* AccessUnit = nullptr;
+		bool bOk;
+		bOk = FromMultistreamBuffer.Pop(AccessUnit);
+		if (AccessUnit)
+		{
+			if (Decoder)
+			{
+			// The decoder has asked to be fed a new AU so it better be able to accept it.
+				/*IAccessUnitBufferInterface::EAUpushResult auPushRes =*/ Decoder->AUdataPushAU(AccessUnit);
+			}
+
+			// The decoder will have added a ref count if it took the AU. If it didnt' for whatever reason
+			// we still release it to get rid of it and not cause a memory leak.
+			FAccessUnit::Release(AccessUnit);
+			AccessUnit = nullptr;
+
+			if (pAvailability)
+			{
+				UpdateDataAvailabilityState(*pAvailability, Metrics::FDataAvailabilityChange::EAvailability::DataAvailable);
+			}
+		}
+	}
+	// An AU is not tagged as being "the last" one. Instead the EOD is handled separately and must be dealt with
+	// by the decoders accordingly.
+	if (FromMultistreamBuffer.IsEODFlagSet() && FromMultistreamBuffer.Num() == 0)
+	{
+		if (pStats && !pStats->DecoderInputBuffer.bEODSignaled)
+		{
+			if (Decoder)
+			{
+				Decoder->AUdataPushEOD();
+			}
+		}
+		if (pAvailability)
+		{
+			UpdateDataAvailabilityState(*pAvailability, Metrics::FDataAvailabilityChange::EAvailability::DataNotAvailable);
+		}
+	}
+}
+
+void FAdaptiveStreamingPlayer::VideoDecoderInputNeeded(const IAccessUnitBufferListener::FBufferStats &currentInputBufferStats)
+{
+	DiagnosticsCriticalSection.Lock();
+	VideoBufferStats.DecoderInputBuffer = currentInputBufferStats;
+	DiagnosticsCriticalSection.Unlock();
+	EDecoderState decState = DecoderState;
+	if (decState == EDecoderState::eDecoder_Running)
+	{
+		FeedDecoder(EStreamType::Video, MultiStreamBufferVid, VideoDecoder.Decoder);
+	}
+}
+
+void FAdaptiveStreamingPlayer::VideoDecoderOutputReady(const IDecoderOutputBufferListener::FDecodeReadyStats &currentReadyStats)
+{
+	DiagnosticsCriticalSection.Lock();
+	VideoBufferStats.DecoderOutputBuffer = currentReadyStats;
+	DiagnosticsCriticalSection.Unlock();
+}
+
+
+void FAdaptiveStreamingPlayer::AudioDecoderInputNeeded(const IAccessUnitBufferListener::FBufferStats &currentInputBufferStats)
+{
+	DiagnosticsCriticalSection.Lock();
+	AudioBufferStats.DecoderInputBuffer = currentInputBufferStats;
+	DiagnosticsCriticalSection.Unlock();
+	EDecoderState decState = DecoderState;
+	if (decState == EDecoderState::eDecoder_Running)
+	{
+		FeedDecoder(EStreamType::Audio, MultiStreamBufferAud, AudioDecoder.Decoder);
+	}
+}
+
+void FAdaptiveStreamingPlayer::AudioDecoderOutputReady(const IDecoderOutputBufferListener::FDecodeReadyStats &currentReadyStats)
+{
+	DiagnosticsCriticalSection.Lock();
+	AudioBufferStats.DecoderOutputBuffer = currentReadyStats;
+	DiagnosticsCriticalSection.Unlock();
+}
+
+
+
+
+bool FAdaptiveStreamingPlayer::FindMatchingStreamInfo(FStreamCodecInformation& OutStreamInfo, int32 MaxWidth, int32 MaxHeight)
+{
+	if (!CurrentPlayPeriod.IsValid())
+	{
+		return false;
+	}
+	TArray<FStreamCodecInformation> VideoCodecInfos;
+
+	TSharedPtrTS<ITimelineMediaAsset> Asset = CurrentPlayPeriod->GetMediaAsset();
+	if (Asset.IsValid())
+	{
+		check(Asset->GetNumberOfAdaptationSets(EStreamType::Video) <= 1); // what if this is more than one?
+		if (Asset->GetNumberOfAdaptationSets(EStreamType::Video) > 0)
+		{
+			TSharedPtrTS<IPlaybackAssetAdaptationSet> VideoSet = Asset->GetAdaptationSetByTypeAndIndex(EStreamType::Video, 0);
+			check(VideoSet.IsValid());
+			if (VideoSet.IsValid())
+			{
+				for(int32 i = 0, iMax = VideoSet->GetNumberOfRepresentations(); i < iMax; ++i)
+				{
+					VideoCodecInfos.Push(VideoSet->GetRepresentationByIndex(i)->GetCodecInformation());
+				}
+				check(VideoCodecInfos.Num());
+				if (VideoCodecInfos.Num())
+				{
+					if (MaxWidth == 0 && MaxHeight == 0)
+					{
+						FStreamCodecInformation Best = VideoCodecInfos[0];
+						for(int32 i=1; i<VideoCodecInfos.Num(); ++i)
+						{
+							const FStreamCodecInformation::FResolution& Res = VideoCodecInfos[i].GetResolution();
+							if (Res.Width > Best.GetResolution().Width)
+							{
+								Best.SetResolution(FStreamCodecInformation::FResolution(Res.Width, Best.GetResolution().Height));
+							}
+							if (Res.Height > Best.GetResolution().Height)
+							{
+								Best.SetResolution(FStreamCodecInformation::FResolution(Best.GetResolution().Width, Res.Height));
+							}
+							// Note: the final RFC 6381 codec string will be bogus since we do not re-create it here.
+							if (VideoCodecInfos[i].GetProfile() > Best.GetProfile())
+							{
+								Best.SetProfile(VideoCodecInfos[i].GetProfile());
+							}
+							if (VideoCodecInfos[i].GetProfileLevel() > Best.GetProfileLevel())
+							{
+								Best.SetProfileLevel(VideoCodecInfos[i].GetProfileLevel());
+							}
+							if (VideoCodecInfos[i].GetExtras().GetValue("b_frames").SafeGetInt64(0))
+							{
+								Best.GetExtras().Set("b_frames", FVariantValue((int64) 1));
+							}
+						}
+						OutStreamInfo = Best;
+					}
+					else
+					{
+						if (MaxWidth == 0)
+						{
+							MaxWidth = 32768;
+						}
+						if (MaxHeight == 0)
+						{
+							MaxHeight = 32768;
+						}
+						FStreamCodecInformation		Best;
+						bool bFirst = true;
+						for(int32 i=0; i<VideoCodecInfos.Num(); ++i)
+						{
+							const FStreamCodecInformation::FResolution& Res = VideoCodecInfos[i].GetResolution();
+							if (Res.ExceedsLimit(MaxWidth, MaxHeight))
+							{
+								continue;
+							}
+							if (bFirst)
+							{
+								bFirst = false;
+								Best = VideoCodecInfos[i];
+							}
+							if (Res.Width > Best.GetResolution().Width)
+							{
+								Best.SetResolution(FStreamCodecInformation::FResolution(Res.Width, Best.GetResolution().Height));
+							}
+							if (Res.Height > Best.GetResolution().Height)
+							{
+								Best.SetResolution(FStreamCodecInformation::FResolution(Best.GetResolution().Width, Res.Height));
+							}
+							// Note: the final RFC 6381 codec string will be bogus since we do not re-create it here.
+							if (VideoCodecInfos[i].GetProfile() > Best.GetProfile())
+							{
+								Best.SetProfile(VideoCodecInfos[i].GetProfile());
+							}
+							if (VideoCodecInfos[i].GetProfileLevel() > Best.GetProfileLevel())
+							{
+								Best.SetProfileLevel(VideoCodecInfos[i].GetProfileLevel());
+							}
+							if (VideoCodecInfos[i].GetExtras().GetValue("b_frames").SafeGetInt64(0))
+							{
+								Best.GetExtras().Set("b_frames", FVariantValue((int64) 1));
+							}
+						}
+						// Found none? (resolution limit set too low)
+						if (bFirst)
+						{
+							// Find smallest by bandwidth
+							Best = VideoCodecInfos[0];
+							int32 BestBW = VideoSet->GetRepresentationByIndex(0)->GetBitrate();
+							for(int32 i=1; i<VideoCodecInfos.Num(); ++i)
+							{
+								if (VideoSet->GetRepresentationByIndex(i)->GetBitrate() < BestBW)
+								{
+									Best = VideoCodecInfos[i];
+									BestBW = VideoSet->GetRepresentationByIndex(i)->GetBitrate();
+								}
+							}
+						}
+						OutStreamInfo = Best;
+					}
+					return true;
+				}
+			}
+		}
+	}
+	return false;
+}
+
+
+//-----------------------------------------------------------------------------
+/**
+ * Updates the ABR and video decoder with maximum stream resolution limits.
+ */
+void FAdaptiveStreamingPlayer::UpdateStreamResolutionLimit()
+{
+	FStreamCodecInformation StreamInfo;
+	// Check if there is a maximum texture resolution limit in place that the video decoder was configured in the initial Open() call with.
+	// If it is then this is an absolute override over everything else.
+	if (MaxVideoTextureResolutionLimitHeight > 0 && VideoResolutionLimitHeight > MaxVideoTextureResolutionLimitHeight)
+	{
+		VideoResolutionLimitHeight = MaxVideoTextureResolutionLimitHeight;
+	}
+
+	if (FindMatchingStreamInfo(StreamInfo, VideoResolutionLimitWidth, VideoResolutionLimitHeight))
+	{
+		VideoResolutionLimitWidth  = StreamInfo.GetResolution().Width;
+		VideoResolutionLimitHeight = StreamInfo.GetResolution().Height;
+		if (VideoDecoder.Decoder)
+		{
+			VideoDecoder.Decoder->SetMaximumDecodeCapability(VideoResolutionLimitWidth, VideoResolutionLimitHeight, StreamInfo.GetProfile(), StreamInfo.GetProfileLevel(), StreamInfo.GetExtras());
+		}
+	}
+	StreamSelector->SetMaxVideoResolution(VideoResolutionLimitWidth, VideoResolutionLimitHeight);
+}
+
+
+//-----------------------------------------------------------------------------
+/**
+ * Checks if the stream readers, decoders and renderers are finished.
+ */
+void FAdaptiveStreamingPlayer::CheckForStreamEnd()
+{
+	if (CurrentState == EPlayerState::eState_Playing)
+	{
+		// When set to loop the stream will not end we we do not need to check further.
+		if (!CurrentLoopParam.bEnableLooping)
+		{
+			if (StreamState == EStreamState::eStream_Running)
+			{
+				// For simplicity we assume all streams are done if we know whether or not they exist.
+				bool bEndVid = bHaveVideoReader.IsSet();
+				bool bEndAud = bHaveAudioReader.IsSet();
+				bool bEndTxt = bHaveTextReader.IsSet();
+
+				DiagnosticsCriticalSection.Lock();
+				FBufferStats vidStats = VideoBufferStats;
+				FBufferStats audStats = AudioBufferStats;
+				DiagnosticsCriticalSection.Unlock();
+
+				// Check for end of video stream
+				if (StreamReaderHandler && bHaveVideoReader.GetWithDefault(false))
+				{
+					// All buffers at end of data?
+					bEndVid = (vidStats.StreamBuffer.bEndOfData && vidStats.DecoderInputBuffer.bEODSignaled && vidStats.DecoderOutputBuffer.bEODreached);
+				}
+
+				// Check for end of audio stream
+				if (StreamReaderHandler && bHaveAudioReader.GetWithDefault(false))
+				{
+					// All buffers at end of data?
+					bEndAud = (audStats.StreamBuffer.bEndOfData && audStats.DecoderInputBuffer.bEODSignaled && audStats.DecoderOutputBuffer.bEODreached);
+				}
+
+				// Text stream
+				// ...
+
+
+				// Everything at EOD?
+				if (bEndVid && bEndAud && bEndTxt)
+				{
+					PostrollVars.Clear();
+					PlaybackState.SetHasEnded(true);
+					StreamState = EStreamState::eStream_ReachedEnd;
+					// Go to pause mode now.
+					InternalPause();
+					DispatchEvent(FMetricEvent::ReportPlaybackEnded());
+					// In case a seek back into the stream will happen we reset the first start state.
+					PrerollVars.bIsVeryFirstStart = true;
+				}
+			}
+			else if (StreamState == EStreamState::eStream_ReachedEnd)
+			{
+				// ... anything special we could still do here?
+			}
+		}
+	}
+}
+
+
+//-----------------------------------------------------------------------------
+/**
+ * Finds fragments for a given time, creates the stream readers
+ * and issues the first fragment request.
+ *
+ * @param NewPosition
+ *
+ * @return
+ */
+bool FAdaptiveStreamingPlayer::InternalStartAt(const FSeekParam& NewPosition)
+{
+	++CurrentPlaybackSequenceID;
+
+	StreamState = EStreamState::eStream_Running;
+
+	PlaybackState.SetHasEnded(false);
+
+	int64 StartingBitrate = StreamSelector->GetAverageBandwidth();
+	if (StartingBitrate <= 0)
+	{
+		StartingBitrate = PlayerOptions.GetValue("initial_bitrate").SafeGetInt64(0);
+	}
+	if (StartingBitrate <= 0)
+	{
+		StartingBitrate = Manifest->GetDefaultStartingBitrate();
+		if (StartingBitrate <= 0)
+		{
+			StartingBitrate = 500000;
+		}
+	}
+
+	// Initialize buffers
+	MultiStreamBufferVid.CapacitySet(PlayerConfig.StreamBufferConfigVideo);
+	MultiStreamBufferAud.CapacitySet(PlayerConfig.StreamBufferConfigAudio);
+	MultiStreamBufferTxt.CapacitySet(PlayerConfig.StreamBufferConfigText);
+	// Auto-select the first track that is pushing data into the buffer.
+	MultiStreamBufferVid.AutoselectFirstTrack();
+	MultiStreamBufferAud.AutoselectFirstTrack();
+	MultiStreamBufferTxt.AutoselectFirstTrack();
+
+	// Update data availability states in case this wasn't done yet.
+	UpdateDataAvailabilityState(DataAvailabilityStateVid, Metrics::FDataAvailabilityChange::EAvailability::DataNotAvailable);
+	UpdateDataAvailabilityState(DataAvailabilityStateAud, Metrics::FDataAvailabilityChange::EAvailability::DataNotAvailable);
+	UpdateDataAvailabilityState(DataAvailabilityStateTxt, Metrics::FDataAvailabilityChange::EAvailability::DataNotAvailable);
+
+	// For an mp4 file we expect several audio tracks to be available and select the one matching the selected index, if there is one.
+	if (ManifestReader && ManifestReader->GetPlaylistType() == "mp4")
+	{
+		// Only select audio for now.
+		if (InitialStreamSelectionAud.IsValid())
+		{
+			MultiStreamBufferAud.SelectTrackByID(InitialStreamSelectionAud->StreamUniqueID);
+		}
+// TODO/FIXME: Cannot start with audio disabled at the moment
+#if 0
+		else
+		{
+			mMultiStreamBufferAud.Deselect();
+		}
+#endif
+	}
+
+	// (Re-)configure the stream selector
+	StreamSelector->SetBandwidth(StartingBitrate);
+	StreamSelector->SetBandwidthCeiling(BitrateCeiling);
+	StreamSelector->SetMaxVideoResolution(VideoResolutionLimitWidth, VideoResolutionLimitHeight);
+	// On the very first playback enforce the starting bitrate.
+	if (bIsPlayStart)
+	{
+		StreamSelector->SetForcedNextBandwidth(StartingBitrate);
+	}
+
+	// From this point on any further start will not be the very first one any more.
+	bIsPlayStart = false;
+
+	// Create the stream reader
+	check(StreamReaderHandler == nullptr);
+	StreamReaderHandler = Manifest->CreateStreamReaderHandler();
+	check(StreamReaderHandler);
+	IStreamReader::CreateParam srhParam;
+//	srhParam.mParam.ReaderConfig    = mConfig.mReader;
+	srhParam.Options = PlayerOptions;
+	srhParam.EventListener  = this;
+	srhParam.MemoryProvider = this;
+	srhParam.PlayerSessionService = this;
+	if (StreamReaderHandler->Create(this, srhParam) != UEMEDIA_ERROR_OK)
+	{
+		FErrorDetail err;
+		err.SetFacility(Facility::EFacility::Player);
+		err.SetMessage("Failed to create stream reader");
+		err.SetCode(INTERR_CREATE_FRAGMENT_READER);
+		PostError(err);
+		return false;
+	}
+
+	check(PendingStartRequest.Get() == nullptr);
+	check(CurrentPlayPeriod.Get() == nullptr);
+	check(PendingFirstSegmentRequest.Get() == nullptr);
+	check(NextPendingSegmentRequests.IsEmpty());
+	check(CompletedSegmentRequests.Num() == 0);
+	NextPendingSegmentRequests.Empty();
+	PendingFirstSegmentRequest.Reset();
+	CompletedSegmentRequests.Empty();
+	CurrentPlayPeriod.Reset();
+	PendingStartRequest = MakeSharedTS<FPendingStartRequest>();
+	PendingStartRequest->SearchType = IManifest::ESearchType::Closest;
+	PendingStartRequest->StartAt.Time = NewPosition.Time;
+	PendingStartRequest->InitialBandwidth.Set(StartingBitrate);
+
+	// The fragment should be the closest to the time stamp unless we are rebuffering in which case it should be the one we failed on.
+/*
+	if (mLastBufferingState == eState_Rebuffering)
+	{
+		mpPendingStartRequest->SearchType = IManifest::ESearchType::Before;
+	}
+	else
+*/
+	{
+		PendingStartRequest->SearchType = IManifest::ESearchType::Closest;
+	}
+
+	return true;
+}
+
+
+
+//-----------------------------------------------------------------------------
+/**
+ * Pauses playback.
+ */
+void FAdaptiveStreamingPlayer::InternalPause()
+{
+	// Pause decoders to stop feeding them data.
+	// NOTE: A decoder may be in InputNeeded() / FeedDecoder() at this point!
+	//       This is to prevent their next ask for new data.
+	DecoderState = EDecoderState::eDecoder_Paused;
+
+	// Stop rendering.
+	StopRendering();
+	PipelineState = EPipelineState::ePipeline_Stopped;
+	// Pause the stream reader
+	PauseStreamReaders();
+
+	// Pause playing.
+	if (CurrentState != EPlayerState::eState_Error)
+	{
+		PlaybackRate = 0.0;
+		CurrentState = EPlayerState::eState_Paused;
+		PlaybackState.SetPausedAndPlaying(true, false);
+		DispatchEvent(FMetricEvent::ReportPlaybackPaused());
+	}
+}
+
+
+//-----------------------------------------------------------------------------
+/**
+ * Resumes playback.
+ */
+void FAdaptiveStreamingPlayer::InternalResume()
+{
+	// Cannot resume when in error state or when stream has reached the end.
+	if (CurrentState != EPlayerState::eState_Error && StreamState != EStreamState::eStream_ReachedEnd)
+	{
+		// Resume the stream reader
+		ResumeStreamReaders();
+
+		// Start rendering.
+		check(PipelineState == EPipelineState::ePipeline_Stopped);
+		StartRendering();
+		PipelineState = EPipelineState::ePipeline_Running;
+
+		// Resume decoders.
+		DecoderState = EDecoderState::eDecoder_Running;
+
+		if (PrerollVars.bIsVeryFirstStart)
+		{
+			PrerollVars.bIsVeryFirstStart = false;
+			// Send play start event
+			DispatchEvent(FMetricEvent::ReportPlaybackStart());
+		}
+
+		// Resume playing.
+		CurrentState = EPlayerState::eState_Playing;
+		PlaybackRate = 1.0;
+		PlaybackState.SetPausedAndPlaying(false, true);
+		PostrollVars.Clear();
+		DispatchEvent(FMetricEvent::ReportPlaybackResumed());
+	}
+}
+
+
+
+
+
+void FAdaptiveStreamingPlayer::InternalStop(bool bHoldCurrentFrame)
+{
+	// Increase the playback sequence ID. This will cause all async segment request message
+	// to be discarded when they are received.
+	++CurrentPlaybackSequenceID;
+
+	// Pause decoders to stop feeding them data.
+	// NOTE: A decoder may be in InputNeeded() / FeedDecoder() at this point!
+	//       This is to prevent their next ask for new data.
+	DecoderState = EDecoderState::eDecoder_Paused;
+
+	// Stop rendering.
+	StopRendering();
+	PipelineState = EPipelineState::ePipeline_Stopped;
+
+	// Destroy the stream reader handler.
+	IStreamReader* CurrentStreamReader = TMediaInterlockedExchangePointer(StreamReaderHandler, (IStreamReader*)nullptr);
+	if (CurrentStreamReader)
+		CurrentStreamReader->Close();
+	delete CurrentStreamReader;
+
+	bHaveVideoReader.Reset();
+	bHaveAudioReader.Reset();
+	bHaveTextReader.Reset();
+
+	// Release any pending segment requests.
+	PendingStartRequest.Reset();
+	PendingFirstSegmentRequest.Reset();
+	NextPendingSegmentRequests.Empty();
+	CompletedSegmentRequests.Empty();
+	CurrentPlayPeriod.Reset();
+
+	PlaybackRate = 0.0;
+	CurrentState = EPlayerState::eState_Paused;
+	PlaybackState.SetPausedAndPlaying(false, false);
+
+	// Flush all access unit buffers.
+	MultiStreamBufferVid.Flush();
+	MultiStreamBufferAud.Flush();
+	MultiStreamBufferTxt.Flush();
+
+	// Flush the renderers once before flushing the decoders.
+	// In case the media samples hold references to the decoder (because of shared frame resources) it could be a possibility
+	// that the decoder cannot flush as long as those references are "out there".
+	AudioRender.Flush();
+	VideoRender.Flush(bHoldCurrentFrame);
+
+	// Flush the decoders.
+	AudioDecoder.Flush();
+	VideoDecoder.Flush();
+
+	// Flush the renderers again, this time to discard everything the decoders may have emitted while being flushed.
+	AudioRender.Flush();
+	VideoRender.Flush(bHoldCurrentFrame);
+}
+
+
+void FAdaptiveStreamingPlayer::InternalClose()
+{
+	// No longer need the manifest reader/updater
+	delete ManifestReader;
+	ManifestReader = nullptr;
+
+	DestroyDecoders();
+	DestroyRenderers();
+
+	// Do not clear the playback state to allow GetPlayPosition() to query the last play position.
+	// This also means queries like IsBuffering() will return the last state but calling those
+	// after a Stop() can be considered at least weird practice.
+	//PlaybackState.Reset();
+	Manifest.Reset();
+
+	HttpManager.Reset();
+	RemoveMetricsReceiver(StreamSelector.Get());
+	StreamSelector.Reset();
+
+	// Reset remaining internal state
+	CurrentState   	= EPlayerState::eState_Idle;
+	PlaybackRate   	= 0.0;
+	StreamState		= EStreamState::eStream_Running;
+	bRebufferPending   = false;
+	LastBufferingState = EPlayerState::eState_Buffering;
+	StartAtTime.Reset();
+	RebufferDetectedAtPlayPos.SetToInvalid();
+	PrerollVars.Clear();
+	PostrollVars.Clear();
+
+	bShouldBePaused = false;
+	bShouldBePlaying = false;
+	bSeekPending = false;
+	CurrentLoopState.Reset();
+	NextLoopStates.Clear();
+
+
+	// Clear error state.
+	LastErrorDetail.Clear();
+	ErrorQueue.Clear();
+
+	// Clear diagnostics.
+	DiagnosticsCriticalSection.Lock();
+	VideoBufferStats.Clear();
+	AudioBufferStats.Clear();
+	TextBufferStats.Clear();
+	DiagnosticsCriticalSection.Unlock();
+}
+
+
+void FAdaptiveStreamingPlayer::InternalSetLoop(const FLoopParam& LoopParam)
+{
+	// Looping only makes sense for on demand presentations since Live does not have an end to loop at.
+	// In case this is called that early on that we can't determine the presentation type yet we allow
+	// looping to be set. It won't get evaluated anyway since the Live presentation doesn't end.
+	// The only issue with this is that until a manifest is loaded querying loop enable will return true.
+	if (!Manifest || Manifest->GetPresentationType() == IManifest::EType::OnDemand)
+	{
+		CurrentLoopParam = LoopParam;
+		// For the lack of better knowledge update the current state immediately.
+		CurrentLoopState.bLoopEnabled = LoopParam.bEnableLooping;
+		PlaybackState.SetLoopStateEnable(LoopParam.bEnableLooping);
+	}
+}
+
+
+
+//-----------------------------------------------------------------------------
+/**
+ * Rebuffers at the current play position.
+ */
+void FAdaptiveStreamingPlayer::InternalRebuffer()
+{
+	// Check if we are configured to throw a playback error instead of doing a rebuffer.
+	// This can be used by a higher logic to force playback at a different position instead.
+	if (PlayerOptions.GetValue("throw_error_when_rebuffering").SafeGetBool(false) == false)
+	{
+		// Pause decoders to stop feeding them data.
+		// NOTE: A decoder may be in InputNeeded() / FeedDecoder() at this point!
+		//       This is to prevent their next ask for new data.
+		DecoderState = EDecoderState::eDecoder_Paused;
+
+		// Stop rendering.
+		StopRendering();
+		PipelineState = EPipelineState::ePipeline_Stopped;
+
+		CurrentState = EPlayerState::eState_Rebuffering;
+		LastBufferingState = EPlayerState::eState_Rebuffering;
+
+		bRebufferPending = false;
+
+		/*
+		This gets handled implicitly after the InternalStartAt()'s initial request.
+		PlaybackState.SetIsBuffering(true);
+		// Send QoS buffering event
+		DispatchBufferingEvent(true, mLastBufferingState);
+		*/
+
+
+		// We do not wait on the current segment that caused the rebuffering to complete.
+		// Instead we perform a seek to that time and start over from scratch.
+		if (Manifest->GetPresentationType() == IManifest::EType::Live)
+		{
+			// A Live presentation shall implicitly go to the Live edge.
+			StartAtTime.Time.SetToInvalid();
+			PostLog(Facility::EFacility::Player, IInfoLog::ELevel::Info, FString::Printf(TEXT("Rebuffering Live stream implicitly goes to the live edge")));
+		}
+		else
+		{
+			check(RebufferDetectedAtPlayPos.IsValid());
+			StartAtTime.Time = RebufferDetectedAtPlayPos;
+			PostLog(Facility::EFacility::Player, IInfoLog::ELevel::Info, FString::Printf(TEXT("Rebuffering on-demand stream at current position")));
+		}
+		RebufferDetectedAtPlayPos.SetToInvalid();
+		InternalStop(PlayerConfig.bHoldLastFrameDuringSeek);
+		CurrentState = EPlayerState::eState_Rebuffering;
+		InternalStartAt(StartAtTime);
+	}
+	else
+	{
+		FErrorDetail err;
+		err.SetFacility(Facility::EFacility::Player);
+		err.SetMessage("Rebuffering is set to generate a playback error");
+		err.SetCode(INTERR_REBUFFER_SHALL_THROW_ERROR);
+		PostError(err);
+	}
+}
+
+//-----------------------------------------------------------------------------
+/**
+*/
+#if PLATFORM_ANDROID
+void FAdaptiveStreamingPlayer::Android_UpdateSurface(const TSharedPtr<IOptionPointerValueContainer>& Surface)
+{
+	if (VideoDecoder.Decoder)
+	{
+		VideoDecoder.Decoder->Android_UpdateSurface(Surface);
+	}
+}
+#endif
+
+//-----------------------------------------------------------------------------
+/**
+ * Debug prints information to the screen
+ *
+ * @param pPrintFN
+ */
+void FAdaptiveStreamingPlayer::DebugPrint(void* pPlayer, void (*pPrintFN)(void* pPlayer, const char *pFmt, ...))
+{
+	DiagnosticsCriticalSection.Lock();
+	const FAccessUnitBufferInfo *pD = nullptr;
+	if (StreamReaderHandler && bHaveVideoReader.GetWithDefault(false))
+	{
+		pD = &VideoBufferStats.StreamBuffer;
+		pPrintFN(pPlayer, "Video buffer  : %3u AUs; %8u/%8u bytes in; eod %d; %#7.4fs", (uint32)pD->NumCurrentAccessUnits, (uint32)pD->CurrentMemInUse, (uint32)pD->MaxDataSize, pD->bEndOfData, pD->PushedDuration.GetAsSeconds());
+		pPrintFN(pPlayer, "Video decoder : %2u in decoder, %zu total, %s; EOD in %d; EOD out %d", (uint32)VideoBufferStats.DecoderOutputBuffer.NumElementsInDecoder, (uint32)VideoBufferStats.DecoderOutputBuffer.MaxDecodedElementsReady, VideoBufferStats.DecoderOutputBuffer.bOutputStalled?"    stalled":"not stalled", VideoBufferStats.DecoderInputBuffer.bEODReached, VideoBufferStats.DecoderOutputBuffer.bEODreached);
+	}
+	if (StreamReaderHandler && bHaveAudioReader.GetWithDefault(false))
+	{
+		pD = &AudioBufferStats.StreamBuffer;
+		pPrintFN(pPlayer, "Audio buffer  : %3u AUs; %8u/%8u bytes in; eod %d; %#7.4fs", (uint32)pD->NumCurrentAccessUnits, (uint32)pD->CurrentMemInUse, (uint32)pD->MaxDataSize, pD->bEndOfData, pD->PushedDuration.GetAsSeconds());
+		pPrintFN(pPlayer, "Audio decoder : %2u in decoder, %2u total, %s; EOD in %d; EOD out %d", (uint32)AudioBufferStats.DecoderOutputBuffer.NumElementsInDecoder, (uint32)AudioBufferStats.DecoderOutputBuffer.MaxDecodedElementsReady, AudioBufferStats.DecoderOutputBuffer.bOutputStalled ? "    stalled" : "not stalled", AudioBufferStats.DecoderInputBuffer.bEODReached, AudioBufferStats.DecoderOutputBuffer.bEODreached);
+	}
+	DiagnosticsCriticalSection.Unlock();
+
+	pPrintFN(pPlayer, "Player state  : %s", GetPlayerStateName(CurrentState));
+	pPrintFN(pPlayer, "Decoder state : %s", GetDecoderStateName(DecoderState));
+	pPrintFN(pPlayer, "Pipeline state: %s", GetPipelineStateName(PipelineState));
+	pPrintFN(pPlayer, "Stream state  : %s", GetStreamStateName(StreamState));
+	pPrintFN(pPlayer, "-----------------------------------");
+	FTimeRange seekable, timeline;
+	GetSeekableRange(seekable);
+	PlaybackState.GetTimelineRange(timeline);
+	pPrintFN(pPlayer, " have metadata: %s", HaveMetadata() ? "Yes" : "No");
+	pPrintFN(pPlayer, "      duration: %.3f", GetDuration().GetAsSeconds());
+	pPrintFN(pPlayer, "timeline range: %.3f - %.3f", timeline.Start.GetAsSeconds(), timeline.End.GetAsSeconds());
+	pPrintFN(pPlayer, "seekable range: %.3f - %.3f", seekable.Start.GetAsSeconds(), seekable.End.GetAsSeconds());
+	pPrintFN(pPlayer, " play position: %.3f", GetPlayPosition().GetAsSeconds());
+	pPrintFN(pPlayer, "  is buffering: %s", IsBuffering() ? "Yes" : "No");
+	pPrintFN(pPlayer, "    is playing: %s", IsPlaying() ? "Yes" : "No");
+	pPrintFN(pPlayer, "     is paused: %s", IsPaused() ? "Yes" : "No");
+	pPrintFN(pPlayer, "    is seeking: %s", IsSeeking() ? "Yes" : "No");
+	pPrintFN(pPlayer, "     has ended: %s", HasEnded() ? "Yes" : "No");
+	pPrintFN(pPlayer, "clock (%s): %.3f", RenderClock->IsRunning()?"running":"paused", RenderClock->GetInterpolatedRenderTime(IMediaRenderClock::ERendererType::Video).GetAsSeconds(-1.0));
+}
+void FAdaptiveStreamingPlayer::DebugHandle(void* pPlayer, void (*debugDrawPrintf)(void* pPlayer, const char *pFmt, ...))
+{
+	if (PointerToLatestPlayer)
+	{
+		PointerToLatestPlayer->DebugPrint(pPlayer, debugDrawPrintf);
+	}
+}
+
+
+
+//---------------------------------------------------------------------------------------------------------------------------------
+//---------------------------------------------------------------------------------------------------------------------------------
+//---------------------------------------------------------------------------------------------------------------------------------
+
+TWeakPtr<FAdaptiveStreamingPlayerEventHandler, ESPMode::ThreadSafe>	FAdaptiveStreamingPlayerEventHandler::SingletonSelf;
+FCriticalSection													FAdaptiveStreamingPlayerEventHandler::SingletonLock;
+
+TSharedPtr<FAdaptiveStreamingPlayerEventHandler, ESPMode::ThreadSafe> FAdaptiveStreamingPlayerEventHandler::Create()
+{
+	FScopeLock lock(&SingletonLock);
+	TSharedPtr<FAdaptiveStreamingPlayerEventHandler, ESPMode::ThreadSafe> Self = SingletonSelf.Pin();
+	if (!Self.IsValid())
+	{
+		FAdaptiveStreamingPlayerEventHandler* Handler = new FAdaptiveStreamingPlayerEventHandler;
+		Handler->StartWorkerThread();
+		Self = MakeShareable(Handler);
+		SingletonSelf = Self;
+	}
+	return Self;
+}
+
+void FAdaptiveStreamingPlayerEventHandler::DispatchEvent(TSharedPtrTS<FMetricEvent> InEvent)
+{
+	EventQueue.SendMessage(InEvent);
+}
+
+void FAdaptiveStreamingPlayerEventHandler::DispatchEventAndWait(TSharedPtrTS<FMetricEvent> InEvent)
+{
+	FMediaEvent Sig;
+	InEvent->EventSignal = &Sig;
+	EventQueue.SendMessage(InEvent);
+	Sig.Wait();
+}
+
+FAdaptiveStreamingPlayerEventHandler::~FAdaptiveStreamingPlayerEventHandler()
+{
+	StopWorkerThread();
+}
+
+void FAdaptiveStreamingPlayerEventHandler::StartWorkerThread()
+{
+	//WorkerThread.ThreadSetPriority(mConfig.WorkerThread.Priority);
+	//WorkerThread.ThreadSetCoreAffinity(mConfig.WorkerThread.CoreAffinity);
+	//WorkerThread.ThreadSetStackSize(mConfig.WorkerThread.StackSize);
+	WorkerThread.ThreadSetName("ElectraPlayer::EventDispatch");
+	WorkerThread.ThreadStart(Electra::MakeDelegate(this, &FAdaptiveStreamingPlayerEventHandler::WorkerThreadFN));
+}
+
+void FAdaptiveStreamingPlayerEventHandler::StopWorkerThread()
+{
+	EventQueue.SendMessage(TSharedPtrTS<FMetricEvent>());
+	WorkerThread.ThreadWaitDone();
+	WorkerThread.ThreadReset();
+}
+
+void FAdaptiveStreamingPlayerEventHandler::WorkerThreadFN()
+{
+	while(1)
+	{
+		TSharedPtrTS<FMetricEvent> pEvt = EventQueue.ReceiveMessage();
+		if (!pEvt.IsValid())
+		{
+			break;
+		}
+
+		// Get the player that sends the event. If it no longer exists ignore the event.
+		TSharedPtr<FAdaptiveStreamingPlayer, ESPMode::ThreadSafe>	Player = pEvt->Player.Pin();
+		if (Player.IsValid())
+		{
+			// Call listeners under lock. They are not allowed to add or remove themselves or other listeners.
+			Player->LockMetricsReceivers();
+			const TArray<IAdaptiveStreamingPlayerMetrics*, TInlineAllocator<4>>& MetricListeners = Player->GetMetricsReceivers();
+			for(int32 i=0; i<MetricListeners.Num(); ++i)
+			{
+				switch(pEvt->Type)
+				{
+					case FMetricEvent::EType::OpenSource:
+						MetricListeners[i]->ReportOpenSource(pEvt->Param.URL);
+						break;
+					case FMetricEvent::EType::ReceivedMasterPlaylist:
+						MetricListeners[i]->ReportReceivedMasterPlaylist(pEvt->Param.URL);
+						break;
+					case FMetricEvent::EType::ReceivedPlaylists:
+						MetricListeners[i]->ReportReceivedPlaylists();
+						break;
+					case FMetricEvent::EType::BufferingStart:
+						MetricListeners[i]->ReportBufferingStart(pEvt->Param.BufferingReason);
+						break;
+					case FMetricEvent::EType::BufferingEnd:
+						MetricListeners[i]->ReportBufferingEnd(pEvt->Param.BufferingReason);
+						break;
+					case FMetricEvent::EType::Bandwidth:
+						MetricListeners[i]->ReportBandwidth(pEvt->Param.Bandwidth.EffectiveBps, pEvt->Param.Bandwidth.ThroughputBps, pEvt->Param.Bandwidth.Latency);
+						break;
+					case FMetricEvent::EType::BufferUtilization:
+						MetricListeners[i]->ReportBufferUtilization(pEvt->Param.BufferStats);
+						break;
+					case FMetricEvent::EType::PlaylistDownload:
+						MetricListeners[i]->ReportPlaylistDownload(pEvt->Param.PlaylistStats);
+						break;
+					case FMetricEvent::EType::SegmentDownload:
+						MetricListeners[i]->ReportSegmentDownload(pEvt->Param.SegmentStats);
+						break;
+					case FMetricEvent::EType::LicenseKey:
+						MetricListeners[i]->ReportLicenseKey(pEvt->Param.LicenseKeyStats);
+						break;
+					case FMetricEvent::EType::DataAvailabilityChange:
+						MetricListeners[i]->ReportDataAvailabilityChange(pEvt->Param.DataAvailability);
+						break;
+					case FMetricEvent::EType::VideoQualityChange:
+						MetricListeners[i]->ReportVideoQualityChange(pEvt->Param.QualityChange.NewBitrate, pEvt->Param.QualityChange.PrevBitrate, pEvt->Param.QualityChange.bIsDrastic);
+						break;
+					case FMetricEvent::EType::PrerollStart:
+						MetricListeners[i]->ReportPrerollStart();
+						break;
+					case FMetricEvent::EType::PrerollEnd:
+						MetricListeners[i]->ReportPrerollEnd();
+						break;
+					case FMetricEvent::EType::PlaybackStart:
+						MetricListeners[i]->ReportPlaybackStart();
+						break;
+					case FMetricEvent::EType::PlaybackPaused:
+						MetricListeners[i]->ReportPlaybackPaused();
+						break;
+					case FMetricEvent::EType::PlaybackResumed:
+						MetricListeners[i]->ReportPlaybackResumed();
+						break;
+					case FMetricEvent::EType::PlaybackEnded:
+						MetricListeners[i]->ReportPlaybackEnded();
+						break;
+					case FMetricEvent::EType::PlaybackJumped:
+						MetricListeners[i]->ReportJumpInPlayPosition(pEvt->Param.TimeJump.ToNewTime, pEvt->Param.TimeJump.FromTime, pEvt->Param.TimeJump.Reason);
+						break;
+					case FMetricEvent::EType::PlaybackStopped:
+						MetricListeners[i]->ReportPlaybackStopped();
+						break;
+					case FMetricEvent::EType::Errored:
+						MetricListeners[i]->ReportError(pEvt->Param.ErrorDetail.GetPrintable());
+						break;
+					case FMetricEvent::EType::LogMessage:
+						MetricListeners[i]->ReportLogMessage(pEvt->Param.LogMessage.Level, pEvt->Param.LogMessage.Message, pEvt->Param.LogMessage.AtMillis);
+						break;
+				}
+			}
+			Player->UnlockMetricsReceivers();
+		}
+		if (pEvt->EventSignal)
+		{
+			pEvt->EventSignal->Signal();
+		}
+	}
+}
+
+
+
+
+
+} // namespace Electra
