ElectraPlayer: RELEASE. Moved from NFL to public. Note that this is done via a Delete and "Add New" operation as discussed.

diff --git a/Engine/Plugins/Media/ElectraPlayer/Source/ElectraPlayerRuntime/Private/Runtime/Player/AdaptiveStreamingPlayer.h b/Engine/Plugins/Media/ElectraPlayer/Source/ElectraPlayerRuntime/Private/Runtime/Player/AdaptiveStreamingPlayer.h
new file mode 100644
--- /dev/null
+++ b/Engine/Plugins/Media/ElectraPlayer/Source/ElectraPlayerRuntime/Private/Runtime/Player/AdaptiveStreamingPlayer.h
@@ -0,0 +1,186 @@
+// Copyright Epic Games, Inc. All Rights Reserved.
+
+#pragma once
+
+#include "PlayerCore.h"
+
+#include "CoreMinimal.h"
+
+#include "ErrorDetail.h"
+#include "ParameterDictionary.h"
+#include "StreamTypes.h"
+#include "AdaptiveStreamingPlayerMetrics.h"
+#include "AdaptiveStreamingPlayerResourceRequest.h"
+
+class IOptionPointerValueContainer;
+
+namespace Electra
+{
+class IMediaRenderer;
+class IVideoDecoderResourceDelegate;
+
+/**
+ *
+**/
+class IAdaptiveStreamingPlayer : public TMediaNoncopyable<IAdaptiveStreamingPlayer>
+{
+public:
+	struct FCreateParam
+	{
+		TSharedPtr<IMediaRenderer, ESPMode::ThreadSafe>					VideoRenderer;
+		TSharedPtr<IMediaRenderer, ESPMode::ThreadSafe>					AudioRenderer;
+		FGuid															ExternalPlayerGUID;
+	};
+
+	static TSharedPtr<IAdaptiveStreamingPlayer, ESPMode::ThreadSafe> Create(const FCreateParam& InCreateParameters);
+
+	//-------------------------------------------------------------------------
+	// Various resource providers/delegates
+	//
+	virtual void SetStaticResourceProviderCallback(const TSharedPtr<IAdaptiveStreamingPlayerResourceProvider, ESPMode::ThreadSafe> & InStaticResourceProvider) = 0;
+	virtual void SetVideoDecoderResourceDelegate(const TSharedPtr<IVideoDecoderResourceDelegate, ESPMode::ThreadSafe>& ResourceDelegate) = 0;
+
+
+	//-------------------------------------------------------------------------
+	// Metric event listener
+	//
+	virtual void AddMetricsReceiver(IAdaptiveStreamingPlayerMetrics* InMetricsReceiver) = 0;
+	virtual void RemoveMetricsReceiver(IAdaptiveStreamingPlayerMetrics* InMetricsReceiver) = 0;
+
+
+	//-------------------------------------------------------------------------
+	// Initializes the player. Options may be passed in to affect behaviour.
+	//
+	virtual bool Initialize(const FParamDict& InOptions) = 0;
+
+
+	//-------------------------------------------------------------------------
+	// Manifest loading functions
+	//
+	//! Issue a load and parse of the manifest/master playlist file.
+	virtual void LoadManifest(const FString& ManifestURL) = 0;
+
+
+	//-------------------------------------------------------------------------
+	// Playback related functions
+	//
+	struct FSeekParam
+	{
+		FSeekParam()
+		{
+			Reset();
+		}
+		void Reset()
+		{
+			Time.SetToInvalid();
+		}
+		FTimeValue	Time;
+	};
+	//! Seek to a new position and play from there. This includes first playstart.
+	//! Playback is initially paused on first player use and must be resumed to begin.
+	//! Query the seekable range (GetSeekableRange()) to get the valid time range.
+	virtual void SeekTo(const FSeekParam& NewPosition) = 0;
+
+	//! Pauses playback.
+	virtual void Pause() = 0;
+
+	//! Resumes playback.
+	virtual void Resume() = 0;
+
+	//! Stops playback. Playback cannot be resumed. Final player events will be sent to registered listeners.
+	virtual void Stop() = 0;
+
+
+	struct FLoopParam
+	{
+		FLoopParam()
+		{
+			Reset();
+		}
+		void Reset()
+		{
+			bEnableLooping = false;
+		}
+		bool	bEnableLooping;
+	};
+
+	//! Puts playback into loop mode if possible. Live streams cannot be made to loop as they have infinite duration.
+	virtual void SetLooping(const FLoopParam& InLoopParams) = 0;
+
+
+	//-------------------------------------------------------------------------
+	// Error related functions
+	//
+	//! Returns the error that has caused playback issues.
+	virtual FErrorDetail GetError() const = 0;
+
+
+	//-------------------------------------------------------------------------
+	// State functions
+	//
+	//! Returns whether or not a manifest has been loaded and assigned yet.
+	virtual bool HaveMetadata() const = 0;
+	//! Returns the duration of the video. Returns invalid time when there is nothing to play. Returns positive infinite for live streams.
+	virtual FTimeValue GetDuration() const = 0;
+	//! Returns the current play position. Returns invalid time when there is nothing to play.
+	virtual FTimeValue GetPlayPosition() const = 0;
+	//! Returns the seekable range.
+	virtual void GetSeekableRange(FTimeRange& OutRange) const = 0;
+	//! Returns the timeline range.
+	virtual void GetTimelineRange(FTimeRange& OutRange) const = 0;
+	//! Fills the provided array with time values that can be seeked to. These are segment start times from the video (or audio if there is no video) track.
+	virtual void GetSeekablePositions(TArray<FTimespan>& OutPositions) const = 0;
+	//! Returns true when playback has finished.
+	virtual bool HasEnded() const = 0;
+	//! Returns true when data is being buffered/rebuffered, false otherwise.
+	virtual bool IsBuffering() const = 0;
+	//! Returns true when seeking is in progress. False if not.
+	virtual bool IsSeeking() const = 0;
+	//! Returns true when playing back, false if not.
+	virtual bool IsPlaying() const = 0;
+	//! Returns true when paused, false if not.
+	virtual bool IsPaused() const = 0;
+
+	//! Returns the current loop state.
+	virtual void GetLoopState(FPlayerLoopState& OutLoopState) const = 0;
+
+	//! Returns stream metadata of the currently active play period.
+	virtual void GetStreamMetadata(TArray<FStreamMetadata>& OutStreamMetadata, EStreamType StreamType) const = 0;
+
+
+	//-------------------------------------------------------------------------
+	// Manual stream selection functions
+	//
+	//! Sets the highest bitrate when selecting a candidate stream.
+	virtual void SetBitrateCeiling(int32 HighestSelectableBitrate) = 0;
+
+	//! Sets the maximum resolution to use. Set both to 0 to disable, set only one to limit widht or height only.
+	//! Setting both will limit on either width or height, whichever limits first.
+	virtual void SetMaxResolution(int32 MaxWidth, int32 MaxHeight) = 0;
+
+	//! Selects a track based from one of the array members returned by GetStreamMetadata().
+	//! NOTE: Presently this must be done before starting playback!
+	virtual void SelectTrackByMetadata(EStreamType StreamType, const FStreamMetadata& StreamMetadata) = 0;
+
+	//! Deselect track. The stream will continue to stream to allow for immediate selection/activation but no data will be fed to the decoder.
+	virtual void DeselectTrack(EStreamType StreamType) = 0;
+
+	//-------------------------------------------------------------------------
+	// Platform specific functions
+	//
+#if PLATFORM_ANDROID
+	virtual void Android_UpdateSurface(const TSharedPtr<IOptionPointerValueContainer>& Surface) = 0;
+#endif
+
+	//-------------------------------------------------------------------------
+	// Debug functions
+	//
+	static void DebugHandle(void* pPlayer, void (*debugDrawPrintf)(void* pPlayer, const char *pFmt, ...));
+
+protected:
+	virtual ~IAdaptiveStreamingPlayer() = default;
+};
+
+
+} // namespace Electra
+
diff --git a/Engine/Plugins/Media/ElectraPlayer/Source/ElectraPlayerRuntime/Private/Runtime/Player/AdaptiveStreamingPlayerABR.cpp b/Engine/Plugins/Media/ElectraPlayer/Source/ElectraPlayerRuntime/Private/Runtime/Player/AdaptiveStreamingPlayerABR.cpp
new file mode 100644
--- /dev/null
+++ b/Engine/Plugins/Media/ElectraPlayer/Source/ElectraPlayerRuntime/Private/Runtime/Player/AdaptiveStreamingPlayerABR.cpp
@@ -0,0 +1,1161 @@
+// Copyright Epic Games, Inc. All Rights Reserved.
+
+#include "Player/AdaptiveStreamingPlayerABR.h"
+
+#include "Player/Manifest.h"
+#include "Player/PlaybackTimeline.h"
+#include "StreamAccessUnitBuffer.h"
+#include "Utilities/Utilities.h"
+
+
+
+namespace Electra
+{
+
+	template <typename T>
+	class TSimpleMovingAverage
+	{
+	public:
+		//! History sample
+		TSimpleMovingAverage(int32 maxNumSMASamples = kMaxSMASamples)
+			: MaxSamples(maxNumSMASamples)
+		{
+			SimpleMovingAverage.Resize(MaxSamples);
+			LastSample = T(0);
+		}
+
+		void SetMaxSMASampleCount(int32 maxSMASamples)
+		{
+			FMediaCriticalSection::ScopedLock lock(CriticalSection);
+			MaxSamples = maxSMASamples;
+			SimpleMovingAverage.Clear();
+			SimpleMovingAverage.Resize(MaxSamples);
+		}
+
+		//! Clears everything
+		void Clear(void)
+		{
+			FMediaCriticalSection::ScopedLock lock(CriticalSection);
+			SimpleMovingAverage.Clear();
+		}
+
+		void InitializeTo(const T& v)
+		{
+			FMediaCriticalSection::ScopedLock lock(CriticalSection);
+			Clear();
+			AddSample(v);
+		}
+
+		//! Adds a sample to the mean.
+		void AddSample(const T& value)
+		{
+			FMediaCriticalSection::ScopedLock lock(CriticalSection);
+
+			LastSample = value;
+
+			// Update SMA (http://en.wikipedia.org/wiki/Moving_average#Simple_moving_average)
+			if (SimpleMovingAverage.Capacity())
+			{
+				if (SimpleMovingAverage.Num() >= MaxSamples)
+					SimpleMovingAverage.Pop();
+				check(!SimpleMovingAverage.IsFull());
+				SimpleMovingAverage.Push(LastSample);
+			}
+		}
+
+		//! Returns the sample value that was added last, if one exists
+		T GetLastSample() const
+		{
+			FMediaCriticalSection::ScopedLock lock(CriticalSection);
+			return(LastSample);
+		}
+
+		//! Returns the simple moving average from all the collected history samples.
+		T GetSMA(void) const
+		{
+			FMediaCriticalSection::ScopedLock lock(CriticalSection);
+			if (SimpleMovingAverage.Num())
+			{
+				T sum = 0;
+				for (SIZE_T i = 0; i < SimpleMovingAverage.Num(); ++i)
+					sum += SimpleMovingAverage[i];
+				return(T(sum / SimpleMovingAverage.Num()));
+			}
+			return(T(0));
+		}
+
+	private:
+		enum
+		{
+			kMaxSMASamples = 5,		//!< Default amount of SMA samples
+		};
+
+		mutable FMediaCriticalSection	CriticalSection;
+		TMediaQueueNoLock<T>			SimpleMovingAverage;	//!< Simple moving average
+		T								LastSample;				//!< The last sample that was added.
+		int32							MaxSamples;				//!< Limit number of SMA samples (0 < mMaxSMASamples <= kMaxSMASamples)
+	};
+
+
+	struct FSimulationResult
+	{
+		FSimulationResult()
+			: SecondsAdded(0.0)
+			, SecondsGained(0.0)
+			, SmallestGain(0.0)
+			, DownloadTime(0.0)
+		{
+		}
+		double		SecondsAdded;
+		double		SecondsGained;
+		double		SmallestGain;
+		double		DownloadTime;
+	};
+
+
+
+
+	class FAdaptiveStreamSelector : public IAdaptiveStreamSelector
+	{
+	public:
+		FAdaptiveStreamSelector(IPlayerSessionServices* PlayerSessionServices, const IAdaptiveStreamSelector::FConfiguration& config);
+		virtual ~FAdaptiveStreamSelector();
+		virtual void SetPresentationType(EMediaPresentationType PresentationType) override;
+		virtual void SetCanSwitchToAlternateStreams(bool bCanSwitch) override;
+		virtual void SetCurrentPlaybackPeriod(TSharedPtrTS<IManifest::IPlayPeriod> InCurrentPlayPeriod) override;
+		virtual void SetBandwidth(int64 bitsPerSecond) override;
+		virtual void SetForcedNextBandwidth(int64 bitsPerSecond) override;
+		virtual void MarkStreamAsUnavailable(const FBlacklistedStream& BlacklistedStream) override;
+		virtual void MarkStreamAsAvailable(const FBlacklistedStream& NoLongerBlacklistedStream) override;
+		virtual int64 GetLastBandwidth() override;
+		virtual int64 GetAverageBandwidth() override;
+		virtual int64 GetAverageThroughput() override;
+		virtual double GetAverageLatency() override;
+
+		virtual void SetBandwidthCeiling(int32 HighestManifestBitrate) override;
+		virtual void SetMaxVideoResolution(int32 MaxWidth, int32 MaxHeight) override;
+
+		virtual ESegmentAction SelectSuitableStreams(FTimeValue& OutDelay, TSharedPtrTS<const IStreamSegment> CurrentSegment) override;
+
+		virtual FABRDownloadProgressDecision ReportDownloadProgress(const Metrics::FSegmentDownloadStats& SegmentDownloadStats) override;
+		virtual void ReportDownloadEnd(const Metrics::FSegmentDownloadStats& SegmentDownloadStats) override;
+
+
+		virtual void ReportBufferingStart(Metrics::EBufferingReason BufferingReason) override;
+		virtual void ReportBufferingEnd(Metrics::EBufferingReason BufferingReason) override;
+		virtual void ReportOpenSource(const FString& URL) override {}
+		virtual void ReportReceivedMasterPlaylist(const FString& EffectiveURL) override {}
+		virtual void ReportReceivedPlaylists() override {}
+		virtual void ReportPlaylistDownload(const Metrics::FPlaylistDownloadStats& PlaylistDownloadStats) override {}
+		virtual void ReportBandwidth(int64 EffectiveBps, int64 ThroughputBps, double LatencyInSeconds) override {}
+		virtual void ReportBufferUtilization(const Metrics::FBufferStats& BufferStats) override {}
+		virtual void ReportSegmentDownload(const Metrics::FSegmentDownloadStats& SegmentDownloadStats) override {}
+		virtual void ReportLicenseKey(const Metrics::FLicenseKeyStats& LicenseKeyStats) override {}
+		virtual void ReportDataAvailabilityChange(const Metrics::FDataAvailabilityChange& DataAvailability) override {}
+		virtual void ReportVideoQualityChange(int32 NewBitrate, int32 PreviousBitrate, bool bIsDrasticDownswitch) override {}
+		virtual void ReportPrerollStart() override {}
+		virtual void ReportPrerollEnd() override {}
+		virtual void ReportPlaybackStart() override {}
+		virtual void ReportPlaybackPaused() override {}
+		virtual void ReportPlaybackResumed() override {}
+		virtual void ReportPlaybackEnded() override {}
+		virtual void ReportJumpInPlayPosition(const FTimeValue& ToNewTime, const FTimeValue& FromTime, Metrics::ETimeJumpReason TimejumpReason) override {}
+		virtual void ReportPlaybackStopped() override {}
+		virtual void ReportError(const FString& ErrorReason) override {}
+		virtual void ReportLogMessage(IInfoLog::ELevel InLogLevel, const FString& LogMessage, int64 PlayerWallclockMilliseconds) override {}
+		virtual void ReportDroppedVideoFrame() override {}
+		virtual void ReportDroppedAudioFrame() override {}
+
+	private:
+		struct FStreamInformation
+		{
+			struct FStreamHealth
+			{
+				FStreamHealth()
+				{
+					Reset();
+				}
+				void Reset()
+				{
+					//PreviousAttempts.Reset();
+					BecomesAvailableAgainAtUTC.SetToZero();
+					LastDownloadStats = {};
+				}
+				//TSharedPtrTS<HTTP::FRetryInfo>				PreviousAttempts;
+				FTimeValue										BecomesAvailableAgainAtUTC;
+				Metrics::FSegmentDownloadStats					LastDownloadStats;
+			};
+
+			TSharedPtrTS<IPlaybackAssetAdaptationSet>		AdaptationSet;
+			TSharedPtrTS<IPlaybackAssetRepresentation>		Representation;
+
+			FString											AdaptationSetUniqueID;
+			FString											RepresentationUniqueID;
+			FString											CDN;
+
+			FStreamHealth									Health;
+			// Set for convenience.
+			FStreamCodecInformation::FResolution			Resolution;
+			int32											Bitrate;
+		};
+
+		void LogMessage(IInfoLog::ELevel Level, const FString& Message)
+		{
+			if (PlayerSessionServices)
+			{
+				PlayerSessionServices->PostLog(Facility::EFacility::ABR, Level, Message);
+			}
+		}
+
+
+		TSharedPtrTS<FStreamInformation> GetStreamInformation(const Metrics::FSegmentDownloadStats& FromDownloadStats);
+
+		void SimulateDownload(FSimulationResult& sim, const TArray<IManifest::IPlayPeriod::FSegmentInformation>& NextSegments, double secondsToSimulate, int64 atBPS);
+
+
+		IPlayerSessionServices* PlayerSessionServices;
+		FConfiguration										Config;
+
+		FMediaCriticalSection								AccessMutex;
+		EMediaPresentationType								PresentationType;
+		TSharedPtrTS<IManifest::IPlayPeriod>				CurrentPlayPeriod;
+		bool												bPlayerIsBuffering;
+		bool												bCanSwitchToAlternateStreams;
+
+		TArray<TSharedPtrTS<FStreamInformation>>			StreamInformationVideo;
+		TArray<TSharedPtrTS<FStreamInformation>>			StreamInformationAudio;
+		FString												CurrentMediaAssetID;
+		FString												CurrentVideoAdaptationSetID;
+		FString												CurrentAudioAdaptationSetID;
+
+		TArray<FBlacklistedStream>							BlacklistedExternally;
+
+		int32												BandwidthCeiling;
+		int32												ForcedInitialBandwidth;
+		FStreamCodecInformation::FResolution				MaxStreamResolution;
+
+		TSimpleMovingAverage<double>						AverageBandwidth;
+		TSimpleMovingAverage<double>						AverageLatency;
+		TSimpleMovingAverage<int64>							AverageThroughput;
+	};
+
+
+
+	IAdaptiveStreamSelector::FConfiguration::FRateScale::FRateScale()
+		// Change no defaults without making adjustments to the values in Configuration::Configuration() below!
+		: LastToAvg(0.0f)				// prefer last rate
+		, Scale(1.0f)					// 100%
+		, MaxClampBPS(50000000)		// limit to 50 Mbps
+	{
+	}
+
+	double IAdaptiveStreamSelector::FConfiguration::FRateScale::Apply(const double lastBps, const double averageBps) const
+	{
+		return Utils::Max(Utils::Min((lastBps + LastToAvg * (averageBps - lastBps)) * Scale, double(MaxClampBPS)), 0.0);
+	}
+
+
+	IAdaptiveStreamSelector::FConfiguration::FConfiguration()
+	{
+		SimulationLongtermRateScale.LastToAvg = 1.0f; 			// For simulation use the average, not last bandwidth.
+		ForwardSimulationDurationMSec = 1000 * 20;		// 20 seconds
+		VideoBufferLowWatermarkMSec = 1000 * 3;
+		VideoBufferHighWatermarkMSec = 1000 * 8;
+		MaxBandwidthHistorySize = 3;				// track download speed of the last 5 fragments.
+		DiscardSmallBandwidthSamplesLess = 64 << 10;
+		FudgeSmallBandwidthSamplesLess = 128 << 10;
+		FudgeSmallBandwidthSamplesFactor = 1000;
+	}
+
+
+
+
+
+
+
+
+
+
+
+
+
+
+
+
+
+	//-----------------------------------------------------------------------------
+	/**
+	 * Create an instance of this class
+	 *
+	 * @param PlayerSessionServices
+	 * @param config
+	 *
+	 * @return
+	 */
+	TSharedPtrTS<IAdaptiveStreamSelector> IAdaptiveStreamSelector::Create(IPlayerSessionServices* PlayerSessionServices, const IAdaptiveStreamSelector::FConfiguration& config)
+	{
+		return TSharedPtrTS<IAdaptiveStreamSelector>(new FAdaptiveStreamSelector(PlayerSessionServices, config));
+	}
+
+	//-----------------------------------------------------------------------------
+	/**
+	 * CTOR
+	 *
+	 * @param config
+	 */
+	FAdaptiveStreamSelector::FAdaptiveStreamSelector(IPlayerSessionServices* InPlayerSessionServices, const IAdaptiveStreamSelector::FConfiguration& config)
+		: PlayerSessionServices(InPlayerSessionServices)
+		, Config(config)
+		, bPlayerIsBuffering(true)
+		, BandwidthCeiling(0x7fffffff)
+		, ForcedInitialBandwidth(0)
+	{
+		AverageBandwidth.SetMaxSMASampleCount(Config.MaxBandwidthHistorySize);
+		AverageLatency.SetMaxSMASampleCount(Config.MaxBandwidthHistorySize);
+		AverageThroughput.SetMaxSMASampleCount(Config.MaxBandwidthHistorySize);
+		MaxStreamResolution.Width = MaxStreamResolution.Height = 8192;
+		PresentationType = EMediaPresentationType::Realtime;
+		bCanSwitchToAlternateStreams = true;
+	}
+
+	//-----------------------------------------------------------------------------
+	/**
+	 * DTOR
+	 */
+	FAdaptiveStreamSelector::~FAdaptiveStreamSelector()
+	{
+	}
+
+
+
+	//-----------------------------------------------------------------------------
+	/**
+	 * Sets the type of presentation.
+	 *
+	 * @param InAssetType
+	 */
+	void FAdaptiveStreamSelector::SetPresentationType(IAdaptiveStreamSelector::EMediaPresentationType InAssetType)
+	{
+		PresentationType = InAssetType;
+	}
+
+
+	//-----------------------------------------------------------------------------
+	/**
+	 * Sets whether or not a segment download failure can potentially be resolved by switching to a different stream. Defaults to yes.
+	 *
+	 * @param bCanSwitch
+	 */
+	void FAdaptiveStreamSelector::SetCanSwitchToAlternateStreams(bool bCanSwitch)
+	{
+		bCanSwitchToAlternateStreams = bCanSwitch;
+	}
+
+
+	//-----------------------------------------------------------------------------
+	/**
+	 * Sets the playback period from which to select streams now.
+	 *
+	 * @param CurrentPlayPeriod
+	 */
+	void FAdaptiveStreamSelector::SetCurrentPlaybackPeriod(TSharedPtrTS<IManifest::IPlayPeriod> InCurrentPlayPeriod)
+	{
+		FMediaCriticalSection::ScopedLock lock(AccessMutex);
+		CurrentPlayPeriod = InCurrentPlayPeriod;
+		if (CurrentPlayPeriod.IsValid())
+		{
+			TSharedPtrTS<ITimelineMediaAsset> Asset = CurrentPlayPeriod->GetMediaAsset();
+			if (Asset.IsValid() && CurrentMediaAssetID != Asset->GetUniqueIdentifier())
+			{
+				CurrentMediaAssetID.Empty();
+				CurrentVideoAdaptationSetID.Empty();
+				CurrentAudioAdaptationSetID.Empty();
+				StreamInformationVideo.Empty();
+				StreamInformationAudio.Empty();
+				// Get video stream information
+				if (Asset->GetNumberOfAdaptationSets(EStreamType::Video) > 0)
+				{
+					TSharedPtrTS<IPlaybackAssetAdaptationSet> Adapt = Asset->GetAdaptationSetByTypeAndIndex(EStreamType::Video, 0);
+					if (Adapt.IsValid())
+					{
+						CurrentMediaAssetID = Asset->GetUniqueIdentifier();
+						CurrentVideoAdaptationSetID = Adapt->GetUniqueIdentifier();
+						for (int32 i = 0, iMax = Adapt->GetNumberOfRepresentations(); i < iMax; ++i)
+						{
+							TSharedPtrTS<FStreamInformation> si = MakeSharedTS<FStreamInformation>();
+							si->AdaptationSet = Adapt;
+							si->Representation = Adapt->GetRepresentationByIndex(i);
+							si->AdaptationSetUniqueID = si->AdaptationSet->GetUniqueIdentifier();
+							si->RepresentationUniqueID = si->Representation->GetUniqueIdentifier();
+							si->CDN = si->Representation->GetCDN();
+							si->Bitrate = si->Representation->GetBitrate();
+							if (si->Representation->GetCodecInformation().IsVideoCodec())
+							{
+								si->Resolution = si->Representation->GetCodecInformation().GetResolution();
+							}
+							StreamInformationVideo.Push(si);
+						}
+						// Sort the representations by ascending bitrate
+						StreamInformationVideo.Sort([](const TSharedPtrTS<FStreamInformation>& a, const TSharedPtrTS<FStreamInformation>& b)
+						{
+							return a->Representation->GetBitrate() < b->Representation->GetBitrate();
+						});
+					}
+				}
+				// Get audio stream information
+				if (Asset->GetNumberOfAdaptationSets(EStreamType::Audio) > 0)
+				{
+					TSharedPtrTS<IPlaybackAssetAdaptationSet> Adapt = Asset->GetAdaptationSetByTypeAndIndex(EStreamType::Audio, 0);
+					if (Adapt.IsValid())
+					{
+						CurrentMediaAssetID = Asset->GetUniqueIdentifier();
+						CurrentAudioAdaptationSetID = Adapt->GetUniqueIdentifier();
+						for (int32 i = 0, iMax = Adapt->GetNumberOfRepresentations(); i < iMax; ++i)
+						{
+							TSharedPtrTS<FStreamInformation> si = MakeSharedTS<FStreamInformation>();
+							si->AdaptationSet = Adapt;
+							si->Representation = Adapt->GetRepresentationByIndex(i);
+							si->AdaptationSetUniqueID = si->AdaptationSet->GetUniqueIdentifier();
+							si->RepresentationUniqueID = si->Representation->GetUniqueIdentifier();
+							si->CDN = si->Representation->GetCDN();
+							si->Bitrate = si->Representation->GetBitrate();
+							/*
+							if (si->Representation->GetCodecInformation().IsAudioCodec())
+							{
+								// ...
+							}
+							*/
+							StreamInformationAudio.Push(si);
+						}
+						// Sort the representations by ascending bitrate
+						StreamInformationAudio.Sort([](const TSharedPtrTS<FStreamInformation>& a, const TSharedPtrTS<FStreamInformation>& b)
+						{
+							return a->Representation->GetBitrate() > b->Representation->GetBitrate();
+						});
+					}
+				}
+			}
+		}
+		else
+		{
+			CurrentMediaAssetID.Empty();
+			CurrentVideoAdaptationSetID.Empty();
+			CurrentAudioAdaptationSetID.Empty();
+			StreamInformationVideo.Empty();
+			StreamInformationAudio.Empty();
+		}
+	}
+
+
+	//-----------------------------------------------------------------------------
+	/**
+	 * Sets the initial bandwidth, either from guessing, past history or a current measurement.
+	 *
+	 * @param bitPerSecond
+	 */
+	void FAdaptiveStreamSelector::SetBandwidth(int64 bitsPerSecond)
+	{
+		FMediaCriticalSection::ScopedLock lock(AccessMutex);
+		AverageBandwidth.InitializeTo((double)bitsPerSecond);
+		AverageLatency.InitializeTo((double)0.1);
+		AverageThroughput.InitializeTo(bitsPerSecond);
+	}
+
+
+	//-----------------------------------------------------------------------------
+	/**
+	 * Sets a forced bitrate for the next segment fetch only.
+	 *
+	 * @param bitsPerSecond
+	 */
+	void FAdaptiveStreamSelector::SetForcedNextBandwidth(int64 bitsPerSecond)
+	{
+		FMediaCriticalSection::ScopedLock lock(AccessMutex);
+		ForcedInitialBandwidth = (int32)bitsPerSecond;
+	}
+
+
+	//-----------------------------------------------------------------------------
+	/**
+	 * Called to (temporarily) mark a stream as unavailable.
+	 * This is used in formats like HLS when a dedicated stream playlist is in error
+	 * and thus segment information for this stream is not available.
+	 *
+	 * @param BlacklistedStream
+	 */
+	void FAdaptiveStreamSelector::MarkStreamAsUnavailable(const FBlacklistedStream& BlacklistedStream)
+	{
+		FMediaCriticalSection::ScopedLock lock(AccessMutex);
+		BlacklistedExternally.Push(BlacklistedStream);
+	}
+
+	//-----------------------------------------------------------------------------
+	/**
+	 * Marks a previously set to unavailable stream as being available again.
+	 * Usually when the stream specific playlist (in HLS) has become available again.
+	 *
+	 * @param NoLongerBlacklistedStream
+	 */
+	void FAdaptiveStreamSelector::MarkStreamAsAvailable(const FBlacklistedStream& NoLongerBlacklistedStream)
+	{
+		FMediaCriticalSection::ScopedLock lock(AccessMutex);
+		BlacklistedExternally.Remove(NoLongerBlacklistedStream);
+	}
+
+
+	//-----------------------------------------------------------------------------
+	/**
+	 * Returns the last measured bandwidth sample in bps.
+	 *
+	 * @return
+	 */
+	int64 FAdaptiveStreamSelector::GetLastBandwidth()
+	{
+		FMediaCriticalSection::ScopedLock lock(AccessMutex);
+		double lastBandwidth = AverageBandwidth.GetLastSample();
+		return((int64)lastBandwidth);
+	}
+
+
+	//-----------------------------------------------------------------------------
+	/**
+	 * Returns the average measured bandwidth sample in bps.
+	 *
+	 * @return
+	 */
+	int64 FAdaptiveStreamSelector::GetAverageBandwidth()
+	{
+		FMediaCriticalSection::ScopedLock lock(AccessMutex);
+		const double averageBandwidth = AverageBandwidth.GetSMA();
+		return((int64)averageBandwidth);
+	}
+
+	//-----------------------------------------------------------------------------
+	/**
+	 * Returns the average measured throughput in bps.
+	 *
+	 * @return
+	 */
+	int64 FAdaptiveStreamSelector::GetAverageThroughput()
+	{
+		FMediaCriticalSection::ScopedLock lock(AccessMutex);
+		return AverageThroughput.GetSMA();
+	}
+
+	//-----------------------------------------------------------------------------
+	/**
+	 * Returns the average measured latency in seconds.
+	 *
+	 * @return
+	 */
+	double FAdaptiveStreamSelector::GetAverageLatency()
+	{
+		FMediaCriticalSection::ScopedLock lock(AccessMutex);
+		return AverageLatency.GetSMA();
+	}
+
+
+
+	//-----------------------------------------------------------------------------
+	/**
+	 * Sets a highest bandwidth limit for a stream type. Call with 0 to disable.
+	 *
+	 * @param HighestManifestBitrate
+	 */
+	void FAdaptiveStreamSelector::SetBandwidthCeiling(int32 HighestManifestBitrate)
+	{
+		FMediaCriticalSection::ScopedLock lock(AccessMutex);
+		BandwidthCeiling = HighestManifestBitrate > 0 ? HighestManifestBitrate : 0x7fffffff;
+	}
+
+
+	//-----------------------------------------------------------------------------
+	/**
+	 * Limits video resolution.
+	 *
+	 * @param MaxWidth
+	 * @param MaxHeight
+	 */
+	void FAdaptiveStreamSelector::SetMaxVideoResolution(int32 MaxWidth, int32 MaxHeight)
+	{
+		FMediaCriticalSection::ScopedLock lock(AccessMutex);
+		MaxStreamResolution.Width = MaxWidth;
+		MaxStreamResolution.Height = MaxHeight;
+	}
+
+
+
+	FABRDownloadProgressDecision FAdaptiveStreamSelector::ReportDownloadProgress(const Metrics::FSegmentDownloadStats& SegmentDownloadStats)
+	{
+		FABRDownloadProgressDecision Decision;
+		uint32 Flags = SegmentDownloadStats.ABRState.ProgressDecision.Flags;
+
+		FAccessUnitBufferInfo CurrentBufferStats;
+		PlayerSessionServices->GetStreamBufferStats(CurrentBufferStats, SegmentDownloadStats.StreamType);
+		const double secondsInBuffer = CurrentBufferStats.PlayableDuration.GetAsSeconds();
+		// Estimate how long we will take to complete the download.
+		double EstimatedTotalDownloadTime = SegmentDownloadStats.Duration;
+		//double currentPercentage = 0.0;
+		// Can we base this off the byte sizes?
+		if (SegmentDownloadStats.NumBytesDownloaded > 0 && SegmentDownloadStats.ByteSize > 0)
+		{
+			EstimatedTotalDownloadTime = (SegmentDownloadStats.TimeToDownload / SegmentDownloadStats.NumBytesDownloaded) * SegmentDownloadStats.ByteSize;
+			//currentPercentage = SegmentDownloadStats.NumBytesDownloaded / SegmentDownloadStats.ByteSize;
+		}
+		// Otherwise, can we base it off the durations?
+		else if (SegmentDownloadStats.DurationDownloaded > 0.0 && SegmentDownloadStats.Duration > 0.0)
+		{
+			EstimatedTotalDownloadTime = (SegmentDownloadStats.TimeToDownload / SegmentDownloadStats.DurationDownloaded) * SegmentDownloadStats.Duration;
+			//currentPercentage = SegmentDownloadStats.DurationDownloaded / SegmentDownloadStats.Duration;
+		}
+		double EstimatedDownloadTimeRemaining = EstimatedTotalDownloadTime - SegmentDownloadStats.TimeToDownload;
+
+		if ((SegmentDownloadStats.ABRState.ProgressDecision.Flags & FABRDownloadProgressDecision::EDecisionFlags::eABR_AbortDownload) == 0)
+		{
+			// Did we take longer to download than the segment duration?
+			double MaxSegmentDownloadDuration = SegmentDownloadStats.Duration;
+			// If there is a forced initial bitrate to start playback with then we allow for the download duration to take longer than normally
+			// as to not abort the first download unless it's download time is really getting problematic.
+			if (ForcedInitialBandwidth && SegmentDownloadStats.StreamType == EStreamType::Video)
+			{
+				MaxSegmentDownloadDuration *= 2.0;
+			}
+			if (SegmentDownloadStats.TimeToDownload >= MaxSegmentDownloadDuration)
+			{
+				// If the estimate to complete download is longer than what is currently in the buffer abort the download and hope
+				// switching to a lower quality will help.
+		//		if (EstimatedDownloadTimeRemaining >= secondsInBuffer)
+		// Be more aggressive and demand the buffer must have the total duration! This is to avoid late switching that may be too late to try to download a whole new segment
+				if (EstimatedTotalDownloadTime >= secondsInBuffer)
+				{
+					Flags = FABRDownloadProgressDecision::EDecisionFlags::eABR_AbortDownload;
+					// If content was already emitted into the buffer we cannot switch. Insert dummy filler data instead.
+					if (SegmentDownloadStats.DurationDelivered > 0.0)
+					{
+						Flags |= FABRDownloadProgressDecision::EDecisionFlags::eABR_InsertFillerData;
+					}
+
+					Decision.Reason = FString::Printf(TEXT("Abort download of %s %s segment @ %d bps for time %.3f with %.3fs in buffer and %.3fs of %.3fs fetched after %.3fs. Estimated total download time %.3fs")
+						, GetStreamTypeName(SegmentDownloadStats.StreamType)
+						, GetSegmentTypeString(SegmentDownloadStats.SegmentType)
+						, SegmentDownloadStats.Bitrate
+						, SegmentDownloadStats.PresentationTime
+						, secondsInBuffer
+						, SegmentDownloadStats.DurationDownloaded
+						, MaxSegmentDownloadDuration
+						, SegmentDownloadStats.TimeToDownload
+						, EstimatedTotalDownloadTime);
+					LogMessage(IInfoLog::ELevel::Info, Decision.Reason);
+				}
+			}
+		}
+
+		// If the buffer is low allow the data being streamed in to already be emitted to the buffer.
+		// This will however mean that in case of remaining problems with this segment we cannot switch to another quality level.
+		if ((SegmentDownloadStats.ABRState.ProgressDecision.Flags & FABRDownloadProgressDecision::EDecisionFlags::eABR_EmitPartialData) == 0)
+		{
+			if (secondsInBuffer < 1.0 && SegmentDownloadStats.DurationDelivered == 0.0)
+			{
+				Flags |= FABRDownloadProgressDecision::EDecisionFlags::eABR_EmitPartialData;
+				/*
+							LogMessage(IInfoLog::ELevel::Info, StringHelpers::SPrintf("Emitting partial download of %s %s segment @ %d bps for time %.3f with %.3fs in buffer and %.3fs of %.3fs fetched after %.3fs. Estimated total download time %.3fs"
+								  , GetStreamTypeName(SegmentDownloadStats.StreamType)
+								  , GetSegmentTypeString(SegmentDownloadStats.SegmentType)
+								  , SegmentDownloadStats.Bitrate
+								  , SegmentDownloadStats.PresentationTime
+								  , secondsInBuffer
+								  , SegmentDownloadStats.DurationDownloaded
+								  , SegmentDownloadStats.Duration
+								  , SegmentDownloadStats.TimeToDownload
+								  , EstimatedTotalDownloadTime));
+				*/
+			}
+		}
+
+		Decision.Flags = (FABRDownloadProgressDecision::EDecisionFlags) Flags;
+		return Decision;
+	}
+
+	void FAdaptiveStreamSelector::ReportDownloadEnd(const Metrics::FSegmentDownloadStats& SegmentDownloadStats)
+	{
+		if (SegmentDownloadStats.SegmentType == Metrics::ESegmentType::Media &&
+			SegmentDownloadStats.StreamType == EStreamType::Video &&
+			(SegmentDownloadStats.bWasSuccessful || SegmentDownloadStats.NumBytesDownloaded))	// any number of bytes received counts!
+		{
+			// At the end of the segment download, successful or otherwise, we clear the forced bitrate and let nature take its course.
+			ForcedInitialBandwidth = 0;
+			// Add bandwidth sample if it can be calculated.
+			if (SegmentDownloadStats.NumBytesDownloaded && SegmentDownloadStats.TimeToDownload > 0.0)
+			{
+				FMediaCriticalSection::ScopedLock lock(AccessMutex);
+				double v = SegmentDownloadStats.NumBytesDownloaded * 8 / SegmentDownloadStats.TimeToDownload;
+				AverageBandwidth.AddSample(v);
+				AverageLatency.AddSample(SegmentDownloadStats.TimeToFirstByte);
+				AverageThroughput.AddSample(SegmentDownloadStats.ThroughputBps > 0 ? SegmentDownloadStats.ThroughputBps : (int64)v);
+			}
+		}
+	}
+
+	void FAdaptiveStreamSelector::ReportBufferingStart(Metrics::EBufferingReason BufferingReason)
+	{
+		/*
+			Initial,
+			Seeking,
+			Rebuffering,
+		*/
+		bPlayerIsBuffering = true;
+
+		if (BufferingReason == Metrics::EBufferingReason::Rebuffering)
+		{
+			double Scale = 0.6;
+			double avgCurrent = AverageBandwidth.GetLastSample();
+			AverageBandwidth.InitializeTo(avgCurrent * Scale);
+			AverageThroughput.InitializeTo((int64)(AverageThroughput.GetLastSample() * Scale));
+			LogMessage(IInfoLog::ELevel::Info, FString::Printf(TEXT("Rebuffering cuts bandwidth by %.2f from %d to %d"), Scale, (int)avgCurrent, (int)(avgCurrent * Scale)));
+		}
+	}
+
+	void FAdaptiveStreamSelector::ReportBufferingEnd(Metrics::EBufferingReason BufferingReason)
+	{
+		bPlayerIsBuffering = false;
+	}
+
+	TSharedPtrTS<FAdaptiveStreamSelector::FStreamInformation> FAdaptiveStreamSelector::GetStreamInformation(const Metrics::FSegmentDownloadStats& FromDownloadStats)
+	{
+		FMediaCriticalSection::ScopedLock lock(AccessMutex);
+		const TArray<TSharedPtrTS<FStreamInformation>>* StreamInfos = nullptr;
+		if (FromDownloadStats.StreamType == EStreamType::Video)
+		{
+			StreamInfos = &StreamInformationVideo;
+		}
+		else if (FromDownloadStats.StreamType == EStreamType::Audio)
+		{
+			StreamInfos = &StreamInformationAudio;
+		}
+		if (StreamInfos)
+		{
+			for (int32 i = 0, iMax = (int32)StreamInfos->Num(); i < iMax; ++i)
+			{
+				if (FromDownloadStats.AdaptationSetID == (*StreamInfos)[i]->AdaptationSetUniqueID &&
+					FromDownloadStats.RepresentationID == (*StreamInfos)[i]->RepresentationUniqueID &&
+					//FromDownloadStats.CDN              == (*StreamInfos)[i]->CDN &&
+					1)
+				{
+					return (*StreamInfos)[i];
+				}
+			}
+		}
+		return nullptr;
+	}
+
+
+	//-----------------------------------------------------------------------------
+	/**
+	 * Simulates download of the next fragments at the current download rate.
+	 *
+	 * @param sim
+	 * @param NextSegments
+	 * @param secondsToSimulate
+	 * @param atBPS
+	 */
+	void FAdaptiveStreamSelector::SimulateDownload(FSimulationResult& sim, const TArray<IManifest::IPlayPeriod::FSegmentInformation>& NextSegments, double secondsToSimulate, int64 atBPS)
+	{
+		sim.SecondsAdded = 0.0;
+		sim.SecondsGained = 0.0;
+		sim.SmallestGain = 0.0;
+		sim.DownloadTime = 0.0;
+
+		bool bFirst = true;
+		int32 infoIdx = 0;
+		check(NextSegments.Num() > 0);
+		double s = 8.0 / atBPS;
+		while (secondsToSimulate > 0.0)
+		{
+			double Dur = NextSegments[infoIdx].Duration.GetAsSeconds();
+
+			double timeToDownload = NextSegments[infoIdx].ByteSize * s;
+
+			secondsToSimulate -= Dur;
+			sim.SecondsAdded += Dur;
+			sim.DownloadTime += timeToDownload;
+			sim.SecondsGained += Dur - timeToDownload;
+
+			if (sim.SecondsGained < sim.SmallestGain || bFirst)
+			{
+				sim.SmallestGain = sim.SecondsGained;
+			}
+			bFirst = false;
+
+			// Advance to the next info, if present. Otherwise repeat the last one ad nauseum.
+			if (infoIdx + 1 < NextSegments.Num())
+			{
+				++infoIdx;
+			}
+		}
+	}
+
+
+	//-----------------------------------------------------------------------------
+	/**
+	 * Returns the index of a feasible video or multiplex stream.
+	 *
+	 * @param OutDelay
+	 * @param CurrentSegment
+	 *
+	 * @return
+	 */
+	IAdaptiveStreamSelector::ESegmentAction FAdaptiveStreamSelector::SelectSuitableStreams(FTimeValue& OutDelay, TSharedPtrTS<const IStreamSegment> CurrentSegment)
+	{
+		OutDelay.SetToZero();
+		if (!CurrentPlayPeriod.IsValid())
+		{
+			return ESegmentAction::FetchNext;
+		}
+
+		FTimeValue TimeNow = PlayerSessionServices->GetSynchronizedUTCTime()->GetTime();
+		// Make all streams available again if their blacklist time has expired
+		AccessMutex.Lock();
+		for (int32 i = 0, iMax = StreamInformationVideo.Num(); i < iMax; ++i)
+		{
+			if (StreamInformationVideo[i]->Health.BecomesAvailableAgainAtUTC <= TimeNow)
+			{
+				StreamInformationVideo[i]->Health.BecomesAvailableAgainAtUTC.SetToZero();
+			}
+		}
+		for (int32 i = 0, iMax = StreamInformationAudio.Num(); i < iMax; ++i)
+		{
+			if (StreamInformationAudio[i]->Health.BecomesAvailableAgainAtUTC <= TimeNow)
+			{
+				StreamInformationAudio[i]->Health.BecomesAvailableAgainAtUTC.SetToZero();
+			}
+		}
+		AccessMutex.Unlock();
+
+
+		EStreamType StreamType = CurrentSegment.IsValid() ? CurrentSegment->GetType() : StreamInformationVideo.Num() ? EStreamType::Video : EStreamType::Audio;
+		// Was the previous segment download in error?
+		bool bRetryIfPossible = false;
+		bool bSkipWithFiller = false;
+		if (CurrentSegment.IsValid())
+		{
+			Metrics::FSegmentDownloadStats Stats;
+			CurrentSegment->GetDownloadStats(Stats);
+			TSharedPtrTS<FAdaptiveStreamSelector::FStreamInformation> CurrentStreamInfo = GetStreamInformation(Stats);
+			check(CurrentStreamInfo.IsValid());
+			if (CurrentStreamInfo.IsValid())
+			{
+				if (!Stats.bWasSuccessful)
+				{
+					// Was not successful. Figure out what to do now.
+
+					/*
+							EStreamType		StreamType;						//!< Type of stream
+							ESegmentType	SegmentType;					//!< Type of segment (init or media)
+							FString	URL;							//!< Effective URL used to download from
+							FString	CDN;							//!< CDN
+							FString	MediaAssetID;
+							FString	AdaptationSetID;
+							FString	RepresentationID;
+							double			PresentationTime;				//!< Presentation time on media timeline
+							double			Duration;						//!< Duration of segment as specified in manifest
+							int32			Bitrate;						//!< Stream bitrate as specified in manifest
+							int32			RetryNumber;
+							bool			bIsMissingSegment;				//!< true if the segment was not actually downloaded because it is missing on the timeline.
+
+							// Outputs from stream reader
+							FString	FailureReason;					//!< Human readable failure reason. Only for display purposes.
+							double			DurationDownloaded;				//!< Duration of content successfully downloaded. May be less than Duration in case of errors.
+							double			DurationDelivered;				//!< Duration of content delivered to buffer. If larger than DurationDownloaded indicates dummy data was inserted into buffer.
+							double			TimeToFirstByte;				//!< Time in seconds until first data byte was received
+							double			TimeToDownload;					//!< Total time in seconds for entire download
+							int64			ByteSize;						//!< Content-Length, may be -1 if unknown (either on error or chunked transfer)
+							int64			NumBytesDownloaded;				//!< Number of bytes successfully downloaded.
+							int64			ThroughputBps;					//!< Estimated throughput in bits per second
+							int32			HTTPStatusCode;					//!< HTTP status code (0 if not connected to server yet)
+							bool			bWasSuccessful;					//!< true if download was successful, false if not
+							bool			bWasAborted;					//!< true if download was aborted by ABR (not by playback!)
+							bool			bDidTimeout;					//!< true if a timeout occurred. Only set if timeouts are enabled. Usually the ABR will monitor and abort.
+							bool			bParseFailure;					//!< true if the segment could not be parsed
+							bool			bInsertedFillerData;
+					*/
+
+					// Too many failures already? It's unlikely that there is a way that would magically fix itself now.
+					if (Stats.RetryNumber >= 3)
+					{
+						LogMessage(IInfoLog::ELevel::Error, FString::Printf(TEXT("Exceeded permissable number of retries (%d). Failing now."), Stats.RetryNumber));
+						return ESegmentAction::Fail;
+					}
+
+					// If we cannot switch to another stream we need to retry this one.
+					if (!bCanSwitchToAlternateStreams)
+					{
+						return ESegmentAction::Retry;
+					}
+
+					// Is this an init segment?
+					if (Stats.SegmentType == Metrics::ESegmentType::Init)
+					{
+						bRetryIfPossible = true;
+						// Did the init segment fail to be parsed? If so this stream is dead for good.
+						if (Stats.bParseFailure)
+						{
+							CurrentStreamInfo->Health.BecomesAvailableAgainAtUTC.SetToInvalid();
+						}
+						// Did we abort that stream in ReportDownloadProgress() ?
+						else if (Stats.bWasAborted)
+						{
+							// Since the download was aborted for the lack of time we should not be trying anything of a same or better bitrate.
+							double bps = 8.0 * Stats.NumBytesDownloaded / Stats.TimeToDownload;
+							double fallbackBps = Stats.Bitrate / 2;
+							AverageBandwidth.InitializeTo(Utils::Min(fallbackBps / 2.0, bps));
+						}
+						// Any other case of failure
+						else
+						{
+							// Take this stream offline for a brief moment.
+							if (CurrentStreamInfo->Health.BecomesAvailableAgainAtUTC == FTimeValue::GetZero())
+							{
+								CurrentStreamInfo->Health.BecomesAvailableAgainAtUTC = TimeNow + FTimeValue().SetFromMilliseconds(1000);
+							}
+						}
+					}
+					// A media segment failure
+					else
+					{
+						bRetryIfPossible = true;
+						bSkipWithFiller = true;
+
+						// Did we abort that stream in ReportDownloadProgress() ?
+						if (Stats.bWasAborted)
+						{
+							// Since the download was aborted for the lack of time we should not be trying anything of a same or better bitrate.
+							if (Stats.StreamType == EStreamType::Video)
+							{
+								double bps = 8.0 * Stats.NumBytesDownloaded / Stats.TimeToDownload;
+								double fallbackBps = Stats.Bitrate / 2;
+								AverageBandwidth.InitializeTo(Utils::Min(fallbackBps / 2.0, bps));
+							}
+						}
+						else
+						{
+							// Take a video stream offline for a brief moment. Audio is usually the only one with no alternative to switch to, so don't do that!
+							if (Stats.StreamType == EStreamType::Video)
+							{
+								if (CurrentStreamInfo->Health.BecomesAvailableAgainAtUTC == FTimeValue::GetZero())
+								{
+									CurrentStreamInfo->Health.BecomesAvailableAgainAtUTC = TimeNow + FTimeValue().SetFromMilliseconds(1000);
+								}
+							}
+							else
+							{
+								// For VoD we retry the audio segment up to a certain number of times while for Live we skip over it
+								// since retrying is a luxury we do not have in this case.
+								if (PresentationType == IAdaptiveStreamSelector::EMediaPresentationType::OnDemand)
+								{
+									bRetryIfPossible = true;
+								}
+								else
+								{
+									bRetryIfPossible = false;
+								}
+							}
+						}
+
+						// If any content has already been put out into the buffers we cannot retry the segment on another quality level.
+						if (Stats.DurationDelivered > 0.0)
+						{
+							bRetryIfPossible = false;
+							bSkipWithFiller = false;
+						}
+					}
+				}
+				// Update stats
+				CurrentStreamInfo->Health.LastDownloadStats = Stats;
+			}
+		}
+
+		// Anything besides video or audio?
+		// NOTE: For the love of God do not remove this line without re-instating the else if() below. See comments there!
+		if (StreamType != EStreamType::Video && StreamType != EStreamType::Audio)
+		{
+			return ESegmentAction::FetchNext;
+		}
+
+
+		TArray<TSharedPtrTS<FStreamInformation>>			PossibleRepresentations;
+		FAccessUnitBufferInfo								CurrentBufferStats;
+
+		PlayerSessionServices->GetStreamBufferStats(CurrentBufferStats, StreamType);
+
+		// For video streams do a download simulation
+		if (StreamType == EStreamType::Video)
+		{
+			const double secondsInVideoBuffer = CurrentBufferStats.PlayableDuration.GetAsSeconds();
+			// Buffer watermark levels
+			const double videoBufferLowWatermark = double(Config.VideoBufferLowWatermarkMSec) / 1000.0;
+			const double videoBufferHighWatermark = double(Config.VideoBufferHighWatermarkMSec) / 1000.0;
+
+			// Get average and last bitrate.
+			const int64 averageBandwidth = Utils::Max(GetAverageBandwidth(), (int64)100000);
+			int64 lastBandwidth = GetLastBandwidth();
+			if (lastBandwidth <= 0)
+			{
+				lastBandwidth = averageBandwidth;
+			}
+			lastBandwidth = Utils::Max(lastBandwidth, (int64)100000);
+
+			// Setup simulation and last bandwidth such that the current audio bitrate is accounted for and clamped to at least 100 Kbps.
+			const int64 simulationBandwidthLongterm = Utils::Max((int64)Config.SimulationLongtermRateScale.Apply(lastBandwidth, averageBandwidth), (int64)100000);
+			// How long to simulate download for.
+			const double simulationDuration = double(Config.ForwardSimulationDurationMSec) / 1000.0;
+
+			// Bandwidth to simulate download of just the next fragment with.
+			const int64 simulationBandwidthNextFragment = Utils::Max((int64)Config.SimulationNextFragmentRateScale.Apply(lastBandwidth, averageBandwidth), (int64)100000);
+
+			// For each video stream simulate how the buffer will behave if we were to use it.
+			// The lowest quality is always available regardless of bandwidth or resolution or other constraints.
+			if (StreamInformationVideo.Num())
+			{
+				PossibleRepresentations.Push(StreamInformationVideo[0]);
+			}
+			for (int32 nStr = 1; nStr < StreamInformationVideo.Num(); ++nStr)
+			{
+				// Check if bitrate and resolution are acceptable
+				if (StreamInformationVideo[nStr]->Representation->GetBitrate() <= BandwidthCeiling &&
+					!StreamInformationVideo[nStr]->Resolution.ExceedsLimit(MaxStreamResolution))
+				{
+					FSimulationResult sim;
+
+					// Get information on the next segments.
+					TArray<IManifest::IPlayPeriod::FSegmentInformation> NextSegments;
+					FTimeValue AverageSegmentDuration;
+					CurrentPlayPeriod->GetSegmentInformation(NextSegments, AverageSegmentDuration, CurrentSegment, FTimeValue().SetFromSeconds(simulationDuration), StreamInformationVideo[nStr]->AdaptationSet, StreamInformationVideo[nStr]->Representation);
+
+					// Long term prediction: Simulate download (at a user-scaled bandwidth factor)
+					SimulateDownload(sim, NextSegments, simulationDuration, simulationBandwidthLongterm);
+
+					// Mid term prediction: Estimate how much we will have in the video buffer when downloading just the next fragment.
+					double secondsInBufferAfterNext;
+					if (NextSegments.Num())
+					{
+						secondsInBufferAfterNext = secondsInVideoBuffer + NextSegments[0].Duration.GetAsSeconds() - (NextSegments[0].ByteSize * (8.0 / simulationBandwidthNextFragment));
+					}
+					else
+					{
+						secondsInBufferAfterNext = secondsInVideoBuffer;
+					}
+
+					// A stream is feasible if it:
+					//    - gives at least some data when the current buffer level is lower than the low-threshold
+					//      OR
+					//    - even with the smallest gain in data the buffer level AFTER the NEXT fragment will be higher than the low-threshold
+					// AND
+					//    - the simulated gain plus what is in the buffer now will be higher than the high-threshold.
+					bool bCond1 = sim.SmallestGain > 0.0 && secondsInVideoBuffer < videoBufferLowWatermark;
+					bool bCond2 = sim.SmallestGain + secondsInBufferAfterNext > videoBufferLowWatermark;
+					bool bCond3 = sim.SecondsGained + secondsInVideoBuffer > videoBufferHighWatermark;
+					bool bIsFeasible = (bCond1 || bCond2) && bCond3;
+
+					// If there is an enforced initial bitrate then any stream that is within that bitrate is feasible by default.
+					if (ForcedInitialBandwidth && StreamInformationVideo[nStr]->Representation->GetBitrate() <= ForcedInitialBandwidth)
+					{
+						bIsFeasible = true;
+					}
+
+					if (bIsFeasible)
+					{
+						PossibleRepresentations.Push(StreamInformationVideo[nStr]);
+					}
+				}
+			}
+		}
+		// For audio we do no download simulation or anything.
+	// NOTE: removed this else if check because of static code analyzer complaining this would always be true, which it is not.
+	//       it is only true as long as the above check that if this is anything but video or audio we already return.
+	//       otherwise, with that if() above removed *NOT* having this check here will be a major bug!
+		//else if (StreamType == EStreamType::Audio)
+		else
+		{
+			for (int32 nStr = 0; nStr < StreamInformationAudio.Num(); ++nStr)
+			{
+				PossibleRepresentations.Push(StreamInformationAudio[nStr]);
+			}
+		}
+
+		// At this point we have a list of representations that we could potentially use.
+		if (PossibleRepresentations.Num())
+		{
+			// Remove those that have been externally blacklisted. This may include the lowest quality stream!
+			FMediaCriticalSection::ScopedLock lock(AccessMutex);
+			if (BlacklistedExternally.Num())
+			{
+				TArray<TSharedPtrTS<FStreamInformation>> RemainingCandidates;
+				for (int32 j = 0, jMax = PossibleRepresentations.Num(); j < jMax; ++j)
+				{
+					bool bStillGood = true;
+					for (int32 i = 0, iMax = BlacklistedExternally.Num(); i < iMax; ++i)
+					{
+						if (PossibleRepresentations[j]->AdaptationSet->GetUniqueIdentifier() == BlacklistedExternally[i].AdaptationSetUniqueID &&
+							PossibleRepresentations[j]->Representation->GetUniqueIdentifier() == BlacklistedExternally[i].RepresentationUniqueID)
+						{
+							bStillGood = false;
+							break;
+						}
+					}
+					if (bStillGood)
+					{
+						RemainingCandidates.Push(PossibleRepresentations[j]);
+					}
+				}
+				Swap(PossibleRepresentations, RemainingCandidates);
+			}
+
+			// Filter out streams that are currently not healthy
+			//if (1)
+			{
+				TArray<TSharedPtrTS<FStreamInformation>> RemainingCandidates;
+				for (int32 j = 0, jMax = PossibleRepresentations.Num(); j < jMax; ++j)
+				{
+					bool bStillGood = PossibleRepresentations[j]->Health.BecomesAvailableAgainAtUTC <= TimeNow;
+					if (bStillGood)
+					{
+						RemainingCandidates.Push(PossibleRepresentations[j]);
+					}
+				}
+				Swap(PossibleRepresentations, RemainingCandidates);
+			}
+
+			// Still something left to use?
+			if (PossibleRepresentations.Num())
+			{
+				const TSharedPtrTS<FStreamInformation>& Candidate = PossibleRepresentations.Top();
+				CurrentPlayPeriod->SelectStream(Candidate->AdaptationSet, Candidate->Representation, Candidate->Representation->GetCDN());
+
+				if (bRetryIfPossible)
+				{
+					return ESegmentAction::Retry;
+				}
+				else if (bSkipWithFiller)
+				{
+					return ESegmentAction::Fill;
+				}
+				else
+				{
+					return ESegmentAction::FetchNext;
+				}
+			}
+			else
+			{
+				// TODO: Are there representations we could lift the temporary blacklist from?
+				if (bSkipWithFiller)
+				{
+					return ESegmentAction::Fill;
+				}
+			}
+		}
+		return ESegmentAction::Fail;
+	}
+
+
+} // namespace Electra
+
diff --git a/Engine/Plugins/Media/ElectraPlayer/Source/ElectraPlayerRuntime/Private/Runtime/Player/AdaptiveStreamingPlayerABR.h b/Engine/Plugins/Media/ElectraPlayer/Source/ElectraPlayerRuntime/Private/Runtime/Player/AdaptiveStreamingPlayerABR.h
new file mode 100644
--- /dev/null
+++ b/Engine/Plugins/Media/ElectraPlayer/Source/ElectraPlayerRuntime/Private/Runtime/Player/AdaptiveStreamingPlayerABR.h
@@ -0,0 +1,149 @@
+// Copyright Epic Games, Inc. All Rights Reserved.
+
+#pragma once
+
+#include "PlayerCore.h"
+#include "PlayerTime.h"
+#include "StreamTypes.h"
+#include "HTTP/HTTPManager.h"
+#include "Player/Manifest.h"
+#include "Player/PlayerSessionServices.h"
+#include "Player/AdaptiveStreamingPlayerABR_State.h"
+#include "Player/AdaptiveStreamingPlayerMetrics.h"
+
+namespace Electra
+{
+
+
+class IAdaptiveStreamSelector : public IAdaptiveStreamingPlayerMetrics
+{
+public:
+	struct FConfiguration
+	{
+		FConfiguration();
+		struct FRateScale
+		{
+			FRateScale();
+			double Apply(const double lastKbps, const double averageKbps) const;
+			float		LastToAvg;			//!< [last, average] bandwidth scaler such that Rate = last + mLastToAvg * (average - last)
+			float		Scale;				//!< Scale factor to apply next
+			uint32		MaxClampBPS;		//!< Maximum value to clamp result to (in bits per second)
+		};
+		FRateScale		SimulationLongtermRateScale;
+		FRateScale		SimulationNextFragmentRateScale;
+		FRateScale		FinalRateAdjust;
+		uint32			ForwardSimulationDurationMSec;			//!< Forward simulation duration
+		uint32			VideoBufferLowWatermarkMSec;
+		uint32			VideoBufferHighWatermarkMSec;
+		uint32			MaxBandwidthHistorySize;				//!< Number of bandwidth history samples to give the estimated average bandwidth.
+		uint32			DiscardSmallBandwidthSamplesLess;
+		uint32			FudgeSmallBandwidthSamplesLess;
+		uint32			FudgeSmallBandwidthSamplesFactor;
+	};
+
+	//! Create an instance of this class
+	static TSharedPtrTS<IAdaptiveStreamSelector> Create(IPlayerSessionServices* PlayerSessionServices, const FConfiguration& config);
+	virtual ~IAdaptiveStreamSelector() = default;
+
+	enum class EMediaPresentationType
+	{
+		OnDemand,
+		Live,
+		Realtime
+	};
+	virtual void SetPresentationType(EMediaPresentationType AssetType) = 0;
+
+	//! Sets whether or not a segment download failure can potentially be resolved by switching to a different stream. Defaults to yes.
+	virtual void SetCanSwitchToAlternateStreams(bool bCanSwitch) = 0;
+
+	//! Sets the manifest from which to select streams.
+	virtual void SetCurrentPlaybackPeriod(TSharedPtrTS<IManifest::IPlayPeriod> CurrentPlayPeriod) = 0;
+
+	//! Sets the initial bandwidth, either from guessing, past history or a current measurement.
+	virtual void SetBandwidth(int64 bitsPerSecond) = 0;
+
+	//! Sets a forced bitrate for the next segment fetch only.
+	virtual void SetForcedNextBandwidth(int64 bitsPerSecond) = 0;
+
+	struct FBlacklistedStream
+	{
+		FString		AssetUniqueID;
+		FString		AdaptationSetUniqueID;
+		FString		RepresentationUniqueID;
+		FString		CDN;
+		bool operator == (const FBlacklistedStream& rhs) const
+		{
+			return AssetUniqueID == rhs.AssetUniqueID &&
+				   AdaptationSetUniqueID == rhs.AdaptationSetUniqueID &&
+				   RepresentationUniqueID == rhs.RepresentationUniqueID &&
+				   CDN == rhs.CDN;
+		}
+	};
+	virtual void MarkStreamAsUnavailable(const FBlacklistedStream& BlacklistedStream) = 0;
+	virtual void MarkStreamAsAvailable(const FBlacklistedStream& NoLongerBlacklistedStream) = 0;
+
+
+	virtual FABRDownloadProgressDecision ReportDownloadProgress(const Metrics::FSegmentDownloadStats& SegmentDownloadStats) = 0;
+	virtual void ReportDownloadEnd(const Metrics::FSegmentDownloadStats& SegmentDownloadStats) = 0;
+
+
+	//! Returns the last measured bandwidth sample in bps.
+	virtual int64 GetLastBandwidth() = 0;
+
+	//! Returns the average measured bandwidth in bps.
+	virtual int64 GetAverageBandwidth() = 0;
+
+	//! Returns the average measured throughput in bps.
+	virtual int64 GetAverageThroughput() = 0;
+
+	//! Returns the average measured latency in seconds.
+	virtual double GetAverageLatency() = 0;
+
+	//! Sets a highest bandwidth limit. Call with 0 to disable.
+	virtual void SetBandwidthCeiling(int32 HighestManifestBitrate) = 0;
+
+	//! Limits video resolution.
+	virtual void SetMaxVideoResolution(int32 MaxWidth, int32 MaxHeight) = 0;
+
+	enum class ESegmentAction
+	{
+		FetchNext,				//!< Fetch the next segment normally.
+		Retry,					//!< Retry the same segment. Another quality or CDN may have been picked out already.
+		Fill,					//!< Fill the segment's duration worth with dummy data.
+		Fail					//!< Abort playback
+	};
+
+	//! Selects a feasible stream from the stream set to fetch the next fragment from.
+	virtual ESegmentAction SelectSuitableStreams(FTimeValue& OutDelay, TSharedPtrTS<const IStreamSegment> CurrentSegment) = 0;
+
+
+	virtual void ReportOpenSource(const FString& URL) = 0;
+	virtual void ReportReceivedMasterPlaylist(const FString& EffectiveURL) = 0;
+	virtual void ReportReceivedPlaylists() = 0;
+	virtual void ReportPlaylistDownload(const Metrics::FPlaylistDownloadStats& PlaylistDownloadStats) = 0;
+	virtual void ReportBufferingStart(Metrics::EBufferingReason BufferingReason) = 0;
+	virtual void ReportBufferingEnd(Metrics::EBufferingReason BufferingReason) = 0;
+	virtual void ReportBandwidth(int64 EffectiveBps, int64 ThroughputBps, double LatencyInSeconds) = 0;
+	virtual void ReportBufferUtilization(const Metrics::FBufferStats& BufferStats) = 0;
+	virtual void ReportSegmentDownload(const Metrics::FSegmentDownloadStats& SegmentDownloadStats) = 0;
+	virtual void ReportLicenseKey(const Metrics::FLicenseKeyStats& LicenseKeyStats) = 0;
+	virtual void ReportDataAvailabilityChange(const Metrics::FDataAvailabilityChange& DataAvailability) = 0;
+	virtual void ReportVideoQualityChange(int32 NewBitrate, int32 PreviousBitrate, bool bIsDrasticDownswitch) = 0;
+	virtual void ReportPrerollStart() = 0;
+	virtual void ReportPrerollEnd() = 0;
+	virtual void ReportPlaybackStart() = 0;
+	virtual void ReportPlaybackPaused() = 0;
+	virtual void ReportPlaybackResumed() = 0;
+	virtual void ReportPlaybackEnded() = 0;
+	virtual void ReportJumpInPlayPosition(const FTimeValue& ToNewTime, const FTimeValue& FromTime, Metrics::ETimeJumpReason TimejumpReason) = 0;
+	virtual void ReportPlaybackStopped() = 0;
+	virtual void ReportError(const FString& ErrorReason) = 0;
+	virtual void ReportLogMessage(IInfoLog::ELevel LogLevel, const FString& LogMessage, int64 PlayerWallclockMilliseconds) = 0;
+	virtual void ReportDroppedVideoFrame() = 0;
+	virtual void ReportDroppedAudioFrame() = 0;
+};
+
+
+} // namespace Electra
+
+
diff --git a/Engine/Plugins/Media/ElectraPlayer/Source/ElectraPlayerRuntime/Private/Runtime/Player/AdaptiveStreamingPlayerEvents.cpp b/Engine/Plugins/Media/ElectraPlayer/Source/ElectraPlayerRuntime/Private/Runtime/Player/AdaptiveStreamingPlayerEvents.cpp
new file mode 100644
--- /dev/null
+++ b/Engine/Plugins/Media/ElectraPlayer/Source/ElectraPlayerRuntime/Private/Runtime/Player/AdaptiveStreamingPlayerEvents.cpp
@@ -0,0 +1,75 @@
+// Copyright Epic Games, Inc. All Rights Reserved.
+
+#include "Player/AdaptiveStreamingPlayerInternal.h"
+
+namespace Electra
+{
+
+
+//-----------------------------------------------------------------------------
+/**
+ * Dispatches a 'buffering' event.
+ *
+ * @param bBegin
+ * @param reason
+ */
+void FAdaptiveStreamingPlayer::DispatchBufferingEvent(bool bBegin, FAdaptiveStreamingPlayer::EPlayerState Reason)
+{
+	if (bBegin)
+	{
+		DispatchEvent(FMetricEvent::ReportBufferingStart(Reason == EPlayerState::eState_Buffering ? Metrics::EBufferingReason::Initial :
+														 Reason == EPlayerState::eState_Rebuffering ? Metrics::EBufferingReason::Rebuffering :
+														 Reason == EPlayerState::eState_Seeking ? Metrics::EBufferingReason::Seeking : Metrics::EBufferingReason::Rebuffering));
+	}
+	else
+	{
+		DispatchEvent(FMetricEvent::ReportBufferingEnd(Reason == EPlayerState::eState_Buffering ? Metrics::EBufferingReason::Initial :
+													   Reason == EPlayerState::eState_Rebuffering ? Metrics::EBufferingReason::Rebuffering :
+													   Reason == EPlayerState::eState_Seeking ? Metrics::EBufferingReason::Seeking : Metrics::EBufferingReason::Rebuffering));
+	}
+}
+
+//-----------------------------------------------------------------------------
+/**
+ * Dispatches a 'fragment download' event.
+ *
+ * @param Request
+ */
+void FAdaptiveStreamingPlayer::DispatchSegmentDownloadedEvent(TSharedPtrTS<IStreamSegment> Request)
+{
+	if (Request.IsValid())
+	{
+		Metrics::FSegmentDownloadStats stats;
+		Request->GetDownloadStats(stats);
+		DispatchEvent(FMetricEvent::ReportSegmentDownload(stats));
+	}
+}
+
+
+void FAdaptiveStreamingPlayer::DispatchBufferUtilizationEvent(EStreamType BufferType)
+{
+	Metrics::FBufferStats stats;
+	stats.BufferType = BufferType;
+
+	FAccessUnitBufferInfo bufStats;
+
+	if (BufferType == EStreamType::Video)
+	{
+		MultiStreamBufferVid.GetStats(bufStats);
+	}
+	else if (BufferType == EStreamType::Audio)
+	{
+		MultiStreamBufferAud.GetStats(bufStats);
+	}
+
+	stats.MaxDurationInSeconds = bufStats.MaxDuration.GetAsSeconds();
+	stats.DurationInUse 	   = bufStats.PushedDuration.GetAsSeconds();
+	stats.MaxByteCapacity      = bufStats.MaxDataSize;
+	stats.BytesInUse		   = bufStats.CurrentMemInUse;
+	DispatchEvent(FMetricEvent::ReportBufferUtilization(stats));
+}
+
+
+
+} // namespace Electra
+
