ElectraPlayer: RELEASE. Moved from NFL to public. Note that this is done via a Delete and "Add New" operation as discussed.

diff --git a/Engine/Plugins/Media/ElectraPlayer/Source/ElectraPlayerRuntime/Private/Runtime/Player/AdaptiveStreamingPlayerInternal.h b/Engine/Plugins/Media/ElectraPlayer/Source/ElectraPlayerRuntime/Private/Runtime/Player/AdaptiveStreamingPlayerInternal.h
new file mode 100644
--- /dev/null
+++ b/Engine/Plugins/Media/ElectraPlayer/Source/ElectraPlayerRuntime/Private/Runtime/Player/AdaptiveStreamingPlayerInternal.h
@@ -0,0 +1,1388 @@
+// Copyright Epic Games, Inc. All Rights Reserved.
+
+#pragma once
+
+#include "Player/AdaptiveStreamingPlayer.h"
+#include "Player/Manifest.h"
+#include "Player/AdaptiveStreamingPlayerInternalConfig.h"
+
+#include "HTTP/HTTPManager.h"
+#include "Player/PlayerSessionServices.h"
+#include "SynchronizedClock.h"
+
+#include "Demuxer/ParserISO14496-12.h"
+
+#include "Player/PlayerStreamReader.h"
+#include "Player/PlaylistReader.h"
+#include "Player/PlayerStreamFilter.h"
+
+#include "InfoLog.h"
+
+#define INTERR_ALL_STREAMS_HAVE_FAILED			1
+#define INTERR_UNSUPPORTED_FORMAT				2
+#define INTERR_COULD_NOT_LOCATE_START_SEGMENT	3
+#define INTERR_FRAGMENT_NOT_AVAILABLE			0x101
+#define INTERR_CREATE_FRAGMENT_READER			0x103
+#define INTERR_REBUFFER_SHALL_THROW_ERROR		0x200
+
+
+namespace Electra
+{
+class FAdaptiveStreamingPlayer;
+
+
+class FMediaRenderClock : public IMediaRenderClock
+{
+public:
+	FMediaRenderClock()
+		: bIsPaused(true)
+	{
+	}
+
+	virtual ~FMediaRenderClock()
+	{
+	}
+
+	virtual void SetCurrentTime(ERendererType ForRenderer, const FTimeValue& CurrentRenderTime) override
+	{
+		FClock* Clk = GetClock(ForRenderer);
+		if (Clk)
+		{
+			FTimeValue Now(MEDIAutcTime::Current());
+			FMediaCriticalSection::ScopedLock lock(Lock);
+			Clk->RenderTime = CurrentRenderTime;
+			Clk->LastSystemBaseTime = Now;
+			Clk->RunningTimeOffset.SetToZero();
+		}
+	}
+
+	virtual FTimeValue GetInterpolatedRenderTime(ERendererType FromRenderer) override
+	{
+		FClock* Clk = GetClock(FromRenderer);
+		if (Clk)
+		{
+			FTimeValue Now(MEDIAutcTime::Current());
+			FMediaCriticalSection::ScopedLock lock(Lock);
+			FTimeValue diff = !bIsPaused ? Now - Clk->LastSystemBaseTime : FTimeValue::GetZero();
+			return Clk->RenderTime + Clk->RunningTimeOffset + diff;
+		}
+		return FTimeValue();
+	}
+
+	void Start()
+	{
+		FTimeValue Now(MEDIAutcTime::Current());
+		FMediaCriticalSection::ScopedLock lock(Lock);
+		if (bIsPaused)
+		{
+			for(int32 i = 0; i < FMEDIA_STATIC_ARRAY_COUNT(Clock); ++i)
+			{
+				Clock[i].LastSystemBaseTime = Now;
+			}
+			bIsPaused = false;
+		}
+	}
+
+	void Stop()
+	{
+		FTimeValue Now(MEDIAutcTime::Current());
+		FMediaCriticalSection::ScopedLock lock(Lock);
+		if (!bIsPaused)
+		{
+			bIsPaused = true;
+			for(int32 i=0; i<FMEDIA_STATIC_ARRAY_COUNT(Clock); ++i)
+			{
+				FTimeValue diff = Now - Clock[i].LastSystemBaseTime;
+				Clock[i].RunningTimeOffset += diff;
+			}
+		}
+	}
+
+	bool IsRunning() const
+	{
+		return !bIsPaused;
+	}
+
+private:
+	struct FClock
+	{
+		FClock()
+		{
+			LastSystemBaseTime = MEDIAutcTime::Current();
+			RunningTimeOffset.SetToZero();
+		}
+		FTimeValue	RenderTime;
+		FTimeValue	LastSystemBaseTime;
+		FTimeValue	RunningTimeOffset;
+	};
+
+	FClock* GetClock(ERendererType ForRenderer)
+	{
+		FClock* Clk = nullptr;
+		switch(ForRenderer)
+		{
+			case IMediaRenderClock::ERendererType::Video:
+				Clk = &Clock[0];
+				break;
+			case IMediaRenderClock::ERendererType::Audio:
+				Clk = &Clock[1];
+				break;
+			case IMediaRenderClock::ERendererType::Subtitles:
+				Clk = &Clock[2];
+				break;
+			default:
+				Clk = nullptr;
+				break;
+		}
+		return Clk;
+	}
+
+	FMediaCriticalSection		Lock;
+	FClock			Clock[3];
+	bool			bIsPaused;
+};
+
+
+
+/**
+ * Interchange structure between the player worker thread and the public API to
+ * avoid mutex locks all over the place.
+ */
+struct FPlaybackState
+{
+	FPlaybackState()
+	{
+		Reset();
+	}
+	void Reset()
+	{
+		FMediaCriticalSection::ScopedLock lock(Lock);
+		SeekableRange.Reset();
+		Duration.SetToInvalid();
+		CurrentPlayPosition.SetToInvalid();
+		LoopState.Reset();
+		bHaveMetadata = false;
+		bHasEnded     = false;
+		bIsSeeking    = false;
+		bIsBuffering  = false;
+		bIsPlaying    = false;
+		bIsPaused     = false;
+	}
+
+	mutable FMediaCriticalSection			Lock;
+	FTimeRange								SeekableRange;
+	FTimeRange								TimelineRange;
+	FTimeValue								Duration;
+	FTimeValue								CurrentPlayPosition;
+	FPlayerLoopState						LoopState;
+	bool									bHaveMetadata;
+	bool									bHasEnded;
+	bool									bIsSeeking;
+	bool									bIsBuffering;
+	bool									bIsPlaying;
+	bool									bIsPaused;
+	TArray<FStreamMetadata>					VideoStreams;
+	TArray<FStreamMetadata>					AudioStreams;
+	TArray<FTimespan>						SeekablePositions;
+
+
+	void SetSeekableRange(const FTimeRange& TimeRange)
+	{
+		FMediaCriticalSection::ScopedLock lock(Lock);
+		SeekableRange = TimeRange;
+	}
+	void GetSeekableRange(FTimeRange& OutRange) const
+	{
+		FMediaCriticalSection::ScopedLock lock(Lock);
+		OutRange = SeekableRange;
+	}
+
+	void SetSeekablePositions(const TArray<FTimespan>& InSeekablePositions)
+	{
+		FMediaCriticalSection::ScopedLock lock(Lock);
+		SeekablePositions = InSeekablePositions;
+	}
+	void GetSeekablePositions(TArray<FTimespan>& OutSeekablePositions) const
+	{
+		FMediaCriticalSection::ScopedLock lock(Lock);
+		OutSeekablePositions = SeekablePositions;
+	}
+
+	void SetTimelineRange(const FTimeRange& TimeRange)
+	{
+		FMediaCriticalSection::ScopedLock lock(Lock);
+		TimelineRange = TimeRange;
+	}
+	void GetTimelineRange(FTimeRange& OutRange) const
+	{
+		FMediaCriticalSection::ScopedLock lock(Lock);
+		OutRange = TimelineRange;
+	}
+
+	void SetDuration(const FTimeValue& InDuration)
+	{
+		FMediaCriticalSection::ScopedLock lock(Lock);
+		Duration = InDuration;
+	}
+	FTimeValue GetDuration() const
+	{
+		FMediaCriticalSection::ScopedLock lock(Lock);
+		return Duration;
+	}
+
+	void SetPlayPosition(const FTimeValue& InPosition)
+	{
+		FMediaCriticalSection::ScopedLock lock(Lock);
+		CurrentPlayPosition = InPosition;
+	}
+	FTimeValue GetPlayPosition() const
+	{
+		FMediaCriticalSection::ScopedLock lock(Lock);
+		return CurrentPlayPosition;
+	}
+
+	void SetHaveMetadata(bool bInHaveMetadata)
+	{
+		FMediaCriticalSection::ScopedLock lock(Lock);
+		bHaveMetadata = bInHaveMetadata;
+	}
+	bool GetHaveMetadata() const
+	{
+		FMediaCriticalSection::ScopedLock lock(Lock);
+		return bHaveMetadata;
+	}
+
+	void SetHasEnded(bool bInHasEnded)
+	{
+		FMediaCriticalSection::ScopedLock lock(Lock);
+		bHasEnded = bInHasEnded;
+	}
+	bool GetHasEnded() const
+	{
+		FMediaCriticalSection::ScopedLock lock(Lock);
+		return bHasEnded;
+	}
+
+	void SetIsSeeking(bool bInIsSeeking)
+	{
+		FMediaCriticalSection::ScopedLock lock(Lock);
+		bIsSeeking = bInIsSeeking;
+	}
+	bool GetIsSeeking() const
+	{
+		FMediaCriticalSection::ScopedLock lock(Lock);
+		return bIsSeeking;
+	}
+
+	void SetIsBuffering(bool bInIsBuffering)
+	{
+		FMediaCriticalSection::ScopedLock lock(Lock);
+		bIsBuffering = bInIsBuffering;
+	}
+	bool GetIsBuffering() const
+	{
+		FMediaCriticalSection::ScopedLock lock(Lock);
+		return bIsBuffering;
+	}
+
+	void SetIsPlaying(bool bInIsPlaying)
+	{
+		FMediaCriticalSection::ScopedLock lock(Lock);
+		bIsPlaying = bInIsPlaying;
+	}
+	bool GetIsPlaying() const
+	{
+		FMediaCriticalSection::ScopedLock lock(Lock);
+		return bIsPlaying;
+	}
+
+	void SetIsPaused(bool bInIsPaused)
+	{
+		FMediaCriticalSection::ScopedLock lock(Lock);
+		bIsPaused = bInIsPaused;
+	}
+	bool GetIsPaused() const
+	{
+		FMediaCriticalSection::ScopedLock lock(Lock);
+		return bIsPaused;
+	}
+
+	void SetPausedAndPlaying(bool bInIsPaused, bool bInIsPlaying)
+	{
+		FMediaCriticalSection::ScopedLock lock(Lock);
+		bIsPaused  = bInIsPaused;
+		bIsPlaying = bInIsPlaying;
+	}
+
+	void SetStreamMetadata(const TArray<FStreamMetadata> &InVideoStreams, const TArray<FStreamMetadata>& InAudioStreams)
+	{
+		FMediaCriticalSection::ScopedLock lock(Lock);
+		VideoStreams = InVideoStreams;
+		AudioStreams = InAudioStreams;
+	}
+	void GetStreamMetadata(TArray<FStreamMetadata> &OutVideoStreams, TArray<FStreamMetadata>& OutAudioStreams) const
+	{
+		FMediaCriticalSection::ScopedLock lock(Lock);
+		OutVideoStreams = VideoStreams;
+		OutAudioStreams = AudioStreams;
+	}
+
+	void SetLoopState(const FPlayerLoopState& InLoopState)
+	{
+		FMediaCriticalSection::ScopedLock lock(Lock);
+		LoopState = InLoopState;
+	}
+	void GetLoopState(FPlayerLoopState& OutLoopState) const
+	{
+		FMediaCriticalSection::ScopedLock lock(Lock);
+		OutLoopState = LoopState;
+	}
+	void SetLoopStateEnable(bool bEnable)
+	{
+		FMediaCriticalSection::ScopedLock lock(Lock);
+		LoopState.bLoopEnabled = bEnable;
+	}
+
+};
+
+
+
+
+
+
+struct FMetricEvent
+{
+	enum class EType
+	{
+		OpenSource,
+		ReceivedMasterPlaylist,
+		ReceivedPlaylists,
+		BufferingStart,
+		BufferingEnd,
+		Bandwidth,
+		BufferUtilization,
+		PlaylistDownload,
+		SegmentDownload,
+		DataAvailabilityChange,
+		VideoQualityChange,
+		PrerollStart,
+		PrerollEnd,
+		PlaybackStart,
+		PlaybackPaused,
+		PlaybackResumed,
+		PlaybackEnded,
+		PlaybackJumped,
+		PlaybackStopped,
+		LicenseKey,
+		Errored,
+		LogMessage,
+	};
+
+	struct FParam
+	{
+		struct FBandwidth
+		{
+			int64	EffectiveBps;
+			int64	ThroughputBps;
+			double	Latency;
+		};
+		struct FQualityChange
+		{
+			int32	NewBitrate;
+			int32	PrevBitrate;
+			bool	bIsDrastic;
+		};
+		struct FLogMessage
+		{
+			FString				Message;
+			IInfoLog::ELevel	Level;
+			int64				AtMillis;
+		};
+		struct FTimeJumped
+		{
+			FTimeValue					ToNewTime;
+			FTimeValue					FromTime;
+			Metrics::ETimeJumpReason	Reason;
+
+		};
+		FString								URL;
+		Metrics::EBufferingReason			BufferingReason;
+		Metrics::FBufferStats				BufferStats;
+		Metrics::FPlaylistDownloadStats		PlaylistStats;
+		Metrics::FSegmentDownloadStats		SegmentStats;
+		Metrics::FLicenseKeyStats			LicenseKeyStats;
+		Metrics::FDataAvailabilityChange	DataAvailability;
+		FBandwidth							Bandwidth;
+		FQualityChange						QualityChange;
+		FTimeJumped							TimeJump;
+		FErrorDetail						ErrorDetail;
+		FLogMessage							LogMessage;
+	};
+
+	EType			Type;
+	FParam			Param;
+	TWeakPtr<FAdaptiveStreamingPlayer, ESPMode::ThreadSafe>	Player;
+	FMediaEvent	*EventSignal = nullptr;
+
+	static TSharedPtrTS<FMetricEvent> ReportOpenSource(const FString& URL)
+	{
+		TSharedPtrTS<FMetricEvent> Evt = MakeSharedTS<FMetricEvent>();
+		Evt->Type = EType::OpenSource;
+		Evt->Param.URL = URL;
+		return Evt;
+	}
+	static TSharedPtrTS<FMetricEvent> ReportReceivedMasterPlaylist(const FString& EffectiveURL)
+	{
+		TSharedPtrTS<FMetricEvent> Evt = MakeSharedTS<FMetricEvent>();
+		Evt->Type = EType::ReceivedMasterPlaylist;
+		Evt->Param.URL = EffectiveURL;
+		return Evt;
+	}
+	static TSharedPtrTS<FMetricEvent> ReportReceivedPlaylists()
+	{
+		TSharedPtrTS<FMetricEvent> Evt = MakeSharedTS<FMetricEvent>();
+		Evt->Type = EType::ReceivedPlaylists;
+		return Evt;
+	}
+	static TSharedPtrTS<FMetricEvent> ReportBufferingStart(Metrics::EBufferingReason BufferingReason)
+	{
+		TSharedPtrTS<FMetricEvent> Evt = MakeSharedTS<FMetricEvent>();
+		Evt->Type = EType::BufferingStart;
+		Evt->Param.BufferingReason = BufferingReason;
+		return Evt;
+	}
+	static TSharedPtrTS<FMetricEvent> ReportBufferingEnd(Metrics::EBufferingReason BufferingReason)
+	{
+		TSharedPtrTS<FMetricEvent> Evt = MakeSharedTS<FMetricEvent>();
+		Evt->Type = EType::BufferingEnd;
+		Evt->Param.BufferingReason = BufferingReason;
+		return Evt;
+	}
+	static TSharedPtrTS<FMetricEvent> ReportBandwidth(int64 EffectiveBps, int64 ThroughputBps, double LatencyInSeconds)
+	{
+		TSharedPtrTS<FMetricEvent> Evt = MakeSharedTS<FMetricEvent>();
+		Evt->Type = EType::Bandwidth;
+		Evt->Param.Bandwidth.EffectiveBps = EffectiveBps;
+		Evt->Param.Bandwidth.ThroughputBps = ThroughputBps;
+		Evt->Param.Bandwidth.Latency = LatencyInSeconds;
+		return Evt;
+	}
+	static TSharedPtrTS<FMetricEvent> ReportBufferUtilization(const Metrics::FBufferStats& BufferStats)
+	{
+		TSharedPtrTS<FMetricEvent> Evt = MakeSharedTS<FMetricEvent>();
+		Evt->Type = EType::BufferUtilization;
+		Evt->Param.BufferStats = BufferStats;
+		return Evt;
+	}
+	static TSharedPtrTS<FMetricEvent> ReportPlaylistDownload(const Metrics::FPlaylistDownloadStats& PlaylistDownloadStats)
+	{
+		TSharedPtrTS<FMetricEvent> Evt = MakeSharedTS<FMetricEvent>();
+		Evt->Type = EType::PlaylistDownload;
+		Evt->Param.PlaylistStats = PlaylistDownloadStats;
+		return Evt;
+	}
+	static TSharedPtrTS<FMetricEvent> ReportSegmentDownload(const Metrics::FSegmentDownloadStats& SegmentDownloadStats)
+	{
+		TSharedPtrTS<FMetricEvent> Evt = MakeSharedTS<FMetricEvent>();
+		Evt->Type = EType::SegmentDownload;
+		Evt->Param.SegmentStats = SegmentDownloadStats;
+		return Evt;
+	}
+	static TSharedPtrTS<FMetricEvent> ReportLicenseKey(const Metrics::FLicenseKeyStats& LicenseKeyStats)
+	{
+		TSharedPtrTS<FMetricEvent> Evt = MakeSharedTS<FMetricEvent>();
+		Evt->Type = EType::LicenseKey;
+		Evt->Param.LicenseKeyStats = LicenseKeyStats;
+		return Evt;
+	}
+	static TSharedPtrTS<FMetricEvent> ReportDataAvailabilityChange(const Metrics::FDataAvailabilityChange& DataAvailability)
+	{
+		TSharedPtrTS<FMetricEvent> Evt = MakeSharedTS<FMetricEvent>();
+		Evt->Type = EType::DataAvailabilityChange;
+		Evt->Param.DataAvailability = DataAvailability;
+		return Evt;
+	}
+	static TSharedPtrTS<FMetricEvent> ReportVideoQualityChange(int32 NewBitrate, int32 PreviousBitrate, bool bIsDrasticDownswitch)
+	{
+		TSharedPtrTS<FMetricEvent> Evt = MakeSharedTS<FMetricEvent>();
+		Evt->Type = EType::VideoQualityChange;
+		Evt->Param.QualityChange.NewBitrate = NewBitrate;
+		Evt->Param.QualityChange.PrevBitrate = PreviousBitrate;
+		Evt->Param.QualityChange.bIsDrastic = bIsDrasticDownswitch;
+		return Evt;
+	}
+	static TSharedPtrTS<FMetricEvent> ReportPrerollStart()
+	{
+		TSharedPtrTS<FMetricEvent> Evt = MakeSharedTS<FMetricEvent>();
+		Evt->Type = EType::PrerollStart;
+		return Evt;
+	}
+	static TSharedPtrTS<FMetricEvent> ReportPrerollEnd()
+	{
+		TSharedPtrTS<FMetricEvent> Evt = MakeSharedTS<FMetricEvent>();
+		Evt->Type = EType::PrerollEnd;
+		return Evt;
+	}
+	static TSharedPtrTS<FMetricEvent> ReportPlaybackStart()
+	{
+		TSharedPtrTS<FMetricEvent> Evt = MakeSharedTS<FMetricEvent>();
+		Evt->Type = EType::PlaybackStart;
+		return Evt;
+	}
+	static TSharedPtrTS<FMetricEvent> ReportPlaybackPaused()
+	{
+		TSharedPtrTS<FMetricEvent> Evt = MakeSharedTS<FMetricEvent>();
+		Evt->Type = EType::PlaybackPaused;
+		return Evt;
+	}
+	static TSharedPtrTS<FMetricEvent> ReportPlaybackResumed()
+	{
+		TSharedPtrTS<FMetricEvent> Evt = MakeSharedTS<FMetricEvent>();
+		Evt->Type = EType::PlaybackResumed;
+		return Evt;
+	}
+	static TSharedPtrTS<FMetricEvent> ReportPlaybackEnded()
+	{
+		TSharedPtrTS<FMetricEvent> Evt = MakeSharedTS<FMetricEvent>();
+		Evt->Type = EType::PlaybackEnded;
+		return Evt;
+	}
+	static TSharedPtrTS<FMetricEvent> ReportJumpInPlayPosition(const FTimeValue& ToNewTime, const FTimeValue& FromTime, Metrics::ETimeJumpReason TimejumpReason)
+	{
+		TSharedPtrTS<FMetricEvent> Evt = MakeSharedTS<FMetricEvent>();
+		Evt->Type = EType::PlaybackJumped;
+		Evt->Param.TimeJump.ToNewTime = ToNewTime;
+		Evt->Param.TimeJump.FromTime = FromTime;
+		Evt->Param.TimeJump.Reason = TimejumpReason;
+		return Evt;
+	}
+	static TSharedPtrTS<FMetricEvent> ReportPlaybackStopped()
+	{
+		TSharedPtrTS<FMetricEvent> Evt = MakeSharedTS<FMetricEvent>();
+		Evt->Type = EType::PlaybackStopped;
+		return Evt;
+	}
+	static TSharedPtrTS<FMetricEvent> ReportError(const FErrorDetail& ErrorDetail)
+	{
+		TSharedPtrTS<FMetricEvent> Evt = MakeSharedTS<FMetricEvent>();
+		Evt->Type = EType::Errored;
+		Evt->Param.ErrorDetail = ErrorDetail;
+		return Evt;
+	}
+	static TSharedPtrTS<FMetricEvent> ReportLogMessage(IInfoLog::ELevel inLogLevel, const FString& LogMessage, int64 PlayerWallclockMilliseconds)
+	{
+		TSharedPtrTS<FMetricEvent> Evt = MakeSharedTS<FMetricEvent>();
+		Evt->Type = EType::LogMessage;
+		Evt->Param.LogMessage.Level = inLogLevel;
+		Evt->Param.LogMessage.Message = LogMessage;
+		Evt->Param.LogMessage.AtMillis = PlayerWallclockMilliseconds;
+		return Evt;
+	}
+
+};
+
+
+
+class FAdaptiveStreamingPlayerEventHandler
+{
+public:
+	static TSharedPtr<FAdaptiveStreamingPlayerEventHandler, ESPMode::ThreadSafe> Create();
+
+	void DispatchEvent(TSharedPtrTS<FMetricEvent> InEvent);
+	void DispatchEventAndWait(TSharedPtrTS<FMetricEvent> InEvent);
+
+	virtual ~FAdaptiveStreamingPlayerEventHandler();
+private:
+	void StartWorkerThread();
+	void StopWorkerThread();
+	void WorkerThreadFN();
+
+	FMediaThread													WorkerThread;
+	TMediaMessageQueueDynamicNoTimeout<TSharedPtrTS<FMetricEvent>>		EventQueue;
+	static TWeakPtr<FAdaptiveStreamingPlayerEventHandler, ESPMode::ThreadSafe>	SingletonSelf;
+	static FCriticalSection														SingletonLock;
+};
+
+
+
+class FAdaptiveStreamingPlayer : public IAdaptiveStreamingPlayer
+							   , public IPlayerSessionServices
+							   , public IStreamReader::StreamReaderEventListener
+							   , public IAccessUnitMemoryProvider
+							   , public IPlayerStreamFilter
+							   , public TSharedFromThis<FAdaptiveStreamingPlayer, ESPMode::ThreadSafe>
+{
+public:
+	FAdaptiveStreamingPlayer(const IAdaptiveStreamingPlayer::FCreateParam& InCreateParameters);
+	virtual ~FAdaptiveStreamingPlayer();
+
+	virtual void SetStaticResourceProviderCallback(const TSharedPtr<IAdaptiveStreamingPlayerResourceProvider, ESPMode::ThreadSafe>& InStaticResourceProvider) override;
+	virtual void SetVideoDecoderResourceDelegate(const TSharedPtr<IVideoDecoderResourceDelegate, ESPMode::ThreadSafe>& ResourceDelegate) override;
+
+	virtual void AddMetricsReceiver(IAdaptiveStreamingPlayerMetrics* InMetricsReceiver) override;
+	virtual void RemoveMetricsReceiver(IAdaptiveStreamingPlayerMetrics* InMetricsReceiver) override;
+
+	// Matrics receiver accessor for event dispatcher thread.
+	void LockMetricsReceivers()
+	{
+		MetricListenerCriticalSection.Lock();
+	}
+	void UnlockMetricsReceivers()
+	{
+		MetricListenerCriticalSection.Unlock();
+	}
+	const TArray<IAdaptiveStreamingPlayerMetrics*, TInlineAllocator<4>>& GetMetricsReceivers()
+	{
+		return MetricListeners;
+	}
+
+
+	virtual bool Initialize(const FParamDict& Options) override;
+
+	virtual void LoadManifest(const FString& manifestURL) override;
+
+	virtual void SeekTo(const FSeekParam& NewPosition) override;
+	virtual void Pause() override;
+	virtual void Resume() override;
+	virtual void Stop() override;
+	virtual void SetLooping(const FLoopParam& InLoopParams) override;
+
+	virtual FErrorDetail GetError() const override;
+
+	virtual bool HaveMetadata() const override;
+	virtual FTimeValue GetDuration() const override;
+	virtual FTimeValue GetPlayPosition() const override;
+	virtual void GetTimelineRange(FTimeRange& OutRange) const override;
+	virtual void GetSeekableRange(FTimeRange& OutRange) const override;
+	virtual void GetSeekablePositions(TArray<FTimespan>& OutPositions) const override;
+	virtual bool HasEnded() const override;
+	virtual bool IsSeeking() const override;
+	virtual bool IsBuffering() const override;
+	virtual bool IsPlaying() const override;
+	virtual bool IsPaused() const override;
+
+	virtual void GetLoopState(FPlayerLoopState& OutLoopState) const override;
+	virtual void GetStreamMetadata(TArray<FStreamMetadata>& OutStreamMetadata, EStreamType StreamType) const override;
+
+	virtual void SetBitrateCeiling(int32 highestSelectableBitrate) override;
+	virtual void SetMaxResolution(int32 MaxWidth, int32 MaxHeight) override;
+
+	virtual void SelectTrackByMetadata(EStreamType StreamType, const FStreamMetadata& StreamMetadata) override;
+	virtual void DeselectTrack(EStreamType StreamType) override;
+
+#if PLATFORM_ANDROID
+	virtual void Android_UpdateSurface(const TSharedPtr<IOptionPointerValueContainer>& Surface) override;
+#endif
+
+	void DebugPrint(void* pPlayer, void (*debugDrawPrintf)(void* pPlayer, const char *pFmt, ...));
+	static void DebugHandle(void* pPlayer, void (*debugDrawPrintf)(void* pPlayer, const char *pFmt, ...));
+
+private:
+	// Methods from IPlayerSessionServices
+	virtual void PostError(const FErrorDetail& Error) override;
+	virtual void PostLog(Facility::EFacility FromFacility, IInfoLog::ELevel LogLevel, const FString& Message) override;
+	virtual void SendMessageToPlayer(TSharedPtrTS<IPlayerMessage> PlayerMessage) override;
+	virtual void GetExternalGuid(FGuid& OutExternalGuid) override;
+	virtual ISynchronizedUTCTime* GetSynchronizedUTCTime() override;
+	virtual TSharedPtr<IAdaptiveStreamingPlayerResourceProvider, ESPMode::ThreadSafe> GetStaticResourceProvider() override;
+	virtual IElectraHttpManager* GetHTTPManager() override;
+	virtual TSharedPtrTS<IAdaptiveStreamSelector> GetStreamSelector() override;
+	virtual void GetStreamBufferStats(FAccessUnitBufferInfo& OutBufferStats, EStreamType ForStream) override;
+	virtual IPlayerStreamFilter* GetStreamFilter() override;
+
+	// Methods from IPlayerStreamFilter
+	virtual bool CanDecodeStream(const FStreamCodecInformation& InStreamCodecInfo) const override;
+
+
+	enum class EPlayerState
+	{
+		eState_Idle,
+		eState_ParsingManifest,
+		eState_PreparingStreams,
+		eState_Ready,
+		eState_Buffering,					//!< Initial buffering at start or after a seek (an expected buffering)
+		eState_Playing,						//!<
+		eState_Paused,						//!<
+		eState_Rebuffering,					//!< Rebuffering due to buffer underrun. Temporary state only.
+		eState_Seeking,						//!< Seeking. Temporary state only.
+		eState_Error,
+	};
+	static const char* GetPlayerStateName(EPlayerState s)
+	{
+		switch(s)
+		{
+			case EPlayerState::eState_Idle:				return("Idle");
+			case EPlayerState::eState_ParsingManifest:	return("Parsing manifest");
+			case EPlayerState::eState_PreparingStreams:	return("Preparing streams");
+			case EPlayerState::eState_Ready:				return("Ready");
+			case EPlayerState::eState_Buffering:			return("Buffering");
+			case EPlayerState::eState_Playing:			return("Playing");
+			case EPlayerState::eState_Paused:			return("Paused");
+			case EPlayerState::eState_Rebuffering:		return("Rebuffering");
+			case EPlayerState::eState_Seeking:			return("Seeking");
+			case EPlayerState::eState_Error:				return("Error");
+			default:									return("undefined");
+		}
+	}
+
+	enum class EDecoderState
+	{
+		eDecoder_Paused,
+		eDecoder_Running,
+	};
+	static const char* GetDecoderStateName(EDecoderState s)
+	{
+		switch(s)
+		{
+			case EDecoderState::eDecoder_Paused:	return("Paused");
+			case EDecoderState::eDecoder_Running:	return("Running");
+			default:								return("undefined");
+		}
+	}
+
+	enum class EPipelineState
+	{
+		ePipeline_Stopped,
+		ePipeline_Prerolling,
+		ePipeline_Running,
+	};
+	static const char* GetPipelineStateName(EPipelineState s)
+	{
+		switch(s)
+		{
+			case EPipelineState::ePipeline_Stopped:		return("Stopped");
+			case EPipelineState::ePipeline_Prerolling:	return("Prerolling");
+			case EPipelineState::ePipeline_Running:		return("Running");
+			default:									return("undefined");
+		}
+	}
+
+	enum class EStreamState
+	{
+		eStream_Running,
+		eStream_ReachedEnd,
+	};
+	static const char* GetStreamStateName(EStreamState s)
+	{
+		switch(s)
+		{
+			case EStreamState::eStream_Running:		return("Running");
+			case EStreamState::eStream_ReachedEnd:	return("Reached end");
+			default:								return("undefined");
+		}
+	}
+
+	struct FVideoRenderer
+	{
+		void Close()
+		{
+			Renderer.Reset();
+		}
+		void Flush(bool bHoldCurrentFrame)
+		{
+			if (Renderer.IsValid())
+			{
+				FParamDict options;
+				options.Set("hold_current_frame", FVariantValue(bHoldCurrentFrame));
+				Renderer->Flush(options);
+			}
+		}
+		TSharedPtr<IMediaRenderer, ESPMode::ThreadSafe>	Renderer;
+	};
+
+	struct FAudioRenderer
+	{
+		void Close()
+		{
+			Renderer.Reset();
+		}
+		void Flush()
+		{
+			if (Renderer.IsValid())
+			{
+				FParamDict options;
+				Renderer->Flush(options);
+			}
+		}
+		TSharedPtr<IMediaRenderer, ESPMode::ThreadSafe>	Renderer;
+	};
+
+	struct FVideoDecoder : public IAccessUnitBufferListener, public IDecoderOutputBufferListener
+		{
+		FVideoDecoder()
+			: Parent(nullptr)
+			, Decoder(nullptr)
+		{
+		}
+		void Close()
+		{
+			if (Decoder)
+			{
+				Decoder->SetAUInputBufferListener(nullptr);
+				Decoder->SetReadyBufferListener(nullptr);
+			}
+			delete Decoder;
+			Decoder = nullptr;
+		}
+		void Flush()
+		{
+			if (Decoder)
+			{
+				Decoder->AUdataFlushEverything();
+			}
+		}
+		virtual void DecoderInputNeeded(const IAccessUnitBufferListener::FBufferStats& currentInputBufferStats)
+		{
+			if (Parent)
+			{
+				Parent->VideoDecoderInputNeeded(currentInputBufferStats);
+			}
+		}
+
+		virtual void DecoderOutputReady(const IDecoderOutputBufferListener::FDecodeReadyStats& currentReadyStats)
+		{
+			if (Parent)
+			{
+				Parent->VideoDecoderOutputReady(currentReadyStats);
+			}
+		}
+
+		FAdaptiveStreamingPlayer*		Parent;
+		IVideoDecoderH264*				Decoder;
+		};
+
+	struct FAudioDecoder : public IAccessUnitBufferListener, public IDecoderOutputBufferListener
+		{
+		FAudioDecoder()
+			: Parent(nullptr)
+			, Decoder(nullptr)
+		{
+		}
+		void Close()
+		{
+			if (Decoder)
+			{
+				Decoder->SetAUInputBufferListener(nullptr);
+				Decoder->SetReadyBufferListener(nullptr);
+			}
+			delete Decoder;
+			Decoder = nullptr;
+		}
+		void Flush()
+		{
+			if (Decoder)
+			{
+				Decoder->AUdataFlushEverything();
+			}
+		}
+		virtual void DecoderInputNeeded(const IAccessUnitBufferListener::FBufferStats& currentInputBufferStats)
+		{
+			if (Parent)
+			{
+				Parent->AudioDecoderInputNeeded(currentInputBufferStats);
+			}
+		}
+
+		virtual void DecoderOutputReady(const IDecoderOutputBufferListener::FDecodeReadyStats& currentReadyStats)
+		{
+			if (Parent)
+			{
+				Parent->AudioDecoderOutputReady(currentReadyStats);
+			}
+		}
+
+		FAdaptiveStreamingPlayer*		Parent;
+		IAudioDecoderAAC*				Decoder;
+		};
+
+
+	class FWorkerThread
+	{
+	public:
+		FWorkerThread()
+			: bStarted(false)
+		{
+		}
+
+		struct FMessage
+		{
+			enum class EType
+			{
+				// Commands
+				LoadManifest,
+				Quit,
+				Pause,
+				Resume,
+				Seek,
+				Loop,
+				Close,
+				ChangeBitrate,
+				LimitResolution,
+				SelectTrackByMetadata,
+				DeselectTrack,
+				// Player session message
+				PlayerSession,
+				// Fragment reader messages
+				FragmentOpen,
+				FragmentClose,
+				// Decoder buffer messages
+				BufferUnderrun,
+			};
+
+			struct FData
+			{
+				struct FLoadManifest
+				{
+					FString											URL;
+					FString											MimeType;
+				};
+				struct FStreamReader
+				{
+					FStreamReader()
+						: AUType(EStreamType::Video), AUSize(0)
+					{
+					}
+					TSharedPtrTS<IStreamSegment>					Request;
+					EStreamType										AUType;
+					SIZE_T											AUSize;
+				};
+				struct FEvent
+				{
+					FEvent()
+						: Event(nullptr)
+					{
+					}
+					FMediaEvent*									Event;
+				};
+				struct FStartPlay
+				{
+					FSeekParam										Position;
+				};
+				struct FLoop
+				{
+					FLoop()
+						: Signal(nullptr)
+					{
+					}
+					FLoopParam										Loop;
+					FMediaEvent*									Signal;
+				};
+				struct FBitrate
+				{
+					int32											Value;
+				};
+				struct FSession
+				{
+					TSharedPtrTS<IPlayerMessage>					PlayerMessage;
+				};
+				struct FResolution
+				{
+					FResolution() : Width(0), Height(0) { }
+					int32											Width;
+					int32											Height;
+				};
+				struct FMetadataTrackSelection
+				{
+					EStreamType										StreamType;
+					FStreamMetadata									StreamMetadata;
+				};
+
+				FLoadManifest				ManifestToLoad;
+				FStreamReader				StreamReader;
+				FEvent						MediaEvent;
+				FStartPlay					StartPlay;
+				FLoop						Looping;
+				FBitrate					Bitrate;
+				FSession					Session;
+				FResolution					Resolution;
+				FMetadataTrackSelection		TrackSelection;
+			};
+			EType					Type;
+			FData   				Data;
+		};
+
+		void SendMessage(FMessage::EType type)
+		{
+			FMessage Msg;
+			Msg.Type = type;
+			WorkMessages.SendMessage(Msg);
+		}
+		void SendMessage(FMessage::EType type, TSharedPtrTS<IStreamSegment> pRequest)
+		{
+			FMessage Msg;
+			Msg.Type = type;
+			Msg.Data.StreamReader.Request = pRequest;
+			WorkMessages.SendMessage(Msg);
+		}
+		void SendMessage(FMessage::EType type, FMediaEvent* pEventSignal)
+		{
+			FMessage Msg;
+			Msg.Type = type;
+			Msg.Data.MediaEvent.Event = pEventSignal;
+			WorkMessages.SendMessage(Msg);
+		}
+		void SendLoadManifestMessage(const FString& URL, const FString& MimeType)
+		{
+			FMessage Msg;
+			Msg.Type = FMessage::EType::LoadManifest;
+			Msg.Data.ManifestToLoad.URL = URL;
+			Msg.Data.ManifestToLoad.MimeType = MimeType;
+			WorkMessages.SendMessage(Msg);
+		}
+		void SendPauseMessage()
+		{
+			FMessage Msg;
+			Msg.Type = FMessage::EType::Pause;
+			WorkMessages.SendMessage(Msg);
+		}
+		void SendResumeMessage()
+		{
+			FMessage Msg;
+			Msg.Type = FMessage::EType::Resume;
+			WorkMessages.SendMessage(Msg);
+		}
+		void SendSeekMessage(const FSeekParam& NewPosition)
+		{
+			FMessage Msg;
+			Msg.Type = FMessage::EType::Seek;
+			Msg.Data.StartPlay.Position = NewPosition;
+			WorkMessages.SendMessage(Msg);
+		}
+		void SendLoopMessage(const FLoopParam& InLoopParams)
+		{
+			FMessage Msg;
+			Msg.Type = FMessage::EType::Loop;
+			Msg.Data.Looping.Loop = InLoopParams;
+			WorkMessages.SendMessage(Msg);
+		}
+		void SendCloseMessage(FMediaEvent* pEventSignal)
+		{
+			FMessage Msg;
+			Msg.Type = FMessage::EType::Close;
+			Msg.Data.MediaEvent.Event = pEventSignal;
+			WorkMessages.SendMessage(Msg);
+		}
+		void SendBitrateMessage(EStreamType type, int32 value, int32 which)
+		{
+			FMessage Msg;
+			Msg.Type = FMessage::EType::ChangeBitrate;
+			Msg.Data.Bitrate.Value = value;
+			WorkMessages.SendMessage(Msg);
+		}
+		void SendPlayerSessionMessage(const TSharedPtrTS<IPlayerMessage>& message)
+		{
+			FMessage Msg;
+			Msg.Type = FMessage::EType::PlayerSession;
+			Msg.Data.Session.PlayerMessage = message;
+			WorkMessages.SendMessage(Msg);
+		}
+		void SendResolutionMessage(int32 Width, int32 Height)
+		{
+			FMessage Msg;
+			Msg.Type = FMessage::EType::LimitResolution;
+			Msg.Data.Resolution.Width = Width;
+			Msg.Data.Resolution.Height = Height;
+			WorkMessages.SendMessage(Msg);
+		}
+		void SendTrackSelectByMetadataMessage(EStreamType StreamType, const FStreamMetadata& StreamMetadata)
+		{
+			FMessage Msg;
+			Msg.Type = FMessage::EType::SelectTrackByMetadata;
+			Msg.Data.TrackSelection.StreamType = StreamType;
+			Msg.Data.TrackSelection.StreamMetadata = StreamMetadata;
+			WorkMessages.SendMessage(Msg);
+		}
+		void SendTrackDeselectMessage(EStreamType StreamType)
+		{
+			FMessage Msg;
+			Msg.Type = FMessage::EType::DeselectTrack;
+			Msg.Data.TrackSelection.StreamType = StreamType;
+			WorkMessages.SendMessage(Msg);
+		}
+
+		FMediaThread									MediaThread;
+		TMediaMessageQueueDynamicWithTimeout<FMessage>	WorkMessages;
+		bool											bStarted;
+	};
+
+
+	void InternalLoadManifest(const FString& URL, const FString& MimeType);
+	bool SelectManifest();
+	void UpdateManifest();
+
+	void VideoDecoderInputNeeded(const IAccessUnitBufferListener::FBufferStats& currentInputBufferStats);
+	void VideoDecoderOutputReady(const IDecoderOutputBufferListener::FDecodeReadyStats& currentReadyStats);
+
+	void AudioDecoderInputNeeded(const IAccessUnitBufferListener::FBufferStats& currentInputBufferStats);
+	void AudioDecoderOutputReady(const IDecoderOutputBufferListener::FDecodeReadyStats& currentReadyStats);
+
+
+
+	void WorkerThreadFN();
+	void StartWorkerThread();
+	void StopWorkerThread();
+
+	int32 CreateRenderers();
+	void DestroyRenderers();
+
+	int32 CreateInitialDecoder(EStreamType type);
+	void DestroyDecoders();
+	bool FindMatchingStreamInfo(FStreamCodecInformation& OutStreamInfo, int32 MaxWidth, int32 MaxHeight);
+	void UpdateStreamResolutionLimit();
+
+	// AU memory / generic stream reader memory
+	void* AUAllocate(IAccessUnitMemoryProvider::EDataType type, SIZE_T size, SIZE_T alignment) override;
+	void AUDeallocate(IAccessUnitMemoryProvider::EDataType type, void *pAddr) override;
+
+	// Stream reader events
+	void OnFragmentOpen(TSharedPtrTS<IStreamSegment> pRequest) override;
+	bool OnFragmentAccessUnitReceived(FAccessUnit* pAccessUnit) override;
+	void OnFragmentReachedEOS(EStreamType InStreamType, TSharedPtr<const FStreamSourceInfo, ESPMode::ThreadSafe> InStreamSourceInfo) override;
+	void OnFragmentClose(TSharedPtrTS<IStreamSegment> pRequest) override;
+
+
+
+
+
+	struct FPendingStartRequest
+	{
+		FPlayStartPosition				StartAt;
+		IManifest::ESearchType			SearchType;
+		FTimeValue						RetryAtTime;
+		TMediaOptionalValue<int32>		InitialBandwidth;
+	};
+
+
+
+
+	struct FBufferStats
+	{
+		void Clear()
+		{
+			StreamBuffer.Clear();
+			DecoderInputBuffer.Clear();
+			DecoderOutputBuffer.Clear();
+		}
+		FAccessUnitBufferInfo								StreamBuffer;
+		IAccessUnitBufferListener::FBufferStats				DecoderInputBuffer;
+		IDecoderOutputBufferListener::FDecodeReadyStats		DecoderOutputBuffer;
+	};
+	void UpdateDiagnostics();
+
+
+	void HandleNewBufferedData();
+	void HandleNewOutputData();
+	void HandleSessionMessage(TSharedPtrTS<IPlayerMessage> SessionMessage);
+	void HandlePlayStateChanges();
+	void HandlePendingMediaSegmentRequests();
+	void HandleDeselectedBuffers();
+	void HandleDecoderChanges();
+
+	void CheckForStreamEnd();
+
+	void CheckForErrors();
+
+	void FeedDecoder(EStreamType Type, FMultiTrackAccessUnitBuffer& FromMultistreamBuffer, IAccessUnitBufferInterface* Decoder);
+
+	void PauseStreamReaders();
+	void ResumeStreamReaders();
+
+	bool InternalStartAt(const FSeekParam& NewPosition);
+	void InternalPause();
+	void InternalResume();
+	void InternalRebuffer();
+	void InternalStop(bool bHoldCurrentFrame);
+	void InternalClose();
+	void InternalSetLoop(const FLoopParam& LoopParam);
+
+
+	void DispatchEvent(TSharedPtrTS<FMetricEvent> Event);
+	void DispatchEventAndWait(TSharedPtrTS<FMetricEvent> Event);
+	void DispatchBufferingEvent(bool bBegin, EPlayerState Reason);
+	void DispatchSegmentDownloadedEvent(TSharedPtrTS<IStreamSegment> Request);
+	void DispatchBufferUtilizationEvent(EStreamType BufferType);
+
+	void UpdateDataAvailabilityState(Metrics::FDataAvailabilityChange& DataAvailabilityState, Metrics::FDataAvailabilityChange::EAvailability NewAvailability);
+
+	void StartRendering();
+	void StopRendering();
+
+
+	struct FPrerollVars
+	{
+		FPrerollVars()
+		{
+			Clear();
+			bIsVeryFirstStart = true;
+		}
+		void Clear()
+		{
+			StartTime  	  = -1;
+			bHaveEnoughVideo = false;
+			bHaveEnoughAudio = false;
+			bHaveEnoughText  = false;
+		}
+		int64					StartTime;
+		bool					bHaveEnoughVideo;
+		bool					bHaveEnoughAudio;
+		bool					bHaveEnoughText;
+		bool					bIsVeryFirstStart;
+	};
+
+	struct FPostrollVars
+	{
+		struct FRenderState
+		{
+			FRenderState()
+			{
+				Clear();
+			}
+			void Clear()
+			{
+				LastCheckTime   = -1;
+				LastBufferCount = 0;
+			}
+			int64	LastCheckTime;
+			int32	LastBufferCount;
+		};
+
+		FPostrollVars()
+		{
+			Clear();
+		}
+
+		void Clear()
+		{
+			Video.Clear();
+			Audio.Clear();
+		}
+
+		FRenderState		Video;
+		FRenderState		Audio;
+	};
+
+
+	//
+	// Member variables
+	//
+	FGuid																ExternalPlayerGUID;
+
+	TWeakPtr<IAdaptiveStreamingPlayerResourceProvider, ESPMode::ThreadSafe> StaticResourceProvider;
+	TWeakPtr<IVideoDecoderResourceDelegate, ESPMode::ThreadSafe>		VideoDecoderResourceDelegate;
+
+	FWorkerThread															WorkerThread;
+	TSharedPtr<FAdaptiveStreamingPlayerEventHandler, ESPMode::ThreadSafe>	EventDispatcher;
+
+	FParamDict															PlayerOptions;
+	FPlaybackState														PlaybackState;
+	FStreamPreferences													StreamPreferences;
+	ISynchronizedUTCTime*												SynchronizedUTCTime;
+
+	TSharedPtrTS<FMediaRenderClock>										RenderClock;
+
+	TSharedPtrTS<IElectraHttpManager>									HttpManager;
+
+	TMediaQueueDynamic<TSharedPtrTS<FErrorDetail>>						ErrorQueue;
+
+	TSharedPtrTS<IPlaybackAssetTimeline>								CurrentTimeline;
+	TSharedPtrTS<IManifest>												Manifest;
+	IPlaylistReader*													ManifestReader;
+
+	IStreamReader*														StreamReaderHandler;
+	TMediaOptionalValue<bool> 											bHaveVideoReader;
+	TMediaOptionalValue<bool> 											bHaveAudioReader;
+	TMediaOptionalValue<bool> 											bHaveTextReader;
+
+
+	FMultiTrackAccessUnitBuffer											MultiStreamBufferVid;
+	FMultiTrackAccessUnitBuffer											MultiStreamBufferAud;
+	FMultiTrackAccessUnitBuffer											MultiStreamBufferTxt;
+
+	Metrics::FDataAvailabilityChange									DataAvailabilityStateVid;
+	Metrics::FDataAvailabilityChange									DataAvailabilityStateAud;
+	Metrics::FDataAvailabilityChange									DataAvailabilityStateTxt;
+
+	TSharedPtr<FStreamMetadata, ESPMode::ThreadSafe>					InitialStreamSelectionVid;
+	TSharedPtr<FStreamMetadata, ESPMode::ThreadSafe>					InitialStreamSelectionAud;
+	TSharedPtr<FStreamMetadata, ESPMode::ThreadSafe>					InitialStreamSelectionTxt;
+
+	EPlayerState														CurrentState;
+	EPipelineState														PipelineState;
+	EDecoderState														DecoderState;
+	EStreamState														StreamState;
+
+	FPrerollVars														PrerollVars;
+	FPostrollVars														PostrollVars;
+	EPlayerState														LastBufferingState;
+	FSeekParam															StartAtTime;
+	double																PlaybackRate;
+	FTimeValue															RebufferDetectedAtPlayPos;
+	bool																bRebufferPending;
+	bool																bIsPlayStart;
+	bool																bIsClosing;
+
+	TSharedPtrTS<IAdaptiveStreamSelector>								StreamSelector;
+	int32																BitrateCeiling;
+	int32																VideoResolutionLimitWidth;
+	int32																VideoResolutionLimitHeight;
+	int32																MaxVideoTextureResolutionLimitHeight;
+	struct FStreamBitrateInfo
+	{
+		FStreamBitrateInfo()
+		{
+			Clear();
+		}
+		void Clear()
+		{
+			Bitrate 	 = 0;
+			QualityLevel = 0;
+		}
+		int32			Bitrate;
+		int32			QualityLevel;
+	};
+	FStreamBitrateInfo													CurrentVideoStreamBitrate;
+
+
+	bool																bShouldBePaused;
+	bool																bShouldBePlaying;
+	bool																bSeekPending;
+
+	FPlayerLoopState													CurrentLoopState;
+	FLoopParam															CurrentLoopParam;
+	TMediaQueueDynamicNoLock<FPlayerLoopState>							NextLoopStates;
+
+	struct FPendingSegmentRequest
+	{
+		TSharedPtrTS<IStreamSegment>		Request;
+		FTimeValue							AtTime;
+	};
+
+	TSharedPtrTS<IManifest::IPlayPeriod>								CurrentPlayPeriod;
+	TSharedPtrTS<FPendingStartRequest>									PendingStartRequest;
+	TSharedPtrTS<IStreamSegment>										PendingFirstSegmentRequest;
+	TQueue<FPendingSegmentRequest>										NextPendingSegmentRequests;
+	TMultiMap<EStreamType, TSharedPtrTS<IStreamSegment>>				CompletedSegmentRequests;
+
+	uint32																CurrentPlaybackSequenceID;
+
+	FVideoRenderer														VideoRender;
+	FAudioRenderer														AudioRender;
+	FVideoDecoder														VideoDecoder;
+	FAudioDecoder														AudioDecoder;
+
+	FMediaCriticalSection												MetricListenerCriticalSection;
+	TArray<IAdaptiveStreamingPlayerMetrics*, TInlineAllocator<4>>		MetricListeners;
+
+	FMediaCriticalSection												DiagnosticsCriticalSection;
+	FBufferStats														VideoBufferStats;
+	FBufferStats														AudioBufferStats;
+	FBufferStats														TextBufferStats;
+	FErrorDetail														LastErrorDetail;
+
+	AdaptiveStreamingPlayerConfig::FConfiguration						PlayerConfig;
+
+	static FAdaptiveStreamingPlayer* 									PointerToLatestPlayer;
+};
+
+} // namespace Electra
+
diff --git a/Engine/Plugins/Media/ElectraPlayer/Source/ElectraPlayerRuntime/Private/Runtime/Player/AdaptiveStreamingPlayerInternalConfig.h b/Engine/Plugins/Media/ElectraPlayer/Source/ElectraPlayerRuntime/Private/Runtime/Player/AdaptiveStreamingPlayerInternalConfig.h
new file mode 100644
--- /dev/null
+++ b/Engine/Plugins/Media/ElectraPlayer/Source/ElectraPlayerRuntime/Private/Runtime/Player/AdaptiveStreamingPlayerInternalConfig.h
@@ -0,0 +1,78 @@
+// Copyright Epic Games, Inc. All Rights Reserved.
+
+#pragma once
+
+#include "Decoder/VideoDecoderH264.h"
+#include "Decoder/AudioDecoderAAC.h"
+#include "StreamAccessUnitBuffer.h"
+#include "Player/AdaptiveStreamingPlayerABR.h"
+
+
+namespace Electra
+{
+
+
+	namespace AdaptiveStreamingPlayerConfig
+	{
+
+
+		struct FConfiguration
+		{
+			FConfiguration()
+			{
+				// TODO: set certain configuration items for this type of player that make sense (like image size auto config)
+
+				WorkerThread.Priority = TPri_Normal;
+				WorkerThread.StackSize = 32768;
+				WorkerThread.CoreAffinity = -1;
+
+				// Set default values to maximum permitted values.
+				StreamBufferConfigVideo.MaxDataSize = 16 << 20;
+				StreamBufferConfigVideo.MaxDuration.SetFromSeconds(20.0);
+
+				StreamBufferConfigAudio.MaxDataSize = 2 << 20;
+				StreamBufferConfigAudio.MaxDuration.SetFromSeconds(20.0);
+
+				bHoldLastFrameDuringSeek = true;
+
+				MaxPrerollDurationUntilForcedStart = 3.0;
+
+				InitialBufferMinTimeAvailBeforePlayback = 5.0;
+				SeekBufferMinTimeAvailBeforePlayback = 5.0;
+				RebufferMinTimeAvailBeforePlayback = 5.0;
+
+				// And apply the following scaling factors to determine the individual state buffering amounts
+				ScaleMPDInitialBufferMinTimeBeforePlayback = 1.0;
+				ScaleMPDSeekBufferMinTimeAvailBeforePlayback = 1.0;
+				ScaleMPDRebufferMinTimeAvailBeforePlayback = 1.0;
+			}
+
+			FMediaRunnable::Param								WorkerThread;
+
+			FAccessUnitBuffer::FConfiguration					StreamBufferConfigVideo;
+			FAccessUnitBuffer::FConfiguration					StreamBufferConfigAudio;
+			FAccessUnitBuffer::FConfiguration					StreamBufferConfigText;
+
+			IAdaptiveStreamSelector::FConfiguration				StreamSelectorConfig;
+
+			IVideoDecoderH264::FInstanceConfiguration			DecoderCfg264;
+			IAudioDecoderAAC::FInstanceConfiguration			DecoderCfgAAC;
+
+			bool												bHoldLastFrameDuringSeek;
+
+			double												MaxPrerollDurationUntilForcedStart;
+
+			double												InitialBufferMinTimeAvailBeforePlayback;
+			double												SeekBufferMinTimeAvailBeforePlayback;
+			double												RebufferMinTimeAvailBeforePlayback;
+
+			double												ScaleMPDInitialBufferMinTimeBeforePlayback;
+			double												ScaleMPDSeekBufferMinTimeAvailBeforePlayback;
+			double												ScaleMPDRebufferMinTimeAvailBeforePlayback;
+		};
+
+	} // namespace AdaptiveStreamingPlayerConfig
+
+} // namespace Electra
+
+
diff --git a/Engine/Plugins/Media/ElectraPlayer/Source/ElectraPlayerRuntime/Private/Runtime/Player/AdaptiveStreamingPlayerManifest.cpp b/Engine/Plugins/Media/ElectraPlayer/Source/ElectraPlayerRuntime/Private/Runtime/Player/AdaptiveStreamingPlayerManifest.cpp
new file mode 100644
--- /dev/null
+++ b/Engine/Plugins/Media/ElectraPlayer/Source/ElectraPlayerRuntime/Private/Runtime/Player/AdaptiveStreamingPlayerManifest.cpp
@@ -0,0 +1,249 @@
+// Copyright Epic Games, Inc. All Rights Reserved.
+
+#include "PlayerCore.h"
+#include "Misc/Paths.h"
+#include "Player/AdaptiveStreamingPlayerInternal.h"
+#include "Player/HLS/PlaylistReaderHLS.h"
+#include "Player/mp4/PlaylistReaderMP4.h"
+#include "Utilities/Utilities.h"
+#include "Utilities/StringHelpers.h"
+#include "Utilities/URLParser.h"
+#include "ParameterDictionary.h"
+
+
+
+namespace Electra
+{
+
+namespace Playlist
+{
+static const FString MIMETypeMP4(TEXT("video/mp4"));
+static const FString MIMETypeHLS(TEXT("application/vnd.apple.mpegURL"));
+static const FString MIMETypeDASH(TEXT("application/dash+xml"));
+
+
+/**
+ * Returns the MIME type of the media playlist pointed to by the given URL.
+ * This only parses the URL for known extensions or scheme.
+ * If the MIME type cannot be precisely determined an empty string is returned.
+ *
+ * @param URL    Media playlist URL
+ *
+ * @return MIME type or empty string if it cannot be determined
+ */
+FString GetMIMETypeForURL(const FString &URL)
+{
+	FString MimeType;
+	TUniquePtr<IURLParser> UrlParser(IURLParser::Create());
+	if (!UrlParser)
+	{
+		return MimeType;
+	}
+	UEMediaError Error = UrlParser->ParseURL(URL);
+	if (Error != UEMEDIA_ERROR_OK)
+	{
+		return MimeType;
+	}
+
+	//	FString PathOnly = UrlParser->GetPath();
+	TArray<FString> PathComponents;
+	UrlParser->GetPathComponents(PathComponents);
+	if (PathComponents.Num())
+	{
+		FString LowerCaseExtension = FPaths::GetExtension(PathComponents.Last().ToLower());
+
+		// Check for an ".mp4" extension.
+
+		static const FString kTextMP4(TEXT("mp4"));
+		static const FString kTextMPD(TEXT("mpd"));
+		static const FString kTextM3U8(TEXT("m3u8"));
+		if (LowerCaseExtension == kTextMP4)
+		{
+			MimeType = MIMETypeMP4;
+		}
+		else if (LowerCaseExtension == kTextMPD)
+		{
+			MimeType = MIMETypeDASH;
+		}
+		else if (LowerCaseExtension == kTextM3U8)
+		{
+			MimeType = MIMETypeHLS;
+		}
+
+	}
+
+	return MimeType;
+}
+
+
+} // namespace Playlist
+
+
+
+
+
+
+
+
+
+//-----------------------------------------------------------------------------
+/**
+ * Starts asynchronous loading and parsing of a manifest.
+ *
+ * @param URL
+ * @param MimeType
+ */
+void FAdaptiveStreamingPlayer::InternalLoadManifest(const FString& URL, const FString& MimeType)
+{
+	if (CurrentState == EPlayerState::eState_Idle)
+	{
+		FString mimeType = Playlist::GetMIMETypeForURL(URL);
+		if (mimeType.Len())
+		{
+			CurrentState = EPlayerState::eState_ParsingManifest;
+
+			check(ManifestReader == nullptr);
+			if (mimeType == Playlist::MIMETypeHLS)
+			{
+				check(ManifestReader == nullptr);
+				ManifestReader = IPlaylistReaderHLS::Create(this);
+				DispatchEvent(FMetricEvent::ReportOpenSource(URL));
+				ManifestReader->LoadAndParse(URL, StreamPreferences, PlayerOptions);
+			}
+			else if (mimeType == Playlist::MIMETypeMP4)
+			{
+				check(ManifestReader == nullptr);
+				ManifestReader = IPlaylistReaderMP4::Create(this);
+				DispatchEvent(FMetricEvent::ReportOpenSource(URL));
+				ManifestReader->LoadAndParse(URL, StreamPreferences, PlayerOptions);
+			}
+			else
+			{
+				FErrorDetail err;
+				err.SetFacility(Facility::EFacility::Player);
+				err.SetMessage("Unsupported stream MIME type");
+				err.SetCode(INTERR_UNSUPPORTED_FORMAT);
+				PostError(err);
+			}
+		}
+		else
+		{
+			FErrorDetail err;
+			err.SetFacility(Facility::EFacility::Player);
+			err.SetMessage("Could not determine stream MIME type");
+			err.SetCode(INTERR_UNSUPPORTED_FORMAT);
+			PostError(err);
+		}
+	}
+	else
+	{
+// TODO: Error: Not idle
+	}
+}
+
+
+//-----------------------------------------------------------------------------
+/**
+ * Selects the internal presentation for playback after having selected/disabled candidate streams via AccessManifest().
+ *
+ * @return
+ */
+bool FAdaptiveStreamingPlayer::SelectManifest()
+{
+	if (ManifestReader)
+	{
+		check(Manifest == nullptr);
+		if (ManifestReader->GetPlaylistType() == "hls")
+		{
+			TArray<FTimespan> SeekablePositions;
+			TSharedPtrTS<IManifest> pPresentation = ManifestReader->GetManifest();
+			check(pPresentation.IsValid());
+
+			CurrentTimeline = pPresentation->GetTimeline();
+			PlaybackState.SetSeekableRange(CurrentTimeline->GetSeekableTimeRange());
+			CurrentTimeline->GetSeekablePositions(SeekablePositions);
+			PlaybackState.SetSeekablePositions(SeekablePositions);
+			PlaybackState.SetTimelineRange(CurrentTimeline->GetTotalTimeRange());
+			PlaybackState.SetDuration(CurrentTimeline->GetDuration());
+
+			TArray<FStreamMetadata> VideoStreamMetadata;
+			TArray<FStreamMetadata> AudioStreamMetadata;
+			pPresentation->GetStreamMetadata(VideoStreamMetadata, EStreamType::Video);
+			pPresentation->GetStreamMetadata(AudioStreamMetadata, EStreamType::Audio);
+			PlaybackState.SetStreamMetadata(VideoStreamMetadata, AudioStreamMetadata);
+			PlaybackState.SetHaveMetadata(true);
+
+			Manifest = pPresentation;
+
+			CurrentState = EPlayerState::eState_Ready;
+
+			double minBufTimeMPD = Manifest->GetMinBufferTime().GetAsSeconds();
+			PlayerConfig.InitialBufferMinTimeAvailBeforePlayback = Utils::Min(minBufTimeMPD * PlayerConfig.ScaleMPDInitialBufferMinTimeBeforePlayback,   PlayerConfig.InitialBufferMinTimeAvailBeforePlayback);
+			PlayerConfig.SeekBufferMinTimeAvailBeforePlayback    = Utils::Min(minBufTimeMPD * PlayerConfig.ScaleMPDSeekBufferMinTimeAvailBeforePlayback, PlayerConfig.SeekBufferMinTimeAvailBeforePlayback);
+			PlayerConfig.RebufferMinTimeAvailBeforePlayback 	 = Utils::Min(minBufTimeMPD * PlayerConfig.ScaleMPDRebufferMinTimeAvailBeforePlayback,   PlayerConfig.RebufferMinTimeAvailBeforePlayback);
+
+			return true;
+		}
+		else if (ManifestReader->GetPlaylistType() == "mp4")
+		{
+			TArray<FTimespan> SeekablePositions;
+			TSharedPtrTS<IManifest> pPresentation = ManifestReader->GetManifest();
+			check(pPresentation.IsValid());
+
+			CurrentTimeline = pPresentation->GetTimeline();
+			PlaybackState.SetSeekableRange(CurrentTimeline->GetSeekableTimeRange());
+			CurrentTimeline->GetSeekablePositions(SeekablePositions);
+			PlaybackState.SetSeekablePositions(SeekablePositions);
+			PlaybackState.SetTimelineRange(CurrentTimeline->GetTotalTimeRange());
+			PlaybackState.SetDuration(CurrentTimeline->GetDuration());
+
+			TArray<FStreamMetadata> VideoStreamMetadata;
+			TArray<FStreamMetadata> AudioStreamMetadata;
+			pPresentation->GetStreamMetadata(VideoStreamMetadata, EStreamType::Video);
+			pPresentation->GetStreamMetadata(AudioStreamMetadata, EStreamType::Audio);
+			PlaybackState.SetStreamMetadata(VideoStreamMetadata, AudioStreamMetadata);
+			PlaybackState.SetHaveMetadata(true);
+
+			Manifest = pPresentation;
+
+			CurrentState = EPlayerState::eState_Ready;
+
+			double minBufTimeMPD = Manifest->GetMinBufferTime().GetAsSeconds();
+			PlayerConfig.InitialBufferMinTimeAvailBeforePlayback = Utils::Min(minBufTimeMPD * PlayerConfig.ScaleMPDInitialBufferMinTimeBeforePlayback,   PlayerConfig.InitialBufferMinTimeAvailBeforePlayback);
+			PlayerConfig.SeekBufferMinTimeAvailBeforePlayback    = Utils::Min(minBufTimeMPD * PlayerConfig.ScaleMPDSeekBufferMinTimeAvailBeforePlayback, PlayerConfig.SeekBufferMinTimeAvailBeforePlayback);
+			PlayerConfig.RebufferMinTimeAvailBeforePlayback 	 = Utils::Min(minBufTimeMPD * PlayerConfig.ScaleMPDRebufferMinTimeAvailBeforePlayback,   PlayerConfig.RebufferMinTimeAvailBeforePlayback);
+
+			return true;
+		}
+		// Handle other types of playlist here.
+		check(!"TODO");
+	}
+
+	return false;
+}
+
+void FAdaptiveStreamingPlayer::UpdateManifest()
+{
+	check(Manifest != nullptr);
+
+	TArray<FTimespan> SeekablePositions;
+
+	// Update the current timeline.
+	CurrentTimeline = Manifest->GetTimeline();
+	PlaybackState.SetSeekableRange(CurrentTimeline->GetSeekableTimeRange());
+	CurrentTimeline->GetSeekablePositions(SeekablePositions);
+	PlaybackState.SetSeekablePositions(SeekablePositions);
+	PlaybackState.SetTimelineRange(CurrentTimeline->GetTotalTimeRange());
+	PlaybackState.SetDuration(CurrentTimeline->GetDuration());
+
+	TArray<FStreamMetadata> VideoStreamMetadata;
+	TArray<FStreamMetadata> AudioStreamMetadata;
+	Manifest->GetStreamMetadata(VideoStreamMetadata, EStreamType::Video);
+	Manifest->GetStreamMetadata(AudioStreamMetadata, EStreamType::Audio);
+	PlaybackState.SetStreamMetadata(VideoStreamMetadata, AudioStreamMetadata);
+	PlaybackState.SetHaveMetadata(true);
+}
+
+} // namespace Electra
+
+
diff --git a/Engine/Plugins/Media/ElectraPlayer/Source/ElectraPlayerRuntime/Private/Runtime/Player/AdaptiveStreamingPlayerMetrics.h b/Engine/Plugins/Media/ElectraPlayer/Source/ElectraPlayerRuntime/Private/Runtime/Player/AdaptiveStreamingPlayerMetrics.h
new file mode 100644
--- /dev/null
+++ b/Engine/Plugins/Media/ElectraPlayer/Source/ElectraPlayerRuntime/Private/Runtime/Player/AdaptiveStreamingPlayerMetrics.h
@@ -0,0 +1,361 @@
+// Copyright Epic Games, Inc. All Rights Reserved.
+
+#pragma once
+
+#include "PlayerCore.h"
+#include "StreamTypes.h"
+#include "InfoLog.h"
+#include "Player/AdaptiveStreamingPlayerABR_State.h"
+#include "Player/Playlist.h"
+
+namespace Electra
+{
+
+namespace Metrics
+{
+	enum class ESegmentType
+	{
+		Init,
+		Media
+	};
+	static const TCHAR * const GetSegmentTypeString(ESegmentType SegmentType)
+	{
+		switch(SegmentType)
+		{
+			case ESegmentType::Init:
+				return TEXT("Init");
+			case ESegmentType::Media:
+				return TEXT("Media");
+		}
+		return TEXT("n/a");
+	}
+
+	enum class EBufferingReason
+	{
+		Initial,
+		Seeking,
+		Rebuffering,
+	};
+	static const TCHAR * const GetBufferingReasonString(EBufferingReason BufferingReason)
+	{
+		switch(BufferingReason)
+		{
+			case EBufferingReason::Initial:
+				return TEXT("Initial");
+			case EBufferingReason::Seeking:
+				return TEXT("Seeking");
+			case EBufferingReason::Rebuffering:
+				return TEXT("Rebuffering");
+		}
+		return TEXT("n/a");
+	}
+
+	enum class ETimeJumpReason
+	{
+		UserSeek,
+		FellOffTimeline,
+		FellBehindWallclock,
+		Looping
+	};
+	static const TCHAR * const GetTimejumpReasonString(ETimeJumpReason TimejumpReason)
+	{
+		switch(TimejumpReason)
+		{
+			case ETimeJumpReason::UserSeek:
+				return TEXT("User seek");
+			case ETimeJumpReason::FellOffTimeline:
+				return TEXT("Fell off timeline");
+			case ETimeJumpReason::FellBehindWallclock:
+				return TEXT("Fell behind wallclock");
+			case ETimeJumpReason::Looping:
+				return TEXT("Looping");
+		}
+		return TEXT("n/a");
+	}
+
+	struct FBufferStats
+	{
+		FBufferStats()
+		{
+			BufferType  		 = EStreamType::Video;
+			MaxDurationInSeconds = 0.0;
+			DurationInUse   	 = 0.0;
+			MaxByteCapacity 	 = 0;
+			BytesInUse  		 = 0;
+		}
+		EStreamType		BufferType;
+		double			MaxDurationInSeconds;
+		double			DurationInUse;
+		int64			MaxByteCapacity;
+		int64			BytesInUse;
+	};
+
+	struct FABRState
+	{
+		FABRState()
+		{
+			Reset();
+		}
+
+		void Reset()
+		{
+			ProgressDecision.Reset();
+		}
+		FABRDownloadProgressDecision		ProgressDecision;
+	};
+
+
+	struct FPlaylistDownloadStats
+	{
+		FPlaylistDownloadStats()
+		{
+			ListType	   = Playlist::EListType::Master;
+			LoadType	   = Playlist::ELoadType::Initial;
+			HTTPStatusCode = 0;
+			RetryNumber    = 0;
+			bWasSuccessful = false;
+		}
+
+		FString					URL;
+		//FString				CDN;							// do not have this yet
+		FString					FailureReason;					//!< Human readable failure reason. Only for display purposes.
+		Playlist::EListType		ListType;
+		Playlist::ELoadType		LoadType;
+		int32					HTTPStatusCode;					//!< HTTP status code (0 if not connected to server yet)
+		int32					RetryNumber;
+		bool					bWasSuccessful;
+	};
+
+	struct FSegmentDownloadStats
+	{
+		FSegmentDownloadStats()
+		{
+			StreamType  		 = EStreamType::Video;
+			SegmentType 		 = ESegmentType::Media;
+			PresentationTime	 = 0.0;
+			Bitrate 			 = 0;
+			Duration			 = 0.0;
+			RetryNumber 		 = 0;
+			bIsMissingSegment    = false;
+			StatsID 			 = 0;
+			DurationDownloaded   = 0.0;
+			DurationDelivered    = 0.0;
+			TimeToFirstByte 	 = 0.0;
+			TimeToDownload  	 = 0.0;
+			ByteSize			 = 0;
+			NumBytesDownloaded   = 0;
+			ThroughputBps   	 = 0;
+			HTTPStatusCode  	 = 0;
+			bWasSuccessful  	 = false;
+			bWasAborted 		 = false;
+			bDidTimeout 		 = false;
+			bParseFailure   	 = false;
+			bInsertedFillerData  = false;
+			ABRState.Reset();
+		}
+		// Inputs from stream request
+		EStreamType		StreamType;						//!< Type of stream
+		ESegmentType	SegmentType;					//!< Type of segment (init or media)
+		FString			URL;							//!< Effective URL used to download from
+		FString			CDN;							//!< CDN
+		FString			MediaAssetID;
+		FString			AdaptationSetID;
+		FString			RepresentationID;
+		double			PresentationTime;				//!< Presentation time on media timeline
+		double			Duration;						//!< Duration of segment as specified in manifest
+		int32			Bitrate;						//!< Stream bitrate as specified in manifest
+		int32			RetryNumber;
+		bool			bIsMissingSegment;				//!< true if the segment was not actually downloaded because it is missing on the timeline.
+
+		// Outputs from stream reader
+		uint32			StatsID;						//!< ID uniquely identifying this download
+		FString			FailureReason;					//!< Human readable failure reason. Only for display purposes.
+		double			DurationDownloaded;				//!< Duration of content successfully downloaded. May be less than Duration in case of errors.
+		double			DurationDelivered;				//!< Duration of content delivered to buffer. If larger than DurationDownloaded indicates dummy data was inserted into buffer.
+		double			TimeToFirstByte;				//!< Time in seconds until first data byte was received
+		double			TimeToDownload;					//!< Total time in seconds for entire download
+		int64			ByteSize;						//!< Content-Length, may be -1 if unknown (either on error or chunked transfer)
+		int64			NumBytesDownloaded;				//!< Number of bytes successfully downloaded.
+		int64			ThroughputBps;					//!< Estimated throughput in bits per second
+		int32			HTTPStatusCode;					//!< HTTP status code (0 if not connected to server yet)
+		bool			bWasSuccessful;					//!< true if download was successful, false if not
+		bool			bWasAborted;					//!< true if download was aborted by ABR (not by playback!)
+		bool			bDidTimeout;					//!< true if a timeout occurred. Only set if timeouts are enabled. Usually the ABR will monitor and abort.
+		bool			bParseFailure;					//!< true if the segment could not be parsed
+		bool			bInsertedFillerData;
+
+		// ABR state
+		FABRState		ABRState;
+	};
+
+	struct FLicenseKeyStats
+	{
+		FLicenseKeyStats()
+		{
+			bWasSuccessful = false;
+		}
+		FString			URL;
+		FString			FailureReason;					//!< Human readable failure reason. Only for display purposes.
+		bool			bWasSuccessful;
+	};
+
+	struct FDataAvailabilityChange
+	{
+		enum class EAvailability
+		{
+			DataAvailable,
+			DataNotAvailable
+		};
+		FDataAvailabilityChange()
+			: StreamType(EStreamType::Unsupported), Availability(EAvailability::DataNotAvailable)
+		{ }
+
+		EStreamType		StreamType;						//!< Type of stream
+		EAvailability	Availability;
+	};
+
+} // namespace Metrics
+
+
+/**
+ *
+**/
+class IAdaptiveStreamingPlayerMetrics
+{
+public:
+	virtual ~IAdaptiveStreamingPlayerMetrics() = default;
+
+	//=================================================================================================================
+	// Methods called from the media player.
+	//
+
+	/**
+	 * Called when the source will be opened.
+	 */
+	virtual void ReportOpenSource(const FString& URL) = 0;
+
+	/**
+	 * Called when the source's master playlist has beed loaded successfully.
+	 */
+	virtual void ReportReceivedMasterPlaylist(const FString& EffectiveURL) = 0;
+
+	/**
+	 * Called when the dependent child playlists have been loaded successfully.
+	 */
+	virtual void ReportReceivedPlaylists() = 0;
+
+	/**
+	 * Called at the end of every downloaded playlist.
+	 */
+	virtual void ReportPlaylistDownload(const Metrics::FPlaylistDownloadStats& PlaylistDownloadStats) = 0;
+
+	/**
+	 * Called when buffering of data begins.
+	 */
+	virtual void ReportBufferingStart(Metrics::EBufferingReason BufferingReason) = 0;
+
+	/**
+	 * Called when buffering of data ends.
+	 * This does not necessarily coincide with a segment download as buffering ends as soon as
+	 * sufficient data has been received to start/resume playback.
+	 */
+	virtual void ReportBufferingEnd(Metrics::EBufferingReason BufferingReason) = 0;
+
+	/**
+	 * Called at the end of each downloaded video segment.
+	 * The order will be ReportSegmentDownload() followed by ReportBandwidth()
+	 */
+	virtual void ReportBandwidth(int64 EffectiveBps, int64 ThroughputBps, double LatencyInSeconds) = 0;
+
+	/**
+	 * Called before and after a segment download.
+	 */
+	virtual void ReportBufferUtilization(const Metrics::FBufferStats& BufferStats) = 0;
+
+	/**
+	 * Called at the end of each downloaded segment.
+	 */
+	virtual void ReportSegmentDownload(const Metrics::FSegmentDownloadStats& SegmentDownloadStats) = 0;
+
+	/**
+	 * Called for license key events.
+	 */
+	virtual void ReportLicenseKey(const Metrics::FLicenseKeyStats& LicenseKeyStats) = 0;
+
+	/**
+	 * Called when a new video stream segment is fetched at a different bitrate than before.
+	 * A drastic change is one where quality _drops_ more than _one_ level.
+	 */
+	virtual void ReportVideoQualityChange(int32 NewBitrate, int32 PreviousBitrate, bool bIsDrasticDownswitch) = 0;
+
+	/**
+	 * Called when stream data availability changes when feeding the decoder.
+	 */
+	virtual void ReportDataAvailabilityChange(const Metrics::FDataAvailabilityChange& DataAvailability) = 0;
+
+	/**
+	 * Called when decoders start to decode first data to pre-roll the pipeline.
+	 */
+	virtual void ReportPrerollStart() = 0;
+
+	/**
+	 * Called when enough initial data has been decoded to pre-roll the pipeline.
+	 */
+	virtual void ReportPrerollEnd() = 0;
+
+	/**
+	 * Called when playback starts for the first time. This is not called after resuming a paused playback.
+	 * If playback has ended(!) and is then begun again by seeking back to an earlier point in time this
+	 * callback will be triggered again. Seeking during playback will not cause this callback.
+	 */
+	virtual void ReportPlaybackStart() = 0;
+
+	/**
+	 * Called when the player enters pause mode, either through user request or an internal state change.
+	 */
+	virtual void ReportPlaybackPaused() = 0;
+
+	/**
+	 * Called when the player resumes from pause mode, either through user request or an internal state change.
+	 */
+	virtual void ReportPlaybackResumed() = 0;
+
+	/**
+	 * Called when playback has reached the end. See ReportPlaybackStart().
+	 */
+	virtual void ReportPlaybackEnded() = 0;
+
+	/**
+	 * Called when the play position jumps either because of a user induced seek or because the play position fell off the timeline
+	 * or because a wallclock synchronized Live playback fell too far behind the actual time.
+	 */
+	virtual void ReportJumpInPlayPosition(const FTimeValue& ToNewTime, const FTimeValue& FromTime, Metrics::ETimeJumpReason TimejumpReason) = 0;
+
+	/**
+	 * Called when playback is terminally stopped.
+	 */
+	virtual void ReportPlaybackStopped() = 0;
+
+	/**
+	 * Called when an error occurs. Errors always result in termination of playback.
+	 */
+	virtual void ReportError(const FString& ErrorReason) = 0;
+
+	/**
+	 * Called to output a log message. The log level is a player internal level.
+	 */
+	virtual void ReportLogMessage(IInfoLog::ELevel LogLevel, const FString& LogMessage, int64 PlayerWallclockMilliseconds) = 0;
+
+
+	//=================================================================================================================
+	// Methods called from the renderers.
+	//
+
+	virtual void ReportDroppedVideoFrame() = 0;
+	virtual void ReportDroppedAudioFrame() = 0;
+};
+
+
+} // namespace Electra
+
+
diff --git a/Engine/Plugins/Media/ElectraPlayer/Source/ElectraPlayerRuntime/Private/Runtime/Player/HLS/EpicTags.h b/Engine/Plugins/Media/ElectraPlayer/Source/ElectraPlayerRuntime/Private/Runtime/Player/HLS/EpicTags.h
new file mode 100644
--- /dev/null
+++ b/Engine/Plugins/Media/ElectraPlayer/Source/ElectraPlayerRuntime/Private/Runtime/Player/HLS/EpicTags.h
@@ -0,0 +1,25 @@
+// Copyright Epic Games, Inc. All Rights Reserved.
+
+#pragma once
+
+#include "MediaTag.h"
+#include "CoreMinimal.h"
+
+namespace Electra
+{
+	namespace HLSPlaylistParser
+	{
+		/**
+		 * Epic namespace provides Tags specific to the epicgames streaming backend.
+		 */
+		namespace Epic
+		{
+			const FString ExtXEpicGamesCustom = "EXT-X-EPICGAMES-CUSTOM";
+
+			const TMap<FString, HLSPlaylistParser::FMediaTag> TagMap = {
+					{ExtXEpicGamesCustom, HLSPlaylistParser::FMediaTag(ExtXEpicGamesCustom, HLSPlaylistParser::TargetPlaylistURL |
+																					   HLSPlaylistParser::HasAttributeList)},
+			};
+		}
+	}
+} // namespace Electra
\ No newline at end of file
diff --git a/Engine/Plugins/Media/ElectraPlayer/Source/ElectraPlayerRuntime/Private/Runtime/Player/HLS/InitSegmentCacheHLS.cpp b/Engine/Plugins/Media/ElectraPlayer/Source/ElectraPlayerRuntime/Private/Runtime/Player/HLS/InitSegmentCacheHLS.cpp
new file mode 100644
--- /dev/null
+++ b/Engine/Plugins/Media/ElectraPlayer/Source/ElectraPlayerRuntime/Private/Runtime/Player/HLS/InitSegmentCacheHLS.cpp
@@ -0,0 +1,90 @@
+// Copyright Epic Games, Inc. All Rights Reserved.
+
+#include "InitSegmentCacheHLS.h"
+#include "Player/PlayerSessionServices.h"
+#include "Demuxer/ParserISO14496-12.h"
+#include "Containers/LruCache.h"
+
+namespace Electra
+{
+
+
+class FInitSegmentCacheHLS : public IInitSegmentCacheHLS
+{
+public:
+	FInitSegmentCacheHLS(IPlayerSessionServices* SessionServices, const FParamDict& Options);
+	virtual ~FInitSegmentCacheHLS();
+	virtual TSharedPtrTS<const IParserISO14496_12> GetInitSegmentFor(const TSharedPtrTS<const FManifestHLSInternal::FMediaStream::FInitSegmentInfo>& InitSegmentInfo) override;
+	virtual void AddInitSegment(const TSharedPtrTS<const IParserISO14496_12>& InitSegment, const TSharedPtrTS<const FManifestHLSInternal::FMediaStream::FInitSegmentInfo>& InitSegmentInfo, const FTimeValue& ExpiresAtUTC) override;
+private:
+
+	static const uint32 kMaxCacheEntries = 32;
+
+	FParamDict				Options;
+	IPlayerSessionServices* SessionServices;
+
+	struct FEntry
+	{
+		FManifestHLSInternal::FMediaStream::FInitSegmentInfo	SegmentInfo;
+		TSharedPtrTS<const IParserISO14496_12>					ParsedData;
+		FTimeValue												ExpiresAt;
+	};
+
+	FCriticalSection						Lock;
+	TLruCache<FString, TSharedPtr<FEntry>> 	Cache;
+};
+
+/***************************************************************************************************************************************************/
+/***************************************************************************************************************************************************/
+/***************************************************************************************************************************************************/
+
+IInitSegmentCacheHLS* IInitSegmentCacheHLS::Create(IPlayerSessionServices* SessionServices, const FParamDict& Options)
+{
+	return new FInitSegmentCacheHLS(SessionServices, Options);
+}
+
+/***************************************************************************************************************************************************/
+/***************************************************************************************************************************************************/
+/***************************************************************************************************************************************************/
+
+FInitSegmentCacheHLS::FInitSegmentCacheHLS(IPlayerSessionServices* InSessionServices, const FParamDict& InOptions)
+	: Options(InOptions)
+	, SessionServices(InSessionServices)
+	, Cache(kMaxCacheEntries)
+{
+}
+
+FInitSegmentCacheHLS::~FInitSegmentCacheHLS()
+{
+	FScopeLock ScopeLock(&Lock);
+	Cache.Empty();
+}
+
+TSharedPtrTS<const IParserISO14496_12> FInitSegmentCacheHLS::GetInitSegmentFor(const TSharedPtrTS<const FManifestHLSInternal::FMediaStream::FInitSegmentInfo>& InitSegmentInfo)
+{
+	FScopeLock ScopeLock(&Lock);
+	const TSharedPtr<FEntry>* CacheEntry = Cache.FindAndTouch(InitSegmentInfo->URI);
+	if (CacheEntry)
+	{
+		return (*CacheEntry)->ParsedData;
+	}
+	else
+	{
+		return TSharedPtrTS<const IParserISO14496_12>();
+	}
+
+}
+
+void FInitSegmentCacheHLS::AddInitSegment(const TSharedPtrTS<const IParserISO14496_12>& InitSegment, const TSharedPtrTS<const FManifestHLSInternal::FMediaStream::FInitSegmentInfo>& InitSegmentInfo, const FTimeValue& ExpiresAtUTC)
+{
+	TSharedPtr<FEntry> Entry = MakeShared<FEntry>();
+	Entry->ParsedData  = InitSegment;
+	Entry->SegmentInfo = *InitSegmentInfo;
+	Entry->ExpiresAt   = ExpiresAtUTC;
+	FScopeLock ScopeLock(&Lock);
+	Cache.Add(InitSegmentInfo->URI, Entry);
+}
+
+
+} // namespace Electra
+
diff --git a/Engine/Plugins/Media/ElectraPlayer/Source/ElectraPlayerRuntime/Private/Runtime/Player/HLS/InitSegmentCacheHLS.h b/Engine/Plugins/Media/ElectraPlayer/Source/ElectraPlayerRuntime/Private/Runtime/Player/HLS/InitSegmentCacheHLS.h
new file mode 100644
--- /dev/null
+++ b/Engine/Plugins/Media/ElectraPlayer/Source/ElectraPlayerRuntime/Private/Runtime/Player/HLS/InitSegmentCacheHLS.h
@@ -0,0 +1,30 @@
+// Copyright Epic Games, Inc. All Rights Reserved.
+
+#pragma once
+
+#include "PlayerCore.h"
+#include "Player/Manifest.h"
+#include "ManifestBuilderHLS.h"
+
+namespace Electra
+{
+class IPlayerSessionServices;
+class IParserISO14496_12;
+
+
+
+class IInitSegmentCacheHLS
+{
+public:
+	static IInitSegmentCacheHLS* Create(IPlayerSessionServices* SessionServices, const FParamDict& Options);
+
+	virtual ~IInitSegmentCacheHLS() = default;
+
+	virtual TSharedPtrTS<const IParserISO14496_12> GetInitSegmentFor(const TSharedPtrTS<const FManifestHLSInternal::FMediaStream::FInitSegmentInfo>& InitSegmentInfo) = 0;
+
+	virtual void AddInitSegment(const TSharedPtrTS<const IParserISO14496_12>& InitSegment, const TSharedPtrTS<const FManifestHLSInternal::FMediaStream::FInitSegmentInfo>& InitSegmentInfo, const FTimeValue& ExpiresAtUTC) = 0;
+};
+
+
+} // namespace Electra
+
diff --git a/Engine/Plugins/Media/ElectraPlayer/Source/ElectraPlayerRuntime/Private/Runtime/Player/HLS/LHLSTags.h b/Engine/Plugins/Media/ElectraPlayer/Source/ElectraPlayerRuntime/Private/Runtime/Player/HLS/LHLSTags.h
new file mode 100644
--- /dev/null
+++ b/Engine/Plugins/Media/ElectraPlayer/Source/ElectraPlayerRuntime/Private/Runtime/Player/HLS/LHLSTags.h
@@ -0,0 +1,41 @@
+// Copyright Epic Games, Inc. All Rights Reserved.
+
+#pragma once
+
+#include "MediaTag.h"
+#include "CoreMinimal.h"
+
+namespace Electra
+{
+    namespace HLSPlaylistParser
+    {
+        /**
+         * LHLS namespace provides tags for the low-latency HLS (LHLS) by the open LHLS spec provided by hls.js team.
+         * See: [REDACTED_PROFILE_URL]/hlsjs-rfcs/blob/lhls-spec/proposals/0001-lhls.md
+         */
+        namespace LHLS
+        {
+			/**
+			 * Identifies prefetch segments.
+			 *
+			 * A prefetchable segment will appear in the segment list of a playlist and can be distinguished as such
+			 * by checking for the tag using `Segment.ContainsTag(ExtXPrefetch)`.
+			 */
+            const FString ExtXPrefetch = "EXT-X-PREFETCH";
+            const FString ExtXPrefetchDiscontinuity = "EXT-X-PREFETCH-DISCONTINUITY";
+
+            /**
+             * Provides the tags and their proper configuration, can be used by calling:
+             *  Parser.Configure(HLSPlaylistParser::LHLS::TagMap);
+             */
+            const TMap<FString, HLSPlaylistParser::FMediaTag> TagMap = {
+                    {ExtXPrefetch,              HLSPlaylistParser::FMediaTag(ExtXPrefetch,
+                                                                             HLSPlaylistParser::TargetSegmentURL |
+                                                                             HLSPlaylistParser::HasDirectValue | HLSPlaylistParser::URLDefinition)},
+                    {ExtXPrefetchDiscontinuity, HLSPlaylistParser::FMediaTag(ExtXPrefetchDiscontinuity,
+                                                                             HLSPlaylistParser::TargetSegmentURL)}
+            };
+        }
+    }
+} // namespace Electra
+
diff --git a/Engine/Plugins/Media/ElectraPlayer/Source/ElectraPlayerRuntime/Private/Runtime/Player/HLS/LicenseKeyCacheHLS.cpp b/Engine/Plugins/Media/ElectraPlayer/Source/ElectraPlayerRuntime/Private/Runtime/Player/HLS/LicenseKeyCacheHLS.cpp
new file mode 100644
--- /dev/null
+++ b/Engine/Plugins/Media/ElectraPlayer/Source/ElectraPlayerRuntime/Private/Runtime/Player/HLS/LicenseKeyCacheHLS.cpp
@@ -0,0 +1,89 @@
+// Copyright Epic Games, Inc. All Rights Reserved.
+
+#include "LicenseKeyCacheHLS.h"
+#include "Player/PlayerSessionServices.h"
+#include "Containers/LruCache.h"
+
+namespace Electra
+{
+
+
+class FLicenseKeyCacheHLS : public ILicenseKeyCacheHLS
+{
+public:
+	FLicenseKeyCacheHLS(IPlayerSessionServices* SessionServices, const FParamDict& Options);
+	virtual ~FLicenseKeyCacheHLS();
+
+	virtual TSharedPtr<TArray<uint8>, ESPMode::ThreadSafe> GetLicenseKeyFor(const TSharedPtr<const FManifestHLSInternal::FMediaStream::FDRMKeyInfo, ESPMode::ThreadSafe>& LicenseKeyInfo) override;
+	virtual void AddLicenseKey(const TSharedPtr<TArray<uint8>, ESPMode::ThreadSafe>& LicenseKey, const TSharedPtr<const FManifestHLSInternal::FMediaStream::FDRMKeyInfo, ESPMode::ThreadSafe>& LicenseKeyInfo, const FTimeValue& ExpiresAtUTC) override;
+
+private:
+
+	static const uint32 kMaxCacheEntries = 32;
+
+	FParamDict				Options;
+	IPlayerSessionServices* SessionServices;
+
+	struct FEntry
+	{
+		FManifestHLSInternal::FMediaStream::FDRMKeyInfo		LicenseKeyInfo;
+		TSharedPtr<TArray<uint8>, ESPMode::ThreadSafe>		LicenseKeyData;
+		FTimeValue											ExpiresAt;
+	};
+
+	FCriticalSection					   	Lock;
+	TLruCache<FString, TSharedPtr<FEntry>> 	Cache;
+};
+
+/***************************************************************************************************************************************************/
+/***************************************************************************************************************************************************/
+/***************************************************************************************************************************************************/
+
+ILicenseKeyCacheHLS* ILicenseKeyCacheHLS::Create(IPlayerSessionServices* SessionServices, const FParamDict& Options)
+{
+	return new FLicenseKeyCacheHLS(SessionServices, Options);
+}
+
+/***************************************************************************************************************************************************/
+/***************************************************************************************************************************************************/
+/***************************************************************************************************************************************************/
+
+FLicenseKeyCacheHLS::FLicenseKeyCacheHLS(IPlayerSessionServices* InSessionServices, const FParamDict& InOptions)
+	: Options(InOptions)
+	, SessionServices(InSessionServices)
+{
+}
+
+FLicenseKeyCacheHLS::~FLicenseKeyCacheHLS()
+{
+	FScopeLock ScopeLock(&Lock);
+	Cache.Empty();
+}
+
+TSharedPtr<TArray<uint8>, ESPMode::ThreadSafe> FLicenseKeyCacheHLS::GetLicenseKeyFor(const TSharedPtr<const FManifestHLSInternal::FMediaStream::FDRMKeyInfo, ESPMode::ThreadSafe>& LicenseKeyInfo)
+{
+	FScopeLock ScopeLock(&Lock);
+	const TSharedPtr<FEntry>* CacheEntry = Cache.FindAndTouch(LicenseKeyInfo->URI);
+	if (CacheEntry)
+	{
+		return (*CacheEntry)->LicenseKeyData;
+	}
+	else
+	{
+		return nullptr;
+	}
+}
+
+void FLicenseKeyCacheHLS::AddLicenseKey(const TSharedPtr<TArray<uint8>, ESPMode::ThreadSafe>& LicenseKey, const TSharedPtr<const FManifestHLSInternal::FMediaStream::FDRMKeyInfo, ESPMode::ThreadSafe>& LicenseKeyInfo, const FTimeValue& ExpiresAtUTC)
+{
+	TSharedPtr<FEntry> Entry = MakeShared<FEntry>();
+	Entry->LicenseKeyData  = LicenseKey;
+	Entry->LicenseKeyInfo  = *LicenseKeyInfo;
+	Entry->ExpiresAt	   = ExpiresAtUTC;
+	FScopeLock ScopeLock(&Lock);
+	Cache.Add(LicenseKeyInfo->URI, Entry);
+}
+
+
+} // namespace Electra
+
diff --git a/Engine/Plugins/Media/ElectraPlayer/Source/ElectraPlayerRuntime/Private/Runtime/Player/HLS/LicenseKeyCacheHLS.h b/Engine/Plugins/Media/ElectraPlayer/Source/ElectraPlayerRuntime/Private/Runtime/Player/HLS/LicenseKeyCacheHLS.h
new file mode 100644
--- /dev/null
+++ b/Engine/Plugins/Media/ElectraPlayer/Source/ElectraPlayerRuntime/Private/Runtime/Player/HLS/LicenseKeyCacheHLS.h
@@ -0,0 +1,28 @@
+// Copyright Epic Games, Inc. All Rights Reserved.
+
+#pragma once
+
+#include "PlayerCore.h"
+#include "Player/Manifest.h"
+#include "ManifestBuilderHLS.h"
+
+namespace Electra
+{
+class IPlayerSessionServices;
+
+
+class ILicenseKeyCacheHLS
+{
+public:
+	static ILicenseKeyCacheHLS* Create(IPlayerSessionServices* SessionServices, const FParamDict& Options);
+
+	virtual ~ILicenseKeyCacheHLS() = default;
+
+	virtual TSharedPtr<TArray<uint8>, ESPMode::ThreadSafe> GetLicenseKeyFor(const TSharedPtr<const FManifestHLSInternal::FMediaStream::FDRMKeyInfo, ESPMode::ThreadSafe>& LicenseKeyInfo) = 0;
+
+	virtual void AddLicenseKey(const TSharedPtr<TArray<uint8>, ESPMode::ThreadSafe>& LicenseKey, const TSharedPtr<const FManifestHLSInternal::FMediaStream::FDRMKeyInfo, ESPMode::ThreadSafe>& LicenseKeyInfo, const FTimeValue& ExpiresAtUTC) = 0;
+};
+
+
+} // namespace Electra
+
