ElectraPlayer: RELEASE. Moved from NFL to public. Note that this is done via a Delete and "Add New" operation as discussed.

diff --git a/Engine/Plugins/Media/ElectraPlayer/Source/ElectraPlayerRuntime/Private/Runtime/Player/HLS/ManifestBuilderHLS.cpp b/Engine/Plugins/Media/ElectraPlayer/Source/ElectraPlayerRuntime/Private/Runtime/Player/HLS/ManifestBuilderHLS.cpp
new file mode 100644
--- /dev/null
+++ b/Engine/Plugins/Media/ElectraPlayer/Source/ElectraPlayerRuntime/Private/Runtime/Player/HLS/ManifestBuilderHLS.cpp
@@ -0,0 +1,2085 @@
+// Copyright Epic Games, Inc. All Rights Reserved.
+
+#include "PlayerCore.h"
+#include "Player/Manifest.h"
+#include "Player/PlaybackTimeline.h"
+
+#include "Player/HLS/ManifestBuilderHLS.h"
+#include "Player/HLS/PlaylistReaderHLS.h"
+#include "Player/AdaptiveStreamingPlayerABR.h"
+#include "Player/PlayerSessionServices.h"
+#include "Player/PlayerStreamFilter.h"
+
+#include "InitSegmentCacheHLS.h"
+#include "LicenseKeyCacheHLS.h"
+#include "OptionalValue.h"
+
+#include "Tags.h"
+#include "LHLSTags.h"
+#include "EpicTags.h"
+
+#include "Utilities/StringHelpers.h"
+#include "SynchronizedClock.h"
+#include "StreamTypes.h"
+#include "ErrorDetail.h"
+
+#include "Utilities/URLParser.h"
+#include "Utilities/Utilities.h"
+#include "Utilities/TimeUtilities.h"
+
+
+#define ERRCODE_HLS_BUILDER_INTERNAL							1
+#define ERRCODE_HLS_BUILDER_UNSUPPORTED_FEATURE					2
+#define ERRCODE_HLS_BUILDER_OPTIONAL_IS_NEEDED					3
+#define ERRCODE_HLS_BUILDER_ID_NOT_FOUND						4
+#define ERRCODE_HLS_BUILDER_NOT_A_MASTER_PLAYLIST				5
+#define ERRCODE_HLS_BUILDER_NOT_A_VARIANT_PLAYLIST				6
+#define ERRCODE_HLS_BUILDER_ATTRIBUTE_PARSE_ERROR				7
+#define ERRCODE_HLS_BUILDER_ATTRIBUTE_INVALID_VALUE				8
+#define ERRCODE_HLS_BUILDER_MISSING_EXTX_TARGETDURATION			9
+#define ERRCODE_HLS_BUILDER_MISSING_EXTINF						10
+#define ERRCODE_HLS_BUILDER_MISSING_EXTXMEDIA_TYPE				11
+#define ERRCODE_HLS_BUILDER_MISSING_EXTXMEDIA_GROUP				12
+#define ERRCODE_HLS_BUILDER_MISSING_EXTXMEDIA_NAME				13
+#define ERRCODE_HLS_BUILDER_NO_VARIANTS_IN_MASTER_PLAYLIST		14
+#define ERRCODE_HLS_BUILDER_MISSING_STREAMINF_BANDWIDTH			15
+#define ERRCODE_HLS_BUILDER_UNSUPPORTED_CODEC					16
+#define ERRCODE_HLS_BUILDER_NO_USABLE_VARIANT_FOUND				17
+#define ERRCODE_HLS_BUILDER_RENDITION_URI_MISSING				18
+#define ERRCODE_HLS_BUILDER_RENDITION_NOT_FOUND_IN_GROUP		19
+#define ERRCODE_HLS_BUILDER_REFERENCED_GROUP_NOT_DEFINED		20
+
+
+namespace Electra
+{
+
+namespace
+{
+
+static void SplitOnOneOf(TArray<FString>& OutResults, const FString& Subject, const FString& SplitAt)
+{
+	OutResults.Empty();
+	if (Subject.Len())
+	{
+		int32 FirstPos = 0;
+		while(1)
+		{
+			int32 SplitPos = StringHelpers::FindFirstOf(Subject, SplitAt, FirstPos);
+			FString subs = Subject.Mid(FirstPos, SplitPos == INDEX_NONE ? MAX_int32 : SplitPos - FirstPos );
+
+			if (subs.Len())
+			{
+				OutResults.Push(subs);
+			}
+			if (SplitPos == INDEX_NONE)
+			{
+				break;
+			}
+			FirstPos = SplitPos + 1;
+		}
+	}
+}
+
+static void SplitOnCommaOrSpace(TArray<FString>& OutResults, const FString& Subject)
+{
+	static const FString kTextCommaOrSpace(TEXT(", "));
+	SplitOnOneOf(OutResults, Subject, kTextCommaOrSpace);
+}
+
+static bool IsAACCodec(const FString& Codec)
+{
+	return Codec.Find(TEXT("mp4a.40.")) == 0;
+}
+
+static bool IsH264Codec(TArray<FString>& OutResults, const FString& Codec)
+{
+	if (Codec.Find(TEXT("avc1.")) == 0)
+	{
+		FString prflv = Codec.Mid(5);
+		// avc1.xxyyzz format?
+		int32 DotIndex;
+		prflv.FindChar(TCHAR('.'), DotIndex);
+		if (DotIndex == INDEX_NONE)
+		{
+			// 6 hex digits?
+			if (prflv.Len() == 6)
+			{
+				OutResults.Push(prflv.Mid(0,2));
+				OutResults.Push(prflv.Mid(2,2));
+				OutResults.Push(prflv.Mid(4,2));
+				return true;
+			}
+		}
+		else
+		{
+			// avc1.xx.zz format
+			OutResults.Push(prflv.Mid(0, DotIndex));
+			OutResults.Push(prflv.Mid(DotIndex + 1));
+			return true;
+		}
+	}
+	return false;
+}
+
+}
+
+
+
+namespace HLSTimeline
+{
+
+class FPlaybackAssetRepresentation : public IPlaybackAssetRepresentation
+{
+public:
+	FPlaybackAssetRepresentation()
+		: Bitrate(0)
+	{
+	}
+	virtual ~FPlaybackAssetRepresentation()
+	{
+	}
+	virtual FString GetUniqueIdentifier() const override
+	{
+		return UniqueIdentifier;
+	}
+	virtual const FStreamCodecInformation& GetCodecInformation() const override
+	{
+		return CodecInformation;
+	}
+	virtual int32 GetBitrate() const override
+	{
+		return Bitrate;
+	}
+	virtual FString GetCDN() const override
+	{
+		return CDN;
+	}
+	FStreamCodecInformation	CodecInformation;
+	FString					UniqueIdentifier;
+	FString					CDN;
+	int32					Bitrate;
+};
+
+class FPlaybackAssetAdaptationSet : public IPlaybackAssetAdaptationSet
+{
+public:
+	virtual ~FPlaybackAssetAdaptationSet()
+	{
+	}
+	virtual FString GetUniqueIdentifier() const override
+	{
+		return UniqueIdentifier;
+	}
+	virtual FString GetListOfCodecs() const override
+	{
+		return Codecs;
+	}
+	virtual FString GetLanguage() const override
+	{
+		return Language;
+	}
+	virtual int32 GetNumberOfRepresentations() const override
+	{
+		return Representations.Num();
+	}
+	virtual TSharedPtrTS<IPlaybackAssetRepresentation> GetRepresentationByIndex(int32 RepresentationIndex) const override
+	{
+		check(RepresentationIndex < Representations.Num());
+		if (RepresentationIndex < Representations.Num())
+		{
+			return Representations[RepresentationIndex];
+		}
+		return TSharedPtrTS<IPlaybackAssetRepresentation>();
+	}
+	virtual TSharedPtrTS<IPlaybackAssetRepresentation> GetRepresentationByUniqueIdentifier(const FString& InUniqueIdentifier) const override
+	{
+		for(int32 i=0, iMax=Representations.Num(); i<iMax; ++i)
+		{
+			if (Representations[i]->GetUniqueIdentifier() == InUniqueIdentifier)
+			{
+				return Representations[i];
+			}
+		}
+		return TSharedPtrTS<IPlaybackAssetRepresentation>();
+	}
+
+	TArray<TSharedPtrTS<IPlaybackAssetRepresentation>>		Representations;
+	FString													UniqueIdentifier;
+	FString													Language;
+	FString													Codecs;
+};
+
+
+class FTimelineMediaAsset : public ITimelineMediaAsset
+{
+public:
+	virtual ~FTimelineMediaAsset()
+	{
+	}
+	virtual FTimeRange GetTimeRange() const override
+	{
+		return TimeRange;
+	}
+	virtual FTimeValue GetDuration() const override
+	{
+		return Duration;
+	}
+	virtual FString GetAssetIdentifier() const override
+	{
+		return AssetIdentifier;
+	}
+	virtual FString GetUniqueIdentifier() const override
+	{
+		return UniqueIdentifier;
+	}
+	virtual int32 GetNumberOfAdaptationSets(EStreamType OfStreamType) const override
+	{
+		switch(OfStreamType)
+		{
+			case EStreamType::Video:
+				return VideoAdaptationSet.IsValid() ? 1 : 0;
+			case EStreamType::Audio:
+				return (int32) AudioAdaptationSets.Num();
+			default:
+				return 0;
+		}
+	}
+	virtual TSharedPtrTS<IPlaybackAssetAdaptationSet> GetAdaptationSetByTypeAndIndex(EStreamType OfStreamType, int32 AdaptationSetIndex) const override
+	{
+		switch(OfStreamType)
+		{
+			case EStreamType::Video:
+			{
+				check(AdaptationSetIndex == 0);
+				return AdaptationSetIndex == 0 ? VideoAdaptationSet : TSharedPtrTS<FPlaybackAssetAdaptationSet>();
+			}
+			case EStreamType::Audio:
+			{
+				check(AdaptationSetIndex < AudioAdaptationSets.Num());
+				return AdaptationSetIndex < AudioAdaptationSets.Num() ? AudioAdaptationSets[AdaptationSetIndex] : TSharedPtrTS<FPlaybackAssetAdaptationSet>();
+			}
+			default:
+			{
+				return TSharedPtrTS<FPlaybackAssetAdaptationSet>();
+			}
+		}
+	}
+	TSharedPtrTS<IPlaybackAssetAdaptationSet> GetAdaptationSetByTypeAndUniqueIdentifier(EStreamType OfStreamType, const FString& InUniqueIdentifier) const
+	{
+		switch(OfStreamType)
+		{
+			case EStreamType::Video:
+			{
+				return VideoAdaptationSet;
+			}
+			case EStreamType::Audio:
+			{
+				for(int32 i=0, iMax=(int32)AudioAdaptationSets.Num(); i<iMax; ++i)
+				{
+					if (AudioAdaptationSets[i]->GetUniqueIdentifier() == InUniqueIdentifier)
+					{
+						return AudioAdaptationSets[i];
+					}
+				}
+				return TSharedPtrTS<FPlaybackAssetAdaptationSet>();
+			}
+			default:
+			{
+				return TSharedPtrTS<FPlaybackAssetAdaptationSet>();
+			}
+		}
+	}
+
+	TSharedPtrTS<FPlaybackAssetAdaptationSet>				VideoAdaptationSet;
+	TArray<TSharedPtrTS<FPlaybackAssetAdaptationSet>>		AudioAdaptationSets;
+	FString													UniqueIdentifier;
+	FString													AssetIdentifier;
+	FTimeRange												TimeRange;
+	FTimeValue												Duration;
+};
+
+
+class FPlaybackTimeline : public IPlaybackAssetTimeline
+{
+public:
+	virtual ~FPlaybackTimeline()
+	{
+	}
+	virtual FTimeValue GetAnchorTime() const override
+	{
+		return FTimeValue::GetZero();
+	}
+	virtual FTimeRange GetTotalTimeRange() const override
+	{
+		return TotalTimeRange;
+	}
+	virtual FTimeRange GetSeekableTimeRange() const override
+	{
+		return SeekableTimeRange;
+	}
+	virtual void GetSeekablePositions(TArray<FTimespan>& OutPositions) const override
+	{
+		OutPositions = SeekablePositions;
+	}
+	virtual FTimeValue GetDuration() const override
+	{
+		return Duration;
+	}
+	virtual int32 GetNumberOfMediaAssets() const override
+	{
+		return 1;
+	}
+	virtual TSharedPtrTS<ITimelineMediaAsset> GetMediaAssetByIndex(int32 MediaAssetIndex) const override
+	{
+		return PlayAsset;
+	}
+
+	FTimeRange								TotalTimeRange;
+	FTimeRange								SeekableTimeRange;
+	FTimeValue								Duration;
+	TSharedPtrTS<FTimelineMediaAsset>		PlayAsset;
+	TArray<FTimespan>						SeekablePositions;
+};
+
+}
+
+
+
+
+class FManifestBuilderHLS : public IManifestBuilderHLS
+{
+public:
+	FManifestBuilderHLS();
+	void Initialize(IPlayerSessionServices* PlayerSessionServices);
+	virtual ~FManifestBuilderHLS();
+
+	virtual FErrorDetail BuildFromMasterPlaylist(TSharedPtrTS<FManifestHLSInternal>& OutHLSPlaylist, const HLSPlaylistParser::FPlaylist& Playlist, const FPlaylistLoadRequestHLS& SourceRequest, const HTTP::FConnectionInfo* ConnectionInfo, const FStreamPreferences& Preferences, const FParamDict& Options) override;
+
+	virtual FErrorDetail GetInitialPlaylistLoadRequests(TArray<FPlaylistLoadRequestHLS>& OutRequests, TSharedPtrTS<FManifestHLSInternal> Manifest, const FStreamPreferences& Preferences, const FParamDict& Options) override;
+	virtual UEMediaError UpdateFailedInitialPlaylistLoadRequest(FPlaylistLoadRequestHLS& InOutFailedRequest, const HTTP::FConnectionInfo* ConnectionInfo, TSharedPtrTS<HTTP::FRetryInfo> PreviousAttempts, const FTimeValue& BlacklistUntilUTC, TSharedPtrTS<FManifestHLSInternal> Manifest, const FStreamPreferences& Preferences, const FParamDict& Options) override;
+
+	virtual FErrorDetail UpdateFromVariantPlaylist(TSharedPtrTS<FManifestHLSInternal> InOutHLSPlaylist, const HLSPlaylistParser::FPlaylist& VariantPlaylist, const FPlaylistLoadRequestHLS& SourceRequest, const HTTP::FConnectionInfo* ConnectionInfo, uint32 ResponseCRC, const FStreamPreferences& Preferences, const FParamDict& Options) override;
+	virtual void SetVariantPlaylistFailure(TSharedPtrTS<FManifestHLSInternal> InHLSPlaylist, const FPlaylistLoadRequestHLS& SourceRequest, const HTTP::FConnectionInfo* ConnectionInfo, TSharedPtrTS<HTTP::FRetryInfo> PreviousAttempts, const FTimeValue& BlacklistUntilUTC) override;
+
+private:
+	FErrorDetail SetupRenditions(FManifestHLSInternal* Manifest, const HLSPlaylistParser::FPlaylist& Playlist);
+	FErrorDetail SetupVariants(FManifestHLSInternal* Manifest, const HLSPlaylistParser::FPlaylist& Playlist);
+	void UpdateManifestMetadataFromStream(FManifestHLSInternal* Manifest, const FManifestHLSInternal::FPlaylistBase* Playlist, FManifestHLSInternal::FMediaStream* Stream, const FStreamPreferences& Preferences, const FParamDict& Options);
+	void LogMessage(IInfoLog::ELevel Level, const FString& Message);
+	FErrorDetail CreateErrorAndLog(const FString& Message, uint16 Code, UEMediaError Error = UEMEDIA_ERROR_OK);
+
+	IPlayerSessionServices*		PlayerSessionServices;
+	uint32						NextUniqueID;
+};
+
+
+IManifestBuilderHLS* IManifestBuilderHLS::Create(IPlayerSessionServices* PlayerSessionServices)
+{
+	FManifestBuilderHLS* ManifestBuilder = new FManifestBuilderHLS;
+	if (ManifestBuilder)
+	{
+		ManifestBuilder->Initialize(PlayerSessionServices);
+	}
+	return ManifestBuilder;
+}
+
+
+FManifestBuilderHLS::FManifestBuilderHLS()
+	: PlayerSessionServices(nullptr)
+	, NextUniqueID(0)
+{
+}
+
+FManifestBuilderHLS::~FManifestBuilderHLS()
+{
+}
+
+void FManifestBuilderHLS::Initialize(IPlayerSessionServices* InPlayerSessionServices)
+{
+	PlayerSessionServices = InPlayerSessionServices;
+}
+
+FErrorDetail FManifestBuilderHLS::CreateErrorAndLog(const FString& InMessage, uint16 InCode, UEMediaError InError)
+{
+	FErrorDetail err;
+	err.SetError(InError != UEMEDIA_ERROR_OK ? InError : UEMEDIA_ERROR_DETAIL);
+	err.SetFacility(Facility::EFacility::HLSPlaylistBuilder);
+	err.SetCode(InCode);
+	err.SetMessage(InMessage);
+	check(PlayerSessionServices);
+	if (PlayerSessionServices)
+	{
+		PlayerSessionServices->PostLog(Facility::EFacility::HLSPlaylistReader, IInfoLog::ELevel::Error, err.GetPrintable());
+	}
+	return err;
+}
+
+void FManifestBuilderHLS::LogMessage(IInfoLog::ELevel Level, const FString& Message)
+{
+	if (PlayerSessionServices)
+	{
+		PlayerSessionServices->PostLog(Facility::EFacility::HLSPlaylistBuilder, Level, Message);
+	}
+}
+
+
+FErrorDetail FManifestBuilderHLS::BuildFromMasterPlaylist(TSharedPtrTS<FManifestHLSInternal>& OutHLSPlaylist, const HLSPlaylistParser::FPlaylist& Playlist, const FPlaylistLoadRequestHLS& SourceRequest, const HTTP::FConnectionInfo* ConnectionInfo, const FStreamPreferences& Preferences, const FParamDict& Options)
+{
+	FErrorDetail Error;
+
+	if (Playlist.Type != HLSPlaylistParser::EPlaylistType::Master)
+	{
+		return CreateErrorAndLog(FString::Printf(TEXT("Not a master playlist")), ERRCODE_HLS_BUILDER_NOT_A_MASTER_PLAYLIST, UEMEDIA_ERROR_FORMAT_ERROR);
+	}
+
+	TUniquePtr<FManifestHLSInternal> Manifest(new FManifestHLSInternal);
+
+	Manifest->InitSegmentCache = MakeShareable(IInitSegmentCacheHLS::Create(PlayerSessionServices, Options));
+	Manifest->LicenseKeyCache = MakeShareable(ILicenseKeyCacheHLS::Create(PlayerSessionServices, Options));
+
+	// Remember the URL we got the playlist from and the time we got it.
+	Manifest->MasterPlaylistVars.PlaylistLoadRequest = SourceRequest;
+	check(ConnectionInfo);
+	if (ConnectionInfo)
+	{
+		Manifest->MasterPlaylistVars.PlaylistLoadRequest.URL = ConnectionInfo->EffectiveURL;
+	}
+
+
+	// Does the master playlist specify independent segments?
+	if (Playlist.ContainsTag(HLSPlaylistParser::ExtXIndependentSegments))
+	{
+		Manifest->bHasIndependentSegments = true;
+	}
+
+	// Do we have a specified start position?
+	if (Playlist.ContainsTag(HLSPlaylistParser::ExtXStart))
+	{
+		// TODO: some future addition
+	}
+
+	// Are there global DRM key sessions?
+	// TODO: future additions. Note: there can be multiple session keys!!!
+	//		if (Playlist.ContainsTag(HLSPlaylistParser::ExtXSessionKey))
+
+
+	// Create timeline and its primary asset object
+	HLSTimeline::FPlaybackTimeline* Timeline = new HLSTimeline::FPlaybackTimeline;
+	HLSTimeline::FTimelineMediaAsset* Asset = new HLSTimeline::FTimelineMediaAsset;
+	Asset->UniqueIdentifier = TEXT("Asset0");
+	Timeline->PlayAsset = MakeShareable(Asset);
+	Manifest->PlaybackTimeline = MakeShareable(Timeline);
+	Timeline = nullptr;
+	Asset = nullptr;
+
+
+	// Set up any renditions specified in the master playlist.
+	Error = SetupRenditions(Manifest.Get(), Playlist);
+	if (Error.IsSet())
+	{
+		return Error;
+	}
+
+	// Set up the variant streams.
+	Error = SetupVariants(Manifest.Get(), Playlist);
+	if (Error.IsSet())
+	{
+		return Error;
+	}
+
+	// Pass the new internal manifest out to the caller.
+	OutHLSPlaylist = MakeShareable(Manifest.Release());
+	return Error;
+}
+
+
+FErrorDetail FManifestBuilderHLS::SetupRenditions(FManifestHLSInternal* Manifest, const HLSPlaylistParser::FPlaylist& Playlist)
+{
+	// Do we have a specified start position?
+	if (Playlist.ContainsTag(HLSPlaylistParser::ExtXStart))
+	{
+		// TODO: some future addition
+	}
+
+	// Are there global DRM key sessions?
+	// TODO: future additions. Note: there can be multiple session keys!!!
+	//		if (Playlist.ContainsTag(HLSPlaylistParser::ExtXSessionKey))
+
+
+	// Get the alternate renditions, if any.
+	int32 NumMediaTags = Playlist.TagNum(HLSPlaylistParser::ExtXMedia);
+	for(int32 nMediaTag=0; nMediaTag<NumMediaTags; ++nMediaTag)
+	{
+		// Helper lambda to get an attribute string value from an EXT-X-MEDIA tag.
+		auto GetStringAttribute = [&Playlist, nMediaTag](FString& OutValue, const FString& Attribute) -> bool
+		{
+			FString						AttributeValueStr;
+			HLSPlaylistParser::EPlaylistError	PlaylistError = Playlist.GetTagAttributeValue(HLSPlaylistParser::ExtXMedia, Attribute, AttributeValueStr, nMediaTag);
+			if (PlaylistError == HLSPlaylistParser::EPlaylistError::None)
+			{
+				OutValue = AttributeValueStr;
+				return true;
+			}
+			return false;
+		};
+		// Helper lambda to get an attribute boolean value from an EXT-X-MEDIA tag.
+		auto GetBooleanAttribute = [&Playlist, nMediaTag](bool& OutValue, const FString& Attribute) -> bool
+		{
+			FString						AttributeValueStr;
+			HLSPlaylistParser::EPlaylistError	PlaylistError = Playlist.GetTagAttributeValue(HLSPlaylistParser::ExtXMedia, Attribute, AttributeValueStr, nMediaTag);
+			if (PlaylistError == HLSPlaylistParser::EPlaylistError::None)
+			{
+				OutValue = AttributeValueStr == TEXT("YES");
+				return true;
+			}
+			return false;
+		};
+
+		// Get the rendition type. This attribute is required.
+		FString Type;
+		if (!GetStringAttribute(Type, TEXT("TYPE")))
+		{
+			return CreateErrorAndLog(FString::Printf(TEXT("EXT-X-MEDIA attribute does not have mandatory TYPE attribute")), ERRCODE_HLS_BUILDER_MISSING_EXTXMEDIA_TYPE, UEMEDIA_ERROR_FORMAT_ERROR);
+		}
+		// Check for valid type
+		if (Type != TEXT("VIDEO") && Type != TEXT("AUDIO") && Type != TEXT("SUBTITLES") && Type != TEXT("CLOSED-CAPTIONS"))
+		{
+			return CreateErrorAndLog(FString::Printf(TEXT("EXT-X-MEDIA TYPE attribute (\"%s\") value is invalid"), *Type), ERRCODE_HLS_BUILDER_ATTRIBUTE_INVALID_VALUE, UEMEDIA_ERROR_FORMAT_ERROR);
+		}
+
+		// Get the rendition group id. This attribute is required.
+		FString GroupID;
+		if (!GetStringAttribute(GroupID, TEXT("GROUP-ID")))
+		{
+			return CreateErrorAndLog(FString::Printf(TEXT("EXT-X-MEDIA attribute does not have mandatory GROUP-ID attribute")), ERRCODE_HLS_BUILDER_MISSING_EXTXMEDIA_GROUP, UEMEDIA_ERROR_FORMAT_ERROR);
+		}
+
+		// Get the remaining required attributes
+		TSharedPtrTS<FManifestHLSInternal::FRendition> Rendition = MakeSharedTS<FManifestHLSInternal::FRendition>();
+		Rendition->Type = Type;
+		Rendition->GroupID = GroupID;
+		if (!GetStringAttribute(Rendition->Name, TEXT("NAME")))
+		{
+			return CreateErrorAndLog(FString::Printf(TEXT("EXT-X-MEDIA attribute does not have mandatory NAME attribute")), ERRCODE_HLS_BUILDER_MISSING_EXTXMEDIA_NAME, UEMEDIA_ERROR_FORMAT_ERROR);
+		}
+		// Get optional attributes.
+		// NOTE: Some are conditionally required (like 'channels' that must be present if there is more than one rendition in the group)
+		//       but we do not check for this here.
+		GetStringAttribute(Rendition->URI, TEXT("URI"));
+		GetStringAttribute(Rendition->Language, TEXT("LANGUAGE"));
+		GetStringAttribute(Rendition->AssocLanguage, TEXT("ASSOC-LANGUAGE"));
+		GetStringAttribute(Rendition->InStreamID, TEXT("INSTREAM-ID"));
+		GetStringAttribute(Rendition->Characteristics, TEXT("CHARACTERISTICS"));
+		GetStringAttribute(Rendition->Channels, TEXT("CHANNELS"));
+		GetBooleanAttribute(Rendition->bDefault, TEXT("DEFAULT"));
+		GetBooleanAttribute(Rendition->bAutoSelect, TEXT("AUTOSELECT"));
+		GetBooleanAttribute(Rendition->bForced, TEXT("FORCED"));
+
+		// Set up an internal ID we can use to track any playlist requests/updates for this rendition.
+		Rendition->Internal.UniqueID = FMediaInterlockedIncrement(NextUniqueID) + 1;
+
+		// Add to appropriate rendition group and to the lookup map.
+		if (Type == TEXT("VIDEO"))
+		{
+			Rendition->Internal.bHasVideo = true;
+			Manifest->VideoRenditions.Add(GroupID, Rendition);
+			Manifest->PlaylistIDMap.Add(Rendition->Internal.UniqueID, Rendition);
+		}
+		else if (Type == TEXT("AUDIO"))
+		{
+			Rendition->Internal.bHasAudio = true;
+			Manifest->AudioRenditions.Add(GroupID, Rendition);
+			Manifest->PlaylistIDMap.Add(Rendition->Internal.UniqueID, Rendition);
+		}
+		else if (Type == TEXT("SUBTITLES"))
+		{
+			Manifest->SubtitleRenditions.Add(GroupID, Rendition);
+			Manifest->PlaylistIDMap.Add(Rendition->Internal.UniqueID, Rendition);
+		}
+		else if (Type == TEXT("CLOSED-CAPTIONS"))
+		{
+			Manifest->ClosedCaptionRenditions.Add(GroupID, Rendition);
+			Manifest->PlaylistIDMap.Add(Rendition->Internal.UniqueID, Rendition);
+		}
+	}
+
+	// Get the timeline asset object.
+	HLSTimeline::FTimelineMediaAsset* Asset = static_cast<HLSTimeline::FTimelineMediaAsset*>(Manifest->PlaybackTimeline->GetMediaAssetByIndex(0).Get());
+
+	// Create adaptation sets for the audio rendition groups
+// FIXME: We will need to do this for the other rendition types (VIDEO, SUBTITLES and CLOSED-CAPTIONS as well)
+	for (TMultiMap<FString, TSharedPtrTS<FManifestHLSInternal::FRendition>>::TConstIterator It = Manifest->AudioRenditions.CreateConstIterator(); It; ++It)
+	{
+		HLSTimeline::FPlaybackAssetAdaptationSet* Adapt = static_cast<HLSTimeline::FPlaybackAssetAdaptationSet*>(Asset->GetAdaptationSetByTypeAndUniqueIdentifier(EStreamType::Audio, It.Key()).Get());
+		if (!Adapt)
+		{
+			Adapt = new HLSTimeline::FPlaybackAssetAdaptationSet;
+			Adapt->UniqueIdentifier = It.Key();
+			TSharedPtrTS<HLSTimeline::FPlaybackAssetAdaptationSet> IAdapt(Adapt);
+			Asset->AudioAdaptationSets.Push(IAdapt);
+		}
+
+		TSharedPtrTS<FManifestHLSInternal::FRendition> Rendition = It.Value();
+		HLSTimeline::FPlaybackAssetRepresentation* Repr = new HLSTimeline::FPlaybackAssetRepresentation;
+		Repr->UniqueIdentifier = LexToString(Rendition->Internal.UniqueID);
+		Rendition->Internal.AdaptationSetUniqueID = It.Key();
+		Rendition->Internal.RepresentationUniqueID = Repr->UniqueIdentifier;
+	// FIXME: need to setup CDNs first from all URIs and then assign the correct one here.
+	// NOTE: URI is optional and may not even exist.
+		Rendition->Internal.CDN = Rendition->URI;
+		Repr->CDN = Rendition->Internal.CDN;
+		Adapt->Representations.Push(TSharedPtrTS<IPlaybackAssetRepresentation>(Repr));
+	}
+
+	return FErrorDetail();
+}
+
+
+FErrorDetail FManifestBuilderHLS::SetupVariants(FManifestHLSInternal* Manifest, const HLSPlaylistParser::FPlaylist& Playlist)
+{
+	if (Playlist.GetPlaylists().Num() == 0)
+	{
+		return CreateErrorAndLog(FString::Printf(TEXT("No variant playlists in master playlist")), ERRCODE_HLS_BUILDER_NO_VARIANTS_IN_MASTER_PLAYLIST, UEMEDIA_ERROR_FORMAT_ERROR);
+	}
+
+	// Get the timeline asset object.
+	HLSTimeline::FTimelineMediaAsset* Asset = static_cast<HLSTimeline::FTimelineMediaAsset*>(Manifest->PlaybackTimeline->GetMediaAssetByIndex(0).Get());
+
+	TUniquePtr<IURLParser> UrlBuilder(IURLParser::Create());
+	UrlBuilder->ParseURL(Manifest->MasterPlaylistVars.PlaylistLoadRequest.URL);
+
+	for(const HLSPlaylistParser::FMediaPlaylist& VariantStream : Playlist.GetPlaylists())
+	{
+		// Helper lambda to get an attribute string value from an EXT-X-STREAM-INF tag.
+		auto GetStringAttribute = [&VariantStream](FString& OutValue, const FString& Attribute) -> bool
+		{
+			FString						AttributeValueStr;
+			HLSPlaylistParser::EPlaylistError	PlaylistError = VariantStream.GetTagAttributeValue(HLSPlaylistParser::ExtXStreamInf, Attribute, AttributeValueStr);
+			if (PlaylistError == HLSPlaylistParser::EPlaylistError::None)
+			{
+				OutValue = AttributeValueStr;
+				return true;
+			}
+			return false;
+		};
+
+
+
+		// Check if we have a company specific tag with custom attributes.
+		FParamDict CompanyCustomExtraOptions;
+		if (VariantStream.ContainsTag(HLSPlaylistParser::Epic::ExtXEpicGamesCustom))
+		{
+			auto GetCustomAttribute = [&VariantStream](FString& OutValue, const FString& Attribute) -> bool
+			{
+				FString						AttributeValueStr;
+				HLSPlaylistParser::EPlaylistError	PlaylistError = VariantStream.GetTagAttributeValue(HLSPlaylistParser::Epic::ExtXEpicGamesCustom, Attribute, AttributeValueStr);
+				if (PlaylistError == HLSPlaylistParser::EPlaylistError::None)
+				{
+					OutValue = AttributeValueStr;
+					return true;
+				}
+				return false;
+			};
+
+			// Check indicator for b-frame presence.
+			FString CustAttrBFrame;
+			if (GetCustomAttribute(CustAttrBFrame, TEXT("BFRAMES")))
+			{
+				// The value is an int with 0=no B frames, 1=B frames present.
+				int32 bCustAttrBFrame = 0;
+				LexFromString(bCustAttrBFrame, *CustAttrBFrame);
+				if (bCustAttrBFrame != 0)
+				{
+					// We set custom options only when the attribute is explicitly set to non-0.
+					CompanyCustomExtraOptions.Set(TEXT("b_frames"), FVariantValue((int64) 1));
+				}
+			}
+		}
+
+		// The only required attribute is bandwidth, so let's make sure it's there.
+		FString Bandwidth;
+		if (!GetStringAttribute(Bandwidth, TEXT("BANDWIDTH")))
+		{
+			return CreateErrorAndLog(FString::Printf(TEXT("EXT-X-STREAM-INF does not have mandatory BANDWIDTH attribute")), ERRCODE_HLS_BUILDER_MISSING_STREAMINF_BANDWIDTH, UEMEDIA_ERROR_FORMAT_ERROR);
+		}
+
+		TSharedPtrTS<FManifestHLSInternal::FVariantStream> vs = MakeSharedTS<FManifestHLSInternal::FVariantStream>();
+		LexFromString(vs->Bandwidth, *Bandwidth);
+		vs->URI = VariantStream.URL;
+
+		// Get the remaining optional attributes.
+		FString Temp;
+		if (GetStringAttribute(Temp, TEXT("AVERAGE-BANDWIDTH")))
+		{
+			LexFromString(vs->AverageBandwidth, *Temp);
+		}
+		if (GetStringAttribute(Temp, TEXT("HDCP-LEVEL")))
+		{
+			// NOTE: So far only "NONE" and "TYPE-0" are defined. In case there's anything else here we treat it as "NONE".
+			vs->HDCPLevel = Temp == TEXT("NONE") ? 0 : 1;
+		}
+		bool bHaveVideoGroup = GetStringAttribute(vs->VideoGroupID, TEXT("VIDEO"));
+		bool bHaveAudioGroup = GetStringAttribute(vs->AudioGroupID, TEXT("AUDIO"));
+		GetStringAttribute(vs->SubtitleGroupID, TEXT("SUBTITLES"));
+		GetStringAttribute(vs->ClosedCaptionGroupID, TEXT("CLOSED-CAPTIONS"));
+
+		// Get codecs, resolution and framerate.
+		HLSPlaylistParser::FResolution Resolution;
+		FString Codecs, FrameRate, ResolutionStr;
+		bool bHaveResolution = GetStringAttribute(ResolutionStr, TEXT("RESOLUTION"));// == HLSPlaylistParser::EPlaylistError::None;
+		if (bHaveResolution)
+		{
+			TArray<FString> ResolutionDimensions;
+			SplitOnOneOf(ResolutionDimensions, ResolutionStr, TEXT("xX"));
+			if (ResolutionDimensions.Num() == 2)
+			{
+				LexFromString(Resolution.Width,  *ResolutionDimensions[0]);
+				LexFromString(Resolution.Height, *ResolutionDimensions[1]);
+			}
+		}
+		GetStringAttribute(Codecs, TEXT("CODECS"));
+		GetStringAttribute(FrameRate, TEXT("FRAME-RATE"));
+
+		// A problem is that all of these are optional. If not specified it is a guessing game as to what this variant actually represents.
+		bool bHasVideo = false;
+		bool bHasAudio = false;
+		// Are codecs specified?
+		if (Codecs.Len())
+		{
+			TArray<FString> CodecList;
+			SplitOnCommaOrSpace(CodecList, Codecs);
+			// Parse each codec
+			for(int32 nCodec=0; nCodec<CodecList.Num(); ++nCodec)
+			{
+				FStreamCodecInformation si;
+				// Set the custom attributes with the codec information. This sets all extra options regardless of codec type!
+				si.GetExtras() = CompanyCustomExtraOptions;
+
+				// Does it match the standard AVC OTI?
+				TArray<FString> CodecOTI;
+				if (IsH264Codec(CodecOTI, CodecList[nCodec]))
+				{
+					if (CodecOTI.Num() == 3)
+					{
+						si.SetStreamType(EStreamType::Video);
+						si.SetCodec(FStreamCodecInformation::ECodec::H264);
+						si.SetCodecSpecifierRFC6381(CodecList[nCodec]);
+						int32 TempValue;
+						LexFromStringHex(TempValue, *CodecOTI[0]);
+						si.SetProfile(TempValue);
+						LexFromStringHex(TempValue, *CodecOTI[1]);
+						si.SetProfileConstraints(TempValue);
+						LexFromStringHex(TempValue, *CodecOTI[2]);
+						si.SetProfileLevel(TempValue);
+						bHasVideo = true;
+					}
+					// Does it match the alternate AVC form (invalid, but we allow it) of avc1.profile.level ?
+					else if (CodecOTI.Num() == 2)
+					{
+						si.SetStreamType(EStreamType::Video);
+						si.SetCodec(FStreamCodecInformation::ECodec::H264);
+						int32 TempValue;
+						LexFromString(TempValue, *CodecOTI[0]);
+						si.SetProfile(TempValue);
+						LexFromString(TempValue, *CodecOTI[1]);
+						si.SetProfileLevel(TempValue);
+						// Convert the .profile.level integers into the normal hexdigit grouping notation.
+						si.SetCodecSpecifierRFC6381(FString::Printf(TEXT("avc1.%02x00%02x"), si.GetProfile(), si.GetProfileLevel()));
+						bHasVideo = true;
+					}
+				}
+				// Match for AAC audio?
+				else if (IsAACCodec(CodecList[nCodec]))
+				{
+					si.SetStreamType(EStreamType::Audio);
+					si.SetCodec(FStreamCodecInformation::ECodec::AAC);
+					si.SetCodecSpecifierRFC6381(CodecList[nCodec]);
+					bHasAudio = true;
+				}
+				else
+				{
+					return CreateErrorAndLog(FString::Printf(TEXT("EXT-X-STREAM-INF CODECS lists unsupported codec \"%s\"!"), *CodecList[nCodec]), ERRCODE_HLS_BUILDER_UNSUPPORTED_CODEC, UEMEDIA_ERROR_FORMAT_ERROR);
+				}
+
+				// If this was a video stream set the resolution and framerate if those are specified.
+				if (si.GetStreamType() == EStreamType::Video)
+				{
+					if (bHaveResolution)
+					{
+						si.SetResolution(FStreamCodecInformation::FResolution(Resolution.Width, Resolution.Height));
+					}
+					if (FrameRate.Len())
+					{
+						si.SetFrameRate(FTimeFraction().SetFromFloatString(FrameRate));
+						if (!si.GetFrameRate().IsValid())
+						{
+							return CreateErrorAndLog(FString::Printf(TEXT("EXT-X-STREAM-INF FRAME-RATE attribute \"%s\" did not parse!"), *FrameRate), ERRCODE_HLS_BUILDER_ATTRIBUTE_PARSE_ERROR, UEMEDIA_ERROR_FORMAT_ERROR);
+						}
+					}
+				}
+
+				// Add to the list of codecs for this variant stream
+				vs->StreamCodecInformationList.Push(si);
+			}
+		}
+		else
+		{
+			LogMessage(IInfoLog::ELevel::Warning, FString::Printf(TEXT("EXT-X-STREAM-INF CODECS not specified. Assuming H.264 HIGH profile level 4.2 and LC-AAC audio.")));
+
+			// With either resolution or framerate specified we assume this is video.
+			if (bHaveResolution || FrameRate.Len())
+			{
+				FStreamCodecInformation si;
+				// Set the custom attributes with the codec information. This sets all extra options regardless of codec type!
+				si.GetExtras() = CompanyCustomExtraOptions;
+				si.SetStreamType(EStreamType::Video);
+				// For codec assume H.264 HIGH profile level 4.2
+				si.SetCodec(FStreamCodecInformation::ECodec::H264);
+				si.SetProfile(100);
+				si.SetProfileLevel(42);
+				si.SetCodecSpecifierRFC6381(FString::Printf(TEXT("avc1.%02x00%02x"), si.GetProfile(), si.GetProfileLevel()));
+
+				if (bHaveResolution)
+				{
+					si.SetResolution(FStreamCodecInformation::FResolution(Resolution.Width, Resolution.Height));
+				}
+				if (FrameRate.Len())
+				{
+					si.SetFrameRate(FTimeFraction().SetFromFloatString(FrameRate));
+					if (!si.GetFrameRate().IsValid())
+					{
+						return CreateErrorAndLog(FString::Printf(TEXT("EXT-X-STREAM-INF FRAME-RATE attribute \"%s\" did not parse!"), *FrameRate), ERRCODE_HLS_BUILDER_ATTRIBUTE_PARSE_ERROR, UEMEDIA_ERROR_FORMAT_ERROR);
+					}
+				}
+				vs->StreamCodecInformationList.Push(si);
+				bHasVideo = true;
+			}
+			// If an audio group is specified we assume there is audio.
+			if (bHaveAudioGroup)
+			{
+				FStreamCodecInformation si;
+				// Set the custom attributes with the codec information. This sets all extra options regardless of codec type!
+				si.GetExtras() = CompanyCustomExtraOptions;
+				si.SetStreamType(EStreamType::Audio);
+				// For codec we assume standard LC-AAC
+				si.SetCodec(FStreamCodecInformation::ECodec::AAC);
+				si.SetCodecSpecifierRFC6381(TEXT("mp4a.40.2"));
+				vs->StreamCodecInformationList.Push(si);
+				bHasAudio = true;
+			}
+		}
+
+		// Usable stream?
+		check(bHasVideo || bHasAudio);
+		if ((bHasVideo || bHasAudio) && vs.IsValid())
+		{
+			// Check if the stream can be used on this platform.
+			check(PlayerSessionServices);
+			IPlayerStreamFilter* StreamFilter = PlayerSessionServices->GetStreamFilter();
+			check(StreamFilter);
+			bool bCanDecodeStream = true;
+			if (StreamFilter)
+			{
+				for(int32 i = 0; i < vs->StreamCodecInformationList.Num(); ++i)
+				{
+					if (!StreamFilter->CanDecodeStream(vs->StreamCodecInformationList[i]))
+					{
+						bCanDecodeStream = false;
+						vs.Reset();
+						break;
+					}
+				}
+			}
+		}
+
+		// Still a usable stream?
+		if ((bHasVideo || bHasAudio) && vs.IsValid())
+		{
+			vs->Internal.bHasVideo = bHasVideo;
+			vs->Internal.bHasAudio = bHasAudio;
+			// Set up an internal ID we can use to track any playlist requests/updates for this variant stream.
+			vs->Internal.UniqueID = FMediaInterlockedIncrement(NextUniqueID) + 1;
+			// Add to the available variants.
+			// First check if this is variant includes video.
+			if (bHasVideo)
+			{
+				Manifest->VariantStreams.Push(vs);
+				Manifest->BandwidthToQualityIndex.Add(vs->Bandwidth, vs->Bandwidth);
+
+				FStreamMetadata StreamMetaData;
+				StreamMetaData.Bandwidth	  = vs->Bandwidth;
+				StreamMetaData.StreamUniqueID = vs->Internal.UniqueID;
+				StreamMetaData.PlaylistID     = UrlBuilder->ResolveWith(vs->URI);
+				for(int32 i=0; i<vs->StreamCodecInformationList.Num(); ++i)
+				{
+					if (vs->StreamCodecInformationList[i].IsVideoCodec())
+					{
+						StreamMetaData.CodecInformation = vs->StreamCodecInformationList[i];
+						break;
+					}
+				}
+				Manifest->StreamMetadataVideo.Push(StreamMetaData);
+
+				// Add to timeline asset
+				HLSTimeline::FPlaybackAssetAdaptationSet* Adapt = nullptr;
+				if (!Asset->VideoAdaptationSet.IsValid())
+				{
+					Adapt = new HLSTimeline::FPlaybackAssetAdaptationSet;
+					Asset->VideoAdaptationSet = MakeShareable(Adapt);
+					Adapt->UniqueIdentifier = TEXT("$video$");
+				}
+				Adapt = static_cast<HLSTimeline::FPlaybackAssetAdaptationSet*>(Asset->VideoAdaptationSet.Get());
+				if (Adapt->Codecs.Find(StreamMetaData.CodecInformation.GetCodecSpecifierRFC6381()) == INDEX_NONE)
+				{
+					if (Adapt->Codecs.Len())
+					{
+						Adapt->Codecs.AppendChar(TCHAR(','));
+					}
+					Adapt->Codecs.Append(StreamMetaData.CodecInformation.GetCodecSpecifierRFC6381());
+				}
+				HLSTimeline::FPlaybackAssetRepresentation* Repr = new HLSTimeline::FPlaybackAssetRepresentation;
+				Repr->UniqueIdentifier = LexToString(StreamMetaData.StreamUniqueID);
+				Repr->CodecInformation = StreamMetaData.CodecInformation;
+				Repr->CDN   		   = StreamMetaData.PlaylistID;
+				Repr->Bitrate   	   = StreamMetaData.Bandwidth;
+				Adapt->Representations.Push(TSharedPtrTS<IPlaybackAssetRepresentation>(Repr));
+
+				vs->Internal.AdaptationSetUniqueID  = Adapt->UniqueIdentifier;
+				vs->Internal.RepresentationUniqueID = Repr->UniqueIdentifier;
+				vs->Internal.CDN					= StreamMetaData.PlaylistID;
+
+				// Check if there is an audio adaptation set already due to an audio rendition group referenced by this variant.
+				if (bHaveAudioGroup)
+				{
+					Adapt = static_cast<HLSTimeline::FPlaybackAssetAdaptationSet*>(Asset->GetAdaptationSetByTypeAndUniqueIdentifier(EStreamType::Audio, vs->AudioGroupID).Get());
+					if (!Adapt)
+					{
+						return CreateErrorAndLog(FString::Printf(TEXT("EXT-X-STREAM-INF references AUDIO rendition group \"%s\" that has not been defined!"), *vs->AudioGroupID), ERRCODE_HLS_BUILDER_REFERENCED_GROUP_NOT_DEFINED, UEMEDIA_ERROR_FORMAT_ERROR);
+					}
+					// Propagate audio codecs from variant to rendition since they are not specified there.
+					for(int32 i=0; i<vs->StreamCodecInformationList.Num(); ++i)
+					{
+						if (vs->StreamCodecInformationList[i].IsAudioCodec())
+						{
+							StreamMetaData.CodecInformation = vs->StreamCodecInformationList[i];
+							if (Adapt->Codecs.Find(vs->StreamCodecInformationList[i].GetCodecSpecifierRFC6381()) == INDEX_NONE)
+							{
+								if (Adapt->Codecs.Len())
+								{
+									Adapt->Codecs.AppendChar(TCHAR(','));
+								}
+								Adapt->Codecs.Append(vs->StreamCodecInformationList[i].GetCodecSpecifierRFC6381());
+							}
+						}
+					}
+					for(int32 i=0; i<Adapt->GetNumberOfRepresentations(); ++i)
+					{
+						Repr = static_cast<HLSTimeline::FPlaybackAssetRepresentation*>(Adapt->GetRepresentationByIndex(i).Get());
+						// FIXME: This is setting the last codec info from all codecs onto all renditions. Not sure if this is the correct thing to do.
+						Repr->CodecInformation = StreamMetaData.CodecInformation;
+						//Repr->CDN;
+					}
+				}
+			}
+			else
+			{
+				// In absence of video assume this is audio-only
+				Manifest->AudioOnlyStreams.Push(vs);
+
+				FStreamMetadata StreamMetaData;
+				StreamMetaData.Bandwidth	  = vs->Bandwidth;
+				StreamMetaData.StreamUniqueID = vs->Internal.UniqueID;
+				StreamMetaData.PlaylistID     = UrlBuilder->ResolveWith(vs->URI);
+				for(int32 i=0; i<vs->StreamCodecInformationList.Num(); ++i)
+				{
+					if (vs->StreamCodecInformationList[i].IsAudioCodec())
+					{
+						StreamMetaData.CodecInformation = vs->StreamCodecInformationList[i];
+						break;
+					}
+				}
+				Manifest->StreamMetadataAudio.Push(StreamMetaData);
+
+				// Check if there is an audio adaptation set already due to an audio rendition group referenced by this variant.
+				HLSTimeline::FPlaybackAssetAdaptationSet* Adapt = nullptr;
+				if (bHaveAudioGroup)
+				{
+					Adapt = static_cast<HLSTimeline::FPlaybackAssetAdaptationSet*>(Asset->GetAdaptationSetByTypeAndUniqueIdentifier(EStreamType::Audio, vs->AudioGroupID).Get());
+					if (!Adapt)
+					{
+						return CreateErrorAndLog(FString::Printf(TEXT("EXT-X-STREAM-INF references AUDIO rendition group \"%s\" that has not been defined!"), *vs->AudioGroupID), ERRCODE_HLS_BUILDER_REFERENCED_GROUP_NOT_DEFINED, UEMEDIA_ERROR_FORMAT_ERROR);
+					}
+				}
+				if (!Adapt)
+				{
+					// FIXME: What do we do if there are several audio-only variant streams with different language codes? Or codecs?
+					Adapt = static_cast<HLSTimeline::FPlaybackAssetAdaptationSet*>(Asset->GetAdaptationSetByTypeAndUniqueIdentifier(EStreamType::Audio, "$audio$").Get());
+					if (!Adapt)
+					{
+						Adapt = new HLSTimeline::FPlaybackAssetAdaptationSet;
+						Adapt->UniqueIdentifier = TEXT("$audio$");
+						TSharedPtrTS<HLSTimeline::FPlaybackAssetAdaptationSet> IAdapt(Adapt);
+						Asset->AudioAdaptationSets.Push(IAdapt);
+					}
+				}
+				// Propagate audio codecs from variant to rendition since they are not specified there.
+				if (Adapt->Codecs.Find(StreamMetaData.CodecInformation.GetCodecSpecifierRFC6381()) == INDEX_NONE)
+				{
+					if (Adapt->Codecs.Len())
+					{
+						Adapt->Codecs.AppendChar(TCHAR(','));
+					}
+					Adapt->Codecs.Append(StreamMetaData.CodecInformation.GetCodecSpecifierRFC6381());
+				}
+				if (!bHaveAudioGroup)
+				{
+					HLSTimeline::FPlaybackAssetRepresentation* Repr = new HLSTimeline::FPlaybackAssetRepresentation;
+					Repr->UniqueIdentifier = LexToString(StreamMetaData.StreamUniqueID);
+					Repr->CodecInformation = StreamMetaData.CodecInformation;
+					Repr->CDN   		   = StreamMetaData.PlaylistID;
+					Repr->Bitrate   	   = StreamMetaData.Bandwidth;
+					Adapt->Representations.Push(TSharedPtrTS<IPlaybackAssetRepresentation>(Repr));
+
+					vs->Internal.AdaptationSetUniqueID  = Adapt->UniqueIdentifier;
+					vs->Internal.RepresentationUniqueID = Repr->UniqueIdentifier;
+					vs->Internal.CDN					= StreamMetaData.PlaylistID;
+				}
+				else
+				{
+					for(int32 i=0; i<Adapt->GetNumberOfRepresentations(); ++i)
+					{
+						HLSTimeline::FPlaybackAssetRepresentation* Repr = static_cast<HLSTimeline::FPlaybackAssetRepresentation*>(Adapt->GetRepresentationByIndex(i).Get());
+						// FIXME: This is setting the last codec info from all codecs onto all renditions. Not sure if this is the correct thing to do.
+						Repr->CodecInformation = StreamMetaData.CodecInformation;
+						//Repr->CDN;
+
+						// Get the range of renditions in this group and find those that do not have a dedicated playlist.
+						// We let those point back to this variant stream since that is what we effectively have to use.
+						TArray< const TSharedPtrTS<FManifestHLSInternal::FRendition>* > RenditionRange;
+						Manifest->AudioRenditions.MultiFindPointer(vs->AudioGroupID, RenditionRange);
+						if (!RenditionRange.IsEmpty())
+						{
+							for(int32 ii=0; ii < RenditionRange.Num(); ++ii)
+							{
+								const TSharedPtrTS<FManifestHLSInternal::FRendition>& Rendition = *RenditionRange[ii];
+								if (Rendition->URI.Len() == 0)
+								{
+									// When the rendition has no dedicated URL then it is merely informational and this audio-only variant is the stream itself to use.
+									Repr->UniqueIdentifier = LexToString(StreamMetaData.StreamUniqueID);
+									Repr->CodecInformation = StreamMetaData.CodecInformation;
+									Repr->CDN   		   = StreamMetaData.PlaylistID;
+									Repr->Bitrate   	   = StreamMetaData.Bandwidth;
+									vs->Internal.AdaptationSetUniqueID  = Adapt->UniqueIdentifier;
+									vs->Internal.RepresentationUniqueID = Repr->UniqueIdentifier;
+									vs->Internal.CDN					= StreamMetaData.PlaylistID;
+								}
+							}
+						}
+					}
+				}
+			}
+			// And to the lookup map.
+			Manifest->PlaylistIDMap.Add(vs->Internal.UniqueID, vs);
+		}
+	}
+
+	// Index the stream quality level map. Lower quality = lower index
+	int32 QualityIndex = 0;
+	for (TMap<int32, int32>::TIterator It = Manifest->BandwidthToQualityIndex.CreateIterator(); It; ++It)
+	{
+		It.Value() = QualityIndex++;
+	}
+
+	return FErrorDetail();
+}
+
+
+
+FErrorDetail FManifestBuilderHLS::GetInitialPlaylistLoadRequests(TArray<FPlaylistLoadRequestHLS>& OutRequests, TSharedPtrTS<FManifestHLSInternal> Manifest, const FStreamPreferences& Preferences, const FParamDict& Options)
+{
+	OutRequests.Empty();
+
+	if (!Manifest.IsValid())
+	{
+		return CreateErrorAndLog(FString::Printf(TEXT("No master playlist")), ERRCODE_HLS_BUILDER_INTERNAL, UEMEDIA_ERROR_BAD_ARGUMENTS);
+	}
+	if (Manifest->VariantStreams.Num() == 0 && Manifest->AudioOnlyStreams.Num() == 0)
+	{
+		// NOTE: There might be a some text/caption-only variant but is this really a valid scenario?
+		return CreateErrorAndLog(FString::Printf(TEXT("No variant playlists in master playlist")), ERRCODE_HLS_BUILDER_NO_VARIANTS_IN_MASTER_PLAYLIST, UEMEDIA_ERROR_FORMAT_ERROR);
+	}
+
+	TSharedPtrTS<FManifestHLSInternal::FVariantStream> StartingVariantStream;
+	int32 InitialBitrate = (int32) Options.GetValue(IPlaylistReaderHLS::OptionKeyInitialBitrate).SafeGetInt64(0);
+	bool bIsAudioOnlyVariant = false;
+
+	if (Manifest->VariantStreams.Num())
+	{
+		// First we need to find the variant stream not exceeding the initial bitrate.
+		if (InitialBitrate <= 0)
+		{
+			// By default we use the bitrate of the first variant stream. The variant streams are not sorted by bandwidth in the master playlist on purpose
+			// and the first one is to be used in absence of any other criteria.
+			InitialBitrate = Manifest->VariantStreams[0]->Bandwidth;
+		}
+
+	// TODO: Consider additional options, like maximum resolution. While not strictly necessary it will save loading a playlist we won't be needing for streaming
+
+		for(int32 nVariant=0; nVariant<Manifest->VariantStreams.Num(); ++nVariant)
+		{
+			if (Manifest->VariantStreams[nVariant]->Bandwidth <= InitialBitrate)
+			{
+				if (StartingVariantStream == nullptr || StartingVariantStream->Bandwidth < Manifest->VariantStreams[nVariant]->Bandwidth)
+				{
+					StartingVariantStream = Manifest->VariantStreams[nVariant];
+				}
+			}
+		}
+		// No variant stream found? Pick the one with lowest bitrate then.
+		if (!StartingVariantStream)
+		{
+			for(int32 nVariant = 0; nVariant < Manifest->VariantStreams.Num(); ++nVariant)
+			{
+				if (StartingVariantStream == nullptr || StartingVariantStream->Bandwidth > Manifest->VariantStreams[nVariant]->Bandwidth)
+				{
+					StartingVariantStream = Manifest->VariantStreams[nVariant];
+				}
+			}
+		}
+	}
+	else if (Manifest->AudioOnlyStreams.Num())
+	{
+		// TODO: This might need to check for additional options like a selected language.
+		StartingVariantStream = Manifest->AudioOnlyStreams[0];
+		bIsAudioOnlyVariant = true;
+	}
+
+	// Still nothing found (can only happen when we use additional criteria to filter streams)
+	if (!StartingVariantStream.IsValid())
+	{
+		return CreateErrorAndLog(FString::Printf(TEXT("No usable variant stream found in master playlist")), ERRCODE_HLS_BUILDER_NO_USABLE_VARIANT_FOUND, UEMEDIA_ERROR_FORMAT_ERROR);
+	}
+
+	// Add the variant playlist to the list of playlists to be fetched.
+	TUniquePtr<IURLParser> UrlBuilder(IURLParser::Create());
+	UrlBuilder->ParseURL(Manifest->MasterPlaylistVars.PlaylistLoadRequest.URL);
+	FPlaylistLoadRequestHLS& VariantRequest = OutRequests.AddDefaulted_GetRef();
+	check(StartingVariantStream->Internal.LoadState == FManifestHLSInternal::FPlaylistBase::FInternal::ELoadState::NotLoaded);
+	StartingVariantStream->Internal.LoadState = FManifestHLSInternal::FPlaylistBase::FInternal::ELoadState::Pending;
+	VariantRequest.InternalUniqueID 		  = StartingVariantStream->Internal.UniqueID;
+	VariantRequest.RequestedAtTime  		  = PlayerSessionServices->GetSynchronizedUTCTime()->GetTime();
+	VariantRequest.URL  					  = UrlBuilder->ResolveWith(StartingVariantStream->URI);
+	VariantRequest.LoadType 				  = FPlaylistLoadRequestHLS::ELoadType::Initial;
+	VariantRequest.AdaptationSetUniqueID	  = StartingVariantStream->Internal.AdaptationSetUniqueID;
+	VariantRequest.RepresentationUniqueID     = StartingVariantStream->Internal.RepresentationUniqueID;
+	VariantRequest.CDN  					  = StartingVariantStream->Internal.CDN;
+
+	// Check for possible alternate renditions.
+// TODO: FIXME: Extend this to additional stream types, not only audio!
+	if (StartingVariantStream->AudioGroupID.Len())
+	{
+		// Get the range of renditions
+		TArray< const TSharedPtrTS<FManifestHLSInternal::FRendition>* > RenditionRange;
+		Manifest->AudioRenditions.MultiFindPointer(StartingVariantStream->AudioGroupID, RenditionRange);
+		if (!RenditionRange.IsEmpty())
+		{
+			for(int32 ii=0; ii < RenditionRange.Num(); ++ii)
+			{
+				const TSharedPtrTS<FManifestHLSInternal::FRendition>& Rendition = *RenditionRange[ii];
+
+			// TODO: FIXME: look at all renditions and find one best matching the user preference (language and/or codec, possibly characteristics)
+			//       especially important when we the variant stream is an audio-only variant already but not the language we want to play.
+
+				// Does the rendition have a URI, indicating this to be a separate stream?
+				// If the URI is empty then the rendition merely has informational value.
+				if (Rendition->URI.Len())
+				{
+				// NOTE: If the rendition type is CLOSED-CAPTIONS then it MUST NOT have a URI!!!
+					FPlaylistLoadRequestHLS& RenditionRequest = OutRequests.AddDefaulted_GetRef();
+					check(Rendition->Internal.LoadState == FManifestHLSInternal::FPlaylistBase::FInternal::ELoadState::NotLoaded);
+					Rendition->Internal.LoadState   		= FManifestHLSInternal::FPlaylistBase::FInternal::ELoadState::Pending;
+					RenditionRequest.InternalUniqueID   	= Rendition->Internal.UniqueID;
+					RenditionRequest.RequestedAtTime		= PlayerSessionServices->GetSynchronizedUTCTime()->GetTime();
+					RenditionRequest.URL					= UrlBuilder->ResolveWith(Rendition->URI);
+					RenditionRequest.LoadType   			= FPlaylistLoadRequestHLS::ELoadType::Initial;
+					RenditionRequest.AdaptationSetUniqueID  = Rendition->Internal.AdaptationSetUniqueID;
+					RenditionRequest.RepresentationUniqueID = Rendition->Internal.RepresentationUniqueID;
+					RenditionRequest.CDN					= Rendition->Internal.CDN;
+
+			// TODO: we pick just the first one from all the renditions for now!
+					break;
+				}
+				else
+				{
+					//return CreateErrorAndLog(FString::Printf("%s rendition in GROUP-ID \"%s\" has no URI", "AUDIO", StartingVariantStream->AudioGroupID.c_str()), ERRCODE_HLS_BUILDER_RENDITION_URI_MISSING, UEMEDIA_ERROR_FORMAT_ERROR);
+				}
+			}
+		}
+		else
+		{
+			// No match found. This is an error.
+			return CreateErrorAndLog(FString::Printf(TEXT("No %s rendition found for GROUP-ID \"%s\""), TEXT("AUDIO"), *StartingVariantStream->AudioGroupID), ERRCODE_HLS_BUILDER_RENDITION_NOT_FOUND_IN_GROUP, UEMEDIA_ERROR_FORMAT_ERROR);
+		}
+	}
+
+	return FErrorDetail();
+}
+
+UEMediaError FManifestBuilderHLS::UpdateFailedInitialPlaylistLoadRequest(FPlaylistLoadRequestHLS& InOutFailedRequest, const HTTP::FConnectionInfo* ConnectionInfo, TSharedPtrTS<HTTP::FRetryInfo> PreviousAttempts, const FTimeValue& BlacklistUntilUTC, TSharedPtrTS<FManifestHLSInternal> Manifest, const FStreamPreferences& Preferences, const FParamDict& Options)
+{
+	// Is the failed request a video variant?
+	for(int32 i=0, iMax=Manifest->VariantStreams.Num(); i<iMax; ++i)
+	{
+		TSharedPtrTS<FManifestHLSInternal::FVariantStream> VideoVariant = Manifest->VariantStreams[i];
+		if (VideoVariant->Internal.UniqueID == InOutFailedRequest.InternalUniqueID)
+		{
+			// Set load state back to NotLoaded
+			VideoVariant->Internal.PlaylistLoadRequest = InOutFailedRequest;
+			if (ConnectionInfo)
+			{
+				VideoVariant->Internal.PlaylistLoadRequest.URL = ConnectionInfo->EffectiveURL;
+			}
+			VideoVariant->Internal.ExpiresAtTime.SetToPositiveInfinity();
+			VideoVariant->Internal.LoadState = FManifestHLSInternal::FPlaylistBase::FInternal::ELoadState::NotLoaded;
+
+			// Try to find an alternate video stream to fetch.
+			int32 FailedBitrate = VideoVariant->Bandwidth;
+			TSharedPtrTS<FManifestHLSInternal::FVariantStream> AlternateVariantStream;
+			for(int32 nVariant=0; nVariant<Manifest->VariantStreams.Num(); ++nVariant)
+			{
+				// Skip over ourselves here. We already failed.
+				if (nVariant == i)
+				{
+					continue;
+				}
+				if (Manifest->VariantStreams[nVariant]->Bandwidth < FailedBitrate)
+				{
+					if (!AlternateVariantStream.IsValid() || AlternateVariantStream->Bandwidth < Manifest->VariantStreams[nVariant]->Bandwidth)
+					{
+						if (!Manifest->VariantStreams[nVariant]->Internal.Blacklisted.Get())
+						{
+							AlternateVariantStream = Manifest->VariantStreams[nVariant];
+						}
+					}
+				}
+			}
+			if (!AlternateVariantStream)
+			{
+				for(int32 nVariant = 0; nVariant < Manifest->VariantStreams.Num(); ++nVariant)
+				{
+					// Skip over ourselves here. We already failed.
+					if (nVariant == i)
+					{
+						continue;
+					}
+					if (!AlternateVariantStream.IsValid() || AlternateVariantStream->Bandwidth > Manifest->VariantStreams[nVariant]->Bandwidth)
+					{
+						if (!Manifest->VariantStreams[nVariant]->Internal.Blacklisted.Get())
+						{
+							AlternateVariantStream = Manifest->VariantStreams[nVariant];
+						}
+					}
+				}
+			}
+			// FIXME: We could also check if there is any variant that isn't blacklisted yet (or whose blacklist can be lifted again).
+			if (!AlternateVariantStream.IsValid())
+			{
+				return UEMEDIA_ERROR_END_OF_STREAM;
+			}
+
+			// Blacklist the failed variant only when we have found an alternative. Otherwise we may want to retry the failed stream.
+			VideoVariant->Internal.Blacklisted = MakeSharedTS<FManifestHLSInternal::FBlacklist>();
+			VideoVariant->Internal.Blacklisted->PreviousAttempts				= PreviousAttempts;
+			VideoVariant->Internal.Blacklisted->BecomesAvailableAgainAtUTC  	= BlacklistUntilUTC;
+			VideoVariant->Internal.Blacklisted->AssetIDs.AssetUniqueID  		= Manifest->PlaybackTimeline->GetMediaAssetByIndex(0)->GetUniqueIdentifier();
+			VideoVariant->Internal.Blacklisted->AssetIDs.AdaptationSetUniqueID  = InOutFailedRequest.AdaptationSetUniqueID;
+			VideoVariant->Internal.Blacklisted->AssetIDs.RepresentationUniqueID = InOutFailedRequest.RepresentationUniqueID;
+			VideoVariant->Internal.Blacklisted->AssetIDs.CDN					= InOutFailedRequest.CDN;
+			// Tell the stream selector that this stream is temporarily unavailable.
+			TSharedPtrTS<IAdaptiveStreamSelector> StreamSelector(PlayerSessionServices->GetStreamSelector());
+			StreamSelector->MarkStreamAsUnavailable(VideoVariant->Internal.Blacklisted->AssetIDs);
+
+			TUniquePtr<IURLParser> UrlBuilder(IURLParser::Create());
+			UrlBuilder->ParseURL(Manifest->MasterPlaylistVars.PlaylistLoadRequest.URL);
+			AlternateVariantStream->Internal.LoadState = FManifestHLSInternal::FPlaylistBase::FInternal::ELoadState::Pending;
+			// Update the changed request parameters.
+			InOutFailedRequest.InternalUniqueID 		  = AlternateVariantStream->Internal.UniqueID;
+			InOutFailedRequest.URL  					  = UrlBuilder->ResolveWith(AlternateVariantStream->URI);
+			InOutFailedRequest.AdaptationSetUniqueID	  = AlternateVariantStream->Internal.AdaptationSetUniqueID;
+			InOutFailedRequest.RepresentationUniqueID     = AlternateVariantStream->Internal.RepresentationUniqueID;
+			InOutFailedRequest.CDN  					  = AlternateVariantStream->Internal.CDN;
+
+			return UEMEDIA_ERROR_OK;
+		}
+	}
+
+	// Audio rendition?
+	for (TMultiMap<FString, TSharedPtrTS<FManifestHLSInternal::FRendition>>::TConstIterator It = Manifest->AudioRenditions.CreateConstIterator(); It; ++It)
+	{
+		TSharedPtrTS<FManifestHLSInternal::FRendition> AudioRendition = It.Value();
+		if (AudioRendition->Internal.UniqueID == InOutFailedRequest.InternalUniqueID)
+		{
+			// Set load state back to NotLoaded
+			AudioRendition->Internal.PlaylistLoadRequest = InOutFailedRequest;
+			if (ConnectionInfo)
+			{
+				AudioRendition->Internal.PlaylistLoadRequest.URL = ConnectionInfo->EffectiveURL;
+			}
+			AudioRendition->Internal.ExpiresAtTime.SetToPositiveInfinity();
+			AudioRendition->Internal.LoadState = FManifestHLSInternal::FPlaylistBase::FInternal::ELoadState::NotLoaded;
+
+			// Try to find an alternate audio stream to fetch.
+
+		// FIXME: For now we fail for audio renditions. There are typically no other renditions to choose from.
+			return UEMEDIA_ERROR_END_OF_STREAM;
+		// Implement this later when there actually are different auiod variants to choose from.
+			#if 0
+				// Blacklist the failed variant only when we have found an alternative. Otherwise we may want to retry the failed stream.
+				AudioRendition->Internal.Blacklisted = MakeSharedTS<FManifestHLSInternal::FBlacklist>();
+				AudioRendition->Internal.Blacklisted->PreviousAttempts  			  = PreviousAttempts;
+				AudioRendition->Internal.Blacklisted->BecomesAvailableAgainAtUTC	  = BlacklistUntilUTC;
+				AudioRendition->Internal.Blacklisted->AssetIDs.AssetUniqueID		  = Manifest->PlaybackTimeline->GetMediaAssetByIndex(0)->GetUniqueIdentifier();
+				AudioRendition->Internal.Blacklisted->AssetIDs.AdaptationSetUniqueID  = InOutFailedRequest.AdaptationSetUniqueID;
+				AudioRendition->Internal.Blacklisted->AssetIDs.RepresentationUniqueID = InOutFailedRequest.RepresentationUniqueID;
+				AudioRendition->Internal.Blacklisted->AssetIDs.CDN  				  = InOutFailedRequest.CDN;
+				// Tell the stream selector that this stream is temporarily unavailable.
+				TSharedPtrTS<IAdaptiveStreamSelector> StreamSelector(PlayerSessionServices->GetStreamSelector());
+				StreamSelector->MarkStreamAsUnavailable(AudioRendition->Internal.Blacklisted->AssetIDs);
+				return UEMEDIA_ERROR_OK;
+			#endif
+		}
+	}
+
+	return UEMEDIA_ERROR_END_OF_STREAM;
+}
+
+
+void FManifestBuilderHLS::SetVariantPlaylistFailure(TSharedPtrTS<FManifestHLSInternal> InHLSPlaylist, const FPlaylistLoadRequestHLS& SourceRequest, const HTTP::FConnectionInfo* ConnectionInfo, TSharedPtrTS<HTTP::FRetryInfo> PreviousAttempts, const FTimeValue& BlacklistUntilUTC)
+{
+	// Find the variant or rendition this playlist updates.
+	InHLSPlaylist->LockPlaylists();
+
+	TWeakPtrTS<FManifestHLSInternal::FPlaylistBase>* PlaylistID = InHLSPlaylist->PlaylistIDMap.Find(SourceRequest.InternalUniqueID);
+	if(PlaylistID != nullptr)
+	{
+		TSharedPtrTS<FManifestHLSInternal::FPlaylistBase> Playlist = PlaylistID->Pin();
+		if (Playlist.IsValid())
+		{
+			Playlist->Internal.PlaylistLoadRequest = SourceRequest;
+			Playlist->Internal.LoadState		   = FManifestHLSInternal::FPlaylistBase::FInternal::ELoadState::NotLoaded;
+			Playlist->Internal.bReloadTriggered    = false;
+			Playlist->Internal.bNewlySelected      = false;
+			if (ConnectionInfo)
+			{
+				Playlist->Internal.PlaylistLoadRequest.URL = ConnectionInfo->EffectiveURL;
+			}
+			Playlist->Internal.ExpiresAtTime.SetToPositiveInfinity();
+
+			Playlist->Internal.Blacklisted = MakeSharedTS<FManifestHLSInternal::FBlacklist>();
+			Playlist->Internal.Blacklisted->PreviousAttempts		   = PreviousAttempts;
+			Playlist->Internal.Blacklisted->BecomesAvailableAgainAtUTC = BlacklistUntilUTC;
+
+			Playlist->Internal.Blacklisted->AssetIDs.AssetUniqueID  		= InHLSPlaylist->PlaybackTimeline->GetMediaAssetByIndex(0)->GetUniqueIdentifier();
+			Playlist->Internal.Blacklisted->AssetIDs.AdaptationSetUniqueID  = SourceRequest.AdaptationSetUniqueID;
+			Playlist->Internal.Blacklisted->AssetIDs.RepresentationUniqueID = SourceRequest.RepresentationUniqueID;
+			Playlist->Internal.Blacklisted->AssetIDs.CDN					= SourceRequest.CDN;
+
+			// Tell the stream selector that this stream is temporarily unavailable.
+			TSharedPtrTS<IAdaptiveStreamSelector> StreamSelector(PlayerSessionServices->GetStreamSelector());
+			StreamSelector->MarkStreamAsUnavailable(Playlist->Internal.Blacklisted->AssetIDs);
+		}
+	}
+	InHLSPlaylist->UnlockPlaylists();
+}
+
+
+
+
+
+FErrorDetail FManifestBuilderHLS::UpdateFromVariantPlaylist(TSharedPtrTS<FManifestHLSInternal> InOutHLSPlaylist, const HLSPlaylistParser::FPlaylist& VariantPlaylist, const FPlaylistLoadRequestHLS& SourceRequest, const HTTP::FConnectionInfo* ConnectionInfo, uint32 ResponseCRC, const FStreamPreferences& Preferences, const FParamDict& Options)
+{
+	// Find the variant or rendition this playlist updates.
+	TSharedPtrTS<FManifestHLSInternal::FPlaylistBase> Playlist;
+	InOutHLSPlaylist->LockPlaylists();
+	TWeakPtrTS<FManifestHLSInternal::FPlaylistBase>* PlaylistID = InOutHLSPlaylist->PlaylistIDMap.Find(SourceRequest.InternalUniqueID);
+	if(PlaylistID != nullptr)
+	{
+		Playlist = PlaylistID->Pin();
+	}
+	InOutHLSPlaylist->UnlockPlaylists();
+
+	if (!Playlist.IsValid())
+	{
+		LogMessage(IInfoLog::ELevel::Warning, FString::Printf(TEXT("No variant or rendition found to be updated with new incoming playlist. Ignoring.")));
+		return FErrorDetail();
+	}
+
+
+	// Did the playlist response not change from last time?
+	if (ResponseCRC == SourceRequest.LastUpdateCRC32)
+	{
+		FManifestHLSInternal::ScopedLockPlaylists lock(InOutHLSPlaylist);
+		check(Playlist->Internal.MediaStream.IsValid());
+		if (Playlist->Internal.MediaStream.IsValid())
+		{
+			Playlist->Internal.PlaylistLoadRequest = SourceRequest;
+			check(ConnectionInfo);
+			if (ConnectionInfo)
+			{
+				Playlist->Internal.PlaylistLoadRequest.URL = ConnectionInfo->EffectiveURL;
+			}
+			Playlist->Internal.ExpiresAtTime	   = SourceRequest.RequestedAtTime + (Playlist->Internal.MediaStream->TargetDuration / 2);
+			Playlist->Internal.bReloadTriggered    = false;
+			Playlist->Internal.Blacklisted.Reset();
+		}
+		else
+		{
+			return CreateErrorAndLog(FString::Printf(TEXT("Unchanged playlist update did not find original stream!")), ERRCODE_HLS_BUILDER_ID_NOT_FOUND, UEMEDIA_ERROR_INTERNAL);
+		}
+	}
+	else
+	{
+		if (VariantPlaylist.Type != HLSPlaylistParser::EPlaylistType::Media)
+		{
+			return CreateErrorAndLog(FString::Printf(TEXT("Not a variant playlist")), ERRCODE_HLS_BUILDER_NOT_A_VARIANT_PLAYLIST, UEMEDIA_ERROR_FORMAT_ERROR);
+		}
+
+		TUniquePtr<FManifestHLSInternal::FMediaStream> MediaStream(new FManifestHLSInternal::FMediaStream);
+
+		// EXT-X-INDEPENDENT-SEGMENTS from the master playlist (only if set there!) applies to all segments in a variant
+		if (InOutHLSPlaylist->bHasIndependentSegments)
+		{
+			MediaStream->bHasIndependentSegments = true;
+		}
+		else if (VariantPlaylist.ContainsTag(HLSPlaylistParser::ExtXIndependentSegments))
+		{
+			MediaStream->bHasIndependentSegments = true;
+		}
+
+		// TODO: get the EXT-X-START either from this playlist or the master playlist (not sure which one takes precendece)
+		//       The attribute is not used at the moment anyway...
+
+
+		// Helper lambdas
+	/*
+		auto GetTagAttributeString = [&VariantPlaylist](FString& OutValue, const FString& Tag, const FString& Attribute) -> bool
+		{
+			FString AttributeValueStr;
+			if (VariantPlaylist.GetTagAttributeValue(Tag, Attribute, AttributeValueStr) == HLSPlaylistParser::EPlaylistError::None)
+			{
+				OutValue = AttributeValueStr;
+				return true;
+			}
+			return false;
+		};
+	*/
+		auto GetTagString = [&VariantPlaylist](FString& OutValue, const FString& Tag) -> bool
+		{
+			FString AttributeValueStr;
+			if (VariantPlaylist.GetTagValue(Tag, AttributeValueStr) == HLSPlaylistParser::EPlaylistError::None)
+			{
+				OutValue = AttributeValueStr;
+				return true;
+			}
+			return false;
+		};
+
+
+		FString AttrValue;
+
+		// Get target duration
+		if (GetTagString(AttrValue, HLSPlaylistParser::ExtXTargetDuration))
+		{
+			int32 TempValue;
+			LexFromString(TempValue, *AttrValue);
+			MediaStream->TargetDuration.SetFromSeconds(TempValue);
+		}
+		else
+		{
+			// That is a mandatory tag!
+			return CreateErrorAndLog(FString::Printf(TEXT("Required EXT-X_TARGETDURATION tag is missing")), ERRCODE_HLS_BUILDER_MISSING_EXTX_TARGETDURATION, UEMEDIA_ERROR_FORMAT_ERROR);
+		}
+
+		// Get the type of playlist
+		if (GetTagString(AttrValue, HLSPlaylistParser::ExtXPlaylistType))
+		{
+			MediaStream->PlaylistType = AttrValue == TEXT("VOD") ? FManifestHLSInternal::FMediaStream::EPlaylistType::VOD : AttrValue == TEXT("EVENT") ? FManifestHLSInternal::FMediaStream::EPlaylistType::Event : FManifestHLSInternal::FMediaStream::EPlaylistType::Live;
+		}
+		else
+		{
+			MediaStream->PlaylistType = FManifestHLSInternal::FMediaStream::EPlaylistType::Live;
+		}
+
+		// EXT-X-ENDLIST present?
+		if (VariantPlaylist.ContainsTag(HLSPlaylistParser::ExtXEndlist))
+		{
+			MediaStream->bHasListEnd = true;
+		}
+
+		// EXT-X-MEDIA-SEQUENCE present?
+		if (GetTagString(AttrValue, HLSPlaylistParser::ExtXMediaSequence))
+		{
+			LexFromString(MediaStream->MediaSequence, *AttrValue);
+		}
+
+		// EXT-X-DISCONTINUITY-SEQUENCE present?
+		if (GetTagString(AttrValue, HLSPlaylistParser::ExtXDiscontinuitySequence))
+		{
+			LexFromString(MediaStream->DiscontinuitySequence, *AttrValue);
+		}
+
+		// EXT-X-I-FRAMES-ONLY present?
+		if (VariantPlaylist.ContainsTag(HLSPlaylistParser::ExtXIFramesOnly))
+		{
+			MediaStream->bIsIFramesOnly = true;
+			return CreateErrorAndLog(FString::Printf(TEXT("EXT-X-I-FRAMES-ONLY is currently not supported")), ERRCODE_HLS_BUILDER_UNSUPPORTED_FEATURE, UEMEDIA_ERROR_FORMAT_ERROR);
+		}
+
+		// Initialization segment.
+		TSharedPtrTS<FManifestHLSInternal::FMediaStream::FInitSegmentInfo>	InitSegmentInfo;
+
+		// License key.
+		TSharedPtr<FManifestHLSInternal::FMediaStream::FDRMKeyInfo, ESPMode::ThreadSafe>	LicenseKeyInfo;
+
+		// Now process the media segments
+		FTimeValue LastSegmentDuration;
+		MediaStream->TotalAccumulatedSegmentDuration.SetToZero();
+		int64 ByteRangeStartPos = 0;
+		const TArray<HLSPlaylistParser::FMediaSegment>& MediaSegmentList = VariantPlaylist.GetSegments();
+		TArray<FTimespan> SeekablePositions;
+		if (MediaSegmentList.Num())
+		{
+			TArray<FString> SplitResults;
+			TUniquePtr<IURLParser> UrlBuilder(IURLParser::Create());
+			UrlBuilder->ParseURL(SourceRequest.URL);
+
+			int64 NextSequenceNum      = MediaStream->MediaSequence;
+			int64 NextDiscontinuityNum = MediaStream->DiscontinuitySequence;
+			int64 NextRangeStartOffset = 0;
+			MediaStream->SegmentList.Reserve(MediaSegmentList.Num());
+			FTimeValue AccumulatedDuration(FTimeValue::GetZero());
+			for(int32 nSeg=0,nSegMax=MediaSegmentList.Num(); nSeg<nSegMax; ++nSeg)
+			{
+				const HLSPlaylistParser::FMediaSegment& MediaSegment = MediaSegmentList[nSeg];
+
+				auto GetSegmentTagString = [&MediaSegment](FString& OutValue, const FString& Tag) -> bool
+				{
+					FString AttributeValueStr;
+					if (MediaSegment.GetTagValue(Tag, AttributeValueStr) == HLSPlaylistParser::EPlaylistError::None)
+					{
+						OutValue = AttributeValueStr;
+						return true;
+					}
+					return false;
+				};
+				auto GetSegmentAttributeString = [&MediaSegment](FString& OutValue, const FString& Tag, const FString& Attribute) -> bool
+				{
+					FString AttributeValueStr;
+					if (MediaSegment.GetTagAttributeValue(Tag, Attribute, AttributeValueStr) == HLSPlaylistParser::EPlaylistError::None)
+					{
+						OutValue = AttributeValueStr;
+						return true;
+					}
+					return false;
+				};
+
+				FManifestHLSInternal::FMediaStream::FMediaSegment& Seg = MediaStream->SegmentList.AddDefaulted_GetRef();
+
+				Seg.RelativeStartTime = AccumulatedDuration;
+
+				// Get the duration
+				if (GetSegmentTagString(AttrValue, HLSPlaylistParser::ExtINF))
+				{
+					SplitOnCommaOrSpace(SplitResults, AttrValue);
+					if (SplitResults.Num())
+					{
+						Seg.Duration.SetFromTimeFraction(FTimeFraction().SetFromFloatString(SplitResults[0]));
+						AccumulatedDuration += Seg.Duration;
+					}
+					else
+					{
+						return CreateErrorAndLog(FString::Printf(TEXT("EXTINF \"%s\" failed to parse!"), *AttrValue), ERRCODE_HLS_BUILDER_ATTRIBUTE_PARSE_ERROR, UEMEDIA_ERROR_FORMAT_ERROR);
+					}
+				}
+				else
+				{
+					// That is a mandatory tag!
+					return CreateErrorAndLog(FString::Printf(TEXT("Required EXTINF tag is missing")), ERRCODE_HLS_BUILDER_MISSING_EXTINF, UEMEDIA_ERROR_FORMAT_ERROR);
+				}
+
+				// Check if there is an EXT-X-DISCONTINUITY
+				if (MediaSegment.ContainsTag(HLSPlaylistParser::ExtXDiscontinuity))
+				{
+					// Increment the discontinuity counter.
+					++NextDiscontinuityNum;
+				}
+				Seg.DiscontinuityCount = NextDiscontinuityNum;
+
+				// Set sequence number and increment by one for the next segment.
+				Seg.SequenceNumber     = NextSequenceNum++;
+
+				// Set the URL
+				Seg.URI 			   = MediaSegment.URL;
+				// Byte range?
+				if (GetSegmentTagString(AttrValue, HLSPlaylistParser::ExtXByteRange))
+				{
+					TArray<FString> brParams;
+					StringHelpers::SplitByDelimiter(brParams, AttrValue, TEXT("@"));
+					if (brParams.Num() > 1)
+					{
+						LexFromString(ByteRangeStartPos, *brParams[1]);
+					}
+					int64 brLen;
+					LexFromString(brLen, *brParams[0]);
+
+					Seg.ByteRange.Start = ByteRangeStartPos;
+					Seg.ByteRange.End = ByteRangeStartPos + brLen - 1;
+					ByteRangeStartPos += brLen;
+				}
+
+				// Regular segment
+				Seg.bIsPrefetch = false;
+
+				// Is there an EXT-X-PROGRAM-DATE-TIME ?
+				if (GetSegmentTagString(AttrValue, HLSPlaylistParser::ExtXProgramDateTime))
+				{
+					if (ISO8601::ParseDateTime(Seg.AbsoluteDateTime, AttrValue)== UEMEDIA_ERROR_OK)
+					{
+						if (nSeg)
+						{
+							FTimeValue Diff = Seg.AbsoluteDateTime - (MediaStream->SegmentList[nSeg - 1].AbsoluteDateTime + MediaStream->SegmentList[nSeg - 1].Duration);
+							// Overlap or gap of more than half a second?
+							if (Utils::AbsoluteValue(Diff.GetAsSeconds()) > 0.5)
+							{
+								LogMessage(IInfoLog::ELevel::Warning, FString::Printf(TEXT("Warning: %s on timeline detected. %#.3f seconds between segment %d and %d\n"), Diff<FTimeValue::GetZero()?TEXT("Overlap"):TEXT("Gap"), Diff.GetAsSeconds(), nSeg-1, nSeg));
+							}
+						}
+					}
+					else
+					{
+						return CreateErrorAndLog(FString::Printf(TEXT("Failed to parse EXT-X-PROGRAM-DATE-TIME \"%s\"!"), *AttrValue), ERRCODE_HLS_BUILDER_ATTRIBUTE_PARSE_ERROR, UEMEDIA_ERROR_FORMAT_ERROR);
+					}
+				}
+				else
+				{
+					if (nSeg)
+					{
+						Seg.AbsoluteDateTime = MediaStream->SegmentList[nSeg - 1].AbsoluteDateTime + MediaStream->SegmentList[nSeg - 1].Duration;
+					}
+					else
+					{
+						// In absence of a preceeding segment we start the absolute time at zero.
+						Seg.AbsoluteDateTime.SetToZero();
+					}
+				}
+
+
+				// DRM license key
+				if (GetSegmentAttributeString(AttrValue, HLSPlaylistParser::ExtXKey, TEXT("METHOD")))
+				{
+					// If the method is "NONE" the following segments are no longer encrypted.
+					if (AttrValue == TEXT("NONE"))
+					{
+						LicenseKeyInfo.Reset();
+					}
+					else if (AttrValue == TEXT("AES-128") || AttrValue == TEXT("SAMPLE-AES"))
+					{
+						// Take note that at least one segment in this playlist is encrypted.
+						MediaStream->bHasEncryptedSegments = true;
+						FManifestHLSInternal::FMediaStream::FDRMKeyInfo::EMethod Method = AttrValue == TEXT("AES-128") ? FManifestHLSInternal::FMediaStream::FDRMKeyInfo::EMethod::AES128 : FManifestHLSInternal::FMediaStream::FDRMKeyInfo::EMethod::SampleAES;
+						FString IV, URI;
+						GetSegmentAttributeString(URI, HLSPlaylistParser::ExtXKey, TEXT("URI"));
+						GetSegmentAttributeString(IV, HLSPlaylistParser::ExtXKey, TEXT("IV"));
+						// Resolve the URL if it is relative.
+						FString LicenseKeyURL = UrlBuilder->ResolveWith(URI);
+						// Is there a change in attributes?
+						if (LicenseKeyInfo.IsValid() && (LicenseKeyInfo->Method != Method || LicenseKeyInfo->URI != LicenseKeyURL || LicenseKeyInfo->IV != IV))
+						{
+							LicenseKeyInfo.Reset();
+						}
+						if (!LicenseKeyInfo.IsValid())
+						{
+							LicenseKeyInfo = MakeShared<FManifestHLSInternal::FMediaStream::FDRMKeyInfo, ESPMode::ThreadSafe>();
+							LicenseKeyInfo->Method = Method;
+							LicenseKeyInfo->URI = LicenseKeyURL;
+							LicenseKeyInfo->IV = IV;
+						}
+					}
+					else
+					{
+						return CreateErrorAndLog(FString::Printf(TEXT("Unsupported EXT-X-KEY method \"%s\"!"), *AttrValue), ERRCODE_HLS_BUILDER_ATTRIBUTE_INVALID_VALUE, UEMEDIA_ERROR_FORMAT_ERROR);
+					}
+				}
+				Seg.DRMKeyInfo = LicenseKeyInfo;
+
+
+				// Is there a new init segment specified with EXT-X-MAP ?
+				if (GetSegmentAttributeString(AttrValue, HLSPlaylistParser::ExtXMap, TEXT("URI")))
+				{
+					FString InitSegmentURL = UrlBuilder->ResolveWith(AttrValue);
+					InitSegmentInfo = MakeSharedTS<FManifestHLSInternal::FMediaStream::FInitSegmentInfo>();
+					InitSegmentInfo->URI = InitSegmentURL;
+// FIXME: The init segment could/can have a separate EXT-X-KEY applied to it!!
+//        For the moment init segments are not supposed to be encrypted!
+//					InitSegmentInfo->DRMKeyInfo = LicenseKeyInfo;
+					// Is there a byte range?
+					if (GetSegmentAttributeString(AttrValue, HLSPlaylistParser::ExtXMap, TEXT("BYTERANGE")))
+					{
+						TArray<FString> brParams;
+						StringHelpers::SplitByDelimiter(brParams, AttrValue, TEXT("@"));
+						int64 brStart = 0;
+						if (brParams.Num() > 1)
+						{
+							LexFromString(brStart, *brParams[1]);
+						}
+						int64 brLen;
+						LexFromString(brLen, *brParams[0]);
+						InitSegmentInfo->ByteRange.Start = brStart;
+						InitSegmentInfo->ByteRange.End = brStart + brLen - 1;
+					}
+				}
+				Seg.InitSegmentInfo = InitSegmentInfo;
+
+				LastSegmentDuration = Seg.Duration;
+
+				// TODO: EXT-X-DATERANGE  (in the very distant future. this is rather complex and very likely not needed)
+
+
+				// Add the absolute time to the array of seek positions.
+				SeekablePositions.Emplace(FTimespan(Seg.AbsoluteDateTime.GetAsHNS()));
+			}
+
+
+			// Check for EXT-X-PREFETCH as per [REDACTED_PROFILE_URL]/hlsjs-rfcs/blob/lhls-spec/proposals/0001-lhls.md
+// FIXME: This is currently a playlist tag instead of a segment. We'll change that later.
+			if (GetTagString(AttrValue, HLSPlaylistParser::LHLS::ExtXPrefetch))
+			{
+				// Add a new segment entry.
+				FManifestHLSInternal::FMediaStream::FMediaSegment PrefetchSeg;
+
+				PrefetchSeg.RelativeStartTime  = AccumulatedDuration;
+				PrefetchSeg.Duration		   = MediaStream->TargetDuration;	// For lack of a better value.
+// FIXME: The discontinuity count is subject to possible EXT-X-PREFETCH-DISCONTINUITY tags!
+				PrefetchSeg.DiscontinuityCount = NextDiscontinuityNum;
+				PrefetchSeg.SequenceNumber     = NextSequenceNum++;
+				PrefetchSeg.URI 			   = AttrValue;
+				PrefetchSeg.InitSegmentInfo    = InitSegmentInfo;
+				PrefetchSeg.bIsPrefetch 	   = true;
+
+				int32 NumRegularSegments = MediaStream->SegmentList.Num();
+				if (NumRegularSegments)
+				{
+					// As per:
+					// "The duration of a Media Playlist containing prefetch segments is considered to be equal to the duration of all complete segments plus the expected duration of all prefetch segments."
+					AccumulatedDuration += PrefetchSeg.Duration;
+
+					// We do not set this here. Instead we keep the last full segment duration to determine when to reload the playlist.
+					//LastSegmentDuration = MediaStream->TargetDuration;
+
+					PrefetchSeg.AbsoluteDateTime = MediaStream->SegmentList[NumRegularSegments - 1].AbsoluteDateTime + MediaStream->SegmentList[NumRegularSegments - 1].Duration;
+					MediaStream->SegmentList.Push(PrefetchSeg);
+				}
+				else
+				{
+					LogMessage(IInfoLog::ELevel::Warning, FString::Printf(TEXT("Warning: Found EXT-X-PREFETCH segment without preceeding regular segment. This is not supported right now, ignoring this prefetch.")));
+				}
+			}
+
+			MediaStream->TotalAccumulatedSegmentDuration = AccumulatedDuration;
+			MediaStream->SeekablePositions = SeekablePositions;
+
+		}
+
+
+		// Find the variant or rendition this new playlist updates.
+		TSharedPtrTS<FManifestHLSInternal::FMediaStream> UpdatedMediaStream = MakeShareable(MediaStream.Release());
+		InOutHLSPlaylist->LockPlaylists();
+
+		Playlist->Internal.PlaylistLoadRequest  			   = SourceRequest;
+		Playlist->Internal.PlaylistLoadRequest.LastUpdateCRC32 = ResponseCRC;
+		Playlist->Internal.MediaStream  					   = UpdatedMediaStream;
+		Playlist->Internal.LoadState						   = FManifestHLSInternal::FPlaylistBase::FInternal::ELoadState::Loaded;
+		Playlist->Internal.Blacklisted.Reset();
+
+		// Set the expiration time of this playlist.
+		// A VOD or list with an end tag does not expire.
+		if (UpdatedMediaStream->bHasListEnd || UpdatedMediaStream->PlaylistType == FManifestHLSInternal::FMediaStream::EPlaylistType::VOD)
+		{
+			Playlist->Internal.ExpiresAtTime.SetToPositiveInfinity();
+		}
+		else
+		{
+			// Otherwise the playlist needs to be reloaded periodically after one target duration has passed from the time it was loaded last.
+			// NOTE: As per https://developer.apple.com/streaming/HLS-WWDC-2017-Preliminary-Spec.pdf this is now the duration of the last segment!
+			if (LastSegmentDuration.IsValid())
+			{
+				Playlist->Internal.ExpiresAtTime = SourceRequest.RequestedAtTime + LastSegmentDuration;
+			}
+			else
+			{
+				Playlist->Internal.ExpiresAtTime = SourceRequest.RequestedAtTime + UpdatedMediaStream->TargetDuration;
+			}
+		}
+		Playlist->Internal.bReloadTriggered = false;
+		Playlist->Internal.bNewlySelected   = false;
+
+		UpdateManifestMetadataFromStream(InOutHLSPlaylist.Get(), Playlist.Get(), UpdatedMediaStream.Get(), Preferences, Options);
+		InOutHLSPlaylist->UnlockPlaylists();
+	}
+
+	return FErrorDetail();
+}
+
+
+void FManifestBuilderHLS::UpdateManifestMetadataFromStream(FManifestHLSInternal* Manifest, const FManifestHLSInternal::FPlaylistBase* Playlist, FManifestHLSInternal::FMediaStream* Stream, const FStreamPreferences& Preferences, const FParamDict& Options)
+{
+	// NOTE: The master playlist mutex must be locked already!
+
+	// Update the seekable and timeline ranges of this stream.
+	FTimeValue EarliestSeekable, LatestSeekable;
+	FTimeValue EarliestTimeline, LatestTimeline;
+	bool bIsAudioOnly = Manifest->AudioOnlyStreams.Num() && Manifest->VariantStreams.Num() == 0;
+	const TArray<FManifestHLSInternal::FMediaStream::FMediaSegment>& SegmentList = Stream->SegmentList;
+	check(SegmentList.Num());
+	if (SegmentList.Num())
+	{
+		EarliestTimeline = SegmentList[0].AbsoluteDateTime;
+		LatestTimeline   = SegmentList.Last().AbsoluteDateTime + SegmentList.Last().Duration;
+
+		// For VOD and EVENT no segment can be removed from the playlist. Similarly for a LIVE playlist that has ended because it will no longer be updated.
+		if (Stream->PlaylistType == FManifestHLSInternal::FMediaStream::EPlaylistType::VOD ||
+			Stream->PlaylistType == FManifestHLSInternal::FMediaStream::EPlaylistType::Event ||
+			Stream->bHasListEnd)
+		{
+			EarliestSeekable = SegmentList[0].AbsoluteDateTime;
+		}
+		else
+		{
+			FTimeValue FirstOffset = Options.GetValue(IPlaylistReaderHLS::OptionKeyLiveSeekableStartOffset).SafeGetTimeValue(FTimeValue::GetZero());
+			for(int32 i=0, iMax=SegmentList.Num(); i<iMax; ++i)
+			{
+				if (SegmentList[i].AbsoluteDateTime >= FirstOffset)
+				{
+					EarliestSeekable = SegmentList[i].AbsoluteDateTime;
+					break;
+				}
+			}
+		}
+
+		// For VOD and playlists that have ended we can use the last segment.
+		if (Stream->PlaylistType == FManifestHLSInternal::FMediaStream::EPlaylistType::VOD ||
+			Stream->bHasListEnd)
+		{
+			LatestSeekable = SegmentList.Last().AbsoluteDateTime;
+		}
+		else
+		{
+			// As per the specification the distance to the Live edge is supposed to be 3 target durations.
+			// NOTE: It is not really clear if the "three target duration" rule is for the _time_ of three target durations *OR* if it simply means
+			//       to back off the last three media segments, which could actually be significantly less.
+			FTimeValue ThreeTargetDurations = Stream->TargetDuration * 3;
+
+			// Get the configured offset from the Live edge. This can be configured differently for audio-only presentations.
+			FTimeValue LastOffset;
+			if (bIsAudioOnly)
+			{
+				// First try the audio-only config.
+				if (Options.HaveKey(IPlaylistReaderHLS::OptionKeyLiveSeekableEndOffsetAudioOnly))
+				{
+					LastOffset = Options.GetValue(IPlaylistReaderHLS::OptionKeyLiveSeekableEndOffsetAudioOnly).SafeGetTimeValue(ThreeTargetDurations);
+				}
+			}
+			// Try the video config value regardless of presentation type.
+			if (!LastOffset.IsValid() && Options.HaveKey(IPlaylistReaderHLS::OptionKeyLiveSeekableEndOffset))
+			{
+				LastOffset = Options.GetValue(IPlaylistReaderHLS::OptionKeyLiveSeekableEndOffset).SafeGetTimeValue(ThreeTargetDurations);
+			}
+			// If still not valid use the default from the HLS specification.
+			if (!LastOffset.IsValid())
+			{
+				LastOffset = ThreeTargetDurations;
+			}
+
+			// Now back off from the end until we have reached at least the required offset.
+			FTimeValue BackedOff(FTimeValue::GetZero());
+			FTimeValue PrevBackedOff;
+			bool bUseConservativeDistance = Options.GetValue(IPlaylistReaderHLS::OptionKeyLiveSeekableEndOffsetBeConservative).SafeGetBool(false);
+			for(int32 i=SegmentList.Num()-1; i>=0; --i)
+			{
+				BackedOff += SegmentList[i].Duration;
+				if (BackedOff >= LastOffset)
+				{
+					LatestSeekable = SegmentList[i].AbsoluteDateTime;
+					// Be conservative and use a larger than configured value if need be?
+					if (!bUseConservativeDistance && PrevBackedOff.IsValid())
+					{
+						// Check if the difference to what is configured is possibly closer when we go one segment closer.
+						FTimeValue diffHere = BackedOff - LastOffset;
+						FTimeValue diffPrev = LastOffset - PrevBackedOff;
+						if (diffPrev < diffHere)
+						{
+							LatestSeekable = SegmentList[i+1].AbsoluteDateTime;
+						}
+					}
+					break;
+				}
+				PrevBackedOff = BackedOff;
+			}
+			// If we have not found a segment backing off from the Live edge the playlist must be too short (or the configured desired offset way too large).
+			if (!LatestSeekable.IsValid() && SegmentList.Num())
+			{
+				LatestSeekable = SegmentList[0].AbsoluteDateTime;
+			}
+		}
+		// At this point the latest must be valid!
+		check(LatestSeekable.IsValid());
+		if (LatestSeekable.IsValid())
+		{
+			// Adjust first against latest if necessary.
+			if (!EarliestSeekable.IsValid() || EarliestSeekable > LatestSeekable)
+			{
+				EarliestSeekable = LatestSeekable;
+			}
+		}
+	}
+	Stream->SeekableRange.Start = EarliestSeekable;
+	Stream->SeekableRange.End   = LatestSeekable;
+	Stream->TimelineRange.Start = EarliestTimeline;
+	Stream->TimelineRange.End   = LatestTimeline;
+
+	// To allow for elementary streams of different duration we need to look at all streams to calculate the overall
+	// presentation timeline regarding total and seekable ranges.
+	// This means we need to merge one or more streams into a common timeline.
+	// For a Live presentation we need to be careful about the streams we use for this merge to avoid stale playlists
+	// that are not actively updating to contribute.
+	// The easiest way to do this is to only look at the actively selected streams, but until playback begins there
+	// won't be any and we need to merge everything that is an initial playlist load.
+	TSet<uint32> StreamIDsToConsider = Manifest->ActivelyReferencedStreamIDs;
+	if (StreamIDsToConsider.Num() == 0)
+	{
+		// No active streams yet. Use all playlists we have loaded up so far.
+		for(int32 i=0; i<Manifest->VariantStreams.Num(); ++i)
+		{
+			if (Manifest->VariantStreams[i]->Internal.LoadState == FManifestHLSInternal::FPlaylistBase::FInternal::ELoadState::Loaded)
+			{
+				StreamIDsToConsider.Add(Manifest->VariantStreams[i]->Internal.UniqueID);
+			}
+		}
+		for (TMultiMap<FString, TSharedPtrTS<FManifestHLSInternal::FRendition>>::TConstIterator It = Manifest->AudioRenditions.CreateConstIterator(); It; ++It)
+		{
+			if (It.Value()->Internal.LoadState == FManifestHLSInternal::FPlaylistBase::FInternal::ELoadState::Loaded)
+			{
+				StreamIDsToConsider.Add(It.Value()->Internal.UniqueID);
+			}
+		}
+		for(int32 i=0; i<Manifest->AudioOnlyStreams.Num(); ++i)
+		{
+			if (Manifest->AudioOnlyStreams[i]->Internal.LoadState == FManifestHLSInternal::FPlaylistBase::FInternal::ELoadState::Loaded)
+			{
+				StreamIDsToConsider.Add(Manifest->AudioOnlyStreams[i]->Internal.UniqueID);
+			}
+		}
+	}
+	// There have to be streams here now. If not we throw an internal error.
+	if (StreamIDsToConsider.Num())
+	{
+		IManifest::EType PresentationType = IManifest::EType::OnDemand;
+
+		EarliestTimeline.SetToZero();
+		LatestTimeline.SetToZero();
+		EarliestSeekable.SetToZero();
+		LatestSeekable.SetToZero();
+
+		for(TSet<uint32>::TConstIterator StreamID = StreamIDsToConsider.CreateConstIterator(); StreamID; ++StreamID)
+		{
+			TSharedPtrTS<FManifestHLSInternal::FPlaylistBase> pl = Manifest->GetPlaylistForUniqueID(*StreamID);
+			check(pl.IsValid());
+			TSharedPtrTS<FManifestHLSInternal::FMediaStream> ms = pl->Internal.MediaStream;
+			if (ms.IsValid())
+			{
+				if (!(ms->bHasListEnd || ms->PlaylistType == FManifestHLSInternal::FMediaStream::EPlaylistType::VOD))
+				{
+					PresentationType = IManifest::EType::Live;
+				}
+				// Even for streams of different duration we still want them all to start at the same time.
+				// We only allow for one of them to be longer than the others.
+				if (ms->TimelineRange.Start > EarliestTimeline)
+				{
+					EarliestTimeline = ms->TimelineRange.Start;
+					EarliestSeekable = ms->SeekableRange.Start;
+				}
+				if (ms->TimelineRange.End > LatestTimeline)
+				{
+					LatestTimeline = ms->TimelineRange.End;
+					LatestSeekable = ms->SeekableRange.End;
+				}
+			}
+		}
+
+
+		// Set the presentation type and duration.
+		Manifest->MasterPlaylistVars.PresentationType = PresentationType;
+		if (PresentationType == IManifest::EType::OnDemand)
+		{
+			Manifest->MasterPlaylistVars.PresentationDuration = LatestTimeline - EarliestTimeline;
+			// For malformed playlists this could actually become negative. Just set to zero here.
+			if (Manifest->MasterPlaylistVars.PresentationDuration < FTimeValue::GetZero())
+			{
+				Manifest->MasterPlaylistVars.PresentationDuration.SetToZero();
+			}
+		}
+		else
+		{
+			Manifest->MasterPlaylistVars.PresentationDuration.SetToPositiveInfinity();
+		}
+
+		Manifest->MasterPlaylistVars.SeekableRange.Start = EarliestSeekable;
+		Manifest->MasterPlaylistVars.SeekableRange.End   = LatestSeekable;
+		Manifest->MasterPlaylistVars.TimelineRange.Start = EarliestTimeline;
+		Manifest->MasterPlaylistVars.TimelineRange.End   = LatestTimeline;
+
+		// Set the seekable position from either the video stream or the audio stream, depending on whether there is a video stream or not.
+		if ((bIsAudioOnly && !Playlist->Internal.bHasVideo) ||
+			(!bIsAudioOnly && Playlist->Internal.bHasVideo))
+		{
+			Manifest->MasterPlaylistVars.SeekablePositions = Stream->SeekablePositions;
+		}
+
+		// Update the timeline.
+		// A new timeline object is created to avoid updating the existing one that may be held onto by some client code.
+		// This isn't overly expensive since we do share the actual streams.
+		check(Manifest->PlaybackTimeline.IsValid());
+		if (Manifest->PlaybackTimeline.IsValid())
+		{
+			HLSTimeline::FPlaybackTimeline* CurrentTimeline = static_cast<HLSTimeline::FPlaybackTimeline*>(Manifest->PlaybackTimeline.Get());
+			HLSTimeline::FPlaybackTimeline* Timeline		= new HLSTimeline::FPlaybackTimeline;
+			HLSTimeline::FTimelineMediaAsset* Asset 		= new HLSTimeline::FTimelineMediaAsset;
+			Timeline->PlayAsset = MakeShareable(Asset);
+			Timeline->SeekableTimeRange = Manifest->MasterPlaylistVars.SeekableRange;
+			Timeline->SeekablePositions = Manifest->MasterPlaylistVars.SeekablePositions;
+			Timeline->TotalTimeRange	= Asset->TimeRange = Manifest->MasterPlaylistVars.TimelineRange;
+			Timeline->Duration  		= Asset->Duration  = Manifest->MasterPlaylistVars.PresentationDuration;
+			Asset->UniqueIdentifier    = CurrentTimeline->PlayAsset->UniqueIdentifier;
+			Asset->AssetIdentifier     = CurrentTimeline->PlayAsset->AssetIdentifier;
+			Asset->VideoAdaptationSet  = CurrentTimeline->PlayAsset->VideoAdaptationSet;
+			Asset->AudioAdaptationSets = CurrentTimeline->PlayAsset->AudioAdaptationSets;
+			// Switch over
+			Manifest->PlaybackTimeline = MakeShareable(Timeline);
+		}
+	}
+	else
+	{
+		PlayerSessionServices->PostError(CreateErrorAndLog(FString::Printf(TEXT("No stream found to construct media timeline from!")), ERRCODE_HLS_BUILDER_INTERNAL, UEMEDIA_ERROR_INTERNAL));
+	}
+}
+
+
+
+
+
+
+
+} // namespace Electra
+
+
+
