ElectraPlayer: RELEASE. Moved from NFL to public. Note that this is done via a Delete and "Add New" operation as discussed.

diff --git a/Engine/Plugins/Media/ElectraPlayer/Source/ElectraPlayerRuntime/Private/Runtime/Player/HLS/ManifestBuilderHLS.h b/Engine/Plugins/Media/ElectraPlayer/Source/ElectraPlayerRuntime/Private/Runtime/Player/HLS/ManifestBuilderHLS.h
new file mode 100644
--- /dev/null
+++ b/Engine/Plugins/Media/ElectraPlayer/Source/ElectraPlayerRuntime/Private/Runtime/Player/HLS/ManifestBuilderHLS.h
@@ -0,0 +1,480 @@
+// Copyright Epic Games, Inc. All Rights Reserved.
+
+#pragma once
+
+#include "PlayerCore.h"
+#include "Containers/List.h"
+#include "Player/PlaybackTimeline.h"
+#include "Player/AdaptiveStreamingPlayerABR.h"
+#include "Player/PlaylistReader.h"
+#include "Player/PlayerSessionServices.h"
+#include "StreamTypes.h"
+#include "ErrorDetail.h"
+#include "Parser.h"
+
+
+namespace Electra
+{
+class IInitSegmentCacheHLS;
+class ILicenseKeyCacheHLS;
+
+struct FPlaylistLoadRequestHLS
+{
+	enum class ELoadType
+	{
+		Master,
+		Initial,
+		First,
+		Update
+	};
+	FPlaylistLoadRequestHLS() : LoadType(ELoadType::Master), InternalUniqueID(0), LastUpdateCRC32(~uint32(0)) {}
+	FString			URL;
+	FTimeValue		RequestedAtTime;
+	ELoadType		LoadType;
+	uint32			InternalUniqueID;
+	uint32			LastUpdateCRC32;
+	FString			AdaptationSetUniqueID;
+	FString			RepresentationUniqueID;
+	FString			CDN;
+};
+
+
+struct FManifestHLSInternal
+{
+	struct FMediaStream
+	{
+		enum class EPlaylistType
+		{
+			Live,
+			Event,
+			VOD
+		};
+
+		struct FByteRange
+		{
+			FByteRange() : Start(-1), End(-1)
+			{
+			}
+			int64		Start;
+			int64		End;
+			bool IsSet() const
+			{
+				return Start >=0 && End >= 0;
+			}
+			int64 GetStart() const
+			{
+				return Start;
+			}
+			int64 GetEnd() const
+			{
+				return End;
+			}
+			int64 GetNumBytes() const
+			{
+				return IsSet() ? End - Start + 1 : 0;
+			}
+		};
+
+		struct FDRMKeyInfo
+		{
+			enum class EMethod
+			{
+				None,
+				AES128,
+				SampleAES
+			};
+			FDRMKeyInfo() : Method(EMethod::None) {}
+			EMethod				Method;
+			FString				URI;
+			FString				IV;
+			//FString		Keyformat;
+			//FString		KeyformatVersions;
+		};
+
+		struct FInitSegmentInfo
+		{
+			FString											URI;
+			FByteRange										ByteRange;
+			TSharedPtr<FDRMKeyInfo, ESPMode::ThreadSafe>	DRMKeyInfo;			//!< If set the init segment is encrypted with the specified parameters.
+		};
+
+		struct FMediaSegment
+		{
+			FString											URI;
+			FByteRange										ByteRange;
+			FTimeValue										Duration;
+			FTimeValue										RelativeStartTime;	//!< Accumulated durations to give a total start time of this segment relative to the beginning of the playlist.
+			FTimeValue										AbsoluteDateTime;	//!< If valid this gives the absolute time of the first sample
+			TSharedPtr<FDRMKeyInfo, ESPMode::ThreadSafe>	DRMKeyInfo;			//!< If not set the segment is not encrypted. Otherwise it is with the specified parameters.
+			TSharedPtrTS<FInitSegmentInfo>					InitSegmentInfo;	//!< If not set the segment contains the necessary init data. Otherwise this points to the init segment.
+			int64											SequenceNumber;
+			int64											DiscontinuityCount;
+			bool											bIsPrefetch;		//!< Whether or not this is an EXT-X-PREFETCH segment
+		};
+
+		struct FStartTime
+		{
+			FStartTime() : bPrecise(false) {}
+			FTimeValue			Offset;
+			bool				bPrecise;
+		};
+
+		FMediaStream()
+			: PlaylistType(EPlaylistType::Live)
+			, MediaSequence(0)
+			, DiscontinuitySequence(0)
+			, bHasListEnd(false)
+			, bIsIFramesOnly(false)
+			, bHasIndependentSegments(false)
+			, bHasEncryptedSegments(false)
+		{
+		}
+
+	// TODO: replace with a scattered vector!!!!
+		TArray<FMediaSegment>			SegmentList;
+		FStartTime						StartTime;
+		EPlaylistType					PlaylistType;
+		FTimeValue						TargetDuration;
+		int64							MediaSequence;
+		int64							DiscontinuitySequence;
+		bool							bHasListEnd;
+		bool							bIsIFramesOnly;
+		bool							bHasIndependentSegments;
+		bool							bHasEncryptedSegments;
+
+		// Internal use
+		FTimeRange						SeekableRange;						//!< Current seekable range within the stream.
+		FTimeRange						TimelineRange;						//!< Current timeline range of the stream.
+		FTimeValue						TotalAccumulatedSegmentDuration;	//!< Accumulated duration of all segments.
+		TArray<FTimespan>				SeekablePositions;					//!< List of segment absolute start times
+	};
+
+
+	struct FBlacklist
+	{
+		TSharedPtrTS<HTTP::FRetryInfo>					PreviousAttempts;
+		FTimeValue										BecomesAvailableAgainAtUTC;
+		IAdaptiveStreamSelector::FBlacklistedStream		AssetIDs;
+	};
+
+
+	struct FPlaylistBase
+	{
+		FPlaylistBase() = default;
+		virtual ~FPlaylistBase() = default;
+
+		virtual bool IsVariantStream() const = 0;
+
+		virtual const FString& GetURL() const = 0;
+
+		virtual int32 GetBitrate() const = 0;
+
+		// Internal use.
+		struct FInternal
+		{
+			enum class ELoadState
+			{
+				NotLoaded,
+				Pending,
+				Loaded
+			};
+			FInternal() : LoadState(ELoadState::NotLoaded), UniqueID(0), bReloadTriggered(false), bNewlySelected(true), bHasVideo(false), bHasAudio(false) {}
+			FPlaylistLoadRequestHLS			PlaylistLoadRequest;
+			TSharedPtrTS<FMediaStream>		MediaStream;
+			TSharedPtrTS<FBlacklist>		Blacklisted;
+			FTimeValue						ExpiresAtTime;							//!< Synchronized UTC time (from session service GetSynchronizedUTCTime()) at which this list expires. Set to infinite if it does not.
+			ELoadState						LoadState;
+			uint32							UniqueID;
+			bool							bReloadTriggered;
+			bool							bNewlySelected;
+			FString							AdaptationSetUniqueID;
+			FString							RepresentationUniqueID;
+			FString							CDN;
+			bool							bHasVideo;
+			bool							bHasAudio;
+		};
+		FInternal			Internal;
+	};
+
+	// 4.3.4.1. EXT-X-MEDIA
+	struct FRendition : public FPlaylistBase
+	{
+		FRendition()
+			: FPlaylistBase(), bDefault(false), bAutoSelect(false), bForced(false)
+		{
+		}
+		virtual bool IsVariantStream() const override
+		{
+			return false;
+		}
+		virtual const FString& GetURL() const override
+		{
+			return URI;
+		}
+		virtual int32 GetBitrate() const override
+		{
+			return 0;
+		}
+
+		FString		Type;
+		FString		GroupID;
+		FString		URI;
+		FString		Language;
+		FString		AssocLanguage;
+		FString		Name;
+		FString		InStreamID;
+		FString		Characteristics;
+		FString		Channels;
+		bool				bDefault;
+		bool				bAutoSelect;
+		bool				bForced;
+	};
+
+	// 4.3.4.2. EXT-X-STREAM-INF
+	struct FVariantStream : public FPlaylistBase
+	{
+		FVariantStream()
+			: FPlaylistBase(), Bandwidth(0), AverageBandwidth(0), HDCPLevel(0)
+		{
+		}
+		virtual bool IsVariantStream() const override
+		{
+			return true;
+		}
+		virtual const FString& GetURL() const override
+		{
+			return URI;
+		}
+		virtual int32 GetBitrate() const override
+		{
+			return Bandwidth;
+		}
+
+		TArray<FStreamCodecInformation>			StreamCodecInformationList;
+		FString									URI;
+		FString									VideoGroupID;
+		FString									AudioGroupID;
+		FString									SubtitleGroupID;
+		FString									ClosedCaptionGroupID;
+		int32									Bandwidth;
+		int32									AverageBandwidth;
+		int32									HDCPLevel;
+	};
+
+
+	// Internal use
+	struct FMasterPlaylistVars
+	{
+		FMasterPlaylistVars() : PresentationType(IManifest::EType::Live) {}
+		FPlaylistLoadRequestHLS								PlaylistLoadRequest;				//!< The master playlist load request. This holds the URL for resolving child playlists.
+		IManifest::EType									PresentationType;					//!< Type of the presentation (Live or VOD).
+		FTimeValue											PresentationDuration;				//!< Current duration of the presentation. Will be reset when a playlist is refreshed.
+		FTimeRange											SeekableRange;						//!< Current seekable range within the presentation. Will be updated when a playlist is refreshed.
+		FTimeRange											TimelineRange;						//!< Current media timeline range of the presentation. Will be updated when a playlist is refreshed.
+		TArray<FTimespan>									SeekablePositions;					//!< Segment start times.
+	};
+
+	FManifestHLSInternal()
+		: bHasIndependentSegments(false)
+	{
+	}
+
+
+	TMultiMap<FString, TSharedPtrTS<FRendition>>			VideoRenditions;
+	TMultiMap<FString, TSharedPtrTS<FRendition>>			AudioRenditions;
+	TMultiMap<FString, TSharedPtrTS<FRendition>>			SubtitleRenditions;
+	TMultiMap<FString, TSharedPtrTS<FRendition>>			ClosedCaptionRenditions;
+	// TODO: if provided we can put the EXT-X-START here
+	// TODO: For future DRM support we can put the EXT-X-SESSION-KEY here.
+
+	TArray<TSharedPtrTS<FVariantStream>>							VariantStreams;
+	TArray<TSharedPtrTS<FVariantStream>>							AudioOnlyStreams;				//!< Variants that are audio-only (have only audio codec specified in master playlist)
+	TArray<TSharedPtrTS<FVariantStream>>							IFrameOnlyStreams;
+
+	TMap<uint32, TWeakPtrTS<FPlaylistBase>>							PlaylistIDMap;
+	TMap<int32, int32>												BandwidthToQualityIndex;
+
+	bool															bHasIndependentSegments;
+
+	FMasterPlaylistVars												MasterPlaylistVars;
+
+
+	FMediaCriticalSection														VariantPlaylistAccessMutex;			//!< This mutex is used to access any(!) of the variant and rendition playlists. The master playlist is immutable.
+	TSet<uint32>													ActivelyReferencedStreamIDs;		//!< Unique IDs of streams from which segments are being fetched.
+	TSharedPtrTS<IInitSegmentCacheHLS>								InitSegmentCache;
+	TSharedPtrTS<ILicenseKeyCacheHLS>								LicenseKeyCache;
+
+	TArray<FStreamMetadata>											StreamMetadataVideo;
+	TArray<FStreamMetadata>											StreamMetadataAudio;
+
+	TSharedPtrTS<IPlaybackAssetTimeline>							PlaybackTimeline;
+
+
+	TSharedPtrTS<FPlaylistBase> GetPlaylistForUniqueID(uint32 UniqueID) const
+	{
+		const TWeakPtrTS<FPlaylistBase>* PlaylistID = PlaylistIDMap.Find(UniqueID);
+		if (PlaylistID != nullptr)
+		{
+			return PlaylistID->Pin();
+		}
+		return TSharedPtrTS<FManifestHLSInternal::FPlaylistBase>();
+	}
+
+
+	void SelectActiveStreamID(uint32 NewActiveStreamID, uint32 NowInactiveStreamID)
+	{
+		VariantPlaylistAccessMutex.Lock();
+		ActivelyReferencedStreamIDs.Remove(NowInactiveStreamID);
+
+		// Add to selected list only if not deselected (the new stream ID is not 0)
+		if (NewActiveStreamID)
+		{
+			ActivelyReferencedStreamIDs.Add(NewActiveStreamID);
+
+			TSharedPtrTS<FPlaylistBase> Playlist = PlaylistIDMap[NewActiveStreamID].Pin();
+			if (Playlist.IsValid())
+			{
+				Playlist->Internal.bNewlySelected = true;
+			}
+		}
+		VariantPlaylistAccessMutex.Unlock();
+	}
+
+	void GetActiveStreams(TArray<TSharedPtrTS<FPlaylistBase>>& OutActiveStreams)
+	{
+		VariantPlaylistAccessMutex.Lock();
+		for(TSet<uint32>::TConstIterator StreamID=ActivelyReferencedStreamIDs.CreateConstIterator(); StreamID; ++StreamID)
+		{
+			TSharedPtrTS<FPlaylistBase> Playlist = PlaylistIDMap[*StreamID].Pin();
+			if (Playlist.IsValid())
+			{
+				OutActiveStreams.Push(Playlist);
+			}
+		}
+		VariantPlaylistAccessMutex.Unlock();
+	}
+
+
+
+	void LockPlaylists()
+	{
+		VariantPlaylistAccessMutex.Lock();
+	}
+	void UnlockPlaylists()
+	{
+		VariantPlaylistAccessMutex.Unlock();
+	}
+	struct ScopedLockPlaylists
+	{
+		ScopedLockPlaylists(TSharedPtrTS<FManifestHLSInternal> InManifest)
+			: Manifest(InManifest)
+		{
+			if (Manifest.IsValid())
+			{
+				Manifest->LockPlaylists();
+			}
+		}
+		~ScopedLockPlaylists()
+		{
+			if (Manifest.IsValid())
+			{
+				Manifest->UnlockPlaylists();
+			}
+		}
+	private:
+		TSharedPtrTS<FManifestHLSInternal> Manifest;
+	};
+};
+
+
+
+
+
+
+class IManifestBuilderHLS
+{
+public:
+	static IManifestBuilderHLS* Create(IPlayerSessionServices* PlayerSessionServices);
+
+	virtual ~IManifestBuilderHLS() = default;
+
+
+	/**
+	 * Builds a new internal manifest from a HLS master playlist.
+	 *
+	 * @param OutHLSPlaylist
+	 * @param Playlist
+	 * @param SourceRequest
+	 * @param ConnectionInfo
+	 * @param Preferences
+	 * @param Options
+	 *
+	 * @return
+	 */
+	virtual FErrorDetail BuildFromMasterPlaylist(TSharedPtrTS<FManifestHLSInternal>& OutHLSPlaylist, const HLSPlaylistParser::FPlaylist& Playlist, const FPlaylistLoadRequestHLS& SourceRequest, const HTTP::FConnectionInfo* ConnectionInfo, const FStreamPreferences& Preferences, const FParamDict& Options) = 0;
+
+
+	/**
+	 * Returns the list of variant and rendition playlists that must be fetched on play start.
+	 *
+	 * @param OutRequests
+	 * @param Manifest
+	 * @param Preferences
+	 * @param Options
+	 *
+	 * @return
+	 */
+	virtual FErrorDetail GetInitialPlaylistLoadRequests(TArray<FPlaylistLoadRequestHLS>& OutRequests, TSharedPtrTS<FManifestHLSInternal> Manifest, const FStreamPreferences& Preferences, const FParamDict& Options) = 0;
+
+	/**
+	 * Updates an initial playlist load request that failed to load or parse with an
+	 * alternative request from a lower bitrate if possible.
+	 *
+	 * @param InOutFailedRequest
+	 *                 Failed request for which an alterative shall be returned.
+	 * @param ConnectionInfo
+	 * @param PreviousAttempts
+	 * @param BlacklistUntilUTC
+	 * @param Manifest
+	 * @param Preferences
+	 * @param Options
+	 *
+	 * @return UEMEDIA_ERROR_OK if an alternative was found or UEMEDIA_ERROR_END_OF_STREAM if no further alternatives exist.
+	 */
+	virtual UEMediaError UpdateFailedInitialPlaylistLoadRequest(FPlaylistLoadRequestHLS& InOutFailedRequest, const HTTP::FConnectionInfo* ConnectionInfo, TSharedPtrTS<HTTP::FRetryInfo> PreviousAttempts, const FTimeValue& BlacklistUntilUTC, TSharedPtrTS<FManifestHLSInternal> Manifest, const FStreamPreferences& Preferences, const FParamDict& Options) = 0;
+
+
+	/**
+	 * Updates the internal manifest with a new/refreshed variant playlist.
+	 *
+	 * @param InOutHLSPlaylist
+	 * @param VariantPlaylist
+	 * @param SourceRequest
+	 * @param ConnectionInfo
+	 * @param ResponseCRC
+	 * @param Preferences
+	 * @param Options
+	 *
+	 * @return
+	 */
+	virtual FErrorDetail UpdateFromVariantPlaylist(TSharedPtrTS<FManifestHLSInternal> InOutHLSPlaylist, const HLSPlaylistParser::FPlaylist& VariantPlaylist, const FPlaylistLoadRequestHLS& SourceRequest, const HTTP::FConnectionInfo* ConnectionInfo, uint32 ResponseCRC, const FStreamPreferences& Preferences, const FParamDict& Options) = 0;
+
+
+	/**
+	 * Marks a variant as failed.
+	 *
+	 * @param InHLSPlaylist
+	 * @param SourceRequest
+	 * @param ConnectionInfo
+	 * @param PreviousAttempts
+	 * @param BlacklistUntilUTC
+	 */
+	virtual void SetVariantPlaylistFailure(TSharedPtrTS<FManifestHLSInternal> InHLSPlaylist, const FPlaylistLoadRequestHLS& SourceRequest, const HTTP::FConnectionInfo* ConnectionInfo, TSharedPtrTS<HTTP::FRetryInfo> PreviousAttempts, const FTimeValue& BlacklistUntilUTC) = 0;
+};
+
+
+} // namespace Electra
+
+
+
+
diff --git a/Engine/Plugins/Media/ElectraPlayer/Source/ElectraPlayerRuntime/Private/Runtime/Player/HLS/ManifestHLS.cpp b/Engine/Plugins/Media/ElectraPlayer/Source/ElectraPlayerRuntime/Private/Runtime/Player/HLS/ManifestHLS.cpp
new file mode 100644
--- /dev/null
+++ b/Engine/Plugins/Media/ElectraPlayer/Source/ElectraPlayerRuntime/Private/Runtime/Player/HLS/ManifestHLS.cpp
@@ -0,0 +1,1265 @@
+// Copyright Epic Games, Inc. All Rights Reserved.
+
+#include "PlayerCore.h"
+#include "ManifestHLS.h"
+#include "ManifestBuilderHLS.h"
+#include "PlaylistReaderHLS.h"
+#include "InitSegmentCacheHLS.h"
+#include "StreamReaderHLSfmp4.h"
+#include "Player/PlayerSessionServices.h"
+#include "SynchronizedClock.h"
+#include "Utilities/StringHelpers.h"
+#include "Utilities/URLParser.h"
+
+
+namespace Electra
+{
+
+/**
+ * Interface to a playback period.
+ */
+class FPlayPeriodHLS : public IManifest::IPlayPeriod
+{
+public:
+	FPlayPeriodHLS(IPlayerSessionServices* SessionServices, IPlaylistReaderHLS* PlaylistReader, TSharedPtrTS<FManifestHLSInternal> Manifest);
+	virtual ~FPlayPeriodHLS();
+
+	virtual void SetStreamPreferences(const FStreamPreferences& Preferences) override;
+
+	virtual EReadyState GetReadyState() override;
+	virtual void PrepareForPlay(const FParamDict& Options) override;
+
+	// TODO: need to provide metadata (duration, streams, languages, etc.)
+
+	virtual IManifest::FResult GetStartingSegment(TSharedPtrTS<IStreamSegment>& OutSegment, const FPlayStartPosition& StartPosition, IManifest::ESearchType SearchType) override;
+	virtual IManifest::FResult GetNextSegment(TSharedPtrTS<IStreamSegment>& OutSegment, TSharedPtrTS<const IStreamSegment> CurrentSegment, const FParamDict& Options) override;
+	virtual IManifest::FResult GetRetrySegment(TSharedPtrTS<IStreamSegment>& OutSegment, TSharedPtrTS<const IStreamSegment> CurrentSegment, const FParamDict& Options) override;
+	virtual IManifest::FResult GetLoopingSegment(TSharedPtrTS<IStreamSegment>& OutSegment, FPlayerLoopState& InOutLoopState, const TMultiMap<EStreamType, TSharedPtrTS<IStreamSegment>>& InFinishedSegments, const FPlayStartPosition& StartPosition, IManifest::ESearchType SearchType) override;
+
+	// Obtains information on the stream segmentation of a particular stream starting at a given current reference segment (optional, if not given returns suitable default values).
+	virtual void GetSegmentInformation(TArray<FSegmentInformation>& OutSegmentInformation, FTimeValue& OutAverageSegmentDuration, TSharedPtrTS<const IStreamSegment> CurrentSegment, const FTimeValue& LookAheadTime, const TSharedPtrTS<IPlaybackAssetAdaptationSet>& AdaptationSet, const TSharedPtrTS<IPlaybackAssetRepresentation>& Representation) override;
+
+	virtual TSharedPtrTS<ITimelineMediaAsset> GetMediaAsset() const override;
+
+	virtual void SelectStream(const TSharedPtrTS<IPlaybackAssetAdaptationSet>& AdaptationSet, const TSharedPtrTS<IPlaybackAssetRepresentation>& Representation, const FString& PreferredCDN) override;
+
+private:
+	struct FSegSearchParam
+	{
+		FSegSearchParam()
+			: MediaSequence(-1)
+			, DiscontinuitySequence(-1)
+			, LocalIndex(-1)
+			, StreamUniqueID(0)
+		{
+		}
+		FTimeValue	Time;						//!< Time to search for.
+		FTimeValue	Duration;					//!< If set this we search for a start time of Time + Duration (aka the next segment)
+		int64		MediaSequence;				//!< If >= 0 we are searching for a specific segment based on media sequence number
+		int64		DiscontinuitySequence;		//!< If >= 0 we are searching for a segment after the this discontinuity sequence number
+		int32		LocalIndex;
+		uint32		StreamUniqueID;				//!< If != 0 the search is for the same stream as was the previous segment. We can use the media sequence index.
+		//FTimeValue	PTS;
+	};
+
+	void LogMessage(IInfoLog::ELevel Level, const FString& Message);
+
+	IManifest::FResult GetMediaStreamForID(TSharedPtrTS<FManifestHLSInternal::FPlaylistBase>& OutPlaylist, TSharedPtrTS<FManifestHLSInternal::FMediaStream>& OutMediaStream, uint32 UniqueID) const;
+
+	IManifest::FResult GetNextOrRetrySegment(TSharedPtrTS<IStreamSegment>& OutSegment, TSharedPtrTS<const IStreamSegment> CurrentSegment, bool bRetry);
+
+	IManifest::FResult FindSegment(TSharedPtrTS<FStreamSegmentRequestHLSfmp4>& OutRequest, TSharedPtrTS<FManifestHLSInternal::FPlaylistBase> InPlaylist, TSharedPtrTS<FManifestHLSInternal::FMediaStream> InStream, uint32 StreamUniqueID, EStreamType StreamType, const FSegSearchParam& SearchParam, IManifest::ESearchType SearchType);
+
+	void RefreshBlacklistState();
+
+	TSharedPtrTS<FManifestHLSInternal>			InternalManifest;
+	IPlayerSessionServices* 					SessionServices;
+	IPlaylistReaderHLS*							PlaylistReader;
+	EReadyState									CurrentReadyState;
+
+	uint32										ActiveVideoUniqueID;
+	uint32										ActiveAudioUniqueID;
+};
+
+
+
+TSharedPtrTS<FManifestHLS> FManifestHLS::Create(IPlayerSessionServices* SessionServices, const FParamDict& Options, IPlaylistReaderHLS* PlaylistReader, TSharedPtrTS<FManifestHLSInternal> Manifest)
+{
+	return TSharedPtrTS<FManifestHLS>(new FManifestHLS(SessionServices, Options, PlaylistReader, Manifest));
+}
+
+FManifestHLS::FManifestHLS(IPlayerSessionServices* InSessionServices, const FParamDict& InOptions, IPlaylistReaderHLS* InPlaylistReader, TSharedPtrTS<FManifestHLSInternal> Manifest)
+	: Options(InOptions)
+	, InternalManifest(Manifest)
+	, SessionServices(InSessionServices)
+	, PlaylistReader(InPlaylistReader)
+{
+}
+
+FManifestHLS::~FManifestHLS()
+{
+}
+
+IManifest::EType FManifestHLS::GetPresentationType() const
+{
+	FManifestHLSInternal::ScopedLockPlaylists lock(InternalManifest);
+	return InternalManifest->MasterPlaylistVars.PresentationType;
+}
+
+
+TSharedPtrTS<IPlaybackAssetTimeline> FManifestHLS::GetTimeline() const
+{
+	FManifestHLSInternal::ScopedLockPlaylists lock(InternalManifest);
+	return InternalManifest->PlaybackTimeline;
+}
+
+
+/**
+ * Returns the bitrate of the default stream (usually the first one specified).
+ *
+ * @return
+ */
+int64 FManifestHLS::GetDefaultStartingBitrate() const
+{
+	FManifestHLSInternal::ScopedLockPlaylists lock(InternalManifest);
+	if (InternalManifest->VariantStreams.Num())
+	{
+		return InternalManifest->VariantStreams[0]->Bandwidth;
+	}
+	return 0;
+}
+
+/**
+ * Returns stream metadata. For period based presentations the streams can be different per period in which case the metadata of the first period is returned.
+ *
+ * @param OutMetadata
+ * @param StreamType
+ */
+void FManifestHLS::GetStreamMetadata(TArray<FStreamMetadata>& OutMetadata, EStreamType StreamType) const
+{
+	FManifestHLSInternal::ScopedLockPlaylists lock(InternalManifest);
+	switch(StreamType)
+	{
+		case EStreamType::Video:
+			OutMetadata = InternalManifest->StreamMetadataVideo;
+			break;
+		case EStreamType::Audio:
+			OutMetadata = InternalManifest->StreamMetadataAudio;
+			break;
+		case EStreamType::Subtitle:
+			OutMetadata.Empty();
+			break;
+	}
+}
+
+
+/**
+ * Returns the duration that should be present in the buffers at all times
+ * (except for the end of the presentation).
+ *
+ * @return
+ */
+FTimeValue FManifestHLS::GetMinBufferTime() const
+{
+	// HLS does not offer a minimum duration to be in the buffers at all times. For expedited startup we use 2 seconds here.
+	return FTimeValue().SetFromSeconds(2.0);
+}
+
+
+
+
+
+/**
+ * Returns a play period for the specified start time.
+ * Since we are not currently splitting the media timeline into individual periods
+ * we simply return a new period here regardless of the starting time.
+ *
+ * @param OutPlayPeriod
+ * @param StartPosition
+ * @param SearchType
+ *
+ * @return
+ */
+IManifest::FResult FManifestHLS::FindPlayPeriod(TSharedPtrTS<IPlayPeriod>& OutPlayPeriod, const FPlayStartPosition& StartPosition, ESearchType SearchType)
+{
+	TSharedPtrTS<FPlayPeriodHLS> Period(new FPlayPeriodHLS(SessionServices, PlaylistReader, InternalManifest));
+	OutPlayPeriod = Period;
+	return IManifest::FResult(IManifest::FResult::EType::Found);
+}
+
+
+
+/**
+ * Creates a stream reader for the media segments.
+ *
+ * @return
+ */
+IStreamReader* FManifestHLS::CreateStreamReaderHandler()
+{
+	return new FStreamReaderHLSfmp4;
+}
+
+
+
+
+
+
+
+
+
+
+FPlayPeriodHLS::FPlayPeriodHLS(IPlayerSessionServices* InSessionServices, IPlaylistReaderHLS* InPlaylistReader, TSharedPtrTS<FManifestHLSInternal> Manifest)
+	: InternalManifest(Manifest)
+	, SessionServices(InSessionServices)
+	, PlaylistReader(InPlaylistReader)
+{
+	check(PlaylistReader);
+
+	CurrentReadyState   			  = IManifest::IPlayPeriod::EReadyState::NotReady;
+
+	// Set the active video and audio stream IDs to 0, which means none are selected.
+	ActiveVideoUniqueID = 0;
+	ActiveAudioUniqueID = 0;
+}
+
+FPlayPeriodHLS::~FPlayPeriodHLS()
+{
+}
+
+
+void FPlayPeriodHLS::LogMessage(IInfoLog::ELevel Level, const FString& Message)
+{
+	if (SessionServices)
+	{
+		SessionServices->PostLog(Facility::EFacility::HLSManifest, Level, Message);
+	}
+}
+
+/**
+ * Sets stream preferences.
+ *
+ * @param Preferences
+ */
+void FPlayPeriodHLS::SetStreamPreferences(const FStreamPreferences& Preferences)
+{
+}
+
+
+
+/**
+ * Returns the current ready state of the period.
+ *
+ * @return
+ */
+IManifest::IPlayPeriod::EReadyState FPlayPeriodHLS::GetReadyState()
+{
+	return CurrentReadyState;
+}
+
+
+
+/**
+ * Prepares the period for playback.
+ *
+ * @param Options
+ */
+void FPlayPeriodHLS::PrepareForPlay(const FParamDict& Options)
+{
+	// For now we just go with the streams for which we loaded the playlists initially.
+	// FIXME: in the future, based on preferences and options, select the streams we want.
+
+	FManifestHLSInternal::ScopedLockPlaylists lock(InternalManifest);
+
+	uint32 OldVideoUniqueID = ActiveVideoUniqueID;
+	uint32 OldAudioUniqueID = ActiveAudioUniqueID;
+
+	for(int32 i=0; i<InternalManifest->VariantStreams.Num(); ++i)
+	{
+		if (InternalManifest->VariantStreams[i]->Internal.LoadState == FManifestHLSInternal::FPlaylistBase::FInternal::ELoadState::Loaded)
+		{
+			ActiveVideoUniqueID = InternalManifest->VariantStreams[i]->Internal.UniqueID;
+			break;
+		}
+	}
+
+	for (TMultiMap<FString, TSharedPtrTS<FManifestHLSInternal::FRendition>>::TConstIterator It = InternalManifest->AudioRenditions.CreateConstIterator(); It; ++It)
+	{
+		if (It.Value()->Internal.LoadState == FManifestHLSInternal::FPlaylistBase::FInternal::ELoadState::Loaded)
+		{
+			ActiveAudioUniqueID = It.Value()->Internal.UniqueID;
+			break;
+		}
+	}
+	// In case there is an audio rendition without a dedicated playlist we look at audio-only variant streams
+	if (ActiveAudioUniqueID == 0)
+	{
+		for(int32 i=0; i<InternalManifest->AudioOnlyStreams.Num(); ++i)
+		{
+			if (InternalManifest->AudioOnlyStreams[i]->Internal.LoadState == FManifestHLSInternal::FPlaylistBase::FInternal::ELoadState::Loaded)
+			{
+				ActiveAudioUniqueID = InternalManifest->AudioOnlyStreams[i]->Internal.UniqueID;
+				break;
+			}
+		}
+	}
+
+	// Tell the manifest which stream IDs are now actively used.
+	InternalManifest->SelectActiveStreamID(ActiveVideoUniqueID, OldVideoUniqueID);
+	InternalManifest->SelectActiveStreamID(ActiveAudioUniqueID, OldAudioUniqueID);
+
+	CurrentReadyState = IManifest::IPlayPeriod::EReadyState::IsReady;
+}
+
+
+
+/**
+ * Returns the media stream for the specified ID.
+ *
+ * @param OutPlaylist
+ * @param OutMediaStream
+ * @param UniqueID
+ *
+ * @return
+ */
+IManifest::FResult FPlayPeriodHLS::GetMediaStreamForID(TSharedPtrTS<FManifestHLSInternal::FPlaylistBase>& OutPlaylist, TSharedPtrTS<FManifestHLSInternal::FMediaStream>& OutMediaStream, uint32 UniqueID) const
+{
+	FManifestHLSInternal::ScopedLockPlaylists lock(InternalManifest);
+	check(UniqueID);
+	if (UniqueID)
+	{
+		TWeakPtrTS<FManifestHLSInternal::FPlaylistBase>* PlaylistID = InternalManifest->PlaylistIDMap.Find(UniqueID);
+		if (PlaylistID != nullptr)
+		{
+			TSharedPtrTS<FManifestHLSInternal::FPlaylistBase> Playlist = PlaylistID->Pin();
+
+			if (Playlist.IsValid())
+			{
+				// Sanity check the ID
+				if (Playlist->Internal.UniqueID == UniqueID)
+				{
+					OutPlaylist = Playlist;
+
+					// Playlist currently blacklisted?
+					if (Playlist->Internal.Blacklisted.IsValid())
+					{
+						// Return and assume a non-blacklisted stream will be selected.
+						return IManifest::FResult().RetryAfterMilliseconds(50);
+					}
+					// Check the load state
+					switch(Playlist->Internal.LoadState)
+					{
+						case FManifestHLSInternal::FPlaylistBase::FInternal::ELoadState::Loaded:
+						{
+							TSharedPtrTS<FManifestHLSInternal::FMediaStream>	MediaStream = Playlist->Internal.MediaStream;
+							// The stream really better be there!
+							if (MediaStream.IsValid())
+							{
+								OutMediaStream = MediaStream;
+								return IManifest::FResult(IManifest::FResult::EType::Found);
+							}
+							else
+							{
+								return IManifest::FResult(IManifest::FResult::EType::NotFound).SetErrorDetail(FErrorDetail().SetMessage(FString::Printf(TEXT("Media stream for unique ID %u is not present!"), UniqueID)));
+							}
+						}
+						case FManifestHLSInternal::FPlaylistBase::FInternal::ELoadState::NotLoaded:
+						{
+							return IManifest::FResult(IManifest::FResult::EType::NotLoaded);
+						}
+						case FManifestHLSInternal::FPlaylistBase::FInternal::ELoadState::Pending:
+						{
+							return IManifest::FResult().RetryAfterMilliseconds(50);
+						}
+					}
+					// Should never get here, but if we do let's bail gracefully.
+					return IManifest::FResult(IManifest::FResult::EType::NotFound).SetErrorDetail(FErrorDetail().SetMessage("Internal error. Unhandled switch case"));
+				}
+				else
+				{
+					return IManifest::FResult(IManifest::FResult::EType::NotFound).SetErrorDetail(FErrorDetail().SetMessage(FString::Printf(TEXT("Playlist unique ID %u does not match requested ID of %u"), Playlist->Internal.UniqueID, UniqueID)));
+				}
+			}
+			else
+			{
+				return IManifest::FResult(IManifest::FResult::EType::NotFound).SetErrorDetail(FErrorDetail().SetMessage(FString::Printf(TEXT("Playlist for unique ID %u has been destroyed"), UniqueID)));
+			}
+		}
+		else
+		{
+			return IManifest::FResult(IManifest::FResult::EType::NotFound).SetErrorDetail(FErrorDetail().SetMessage(FString::Printf(TEXT("No media stream found for unique ID %u"), UniqueID)));
+		}
+	}
+	else
+	{
+		return IManifest::FResult(IManifest::FResult::EType::NotFound).SetErrorDetail(FErrorDetail().SetMessage(FString::Printf(TEXT("Invalid unique media stream ID %u"), UniqueID)));
+	}
+}
+
+
+
+
+
+/**
+ * Locate a segment in the stream's playlist.
+ *
+ * @param OutRequest
+ * @param InPlaylist
+ * @param InStream
+ * @param StreamUniqueID
+ * @param StreamType
+ * @param SearchParam
+ * @param SearchType
+ *
+ * @return
+ */
+IManifest::FResult FPlayPeriodHLS::FindSegment(TSharedPtrTS<FStreamSegmentRequestHLSfmp4>& OutRequest, TSharedPtrTS<FManifestHLSInternal::FPlaylistBase> InPlaylist, TSharedPtrTS<FManifestHLSInternal::FMediaStream> InStream, uint32 StreamUniqueID, EStreamType StreamType, const FPlayPeriodHLS::FSegSearchParam& SearchParam, IManifest::ESearchType SearchType)
+{
+	// VOD or EVENT playlist?
+///////	if (InStream->PlaylistType == FManifestHLSInternal::FMediaStream::EPlaylistType::VOD || InStream->PlaylistType == FManifestHLSInternal::FMediaStream::EPlaylistType::Event)
+
+
+	TUniquePtr<IURLParser> UrlBuilder(IURLParser::Create());
+	UrlBuilder->ParseURL(InPlaylist->Internal.PlaylistLoadRequest.URL);
+	TSharedPtrTS<FStreamSegmentRequestHLSfmp4>	Req(new FStreamSegmentRequestHLSfmp4);
+	Req->StreamType 	= StreamType;
+	Req->StreamUniqueID = StreamUniqueID;
+	if (InternalManifest->PlaybackTimeline.IsValid())
+	{
+		Req->MediaAsset = InternalManifest->PlaybackTimeline->GetMediaAssetByIndex(0);
+		if (!Req->MediaAsset.IsValid())
+		{
+			return IManifest::FResult(IManifest::FResult::EType::NotFound).SetErrorDetail(FErrorDetail().SetMessage("Internal error, media asset not found on asset timeline!"));
+		}
+		if (Req->MediaAsset->GetNumberOfAdaptationSets(StreamType) > 1)
+		{
+			return IManifest::FResult(IManifest::FResult::EType::NotFound).SetErrorDetail(FErrorDetail().SetMessage(FString::Printf(TEXT("Internal error, more than one %s rendition group found on asset timeline!"), GetStreamTypeName(StreamType))));
+		}
+		Req->AdaptationSet = Req->MediaAsset->GetAdaptationSetByTypeAndIndex(StreamType, 0);
+		if (!Req->AdaptationSet.IsValid())
+		{
+			return IManifest::FResult(IManifest::FResult::EType::NotFound).SetErrorDetail(FErrorDetail().SetMessage(FString::Printf(TEXT("Internal error, no %s rendition group found on asset timeline!"), GetStreamTypeName(StreamType))));
+		}
+		Req->Representation = Req->AdaptationSet->GetRepresentationByUniqueIdentifier(LexToString(StreamUniqueID));
+		if (!Req->Representation.IsValid())
+		{
+			return IManifest::FResult(IManifest::FResult::EType::NotFound).SetErrorDetail(FErrorDetail().SetMessage(FString::Printf(TEXT("Internal error, %s rendition not found in group on asset timeline!"), GetStreamTypeName(StreamType))));
+		}
+	}
+	Req->InitSegmentCache   	   = InternalManifest->InitSegmentCache;
+	Req->LicenseKeyCache		   = InternalManifest->LicenseKeyCache;
+	Req->bHasEncryptedSegments     = InStream->bHasEncryptedSegments;
+	if (StreamType == EStreamType::Video)
+	{
+		Req->Bitrate				   = InPlaylist->GetBitrate();
+		check(InternalManifest->BandwidthToQualityIndex.Find(Req->Bitrate) != nullptr);
+		Req->QualityLevel   		   = InternalManifest->BandwidthToQualityIndex[Req->Bitrate];
+	}
+
+	const TArray<FManifestHLSInternal::FMediaStream::FMediaSegment>& SegmentList = InStream->SegmentList;
+	if (SegmentList.Num())
+	{
+		FTimeValue searchTime = SearchParam.Time;
+		int32 SelectedSegmentIndex = -1;
+
+		// Searching for the next segment within the same stream?
+		if (SearchParam.StreamUniqueID != 0)
+		{
+			check(SearchType == IManifest::ESearchType::StrictlyAfter || SearchType == IManifest::ESearchType::Same);
+			if (SearchType != IManifest::ESearchType::StrictlyAfter && SearchType != IManifest::ESearchType::Same)
+			{
+				return IManifest::FResult(IManifest::FResult::EType::NotFound).SetErrorDetail(FErrorDetail().SetMessage("Can only find next or retry segment in same stream right now"));
+			}
+
+			int64 NextSequenceNumber;
+			if (SearchType == IManifest::ESearchType::Same)
+			{
+				NextSequenceNumber = SearchParam.MediaSequence;
+			}
+			else
+			{
+				NextSequenceNumber = SearchParam.MediaSequence + 1;
+			}
+			// We can use the media sequence number.
+			for(int32 i=0,iMax=SegmentList.Num(); i<iMax; ++i)
+			{
+				if (SegmentList[i].SequenceNumber >= NextSequenceNumber)
+				{
+					SelectedSegmentIndex = i;
+					break;
+				}
+			}
+		}
+		else
+		{
+			for(int32 i=0,iMax=SegmentList.Num(); i<iMax; ++i)
+			{
+				// Find the segment whose start time is >= the time we're looking for.
+				if (SegmentList[i].AbsoluteDateTime >= searchTime)
+				{
+					// Do we want the segment with start time >= the search time?
+					if (SearchType == IManifest::ESearchType::After)
+					{
+						// Yes, we're done.
+						SelectedSegmentIndex = i;
+						break;
+					}
+					// Do we want the segment with start time > the time we're looking for?
+					else if (SearchType == IManifest::ESearchType::StrictlyAfter)
+					{
+						// Only go forward if we did hit the search time exactly ( == ) and we're not on the last first segment already!
+						if (SegmentList[i].AbsoluteDateTime == searchTime)
+						{
+							// Continue the loop. The next segment, if it exists, will have a greater search time and we'll catch it then.
+							continue;
+						}
+						SelectedSegmentIndex = i;
+						break;
+					}
+					// Do we want the segment with start time <= the search time?
+					else if (SearchType == IManifest::ESearchType::Before)
+					{
+						SelectedSegmentIndex = i;
+						// Only go back if we did not hit the search time exactly ( == ) and we're not on the very first segment already!
+						if (SegmentList[i].AbsoluteDateTime > searchTime && i > 0)
+						{
+							--SelectedSegmentIndex;
+						}
+						break;
+					}
+					// Do we want the segment with start time < the search time?
+					else if (SearchType == IManifest::ESearchType::StrictlyBefore)
+					{
+						// If we cannot go back one segment we can return.
+						if (i == 0)
+						{
+							return IManifest::FResult(IManifest::FResult::EType::BeforeStart);
+						}
+						SelectedSegmentIndex = i-1;
+						break;
+					}
+					// Do we want the segment whose start time is closest to the search time?
+					else if (SearchType == IManifest::ESearchType::Closest)
+					{
+						SelectedSegmentIndex = i;
+						// If there is an earlier segment we can check which one is closer.
+						if (i > 0)
+						{
+							FTimeValue diffHere   = SegmentList[i].AbsoluteDateTime - searchTime;
+							FTimeValue diffBefore = searchTime - SegmentList[i - 1].AbsoluteDateTime;
+							// In the exceptionally rare case the difference to either segment is the same we pick the earlier one.
+							if (diffBefore <= diffHere)
+							{
+								--SelectedSegmentIndex;
+							}
+						}
+						break;
+					}
+					// Do we want the segment for the exact same start time as the search time?
+					else if (SearchType == IManifest::ESearchType::Same)
+					{
+						// This is used for retrying a failed segment. Usually on another quality level or CDN.
+						// To allow for slight variations in the time we do a 'closest' search if the exact time can't be found.
+						SelectedSegmentIndex = i;
+						// If we hit the time dead on we are done.
+						if (SegmentList[i].AbsoluteDateTime == searchTime)
+						{
+							break;
+						}
+						// Otherwise we do the same as for the 'closest' search.
+						if (i > 0)
+						{
+							FTimeValue diffHere   = SegmentList[i].AbsoluteDateTime - searchTime;
+							FTimeValue diffBefore = searchTime - SegmentList[i - 1].AbsoluteDateTime;
+							// In the exceptionally rare case the difference to either segment is the same we pick the earlier one.
+							if (diffBefore <= diffHere)
+							{
+								--SelectedSegmentIndex;
+							}
+						}
+						break;
+					}
+
+					else
+					{
+						checkNoEntry();
+						return IManifest::FResult(IManifest::FResult::EType::NotFound).SetErrorDetail(FErrorDetail().SetMessage("Internal error, unsupported segment search mode!"));
+					}
+				}
+			}
+			// If not found but there are segments in the list they all have a start time smaller than we are looking for.
+			// Let's look at the last segment to see if the search time falls into the duration of it.
+			if (SelectedSegmentIndex < 0 && SegmentList.Num())
+			{
+				// Whether we can use the last segment also depends on the search mode.
+				if (SearchType == IManifest::ESearchType::Closest ||
+					SearchType == IManifest::ESearchType::Before ||
+					SearchType == IManifest::ESearchType::StrictlyBefore)
+				{
+					int32 LastSegmentIndex = SegmentList.Num() - 1;
+					if (searchTime < SegmentList[LastSegmentIndex].AbsoluteDateTime + SegmentList[LastSegmentIndex].Duration)
+					{
+						SelectedSegmentIndex = LastSegmentIndex;
+					}
+				}
+			}
+		}
+
+		// Did we find the segment?
+		if (SelectedSegmentIndex >= 0 && SelectedSegmentIndex < SegmentList.Num())
+		{
+			//Req->PlaylistRelativeStartTime = SegmentList[SelectedSegmentIndex].RelativeStartTime;
+			Req->AbsoluteDateTime   	   = SegmentList[SelectedSegmentIndex].AbsoluteDateTime;
+			Req->SegmentDuration		   = SegmentList[SelectedSegmentIndex].Duration;
+			Req->MediaSequence  		   = SegmentList[SelectedSegmentIndex].SequenceNumber;
+			Req->DiscontinuitySequence     = SegmentList[SelectedSegmentIndex].DiscontinuityCount;
+			Req->LocalIndex 			   = SelectedSegmentIndex;
+			Req->bIsPrefetch			   = SegmentList[SelectedSegmentIndex].bIsPrefetch;
+			Req->bIsEOSSegment  		   = false;
+			Req->URL					   = UrlBuilder->ResolveWith(SegmentList[SelectedSegmentIndex].URI);
+			Req->FirstAUTimeOffset  	   = searchTime - SegmentList[SelectedSegmentIndex].AbsoluteDateTime;
+			Req->InitSegmentInfo		   = SegmentList[SelectedSegmentIndex].InitSegmentInfo;
+			Req->LicenseKeyInfo 		   = SegmentList[SelectedSegmentIndex].DRMKeyInfo;
+
+			// This can be negative when we're picking the segment after the search time!
+			// But we don't want this since it is a useless case. We will be receiving AUs from behind the search time and
+			// that's simply where we are starting at. Set to zero in this case.
+			if (Req->FirstAUTimeOffset < FTimeValue::GetZero())
+			{
+				Req->FirstAUTimeOffset.SetToZero();
+			}
+
+			if (SegmentList[SelectedSegmentIndex].ByteRange.IsSet())
+			{
+				Req->Range.Start		= SegmentList[SelectedSegmentIndex].ByteRange.GetStart();
+				Req->Range.EndIncluding = SegmentList[SelectedSegmentIndex].ByteRange.GetEnd();
+			}
+			OutRequest = Req;
+			return IManifest::FResult(IManifest::FResult::EType::Found);
+		}
+		else
+		{
+			// Not having found a segment means we're beyond this presentation.
+			// Unless this is a VOD list or it has as ENDLIST tag we have to try this later, assuming that an updated playlist has added additional segments.
+			if (InStream->PlaylistType == FManifestHLSInternal::FMediaStream::EPlaylistType::VOD || InStream->bHasListEnd)
+			{
+				Req->bIsEOSSegment = true;
+				OutRequest = Req;
+				return IManifest::FResult(IManifest::FResult::EType::PastEOS);
+			}
+			else
+			{
+				// Try again after half a target duration.
+				return IManifest::FResult(IManifest::FResult::EType::TryAgainLater).RetryAfterMilliseconds(InStream->TargetDuration.GetAsMilliseconds() / 2);
+			}
+		}
+	}
+	else
+	{
+		// No segments is not really expected. If this occurs we assume the presentation has ended.
+		Req->bIsEOSSegment = true;
+		OutRequest = Req;
+		return IManifest::FResult(IManifest::FResult::EType::PastEOS);
+	}
+	return IManifest::FResult(IManifest::FResult::EType::NotFound);
+}
+
+
+IManifest::FResult FPlayPeriodHLS::GetStartingSegment(TSharedPtrTS<IStreamSegment>& OutSegment, const FPlayStartPosition& StartPosition, IManifest::ESearchType SearchType)
+{
+	FManifestHLSInternal::ScopedLockPlaylists lock(InternalManifest);
+
+	RefreshBlacklistState();
+
+	TSharedPtrTS<FManifestHLSInternal::FPlaylistBase>	VideoPlaylist;
+	TSharedPtrTS<FManifestHLSInternal::FPlaylistBase>	AudioPlaylist;
+	TSharedPtrTS<FManifestHLSInternal::FMediaStream>	VideoStream;
+	TSharedPtrTS<FManifestHLSInternal::FMediaStream>	AudioStream;
+
+	// Get the streams that are selected, if there are selected ones.
+	IManifest::FResult vidResult = ActiveVideoUniqueID ? GetMediaStreamForID(VideoPlaylist, VideoStream, ActiveVideoUniqueID) : IManifest::FResult(IManifest::FResult::EType::Found);
+	IManifest::FResult audResult = ActiveAudioUniqueID ? GetMediaStreamForID(AudioPlaylist, AudioStream, ActiveAudioUniqueID) : IManifest::FResult(IManifest::FResult::EType::Found);
+	if (vidResult.IsSuccess() && audResult.IsSuccess())
+	{
+		TSharedPtrTS<FStreamSegmentRequestHLSfmp4> 	VideoSegmentRequest;
+		TSharedPtrTS<FStreamSegmentRequestHLSfmp4> 	AudioSegmentRequest;
+		FSegSearchParam									SearchParam;
+
+		SearchParam.Time = StartPosition.Time;
+
+		// Do we have both video and audio?
+		if (ActiveVideoUniqueID && ActiveAudioUniqueID)
+		{
+			vidResult = FindSegment(VideoSegmentRequest, VideoPlaylist, VideoStream, ActiveVideoUniqueID, EStreamType::Video, SearchParam, SearchType);
+			// Found and PastEOS are valid results here. Everything else is not.
+			if (vidResult.GetType() != IManifest::FResult::EType::Found && vidResult.GetType() != IManifest::FResult::EType::PastEOS)
+			{
+				return vidResult;
+			}
+			// If the search for video was successful we adjust the search parameters for the audio stream.
+			if (vidResult.IsSuccess())
+			{
+				VideoSegmentRequest->bIsInitialStartRequest = true;
+				VideoSegmentRequest->FirstAUTimeOffset.SetToZero();
+
+				// With the video segment found let's find the corresponding audio segment.
+				SearchParam.Time = VideoSegmentRequest->AbsoluteDateTime;
+				SearchParam.DiscontinuitySequence = VideoSegmentRequest->DiscontinuitySequence;
+				// For audio we start with the segment before the video segment if there is no precise match.
+				// The stream reader will skip over all audio access units before the intended start time.
+				SearchType = IManifest::ESearchType::Before;
+			}
+			// Search for audio.
+			audResult = FindSegment(AudioSegmentRequest, AudioPlaylist, AudioStream, ActiveAudioUniqueID, EStreamType::Audio, SearchParam, SearchType);
+			// Equally here, if successful or PastEOS is acceptable and everything else is not.
+			if (audResult.GetType() != IManifest::FResult::EType::Found && audResult.GetType() != IManifest::FResult::EType::PastEOS)
+			{
+				return audResult;
+			}
+
+			// Both segments found?
+			if (vidResult.IsSuccess() && audResult.IsSuccess())
+			{
+				AudioSegmentRequest->bIsInitialStartRequest = true;
+				VideoSegmentRequest->DependentStreams.Push(AudioSegmentRequest);
+				OutSegment = VideoSegmentRequest;
+				return IManifest::FResult(IManifest::FResult::EType::Found);
+			}
+			// Only audio found?
+			else if (audResult.IsSuccess())
+			{
+				AudioSegmentRequest->bIsInitialStartRequest = true;
+				AudioSegmentRequest->DependentStreams.Push(VideoSegmentRequest);
+				OutSegment = AudioSegmentRequest;
+				return IManifest::FResult(IManifest::FResult::EType::Found);
+			}
+			// Only video found? Or neither?
+			else //if (vidResult.IsSuccess())
+			{
+				VideoSegmentRequest->bIsInitialStartRequest = true;
+				VideoSegmentRequest->DependentStreams.Push(AudioSegmentRequest);
+				OutSegment = VideoSegmentRequest;
+				return IManifest::FResult(IManifest::FResult::EType::Found);
+			}
+		}
+		// Video only?
+		else if (ActiveVideoUniqueID)
+		{
+			vidResult = FindSegment(VideoSegmentRequest, VideoPlaylist, VideoStream, ActiveVideoUniqueID, EStreamType::Video, SearchParam, SearchType);
+			if (vidResult.IsSuccess())
+			{
+				VideoSegmentRequest->bIsInitialStartRequest = true;
+				VideoSegmentRequest->FirstAUTimeOffset.SetToZero();
+				OutSegment = VideoSegmentRequest;
+				return IManifest::FResult(IManifest::FResult::EType::Found);
+			}
+			else
+			{
+				return vidResult;
+			}
+		}
+		// Audio only
+		else
+		{
+			audResult = FindSegment(AudioSegmentRequest, AudioPlaylist, AudioStream, ActiveAudioUniqueID, EStreamType::Audio, SearchParam, SearchType);
+			if (audResult.IsSuccess())
+			{
+				AudioSegmentRequest->bIsInitialStartRequest = true;
+				OutSegment = AudioSegmentRequest;
+				return IManifest::FResult(IManifest::FResult::EType::Found);
+			}
+			else
+			{
+				return audResult;
+			}
+		}
+	}
+	else
+	{
+		// Either playlist not yet loaded?
+		if (vidResult.GetType() == IManifest::FResult::EType::NotLoaded)
+		{
+			check(VideoPlaylist.IsValid());
+
+			VideoPlaylist->Internal.LoadState = FManifestHLSInternal::FPlaylistBase::FInternal::ELoadState::Pending;
+
+			TUniquePtr<IURLParser> UrlBuilder(IURLParser::Create());
+			UrlBuilder->ParseURL(InternalManifest->MasterPlaylistVars.PlaylistLoadRequest.URL);
+			FPlaylistLoadRequestHLS Req;
+			Req.LoadType			   = FPlaylistLoadRequestHLS::ELoadType::First;
+			Req.InternalUniqueID	   = ActiveVideoUniqueID;
+			Req.RequestedAtTime 	   = SessionServices->GetSynchronizedUTCTime()->GetTime();
+			Req.URL 				   = UrlBuilder->ResolveWith(VideoPlaylist->GetURL());
+			Req.AdaptationSetUniqueID  = VideoPlaylist->Internal.AdaptationSetUniqueID;
+			Req.RepresentationUniqueID = VideoPlaylist->Internal.RepresentationUniqueID;
+			Req.CDN 				   = VideoPlaylist->Internal.CDN;
+			PlaylistReader->RequestPlaylistLoad(Req);
+			vidResult.RetryAfterMilliseconds(50);
+		}
+		if (audResult.GetType() == IManifest::FResult::EType::NotLoaded)
+		{
+			check(AudioPlaylist.IsValid());
+
+			AudioPlaylist->Internal.LoadState = FManifestHLSInternal::FPlaylistBase::FInternal::ELoadState::Pending;
+
+			TUniquePtr<IURLParser> UrlBuilder(IURLParser::Create());
+			UrlBuilder->ParseURL(InternalManifest->MasterPlaylistVars.PlaylistLoadRequest.URL);
+			FPlaylistLoadRequestHLS Req;
+			Req.LoadType			   = FPlaylistLoadRequestHLS::ELoadType::First;
+			Req.InternalUniqueID	   = ActiveAudioUniqueID;
+			Req.RequestedAtTime 	   = SessionServices->GetSynchronizedUTCTime()->GetTime();
+			Req.URL 				   = UrlBuilder->ResolveWith(AudioPlaylist->GetURL());
+			Req.AdaptationSetUniqueID  = AudioPlaylist->Internal.AdaptationSetUniqueID;
+			Req.RepresentationUniqueID = AudioPlaylist->Internal.RepresentationUniqueID;
+			Req.CDN 				   = AudioPlaylist->Internal.CDN;
+			PlaylistReader->RequestPlaylistLoad(Req);
+			audResult.RetryAfterMilliseconds(50);
+		}
+
+		if (vidResult.GetType() == IManifest::FResult::EType::TryAgainLater)
+		{
+			return vidResult;
+		}
+		if (audResult.GetType() == IManifest::FResult::EType::TryAgainLater)
+		{
+			return audResult;
+		}
+		// If both are a go, go!
+		if (vidResult.IsSuccess() && audResult.IsSuccess())
+		{
+			return vidResult;
+		}
+		// Return that which is at fault.
+		return !vidResult.IsSuccess() ? vidResult : audResult;
+	}
+}
+
+IManifest::FResult FPlayPeriodHLS::GetLoopingSegment(TSharedPtrTS<IStreamSegment>& OutSegment, FPlayerLoopState& InOutLoopState, const TMultiMap<EStreamType, TSharedPtrTS<IStreamSegment>>& InFinishedSegments, const FPlayStartPosition& StartPosition, IManifest::ESearchType SearchType)
+{
+	if (InFinishedSegments.Num())
+	{
+		// Go over all finished segments and get the largest next expected timestamp from all of them.
+		FTimeValue LargestNextExpectedTimestamp(FTimeValue::GetZero());
+		for(auto It = InFinishedSegments.CreateConstIterator(); It; ++It)
+		{
+			const FStreamSegmentRequestHLSfmp4* FinishedRequest = static_cast<const FStreamSegmentRequestHLSfmp4*>(It->Value.Get());
+			if (FinishedRequest && FinishedRequest->NextLargestExpectedTimestamp > LargestNextExpectedTimestamp)
+			{
+				LargestNextExpectedTimestamp = FinishedRequest->NextLargestExpectedTimestamp;
+			}
+		}
+		IManifest::FResult res = GetStartingSegment(OutSegment, StartPosition, SearchType);
+		if (res.GetType() == IManifest::FResult::EType::Found)
+		{
+			FStreamSegmentRequestHLSfmp4* LoopRequest = static_cast<FStreamSegmentRequestHLSfmp4*>(OutSegment.Get());
+			InOutLoopState.bLoopEnabled = true;
+			InOutLoopState.LoopBasetime = LargestNextExpectedTimestamp;	// This is the absolute playback time at which the loop will occur
+			++InOutLoopState.LoopCount;
+			LoopRequest->PlayerLoopState = InOutLoopState;
+			LoopRequest->PlayerLoopState.LoopBasetime -= LoopRequest->AbsoluteDateTime;	// This is the _offset_ to add internally to the PTS to make it loop.
+			// Set the loop state in the dependent streams as well.
+			for(int32 nDep=0, nDepMax=LoopRequest->DependentStreams.Num(); nDep<nDepMax; ++nDep)
+			{
+				LoopRequest->DependentStreams[nDep]->PlayerLoopState = LoopRequest->PlayerLoopState;
+			}
+			return res;
+		}
+	}
+	// Return past EOS when we can't loop to indicate we're really done now.
+	return IManifest::FResult(IManifest::FResult::EType::PastEOS);
+}
+
+
+
+IManifest::FResult FPlayPeriodHLS::GetNextOrRetrySegment(TSharedPtrTS<IStreamSegment>& OutSegment, TSharedPtrTS<const IStreamSegment> InCurrentSegment, bool bRetry)
+{
+	// Need to have a current segment to find the next one.
+	if (!InCurrentSegment.IsValid())
+	{
+		return IManifest::FResult(IManifest::FResult::EType::NotFound).SetErrorDetail(FErrorDetail().SetMessage("Cannot get next segment without a current segment!"));
+	}
+	const FStreamSegmentRequestHLSfmp4* CurrentRequest = static_cast<const FStreamSegmentRequestHLSfmp4*>(InCurrentSegment.Get());
+	check(CurrentRequest->DependentStreams.Num() == 0);
+	if (CurrentRequest->DependentStreams.Num())
+	{
+		return IManifest::FResult(IManifest::FResult::EType::NotFound).SetErrorDetail(FErrorDetail().SetMessage("Cannot get next segment for a segment with dependent segments!"));
+	}
+	if (CurrentRequest->StreamUniqueID == 0)
+	{
+		return IManifest::FResult(IManifest::FResult::EType::NotFound).SetErrorDetail(FErrorDetail().SetMessage("Cannot get next segment for a segment having no unique stream ID!"));
+	}
+
+	FManifestHLSInternal::ScopedLockPlaylists lock(InternalManifest);
+	uint32 ForStreamID = 0;
+	switch(CurrentRequest->GetType())
+	{
+		case EStreamType::Video:
+		{
+			ForStreamID = ActiveVideoUniqueID;
+			break;
+		}
+		case EStreamType::Audio:
+		{
+			ForStreamID = ActiveAudioUniqueID;
+			break;
+		}
+		default:
+		{
+			return IManifest::FResult(IManifest::FResult::EType::NotFound).SetErrorDetail(FErrorDetail().SetMessage(FString::Printf(TEXT("Cannot get next segment for unsupported stream type \"%s\"!"), GetStreamTypeName(CurrentRequest->GetType()))));
+		}
+	}
+	if (ForStreamID == 0)
+	{
+		return IManifest::FResult(IManifest::FResult::EType::NotFound).SetErrorDetail(FErrorDetail().SetMessage(FString::Printf(TEXT("Cannot get next segment stream type \"%s\" since no stream is actively selected!"), GetStreamTypeName(CurrentRequest->GetType()))));
+	}
+
+	RefreshBlacklistState();
+
+	TSharedPtrTS<FManifestHLSInternal::FMediaStream>	Stream;
+	TSharedPtrTS<FManifestHLSInternal::FPlaylistBase> Playlist;
+	IManifest::FResult Result = GetMediaStreamForID(Playlist, Stream, ForStreamID);
+	if (Result.IsSuccess())
+	{
+		TSharedPtrTS<FStreamSegmentRequestHLSfmp4> 	NextSegmentRequest;
+		FSegSearchParam									SearchParam;
+
+		SearchParam.Time				  = CurrentRequest->AbsoluteDateTime;
+		SearchParam.Duration			  = CurrentRequest->SegmentDuration;
+		SearchParam.MediaSequence   	  = CurrentRequest->MediaSequence;
+		SearchParam.DiscontinuitySequence = CurrentRequest->DiscontinuitySequence;
+		SearchParam.LocalIndex  		  = CurrentRequest->LocalIndex;
+		SearchParam.StreamUniqueID  	  = CurrentRequest->StreamUniqueID == ForStreamID ? ForStreamID : 0;
+		Result = FindSegment(NextSegmentRequest, Playlist, Stream, ForStreamID, CurrentRequest->GetType(), SearchParam, bRetry ? IManifest::ESearchType::Same : IManifest::ESearchType::StrictlyAfter);
+		if (Result.IsSuccess())
+		{
+			// Continuing with the next segment implicitly means there is no AU time offset.
+			if (!bRetry)
+			{
+				NextSegmentRequest->FirstAUTimeOffset.SetToZero();
+			}
+			else
+			{
+				NextSegmentRequest->NumOverallRetries = CurrentRequest->NumOverallRetries + 1;
+			}
+			// Copy over the player loop state.
+			NextSegmentRequest->PlayerLoopState = CurrentRequest->PlayerLoopState;
+			OutSegment = NextSegmentRequest;
+			return IManifest::FResult(IManifest::FResult::EType::Found);
+		}
+		if (Result.GetType() == IManifest::FResult::EType::TryAgainLater)
+		{
+			return Result;
+		}
+		else
+		{
+			checkf(Result.GetType() == IManifest::FResult::EType::PastEOS, TEXT("What error is this?"));
+			return Result;
+		}
+	}
+	else
+	{
+		// Playlist not loaded?
+		if (Result.GetType() == IManifest::FResult::EType::NotLoaded)
+		{
+			check(Playlist.IsValid());
+
+			Playlist->Internal.LoadState = FManifestHLSInternal::FPlaylistBase::FInternal::ELoadState::Pending;
+
+			TUniquePtr<IURLParser> UrlBuilder(IURLParser::Create());
+			UrlBuilder->ParseURL(InternalManifest->MasterPlaylistVars.PlaylistLoadRequest.URL);
+			FPlaylistLoadRequestHLS Req;
+			Req.LoadType			   = FPlaylistLoadRequestHLS::ELoadType::First;
+			Req.InternalUniqueID	   = ForStreamID;
+			Req.RequestedAtTime 	   = SessionServices->GetSynchronizedUTCTime()->GetTime();
+			Req.URL 				   = UrlBuilder->ResolveWith(Playlist->GetURL());
+			Req.AdaptationSetUniqueID  = Playlist->Internal.AdaptationSetUniqueID;
+			Req.RepresentationUniqueID = Playlist->Internal.RepresentationUniqueID;
+			Req.CDN 				   = Playlist->Internal.CDN;
+			PlaylistReader->RequestPlaylistLoad(Req);
+			Result.RetryAfterMilliseconds(50);
+		}
+		return Result;
+	}
+
+	return IManifest::FResult(IManifest::FResult::EType::PastEOS);
+}
+
+/**
+ * Selects the next segment to download.
+ * This might be a segment from a different variant stream after a quality switch.
+ *
+ * @param OutSegment
+ * @param InCurrentSegment
+ * @param Options
+ *
+ * @return
+ */
+IManifest::FResult FPlayPeriodHLS::GetNextSegment(TSharedPtrTS<IStreamSegment>& OutSegment, TSharedPtrTS<const IStreamSegment> InCurrentSegment, const FParamDict& Options)
+{
+	return GetNextOrRetrySegment(OutSegment, InCurrentSegment, false);
+}
+
+
+IManifest::FResult FPlayPeriodHLS::GetRetrySegment(TSharedPtrTS<IStreamSegment>& OutSegment, TSharedPtrTS<const IStreamSegment> InCurrentSegment, const FParamDict& Options)
+{
+	bool bInsertFiller = Options.GetValue("insertFiller").SafeGetBool(false);
+
+	// To insert filler data we can use the current request over again.
+	if (bInsertFiller)
+	{
+		const FStreamSegmentRequestHLSfmp4* CurrentRequest = static_cast<const FStreamSegmentRequestHLSfmp4*>(InCurrentSegment.Get());
+		TSharedPtrTS<FStreamSegmentRequestHLSfmp4> NewRequest(new FStreamSegmentRequestHLSfmp4);
+		NewRequest->CopyFrom(*CurrentRequest);
+		NewRequest->bInsertFillerData = true;
+		// We treat replacing the segment with filler data as a retry.
+		++NewRequest->NumOverallRetries;
+		OutSegment = NewRequest;
+		return IManifest::FResult(IManifest::FResult::EType::Found);
+	}
+	return GetNextOrRetrySegment(OutSegment, InCurrentSegment, true);
+}
+
+
+/**
+ * Checks if any potentially blacklisted stream can be used again.
+ */
+void FPlayPeriodHLS::RefreshBlacklistState()
+{
+	FTimeValue Now = SessionServices->GetSynchronizedUTCTime()->GetTime();
+
+	// Note: the manifest must be locked already.
+	for (TMap<uint32, TWeakPtrTS<FManifestHLSInternal::FPlaylistBase>>::TIterator It = InternalManifest->PlaylistIDMap.CreateIterator(); It; ++It)
+	{
+		TSharedPtrTS<FManifestHLSInternal::FRendition::FPlaylistBase> Stream = It.Value().Pin();
+		if (Stream.IsValid())
+		{
+			if (Stream->Internal.Blacklisted.IsValid())
+			{
+				if (Now >= Stream->Internal.Blacklisted->BecomesAvailableAgainAtUTC)
+				{
+					Stream->Internal.LoadState  	  = FManifestHLSInternal::FPlaylistBase::FInternal::ELoadState::NotLoaded;
+					Stream->Internal.bReloadTriggered = false;
+					Stream->Internal.bNewlySelected   = false;
+					Stream->Internal.ExpiresAtTime.SetToPositiveInfinity();
+
+					// Tell the stream selector that this stream is available again.
+					TSharedPtrTS<IAdaptiveStreamSelector> StreamSelector(SessionServices->GetStreamSelector());
+					StreamSelector->MarkStreamAsAvailable(Stream->Internal.Blacklisted->AssetIDs);
+
+					Stream->Internal.Blacklisted.Reset();
+					LogMessage(IInfoLog::ELevel::Info, FString::Printf(TEXT("Lifting blacklist of playlist \"%s\""), *Stream->Internal.PlaylistLoadRequest.URL));
+				}
+			}
+		}
+	}
+}
+
+
+
+/**
+ * Returns segment information (duration and estimated byte size) of the
+ * next n segments for the indicated stream.
+ *
+ * @param OutSegmentInformation
+ * @param OutAverageSegmentDuration
+ * @param InCurrentSegment
+ * @param LookAheadTime
+ * @param AdaptationSet
+ * @param Representation
+ */
+void FPlayPeriodHLS::GetSegmentInformation(TArray<FSegmentInformation>& OutSegmentInformation, FTimeValue& OutAverageSegmentDuration, TSharedPtrTS<const IStreamSegment> InCurrentSegment, const FTimeValue& LookAheadTime, const TSharedPtrTS<IPlaybackAssetAdaptationSet>& AdaptationSet, const TSharedPtrTS<IPlaybackAssetRepresentation>& Representation)
+{
+	OutSegmentInformation.Empty();
+	OutAverageSegmentDuration.SetToInvalid();
+
+	FTimeValue StartingTime(FTimeValue::GetZero());
+
+	const FStreamSegmentRequestHLSfmp4* CurrentSegment = static_cast<const FStreamSegmentRequestHLSfmp4*>(InCurrentSegment.Get());
+
+	if (CurrentSegment)
+	{
+		// The time of the next segment needs to be larger than that of the current. We add half the duration to the time to do that.
+		// The reason being that adding the whole duration might get us slightly further than the next segment actually is, particularly if it is in
+		// another variant playlist.
+		StartingTime = CurrentSegment->AbsoluteDateTime + (CurrentSegment->SegmentDuration / 2);
+	}
+
+	FManifestHLSInternal::ScopedLockPlaylists lock(InternalManifest);
+	const FManifestHLSInternal* pInt = InternalManifest.Get();
+	check(pInt);
+	if (!pInt)
+	{
+		return;
+	}
+
+	// The representation ID is the unique ID of the stream as a string. Convert it back
+	uint32 UniqueID;
+	LexFromString(UniqueID, *Representation->GetUniqueIdentifier());
+	TSharedPtrTS<FManifestHLSInternal::FPlaylistBase> Playlist = pInt->GetPlaylistForUniqueID(UniqueID);
+	if (!Playlist.IsValid() || !Playlist->IsVariantStream())
+	{
+		return;
+	}
+
+	// Get the bitrate of the intended variant stream
+	const int32 Bitrate = Playlist->GetBitrate();
+	// Is the playlist of this stream loaded?
+	TSharedPtrTS<FManifestHLSInternal::FMediaStream>	MediaStream = Playlist->Internal.MediaStream;
+	bool bIsIntendedStream = false;
+	if (MediaStream.IsValid())
+	{
+		bIsIntendedStream = true;
+	}
+	else
+	{
+		// Not loaded. At this point we have to _assume_ that all video variant streams are segmented the same so we search for any
+		// loaded variant and use its segmentation to return information for.
+		for(int32 VideoStreamIndex=0; VideoStreamIndex<pInt->VariantStreams.Num(); ++VideoStreamIndex)
+		{
+			MediaStream = pInt->VariantStreams[VideoStreamIndex]->Internal.MediaStream;
+			if (MediaStream.IsValid())
+			{
+				break;
+			}
+		}
+	}
+	// This should have yielded a playlist.
+	check(MediaStream.IsValid());
+	if (MediaStream.IsValid())
+	{
+		const TArray<FManifestHLSInternal::FMediaStream::FMediaSegment>& SegmentList = MediaStream->SegmentList;
+
+		FTimeValue TimeToGo(LookAheadTime);
+		FTimeValue AccumulatedDuration(FTimeValue::GetZero());
+
+		// Find the segment we need to start with.
+		int32 FirstIndex = 0;
+		for(; FirstIndex < SegmentList.Num(); ++FirstIndex)
+		{
+			if (SegmentList[FirstIndex].AbsoluteDateTime >= StartingTime)
+			{
+				break;
+			}
+		}
+
+		while(TimeToGo > FTimeValue::GetZero() && FirstIndex < SegmentList.Num())
+		{
+			FSegmentInformation& si = OutSegmentInformation.AddDefaulted_GetRef();
+			si.Duration = SegmentList[FirstIndex].Duration;
+			if (si.Duration <= FTimeValue::GetZero())
+			{
+				break;
+			}
+			// Set the actual byte size only if a byte range is defined and if we are operating on the intended stream. Otherwise use default size for duration and bitrate.
+			si.ByteSize = bIsIntendedStream && SegmentList[FirstIndex].ByteRange.IsSet() ? SegmentList[FirstIndex].ByteRange.GetNumBytes() : static_cast<int64>(Bitrate * si.Duration.GetAsSeconds() / 8);
+			AccumulatedDuration += si.Duration;
+			TimeToGo -= si.Duration;
+			++FirstIndex;
+		}
+		// Fill the remaining duration with the average segment duration or, if that is somehow not valid, the target duration.
+		FTimeValue fillDuration = MediaStream->TotalAccumulatedSegmentDuration;
+		if (!fillDuration.IsValid() || fillDuration <= FTimeValue::GetZero() || SegmentList.Num() == 0)
+		{
+			fillDuration = MediaStream->TargetDuration;
+		}
+		else
+		{
+			fillDuration /= SegmentList.Num();
+		}
+		if (fillDuration > FTimeValue::GetZero())
+		{
+			while(TimeToGo > FTimeValue::GetZero())
+			{
+				FSegmentInformation& si = OutSegmentInformation.AddDefaulted_GetRef();
+				si.Duration = fillDuration;
+				si.ByteSize = static_cast<int64>(Bitrate * si.Duration.GetAsSeconds() / 8);
+				AccumulatedDuration += si.Duration;
+				TimeToGo -= si.Duration;
+			}
+		}
+		// Set up average duration.
+		if (OutSegmentInformation.Num())
+		{
+			OutAverageSegmentDuration = AccumulatedDuration / OutSegmentInformation.Num();
+		}
+	}
+	else
+	{
+		// Not a single playlist active?
+		// Should we synthesize some information based on stream bitrate (which we have) and a fake fixed duration?
+		check(!"Now what?");
+	}
+}
+
+
+
+
+TSharedPtrTS<ITimelineMediaAsset> FPlayPeriodHLS::GetMediaAsset() const
+{
+	FManifestHLSInternal::ScopedLockPlaylists lock(InternalManifest);
+	const FManifestHLSInternal* pInt = InternalManifest.Get();
+	check(pInt && pInt->PlaybackTimeline.IsValid());
+	if (pInt && pInt->PlaybackTimeline.IsValid())
+	{
+		// HLS only has a single playback "asset". (aka Period)
+		return pInt->PlaybackTimeline->GetMediaAssetByIndex(0);
+	}
+	return TSharedPtrTS<ITimelineMediaAsset>();
+}
+
+
+
+void FPlayPeriodHLS::SelectStream(const TSharedPtrTS<IPlaybackAssetAdaptationSet>& AdaptationSet, const TSharedPtrTS<IPlaybackAssetRepresentation>& Representation, const FString& PreferredCDN)
+{
+	if (AdaptationSet.IsValid() && Representation.IsValid())
+	{
+		RefreshBlacklistState();
+
+		// The representation ID is the unique ID of the stream as a string. Convert it back
+		uint32 UniqueID;
+		LexFromString(UniqueID, *Representation->GetUniqueIdentifier());
+
+		// Which stream type is this?
+		switch(Representation->GetCodecInformation().GetStreamType())
+		{
+			case EStreamType::Video:
+			{
+				// Different from what we have actively selected?
+				if (UniqueID != ActiveVideoUniqueID)
+				{
+					// FIXME: We could emit a QoS event here. Although selecting the stream does not mean it will actually get used.
+					//        There is still a chance that (if the playlist is not loaded yet) another stream will be selected right away.
+
+					// Tell the manifest that we are now using a different stream.
+					InternalManifest->SelectActiveStreamID(UniqueID, ActiveVideoUniqueID);
+
+					ActiveVideoUniqueID = UniqueID;
+				}
+				break;
+			}
+			case EStreamType::Audio:
+			{
+				// Different from what we have actively selected?
+				if (UniqueID != ActiveAudioUniqueID)
+				{
+					// FIXME: We could emit a QoS event here. Although selecting the stream does not mean it will actually get used.
+					//        There is still a chance that (if the playlist is not loaded yet) another stream will be selected right away.
+
+					// Tell the manifest that we are now using a different stream.
+					InternalManifest->SelectActiveStreamID(UniqueID, ActiveAudioUniqueID);
+
+					ActiveAudioUniqueID = UniqueID;
+				}
+				break;
+			}
+			default:
+			{
+				break;
+			}
+		}
+	}
+}
+
+
+} // namespace Electra
+
+
diff --git a/Engine/Plugins/Media/ElectraPlayer/Source/ElectraPlayerRuntime/Private/Runtime/Player/HLS/ManifestHLS.h b/Engine/Plugins/Media/ElectraPlayer/Source/ElectraPlayerRuntime/Private/Runtime/Player/HLS/ManifestHLS.h
new file mode 100644
--- /dev/null
+++ b/Engine/Plugins/Media/ElectraPlayer/Source/ElectraPlayerRuntime/Private/Runtime/Player/HLS/ManifestHLS.h
@@ -0,0 +1,42 @@
+// Copyright Epic Games, Inc. All Rights Reserved.
+
+#pragma once
+
+#include "Player/Manifest.h"
+
+namespace Electra
+{
+struct FManifestHLSInternal;
+class IPlaylistReaderHLS;
+
+
+class FManifestHLS : public IManifest
+{
+public:
+	static TSharedPtrTS<FManifestHLS> Create(IPlayerSessionServices* SessionServices, const FParamDict& Options, IPlaylistReaderHLS* PlaylistReader, TSharedPtrTS<FManifestHLSInternal> Manifest);
+
+	virtual ~FManifestHLS();
+	virtual EType GetPresentationType() const override;
+	virtual TSharedPtrTS<IPlaybackAssetTimeline> GetTimeline() const override;
+	virtual int64 GetDefaultStartingBitrate() const override;
+	virtual FTimeValue GetMinBufferTime() const override;
+	virtual void GetStreamMetadata(TArray<FStreamMetadata>& OutMetadata, EStreamType StreamType) const override;
+	virtual IStreamReader *CreateStreamReaderHandler() override;
+
+	virtual FResult FindPlayPeriod(TSharedPtrTS<IPlayPeriod>& OutPlayPeriod, const FPlayStartPosition& StartPosition, ESearchType SearchType) override;
+
+private:
+	FManifestHLS(IPlayerSessionServices* SessionServices, const FParamDict& InOptions, IPlaylistReaderHLS* PlaylistReader, TSharedPtrTS<FManifestHLSInternal> Manifest);
+
+	FParamDict									Options;
+	TSharedPtrTS<FManifestHLSInternal>			InternalManifest;
+	IPlayerSessionServices* 					SessionServices;
+	IPlaylistReaderHLS*							PlaylistReader;
+
+};
+
+
+
+} // namespace Electra
+
+
diff --git a/Engine/Plugins/Media/ElectraPlayer/Source/ElectraPlayerRuntime/Private/Runtime/Player/HLS/MediaTag.h b/Engine/Plugins/Media/ElectraPlayer/Source/ElectraPlayerRuntime/Private/Runtime/Player/HLS/MediaTag.h
new file mode 100644
--- /dev/null
+++ b/Engine/Plugins/Media/ElectraPlayer/Source/ElectraPlayerRuntime/Private/Runtime/Player/HLS/MediaTag.h
@@ -0,0 +1,122 @@
+// Copyright Epic Games, Inc. All Rights Reserved.
+
+#pragma once
+
+#include "CoreMinimal.h"
+#include "PlayerCore.h"
+
+namespace Electra
+{
+    namespace HLSPlaylistParser
+    {
+        typedef int FMediaTagOption;
+
+        /**
+         * If an attribute doesn't support an attribute list, the parse will fail if a ":" occurs after the tag name.
+         */
+        const FMediaTagOption HasAttributeList = 1;
+        /**
+         * Parser will not throw an error if no attribute list occurs after the tag.
+         */
+        const FMediaTagOption OptionalAttributeList = HasAttributeList | (1 << 1);
+        const FMediaTagOption HasDirectValue = 1 << 3;
+
+        /**
+         * Defines tags which are applied to the hole playlist.
+         */
+        const FMediaTagOption TargetMasterPlaylist = 1 << 5;
+
+        /**
+         * Defines tags which are applied to playlist URLs.
+         */
+        const FMediaTagOption TargetPlaylistURL = 1 << 4;
+
+        /**
+         * Defines tags which are applied to the hole playlist which may only contain segments.
+         */
+        const FMediaTagOption TargetMediaPlaylist = 1 << 7;
+
+        /**
+         * Defines tags which are applied to segment URLs.
+         */
+        const FMediaTagOption TargetSegmentURL = 1 << 6;
+
+        /**
+         * The tag will be present in all MediaSegments which are followed until it's set again.
+         */
+        const FMediaTagOption ApplyOnFollowing = 1 << 8;
+
+		/**
+		 * The tag defines a URL in its value which will be recognized as separate URL. The tag should
+		 * still contain one of the Target attributes (TargetMediaPlaylist or TargetSegmentURL). The
+		 * resulting segment or playlist URL will receive all previous tags which have the proper target
+		 * defined including the tag which contains this URLDefinition attribute.
+		 * This is currently only happening for EXT-X-PREFETCH.
+		 */
+		const FMediaTagOption URLDefinition = 1 << 9;
+
+		/**
+		 * Defines whether a tag is required for the defined target. This will be checked for segments and
+		 * playlists. Segments created by a tag defined with the attribute "URLDefinition" will skip the
+		 * validation and therefore ignore "Required" tags.
+		 */
+        const FMediaTagOption Required = 1 << 10;
+
+        typedef int FAttributeOption;
+
+        /**
+         * Attributes are optional per default. The parser will fail with MissingRequiredAttribute if a required attribute
+         * is not present.
+         */
+        const FAttributeOption AttrRequired = 1;
+
+
+        /**
+         * Configuration of a media tag. This is used internally to define the HLS standard and can be used to define
+         * custom tags.
+         */
+        class FMediaTag
+        {
+        public:
+            FMediaTag(const FString& KeyVal, const FMediaTagOption& Options) : Key(KeyVal), OptionFlag(Options)
+            {}
+
+            FMediaTag(const FString& KeyVal, const FMediaTagOption& Options,
+                      const TMap<FString, FAttributeOption>& Attr) : Key(KeyVal),
+                                                                                       OptionFlag(Options),
+                                                                                       Attributes(Attr)
+            {}
+
+            FString Key;
+            FMediaTagOption OptionFlag;
+            TMap<FString, FAttributeOption> Attributes;
+
+            inline const bool HasOption(const FMediaTagOption& Option) const
+            {
+                return (OptionFlag & Option) == Option;
+            }
+
+            /**
+             * Checks if the tag contains an option which matches a bitmask of options.
+             * @param Option
+             * @return
+             */
+            inline const bool HasOptionMap(const FMediaTagOption& Option) const
+            {
+                return (OptionFlag & Option) != 0;
+            }
+        };
+
+        /**
+         * Contains the parsed tag value and attributes.
+         */
+        struct FTagContent
+        {
+            FString TagValue;
+            TMap<FString, FString> AttributeValues;
+        };
+    }
+} // namespace Electra
+
+
+
diff --git a/Engine/Plugins/Media/ElectraPlayer/Source/ElectraPlayerRuntime/Private/Runtime/Player/HLS/Parser.cpp b/Engine/Plugins/Media/ElectraPlayer/Source/ElectraPlayerRuntime/Private/Runtime/Player/HLS/Parser.cpp
new file mode 100644
--- /dev/null
+++ b/Engine/Plugins/Media/ElectraPlayer/Source/ElectraPlayerRuntime/Private/Runtime/Player/HLS/Parser.cpp
@@ -0,0 +1,613 @@
+// Copyright Epic Games, Inc. All Rights Reserved.
+
+#include "PlayerCore.h"
+
+#include "Parser.h"
+#include "Tags.h"
+#include "CoreMinimal.h"
+
+namespace Electra
+{
+    namespace HLSPlaylistParser
+    {
+        FParser::FParser()
+		: LastLineNumber(-1)
+        {
+            Tags = {
+                    {ExtINF,                    FMediaTag(ExtINF, TargetSegmentURL | HasDirectValue | Required)},
+                    {ExtXVersion,               FMediaTag(ExtXVersion,
+                                                          TargetMasterPlaylist | TargetMediaPlaylist | HasDirectValue)},
+                    {ExtXByteRange,             FMediaTag(ExtXByteRange, TargetSegmentURL | HasDirectValue)},
+                    {ExtXDiscontinuity,         FMediaTag(ExtXDiscontinuity, TargetSegmentURL)},
+                    {ExtXKey,                   FMediaTag(ExtXKey,
+                                                          TargetSegmentURL | HasAttributeList | ApplyOnFollowing, {
+                                                                  {"METHOD",            AttrRequired},
+                                                                  {"URI",               0},
+                                                                  {"IV",                0},
+                                                                  {"KEYFORMAT",         0},
+                                                                  {"KEYFORMATVERSIONS", 0},
+                                                          })},
+                    {ExtXMap,                   FMediaTag(ExtXMap,
+                                                          TargetSegmentURL | HasAttributeList | ApplyOnFollowing, {
+                                                                  {"URI",       AttrRequired},
+                                                                  {"BYTERANGE", 0},
+                                                          })},
+                    {ExtXProgramDateTime,       FMediaTag(ExtXProgramDateTime, TargetSegmentURL | HasDirectValue)},
+                    {ExtXDateRange,             FMediaTag(ExtXDateRange, TargetSegmentURL | HasAttributeList, {
+                            {"ID",               AttrRequired},
+                            {"CLASS",            0},
+                            {"START-DATE",       AttrRequired},
+                            {"END-DATE",         0},
+                            {"DURATION",         0},
+                            {"PLANNED-DURATION", 0},
+                            {"END-ON-NEXT",      0},
+                            {"SCTE35-CMD",       0},
+                            {"SCTE35-OUT",       0},
+                            {"SCTE35-IN",        0},
+                            {"X-",               0}, // Custom attributes are allowed.
+                    })},
+                    {ExtXTargetDuration,        FMediaTag(ExtXTargetDuration,
+                                                          TargetMediaPlaylist | Required | HasDirectValue)},
+                    {ExtXMediaSequence,         FMediaTag(ExtXMediaSequence, TargetMediaPlaylist | HasDirectValue)},
+                    {ExtXDiscontinuitySequence, FMediaTag(ExtXDiscontinuitySequence,
+                                                          TargetMediaPlaylist | HasDirectValue)},
+                    {ExtXPlaylistType,          FMediaTag(ExtXPlaylistType, TargetMediaPlaylist | HasDirectValue)},
+                    {ExtXIFramesOnly,           FMediaTag(ExtXIFramesOnly, TargetMediaPlaylist)},
+                    {ExtXAllowCache,            FMediaTag(ExtXAllowCache,
+                                                          TargetMediaPlaylist | TargetMasterPlaylist | HasDirectValue)},
+                    // Master playlist related tags.
+                    {ExtXMedia,                 FMediaTag(ExtXMedia, TargetMasterPlaylist | HasAttributeList, {
+                            {"TYPE",            AttrRequired},
+                            {"URI",             0},
+                            {"GROUP-ID",        AttrRequired},
+                            {"LANGUAGE",        0},
+                            {"ASSOC-LANGUAGE",  0},
+                            {"NAME",            AttrRequired},
+                            {"DEFAULT",         0},
+                            {"AUTOSELECT",      0},
+                            {"FORCED",          0},
+                            {"INSTREAM-ID",     0},
+                            {"CHARACTERISTICS", 0},
+                            {"CHANNELS",        0},
+                    })},
+                    {ExtXStreamInf,             FMediaTag(ExtXStreamInf, TargetPlaylistURL | HasAttributeList, {
+                            {"BANDWIDTH",         AttrRequired},
+                            {"AVERAGE-BANDWIDTH", 0},
+                            {"CODECS",            0}, // Should include but not MUST. So keeping it optional.
+                            {"RESOLUTION",        0},
+                            {"FRAME-RATE",        0},
+                            {"HDCP-LEVEL",        0},
+                            {"AUDIO",             0},
+                            {"VIDEO",             0},
+                            {"SUBTITLES",         0},
+                            {"CLOSED-CAPTIONS",   0},
+                    })},
+                    {ExtXIFrameStreamInf,       FMediaTag(ExtXIFrameStreamInf, TargetMasterPlaylist | HasAttributeList,
+                                                          {
+                                                                  {"URI",               AttrRequired},
+                                                                  {"BANDWIDTH",         AttrRequired},
+                                                                  {"AVERAGE-BANDWIDTH", 0},
+                                                                  {"CODECS",            0}, // Should include but not MUST. So keeping it optional.
+                                                                  {"RESOLUTION",        0},
+                                                                  {"HDCP-LEVEL",        0},
+                                                                  {"VIDEO",             0},
+                                                          })},
+                    {ExtXSessionData,           FMediaTag(ExtXSessionData, TargetMasterPlaylist | HasAttributeList, {
+                            {"DATA-ID",  AttrRequired},
+                            {"VALUE",    0}, // Must contain either.
+                            {"URI",      0},
+                            {"LANGUAGE", 0},
+                    })},
+                    {ExtXSessionKey,            FMediaTag(ExtXSessionKey, TargetMasterPlaylist | HasAttributeList, {
+                            {"METHOD",            AttrRequired},
+                            {"URI",               0},
+                            {"IV",                0},
+                            {"KEYFORMAT",         0},
+                            {"KEYFORMATVERSIONS", 0},
+                    })},
+                    {ExtXIndependentSegments,   FMediaTag(ExtXIndependentSegments,
+                                                          TargetMasterPlaylist | TargetMediaPlaylist)},
+                    {ExtXStart,                 FMediaTag(ExtXStart,
+                                                          TargetMasterPlaylist | TargetMediaPlaylist | HasAttributeList,
+                                                          {
+                                                                  {"TIME-OFFSET", AttrRequired},
+                                                                  {"PRECISE",     0},
+                                                          })},
+                    {ExtXEndlist,               FMediaTag(ExtXEndlist, TargetMasterPlaylist | TargetMediaPlaylist)},
+            };
+        }
+
+        bool FParser::HasUpcomingString(StringHelpers::FStringIterator& ContentIterator, const FString& Expected)
+        {
+            StringHelpers::FStringIterator ExpectedIterator(Expected);
+            int StringMoved = 0;
+
+            while (ExpectedIterator)
+            {
+                if (!ContentIterator || *ContentIterator != *ExpectedIterator)
+                {
+                    ContentIterator -= StringMoved;
+                    return false;
+                }
+
+                ++ContentIterator;
+                ++ExpectedIterator;
+                ++StringMoved;
+            }
+
+            ContentIterator -= StringMoved;
+            return true;
+        }
+
+        EPlaylistError FParser::ParseToken(StringHelpers::FStringIterator& ContentIterator, FString& OutValue, const Token& SearchToken)
+        {
+            bool bIgnoreTokens = (TIgnore & SearchToken) == TIgnore;
+            FString result;
+
+            while (ContentIterator)
+            {
+                Token cToken = LookupToken(*ContentIterator);
+                if (cToken != TUnknown)
+                {
+                    if ((cToken & SearchToken) == cToken)
+                    {
+                        break;
+                    }
+                    else if (!bIgnoreTokens)
+                    {
+                        return EPlaylistError::InvalidToken;
+                    }
+                }
+
+                result += *ContentIterator;
+                ++ContentIterator;
+            }
+
+            if (!ContentIterator && (TEOF & SearchToken) != TEOF)
+            {
+                return EPlaylistError::InvalidToken;
+            }
+
+            OutValue = result;
+            return EPlaylistError::None;
+        }
+
+
+        EPlaylistError FParser::SkipToken(StringHelpers::FStringIterator& ContentIterator, const Token& SearchToken)
+        {
+            bool bIgnoreTokens = (TIgnore & SearchToken) == TIgnore;
+
+            while (ContentIterator)
+            {
+                Token cToken = LookupToken(*ContentIterator);
+                if (cToken != TUnknown)
+                {
+                    if ((cToken & SearchToken) == cToken)
+                    {
+                        break;
+                    }
+                    else if (!bIgnoreTokens)
+                    {
+                        return EPlaylistError::InvalidToken;
+                    }
+                }
+                ++ContentIterator;
+            }
+
+            if (!ContentIterator && (TEOF & SearchToken) != TEOF)
+            {
+                return EPlaylistError::InvalidToken;
+            }
+
+            return EPlaylistError::None;
+        }
+
+
+        EPlaylistError FParser::ParseAttributeValue(StringHelpers::FStringIterator& ContentIterator, FString& OutValue)
+        {
+			static const FString kTextHexPrefixLowercase(TEXT("0x"));
+			static const FString kTextHexPrefixUppercase(TEXT("0X"));
+
+			if (!ContentIterator)
+            {
+                return EPlaylistError::InvalidValue;
+            }
+
+            EPlaylistError Error;
+
+            if (*ContentIterator == TCHAR('"'))
+            {
+                ++ContentIterator;
+                Error = ParseToken(ContentIterator, OutValue, TQuote | TIgnore);
+                if (Error == EPlaylistError::None)
+                {
+                    ++ContentIterator; // Skip the closing quote
+                }
+            }
+            else if (HasUpcomingString(ContentIterator, kTextHexPrefixLowercase) || HasUpcomingString(ContentIterator, kTextHexPrefixUppercase))
+            {
+                Error = ParseToken(ContentIterator, OutValue, TAttributeSeparation | TLineBreak | TEOF);
+            }
+			// Note: cast required to silence static code analyzer (isdigit() expects an int, apparently need to cast char to unsigned first to prevent sign extension)
+            else if (FChar::IsDigit(*ContentIterator) || *ContentIterator == TCHAR('-'))
+            {
+                Error = ParseToken(ContentIterator, OutValue, TAttributeSeparation | TLineBreak | TEOF);
+            }
+			// Note: cast required to silence static code analyzer (isalpha() expects an int, apparently need to cast char to unsigned first to prevent sign extension)
+            else if (FChar::IsAlpha(*ContentIterator))
+            {
+                Error = ParseToken(ContentIterator, OutValue, TAttributeSeparation | TLineBreak | TEOF);
+            }
+            else
+            {
+                Error = EPlaylistError::InvalidValue;
+            }
+
+            return Error;
+        }
+
+        EPlaylistError FParser::ParseAttributeList(StringHelpers::FStringIterator& ContentIterator, FTagContent& TagContent, const TMap<FString, FAttributeOption>& Attributes)
+        {
+            while (ContentIterator && *ContentIterator != CLineBreak && *ContentIterator != TComment)
+            {
+                if (*ContentIterator == CAttributeSeparation)
+                {
+                    ++ContentIterator;
+                    if (!ContentIterator || *ContentIterator == CLineBreak)
+                    {
+                        return EPlaylistError::InvalidToken;
+                    }
+                }
+
+                FString AttributeKey;
+                EPlaylistError Error = ParseToken(ContentIterator, AttributeKey, TAttributeDeclaration);
+                if (Error != EPlaylistError::None)
+                {
+                    return Error;
+                }
+
+                // Skip attribute declaration.
+                ++ContentIterator;
+
+                FString AttributeValue;
+                Error = ParseAttributeValue(ContentIterator, AttributeValue);
+                if (Error != EPlaylistError::None)
+                {
+                    return Error;
+                }
+
+                TagContent.AttributeValues.Add(AttributeKey, AttributeValue);
+            }
+
+            for (auto& Elem: Attributes)
+            {
+                if ((Elem.Value & AttrRequired) == AttrRequired &&
+                    !TagContent.AttributeValues.Contains(Elem.Key))
+                {
+                    return EPlaylistError::MissingRequiredAttribute;
+                }
+            }
+
+            return EPlaylistError::None;
+        }
+
+        void FParser::DetectErrorLocation(StringHelpers::FStringIterator& ContentIterator, StringHelpers::FStringIterator& ContentBegin)
+        {
+            LastLineNumber = 1;
+            int CurrentCharacter = 0;
+
+            while (ContentBegin && ContentBegin != ContentIterator)
+            {
+                if (*ContentBegin == CLineBreak)
+                {
+                    ++LastLineNumber;
+                    CurrentCharacter = 0;
+                }
+				else
+				{
+					++CurrentCharacter;
+				}
+
+                ++ContentBegin;
+            }
+
+			LastErrorMsg = FString::Printf(TEXT("invalid token in line: %d character: %d"), LastLineNumber, CurrentCharacter);
+        }
+
+        void FParser::Configure(const TMap<FString, FMediaTag>& TagConfiguration)
+        {
+			Tags.Append(TagConfiguration);
+        }
+
+        EPlaylistError FParser::Parse(const FString& PlaylistContent, FPlaylist& OutPlaylist)
+        {
+			StringHelpers::FStringIterator ContentBegin(PlaylistContent);
+			StringHelpers::FStringIterator ContentIterator(PlaylistContent);
+
+            bool bFirstLine = true;
+            EPlaylistError Error = EPlaylistError::None;
+            // Keeps all tags which are applied to the next segment or playlist.
+            TMap<FString, TArray<FTagContent>> NextItemTags;
+            // Keeps all tags which are applied to all upcoming segments or playlists.
+            TMap<FString, TArray<FTagContent>> FollowingItemTags;
+
+            // Helper to determine what kind of playlist is currently active and whether the tag which is parsed matches
+            // into this type.
+            FMediaTagOption PlaylistTypeMaster = TargetMasterPlaylist | TargetPlaylistURL;
+            FMediaTagOption PlaylistTypeMedia = TargetMediaPlaylist | TargetSegmentURL;
+
+            FMediaTagOption PlaylistTypeUndefined = PlaylistTypeMaster | PlaylistTypeMedia;
+            FMediaTagOption PlaylistTypeTargets = PlaylistTypeUndefined;
+
+            while (ContentIterator)
+            {
+                // Skip empty lines.
+                if (!bFirstLine && (*ContentIterator == CLineBreak || *ContentIterator == ' '))
+                {
+                    ++ContentIterator;
+                    continue;
+                }
+
+				static const FString kStringEXT = TEXT("#EXT");
+                bool bIsTag = HasUpcomingString(ContentIterator, kStringEXT);
+
+                bool bIsComment = *ContentIterator == CComment && !bIsTag;
+
+                // Header check which MUST be present on every master playlist and playlist.
+                if (bFirstLine)
+                {
+					static const FString kStringEXTM3U = TEXT("#EXTM3U");
+                    if (!HasUpcomingString(ContentIterator, kStringEXTM3U))
+                    {
+                        Error = EPlaylistError::MissingHeader;
+                        break;
+                    }
+                    bFirstLine = false;
+                    // Skip the hole line.
+                    Error = SkipToken(ContentIterator, TLineBreak | TEOF | TIgnore);
+                    if (Error != EPlaylistError::None)
+                    {
+                        break;
+                    }
+                    continue;
+                }
+
+				FString Url;
+
+                if (bIsTag)
+                {
+                    ++ContentIterator; // Skip the #.
+
+					FString tagNameVal;
+                    Error = ParseToken(ContentIterator, tagNameVal, TLineBreak | TColumn | TEOF);
+                    if (Error != EPlaylistError::None)
+                    {
+                        break;
+                    }
+
+					FString tagName = tagNameVal;
+                    const FMediaTag* tag = Tags.Find(tagName);
+                    if (tag == nullptr)
+                    {
+                        // Ignore unrecognized tags completely. If these should be parsed and properly added to either
+                        // playlist or segment, use Configure to define the tags properties for proper parsing.
+                        Error = SkipToken(ContentIterator, TLineBreak | TEOF | TIgnore);
+                        if (Error != EPlaylistError::None)
+                        {
+                            break;
+                        }
+                        continue;
+                    }
+
+                    FTagContent tagItem;
+
+                    // Fetch attribute list if tag requires it.
+                    if (tag->HasOption(HasAttributeList))
+                    {
+                        if (ContentIterator && *ContentIterator == CColumn)
+                        {
+                            ++ContentIterator; // Skip the column.
+                        }
+
+                        if (ContentIterator && *ContentIterator != TComment &&
+                            *ContentIterator != TLineBreak)
+                        {
+                            Error = ParseAttributeList(ContentIterator, tagItem, tag->Attributes);
+                            if (Error != EPlaylistError::None)
+                            {
+                                break;
+                            }
+                        }
+                        else if (!tag->HasOption(OptionalAttributeList))
+                        {
+                            Error = EPlaylistError::MissingAttributeList;
+                            break;
+                        }
+                    }
+                    else if (tag->HasOption(HasDirectValue))
+                    {
+                        if (ContentIterator && *ContentIterator == CColumn)
+                        {
+                            ++ContentIterator; // Skip the column.
+                        }
+
+                        Error = ParseToken(ContentIterator, tagItem.TagValue, TLineBreak | TEOF | TIgnore);
+                        if (Error != EPlaylistError::None)
+                        {
+                            break;
+                        }
+                    }
+                    else
+                    {
+                        if (ContentIterator && *ContentIterator != CLineBreak)
+                        {
+                            Error = EPlaylistError::InvalidToken;
+                            break;
+                        }
+                    }
+
+                    // Check if tag type matches into current playlist type or if we determine the playlist type based on
+                    // the current tag.
+                    if (PlaylistTypeTargets == PlaylistTypeUndefined && tag->HasOptionMap(PlaylistTypeUndefined) &&
+                        // Special case for tags which can applied to both. (ex. EXT-X-VERSION).
+                        !tag->HasOption(TargetMasterPlaylist | TargetMediaPlaylist))
+                    {
+                        if (tag->HasOptionMap(PlaylistTypeMaster))
+                        {
+                            PlaylistTypeTargets = PlaylistTypeMaster;
+                            OutPlaylist.Type = EPlaylistType::Master;
+                        }
+                        else
+                        {
+                            PlaylistTypeTargets = PlaylistTypeMedia;
+                            OutPlaylist.Type = EPlaylistType::Media;
+                        }
+                    }
+                    else if (PlaylistTypeTargets != PlaylistTypeUndefined && tag->HasOptionMap(PlaylistTypeUndefined) &&
+                             !tag->HasOptionMap(PlaylistTypeTargets))
+                    {
+                        Error = EPlaylistError::InvalidPlaylist;
+                        break;
+                    }
+
+					// In case the tag identifies a new URL, this is currently only the case for
+					// LHLS EXT-X-PREFETCH tag.
+					if (tag->HasOption(URLDefinition))
+					{
+						Url = tagItem.TagValue;
+					}
+
+                    if (tag->HasOption(TargetSegmentURL) || tag->HasOption(TargetPlaylistURL))
+                    {
+						if (tag->HasOption(ApplyOnFollowing) && bInheritTags)
+                        {
+							FollowingItemTags.FindOrAdd(tagName).Add(MoveTemp(tagItem));
+                        }
+                        else
+                        {
+							NextItemTags.FindOrAdd(tagName).Add(MoveTemp(tagItem));
+                        }
+                    }
+                    else if (tag->HasOption(TargetMasterPlaylist) || tag->HasOption(TargetMediaPlaylist))
+                    {
+                        // Ensure this tag only occurs once if it's a Media Playlist tag.
+                        if (tag->HasOption(TargetMediaPlaylist) && OutPlaylist.ContainsTag(tagName))
+                        {
+                            Error = EPlaylistError::InvalidPlaylist;
+                            break;
+                        }
+
+                        OutPlaylist.AddTag(tagName, tagItem);
+                    }
+
+                }
+                else if (bIsComment)
+                {
+                    // Skip comments completely.
+                    Error = SkipToken(ContentIterator, TLineBreak | TEOF | TIgnore);
+                    if (Error != EPlaylistError::None)
+                    {
+                        break;
+                    }
+                    continue;
+                }
+                else
+                {
+                    // Is URL.
+                    Error = ParseToken(ContentIterator, Url, TLineBreak | TEOF | TIgnore);
+                    if (Error != EPlaylistError::None)
+                    {
+                        break;
+                    }
+                }
+
+				if (!Url.IsEmpty())
+				{
+					// Before adding an URL to the playlist (variant or master), the type of the
+					// playlist needs to be distinguished. If type couldn't been defined yet than
+					// the playlist isn't valid.
+					if (PlaylistTypeTargets == PlaylistTypeUndefined)
+					{
+						Error = EPlaylistError::InvalidPlaylist;
+						break;
+					}
+
+					if (bInheritTags)
+					{
+						NextItemTags.Append(FollowingItemTags);
+					}
+
+					if ((PlaylistTypeTargets & TargetSegmentURL) == TargetSegmentURL)
+					{
+						FMediaSegment SegmentURL(NextItemTags, Url);
+
+						// Tags which define a URL won't trigger a validation, because these
+						// are custom tags outside of the standard RFC.
+						if (!bIsTag)
+						{
+							Error = SegmentURL.ValidateTags(Tags, TargetSegmentURL);
+							if (Error != EPlaylistError::None)
+							{
+								break;
+							}
+						}
+
+						OutPlaylist.AddSegment(SegmentURL);
+					}
+					else
+					{
+						FMediaPlaylist PlaylistURL(NextItemTags, Url);
+						Error = PlaylistURL.ValidateTags(Tags, TargetPlaylistURL);
+						if (Error != EPlaylistError::None)
+						{
+							break;
+						}
+
+						OutPlaylist.AddPlaylist(PlaylistURL);
+					}
+				}
+
+                if (!ContentIterator)
+                {
+                    break;
+                }
+                ++ContentIterator;
+            }
+
+            if (Error != EPlaylistError::None)
+            {
+                DetectErrorLocation(ContentIterator, ContentBegin);
+                return Error;
+            }
+
+            // Validate playlist related tags.
+            return OutPlaylist.ValidateTags(Tags, (PlaylistTypeTargets & PlaylistTypeMaster) == PlaylistTypeMaster
+                                                  ? TargetMasterPlaylist : TargetMediaPlaylist);
+        }
+
+        Token LookupToken(TCHAR p)
+        {
+            switch (p)
+            {
+                case TCHAR('#'):
+                    return TComment;
+                case TCHAR('\n'):
+                    return TLineBreak;
+                case TCHAR(':'):
+                    return TColumn;
+                case TCHAR('='):
+                    return TAttributeDeclaration;
+                case TCHAR(','):
+                    return TAttributeSeparation;
+                case TCHAR(' '):
+                    return TSpace;
+                case TCHAR('"'):
+                    return TQuote;
+                default:
+                    return TUnknown;
+            }
+        }
+    }
+} // namespace Electra
+
+
