ElectraPlayer: RELEASE. Moved from NFL to public. Note that this is done via a Delete and "Add New" operation as discussed.

diff --git a/Engine/Plugins/Media/ElectraPlayer/Source/ElectraPlayerRuntime/Private/Runtime/Player/HLS/Parser.h b/Engine/Plugins/Media/ElectraPlayer/Source/ElectraPlayerRuntime/Private/Runtime/Player/HLS/Parser.h
new file mode 100644
--- /dev/null
+++ b/Engine/Plugins/Media/ElectraPlayer/Source/ElectraPlayerRuntime/Private/Runtime/Player/HLS/Parser.h
@@ -0,0 +1,330 @@
+// Copyright Epic Games, Inc. All Rights Reserved.
+
+#pragma once
+
+#include "PlayerCore.h"
+
+#include "MediaTag.h"
+#include "Token.h"
+#include "CoreMinimal.h"
+
+#include "Utilities/StringHelpers.h"
+
+namespace Electra
+{
+    namespace HLSPlaylistParser
+    {
+
+        enum class EPlaylistError
+        {
+            None,
+            MissingRequiredAttribute,
+            MissingRequiredTag,
+            NoSuchAttribute,
+            NoSuchTag,
+            NoTagValue,
+            InvalidToken,
+            InvalidValue,
+            MissingHeader,
+            MissingAttributeList,
+            UnknownTag,
+            UnknownAttribute,
+            InvalidPlaylist,
+        };
+
+        enum class EPlaylistType
+        {
+            Master,
+            Media,
+        };
+
+        /**
+         * Custom types used to parse tag or attribute values.
+         */
+
+        /**
+         * Representation of decimal-resolution values: "<WIDTH>'x'<HEIGHT>".
+         */
+        class FResolution
+        {
+        public:
+            FResolution() : Width(0), Height(0)
+            {}
+
+            int Width;
+            int Height;
+        };
+
+        /**
+         * Representation of ByteRange values: "<Length>[@<SubRangeStart>]"
+         */
+        struct FByteRange
+        {
+            int Length;
+            /**
+             * Optional.
+             */
+            int SubRangeStart;
+        };
+
+        /**
+         * Representation of the #EXTINF value: "<Duration>[,<Title>]"
+         */
+        struct FSegmentDuration
+        {
+            double Duration;
+            FString Title;
+        };
+
+        class TTagContainer
+        {
+        private:
+            TMap<FString, TArray<FTagContent>> Tags;
+
+        public:
+
+            TTagContainer()
+            {};
+
+            TTagContainer(TMap<FString, TArray<FTagContent>>& ItemTags) : Tags(MoveTemp(ItemTags))
+            {};
+
+            void AddTag(FString& TagKey, FTagContent& Tag);
+
+            /**
+             * Checks if the tag exists within this TagContainer.
+             * @param TagName
+             * @return
+             */
+            bool ContainsTag(const FString& TagName) const;
+
+			/**
+			 * Returns the number of how often a tag occurred. Returns 0 if the tag doesn't exist.
+			 * @param TagName
+			 * @return
+			 */
+			int TagNum(const FString& TagName) const;
+
+            /**
+             * Gets a tag value based on the tag name and index and writes it to the OutValue.
+             *
+             * If the tag doesn't exist or the index doesn't match it returns NoSuchTag.
+			 * If the tag value is empty it returns NoTagValue.
+             *
+             * @tparam T
+             * @param TagName
+             * @param OutValue
+			 * @param TagIndex
+             * @return
+             */
+			EPlaylistError GetTagValue(const FString& TagName, FString& OutValue, const int& TagIndex = 0) const
+			{
+				const TArray<FTagContent>* Content = Tags.Find(TagName);
+
+				if (Content == nullptr)
+				{
+					return EPlaylistError::NoSuchTag;
+				}
+
+				if (Content->Num() <= TagIndex)
+				{
+					return EPlaylistError::NoSuchTag;
+				}
+
+				FTagContent TagContent = (*Content)[TagIndex];
+
+				if (TagContent.TagValue.IsEmpty())
+				{
+					return EPlaylistError::NoTagValue;
+				}
+
+				OutValue = TagContent.TagValue;
+                return EPlaylistError::None;
+            }
+
+            /**
+             * Gets an attribute value based on the tag and the tag index and assigns it to the OutValue.
+             *
+             * If the attribute doesn't exist it returns NoSuchAttribute.
+             * If the tag doesn't exist it returns NoSuchTag.
+             *
+             * Usage:
+             *
+             * 	FString EncryptionMethod;
+             * 	PlaylistError Error = Parser->GetTagAttributeValue("EXT-X-KEY", "METHOD", &EncryptionMethod);
+             *
+             * 	if (Error == NoSuchAttribute || Error == NoSuchTag || EncryptionMethod == "NONE") {
+             * 	  // Handle none encryption
+             * 	} else if (Error == None) {
+             * 	  // Handle encryption
+             * 	} else {
+             * 	  // Handle error case
+             * 	}
+             *
+             * @param TagName
+             * @param AttributeName
+             * @param OutValue
+             * @return
+             */
+            EPlaylistError GetTagAttributeValue(const FString& TagName, const FString& AttributeName, FString& OutValue, const int& TagIndex = 0) const
+            {
+				const TArray<FTagContent>* TagContents = Tags.Find(TagName);
+
+				if (TagContents == nullptr)
+				{
+					return EPlaylistError::NoSuchTag;
+				}
+
+				if (TagContents->Num() <= TagIndex)
+				{
+					return EPlaylistError::NoSuchTag;
+				}
+
+				FTagContent Content = (*TagContents)[TagIndex];
+				const FString* AttributeValue = Content.AttributeValues.Find(AttributeName);
+
+				if (AttributeValue == nullptr)
+				{
+					return EPlaylistError::NoSuchAttribute;
+				}
+
+				OutValue = *AttributeValue;
+				return EPlaylistError::None;
+            }
+
+            /**
+             * Validates tags against a list of tag definitions for a specific target.
+             * @param Target
+             * @return
+             */
+            EPlaylistError
+            ValidateTags(const TMap<FString, FMediaTag>& TagDefinitions,
+                         const FMediaTagOption& Target);
+        };
+
+        /**
+         * Representation of a segment URL. Please note that the map of tags will be moved to the media segment.
+         */
+        class FMediaSegment : public TTagContainer
+        {
+        public:
+            const FString URL;
+
+            FMediaSegment(TMap<FString, TArray<FTagContent>>& Tags, FString& Url) : TTagContainer(Tags), URL(MoveTemp(Url))
+            {}
+        };
+
+        /**
+         * Representation of a playlist URL. Please note that the map of tags will be moved to the playlist.
+         */
+        class FMediaPlaylist : public FMediaSegment
+        {
+        public:
+            FMediaPlaylist(TMap<FString, TArray<FTagContent>>& Tags, FString& Url) : FMediaSegment(Tags, Url)
+            {};
+        };
+
+        /**
+         * Representation of both Master and Media playlist which either contains MediaPlaylists (GetPlaylists) or
+         * MediaSegments (GetSegments). To determine the type check for PlaylistType provided in Type.
+         */
+        class FPlaylist : public TTagContainer
+        {
+        private:
+            TArray<FMediaSegment> Segments;
+            TArray<FMediaPlaylist> Playlists;
+
+        public:
+            FPlaylist()
+            {}
+
+            /**
+            * The type identifies which vector will contain values either GetSegments or GetPlaylists.
+            * @return
+            */
+            EPlaylistType Type;
+
+            inline const TArray<FMediaSegment>& GetSegments() const
+            {
+                return Segments;
+            }
+
+            inline const TArray<FMediaPlaylist>& GetPlaylists() const
+            {
+                return Playlists;
+            }
+
+            /**
+             * Adds a segment to the list by moving it.
+             * @param Segment
+             */
+            void AddSegment(FMediaSegment& Segment);
+
+            /**
+             * Adds a playlist to the list by moving it.
+             * @param Playlist
+             */
+            void AddPlaylist(FMediaPlaylist& Playlist);
+        };
+
+        class FParser
+        {
+        private:
+            FString	LastErrorMsg;
+            int LastLineNumber;
+
+            TMap<FString, FMediaTag> Tags;
+
+
+            EPlaylistError ParseAttributeList(StringHelpers::FStringIterator& ContentIterator, FTagContent& TagContent, const TMap<FString, FAttributeOption>& Attributes);
+
+            EPlaylistError ParseAttributeValue(StringHelpers::FStringIterator& ContentIterator, FString& OutValue);
+
+            EPlaylistError ParseToken(StringHelpers::FStringIterator& ContentIterator, FString& OutValue, const Token& SearchToken);
+
+            EPlaylistError SkipToken(StringHelpers::FStringIterator& ContentIterator, const Token& SkipToken);
+
+            bool HasUpcomingString(StringHelpers::FStringIterator& ContentIterator, const FString& Expected);
+
+            void DetectErrorLocation(StringHelpers::FStringIterator& ContentIterator, StringHelpers::FStringIterator& ContentBegin);
+
+        public:
+            FParser();
+
+            /**
+             * If set to true all tags which have the property ApplyOnFollowing will be added to the tag map of each
+             * segment after the tag occurred. If set to false the tag is only added to the map of the first segment
+             * and it's the responsibility of the outer implementation to apply values to the upcoming segments.
+             */
+            bool bInheritTags = false;
+
+            /**
+             * Provides further MediaTag configurations on top of the HLS version 6 standard.
+             * @param TagConfiguration Map where
+             */
+            void Configure(const TMap<FString, FMediaTag>& TagConfiguration);
+
+            /**
+             * Parses the PlaylistContent and puts the data into OutPlaylist.
+             * Returns PlaylistError::None if no error occurred.
+             * @param PlaylistContent
+             * @param OutPlaylist
+             * @return
+             */
+            EPlaylistError Parse(const FString& PlaylistContent, FPlaylist& OutPlaylist);
+
+            inline FString GetLastErrorMessage()
+            {
+                return LastErrorMsg;
+            }
+
+            inline int GetLastLineNumber()
+            {
+                return LastLineNumber;
+            }
+        };
+    }
+} // namespace Electra
+
+
+
diff --git a/Engine/Plugins/Media/ElectraPlayer/Source/ElectraPlayerRuntime/Private/Runtime/Player/HLS/Playlist.cpp b/Engine/Plugins/Media/ElectraPlayer/Source/ElectraPlayerRuntime/Private/Runtime/Player/HLS/Playlist.cpp
new file mode 100644
--- /dev/null
+++ b/Engine/Plugins/Media/ElectraPlayer/Source/ElectraPlayerRuntime/Private/Runtime/Player/HLS/Playlist.cpp
@@ -0,0 +1,20 @@
+// Copyright Epic Games, Inc. All Rights Reserved.
+
+#include "Parser.h"
+
+namespace Electra
+{
+    namespace HLSPlaylistParser
+    {
+        void FPlaylist::AddSegment(HLSPlaylistParser::FMediaSegment& Segment)
+        {
+            Segments.Add(MoveTemp(Segment));
+        }
+
+        void FPlaylist::AddPlaylist(FMediaPlaylist& Playlist)
+        {
+            Playlists.Add(MoveTemp(Playlist));
+        }
+    }
+} // namespace Electra
+
diff --git a/Engine/Plugins/Media/ElectraPlayer/Source/ElectraPlayerRuntime/Private/Runtime/Player/HLS/PlaylistReaderHLS.cpp b/Engine/Plugins/Media/ElectraPlayer/Source/ElectraPlayerRuntime/Private/Runtime/Player/HLS/PlaylistReaderHLS.cpp
new file mode 100644
--- /dev/null
+++ b/Engine/Plugins/Media/ElectraPlayer/Source/ElectraPlayerRuntime/Private/Runtime/Player/HLS/PlaylistReaderHLS.cpp
@@ -0,0 +1,1087 @@
+// Copyright Epic Games, Inc. All Rights Reserved.
+
+#include "PlayerCore.h"
+#include "Player/Manifest.h"
+#include "HTTP/HTTPManager.h"
+#include "SynchronizedClock.h"
+#include "InitSegmentCacheHLS.h"
+#include "PlaylistReaderHLS.h"
+#include "ManifestBuilderHLS.h"
+#include "ManifestHLS.h"
+#include "LHLSTags.h"
+#include "EpicTags.h"
+#include "Utilities/HashFunctions.h"
+#include "Utilities/TimeUtilities.h"
+#include "HAL/LowLevelMemTracker.h"
+#include "ElectraPlayerPrivate.h"
+#include "Player/AdaptiveStreamingPlayerResourceRequest.h"
+
+#define ERRCODE_HLS_PARSER_ERROR							1
+#define ERRCODE_HLS_NO_MASTER_PLAYLIST						2
+#define ERRCODE_HLS_MASTER_PLAYLIST_DOWNLOAD_FAILED			3
+#define ERRCODE_HLS_MASTER_PLAYLIST_PARSING_FAILED			4
+#define ERRCODE_HLS_VARIANT_PLAYLIST_DOWNLOAD_FAILED		5
+#define ERRCODE_HLS_VARIANT_PLAYLIST_PARSING_FAILED			6
+
+
+
+namespace Electra
+{
+
+const FString IPlaylistReaderHLS::OptionKeyInitialBitrate(TEXT("initial_bitrate"));														//!< (int64) value indicating the bitrate to start with.
+const FString IPlaylistReaderHLS::OptionKeyLiveSeekableStartOffset(TEXT("seekable_range_live_start_offset"));							//!< (FTimeValue) value specifying how many seconds away from the Live media timeline the seekable range should start.
+const FString IPlaylistReaderHLS::OptionKeyLiveSeekableEndOffset(TEXT("seekable_range_live_end_offset"));								//!< (FTimeValue) value specifying how many seconds away from the Live media timeline the seekable range should end.
+const FString IPlaylistReaderHLS::OptionKeyLiveSeekableEndOffsetAudioOnly(TEXT("seekable_range_live_end_offset_audioonly"));			//!< (FTimeValue) value specifying how many seconds away from the Live media timeline the seekable range should end for audio-only playlists.
+const FString IPlaylistReaderHLS::OptionKeyLiveSeekableEndOffsetBeConservative(TEXT("seekable_range_live_end_offset_conservative"));	//!< (bool) true to use a larger Live edge distance, false to go with the smaller absolute difference
+
+const FString IPlaylistReaderHLS::OptionKeyMasterPlaylistLoadConnectTimeout (TEXT("master_playlist_connection_timeout") );		//!< (FTimeValue) value specifying connection timeout fetching the master playlist
+const FString IPlaylistReaderHLS::OptionKeyMasterPlaylistLoadNoDataTimeout  (TEXT("master_playlist_nodata_timeout")     );		//!< (FTimeValue) value specifying no-data timeout fetching the master playlist
+const FString IPlaylistReaderHLS::OptionKeyVariantPlaylistLoadConnectTimeout(TEXT("variant_playlist_connection_timeout"));		//!< (FTimeValue) value specifying connection timeout fetching a variant playlist the first time
+const FString IPlaylistReaderHLS::OptionKeyVariantPlaylistLoadNoDataTimeout (TEXT("variant_playlist_nodata_timeout")    );		//!< (FTimeValue) value specifying no-data timeout fetching a variant playlist the first time
+const FString IPlaylistReaderHLS::OptionKeyUpdatePlaylistLoadConnectTimeout (TEXT("update_playlist_connection_timeout") );		//!< (FTimeValue) value specifying connection timeout fetching a variant playlist repeatedly
+const FString IPlaylistReaderHLS::OptionKeyUpdatePlaylistLoadNoDataTimeout  (TEXT("update_playlist_nodata_timeout" )    );		//!< (FTimeValue) value specifying no-data timeout fetching a variant playlist repeatedly
+
+
+
+/**
+ * This class is responsible for downloading HLS playlists and parsing them.
+ */
+class FPlaylistReaderHLS : public IPlaylistReaderHLS, public FMediaThread
+{
+public:
+	FPlaylistReaderHLS();
+	void Initialize(IPlayerSessionServices* PlayerSessionServices);
+
+	virtual ~FPlaylistReaderHLS();
+
+	/**
+	 * Returns the type of playlist format.
+	 * For this implementation it will be "hls".
+	 *
+	 * @return "hls" to indicate this is an HLS playlist.
+	 */
+	virtual const FString& GetPlaylistType() const override
+	{
+		static FString Type("hls");
+		return Type;
+	}
+
+	/**
+	 * Loads and parses the playlist.
+	 *
+	 * @param URL     URL of the playlist to load
+	 * @param Preferences
+	 *                User preferences for initial stream selection.
+	 * @param Options Options for the reader and parser
+	 */
+	virtual void LoadAndParse(const FString& URL, const FStreamPreferences& Preferences, const FParamDict& Options) override;
+
+	/**
+	 * Returns the URL from which the playlist was loaded (or supposed to be loaded).
+	 *
+	 * @return The playlist URL
+	 */
+	virtual FString GetURL() const override;
+
+	/**
+	 * Returns an interface to the manifest created from the loaded HLS playlists.
+	 *
+	 * @return A shared manifest interface pointer.
+	 */
+	virtual TSharedPtrTS<IManifest> GetManifest() override;
+
+	/**
+	 * Requests loading of a HLS playlist.
+	 * This can either be the master playlist or a variant playlist needed to carry
+	 * out playback on a certain quality level.
+	 *
+	 * @param LoadRequest
+	 */
+	virtual void RequestPlaylistLoad(const FPlaylistLoadRequestHLS& LoadRequest) override;
+
+private:
+	class FPlaylistRequest : public TSharedFromThis<FPlaylistRequest, ESPMode::ThreadSafe>
+	{
+	public:
+		FPlaylistRequest(const FPlaylistLoadRequestHLS& InPlaylistLoadRequest, const FParamDict& InOptions)
+			: PlaylistLoadRequest(InPlaylistLoadRequest)
+			, Options(InOptions)
+			, bIsMasterPlaylist(false)
+		{
+		}
+		~FPlaylistRequest()
+		{
+			ReceiveBuffer.Reset();
+			HTTPRequest.Reset();
+		}
+
+		TSharedPtrTS<IElectraHttpManager::FReceiveBuffer> GetReceiveBuffer() const
+		{
+			return ReceiveBuffer;
+		}
+
+		void SetExecuteAtUTC(const FTimeValue& AtUTCTime)
+		{
+			ExecuteAtUTC = AtUTCTime;
+		}
+
+		const FTimeValue& GetExecuteAtUTC() const
+		{
+			return ExecuteAtUTC;
+		}
+
+		void SetRetryInfo(const TSharedPtrTS<HTTP::FRetryInfo>& InRetryInfo)
+		{
+			RetryInfo = InRetryInfo;
+		}
+
+		const FPlaylistLoadRequestHLS& GetPlaylistLoadRequest() const
+		{
+			return PlaylistLoadRequest;
+		}
+
+		FPlaylistLoadRequestHLS& GetPlaylistLoadRequest()
+		{
+			return PlaylistLoadRequest;
+		}
+
+		const IElectraHttpManager::FRequest* GetHTTPRequest() const
+		{
+			return HTTPRequest.Get();
+		}
+
+		const HTTP::FConnectionInfo* GetConnectionInfo() const
+		{
+			return HTTPRequest.IsValid() ? &HTTPRequest->ConnectionInfo : &StaticRequestConnectionInfo;
+		}
+
+		void SetIsMasterPlaylist(bool bInIsMasterPlaylist)
+		{
+			bIsMasterPlaylist = bInIsMasterPlaylist;
+		}
+
+		bool IsMasterPlaylist() const
+		{
+			return bIsMasterPlaylist;
+		}
+
+		void Execute(TSharedPtrTS<IElectraHttpManager::FProgressListener> ProgressListener, IPlayerSessionServices* PlayerSessionServices);
+		void Cancel(IPlayerSessionServices* PlayerSessionServices);
+
+		void ExecuteStatic(IAdaptiveStreamingPlayerResourceProvider* InStaticResourceProvider, TWeakPtr<FMediaSemaphore, ESPMode::ThreadSafe> DoneSignal);
+		int32 HandleStaticCompletion();
+		void PrepareStaticResult();
+
+	private:
+		class FStaticResourceRequest : public IAdaptiveStreamingPlayerResourceRequest
+		{
+		public:
+			FStaticResourceRequest(TWeakPtr<FPlaylistRequest, ESPMode::ThreadSafe> InOwningRequest, TWeakPtr<FMediaSemaphore, ESPMode::ThreadSafe> InDoneSignal, FString InURL)
+				: OwningRequest(InOwningRequest), DoneSignal(InDoneSignal), URL(InURL), bIsDone(false)
+			{ }
+
+			virtual ~FStaticResourceRequest()
+			{ }
+
+			virtual EPlaybackResourceType GetResourceType() const override
+			{ return EPlaybackResourceType::Playlist; }
+
+			virtual FString GetResourceURL() const override
+			{ return URL; }
+
+			virtual void SetPlaybackData(TSharedPtr<TArray<uint8>, ESPMode::ThreadSafe>	PlaybackData) override
+			{ Data = PlaybackData; }
+
+			virtual void SignalDataReady() override;
+
+			bool IsDone() const
+			{ return bIsDone;}
+
+			void Cancel()
+			{
+				DoneSignal.Reset();
+				OwningRequest.Reset();
+			}
+			TSharedPtr<TArray<uint8>, ESPMode::ThreadSafe> GetData()
+			{ return Data; }
+
+		private:
+			TWeakPtr<FPlaylistRequest, ESPMode::ThreadSafe>		OwningRequest;
+			TWeakPtr<FMediaSemaphore, ESPMode::ThreadSafe>		DoneSignal;
+			FString												URL;
+			TSharedPtr<TArray<uint8>, ESPMode::ThreadSafe>		Data;
+			bool												bIsDone;
+		};
+
+		FPlaylistLoadRequestHLS									PlaylistLoadRequest;
+		FParamDict												Options;
+		FTimeValue												ExecuteAtUTC;
+		TSharedPtrTS<IElectraHttpManager::FReceiveBuffer>		ReceiveBuffer;
+		TSharedPtrTS<IElectraHttpManager::FRequest>				HTTPRequest;
+		TSharedPtrTS<HTTP::FRetryInfo>							RetryInfo;
+		bool													bIsMasterPlaylist;
+
+		TSharedPtr<FStaticResourceRequest, ESPMode::ThreadSafe>	StaticRequest;
+		HTTP::FConnectionInfo									StaticRequestConnectionInfo;
+	};
+	using FPlaylistRequestPtr = TSharedPtr<FPlaylistRequest, ESPMode::ThreadSafe>;
+	//using FPlaylistRequestRef = TSharedRef<FPlaylistRequest, ESPMode::ThreadSafe>;
+
+	void Close();
+	void StartWorkerThread();
+	void StopWorkerThread();
+	void WorkerThread(void);
+	void HandleEnqueuedPlaylistDownloads(const FTimeValue& TimeNow);
+	void HandleCompletedPlaylistDownloads(const FTimeValue& TimeNow);
+	void HandleStaticRequestCompletions(const FTimeValue& TimeNow);
+	void CheckForPlaylistUpdate(const FTimeValue& TimeNow);
+
+	void PostError(const FString& Message, uint16 Code, UEMediaError Error = UEMEDIA_ERROR_OK);
+	void LogMessage(IInfoLog::ELevel Level, const FString& Message);
+
+	FErrorDetail CreateErrorAndLog(const FString& Message, uint16 Code, UEMediaError Error = UEMEDIA_ERROR_OK);
+
+	void EnqueueLoadPlaylist(const FPlaylistLoadRequestHLS& InPlaylistLoadRequest, bool bIsMasterPlaylist);
+	void EnqueueRetryLoadPlaylist(FPlaylistRequestPtr RequestToRetry, const TSharedPtrTS<HTTP::FRetryInfo>& RetryInfo, const FTimeValue& AtUTCTime);
+
+
+	void HTTPCompletionCallback(const IElectraHttpManager::FRequest* Request);
+
+	FErrorDetail ParsePlaylist(FPlaylistRequestPtr FromRequest);
+
+
+	IPlayerSessionServices*									PlayerSessionServices;
+	FStreamPreferences										StreamPreferences;
+	FParamDict												Options;
+	FString													MasterPlaylistURL;
+	TSharedPtr<FMediaSemaphore, ESPMode::ThreadSafe>		WorkerThreadSignal;
+	bool													bIsWorkerThreadStarted;
+	volatile bool											bTerminateWorkerThread;
+
+	FMediaCriticalSection									PendingPlaylistRequestsLock;
+	TArray<FPlaylistRequestPtr>								PendingPlaylistRequests;
+	TArray<FPlaylistRequestPtr>								StaticResourcePlaylistRequests;
+	TArray<FPlaylistRequestPtr>								EnqueuedPlaylistRequests;
+	TMediaQueueDynamic<FPlaylistRequestPtr>					CompletedPlaylistRequests;
+	TDoubleLinkedList<FPlaylistLoadRequestHLS>				InitiallyRequiredPlaylistLoadRequests;				//!< Playlists that need to be fetched and parsed before we can report metadata ready to the player.
+	TSharedPtrTS<IElectraHttpManager::FProgressListener>	ProgressListener;
+
+	IManifestBuilderHLS* 									Builder;
+	TSharedPtrTS<FManifestHLSInternal>						Manifest;
+	FErrorDetail											LastErrorDetail;
+
+	TSharedPtrTS<IManifest>									PlayerManifest;
+};
+
+
+/***************************************************************************************************************************************************/
+/***************************************************************************************************************************************************/
+/***************************************************************************************************************************************************/
+
+IPlaylistReader* IPlaylistReaderHLS::Create(IPlayerSessionServices* PlayerSessionServices)
+{
+	FPlaylistReaderHLS* PlaylistReader = new FPlaylistReaderHLS;
+	if (PlaylistReader)
+	{
+		PlaylistReader->Initialize(PlayerSessionServices);
+	}
+	return PlaylistReader;
+}
+
+/***************************************************************************************************************************************************/
+/***************************************************************************************************************************************************/
+/***************************************************************************************************************************************************/
+
+
+FPlaylistReaderHLS::FPlaylistReaderHLS()
+	: FMediaThread("ElectraPlayer::HLS Playlist")
+	, PlayerSessionServices(nullptr)
+	, bIsWorkerThreadStarted(false)
+	, bTerminateWorkerThread(false)
+	, Builder(nullptr)
+{
+	WorkerThreadSignal = MakeShared<FMediaSemaphore, ESPMode::ThreadSafe>();
+}
+
+FPlaylistReaderHLS::~FPlaylistReaderHLS()
+{
+	Close();
+}
+
+FString FPlaylistReaderHLS::GetURL() const
+{
+	return MasterPlaylistURL;
+}
+
+TSharedPtrTS<IManifest> FPlaylistReaderHLS::GetManifest()
+{
+	return PlayerManifest;
+}
+
+void FPlaylistReaderHLS::Initialize(IPlayerSessionServices* InPlayerSessionServices)
+{
+	PlayerSessionServices = InPlayerSessionServices;
+	ProgressListener = MakeSharedTS<IElectraHttpManager::FProgressListener>();
+	ProgressListener->CompletionDelegate = Electra::MakeDelegate(this, &FPlaylistReaderHLS::HTTPCompletionCallback);
+}
+
+void FPlaylistReaderHLS::Close()
+{
+	StopWorkerThread();
+}
+
+void FPlaylistReaderHLS::StartWorkerThread()
+{
+	check(!bIsWorkerThreadStarted);
+
+	bTerminateWorkerThread = false;
+	//ThreadSetPriority( ... );
+	//ThreadSetCoreAffinity( ... );
+	//ThreadSetStackSize( ... );
+	ThreadStart(Electra::MakeDelegate(this, &FPlaylistReaderHLS::WorkerThread));
+	bIsWorkerThreadStarted = true;
+}
+
+void FPlaylistReaderHLS::StopWorkerThread()
+{
+	if (bIsWorkerThreadStarted)
+	{
+		bTerminateWorkerThread = true;
+		WorkerThreadSignal->Release();
+		ThreadWaitDone();
+		ThreadReset();
+		bIsWorkerThreadStarted = false;
+	}
+}
+
+void FPlaylistReaderHLS::PostError(const FString& InMessage, uint16 InCode, UEMediaError InError)
+{
+	LastErrorDetail.Clear();
+	LastErrorDetail.SetError(InError != UEMEDIA_ERROR_OK ? InError : UEMEDIA_ERROR_DETAIL);
+	LastErrorDetail.SetFacility(Facility::EFacility::HLSPlaylistReader);
+	LastErrorDetail.SetCode(InCode);
+	LastErrorDetail.SetMessage(InMessage);
+	check(PlayerSessionServices);
+	if (PlayerSessionServices)
+	{
+		PlayerSessionServices->PostError(LastErrorDetail);
+	}
+}
+
+FErrorDetail FPlaylistReaderHLS::CreateErrorAndLog(const FString& InMessage, uint16 InCode, UEMediaError InError)
+{
+	FErrorDetail err;
+	err.SetError(InError != UEMEDIA_ERROR_OK ? InError : UEMEDIA_ERROR_DETAIL);
+	err.SetFacility(Facility::EFacility::HLSPlaylistReader);
+	err.SetCode(InCode);
+	err.SetMessage(InMessage);
+	check(PlayerSessionServices);
+	if (PlayerSessionServices)
+	{
+		PlayerSessionServices->PostLog(Facility::EFacility::HLSPlaylistReader, IInfoLog::ELevel::Error, err.GetPrintable());
+	}
+	return err;
+}
+
+
+void FPlaylistReaderHLS::LogMessage(IInfoLog::ELevel Level, const FString& Message)
+{
+	if (PlayerSessionServices)
+	{
+		PlayerSessionServices->PostLog(Facility::EFacility::HLSPlaylistReader, Level, Message);
+	}
+}
+
+void FPlaylistReaderHLS::LoadAndParse(const FString& URL, const FStreamPreferences& Preferences, const FParamDict& InOptions)
+{
+	StreamPreferences = Preferences;
+	Options 		  = InOptions;
+	MasterPlaylistURL = URL;
+	StartWorkerThread();
+}
+
+void FPlaylistReaderHLS::RequestPlaylistLoad(const FPlaylistLoadRequestHLS& LoadRequest)
+{
+	EnqueueLoadPlaylist(LoadRequest, false);
+}
+
+void FPlaylistReaderHLS::EnqueueLoadPlaylist(const FPlaylistLoadRequestHLS& InPlaylistLoadRequest, bool bIsMasterPlaylist)
+{
+	FPlaylistRequestPtr Request = MakeShared<FPlaylistRequest, ESPMode::ThreadSafe>(InPlaylistLoadRequest, Options);
+	Request->SetIsMasterPlaylist(bIsMasterPlaylist);
+	EnqueuedPlaylistRequests.Push(Request);
+	WorkerThreadSignal->Release();
+}
+
+void FPlaylistReaderHLS::EnqueueRetryLoadPlaylist(FPlaylistRequestPtr RequestToRetry, const TSharedPtrTS<HTTP::FRetryInfo>& RetryInfo, const FTimeValue& AtUTCTime)
+{
+	RequestToRetry->SetExecuteAtUTC(AtUTCTime);
+	RequestToRetry->SetRetryInfo(RetryInfo);
+	EnqueuedPlaylistRequests.Push(RequestToRetry);
+	WorkerThreadSignal->Release();
+}
+
+
+void FPlaylistReaderHLS::HTTPCompletionCallback(const IElectraHttpManager::FRequest* InRequest)
+{
+	FPlaylistRequestPtr Request;
+	PendingPlaylistRequestsLock.Lock();
+	for(int32 i=0,iMax=PendingPlaylistRequests.Num(); i<iMax; ++i)
+	{
+		if (PendingPlaylistRequests[i]->GetHTTPRequest() == InRequest)
+		{
+			Request = PendingPlaylistRequests[i];
+			PendingPlaylistRequests.RemoveAtSwap(i);
+			break;
+		}
+	}
+	PendingPlaylistRequestsLock.Unlock();
+	if (Request)
+	{
+		// Add to the list of completed readers.
+		// Cannot release the actual reader here since we are inside the reader's callback!
+		CompletedPlaylistRequests.Push(Request);
+	}
+	// Notify the worker thread that there is work to be done.
+	WorkerThreadSignal->Release();
+}
+
+void FPlaylistReaderHLS::CheckForPlaylistUpdate(const FTimeValue& TimeNow)
+{
+	if (Manifest.IsValid())
+	{
+		TArray<TSharedPtrTS<FManifestHLSInternal::FPlaylistBase>> ActiveStreams;
+		Manifest->GetActiveStreams(ActiveStreams);
+		for(int32 i=0, iMax=ActiveStreams.Num(); i<iMax; ++i)
+		{
+			if (!ActiveStreams[i]->Internal.bReloadTriggered && TimeNow >= ActiveStreams[i]->Internal.ExpiresAtTime)
+			{
+				ActiveStreams[i]->Internal.bReloadTriggered 					= true;
+				ActiveStreams[i]->Internal.PlaylistLoadRequest.LoadType 		= FPlaylistLoadRequestHLS::ELoadType::Update;
+				ActiveStreams[i]->Internal.PlaylistLoadRequest.RequestedAtTime  = TimeNow;
+				EnqueueLoadPlaylist(ActiveStreams[i]->Internal.PlaylistLoadRequest, false);
+			}
+		}
+	}
+}
+
+void FPlaylistReaderHLS::WorkerThread(void)
+{
+	LLM_SCOPE(ELLMTag::ElectraPlayer);
+	CSV_SCOPED_TIMING_STAT(ElectraPlayer, PlaylistReaderHLS_Worker);
+
+	Builder = IManifestBuilderHLS::Create(PlayerSessionServices);
+
+	// Setup the playlist load request for the master playlist.
+	FPlaylistLoadRequestHLS PlaylistLoadRequest;
+	PlaylistLoadRequest.URL 			 = MasterPlaylistURL;
+	PlaylistLoadRequest.RequestedAtTime  = PlayerSessionServices->GetSynchronizedUTCTime()->GetTime();
+	PlaylistLoadRequest.InternalUniqueID = 0;
+	EnqueueLoadPlaylist(PlaylistLoadRequest, true);
+
+	while(!bTerminateWorkerThread)
+	{
+		WorkerThreadSignal->Obtain(1000 * 100);
+		if (bTerminateWorkerThread)
+		{
+			break;
+		}
+		FTimeValue Now = PlayerSessionServices->GetSynchronizedUTCTime()->GetTime();
+
+		// Check for completed static resource requests.
+		// NOTE: This needs to be done first in order to move the completed ones being moved onto the completed list to be
+		//       handled immediately in the following HandleCompletedPlaylistDownloads() call!!
+		HandleStaticRequestCompletions(Now);
+
+		// Check completed playlist downloads
+		HandleCompletedPlaylistDownloads(Now);
+
+		// Execute enqueued playlist download requests.
+		HandleEnqueuedPlaylistDownloads(Now);
+
+		// Check which playlists need to be reloaded.
+		CheckForPlaylistUpdate(Now);
+	}
+
+
+	// No playlists are required any more.
+	InitiallyRequiredPlaylistLoadRequests.Empty();
+
+	// Cancel any pending static resource requests.
+	for(int32 i=0,iMax=StaticResourcePlaylistRequests.Num(); i<iMax; ++i)
+	{
+		StaticResourcePlaylistRequests[i]->Cancel(PlayerSessionServices);
+	}
+	StaticResourcePlaylistRequests.Empty();
+
+	// Cancel all pending playlist downloads.
+	// NOTE: To avoid a possible race condition and deadlock with downloads that will just complete now and call into the HTTPCompletionCallback()
+	//       where the mutex needs to get locked we swap the pending list out for an empty one and iterate the pending ones outside the lock.
+	PendingPlaylistRequestsLock.Lock();
+	TArray<FPlaylistRequestPtr> ActivePendingRequests;
+	Swap(PendingPlaylistRequests, ActivePendingRequests);
+	PendingPlaylistRequestsLock.Unlock();
+	// First cancel the request. This waits until the request has been removed from the manager and is no longer active.
+	for(int32 i=0,iMax=ActivePendingRequests.Num(); i<iMax; ++i)
+	{
+		ActivePendingRequests[i]->Cancel(PlayerSessionServices);
+	}
+	// With all requests canceled and removed from the manager we can delete them now.
+	for(int32 i=0,iMax=ActivePendingRequests.Num(); i<iMax; ++i)
+	{
+		ActivePendingRequests[i].Reset();
+	}
+	ActivePendingRequests.Empty();
+
+	// Delete any unprocessed enqueued requests
+	while(EnqueuedPlaylistRequests.Num())
+	{
+		EnqueuedPlaylistRequests.Pop();
+	}
+
+	// Delete any unprocessed completed downloads
+	while(!CompletedPlaylistRequests.IsEmpty())
+	{
+		CompletedPlaylistRequests.Pop();
+	}
+
+	delete Builder;
+	Builder = nullptr;
+}
+
+
+void FPlaylistReaderHLS::FPlaylistRequest::ExecuteStatic(IAdaptiveStreamingPlayerResourceProvider* InStaticResourceProvider, TWeakPtr<FMediaSemaphore, ESPMode::ThreadSafe> DoneSignal)
+{
+	StaticRequest = MakeShared<FStaticResourceRequest, ESPMode::ThreadSafe>(SharedThis(this), DoneSignal, PlaylistLoadRequest.URL);
+	InStaticResourceProvider->ProvideStaticPlaybackDataForURL(StaticRequest);
+}
+
+void FPlaylistReaderHLS::FPlaylistRequest::FStaticResourceRequest::SignalDataReady()
+{
+	bIsDone = true;
+	TSharedPtr<FPlaylistReaderHLS::FPlaylistRequest, ESPMode::ThreadSafe> Req = OwningRequest.Pin();
+	if (Req.IsValid())
+	{
+		TSharedPtr<FMediaSemaphore, ESPMode::ThreadSafe> Sig = DoneSignal.Pin();
+		if (Sig.IsValid())
+		{
+			Sig->Release();
+		}
+	}
+}
+
+int32 FPlaylistReaderHLS::FPlaylistRequest::HandleStaticCompletion()
+{
+	if (StaticRequest.Get())
+	{
+		if (StaticRequest->IsDone())
+		{
+			TSharedPtr<TArray<uint8>, ESPMode::ThreadSafe> DataPtr = StaticRequest->GetData();
+			return DataPtr.IsValid() ? 1 : -1;
+		}
+		return 0;
+	}
+	return -1;
+}
+
+void FPlaylistReaderHLS::FPlaylistRequest::PrepareStaticResult()
+{
+	StaticRequestConnectionInfo.EffectiveURL		  = PlaylistLoadRequest.URL;
+	StaticRequestConnectionInfo.bIsConnected		  = true;
+	StaticRequestConnectionInfo.bHaveResponseHeaders  = true;
+	StaticRequestConnectionInfo.bWasAborted 		  = false;
+	StaticRequestConnectionInfo.bHasFinished		  = true;
+	StaticRequestConnectionInfo.ContentLength   	  = 0;
+	StaticRequestConnectionInfo.BytesReadSoFar  	  = 0;
+	StaticRequestConnectionInfo.HTTPVersionReceived   = 11;
+	StaticRequestConnectionInfo.StatusInfo.HTTPStatus = 200;
+
+	ReceiveBuffer = MakeSharedTS<IElectraHttpManager::FReceiveBuffer>();
+	ReceiveBuffer->Buffer.SetEOD();
+
+	TSharedPtr<TArray<uint8>, ESPMode::ThreadSafe> DataPtr = StaticRequest->GetData();
+	if (DataPtr.IsValid())
+	{
+		TArray<uint8>& DataArray = *DataPtr;
+		if (DataArray.Num())
+		{
+			ReceiveBuffer->Buffer.EnlargeTo(DataArray.Num());
+			ReceiveBuffer->Buffer.PushData(DataArray.GetData(), DataArray.Num());
+			StaticRequestConnectionInfo.ContentLength = StaticRequestConnectionInfo.BytesReadSoFar = DataArray.Num();
+		}
+	}
+}
+
+
+
+void FPlaylistReaderHLS::FPlaylistRequest::Execute(TSharedPtrTS<IElectraHttpManager::FProgressListener> InProgressListener, IPlayerSessionServices* InPlayerSessionServices)
+{
+	ReceiveBuffer = MakeSharedTS<IElectraHttpManager::FReceiveBuffer>();
+	HTTPRequest = MakeSharedTS<IElectraHttpManager::FRequest>();
+	HTTPRequest->Parameters.URL = PlaylistLoadRequest.URL;
+	HTTPRequest->Parameters.AcceptEncoding.Set("");		// setting an empty string enables all supported encodings (ie. gzip)
+	HTTPRequest->ConnectionInfo.RetryInfo = RetryInfo;
+
+	// Set connection timeouts for master and variant playlist retrieval.
+	if (bIsMasterPlaylist)
+	{
+		HTTPRequest->Parameters.ConnectTimeout = Options.GetValue(IPlaylistReaderHLS::OptionKeyMasterPlaylistLoadConnectTimeout).SafeGetTimeValue(FTimeValue().SetFromMilliseconds(1000 * 8));
+		HTTPRequest->Parameters.NoDataTimeout  = Options.GetValue(IPlaylistReaderHLS::OptionKeyMasterPlaylistLoadNoDataTimeout).SafeGetTimeValue(FTimeValue().SetFromMilliseconds(1000 * 5));
+	}
+	else
+	{
+		if (PlaylistLoadRequest.LoadType == FPlaylistLoadRequestHLS::ELoadType::Update)
+		{
+			HTTPRequest->Parameters.ConnectTimeout = Options.GetValue(IPlaylistReaderHLS::OptionKeyUpdatePlaylistLoadConnectTimeout).SafeGetTimeValue(FTimeValue().SetFromMilliseconds(1000 * 2));
+			HTTPRequest->Parameters.NoDataTimeout  = Options.GetValue(IPlaylistReaderHLS::OptionKeyUpdatePlaylistLoadNoDataTimeout).SafeGetTimeValue(FTimeValue().SetFromMilliseconds(1000 * 2));
+		}
+		else
+		{
+			HTTPRequest->Parameters.ConnectTimeout = Options.GetValue(IPlaylistReaderHLS::OptionKeyVariantPlaylistLoadConnectTimeout).SafeGetTimeValue(FTimeValue().SetFromMilliseconds(1000 * 5));
+			HTTPRequest->Parameters.NoDataTimeout  = Options.GetValue(IPlaylistReaderHLS::OptionKeyVariantPlaylistLoadNoDataTimeout).SafeGetTimeValue(FTimeValue().SetFromMilliseconds(1000 * 2));
+		}
+	}
+
+	HTTPRequest->ReceiveBuffer    = ReceiveBuffer;
+	HTTPRequest->ProgressListener = InProgressListener;
+	// Add the request.
+	InPlayerSessionServices->GetHTTPManager()->AddRequest(HTTPRequest);
+}
+
+void FPlaylistReaderHLS::FPlaylistRequest::Cancel(IPlayerSessionServices* InPlayerSessionServices)
+{
+	if (StaticRequest.Get())
+	{
+		StaticRequest->Cancel();
+		StaticRequest.Reset();
+	}
+	if (HTTPRequest.IsValid())
+	{
+		InPlayerSessionServices->GetHTTPManager()->RemoveRequest(HTTPRequest);
+		HTTPRequest.Reset();
+	}
+}
+
+
+void FPlaylistReaderHLS::HandleStaticRequestCompletions(const FTimeValue& TimeNow)
+{
+	for(int32 i=0; i<StaticResourcePlaylistRequests.Num(); ++i)
+	{
+		FPlaylistRequestPtr Request = StaticResourcePlaylistRequests[i];
+		switch(Request->HandleStaticCompletion())
+		{
+			// Not done yet.
+			case 0:
+			{
+				break;
+			}
+			// Done, with static data present.
+			case 1:
+			{
+				StaticResourcePlaylistRequests.RemoveAt(i);
+				--i;
+				Request->PrepareStaticResult();
+				CompletedPlaylistRequests.Push(Request);
+				break;
+			}
+			// Done, but no static data available.
+			default:
+			{
+				StaticResourcePlaylistRequests.RemoveAt(i);
+				--i;
+				// Do the actual HTTP request instead.
+				PendingPlaylistRequestsLock.Lock();
+				PendingPlaylistRequests.Push(Request);
+				Request->Execute(ProgressListener, PlayerSessionServices);
+				PendingPlaylistRequestsLock.Unlock();
+				break;
+			}
+		}
+	}
+}
+
+
+void FPlaylistReaderHLS::HandleEnqueuedPlaylistDownloads(const FTimeValue& TimeNow)
+{
+	for(int32 i=0; i<EnqueuedPlaylistRequests.Num(); ++i)
+	{
+		FPlaylistRequestPtr Request = EnqueuedPlaylistRequests[i];
+		if (!Request->GetExecuteAtUTC().IsValid() || TimeNow >= Request->GetExecuteAtUTC())
+		{
+			EnqueuedPlaylistRequests.RemoveAt(i);
+			--i;
+
+			// Is there a static resource provider that we can try?
+			TSharedPtr<IAdaptiveStreamingPlayerResourceProvider, ESPMode::ThreadSafe> StaticResourceProvider = PlayerSessionServices->GetStaticResourceProvider();
+			if (StaticResourceProvider)
+			{
+				StaticResourcePlaylistRequests.Push(Request);
+				Request->ExecuteStatic(StaticResourceProvider.Get(), WorkerThreadSignal);
+			}
+			else
+			{
+				PendingPlaylistRequestsLock.Lock();
+				PendingPlaylistRequests.Push(Request);
+				Request->Execute(ProgressListener, PlayerSessionServices);
+				PendingPlaylistRequestsLock.Unlock();
+			}
+		}
+	}
+}
+
+void FPlaylistReaderHLS::HandleCompletedPlaylistDownloads(const FTimeValue& TimeNow)
+{
+	FErrorDetail ParseError;
+	while(CompletedPlaylistRequests.Num())
+	{
+		FPlaylistRequestPtr Request = CompletedPlaylistRequests.Pop();
+		const HTTP::FConnectionInfo* ConnInfo = Request->GetConnectionInfo();
+		check(ConnInfo);
+		if (ConnInfo && !ConnInfo->bWasAborted)
+		{
+			bool bIsMasterPlaylist = Request->GetPlaylistLoadRequest().LoadType == FPlaylistLoadRequestHLS::ELoadType::Master;
+			bool bIsInitial 	   = Request->GetPlaylistLoadRequest().LoadType != FPlaylistLoadRequestHLS::ELoadType::Update;
+			PlayerSessionServices->SendMessageToPlayer(IPlaylistReader::PlaylistDownloadMessage::Create(Request->GetConnectionInfo(), bIsMasterPlaylist ? Playlist::EListType::Master : Playlist::EListType::Variant, bIsInitial ? Playlist::ELoadType::Initial : Playlist::ELoadType::Update));
+
+			if (!ConnInfo->StatusInfo.ErrorDetail.IsSet())
+			{
+				ParseError = ParsePlaylist(Request);
+			}
+
+			// Was this an initially required playlist?
+			switch(Request->GetPlaylistLoadRequest().LoadType)
+			{
+				// ------------------------------------------------------------
+				// Master playlist
+				//
+				case FPlaylistLoadRequestHLS::ELoadType::Master:
+				{
+					// Has there been a download error or a parse error?
+					if (ConnInfo->StatusInfo.ErrorDetail.IsSet() || ParseError.IsSet())
+					{
+						bool bRetry = false;
+						int32 MaxRetries = 0;
+						// The current idea is that a parse error might be because of somehow garbled data and not because of a bad playlist on the server.
+						// We do retry the playlist in the hope the problem is intermittent and we get the clean playlist on a retry.
+						if (ParseError.IsSet())
+						{
+							bRetry = true;
+							MaxRetries = 2;
+						}
+						else
+						{
+							if (ConnInfo->StatusInfo.ConnectionTimeoutAfterMilliseconds || ConnInfo->StatusInfo.NoDataTimeoutAfterMilliseconds || ConnInfo->StatusInfo.bReadError)
+							{
+								bRetry = true;
+								MaxRetries = 2;
+							}
+							else if (ConnInfo->StatusInfo.HTTPStatus >= 502 && ConnInfo->StatusInfo.HTTPStatus <= 504)
+							{
+								bRetry = true;
+								MaxRetries = 1;
+							}
+						}
+						bool bFail = true;
+						// Should we retry?
+						if (bRetry)
+						{
+							// Already a retry info there?
+							TSharedPtrTS<HTTP::FRetryInfo> RetryInfo = ConnInfo->RetryInfo;
+							if (!RetryInfo.IsValid())
+							{
+								RetryInfo = MakeSharedTS<HTTP::FRetryInfo>();
+								RetryInfo->MaxAttempts = MaxRetries;
+							}
+							// Can still retry?
+							if (RetryInfo->AttemptNumber < RetryInfo->MaxAttempts)
+							{
+								FTimeValue RetryDelay;
+								RetryDelay.SetFromMilliseconds(1000 * (1 << RetryInfo->AttemptNumber));
+								++RetryInfo->AttemptNumber;
+								RetryInfo->PreviousFailureStates.Push(ConnInfo->StatusInfo);
+								LogMessage(IInfoLog::ELevel::Warning, FString::Printf(TEXT("Failed to %s master playlist \"%s\" (%s), retrying..."), ParseError.IsSet() ? TEXT("parse") : TEXT("download"), *ConnInfo->EffectiveURL, ParseError.IsSet() ? *ParseError.GetPrintable() : *ConnInfo->StatusInfo.ErrorDetail.GetMessage()));
+								EnqueueRetryLoadPlaylist(Request, RetryInfo, TimeNow + RetryDelay);
+								bFail = false;
+							}
+						}
+
+						if (bFail)
+						{
+							if (ParseError.IsSet())
+							{
+								PostError(FString::Printf(TEXT("Failed to parse master playlist \"%s\": %s"), *ConnInfo->EffectiveURL, *ParseError.GetPrintable()), ERRCODE_HLS_MASTER_PLAYLIST_PARSING_FAILED, UEMEDIA_ERROR_FORMAT_ERROR);
+							}
+							else
+							{
+								PostError(FString::Printf(TEXT("Failed to download master playlist \"%s\" (%s)"), *ConnInfo->EffectiveURL, *ConnInfo->StatusInfo.ErrorDetail.GetMessage()), ERRCODE_HLS_MASTER_PLAYLIST_DOWNLOAD_FAILED, UEMEDIA_ERROR_READ_ERROR);
+							}
+						}
+					}
+					break;
+				}
+
+				// ------------------------------------------------------------
+				// Initially selected variant playlist
+				//
+				case FPlaylistLoadRequestHLS::ELoadType::Initial:
+				{
+					// Remove the request from the initial ones regardless of download or parse errors.
+					for (TDoubleLinkedList<FPlaylistLoadRequestHLS>::TIterator It(InitiallyRequiredPlaylistLoadRequests.GetHead()); It; ++It )
+					{
+						if ((*It).InternalUniqueID == Request->GetPlaylistLoadRequest().InternalUniqueID)
+						{
+							InitiallyRequiredPlaylistLoadRequests.RemoveNode(It.GetNode());
+							break;
+						}
+					}
+
+					// Has there been a download or a parse error?
+					if (ConnInfo->StatusInfo.ErrorDetail.IsSet() || ParseError.IsSet())
+					{
+						// We want to switch to a different stream here only, not necessarily retry the same unless there are no alternatives remaining.
+						// The initial playlist load is primarily to establish the timeline. Any errors fetching the playlist for actual streaming will be handled separately.
+
+						// Already a retry info there?
+						TSharedPtrTS<HTTP::FRetryInfo> RetryInfo = ConnInfo->RetryInfo;
+						if (!RetryInfo.IsValid())
+						{
+							RetryInfo = MakeSharedTS<HTTP::FRetryInfo>();
+							RetryInfo->MaxAttempts = 2;
+						}
+						RetryInfo->PreviousFailureStates.Push(ConnInfo->StatusInfo);
+						FTimeValue RetryDelay;
+						RetryDelay.SetFromMilliseconds(1000 * 10);
+						if (Builder->UpdateFailedInitialPlaylistLoadRequest(Request->GetPlaylistLoadRequest(), ConnInfo, RetryInfo, TimeNow + RetryDelay, Manifest, StreamPreferences, Options) == UEMEDIA_ERROR_OK)
+						{
+							// Retry with this initial playlist
+							LogMessage(IInfoLog::ELevel::Warning, FString::Printf(TEXT("Failed %s initial playlist \"%s\" (%s), trying alternative variant"), ParseError.IsSet() ? TEXT("parsing") : TEXT("downloading"), *ConnInfo->EffectiveURL, ParseError.IsSet() ? *ParseError.GetPrintable() : *ConnInfo->StatusInfo.ErrorDetail.GetMessage()));
+							InitiallyRequiredPlaylistLoadRequests.AddHead(Request->GetPlaylistLoadRequest());
+							EnqueueRetryLoadPlaylist(Request, RetryInfo, FTimeValue::GetInvalid());
+						}
+						else
+						{
+							// There are no alternative variants to choose from. To not give up that easily we will try this failed one here again.
+							if (RetryInfo->AttemptNumber < RetryInfo->MaxAttempts)
+							{
+								RetryDelay.SetFromMilliseconds(1000 * (1 << RetryInfo->AttemptNumber));
+								++RetryInfo->AttemptNumber;
+								LogMessage(IInfoLog::ELevel::Warning, FString::Printf(TEXT("Failed %s initial playlist \"%s\" (%s), retrying..."), ParseError.IsSet() ? TEXT("parsing") : TEXT("downloading"), *ConnInfo->EffectiveURL, ParseError.IsSet() ? *ParseError.GetPrintable() : *ConnInfo->StatusInfo.ErrorDetail.GetMessage()));
+								InitiallyRequiredPlaylistLoadRequests.AddHead(Request->GetPlaylistLoadRequest());
+								EnqueueRetryLoadPlaylist(Request, RetryInfo, TimeNow + RetryDelay);
+							}
+							else
+							{
+								// No further alternatives and exhausted retries. Fail!
+								PostError(FString::Printf(TEXT("Failed to %s initial variant playlist \"%s\" (%s). No alternatives remain."), ParseError.IsSet() ? TEXT("parse") : TEXT("download"), *ConnInfo->EffectiveURL, ParseError.IsSet() ? *ParseError.GetPrintable() : *ConnInfo->StatusInfo.ErrorDetail.GetMessage()), ERRCODE_HLS_VARIANT_PLAYLIST_DOWNLOAD_FAILED, UEMEDIA_ERROR_READ_ERROR);
+							}
+						}
+					}
+					// Was this the last one?
+					// NOTE: If downloading or parsing failed we will not attempt to create a manifest (would not work anyway)
+					//       and also not send a message to the player telling it that metadata is available.
+					else if (InitiallyRequiredPlaylistLoadRequests.IsEmpty())
+					{
+						// Create a wrapper manifest for use with the player.
+						PlayerManifest = FManifestHLS::Create(PlayerSessionServices, Options, this, Manifest);
+
+						// Yes. We can now create the media timeline and report availability of metadata to the player.
+						PlayerSessionServices->SendMessageToPlayer(IPlaylistReader::PlaylistLoadedMessage::Create(LastErrorDetail, Request->GetConnectionInfo(), Playlist::EListType::Variant, Playlist::ELoadType::Initial));
+					}
+					break;
+				}
+
+				// ------------------------------------------------------------
+				// First fetch of a newly selected stream or an update of a Live playlist
+				//
+				case FPlaylistLoadRequestHLS::ELoadType::First:
+				case FPlaylistLoadRequestHLS::ELoadType::Update:
+				{
+					if (ConnInfo->StatusInfo.ErrorDetail.IsSet() || ParseError.IsSet())
+					{
+						TSharedPtrTS<HTTP::FRetryInfo> RetryInfo = ConnInfo->RetryInfo;
+						bool bRetry = false;
+						int32 MaxRetries = 0;
+						// For variant playlists we follow the same idea as for the master playlist in that a parse error might be because of somehow garbled data and not because of a bad playlist on the server.
+						if (ParseError.IsSet())
+						{
+							bRetry = true;
+							MaxRetries = 2;
+						}
+						else
+						{
+							if (ConnInfo->StatusInfo.ConnectionTimeoutAfterMilliseconds || ConnInfo->StatusInfo.NoDataTimeoutAfterMilliseconds || ConnInfo->StatusInfo.bReadError)
+							{
+								bRetry = true;
+								MaxRetries = 2;
+							}
+							else
+							{
+								switch(ConnInfo->StatusInfo.HTTPStatus)
+								{
+									case 204:
+									case 205:
+									case 403:
+									case 404:
+									case 429:
+									case 502:
+									case 503:
+									case 504:
+										bRetry = true;
+										MaxRetries = 1;
+									default:
+										break;
+								}
+							}
+						}
+						bool bFail = true;
+						// Should we retry?
+						if (bRetry)
+						{
+							// Already a retry info there?
+							if (!RetryInfo.IsValid())
+							{
+								RetryInfo = MakeSharedTS<HTTP::FRetryInfo>();
+								RetryInfo->MaxAttempts = MaxRetries;
+							}
+							// Can still retry?
+							if (RetryInfo->AttemptNumber < RetryInfo->MaxAttempts)
+							{
+								FTimeValue RetryDelay;
+								RetryDelay.SetFromMilliseconds(1000 * (1 << RetryInfo->AttemptNumber));
+								++RetryInfo->AttemptNumber;
+								RetryInfo->PreviousFailureStates.Push(ConnInfo->StatusInfo);
+								LogMessage(IInfoLog::ELevel::Warning, FString::Printf(TEXT("Failed %s playlist \"%s\" (%s), retrying..."), ParseError.IsSet() ? TEXT("parsing") : TEXT("downloading"), *ConnInfo->EffectiveURL, ParseError.IsSet() ? *ParseError.GetPrintable(): *ConnInfo->StatusInfo.ErrorDetail.GetMessage()));
+								EnqueueRetryLoadPlaylist(Request, RetryInfo, TimeNow + RetryDelay);
+								bFail = false;
+							}
+						}
+
+						if (bFail)
+						{
+							int32 BlacklistForMilliseconds = 1000 * 15;
+							Builder->SetVariantPlaylistFailure(Manifest, Request->GetPlaylistLoadRequest(), ConnInfo, RetryInfo, TimeNow + FTimeValue().SetFromMilliseconds(BlacklistForMilliseconds));
+							LogMessage(IInfoLog::ELevel::Warning, FString::Printf(TEXT("Failed %s playlist \"%s\" (%s), blacklisting for %d milliseconds"), ParseError.IsSet() ? TEXT("parsing") : TEXT("downloading"), *ConnInfo->EffectiveURL, ParseError.IsSet() ? *ParseError.GetPrintable() : *ConnInfo->StatusInfo.ErrorDetail.GetMessage(), BlacklistForMilliseconds));
+						}
+					}
+					else
+					{
+						PlayerSessionServices->SendMessageToPlayer(IPlaylistReader::PlaylistLoadedMessage::Create(LastErrorDetail, Request->GetConnectionInfo(), Playlist::EListType::Variant, Playlist::ELoadType::Update));
+					}
+					break;
+				}
+				default:
+					break;
+			}
+		}
+		// Request is done.
+		Request.Reset();
+	}
+}
+
+
+
+FErrorDetail FPlaylistReaderHLS::ParsePlaylist(FPlaylistRequestPtr FromRequest)
+{
+	// Calculate checksum over the response to see if the playlist has changed
+	check(FromRequest->GetReceiveBuffer()->Buffer.GetLinearReadSize() >= FromRequest->GetReceiveBuffer()->Buffer.Num());	// The buffer must be linear!
+	uint32 crc = HashFunctions::CRC::Calc32(FromRequest->GetReceiveBuffer()->Buffer.GetLinearReadData(), FromRequest->GetReceiveBuffer()->Buffer.Num());
+	bool bNoChange = FromRequest->GetPlaylistLoadRequest().LoadType == FPlaylistLoadRequestHLS::ELoadType::Update && FromRequest->GetPlaylistLoadRequest().LastUpdateCRC32 == crc;
+
+	HLSPlaylistParser::FPlaylist		Playlist;
+	HLSPlaylistParser::FParser			Parser;
+	Parser.Configure(HLSPlaylistParser::LHLS::TagMap);
+	Parser.Configure(HLSPlaylistParser::Epic::TagMap);
+
+	HLSPlaylistParser::EPlaylistError ParseError;
+	if (bNoChange)
+	{
+		ParseError = HLSPlaylistParser::EPlaylistError::None;
+	}
+	else
+	{
+		int32 RequestBytes = FromRequest->GetReceiveBuffer()->Buffer.Num();
+		// The FromRequest buffer is not zero terminated. We need to pick the correct FString constructor for converting the chars into TCHARs while adding the terminating zero!
+		ParseError = Parser.Parse(FString(RequestBytes, (TCHAR*)FUTF8ToTCHAR((const ANSICHAR*)FromRequest->GetReceiveBuffer()->Buffer.GetLinearReadData(), RequestBytes).Get()), Playlist);
+	}
+
+	if (ParseError == HLSPlaylistParser::EPlaylistError::None)
+	{
+		FErrorDetail Error;
+		if (FromRequest->IsMasterPlaylist())
+		{
+			// Get the Date header from the response and set our clock to this time.
+			// Do this once only with the date from the master playlist.
+			const HTTP::FConnectionInfo* ConnInfo = FromRequest->GetConnectionInfo();
+			if (ConnInfo)
+			{
+				for(int32 i=0; i<ConnInfo->ResponseHeaders.Num(); ++i)
+				{
+					if (ConnInfo->ResponseHeaders[i].Header == "Date")
+					{
+						// How old is the response already?
+						FTimeValue ResponseAge = MEDIAutcTime::Current() - ConnInfo->RequestStartTime - FTimeValue().SetFromSeconds(ConnInfo->TimeUntilFirstByte);
+						// Parse the header
+						FTimeValue DateFromHeader;
+						UEMediaError DateParseError = RFC7231::ParseDateTime(DateFromHeader, ConnInfo->ResponseHeaders[i].Value);
+						if (DateParseError == UEMEDIA_ERROR_OK && DateFromHeader.IsValid())
+						{
+							PlayerSessionServices->GetSynchronizedUTCTime()->SetTime(DateFromHeader + ResponseAge);
+						}
+						break;
+					}
+				}
+			}
+
+
+			TSharedPtrTS<FManifestHLSInternal> NewManifest;
+			Error = Builder->BuildFromMasterPlaylist(NewManifest, Playlist, FromRequest->GetPlaylistLoadRequest(), FromRequest->GetConnectionInfo(), StreamPreferences, Options);
+			if (Error.IsOK())
+			{
+				Manifest = NewManifest;
+				MasterPlaylistURL = Manifest->MasterPlaylistVars.PlaylistLoadRequest.URL;
+
+				// Notify that we have successully loaded the master playlist.
+				PlayerSessionServices->SendMessageToPlayer(IPlaylistReader::PlaylistLoadedMessage::Create(LastErrorDetail, FromRequest->GetConnectionInfo(), Playlist::EListType::Master, Playlist::ELoadType::Initial));
+
+				// Get the playlists we need to eventually start playback with.
+				TArray<FPlaylistLoadRequestHLS> PlaylistRequests;
+				Error = Builder->GetInitialPlaylistLoadRequests(PlaylistRequests, Manifest, StreamPreferences, Options);
+				if (Error.IsOK())
+				{
+					// Enqueue those for loading.
+					for(int32 i=0; i<PlaylistRequests.Num(); ++i)
+					{
+						check(PlaylistRequests[i].LoadType == FPlaylistLoadRequestHLS::ELoadType::Initial);
+
+						// Add this to the list of initially required playlists. Only when those are ready can we report
+						// to the player that metadata is available and playback can be started.
+						InitiallyRequiredPlaylistLoadRequests.AddHead(PlaylistRequests[i]);
+
+						EnqueueLoadPlaylist(PlaylistRequests[i], false);
+					}
+				}
+			}
+		}
+		else
+		{
+			// For our purposes there needs to be a master playlist.
+			if (Manifest.IsValid())
+			{
+				Error = Builder->UpdateFromVariantPlaylist(Manifest, Playlist, FromRequest->GetPlaylistLoadRequest(), FromRequest->GetConnectionInfo(), crc, StreamPreferences, Options);
+			}
+			else
+			{
+				return CreateErrorAndLog(FString::Printf(TEXT("No master playlist present. Cannot play a variant playlist by itself")), ERRCODE_HLS_NO_MASTER_PLAYLIST);
+			}
+		}
+		return Error;
+	}
+	else
+	{
+		return CreateErrorAndLog(FString::Printf(TEXT("HLS parser returned error %d"), ParseError), ERRCODE_HLS_PARSER_ERROR);
+	}
+}
+
+} // namespace Electra
+
+
diff --git a/Engine/Plugins/Media/ElectraPlayer/Source/ElectraPlayerRuntime/Private/Runtime/Player/HLS/PlaylistReaderHLS.h b/Engine/Plugins/Media/ElectraPlayer/Source/ElectraPlayerRuntime/Private/Runtime/Player/HLS/PlaylistReaderHLS.h
new file mode 100644
--- /dev/null
+++ b/Engine/Plugins/Media/ElectraPlayer/Source/ElectraPlayerRuntime/Private/Runtime/Player/HLS/PlaylistReaderHLS.h
@@ -0,0 +1,67 @@
+// Copyright Epic Games, Inc. All Rights Reserved.
+
+#pragma once
+
+#include "PlayerCore.h"
+#include "Player/PlaylistReader.h"
+#include "Player/PlayerSessionServices.h"
+
+
+namespace Electra
+{
+
+struct FPlaylistLoadRequestHLS;
+
+
+
+class IPlaylistReaderHLS : public IPlaylistReader
+{
+public:
+	static IPlaylistReader* Create(IPlayerSessionServices* PlayerSessionServices);
+
+	virtual ~IPlaylistReaderHLS() = default;
+
+
+	static const FString OptionKeyInitialBitrate;						//!< (int64) value indicating the bitrate to start with.
+	static const FString OptionKeyLiveSeekableStartOffset;				//!< (FTimeValue) value specifying how many seconds away from the Live media timeline the seekable range should start.
+	static const FString OptionKeyLiveSeekableEndOffset;				//!< (FTimeValue) value specifying how many seconds away from the Live media timeline the seekable range should end.
+	static const FString OptionKeyLiveSeekableEndOffsetAudioOnly;		//!< (FTimeValue) value specifying how many seconds away from the Live media timeline the seekable range should end for audio-only playlists.
+	static const FString OptionKeyLiveSeekableEndOffsetBeConservative;	//!< (bool) true to use a larger Live edge distance, false to go with the smaller absolute difference
+	static const FString OptionKeyMasterPlaylistLoadConnectTimeout;		//!< (FTimeValue) value specifying connection timeout fetching the master playlist
+	static const FString OptionKeyMasterPlaylistLoadNoDataTimeout;		//!< (FTimeValue) value specifying no-data timeout fetching the master playlist
+	static const FString OptionKeyVariantPlaylistLoadConnectTimeout;	//!< (FTimeValue) value specifying connection timeout fetching a variant playlist the first time
+	static const FString OptionKeyVariantPlaylistLoadNoDataTimeout;		//!< (FTimeValue) value specifying no-data timeout fetching a variant playlist the first time
+	static const FString OptionKeyUpdatePlaylistLoadConnectTimeout;		//!< (FTimeValue) value specifying connection timeout fetching a variant playlist repeatedly
+	static const FString OptionKeyUpdatePlaylistLoadNoDataTimeout;		//!< (FTimeValue) value specifying no-data timeout fetching a variant playlist repeatedly
+
+	/**
+	 * Loads and parses the playlist.
+	 *
+	 * @param URL     URL of the playlist to load
+	 * @param Preferences
+	 *                User preferences for initial stream selection.
+	 * @param Options Options
+	 */
+	virtual void LoadAndParse(const FString& URL, const FStreamPreferences& Preferences, const FParamDict& Options) = 0;
+
+	/**
+	 * Returns the URL from which the playlist was loaded (or supposed to be loaded).
+	 *
+	 * @return The playlist URL
+	 */
+	virtual FString GetURL() const = 0;
+
+
+	/**
+	 * Requests loading of a variant/rendition playlist.
+	 *
+	 * @param LoadRequest
+	 */
+	virtual void RequestPlaylistLoad(const FPlaylistLoadRequestHLS& LoadRequest) = 0;
+};
+
+
+} // namespace Electra
+
+
+
