ElectraPlayer: RELEASE. Moved from NFL to public. Note that this is done via a Delete and "Add New" operation as discussed.

diff --git a/Engine/Plugins/Media/ElectraPlayer/Source/ElectraPlayerRuntime/Private/Runtime/Player/HLS/StreamReaderHLSfmp4.cpp b/Engine/Plugins/Media/ElectraPlayer/Source/ElectraPlayerRuntime/Private/Runtime/Player/HLS/StreamReaderHLSfmp4.cpp
new file mode 100644
--- /dev/null
+++ b/Engine/Plugins/Media/ElectraPlayer/Source/ElectraPlayerRuntime/Private/Runtime/Player/HLS/StreamReaderHLSfmp4.cpp
@@ -0,0 +1,1586 @@
+// Copyright Epic Games, Inc. All Rights Reserved.
+
+#include "HTTP/HTTPManager.h"
+#include "StreamReaderHLSfmp4.h"
+#include "Demuxer/ParserISO14496-12.h"
+#include "StreamAccessUnitBuffer.h"
+#include "Player/PlayerLicenseKey.h"
+#include "HAL/LowLevelMemTracker.h"
+#include "ElectraPlayerPrivate.h"
+
+#define INTERNAL_ERROR_INIT_SEGMENT_DOWNLOAD_ERROR					1
+#define INTERNAL_ERROR_INIT_SEGMENT_PARSE_ERROR						2
+#define INTERNAL_ERROR_INIT_SEGMENT_LICENSEKEY_ERROR				10
+
+#define PERMIT_INSECURE_SEGMENT_DEMOTING 0
+
+namespace Electra
+{
+
+//! (bool) if false and media segment is using EXT-X-KEY encryption fetch it via http even if it should be https, otherwise keep the original scheme.
+const FString FStreamReaderHLSfmp4::OptionKeyDontUseInsecureForEncryptedMediaSegments("dont_use_insecure_for_media_segments");
+
+//! (bool) if false the init segment is fetched via http even if it should be https, otherwise keep the original scheme.
+const FString FStreamReaderHLSfmp4::OptionKeyDontUseInsecureForInitSegments("dont_use_insecure_for_init_segments");
+
+
+FStreamSegmentRequestHLSfmp4::FStreamSegmentRequestHLSfmp4()
+{
+	StreamType  			  = EStreamType::Video;
+	StreamUniqueID  		  = 0;
+	Bitrate 				  = 0;
+	QualityLevel			  = 0;
+	MediaSequence   		  = -1;
+	DiscontinuitySequence     = -1;
+	LocalIndex  			  = -1;
+	bIsPrefetch 			  = false;
+	bIsEOSSegment   		  = false;
+	bHasEncryptedSegments     = false;
+	NumOverallRetries   	  = 0;
+	bInsertFillerData   	  = false;
+	bIsInitialStartRequest    = false;
+	CurrentPlaybackSequenceID = ~0U;
+}
+
+FStreamSegmentRequestHLSfmp4::~FStreamSegmentRequestHLSfmp4()
+{
+}
+
+
+void FStreamSegmentRequestHLSfmp4::CopyFrom(const FStreamSegmentRequestHLSfmp4& rhs)
+{
+	URL 						 = rhs.URL;
+	Range		  				 = rhs.Range;
+	StreamType  				 = rhs.StreamType;
+	StreamUniqueID  			 = rhs.StreamUniqueID;
+	Bitrate 					 = rhs.Bitrate;
+	QualityLevel				 = rhs.QualityLevel;
+	Representation  			 = rhs.Representation;
+	AdaptationSet   			 = rhs.AdaptationSet;
+	MediaAsset  				 = rhs.MediaAsset;
+	CDN 						 = rhs.CDN;
+	//PlaylistRelativeStartTime  = rhs.PlaylistRelativeStartTime;
+	AbsoluteDateTime			 = rhs.AbsoluteDateTime;
+	SegmentDuration 			 = rhs.SegmentDuration;
+	MediaSequence   			 = rhs.MediaSequence;
+	DiscontinuitySequence   	 = rhs.DiscontinuitySequence;
+	LocalIndex  				 = rhs.LocalIndex;
+	FirstAUTimeOffset   		 = rhs.FirstAUTimeOffset;
+	bIsPrefetch 				 = rhs.bIsPrefetch;
+	bIsEOSSegment   			 = rhs.bIsEOSSegment;
+	bHasEncryptedSegments   	 = rhs.bHasEncryptedSegments;
+	NumOverallRetries   		 = rhs.NumOverallRetries;
+	bInsertFillerData   		 = rhs.bInsertFillerData;
+	InitSegmentCache			 = rhs.InitSegmentCache;
+	InitSegmentInfo 			 = rhs.InitSegmentInfo;
+	DependentStreams			 = rhs.DependentStreams;
+	bIsInitialStartRequest  	 = rhs.bIsInitialStartRequest;
+	PlayerLoopState 			 = rhs.PlayerLoopState;
+	CurrentPlaybackSequenceID    = rhs.CurrentPlaybackSequenceID;
+	DownloadStats   			 = rhs.DownloadStats;
+	ConnectionInfo  			 = rhs.ConnectionInfo;
+	NextLargestExpectedTimestamp = rhs.NextLargestExpectedTimestamp;
+}
+
+
+
+void FStreamSegmentRequestHLSfmp4::SetPlaybackSequenceID(uint32 PlaybackSequenceID)
+{
+	CurrentPlaybackSequenceID = PlaybackSequenceID;
+}
+
+uint32 FStreamSegmentRequestHLSfmp4::GetPlaybackSequenceID() const
+{
+	return CurrentPlaybackSequenceID;
+}
+
+EStreamType FStreamSegmentRequestHLSfmp4::GetType() const
+{
+	return StreamType;
+}
+
+void FStreamSegmentRequestHLSfmp4::GetDependentStreams(TArray<FDependentStreams>& OutDependentStreams) const
+{
+	OutDependentStreams.Empty();
+	for(int32 i=0; i<DependentStreams.Num(); ++i)
+	{
+		FDependentStreams& depStr = OutDependentStreams.AddDefaulted_GetRef();
+		depStr.StreamType = DependentStreams[i]->GetType();
+	}
+}
+
+void FStreamSegmentRequestHLSfmp4::GetEndedStreams(TArray<TSharedPtrTS<IStreamSegment>>& OutAlreadyEndedStreams)
+{
+	OutAlreadyEndedStreams.Empty();
+	if (bIsEOSSegment)
+	{
+		OutAlreadyEndedStreams.Push(SharedThis(this));
+	}
+	for(int32 i=0; i<DependentStreams.Num(); ++i)
+	{
+		if (DependentStreams[i]->bIsEOSSegment)
+		{
+			OutAlreadyEndedStreams.Push(DependentStreams[i]);
+		}
+	}
+}
+
+FTimeValue FStreamSegmentRequestHLSfmp4::GetFirstPTS() const
+{
+	return AbsoluteDateTime;
+}
+
+int32 FStreamSegmentRequestHLSfmp4::GetQualityIndex() const
+{
+	return QualityLevel;
+}
+
+int32 FStreamSegmentRequestHLSfmp4::GetBitrate() const
+{
+	return Bitrate;
+}
+
+void FStreamSegmentRequestHLSfmp4::GetDownloadStats(Metrics::FSegmentDownloadStats& OutStats) const
+{
+	OutStats = DownloadStats;
+}
+
+/***************************************************************************************************************************************************/
+/***************************************************************************************************************************************************/
+/***************************************************************************************************************************************************/
+/***************************************************************************************************************************************************/
+/***************************************************************************************************************************************************/
+
+FStreamReaderHLSfmp4::FStreamReaderHLSfmp4()
+	: PlayerSessionService(nullptr)
+	, bIsStarted(false)
+{
+}
+
+
+FStreamReaderHLSfmp4::~FStreamReaderHLSfmp4()
+{
+	Close();
+}
+
+UEMediaError FStreamReaderHLSfmp4::Create(IPlayerSessionServices* InPlayerSessionService, const IStreamReader::CreateParam& InCreateParam)
+{
+	PlayerSessionService = InPlayerSessionService;
+
+	if (!InCreateParam.MemoryProvider ||
+		!InCreateParam.EventListener ||
+		!InCreateParam.PlayerSessionService)
+	{
+		return UEMEDIA_ERROR_BAD_ARGUMENTS;
+	}
+
+	bIsStarted = true;
+	for(int32 i=0; i<FMEDIA_STATIC_ARRAY_COUNT(StreamHandlers); ++i)
+	{
+		StreamHandlers[i].PlayerSessionService = PlayerSessionService;
+		StreamHandlers[i].Parameters		   = InCreateParam;
+		StreamHandlers[i].bTerminate		   = false;
+		StreamHandlers[i].bRequestCanceled     = false;
+		StreamHandlers[i].bHasErrored   	   = false;
+
+		StreamHandlers[i].ThreadSetPriority(StreamHandlers[i].Parameters.ReaderConfig.ThreadParam.Priority);
+		StreamHandlers[i].ThreadSetCoreAffinity(StreamHandlers[i].Parameters.ReaderConfig.ThreadParam.CoreAffinity);
+		StreamHandlers[i].ThreadSetStackSize(StreamHandlers[i].Parameters.ReaderConfig.ThreadParam.StackSize);
+		StreamHandlers[i].ThreadSetName(i==0?"ElectraPlayer::fmp4 Video":"ElectraPlayer::fmp4 Audio");
+		StreamHandlers[i].ThreadStart(Electra::MakeDelegate(&StreamHandlers[i], &FStreamHandler::WorkerThread));
+	}
+	return UEMEDIA_ERROR_OK;
+}
+
+void FStreamReaderHLSfmp4::Close()
+{
+	if (bIsStarted)
+	{
+		bIsStarted = false;
+		// Signal the worker threads to end.
+		for(int32 i=0; i<FMEDIA_STATIC_ARRAY_COUNT(StreamHandlers); ++i)
+		{
+			StreamHandlers[i].bTerminate = true;
+			StreamHandlers[i].Cancel();
+			StreamHandlers[i].SignalWork();
+		}
+		// Wait until they finished.
+		for(int32 i=0; i<FMEDIA_STATIC_ARRAY_COUNT(StreamHandlers); ++i)
+		{
+			StreamHandlers[i].ThreadWaitDone();
+			StreamHandlers[i].ThreadReset();
+		}
+	}
+}
+
+IStreamReader::EAddResult FStreamReaderHLSfmp4::AddRequest(uint32 CurrentPlaybackSequenceID, TSharedPtrTS<IStreamSegment> InRequest)
+{
+	TSharedPtrTS<FStreamSegmentRequestHLSfmp4> Request = StaticCastSharedPtr<FStreamSegmentRequestHLSfmp4>(InRequest);
+
+	// Video and audio only for now.
+	if (Request->GetType() != EStreamType::Video && Request->GetType() != EStreamType::Audio)
+	{
+		check(!"no good");
+		ErrorDetail.SetMessage(FString::Printf(TEXT("Request is not video or audio")));
+		return IStreamReader::EAddResult::Error;
+	}
+	// Only initial requests are allowed to have a dependent stream for now.
+	if (Request->DependentStreams.Num() && !Request->bIsInitialStartRequest)
+	{
+		check(!"no good");
+		ErrorDetail.SetMessage(FString::Printf(TEXT("Dependent streams only allowed for initial request")));
+		return IStreamReader::EAddResult::Error;
+	}
+	// Also, there may only be one dependent stream.
+	if (Request->DependentStreams.Num() > 1)
+	{
+		check(!"no good");
+		ErrorDetail.SetMessage(FString::Printf(TEXT("Only one dependent streams allowed")));
+		return IStreamReader::EAddResult::Error;
+	}
+
+	// Get the handler for the main request.
+	FStreamHandler* Handler = nullptr;
+	switch(Request->GetType())
+	{
+		case  EStreamType::Video:
+			Handler = &StreamHandlers[0];
+			break;
+		case  EStreamType::Audio:
+			Handler = &StreamHandlers[1];
+			break;
+		default:
+			check(!"Whoops");
+			break;
+	}
+	if (!Handler)
+	{
+		check(!"no good");
+		ErrorDetail.SetMessage(FString::Printf(TEXT("No handler for stream type")));
+		return IStreamReader::EAddResult::Error;
+	}
+	// Is the handler busy?
+	if (Handler->CurrentRequest.IsValid())
+	{
+		check(!"why is the handler busy??");
+		return IStreamReader::EAddResult::TryAgainLater;
+	}
+
+// TODO: make this work with more than just one request (asserted above) if necessary
+	if (Request->DependentStreams.Num())
+	{
+		TSharedPtrTS<FStreamSegmentRequestHLSfmp4> Request2 = Request->DependentStreams[0];
+		FStreamHandler* Handler2 = nullptr;
+		switch(Request2->GetType())
+		{
+			case  EStreamType::Video:
+				Handler2 = &StreamHandlers[0];
+				break;
+			case  EStreamType::Audio:
+				Handler2 = &StreamHandlers[1];
+				break;
+			default:
+				check(!"Whoops");
+				break;
+		}
+		if (!Handler2)
+		{
+			check(!"no good");
+			ErrorDetail.SetMessage(FString::Printf(TEXT("No handler for stream type")));
+			return IStreamReader::EAddResult::Error;
+		}
+		// Is the handler busy?
+		if (Handler2->CurrentRequest.IsValid())
+		{
+			check(!"why is the handler busy??");
+			return IStreamReader::EAddResult::TryAgainLater;
+		}
+
+		Request->DependentStreams.Empty();
+		Request2->SetPlaybackSequenceID(CurrentPlaybackSequenceID);
+		// Only add the request if this is not an EOD segment.
+		if (!Request2->bIsEOSSegment)
+		{
+			Handler2->CurrentRequest = Request2;
+			Handler2->SignalWork();
+		}
+	}
+	Request->SetPlaybackSequenceID(CurrentPlaybackSequenceID);
+	// Only add the request if this is not an EOD segment.
+	if (!Request->bIsEOSSegment)
+	{
+		Handler->CurrentRequest = Request;
+		Handler->SignalWork();
+	}
+	return IStreamReader::EAddResult::Added;
+}
+
+void FStreamReaderHLSfmp4::PauseDownload()
+{
+	// Download will not be paused. The pending segment will complete downloading.
+}
+
+void FStreamReaderHLSfmp4::ResumeDownload()
+{
+	// Since we do not pause we also do not resume.
+}
+
+void FStreamReaderHLSfmp4::CancelRequests()
+{
+	for(int32 i=0; i<FMEDIA_STATIC_ARRAY_COUNT(StreamHandlers); ++i)
+	{
+		StreamHandlers[i].Cancel();
+	}
+}
+
+
+
+
+
+uint32 FStreamReaderHLSfmp4::FStreamHandler::UniqueDownloadID = 1;
+
+FStreamReaderHLSfmp4::FStreamHandler::FStreamHandler()
+{
+	PlayerSessionService = nullptr;
+	bTerminate  		 = false;
+	bRequestCanceled	 = false;
+	bAbortedByABR   	 = false;
+	bHasErrored 		 = false;
+	NumMOOFBoxesFound    = 0;
+	ProgressReportCount  = 0;
+}
+
+FStreamReaderHLSfmp4::FStreamHandler::~FStreamHandler()
+{
+	// NOTE: The thread will have been terminated by the enclosing FStreamReaderHLSfmp4's Close() method!
+}
+
+void FStreamReaderHLSfmp4::FStreamHandler::Cancel()
+{
+	bRequestCanceled = true;
+}
+
+void FStreamReaderHLSfmp4::FStreamHandler::SignalWork()
+{
+	WorkSignal.Release();
+}
+
+void FStreamReaderHLSfmp4::FStreamHandler::WorkerThread()
+{
+	LLM_SCOPE(ELLMTag::ElectraPlayer);
+
+	check(PlayerSessionService);
+	StreamSelector = PlayerSessionService->GetStreamSelector();
+	check(StreamSelector.IsValid());
+	while(!bTerminate)
+	{
+		WorkSignal.Obtain();
+		if (!bTerminate)
+		{
+			if (CurrentRequest.IsValid())
+			{
+				if (!bRequestCanceled)
+				{
+					HandleRequest();
+				}
+				else
+				{
+					CurrentRequest.Reset();
+				}
+			}
+			bRequestCanceled = false;
+		}
+	}
+	StreamSelector.Reset();
+}
+
+
+void FStreamReaderHLSfmp4::FStreamHandler::LogMessage(IInfoLog::ELevel Level, const FString& Message)
+{
+	if (PlayerSessionService)
+	{
+		PlayerSessionService->PostLog(Facility::EFacility::HLSFMP4Reader, Level, Message);
+	}
+}
+
+
+int32 FStreamReaderHLSfmp4::FStreamHandler::HTTPProgressCallback(const IElectraHttpManager::FRequest* Request)
+{
+	HTTPUpdateStats(MEDIAutcTime::Current(), Request);
+	++ProgressReportCount;
+
+	// Aborted?
+	return HasReadBeenAborted() ? 1 : 0;
+}
+
+void FStreamReaderHLSfmp4::FStreamHandler::HTTPCompletionCallback(const IElectraHttpManager::FRequest* Request)
+{
+	HTTPUpdateStats(FTimeValue::GetInvalid(), Request);
+
+	bHasErrored = Request->ConnectionInfo.StatusInfo.ErrorDetail.IsError();
+	DownloadCompleteSignal.Signal();
+}
+
+void FStreamReaderHLSfmp4::FStreamHandler::HTTPUpdateStats(const FTimeValue& CurrentTime, const IElectraHttpManager::FRequest* Request)
+{
+	TSharedPtrTS<FStreamSegmentRequestHLSfmp4> SegmentRequest = CurrentRequest;
+	if (SegmentRequest.IsValid())
+	{
+		FMediaCriticalSection::ScopedLock lock(MetricUpdateLock);
+		SegmentRequest->ConnectionInfo = Request->ConnectionInfo;
+		// Update the current download stats which we report periodically to the ABR.
+		Metrics::FSegmentDownloadStats& ds = SegmentRequest->DownloadStats;
+		if (Request->ConnectionInfo.EffectiveURL.Len())
+		{
+			ds.URL 			  = Request->ConnectionInfo.EffectiveURL;
+		}
+		ds.HTTPStatusCode     = Request->ConnectionInfo.StatusInfo.HTTPStatus;
+		ds.TimeToFirstByte    = Request->ConnectionInfo.TimeUntilFirstByte;
+		ds.TimeToDownload     = ((CurrentTime.IsValid() ? CurrentTime : Request->ConnectionInfo.RequestEndTime) - Request->ConnectionInfo.RequestStartTime).GetAsSeconds();
+		ds.ByteSize 		  = Request->ConnectionInfo.ContentLength;
+		ds.NumBytesDownloaded = Request->ConnectionInfo.BytesReadSoFar;
+	}
+}
+
+
+FString FStreamReaderHLSfmp4::FStreamHandler::DemoteInitURLToHTTP(const FString& InURL, bool bIsEncrypted)
+{
+	FString NewUrl(InURL);
+#if PERMIT_INSECURE_SEGMENT_DEMOTING
+	if (bIsEncrypted &&
+		Parameters.Options.GetValue(OptionKeyDontUseInsecureForInitSegments).SafeGetBool(false) == false &&
+		NewUrl.Left(8) == TEXT("https://"))
+	{
+		// Remove the 's'
+		NewUrl.RemoveAt(4);
+	}
+#endif
+	return NewUrl;
+}
+
+FString FStreamReaderHLSfmp4::FStreamHandler::DemoteMediaURLToHTTP(const FString& InURL, bool bIsEncrypted)
+{
+	FString NewUrl(InURL);
+#if PERMIT_INSECURE_SEGMENT_DEMOTING
+if (bIsEncrypted &&
+		Parameters.Options.GetValue(OptionKeyDontUseInsecureForEncryptedMediaSegments).SafeGetBool(false) == false &&
+		NewUrl.Left(8) == TEXT("https://"))
+	{
+		// Remove the 's'
+		NewUrl.RemoveAt(4);
+	}
+#endif
+	return NewUrl;
+}
+
+
+FStreamReaderHLSfmp4::FStreamHandler::ELicenseKeyResult FStreamReaderHLSfmp4::FStreamHandler::GetLicenseKey(FErrorDetail& OutErrorDetail, TSharedPtr<TArray<uint8>, ESPMode::ThreadSafe>& OutLicenseKeyData, const TSharedPtrTS<FStreamSegmentRequestHLSfmp4>& InRequest, const TSharedPtr<const FManifestHLSInternal::FMediaStream::FDRMKeyInfo, ESPMode::ThreadSafe>& LicenseKeyInfo)
+{
+	if (LicenseKeyInfo.IsValid() && LicenseKeyInfo->URI.Len())
+	{
+		TSharedPtr<TArray<uint8>, ESPMode::ThreadSafe> LicenseKeyData;
+		if (InRequest->LicenseKeyCache.IsValid())
+		{
+			LicenseKeyData = InRequest->LicenseKeyCache->GetLicenseKeyFor(LicenseKeyInfo);
+		}
+		if (LicenseKeyData.IsValid())
+		{
+			OutLicenseKeyData = LicenseKeyData;
+			return FStreamReaderHLSfmp4::FStreamHandler::ELicenseKeyResult::AlreadyCached;
+		}
+		else
+		{
+			TSharedPtrTS<IElectraHttpManager::FRequest> HTTP(new IElectraHttpManager::FRequest);
+			ReadBuffer.Reset();
+			ReadBuffer.ReceiveBuffer = MakeSharedTS<IElectraHttpManager::FReceiveBuffer>();
+
+			Metrics::FSegmentDownloadStats& ds = InRequest->DownloadStats;
+			ds.URL = LicenseKeyInfo->URI;
+
+			HTTP->Parameters.URL   = LicenseKeyInfo->URI;
+			HTTP->ReceiveBuffer    = ReadBuffer.ReceiveBuffer;
+
+			ProgressReportCount = 0;
+			DownloadCompleteSignal.Reset();
+
+			// Is there a static resource provider that we can try?
+			bool bHaveStaticResponse = false;
+			TSharedPtr<IAdaptiveStreamingPlayerResourceProvider, ESPMode::ThreadSafe> StaticResourceProvider = PlayerSessionService->GetStaticResourceProvider();
+			if (StaticResourceProvider)
+			{
+				TSharedPtr<FStaticResourceRequest, ESPMode::ThreadSafe>	StaticRequest = MakeShared<FStaticResourceRequest, ESPMode::ThreadSafe>(/*SharedThis(this), DoneSignal,*/ LicenseKeyInfo->URI, IAdaptiveStreamingPlayerResourceRequest::EPlaybackResourceType::LicenseKey);
+				StaticResourceProvider->ProvideStaticPlaybackDataForURL(StaticRequest);
+				while(!HasReadBeenAborted())
+				{
+					if (StaticRequest->WaitDone(1000 * 100))
+					{
+						TSharedPtr<TArray<uint8>, ESPMode::ThreadSafe> KeyData = StaticRequest->GetData();
+						if (KeyData.IsValid())
+						{
+							// Copy the response over into the receive buffer as if it was received through the http request.
+							ReadBuffer.ReceiveBuffer->Buffer.Reserve(KeyData->Num());
+							ReadBuffer.ReceiveBuffer->Buffer.PushData(KeyData->GetData(), KeyData->Num());
+							ReadBuffer.ReceiveBuffer->Buffer.SetEOD();
+							bHaveStaticResponse = true;
+						}
+						break;
+					}
+				}
+			}
+
+			// When we did not get a static key response we have to issue the request for real.
+			if (!bHaveStaticResponse)
+			{
+				TSharedPtrTS<IElectraHttpManager::FProgressListener>	ProgressListener(new IElectraHttpManager::FProgressListener);
+				ProgressListener->ProgressDelegate   = Electra::MakeDelegate(this, &FStreamHandler::HTTPProgressCallback);
+				ProgressListener->CompletionDelegate = Electra::MakeDelegate(this, &FStreamHandler::HTTPCompletionCallback);
+				HTTP->ProgressListener = ProgressListener;
+				PlayerSessionService->GetHTTPManager()->AddRequest(HTTP);
+				while(!HasReadBeenAborted())
+				{
+					if (DownloadCompleteSignal.WaitTimeout(1000 * 100))
+					{
+						break;
+					}
+				}
+				ProgressListener.Reset();
+				// Note: It is only safe to access the connection info when the HTTP request has completed or the request been removed.
+				PlayerSessionService->GetHTTPManager()->RemoveRequest(HTTP);
+			}
+
+			InRequest->ConnectionInfo = HTTP->ConnectionInfo;
+			if (!HTTP->ConnectionInfo.StatusInfo.ErrorDetail.IsError())
+			{
+				// Notify license key download ok.
+				PlayerSessionService->SendMessageToPlayer(FLicenseKeyMessage::Create(FLicenseKeyMessage::EReason::LicenseKeyDownload, FErrorDetail(), &HTTP->ConnectionInfo));
+
+				// There is not much we can validate here. The key is the direct key data without any wrapping.
+				if (ReadBuffer.ReceiveBuffer->Buffer.Num() == 16)
+				{
+					LicenseKeyData = MakeShared<TArray<uint8>, ESPMode::ThreadSafe>();
+					LicenseKeyData->AddUninitialized(ReadBuffer.ReceiveBuffer->Buffer.Num());
+					FMemory::Memcpy(LicenseKeyData->GetData(), ReadBuffer.ReceiveBuffer->Buffer.GetLinearReadData(), ReadBuffer.ReceiveBuffer->Buffer.GetLinearReadSize());
+					if (InRequest->LicenseKeyCache.IsValid())
+					{
+						InRequest->LicenseKeyCache->AddLicenseKey(LicenseKeyData, LicenseKeyInfo, FTimeValue::GetPositiveInfinity());
+					}
+					OutLicenseKeyData = LicenseKeyData;
+					return FStreamReaderHLSfmp4::FStreamHandler::ELicenseKeyResult::Ok;
+				}
+				else
+				{
+					// Notify license key download failure.
+					PlayerSessionService->SendMessageToPlayer(FLicenseKeyMessage::Create(FLicenseKeyMessage::EReason::LicenseKeyData,
+																						 FErrorDetail().SetError(UEMEDIA_ERROR_FORMAT_ERROR).SetFacility(Facility::EFacility::LicenseKey).SetCode((uint16)FStreamReaderHLSfmp4::FStreamHandler::ELicenseKeyResult::FormatError).SetMessage("Invalid license key length"),
+																						 &HTTP->ConnectionInfo));
+					return FStreamReaderHLSfmp4::FStreamHandler::ELicenseKeyResult::FormatError;
+				}
+			}
+			else
+			{
+				// Notify license key download failure.
+				PlayerSessionService->SendMessageToPlayer(FLicenseKeyMessage::Create(FLicenseKeyMessage::EReason::LicenseKeyDownload,
+																					 FErrorDetail().SetError(UEMEDIA_ERROR_READ_ERROR).SetFacility(Facility::EFacility::LicenseKey).SetCode((uint16)FStreamReaderHLSfmp4::FStreamHandler::ELicenseKeyResult::DownloadError).SetMessage("License key download failure"),
+																					 &HTTP->ConnectionInfo));
+				return FStreamReaderHLSfmp4::FStreamHandler::ELicenseKeyResult::DownloadError;
+			}
+		}
+	}
+	return FStreamReaderHLSfmp4::FStreamHandler::ELicenseKeyResult::DownloadError;
+}
+
+
+FStreamReaderHLSfmp4::FStreamHandler::EInitSegmentResult FStreamReaderHLSfmp4::FStreamHandler::GetInitSegment(FErrorDetail& OutErrorDetail, TSharedPtrTS<const IParserISO14496_12>& OutMP4InitSegment, const TSharedPtrTS<FStreamSegmentRequestHLSfmp4>& Request)
+{
+	if (Request->InitSegmentInfo.IsValid() && Request->InitSegmentInfo->URI.Len())
+	{
+		TSharedPtrTS<const IParserISO14496_12> MP4InitSegment = Request->InitSegmentCache->GetInitSegmentFor(Request->InitSegmentInfo);
+		if (!MP4InitSegment.IsValid())
+		{
+			TSharedPtr<TArray<uint8>, ESPMode::ThreadSafe> LicenseKeyData;
+			// Is the init segment encrypted?
+			if (Request->InitSegmentInfo->DRMKeyInfo.IsValid())
+			{
+				ELicenseKeyResult LicenseKeyResult = GetLicenseKey(OutErrorDetail, LicenseKeyData, Request, Request->InitSegmentInfo->DRMKeyInfo);
+				if (LicenseKeyResult == FStreamReaderHLSfmp4::FStreamHandler::ELicenseKeyResult::DownloadError)
+				{
+					return FStreamReaderHLSfmp4::FStreamHandler::EInitSegmentResult::DownloadError;
+				}
+				else if (LicenseKeyResult == FStreamReaderHLSfmp4::FStreamHandler::ELicenseKeyResult::FormatError)
+				{
+					return FStreamReaderHLSfmp4::FStreamHandler::EInitSegmentResult::LicenseKeyError;
+				}
+			}
+
+			TSharedPtrTS<IElectraHttpManager::FRequest> 			HTTP(new IElectraHttpManager::FRequest);
+			TSharedPtrTS<IElectraHttpManager::FProgressListener>	ProgressListener(new IElectraHttpManager::FProgressListener);
+			ProgressListener->ProgressDelegate   = Electra::MakeDelegate(this, &FStreamHandler::HTTPProgressCallback);
+			ProgressListener->CompletionDelegate = Electra::MakeDelegate(this, &FStreamHandler::HTTPCompletionCallback);
+			ReadBuffer.Reset();
+			ReadBuffer.ReceiveBuffer = MakeSharedTS<IElectraHttpManager::FReceiveBuffer>();
+
+			// FIXME: Presently we demote init segments to HTTP unless explicitly forbidden to do so.
+			//        Since they are not encrypted we have no reliable way to know if the content itself is encrypted.
+			//        We say it is if _any_ segment in the playlist is using encryption, whichever one that is.
+			FString RequestURL = DemoteInitURLToHTTP(Request->InitSegmentInfo->URI, Request->bHasEncryptedSegments);
+
+			Metrics::FSegmentDownloadStats& ds = Request->DownloadStats;
+			ds.URL  	   = RequestURL;
+			ds.SegmentType = Metrics::ESegmentType::Init;
+
+			HTTP->Parameters.URL   = RequestURL;
+			HTTP->ReceiveBuffer    = ReadBuffer.ReceiveBuffer;
+			HTTP->ProgressListener = ProgressListener;
+			if (Request->InitSegmentInfo->ByteRange.IsSet())
+			{
+				HTTP->Parameters.Range.Start		= Request->InitSegmentInfo->ByteRange.GetStart();
+				HTTP->Parameters.Range.EndIncluding = Request->InitSegmentInfo->ByteRange.GetEnd();
+			}
+
+			ProgressReportCount = 0;
+			DownloadCompleteSignal.Reset();
+			PlayerSessionService->GetHTTPManager()->AddRequest(HTTP);
+
+			while(!HasReadBeenAborted())
+			{
+				if (DownloadCompleteSignal.WaitTimeout(1000 * 100))
+				{
+					break;
+				}
+			}
+
+			ProgressListener.Reset();
+			// Note: It is only safe to access the connection info when the HTTP request has completed or the request been removed.
+			PlayerSessionService->GetHTTPManager()->RemoveRequest(HTTP);
+			Request->ConnectionInfo = HTTP->ConnectionInfo;
+			if (!HTTP->ConnectionInfo.StatusInfo.ErrorDetail.IsError())
+			{
+// TODO: If encrypted we must now decrypt it!!
+//		 The question is how to do this if there is no explicit IV and the media sequence number is to be used
+//		 with the init segment being static and not having one.
+//		 Presently the assumption is that the init segment is not encrypted.
+
+				TSharedPtrTS<IParserISO14496_12> InitSegmentParser = IParserISO14496_12::CreateParser();
+				FParamDict parseOptions;
+				UEMediaError parseError = InitSegmentParser->ParseHeader(this, this, parseOptions, Parameters.PlayerSessionService);
+				if (parseError == UEMEDIA_ERROR_OK || parseError == UEMEDIA_ERROR_END_OF_STREAM)
+				{
+					// Parse the tracks of the init segment. We do this mainly to get to the CSD we might need should we have to insert filler data later.
+					parseError = InitSegmentParser->PrepareTracks(TSharedPtrTS<const IParserISO14496_12>());
+					if (parseError == UEMEDIA_ERROR_OK)
+					{
+						Request->InitSegmentCache->AddInitSegment(InitSegmentParser, Request->InitSegmentInfo, FTimeValue::GetPositiveInfinity());
+						OutMP4InitSegment = InitSegmentParser;
+						return FStreamReaderHLSfmp4::FStreamHandler::EInitSegmentResult::Ok;
+					}
+					else
+					{
+						return FStreamReaderHLSfmp4::FStreamHandler::EInitSegmentResult::ParseError;
+					}
+				}
+				else
+				{
+					return FStreamReaderHLSfmp4::FStreamHandler::EInitSegmentResult::ParseError;
+				}
+			}
+			else
+			{
+				return FStreamReaderHLSfmp4::FStreamHandler::EInitSegmentResult::DownloadError;
+			}
+		}
+		else
+		{
+			OutMP4InitSegment = MP4InitSegment;
+		}
+	}
+	return FStreamReaderHLSfmp4::FStreamHandler::EInitSegmentResult::AlreadyCached;
+}
+
+
+void FStreamReaderHLSfmp4::FStreamHandler::HandleRequest()
+{
+//CSV_SCOPED_TIMING_STAT(ElectraPlayer, StreamReaderMP4_Worker);
+
+	UEMediaError								Error;
+	bool										bIsEmptyFillerSegment = false;
+
+	// Get the request into a local TSharedPtrTS to hold on to it.
+	TSharedPtrTS<FStreamSegmentRequestHLSfmp4> Request = CurrentRequest;
+
+	// FIXME: We need to hold on to the retry info across several segment attempts.
+	//        The connection info gets set and cleared in here a few times and we actually need
+	//        to make sure the retry info is not modified.
+	TSharedPtrTS<HTTP::FRetryInfo> CurrentRetryInfo = CurrentRequest->ConnectionInfo.RetryInfo;
+
+	FTimeValue LoopTimestampOffset = Request->PlayerLoopState.LoopBasetime;
+	TSharedPtr<const FPlayerLoopState, ESPMode::ThreadSafe>		PlayerLoopState = MakeShared<const FPlayerLoopState, ESPMode::ThreadSafe>(Request->PlayerLoopState);
+
+	Metrics::FSegmentDownloadStats& ds = CurrentRequest->DownloadStats;
+	ds.StatsID = FMediaInterlockedIncrement(UniqueDownloadID);
+
+	ds.FailureReason.Empty();
+	ds.bWasSuccessful      = true;
+	ds.bWasAborted  	   = false;
+	ds.bDidTimeout  	   = false;
+	ds.HTTPStatusCode      = 0;
+	ds.StreamType   	   = Request->GetType();
+	ds.SegmentType  	   = Metrics::ESegmentType::Media;
+	ds.PresentationTime    = Request->AbsoluteDateTime.GetAsSeconds();
+	ds.Bitrate  		   = Request->Bitrate;
+	ds.Duration 		   = Request->SegmentDuration.GetAsSeconds();
+	ds.DurationDownloaded  = 0.0;
+	ds.DurationDelivered   = 0.0;
+	ds.TimeToFirstByte     = 0.0;
+	ds.TimeToDownload      = 0.0;
+	ds.ByteSize 		   = -1;
+	ds.NumBytesDownloaded  = 0;
+	ds.ThroughputBps	   = 0;
+	ds.bInsertedFillerData = false;
+
+	ds.MediaAssetID 	= Request->MediaAsset.IsValid() ? Request->MediaAsset->GetUniqueIdentifier() : "";
+	ds.AdaptationSetID  = Request->AdaptationSet.IsValid() ? Request->AdaptationSet->GetUniqueIdentifier() : "";
+	ds.RepresentationID = Request->Representation.IsValid() ? Request->Representation->GetUniqueIdentifier() : "";
+	ds.URL  			= Request->URL;
+	ds.CDN  			= Request->CDN;
+	ds.RetryNumber  	= Request->NumOverallRetries;
+
+	bIsEmptyFillerSegment = Request->bInsertFillerData;
+
+	TSharedPtrTS<FAccessUnit::CodecData>		CSD(new FAccessUnit::CodecData);
+	TSharedPtrTS<const IParserISO14496_12>		MP4InitSegment;
+	FErrorDetail								InitSegmentErrorDetail;
+	EInitSegmentResult							InitSegmentResult = EInitSegmentResult::AlreadyCached;
+
+	bHasErrored 		   = false;
+	bAbortedByABR   	   = false;
+	bAllowEarlyEmitting    = false;
+	bFillRemainingDuration = false;
+	DurationSuccessfullyDelivered.SetToZero();
+	AccessUnitFIFO.Clear();
+
+	Parameters.EventListener->OnFragmentOpen(CurrentRequest);
+
+	if (!bIsEmptyFillerSegment)
+	{
+		InitSegmentResult = GetInitSegment(InitSegmentErrorDetail, MP4InitSegment, Request);
+		// If we just downloaded the init segment successfully let the stream selector know.
+		if (InitSegmentResult == EInitSegmentResult::Ok)
+		{
+			StreamSelector->ReportDownloadEnd(ds);
+		}
+	}
+	else
+	{
+		// See if we have the init segment in the cache already. We won't request it if it is not as perhaps we are to insert
+		// filler data because the init segment has already failed.
+		if (Request->InitSegmentCache.IsValid() && Request->InitSegmentInfo.IsValid())
+		{
+			MP4InitSegment = Request->InitSegmentCache->GetInitSegmentFor(Request->InitSegmentInfo);
+			if (MP4InitSegment.IsValid())
+			{
+				// Get the CSD from track 0. This is identical to what we are doing  further down with the actual track.
+				check(MP4InitSegment->GetNumberOfTracks() == 1);
+				if (MP4InitSegment->GetNumberOfTracks() == 1)
+				{
+					const IParserISO14496_12::ITrack* Track = MP4InitSegment->GetTrackByIndex(0);
+					check(Track);
+					if (Track)
+					{
+						CSD->CodecSpecificData = Track->GetCodecSpecificData();
+						CSD->RawCSD = Track->GetCodecSpecificDataRAW();
+						CSD->ParsedInfo = Track->GetCodecInformation();
+					}
+				}
+			}
+		}
+	}
+
+	FTimeValue NextExpectedDTS;
+	FTimeValue DiscardBefore(FTimeValue::GetZero());
+	FTimeValue DurationSuccessfullyRead(FTimeValue::GetZero());
+	FTimeValue LastKnownAUDuration;
+	if (!bIsEmptyFillerSegment)
+	{
+		if (InitSegmentResult == EInitSegmentResult::Ok || InitSegmentResult == EInitSegmentResult::AlreadyCached)
+		{
+			// If the segment is encrypted we need to get the license key to decrypt it.
+			TSharedPtr<TArray<uint8>, ESPMode::ThreadSafe> LicenseKeyData;
+			// Is the init segment encrypted?
+			if (Request->LicenseKeyInfo.IsValid())
+			{
+				ELicenseKeyResult LicenseKeyResult = GetLicenseKey(InitSegmentErrorDetail, LicenseKeyData, Request, Request->LicenseKeyInfo);
+				if (LicenseKeyResult == FStreamReaderHLSfmp4::FStreamHandler::ELicenseKeyResult::DownloadError)
+				{
+					ds.FailureReason = FString::Printf(TEXT("Failed to download license key"));
+					LogMessage(IInfoLog::ELevel::Error, ds.FailureReason);
+					bHasErrored = true;
+				}
+				else if (LicenseKeyResult == FStreamReaderHLSfmp4::FStreamHandler::ELicenseKeyResult::FormatError)
+				{
+					ds.FailureReason = FString::Printf(TEXT("License key format error"));
+					LogMessage(IInfoLog::ELevel::Error, ds.FailureReason);
+					bHasErrored = true;
+				}
+
+				// Check that the encryption is AES-128 for the time being.
+				if (Request->LicenseKeyInfo->Method != FManifestHLSInternal::FMediaStream::FDRMKeyInfo::EMethod::None &&
+					Request->LicenseKeyInfo->Method != FManifestHLSInternal::FMediaStream::FDRMKeyInfo::EMethod::AES128)
+				{
+					ds.FailureReason = FString::Printf(TEXT("Unsupported encryption method"));
+					LogMessage(IInfoLog::ELevel::Error, ds.FailureReason);
+					bHasErrored = true;
+				}
+
+				// Create the decrypter.
+				if (!bHasErrored)
+				{
+					if (LicenseKeyData.IsValid())
+					{
+						IStreamDecrypterAES128::EResult DecrypterResult;
+						Decrypter = IStreamDecrypterAES128::Create();
+
+						// Set up the IV for this segment which is either explicitly provided or the media sequence number.
+						TArray<uint8> IV;
+						if (Request->LicenseKeyInfo->IV.Len())
+						{
+							DecrypterResult = IStreamDecrypterAES128::ConvHexStringToBin(IV, TCHAR_TO_ANSI(*Request->LicenseKeyInfo->IV));
+							if (DecrypterResult != IStreamDecrypterAES128::EResult::Ok)
+							{
+								ds.FailureReason = FString::Printf(TEXT("Bad explicit IV value"));
+								LogMessage(IInfoLog::ELevel::Error, ds.FailureReason);
+								bHasErrored = true;
+							}
+						}
+						else
+						{
+							IStreamDecrypterAES128::MakePaddedIVFromUInt64(IV, Request->MediaSequence);
+						}
+						if (!bHasErrored)
+						{
+							DecrypterResult = Decrypter->CBCInit(*LicenseKeyData, &IV);
+							if (DecrypterResult != IStreamDecrypterAES128::EResult::Ok)
+							{
+								ds.FailureReason = FString::Printf(TEXT("Received bad license key"));
+								LogMessage(IInfoLog::ELevel::Error, ds.FailureReason);
+								bHasErrored = true;
+							}
+						}
+					}
+					else
+					{
+						ds.FailureReason = FString::Printf(TEXT("No valid license key"));
+						LogMessage(IInfoLog::ELevel::Error, ds.FailureReason);
+						bHasErrored = true;
+					}
+				}
+			}
+
+			if (!bHasErrored)
+			{
+				ReadBuffer.Reset();
+				ReadBuffer.ReceiveBuffer = MakeSharedTS<IElectraHttpManager::FReceiveBuffer>();
+
+				// Start downloading the segment. Clear any stats that may have been set by the init segment download.
+				FString RequestURL = DemoteMediaURLToHTTP(Request->URL, Decrypter.IsValid());
+
+				ds.FailureReason.Empty();
+				ds.URL  			   = RequestURL;
+				ds.SegmentType  	   = Metrics::ESegmentType::Media;
+				ds.bWasSuccessful      = true;
+				ds.bWasAborted  	   = false;
+				ds.bDidTimeout  	   = false;
+				ds.HTTPStatusCode      = 0;
+				ds.TimeToFirstByte     = 0.0;
+				ds.TimeToDownload      = 0.0;
+				ds.ByteSize 		   = -1;
+				ds.NumBytesDownloaded  = 0;
+				ds.ThroughputBps	   = 0;
+
+				// Clear out the current connection info which may now be populated with the init segment fetch results.
+				CurrentRequest->ConnectionInfo = {};
+
+				TSharedPtrTS<IElectraHttpManager::FProgressListener>	ProgressListener(new IElectraHttpManager::FProgressListener);
+				ProgressListener->ProgressDelegate   = Electra::MakeDelegate(this, &FStreamHandler::HTTPProgressCallback);
+				ProgressListener->CompletionDelegate = Electra::MakeDelegate(this, &FStreamHandler::HTTPCompletionCallback);
+				TSharedPtrTS<IElectraHttpManager::FRequest> HTTP(new IElectraHttpManager::FRequest);
+				HTTP->Parameters.URL   = RequestURL;
+				HTTP->ReceiveBuffer    = ReadBuffer.ReceiveBuffer;
+				HTTP->ProgressListener = ProgressListener;
+				if (Request->Range.IsSet())
+				{
+					HTTP->Parameters.Range = Request->Range;
+				}
+
+				ProgressReportCount = 0;
+				DownloadCompleteSignal.Reset();
+				PlayerSessionService->GetHTTPManager()->AddRequest(HTTP);
+
+				MP4Parser = IParserISO14496_12::CreateParser();
+				NumMOOFBoxesFound = 0;
+
+				FTimeValue	BaseMediaDecodeTime;
+				FTimeValue	TimeMappingOffset;
+				bool		bDone = false;
+				bool		bTimeOffsetsSet = false;
+				int64		LastSuccessfulFilePos = 0;
+
+				bool bIsFirstAU = true;
+				while(!bDone && !HasErrored() && !HasReadBeenAborted())
+				{
+					UEMediaError parseError = MP4Parser->ParseHeader(this, this, Parameters.Options, Parameters.PlayerSessionService);
+					if (parseError == UEMEDIA_ERROR_OK)
+					{
+						parseError = MP4Parser->PrepareTracks(MP4InitSegment);
+						if (parseError == UEMEDIA_ERROR_OK)
+						{
+							// For the time being we only want to have a single track in the movie segments.
+							check(MP4Parser->GetNumberOfTracks() == 1);
+							if (MP4Parser->GetNumberOfTracks() == 1)
+							{
+								const IParserISO14496_12::ITrack* Track = MP4Parser->GetTrackByIndex(0);
+								check(Track);
+								if (Track)
+								{
+									CSD->CodecSpecificData = Track->GetCodecSpecificData();
+									CSD->RawCSD			= Track->GetCodecSpecificDataRAW();
+									CSD->ParsedInfo 		= Track->GetCodecInformation();
+
+									// TODO: Check that the track format matches the one we're expecting (video, audio, etc)
+									IParserISO14496_12::ITrackIterator* TrackIterator = Track->CreateIterator();
+
+									if (!bTimeOffsetsSet)
+									{
+										bTimeOffsetsSet 	= true;
+										BaseMediaDecodeTime = FTimeValue().SetFromND(TrackIterator->GetBaseMediaDecodeTime(), TrackIterator->GetTimescale());
+										DiscardBefore   	= Request->FirstAUTimeOffset + BaseMediaDecodeTime;
+										TimeMappingOffset   = Request->AbsoluteDateTime - BaseMediaDecodeTime;
+									}
+
+									FTimeValue DTS;
+									FTimeValue PTS;
+									FTimeValue Duration;
+									for(Error = TrackIterator->StartAtTime(FTimeValue::GetZero(), IParserISO14496_12::ITrackIterator::ESearchMode::After); Error == UEMEDIA_ERROR_OK; Error = TrackIterator->Next())
+									{
+										// Get the DTS and PTS. Those are 0-based in a fragment and offeet by the base media decode time of the fragment.
+										DTS.SetFromND(TrackIterator->GetDTS(), TrackIterator->GetTimescale());
+										PTS.SetFromND(TrackIterator->GetPTS(), TrackIterator->GetTimescale());
+
+										#if 0
+										MEDIAdbg::Printf("[%s] %lld/%d/%u: DTS=%lld, PTS=%lld, dur=%lld, scale=%u, sync:%d, size=%lld, off=%lld\n", GetStreamTypeName(Request->GetType()), (long long int)Request->MediaSequence, NumMOOFBoxesFound, TrackIterator->GetSampleNumber(),
+											(long long int)TrackIterator->GetDTS(),
+											(long long int)TrackIterator->GetPTS(),
+											(long long int)TrackIterator->GetDuration(),
+											TrackIterator->GetTimescale(),
+											TrackIterator->IsSyncSample()?1:0,
+											(long long int)TrackIterator->GetSampleSize(),
+											(long long int)TrackIterator->GetSampleFileOffset());
+										#endif
+
+										// Create access unit
+										FAccessUnit *AccessUnit = FAccessUnit::Create(Parameters.MemoryProvider);
+										check(AccessUnit);
+
+										Duration.SetFromND(TrackIterator->GetDuration(), TrackIterator->GetTimescale());
+
+										AccessUnit->ESType			 = Request->GetType();
+										AccessUnit->Duration  		 = Duration;
+										AccessUnit->AUSize			 = (uint32) TrackIterator->GetSampleSize();
+										AccessUnit->AUData   		 = AccessUnit->AllocatePayloadBuffer(AccessUnit->AUSize);
+										check(AccessUnit->AUData);
+										AccessUnit->bIsFirstInSequence = bIsFirstAU;
+										AccessUnit->bIsSyncSample 	 = TrackIterator->IsSyncSample();
+										AccessUnit->bIsDummyData  	 = false;
+										AccessUnit->AUCodecData		 = CSD;
+
+										// Calculate the drop on the fragment local DTS/PTS.
+										AccessUnit->DropState = FAccessUnit::EDropState::None;
+										if (DTS < DiscardBefore)
+										{
+											AccessUnit->DropState |= FAccessUnit::EDropState::DtsTooEarly;
+										}
+										if (PTS < DiscardBefore)
+										{
+											AccessUnit->DropState |= FAccessUnit::EDropState::PtsTooEarly;
+										}
+
+
+										// Offset the AU's DTS and PTS to the time mapping of the segment.
+										AccessUnit->DTS = DTS + TimeMappingOffset + LoopTimestampOffset;
+										AccessUnit->PTS = PTS + TimeMappingOffset + LoopTimestampOffset;
+
+										//AccessUnit->StreamSourceInfo = ....;
+										AccessUnit->PlayerLoopState = PlayerLoopState;
+
+										// Update the current download stats which we report periodically to the ABR.
+										ds.DurationDownloaded = DurationSuccessfullyRead.GetAsSeconds();
+										ds.DurationDelivered  = DurationSuccessfullyDelivered.GetAsSeconds();
+
+										// There should not be any gaps!
+								// TODO: what if there are?
+										check(GetCurrentOffset() == TrackIterator->GetSampleFileOffset());
+										int64 NumRead = ReadData(AccessUnit->AUData, AccessUnit->AUSize);
+										if (NumRead == AccessUnit->AUSize)
+										{
+											DurationSuccessfullyRead += Duration;
+											NextExpectedDTS = AccessUnit->DTS + Duration;
+											LastKnownAUDuration = Duration;
+											LastSuccessfulFilePos = GetCurrentOffset();
+										}
+										else
+										{
+											// Did not get the number of bytes we needed. Either because of a read error or because we got aborted.
+											FAccessUnit::Release(AccessUnit);
+											AccessUnit = nullptr;
+											bDone = true;
+											break;
+										}
+
+
+										if (AccessUnit)
+										{
+											AccessUnitFIFO.Push(AccessUnit);
+										}
+
+										// Shall we pass on any AUs we already read?
+										if (bAllowEarlyEmitting)
+										{
+											while(AccessUnitFIFO.Num() && !HasReadBeenAborted())
+											{
+												FAccessUnit* pNext = AccessUnitFIFO.FrontRef();
+												if (Parameters.EventListener->OnFragmentAccessUnitReceived(pNext))
+												{
+													DurationSuccessfullyDelivered += pNext->Duration;
+													AccessUnitFIFO.Pop();
+												}
+												else
+												{
+													break;
+												}
+											}
+										}
+
+										bIsFirstAU = false;
+									}
+									delete TrackIterator;
+
+									// Check if we are done or if there is additional data that needs parsing, like more moof boxes.
+									if (HasReadBeenAborted() || HasReachedEOF())
+									{
+										bDone = true;
+									}
+								}
+								else
+								{
+									// can't really happen. would indicate an internal screw up
+									LogMessage(IInfoLog::ELevel::Error, FString::Printf(TEXT("Segment \"%s\" has no track"), *Request->URL));
+									bHasErrored = true;
+								}
+							}
+							else
+							{
+								// more than 1 track
+								LogMessage(IInfoLog::ELevel::Error, FString::Printf(TEXT("Segment \"%s\" has more than one track"), *Request->URL));
+								bHasErrored = true;
+							}
+						}
+						else
+						{
+							// error preparing track for iterating
+							LogMessage(IInfoLog::ELevel::Error, FString::Printf(TEXT("Failed to prepare segment \"%s\" for iterating"), *Request->URL));
+							bHasErrored = true;
+						}
+					}
+					else if (parseError == UEMEDIA_ERROR_END_OF_STREAM)
+					{
+						bDone = true;
+					}
+					else
+					{
+						// failed to parse the segment (in general)
+						if (!HasReadBeenAborted())
+						{
+							LogMessage(IInfoLog::ELevel::Error, FString::Printf(TEXT("Failed to download segment \"%s\""), *Request->URL));
+							bHasErrored = true;
+						}
+					}
+				}
+				ProgressListener.Reset();
+				// Note: It is only safe to access the connection info when the HTTP request has completed or the request been removed.
+				PlayerSessionService->GetHTTPManager()->RemoveRequest(HTTP);
+				CurrentRequest->ConnectionInfo = HTTP->ConnectionInfo;
+				HTTP.Reset();
+			}
+		}
+		else
+		{
+			if (!HasReadBeenAborted())
+			{
+				// Init segment failed to download or parse.
+				CurrentRequest->ConnectionInfo.StatusInfo.ErrorDetail.SetFacility(Facility::EFacility::HLSFMP4Reader);
+				if (InitSegmentResult == EInitSegmentResult::ParseError)
+				{
+					CurrentRequest->ConnectionInfo.StatusInfo.ErrorDetail.SetMessage(TEXT("Init segment parse error")).SetCode(INTERNAL_ERROR_INIT_SEGMENT_PARSE_ERROR);
+					ds.bParseFailure  = true;
+				}
+				else if (InitSegmentResult == EInitSegmentResult::LicenseKeyError)
+				{
+					CurrentRequest->ConnectionInfo.StatusInfo.ErrorDetail.SetMessage(FString::Printf(TEXT("Init segment license key format error"))).SetCode(INTERNAL_ERROR_INIT_SEGMENT_LICENSEKEY_ERROR);
+				}
+				else
+				{
+					// This is either a download failure of the init segment or its license key.
+					CurrentRequest->ConnectionInfo.StatusInfo.ErrorDetail.SetMessage(FString::Printf(TEXT("Init segment download error: %s"), *CurrentRequest->ConnectionInfo.StatusInfo.ErrorDetail.GetMessage())).SetCode(INTERNAL_ERROR_INIT_SEGMENT_DOWNLOAD_ERROR);
+				}
+				bHasErrored = true;
+			}
+		}
+	}
+
+	// Do we need to fill remaining duration with dummy data?
+	if (bIsEmptyFillerSegment || bFillRemainingDuration)
+	{
+		// If this is a prefetch segment we will not fill in dummy data as the actual duration is not yet known
+		// and an approximation only. If it is too long we would create an overlap with the next segment which
+		// is not desirable.
+		if (!Request->bIsPrefetch)
+		{
+			// Get the supposed segment duration.
+			FTimeValue SegmentDurationToGo = Request->SegmentDuration;
+
+			// Did we get anything so far?
+			FTimeValue DefaultDuration;
+			if (NextExpectedDTS.IsValid())
+			{
+				check(DurationSuccessfullyRead.IsValid());
+				check(LastKnownAUDuration.IsValid());
+				SegmentDurationToGo -= DurationSuccessfullyRead;
+				DefaultDuration = LastKnownAUDuration;
+			}
+			else
+			{
+				// No. We need to start with the segment time.
+				NextExpectedDTS = Request->AbsoluteDateTime + LoopTimestampOffset;
+				DiscardBefore   = NextExpectedDTS + Request->FirstAUTimeOffset;
+				switch(Request->GetType())
+				{
+					case EStreamType::Video:
+					{
+						DefaultDuration.SetFromND(1, 60);
+						break;
+					}
+					case EStreamType::Audio:
+					{
+						int64 n = 1024;
+						uint32 d = 48000;
+						if (CSD->ParsedInfo.GetSamplingRate())
+						{
+							d = (uint32) CSD->ParsedInfo.GetSamplingRate();
+							switch(CSD->ParsedInfo.GetCodec())
+							{
+								case FStreamCodecInformation::ECodec::AAC:
+								{
+									n = CSD->ParsedInfo.GetExtras().GetValue("samples_per_block").SafeGetInt64(1024);
+									break;
+								}
+							}
+						}
+						DefaultDuration.SetFromND(n, d);
+						break;
+					}
+					default:
+					{
+						DefaultDuration.SetFromND(1, 10);
+						break;
+					}
+				}
+			}
+
+			ds.bInsertedFillerData = SegmentDurationToGo > FTimeValue::GetZero();
+			while(SegmentDurationToGo > FTimeValue::GetZero())
+			{
+				FAccessUnit *AccessUnit = FAccessUnit::Create(Parameters.MemoryProvider);
+				check(AccessUnit);
+
+				AccessUnit->ESType	   = Request->GetType();
+				AccessUnit->Duration     = DefaultDuration;
+				AccessUnit->AUSize	   = 0;
+				AccessUnit->AUData      = nullptr;
+				AccessUnit->bIsDummyData = true;
+				if (CSD->CodecSpecificData.Num())
+				{
+					AccessUnit->AUCodecData = CSD;
+				}
+
+				// Calculate the drop on the fragment local NextExpectedDTS/PTS.
+				AccessUnit->DropState = FAccessUnit::EDropState::None;
+				if (NextExpectedDTS < DiscardBefore)
+				{
+					AccessUnit->DropState |= FAccessUnit::EDropState::DtsTooEarly;
+					AccessUnit->DropState |= FAccessUnit::EDropState::PtsTooEarly;
+				}
+
+				AccessUnit->DTS = NextExpectedDTS;
+				AccessUnit->PTS = NextExpectedDTS;
+
+				NextExpectedDTS += DefaultDuration;
+				// It is possible for the default duration to not be a whole integer multiple of the segment duration
+				// in which case we leave a short gap unfilled.
+				SegmentDurationToGo -= DefaultDuration;
+
+				// Add to the FIFO. We do not need to check for early emitting here as we are not waiting for any
+				// data to be read. We can just shove all the synthesized dummy AUs in there.
+				AccessUnitFIFO.Push(AccessUnit);
+			}
+		}
+	}
+
+	bool bFillWithDummyData = bIsEmptyFillerSegment || bFillRemainingDuration;
+	while(AccessUnitFIFO.Num() && !bTerminate && (!HasReadBeenAborted() || bFillWithDummyData))
+	{
+		FAccessUnit* pNext = AccessUnitFIFO.FrontRef();
+		while(!bTerminate && (!HasReadBeenAborted() || bFillWithDummyData))
+		{
+			if (Parameters.EventListener->OnFragmentAccessUnitReceived(pNext))
+			{
+				DurationSuccessfullyDelivered += pNext->Duration;
+				AccessUnitFIFO.Pop();
+				break;
+			}
+			else
+			{
+				FMediaRunnable::SleepMicroseconds(1000 * 20);
+			}
+		}
+	}
+	// Anything not handed over after an abort we delete
+	while(AccessUnitFIFO.Num())
+	{
+		FAccessUnit::Release(AccessUnitFIFO.Pop());
+	}
+
+
+	// Set up remaining download stat fields.
+	if (ds.FailureReason.Len() == 0)
+	{
+		ds.FailureReason = CurrentRequest->ConnectionInfo.StatusInfo.ErrorDetail.GetMessage();
+	}
+	if (bAbortedByABR)
+	{
+		// If aborted set the reason as the download failure.
+		ds.FailureReason = ds.ABRState.ProgressDecision.Reason;
+	}
+	ds.bWasAborted  	  = bAbortedByABR;
+	ds.bWasSuccessful     = !bHasErrored && !bAbortedByABR;
+	ds.URL  			  = CurrentRequest->ConnectionInfo.EffectiveURL;
+	ds.HTTPStatusCode     = CurrentRequest->ConnectionInfo.StatusInfo.HTTPStatus;
+	ds.DurationDownloaded = DurationSuccessfullyRead.GetAsSeconds();
+	ds.DurationDelivered  = DurationSuccessfullyDelivered.GetAsSeconds();
+	ds.TimeToFirstByte    = CurrentRequest->ConnectionInfo.TimeUntilFirstByte;
+	ds.TimeToDownload     = (CurrentRequest->ConnectionInfo.RequestEndTime - CurrentRequest->ConnectionInfo.RequestStartTime).GetAsSeconds();
+	ds.ByteSize 		  = CurrentRequest->ConnectionInfo.ContentLength;
+	ds.NumBytesDownloaded = CurrentRequest->ConnectionInfo.BytesReadSoFar;
+	ds.ThroughputBps	  = CurrentRequest->ConnectionInfo.Throughput.GetThroughput();
+	if (ds.ThroughputBps == 0)
+	{
+		ds.ThroughputBps = ds.TimeToDownload > 0.0 ? 8 * ds.NumBytesDownloaded / ds.TimeToDownload : 0;
+	}
+
+	StreamSelector->ReportDownloadEnd(ds);
+
+	// Remember the next expected timestamp.
+	CurrentRequest->NextLargestExpectedTimestamp = NextExpectedDTS;
+
+	// Restore the original retry info that may have been reset in all the changes and assignments in here.
+	CurrentRequest->ConnectionInfo.RetryInfo = CurrentRetryInfo;
+	// Clean out everything before reporting OnFragmentClose().
+	TSharedPtrTS<FStreamSegmentRequestHLSfmp4> FinishedRequest = CurrentRequest;
+	CurrentRequest.Reset();
+	ReadBuffer.Reset();
+	MP4Parser.Reset();
+	Decrypter.Reset();
+
+	Parameters.EventListener->OnFragmentClose(FinishedRequest);
+}
+
+
+bool FStreamReaderHLSfmp4::FStreamHandler::HasErrored() const
+{
+	return bHasErrored;
+}
+
+
+
+
+/**
+ * Read n bytes of data into the provided buffer.
+ *
+ * Reading must return the number of bytes asked to get, if necessary by blocking.
+ * If a read error prevents reading the number of bytes -1 must be returned.
+ *
+ * @param IntoBuffer Buffer into which to store the data bytes. If nullptr is passed the data must be skipped over.
+ * @param NumBytesToRead The number of bytes to read. Must not read more bytes and no less than requested.
+ * @return The number of bytes read or -1 on a read error.
+ */
+int64 FStreamReaderHLSfmp4::FStreamHandler::ReadData(void* IntoBuffer, int64 NumBytesToRead)
+{
+	FPODRingbuffer& SourceBuffer = ReadBuffer.ReceiveBuffer->Buffer;
+	// Make sure the buffer will have the amount of data we need.
+	while(1)
+	{
+		// Check if a HTTP reader progress event fired in the meantime.
+		if (ProgressReportCount)
+		{
+			ProgressReportCount = 0;
+			if (CurrentRequest.IsValid())
+			{
+				MetricUpdateLock.Lock();
+				Metrics::FSegmentDownloadStats currentDownloadStats = CurrentRequest->DownloadStats;
+				MetricUpdateLock.Unlock();
+
+				Metrics::FSegmentDownloadStats& ds = CurrentRequest->DownloadStats;
+				FABRDownloadProgressDecision StreamSelectorDecision = StreamSelector->ReportDownloadProgress(currentDownloadStats);
+				ds.ABRState.ProgressDecision = StreamSelectorDecision;
+				if ((StreamSelectorDecision.Flags & FABRDownloadProgressDecision::EDecisionFlags::eABR_EmitPartialData) != 0)
+				{
+					bAllowEarlyEmitting = true;
+					// Deliver all enqueued AUs right now. Unless the request also gets aborted we could be stuck
+					// in here for a while longer.
+					while(AccessUnitFIFO.Num())
+					{
+						FAccessUnit* pNext = AccessUnitFIFO.FrontRef();
+						if (Parameters.EventListener->OnFragmentAccessUnitReceived(pNext))
+						{
+							DurationSuccessfullyDelivered += pNext->Duration;
+							AccessUnitFIFO.Pop();
+						}
+						else
+						{
+							break;
+						}
+					}
+				}
+				if ((StreamSelectorDecision.Flags & FABRDownloadProgressDecision::EDecisionFlags::eABR_InsertFillerData) != 0)
+				{
+					bFillRemainingDuration = true;
+				}
+				if ((StreamSelectorDecision.Flags & FABRDownloadProgressDecision::EDecisionFlags::eABR_AbortDownload) != 0)
+				{
+					// When aborted and early emitting did place something into the buffers we need to fill
+					// the remainder no matter what.
+					if (DurationSuccessfullyDelivered > FTimeValue::GetZero())
+					{
+						bFillRemainingDuration = true;
+					}
+					bAbortedByABR = true;
+					return -1;
+				}
+			}
+		}
+
+		// Reading from an encrypted segment?
+		if (Decrypter.IsValid())
+		{
+			// We are handling full segment encryption here only (AES-128) and not sample encryption or any other scheme.
+			// Because AES is a block cipher we can only decrypt multiple of 16 byte chunks. Due to PKCS7 padding the encrypted segment
+			// will also be 1 to 16 bytes larger than its unencrypted original which we have to consider and remove as to not make
+			// the excess available to the caller here.
+
+			// Get the encrypted size we need to have in order to decrypt it and get the data we need.
+			int32 RequiredEncryptedSize = Decrypter->CBCGetEncryptionDataSize(ReadBuffer.ParsePos + NumBytesToRead);
+
+			// Reading data from a network stream has no reliable end-of-data marker since the data can be unbounded in length
+			// when using chunked transfer encoding. EOD is only signaled in the receive buffer when the connection is closed.
+			// To ensure we get a reliable EOD signal we wait for more data to arrive in the buffer than we actually want since
+			// the wait will either be satisfied with enough data (so what we want to have cannot be at the end yet) or when
+			// the EOD flag gets set at the end of the transfer (with fewer data than we waited for but ideally the amount we
+			// wanted to read (unless an error occurred)).
+			// One AES block size is sufficient here. It could be set to higher values (1K or 16K even) with little harm other
+			// than this will wait a tiny bit longer for new data then as long as we are not actually at EOD.
+			const int32 NumExtraRequiredToCatchEOD = 16;
+
+			if (!SourceBuffer.WaitUntilSizeAvailable(RequiredEncryptedSize + NumExtraRequiredToCatchEOD, 1000 * 100))
+			{
+				if (HasErrored() || HasReadBeenAborted() || SourceBuffer.WasAborted())
+				{
+					return -1;
+				}
+			}
+			else
+			{
+				SourceBuffer.Lock();
+				// Check the available size. If the read was aborted there may not be enough in here as the wait got released early.
+				if (SourceBuffer.Num() >= RequiredEncryptedSize)
+				{
+					int64 ReadUpToPos = ReadBuffer.ParsePos + NumBytesToRead;
+					// Have enough data. See if it needs to be decrypted.
+					if (ReadUpToPos > ReadBuffer.DecryptedPos)
+					{
+						// Decrypt from the last pos to the new pos now.
+						check((ReadBuffer.DecryptedPos & 15) == 0);
+						check((RequiredEncryptedSize & 15) == 0);
+						check(RequiredEncryptedSize > ReadBuffer.DecryptedPos);
+						uint8* EncryptedData = SourceBuffer.GetLinearReadData() + ReadBuffer.DecryptedPos;
+						int32 EncryptedSize = RequiredEncryptedSize - ReadBuffer.DecryptedPos;
+						bool bIsFinalBlock = false;
+						// See comment above on NumExtraRequiredToCatchEOD why we can check for EOD on the buffer here.
+						if (SourceBuffer.GetEOD() && RequiredEncryptedSize >= SourceBuffer.Num())
+						{
+							bIsFinalBlock = true;
+						}
+
+						// Decrypt data in place
+						IStreamDecrypterAES128::EResult DecrypterResult;
+						int32 NumDecryptedBytes = 0;
+						DecrypterResult = Decrypter->CBCDecryptInPlace(NumDecryptedBytes, EncryptedData, EncryptedSize, bIsFinalBlock);
+						// This cannot fail since we ensured it to be set up correctly and pass only properly aligned data, but just in case.
+						if (DecrypterResult == IStreamDecrypterAES128::EResult::Ok)
+						{
+							// Advance the decrypted pos by the entire encrypted block size, not the amount of decrypted bytes!
+							// This is required for the RequiredEncryptedSize test above to work correctly.
+							ReadBuffer.DecryptedPos += EncryptedSize;
+							if (bIsFinalBlock)
+							{
+								// On the final block adjust the maximum parse position to the end of the decrypted data
+								// which is less than the encrypted total size due to padding.
+								// This is needed for HasReachedEOF() to work correctly and not allow to read data from the padded area!
+								ReadBuffer.MaxParsePos = ReadBuffer.DecryptedPos - EncryptedSize + NumDecryptedBytes;
+							}
+						}
+						else
+						{
+							SourceBuffer.Unlock();
+							LogMessage(IInfoLog::ELevel::Error, FString::Printf(TEXT("Failed to decrypt (%s)"), IStreamDecrypterAES128::GetResultText(DecrypterResult)));
+							return -1;
+						}
+					}
+
+					// Enough decrypted data available?
+					if (ReadUpToPos <= ReadBuffer.DecryptedPos)
+					{
+						// Trying to read past the end of the decrypted data which would get into the padding area?
+						if (ReadUpToPos <= ReadBuffer.MaxParsePos)
+						{
+							// No, read is ok.
+							if (IntoBuffer)
+							{
+								FMemory::Memcpy(IntoBuffer, SourceBuffer.GetLinearReadData() + ReadBuffer.ParsePos, NumBytesToRead);
+							}
+							SourceBuffer.Unlock();
+							ReadBuffer.ParsePos += NumBytesToRead;
+							return NumBytesToRead;
+						}
+						else
+						{
+							// This means EOF now. Return 0 or -1 on error.
+							SourceBuffer.Unlock();
+							return HasErrored() ? -1 : 0;
+						}
+					}
+				}
+				else
+				{
+					// Return 0 at EOF and -1 on error.
+					SourceBuffer.Unlock();
+					return HasErrored() ? -1 : 0;
+				}
+				SourceBuffer.Unlock();
+			}
+		}
+		else
+		{
+			if (!SourceBuffer.WaitUntilSizeAvailable(ReadBuffer.ParsePos + NumBytesToRead, 1000 * 100))
+			{
+				if (HasErrored() || HasReadBeenAborted() || SourceBuffer.WasAborted())
+				{
+					return -1;
+				}
+			}
+			else
+			{
+				SourceBuffer.Lock();
+				if (SourceBuffer.Num() >= ReadBuffer.ParsePos + NumBytesToRead)
+				{
+					if (IntoBuffer)
+					{
+						FMemory::Memcpy(IntoBuffer, SourceBuffer.GetLinearReadData() + ReadBuffer.ParsePos, NumBytesToRead);
+					}
+					SourceBuffer.Unlock();
+					ReadBuffer.ParsePos += NumBytesToRead;
+					return NumBytesToRead;
+				}
+				else
+				{
+					// Return 0 at EOF and -1 on error.
+					SourceBuffer.Unlock();
+					return HasErrored() ? -1 : 0;
+				}
+			}
+		}
+	}
+	return -1;
+}
+
+/**
+ * Checks if the data source has reached the End Of File (EOF) and cannot provide any additional data.
+ *
+ * @return If EOF has been reached returns true, otherwise false.
+ */
+bool FStreamReaderHLSfmp4::FStreamHandler::HasReachedEOF() const
+{
+	const FPODRingbuffer& SourceBuffer = ReadBuffer.ReceiveBuffer->Buffer;
+	return !HasErrored() && SourceBuffer.GetEOD() && (ReadBuffer.ParsePos >= SourceBuffer.Num() || ReadBuffer.ParsePos >= ReadBuffer.MaxParsePos);
+}
+
+/**
+ * Checks if reading of the file and therefor parsing has been aborted.
+ *
+ * @return true if reading/parsing has been aborted, false otherwise.
+ */
+bool FStreamReaderHLSfmp4::FStreamHandler::HasReadBeenAborted() const
+{
+	return bTerminate || bRequestCanceled || bAbortedByABR;
+}
+
+/**
+ * Returns the current read offset.
+ *
+ * The first read offset is not necessarily zero. It could be anywhere inside the source.
+ *
+ * @return The current byte offset in the source.
+ */
+int64 FStreamReaderHLSfmp4::FStreamHandler::GetCurrentOffset() const
+{
+	return ReadBuffer.ParsePos;
+}
+
+
+IParserISO14496_12::IBoxCallback::EParseContinuation FStreamReaderHLSfmp4::FStreamHandler::OnFoundBox(IParserISO14496_12::FBoxType Box, int64 BoxSizeInBytes, int64 FileDataOffset, int64 BoxDataOffset)
+{
+	// Check which box is being parsed next.
+	switch(Box)
+	{
+		case IParserISO14496_12::BoxType_moov:
+		case IParserISO14496_12::BoxType_sidx:
+		{
+			return IParserISO14496_12::IBoxCallback::EParseContinuation::Continue;
+		}
+		case IParserISO14496_12::BoxType_moof:
+		{
+			++NumMOOFBoxesFound;
+			return IParserISO14496_12::IBoxCallback::EParseContinuation::Continue;
+		}
+		case IParserISO14496_12::BoxType_mdat:
+		{
+			return IParserISO14496_12::IBoxCallback::EParseContinuation::Stop;
+		}
+		default:
+		{
+			return IParserISO14496_12::IBoxCallback::EParseContinuation::Continue;
+		}
+	}
+}
+
+
+
+} // namespace Electra
+
diff --git a/Engine/Plugins/Media/ElectraPlayer/Source/ElectraPlayerRuntime/Private/Runtime/Player/HLS/StreamReaderHLSfmp4.h b/Engine/Plugins/Media/ElectraPlayer/Source/ElectraPlayerRuntime/Private/Runtime/Player/HLS/StreamReaderHLSfmp4.h
new file mode 100644
--- /dev/null
+++ b/Engine/Plugins/Media/ElectraPlayer/Source/ElectraPlayerRuntime/Private/Runtime/Player/HLS/StreamReaderHLSfmp4.h
@@ -0,0 +1,265 @@
+// Copyright Epic Games, Inc. All Rights Reserved.
+
+#pragma once
+
+#include "Player/PlayerStreamReader.h"
+#include "Player/AdaptiveStreamingPlayerResourceRequest.h"
+#include "Demuxer/ParserISO14496-12.h"
+#include "InitSegmentCacheHLS.h"
+#include "LicenseKeyCacheHLS.h"
+#include "HTTP/HTTPManager.h"
+#include "Crypto/StreamCryptoAES128.h"
+
+namespace Electra
+{
+
+class FStreamSegmentRequestHLSfmp4 : public IStreamSegment
+{
+public:
+	FStreamSegmentRequestHLSfmp4();
+	virtual ~FStreamSegmentRequestHLSfmp4();
+
+	virtual void SetPlaybackSequenceID(uint32 PlaybackSequenceID) override;
+	virtual uint32 GetPlaybackSequenceID() const override;
+
+
+	virtual EStreamType GetType() const override;
+
+	virtual void GetDependentStreams(TArray<FDependentStreams>& OutDependentStreams) const override;
+	virtual void GetEndedStreams(TArray<TSharedPtrTS<IStreamSegment>>& OutAlreadyEndedStreams) override;
+
+	//! Returns the first PTS value as indicated by the media timeline. This should correspond to the actual absolute PTS of the sample.
+	virtual FTimeValue GetFirstPTS() const override;
+
+	virtual int32 GetQualityIndex() const override;
+	virtual int32 GetBitrate() const override;
+
+	virtual void GetDownloadStats(Metrics::FSegmentDownloadStats& OutStats) const override;
+
+	void CopyFrom(const FStreamSegmentRequestHLSfmp4& rhs);
+
+	FString																		URL;
+	IElectraHttpManager::FParams::FRange										Range;
+
+	EStreamType																	StreamType;							//!< Type of stream (video, audio, etc.)
+	uint32																		StreamUniqueID;						//!< The unique stream ID identifying the stream for which this is a request.
+	int32																		Bitrate;
+	int32																		QualityLevel;
+
+	TSharedPtrTS<IPlaybackAssetRepresentation>									Representation;
+	TSharedPtrTS<IPlaybackAssetAdaptationSet>									AdaptationSet;
+	TSharedPtrTS<ITimelineMediaAsset>											MediaAsset;
+	FString																		CDN;
+
+	//FTimeValue																	PlaylistRelativeStartTime;			//!< The start time of this segment within the media playlist.
+	FTimeValue																	AbsoluteDateTime;					//!< The absolute start time of this segment as declared through EXT-X-PROGRAM-DATE-TIME mapping
+	FTimeValue																	SegmentDuration;					//!< Duration of the segment as specified in the media playlist.
+	int64																		MediaSequence;						//!< The media sequence number of this segment.
+	int64																		DiscontinuitySequence;				//!< The discontinuity index after which this segment is located in the media playlist.
+	int32																		LocalIndex;							//!< Local index of the segment in the media playlist at the time the request was generated.
+
+	FTimeValue																	FirstAUTimeOffset;					//!< A time offset into the segment to the first access unit to be sent to the decoder (audio).
+
+	bool																		bIsPrefetch;
+	bool																		bIsEOSSegment;
+
+	int32																		NumOverallRetries;					//!< Number of retries for this _segment_ across all possible quality levels and CDNs.
+	bool																		bInsertFillerData;
+
+	bool																		bHasEncryptedSegments;
+
+	TSharedPtrTS<IInitSegmentCacheHLS> 											InitSegmentCache;
+	TSharedPtrTS<const FManifestHLSInternal::FMediaStream::FInitSegmentInfo>	InitSegmentInfo;
+
+	TSharedPtrTS<ILicenseKeyCacheHLS> 											LicenseKeyCache;
+	TSharedPtrTS<const FManifestHLSInternal::FMediaStream::FDRMKeyInfo>			LicenseKeyInfo;
+
+	// List of dependent streams. Usually set for initial playback start requests.
+	TArray<TSharedPtrTS<FStreamSegmentRequestHLSfmp4>>							DependentStreams;
+	bool																		bIsInitialStartRequest;
+
+	FPlayerLoopState															PlayerLoopState;
+
+	uint32																		CurrentPlaybackSequenceID;			//!< Set by the player before adding the request to the stream reader.
+
+	Metrics::FSegmentDownloadStats												DownloadStats;
+	HTTP::FConnectionInfo														ConnectionInfo;
+	FTimeValue																	NextLargestExpectedTimestamp;	//!< Largest timestamp of all samples (plus its duration) across all tracks.
+};
+
+
+
+/**
+ *
+**/
+class FStreamReaderHLSfmp4 : public IStreamReader
+{
+public:
+	FStreamReaderHLSfmp4();
+	virtual ~FStreamReaderHLSfmp4();
+
+	virtual UEMediaError Create(IPlayerSessionServices* PlayerSessionService, const CreateParam& InCreateParam) override;
+	virtual void Close() override;
+
+	//! Adds a request to read from a stream
+	virtual EAddResult AddRequest(uint32 CurrentPlaybackSequenceID, TSharedPtrTS<IStreamSegment> Request) override;
+
+	//! Pauses all pending requests.
+	virtual void PauseDownload() override;
+	//! Resumes all pending requests.
+	virtual void ResumeDownload() override;
+	//! Cancels all pending requests.
+	virtual void CancelRequests() override;
+
+private:
+
+	struct FStreamHandler : public FMediaThread, public IParserISO14496_12::IReader, public IParserISO14496_12::IBoxCallback
+	{
+		struct FReadBuffer
+		{
+			FReadBuffer()
+			{
+				Reset();
+			}
+			void Reset()
+			{
+				ReceiveBuffer.Reset();
+				ParsePos = 0;
+				MaxParsePos = TNumericLimits<int64>::Max();
+				DecryptedPos = 0;
+			}
+			TSharedPtrTS<IElectraHttpManager::FReceiveBuffer>	ReceiveBuffer;
+			int64												ParsePos;
+			int64												MaxParsePos;
+			int32												DecryptedPos;
+		};
+
+
+		class FStaticResourceRequest : public IAdaptiveStreamingPlayerResourceRequest
+		{
+		public:
+			FStaticResourceRequest(FString InURL, EPlaybackResourceType InType)
+				: URL(InURL)
+				, Type(InType)
+			{ }
+
+			virtual ~FStaticResourceRequest()
+			{ }
+
+			virtual EPlaybackResourceType GetResourceType() const override
+			{ return Type; }
+
+			virtual FString GetResourceURL() const override
+			{ return URL; }
+
+			virtual void SetPlaybackData(TSharedPtr<TArray<uint8>, ESPMode::ThreadSafe>	PlaybackData) override
+			{ Data = PlaybackData; }
+
+			virtual void SignalDataReady() override
+			{ DoneSignal.Signal(); }
+
+			bool IsDone() const
+			{ return DoneSignal.IsSignaled(); }
+			
+			bool WaitDone(int32 WaitMicros)
+			{ return DoneSignal.WaitTimeout(WaitMicros); }
+
+			TSharedPtr<TArray<uint8>, ESPMode::ThreadSafe> GetData()
+			{ return Data; }
+
+		private:
+			FString												URL;
+			TSharedPtr<TArray<uint8>, ESPMode::ThreadSafe>		Data;
+			FMediaEvent											DoneSignal;
+			EPlaybackResourceType								Type;
+		};
+
+		enum class EInitSegmentResult
+		{
+			Ok,
+			AlreadyCached,
+			DownloadError,
+			ParseError,
+			LicenseKeyError
+		};
+
+		enum class ELicenseKeyResult
+		{
+			Ok,
+			AlreadyCached,
+			DownloadError,
+			FormatError
+		};
+
+		static uint32											UniqueDownloadID;
+
+		IStreamReader::CreateParam								Parameters;
+		TSharedPtrTS<FStreamSegmentRequestHLSfmp4>				CurrentRequest;
+		FMediaSemaphore											WorkSignal;
+		volatile bool											bTerminate;
+		volatile bool											bRequestCanceled;
+		volatile bool											bHasErrored;
+		bool													bAbortedByABR;
+		bool													bAllowEarlyEmitting;
+		bool													bFillRemainingDuration;
+
+		IPlayerSessionServices*									PlayerSessionService;
+		FReadBuffer												ReadBuffer;
+		TSharedPtr<IStreamDecrypterAES128, ESPMode::ThreadSafe>	Decrypter;
+		FMediaEvent												DownloadCompleteSignal;
+		TSharedPtrTS<IParserISO14496_12>						MP4Parser;
+		int32													NumMOOFBoxesFound;
+
+		TMediaQueueDynamicNoLock<FAccessUnit *>					AccessUnitFIFO;
+		FTimeValue 												DurationSuccessfullyDelivered;
+
+		FMediaCriticalSection									MetricUpdateLock;
+		int32													ProgressReportCount;
+		TSharedPtrTS<IAdaptiveStreamSelector>					StreamSelector;
+
+
+		FStreamHandler();
+		virtual ~FStreamHandler();
+		void Cancel();
+		void SignalWork();
+		void WorkerThread();
+		void HandleRequest();
+		FString DemoteMediaURLToHTTP(const FString& InURL, bool bIsEncrypted);
+		FString DemoteInitURLToHTTP(const FString& InURL, bool bIsEncrypted);
+		EInitSegmentResult GetInitSegment(FErrorDetail& OutErrorDetail, TSharedPtrTS<const IParserISO14496_12>& OutMP4InitSegment, const TSharedPtrTS<FStreamSegmentRequestHLSfmp4>& InRequest);
+		ELicenseKeyResult GetLicenseKey(FErrorDetail& OutErrorDetail, TSharedPtr<TArray<uint8>, ESPMode::ThreadSafe>& OutLicenseKeyData, const TSharedPtrTS<FStreamSegmentRequestHLSfmp4>& InRequest, const TSharedPtr<const FManifestHLSInternal::FMediaStream::FDRMKeyInfo, ESPMode::ThreadSafe>& LicenseKeyInfo);
+
+		void LogMessage(IInfoLog::ELevel Level, const FString& Message);
+
+		int32 HTTPProgressCallback(const IElectraHttpManager::FRequest* Request);
+		void HTTPCompletionCallback(const IElectraHttpManager::FRequest* Request);
+		void HTTPUpdateStats(const FTimeValue& CurrentTime, const IElectraHttpManager::FRequest* Request);
+
+
+		bool HasErrored() const;
+
+		// Methods from IParserISO14496_12::IReader
+		virtual int64 ReadData(void* IntoBuffer, int64 NumBytesToRead) override;
+		virtual bool HasReachedEOF() const override;
+		virtual bool HasReadBeenAborted() const override;
+		virtual int64 GetCurrentOffset() const override;
+		// Methods from IParserISO14496_12::IBoxCallback
+		virtual IParserISO14496_12::IBoxCallback::EParseContinuation OnFoundBox(IParserISO14496_12::FBoxType Box, int64 BoxSizeInBytes, int64 FileDataOffset, int64 BoxDataOffset) override;
+	};
+
+	// Currently set to use 2 handlers, one for video and one for audio. This could become a pool of n if we need to stream
+	// multiple dependent segments, keeping a pool of available and active handlers to cycle between.
+	FStreamHandler						StreamHandlers[2];		// 0 = video (MEDIAstreamType_Video), 1 = audio (MEDIAstreamType_Audio)
+	IPlayerSessionServices*				PlayerSessionService;
+	bool								bIsStarted;
+	FErrorDetail						ErrorDetail;
+
+	static const FString		OptionKeyDontUseInsecureForEncryptedMediaSegments;		//!< (bool) if false and media segment is using EXT-X-KEY encryption fetch it via http even if it should be https, otherwise keep the original scheme.
+	static const FString		OptionKeyDontUseInsecureForInitSegments;				//!< (bool) if false the init segment is fetched via http even if it should be https, otherwise keep the original scheme.
+};
+
+
+
+
+} // namespace Electra
+
diff --git a/Engine/Plugins/Media/ElectraPlayer/Source/ElectraPlayerRuntime/Private/Runtime/Player/HLS/TagContainer.cpp b/Engine/Plugins/Media/ElectraPlayer/Source/ElectraPlayerRuntime/Private/Runtime/Player/HLS/TagContainer.cpp
new file mode 100644
--- /dev/null
+++ b/Engine/Plugins/Media/ElectraPlayer/Source/ElectraPlayerRuntime/Private/Runtime/Player/HLS/TagContainer.cpp
@@ -0,0 +1,52 @@
+// Copyright Epic Games, Inc. All Rights Reserved.
+
+#include "CoreMinimal.h"
+#include "Parser.h"
+
+namespace Electra
+{
+	namespace HLSPlaylistParser
+	{
+
+		void TTagContainer::AddTag(FString& TagKey, FTagContent& Tag)
+		{
+			Tags.FindOrAdd(TagKey).Add(MoveTemp(Tag));
+        }
+
+        bool TTagContainer::ContainsTag(const FString& TagName) const
+        {
+			return Tags.Contains(TagName);
+        }
+
+		int TTagContainer::TagNum(const FString& TagName) const
+		{
+			const TArray<FTagContent>* Content = Tags.Find(TagName);
+			if (Content == nullptr)
+			{
+				return 0;
+			}
+
+			return Content->Num();
+		}
+
+        EPlaylistError TTagContainer::ValidateTags(const TMap<FString, FMediaTag>& TagDefinitions,
+                                                   const FMediaTagOption& Target)
+        {
+			for (auto& Elem : TagDefinitions)
+			{
+				FMediaTag Tag = Elem.Value;
+				if (Tag.HasOption(Target))
+				{
+					if (Tag.HasOption(Required) && !Tags.Contains(Elem.Key))
+					{
+						return EPlaylistError::MissingRequiredTag;
+					}
+				}
+			}
+
+            return EPlaylistError::None;
+        }
+    }
+} // namespace Electra
+
+
diff --git a/Engine/Plugins/Media/ElectraPlayer/Source/ElectraPlayerRuntime/Private/Runtime/Player/HLS/Tags.h b/Engine/Plugins/Media/ElectraPlayer/Source/ElectraPlayerRuntime/Private/Runtime/Player/HLS/Tags.h
new file mode 100644
--- /dev/null
+++ b/Engine/Plugins/Media/ElectraPlayer/Source/ElectraPlayerRuntime/Private/Runtime/Player/HLS/Tags.h
@@ -0,0 +1,41 @@
+// Copyright Epic Games, Inc. All Rights Reserved.
+
+#pragma once
+
+#include "CoreMinimal.h"
+
+namespace Electra
+{
+    namespace HLSPlaylistParser
+    {
+        // Media Segment related tags.
+        const FString ExtXByteRange = TEXT("EXT-X-BYTERANGE");
+        const FString ExtXVersion = TEXT("EXT-X-VERSION");
+        const FString ExtINF = TEXT("EXTINF");
+        const FString ExtXDiscontinuity = TEXT("EXT-X-DISCONTINUITY");
+        const FString ExtXKey = TEXT("EXT-X-KEY");
+        const FString ExtXMap = TEXT("EXT-X-MAP");
+        const FString ExtXProgramDateTime = TEXT("EXT-X-PROGRAM-DATE-TIME");
+        const FString ExtXDateRange = TEXT("EXT-X-DATERANGE");
+        // Playlist related tags for playlists of Media Segments.
+        const FString ExtXTargetDuration = TEXT("EXT-X-TARGETDURATION");
+        const FString ExtXMediaSequence = TEXT("EXT-X-MEDIA-SEQUENCE");
+        const FString ExtXDiscontinuitySequence = TEXT("EXT-X-DISCONTINUITY-SEQUENCE");
+        const FString ExtXEndlist = TEXT("EXT-X-ENDLIST");
+        const FString ExtXPlaylistType = TEXT("EXT-X-PLAYLIST-TYPE");
+        const FString ExtXIFramesOnly = TEXT("EXT-X-I-FRAMES-ONLY");
+        // Master Playlist related tags.
+        const FString ExtXMedia = TEXT("EXT-X-MEDIA");
+        const FString ExtXStreamInf = TEXT("EXT-X-STREAM-INF");
+        const FString ExtXIFrameStreamInf = TEXT("EXT-X-I-FRAME-STREAM-INF");
+        const FString ExtXSessionData = TEXT("EXT-X-SESSION-DATA");
+        const FString ExtXSessionKey = TEXT("EXT-X-SESSION-KEY");
+        // Master or Media Playlist related tags which are applied to the playlist.
+        const FString ExtXIndependentSegments = TEXT("EXT-X-INDEPENDENT-SEGMENTS");
+        const FString ExtXStart = TEXT("EXT-X-START");
+        const FString ExtXAllowCache = TEXT("EXT-X-ALLOW-CACHE");
+    }
+} // namespace Electra
+
+
+
diff --git a/Engine/Plugins/Media/ElectraPlayer/Source/ElectraPlayerRuntime/Private/Runtime/Player/HLS/Token.h b/Engine/Plugins/Media/ElectraPlayer/Source/ElectraPlayerRuntime/Private/Runtime/Player/HLS/Token.h
new file mode 100644
--- /dev/null
+++ b/Engine/Plugins/Media/ElectraPlayer/Source/ElectraPlayerRuntime/Private/Runtime/Player/HLS/Token.h
@@ -0,0 +1,32 @@
+// Copyright Epic Games, Inc. All Rights Reserved.
+
+#pragma once
+
+namespace Electra
+{
+    namespace HLSPlaylistParser
+    {
+        const TCHAR CComment = TCHAR('#');
+        const TCHAR CColumn = TCHAR(':');
+        const TCHAR CLineBreak = TCHAR('\n');
+        const TCHAR CAttributeSeparation = TCHAR(',');
+
+        typedef unsigned int Token;
+
+        const Token TUnknown = 0xFFFFFFFF;
+        const Token TComment = 1;
+        const Token TLineBreak = 1 << 1;
+        const Token TColumn = 1 << 2;
+        const Token TAttributeDeclaration = 1 << 3;
+        const Token TSpace = 1 << 4;
+        const Token TQuote = 1 << 5;
+        const Token TAttributeSeparation = 1 << 6;
+        const Token TEOF = 1 << 7;
+        const Token TIgnore = 1 << 8;
+
+        Token LookupToken(TCHAR p);
+    }
+} // namespace Electra
+
+
+
diff --git a/Engine/Plugins/Media/ElectraPlayer/Source/ElectraPlayerRuntime/Private/Runtime/Player/Manifest.h b/Engine/Plugins/Media/ElectraPlayer/Source/ElectraPlayerRuntime/Private/Runtime/Player/Manifest.h
new file mode 100644
--- /dev/null
+++ b/Engine/Plugins/Media/ElectraPlayer/Source/ElectraPlayerRuntime/Private/Runtime/Player/Manifest.h
@@ -0,0 +1,314 @@
+// Copyright Epic Games, Inc. All Rights Reserved.
+
+#pragma once
+#include "PlayerCore.h"
+#include "PlayerTime.h"
+#include "StreamTypes.h"
+#include "HTTP/HTTPManager.h"
+#include "Player/PlaybackTimeline.h"
+#include "Player/AdaptiveStreamingPlayerMetrics.h"
+#include "Player/PlayerSessionServices.h"
+#include "ParameterDictionary.h"
+
+
+
+//
+//
+//
+
+
+namespace Electra
+{
+
+	class IStreamReader;
+
+
+	struct FPlayStartPosition
+	{
+		FTimeValue		Time;
+	};
+
+
+	class IStreamSegment : private TMediaNoncopyable<IStreamSegment>, public TSharedFromThis<IStreamSegment, ESPMode::ThreadSafe>
+	{
+	public:
+		virtual ~IStreamSegment() = default;
+
+		/**
+		 * Sets a playback sequence ID. This is useful to check the ID later in asynchronously received messages
+		 * to validate the request still being valid or having become outdated.
+		 */
+		virtual void SetPlaybackSequenceID(uint32 PlaybackSequenceID) = 0;
+		/**
+		 * Returns the playback sequence ID. See SetPlaybackSequenceID()
+		 */
+		virtual uint32 GetPlaybackSequenceID() const = 0;
+
+
+		virtual EStreamType GetType() const = 0;
+
+		struct FDependentStreams
+		{
+			EStreamType		StreamType;
+		};
+		/**
+		 * Returns a list of dependent stream types.
+		 */
+		virtual void GetDependentStreams(TArray<FDependentStreams>& OutDependentStreams) const = 0;
+
+		/**
+		 * Returns a list of streams, primary or dependent, that are already at EOS in this request.
+		 */
+		virtual void GetEndedStreams(TArray<TSharedPtrTS<IStreamSegment>>& OutAlreadyEndedStreams) = 0;
+
+		//! Returns the first PTS value as indicated by the media timeline. This should correspond to the actual absolute PTS of the sample.
+		virtual FTimeValue GetFirstPTS() const = 0;
+
+		virtual int32 GetQualityIndex() const = 0;
+		virtual int32 GetBitrate() const = 0;
+
+		virtual void GetDownloadStats(Metrics::FSegmentDownloadStats& OutStats) const = 0;
+	};
+
+
+
+	class IManifest : public TMediaNoncopyable<IManifest>
+	{
+	public:
+		enum class EType
+		{
+			OnDemand,						//!< An on-demand presentation
+			Live,							//!< A live presentation
+		};
+
+		enum class ESearchType
+		{
+			Closest,				//!< Find closest match
+			After,					//!< Find match only for fragment times >= target time
+			Before,					//!< Find match only for fragment times <= target time
+			StrictlyAfter,			//!< Match must be strictly after (>). Used to locate the next segment.
+			StrictlyBefore,			//!< Match must be strictly before (<). Used to locate the previous segment.
+			Same,					//!< Match must be for the same fragment
+		};
+
+		class FResult
+		{
+		public:
+			enum class EType
+			{
+				Found,				//!< Found
+				NotFound,			//!< Not found
+				PastEOS,			//!< Time is beyond the duration
+				BeforeStart,		//!< Time is before the start time
+				TryAgainLater,		//!< Not found at the moment. Playlist load may be pending
+				NotLoaded,			//!< Not loaded (playlist has not been requested)
+			};
+			FResult(EType InType = EType::NotFound)
+				: Type(InType)
+			{
+			}
+			FResult& RetryAfterMilliseconds(int32 Milliseconds)
+			{
+				Type = EType::TryAgainLater;
+				RetryAgainAtTime = MEDIAutcTime::Current() + FTimeValue().SetFromMilliseconds(Milliseconds);
+				return *this;
+			}
+			FResult& SetErrorDetail(const FErrorDetail& InErrorDetail)
+			{
+				ErrorDetail = InErrorDetail;
+				return *this;
+			}
+			EType GetType() const
+			{
+				return Type;
+			}
+			bool IsSuccess() const
+			{
+				return GetType() == EType::Found;
+			}
+			const FTimeValue& GetRetryAgainAtTime() const
+			{
+				return RetryAgainAtTime;
+			}
+			const FErrorDetail& GetErrorDetail() const
+			{
+				return ErrorDetail;
+			}
+			static const TCHAR* GetTypeName(EType s)
+			{
+				switch (s)
+				{
+					case EType::Found:
+						return TEXT("Found");
+					case EType::NotFound:
+						return TEXT("Not found");
+					case EType::PastEOS:
+						return TEXT("Behind EOF");
+					case EType::BeforeStart:
+						return TEXT("Before start");
+					case EType::TryAgainLater:
+						return TEXT("Try again later");
+					case EType::NotLoaded:
+						return TEXT("Not loaded");
+					default:
+						return TEXT("undefined");
+				}
+			}
+		protected:
+			EType			Type;
+			FTimeValue		RetryAgainAtTime;
+			FErrorDetail	ErrorDetail;
+		};
+
+		//!
+		virtual ~IManifest() = default;
+
+
+
+		//-------------------------------------------------------------------------
+		// Presentation related functions
+		//
+		//! Returns the type of this presentation, either on-demand or live.
+		virtual EType GetPresentationType() const = 0;
+
+
+
+		//-------------------------------------------------------------------------
+		// Presentation and play time related functions.
+		//
+		//! Returns the current timeline object.
+		virtual TSharedPtrTS<IPlaybackAssetTimeline> GetTimeline() const = 0;
+
+		//! Returns the bitrate of the default stream (usually the first one specified).
+		virtual int64 GetDefaultStartingBitrate() const = 0;
+
+		//! Returns stream metadata. For period based presentations the streams can be different per period in which case the metadata of the first period is returned.
+		virtual void GetStreamMetadata(TArray<FStreamMetadata>& OutMetadata, EStreamType StreamType) const = 0;
+
+
+		//
+		virtual FTimeValue GetMinBufferTime() const = 0;
+
+
+
+		//-------------------------------------------------------------------------
+		// Stream fragment reader
+		//
+		//! Create a stream reader handler suitable for streams from this manifest.
+		virtual IStreamReader* CreateStreamReaderHandler() = 0;
+
+
+
+		/**
+		 * Interface to a playback period.
+		 */
+		class IPlayPeriod : private TMediaNoncopyable<IPlayPeriod>
+		{
+		public:
+			virtual ~IPlayPeriod() = default;
+
+			virtual void SetStreamPreferences(const FStreamPreferences& Preferences) = 0;
+
+			enum class EReadyState
+			{
+				NotReady,
+				Preparing,
+				IsReady,
+				MustReselect						//!< Player must forget this Play Period and reselect it through FindPlayPeriod(). Needed if the timeline got altered.
+			};
+			virtual EReadyState GetReadyState() = 0;
+			virtual void PrepareForPlay(const FParamDict& Options) = 0;
+
+			/**
+			 * Returns the media asset for this play period.
+			 *
+			 * @return Pointer to the media asset.
+			 */
+			virtual TSharedPtrTS<ITimelineMediaAsset> GetMediaAsset() const = 0;
+
+			/**
+			 * Selects a specific stream to be used for all next segment downloads.
+			 *
+			 * @param AdaptationSet
+			 * @param Representation
+			 * @param PreferredCDN
+			 *               Preferred CDN to get the segment from (retrieve through GetMediaAsset() for now)
+			 */
+			virtual void SelectStream(const TSharedPtrTS<IPlaybackAssetAdaptationSet>& AdaptationSet, const TSharedPtrTS<IPlaybackAssetRepresentation>& Representation, const FString& PreferredCDN) = 0;
+
+			/**
+			 * Sets up a starting segment request to begin playback at the specified time.
+			 * The streams selected through SelectStream() will be used.
+			 *
+			 * @param OutSegment
+			 * @param StartPosition
+			 * @param SearchType
+			 *
+			 * @return
+			 */
+			virtual FResult GetStartingSegment(TSharedPtrTS<IStreamSegment>& OutSegment, const FPlayStartPosition& StartPosition, ESearchType SearchType) = 0;
+
+
+			/**
+			 * Sets up a starting segment request to loop playback to.
+			 * The streams selected through SelectStream() will be used.
+			 *
+			 * @param OutSegment
+			 * @param InOutLoopState
+			 * @param InFinishedSegments
+			 * @param StartPosition
+			 * @param SearchType
+			 *
+			 * @return
+			 */
+			virtual FResult GetLoopingSegment(TSharedPtrTS<IStreamSegment>& OutSegment, FPlayerLoopState& InOutLoopState, const TMultiMap<EStreamType, TSharedPtrTS<IStreamSegment>>& InFinishedSegments, const FPlayStartPosition& StartPosition, ESearchType SearchType) = 0;
+
+			/**
+			 * Gets the segment request for the segment following the specified earlier request.
+			 *
+			 * @param OutSegment
+			 * @param CurrentSegment
+			 * @param Options
+			 *
+			 * @return
+			 */
+			virtual FResult GetNextSegment(TSharedPtrTS<IStreamSegment>& OutSegment, TSharedPtrTS<const IStreamSegment> CurrentSegment, const FParamDict& Options) = 0;
+
+			/**
+			 * Gets the segment request for the same segment on a different quality level or CDN.
+			 *
+			 * @param OutSegment
+			 * @param CurrentSegment
+			 * @param Options
+			 *
+			 * @return
+			 */
+			virtual FResult GetRetrySegment(TSharedPtrTS<IStreamSegment>& OutSegment, TSharedPtrTS<const IStreamSegment> CurrentSegment, const FParamDict& Options) = 0;
+
+			struct FSegmentInformation
+			{
+				FTimeValue	Duration;
+				int64		ByteSize;
+			};
+
+			/**
+			 * Obtains information on the stream segmentation of a particular stream starting at a given current reference segment (optional, if not given returns suitable default values).
+			 *
+			 * @param OutSegmentInformation
+			 * @param OutAverageSegmentDuration
+			 * @param CurrentSegment
+			 * @param LookAheadTime
+			 * @param AdaptationSet
+			 * @param Representation
+			 */
+			virtual void GetSegmentInformation(TArray<FSegmentInformation>& OutSegmentInformation, FTimeValue& OutAverageSegmentDuration, TSharedPtrTS<const IStreamSegment> CurrentSegment, const FTimeValue& LookAheadTime, const TSharedPtrTS<IPlaybackAssetAdaptationSet>& AdaptationSet, const TSharedPtrTS<IPlaybackAssetRepresentation>& Representation) = 0;
+		};
+
+		virtual FResult FindPlayPeriod(TSharedPtrTS<IPlayPeriod>& OutPlayPeriod, const FPlayStartPosition& StartPosition, ESearchType SearchType) = 0;
+	};
+
+
+
+} // namespace Electra
+
+
diff --git a/Engine/Plugins/Media/ElectraPlayer/Source/ElectraPlayerRuntime/Private/Runtime/Player/PlaybackTimeline.h b/Engine/Plugins/Media/ElectraPlayer/Source/ElectraPlayerRuntime/Private/Runtime/Player/PlaybackTimeline.h
new file mode 100644
--- /dev/null
+++ b/Engine/Plugins/Media/ElectraPlayer/Source/ElectraPlayerRuntime/Private/Runtime/Player/PlaybackTimeline.h
@@ -0,0 +1,194 @@
+// Copyright Epic Games, Inc. All Rights Reserved.
+
+#pragma once
+#include "PlayerCore.h"
+#include "PlayerTime.h"
+#include "StreamTypes.h"
+
+
+namespace Electra
+{
+	class IPlaybackAssetRepresentation
+	{
+	public:
+		virtual ~IPlaybackAssetRepresentation() = default;
+
+		virtual FString GetUniqueIdentifier() const = 0;
+
+		virtual const FStreamCodecInformation& GetCodecInformation() const = 0;
+
+		virtual int32 GetBitrate() const = 0;
+
+		virtual FString GetCDN() const = 0;
+	};
+
+	class IPlaybackAssetAdaptationSet
+	{
+	public:
+		virtual ~IPlaybackAssetAdaptationSet() = default;
+
+		/**
+		 * Returns a unique identifier for this adaptation set.
+		 * This may be a value from the manifest or an internally generated one.
+		 * NOTE: The identifier is unique only within the owning media asset!
+		 *
+		 * @return Unique identifier for this adaptation set with the owning media asset.
+		 */
+		virtual FString GetUniqueIdentifier() const = 0;
+
+		virtual FString GetListOfCodecs() const = 0;
+
+		virtual FString GetLanguage() const = 0;
+
+		virtual int32 GetNumberOfRepresentations() const = 0;
+
+		virtual TSharedPtrTS<IPlaybackAssetRepresentation> GetRepresentationByIndex(int32 RepresentationIndex) const = 0;
+
+		virtual TSharedPtrTS<IPlaybackAssetRepresentation> GetRepresentationByUniqueIdentifier(const FString& UniqueIdentifier) const = 0;
+	};
+
+
+	class ITimelineMediaAsset
+	{
+	public:
+		virtual ~ITimelineMediaAsset() = default;
+
+		/**
+		 * Returns the time range of this asset on the playback timeline.
+		 * The timeline anchor time is included in this range.
+		 *
+		 * @return Time range of this asset on the playback timeline.
+		 */
+		virtual FTimeRange GetTimeRange() const = 0;
+
+		/**
+		 * Returns the duration of this asset, which is typically the difference between
+		 * the end and start value of GetTimeRange() unless it is last asset of a Live
+		 * presentation timeline for which the duration is infinite.
+		 *
+		 * @return Duration of this asset.
+		 */
+		virtual FTimeValue GetDuration() const = 0;
+
+		/**
+		 * Returns the asset identifier, if present in the manifest, for this asset.
+		 *
+		 * @return Asset identifier as it appears in the manifest.
+		 */
+		virtual FString GetAssetIdentifier() const = 0;
+
+
+		/**
+		 * Returns a unique identifier for this asset.
+		 * This may be a value from the manifest or an internally generated one.
+		 *
+		 * @return Unique identifier for this asset.
+		 */
+		virtual FString GetUniqueIdentifier() const = 0;
+
+
+		/**
+		 * Returns the number of "adaptation sets" for a particular type of stream.
+		 * An adaptation set is defined as a group of streams representing the same content at different
+		 * quality levels to dynamically switch between.
+		 * Streams are further grouped by language and codec (there will not be a mix of either in a
+		 * single adaptation set).
+		 *
+		 * @param OfStreamType
+		 *               Type of elementary stream for which to the the number of available adaptation sets.
+		 *
+		 * @return Number of adaptation sets for the specified type.
+		 */
+		virtual int32 GetNumberOfAdaptationSets(EStreamType OfStreamType) const = 0;
+
+		/**
+		 * Returns the media asset's adaptation set by index.
+		 *
+		 * @param OfStreamType
+		 *               Type of elementary stream for which to get the adaptation set.
+		 * @param AdaptationSetIndex
+		 *               Index of the adaptation set to return.
+		 *
+		 * @return Shared pointer to the requested adaptation set.
+		 */
+		virtual TSharedPtrTS<IPlaybackAssetAdaptationSet> GetAdaptationSetByTypeAndIndex(EStreamType OfStreamType, int32 AdaptationSetIndex) const = 0;
+	};
+
+
+	/**
+	 * Interface class to the media playback asset timeline.
+	 * This is a snapshot of the current timeline which may change periodically.
+	 */
+	class IPlaybackAssetTimeline
+	{
+	public:
+		virtual ~IPlaybackAssetTimeline() = default;
+
+		/**
+		 * Returns the time value the timeline is anchored at.
+		 * Usually this is a UTC timestamp at which the presentation began but it could
+		 * literally be anything.
+		 * All time values are absolute and have this anchor added in.
+		 *
+		 * @return The time this presentation timeline is anchored at.
+		 */
+		virtual FTimeValue GetAnchorTime() const = 0;
+
+		/**
+		 * Returns the total time range of the timeline, including the end time of the last sample
+		 * to which the player will not be able to seek to.
+		 *
+		 * @return Time range of assets on the timeline.
+		 */
+		virtual FTimeRange GetTotalTimeRange() const = 0;
+
+
+		/**
+		 * Returns the seekable range of the timeline, which is a subset of the total
+		 * time range. Playback can start at any point in the seekable range with
+		 * additional constraints imposed by the format.
+		 *
+		 * @return Time range in which playback can start.
+		 */
+		virtual FTimeRange GetSeekableTimeRange() const = 0;
+
+		/**
+		 * Returns the timestamps of the segments from the video or audio track (if no video is present).
+		 * Segments are required to start with a keyframe and can thus be used to start playback with.
+		 */
+		virtual void GetSeekablePositions(TArray<FTimespan>& OutPositions) const = 0;
+
+		/**
+		 * Returns the duration of the assets on the timeline.
+		 * Typically this is the difference of the end and start values of the total time range
+		 * unless this is a Live presentation for which the duration will be set to infinite.
+		 *
+		 * @return Duration of the timeline.
+		 */
+		virtual FTimeValue GetDuration() const = 0;
+
+		/**
+		 * Returns the number of media assets on the playback timeline.
+		 *
+		 * @return Number of media assets on the playback timeline.
+		 */
+		virtual int32 GetNumberOfMediaAssets() const = 0;
+
+		/**
+		 * Returns a media asset from the timeline by its index.
+		 * Please note that media assets are not sorted by their time range on the timeline.
+		 * This is done on purpose to keep the order of assets as they appear in the manifest.
+		 *
+		 * @param MediaAssetIndex
+		 *               Index (0 to GetNumberOfMediaAssets()-1) of the asset to get.
+		 *
+		 * @return Shared pointer to the requested media asset.
+		 */
+		virtual TSharedPtrTS<ITimelineMediaAsset> GetMediaAssetByIndex(int32 MediaAssetIndex) const = 0;
+	};
+
+
+
+} // namespace Electra
+
+
diff --git a/Engine/Plugins/Media/ElectraPlayer/Source/ElectraPlayerRuntime/Private/Runtime/Player/PlayerStreamReader.h b/Engine/Plugins/Media/ElectraPlayer/Source/ElectraPlayerRuntime/Private/Runtime/Player/PlayerStreamReader.h
new file mode 100644
--- /dev/null
+++ b/Engine/Plugins/Media/ElectraPlayer/Source/ElectraPlayerRuntime/Private/Runtime/Player/PlayerStreamReader.h
@@ -0,0 +1,96 @@
+// Copyright Epic Games, Inc. All Rights Reserved.
+
+#pragma once
+
+#include "PlayerCore.h"
+#include "StreamTypes.h"
+#include "Player/Manifest.h"
+#include "Player/PlayerSessionServices.h"
+
+
+namespace Electra
+{
+	class IAccessUnitMemoryProvider;
+	struct FAccessUnit;
+	struct FStreamSourceInfo;
+
+
+	/**
+	 * Base class for reading fragments from a stream.
+	 * Specialized versions for particular container formats shall derive from this class.
+	**/
+	class IStreamReader : private TMediaNoncopyable<IStreamReader>
+	{
+	public:
+		virtual ~IStreamReader() = default;
+
+		class StreamReaderEventListener
+		{
+		public:
+			virtual ~StreamReaderEventListener() = default;
+			virtual bool OnFragmentAccessUnitReceived(FAccessUnit* pAccessUnit) = 0;
+			virtual void OnFragmentReachedEOS(EStreamType InStreamType, TSharedPtr<const FStreamSourceInfo, ESPMode::ThreadSafe> InStreamSourceInfo) = 0;
+			virtual void OnFragmentOpen(TSharedPtrTS<IStreamSegment> pRequest) = 0;
+			virtual void OnFragmentClose(TSharedPtrTS<IStreamSegment> pRequest) = 0;
+		};
+
+		struct ReaderConfiguration
+		{
+			ReaderConfiguration()
+			{
+				ThreadParam.Priority = TPri_Normal;
+				ThreadParam.StackSize = 65536;
+				ThreadParam.CoreAffinity = -1;
+
+				// No limit on download duration
+				MaxDLDurationScale = 0.0;
+			}
+
+			FMediaRunnable::Param							ThreadParam;
+			double											MaxDLDurationScale;
+		};
+
+		struct CreateParam
+		{
+			CreateParam()
+			{
+				MemoryProvider = nullptr;
+				EventListener = nullptr;
+				PlayerSessionService = nullptr;
+			}
+
+			ReaderConfiguration		   ReaderConfig;
+			IAccessUnitMemoryProvider* MemoryProvider;
+			StreamReaderEventListener* EventListener;
+
+			FParamDict					Options;
+			IPlayerSessionServices*		PlayerSessionService;
+		};
+
+		//! Creates the instance internally based on the creation parameters.
+		virtual UEMediaError Create(IPlayerSessionServices* PlayerSessionService, const CreateParam& createParam) = 0;
+
+		//! Cancels all pending requests and stops processing. The instance can no longer be used after that!
+		virtual void Close() = 0;
+
+		enum class EAddResult
+		{
+			Added,
+			TryAgainLater,
+			Error
+		};
+		//! Adds a request to read from a stream
+		virtual EAddResult AddRequest(uint32 CurrentPlaybackSequenceID, TSharedPtrTS<IStreamSegment> Request) = 0;
+
+		//! Cancels all pending requests.
+		virtual void CancelRequests() = 0;
+
+		//! Pauses all pending requests.
+		virtual void PauseDownload() = 0;
+		//! Resumes all pending requests.
+		virtual void ResumeDownload() = 0;
+	};
+
+} // namespace Electra
+
+
