ElectraPlayer: RELEASE. Moved from NFL to public. Note that this is done via a Delete and "Add New" operation as discussed.

diff --git a/Engine/Plugins/Media/ElectraPlayer/Source/ElectraPlayerRuntime/Private/Runtime/Player/PlaylistReader.h b/Engine/Plugins/Media/ElectraPlayer/Source/ElectraPlayerRuntime/Private/Runtime/Player/PlaylistReader.h
new file mode 100644
--- /dev/null
+++ b/Engine/Plugins/Media/ElectraPlayer/Source/ElectraPlayerRuntime/Private/Runtime/Player/PlaylistReader.h
@@ -0,0 +1,171 @@
+// Copyright Epic Games, Inc. All Rights Reserved.
+
+#pragma once
+
+#include "PlayerCore.h"
+#include "Player/PlayerSessionServices.h"
+#include "StreamTypes.h"
+#include "Player/Playlist.h"
+#include "ParameterDictionary.h"
+
+
+
+namespace Electra
+{
+	// Forward declarations
+	class IManifest;
+
+	class IPlaylistReader
+	{
+	public:
+		virtual ~IPlaylistReader() = default;
+
+		/**
+		 * Returns the type of the playlist (eg. "hls", "dash", etc.)
+		 */
+		virtual const FString& GetPlaylistType() const = 0;
+
+		/**
+		 * Loads and parses the playlist.
+		 *
+		 * @param URL     URL of the playlist to load
+		 * @param Preferences
+		 *                User preferences (for initial stream selection)
+		 * @param Options Options for the playlist reader and parser specific to the format.
+		 *                See specifics in the format's implementation.
+		 */
+		virtual void LoadAndParse(const FString& URL, const FStreamPreferences& Preferences, const FParamDict& Options) = 0;
+
+		/**
+		 * Returns the URL from which the playlist was loaded (or supposed to be loaded).
+		 *
+		 * @return The playlist URL
+		 */
+		virtual FString GetURL() const = 0;
+
+
+		/**
+		 * Returns the manifest interface to access the playlist in a uniform way.
+		 *
+		 * @return
+		 */
+		virtual TSharedPtrTS<IManifest> GetManifest() = 0;
+
+
+		class PlaylistDownloadMessage : public IPlayerMessage
+		{
+		public:
+			static TSharedPtrTS<IPlayerMessage> Create(const HTTP::FConnectionInfo* ConnectionInfo, Playlist::EListType InListType, Playlist::ELoadType InLoadType)
+			{
+				TSharedPtrTS<PlaylistDownloadMessage> p(new PlaylistDownloadMessage(ConnectionInfo, InListType, InLoadType));
+				return p;
+			}
+
+			static const FString& Type()
+			{
+				static FString TypeName("PlaylistDownload");
+				return TypeName;
+			}
+
+			virtual const FString& GetType() const
+			{
+				return Type();
+			}
+
+			Playlist::EListType GetListType() const
+			{
+				return ListType;
+			}
+
+			Playlist::ELoadType GetLoadType() const
+			{
+				return LoadType;
+			}
+
+			const HTTP::FConnectionInfo& GetConnectionInfo() const
+			{
+				return ConnectionInfo;
+			}
+
+		private:
+			PlaylistDownloadMessage(const HTTP::FConnectionInfo* InConnectionInfo, Playlist::EListType InListType, Playlist::ELoadType InLoadType)
+				: ListType(InListType)
+				, LoadType(InLoadType)
+			{
+				if (InConnectionInfo)
+				{
+					// Have to make a dedicated copy of the connection info in order to get a copy of the retry info at this point in time.
+					ConnectionInfo.CopyFrom(*InConnectionInfo);
+					int32 y = 0;
+				}
+			}
+			HTTP::FConnectionInfo	ConnectionInfo;
+			Playlist::EListType		ListType;
+			Playlist::ELoadType		LoadType;
+		};
+
+
+		class PlaylistLoadedMessage : public IPlayerMessage
+		{
+		public:
+			static TSharedPtrTS<IPlayerMessage> Create(const FErrorDetail& PlayerResult, const HTTP::FConnectionInfo* ConnectionInfo, Playlist::EListType InListType, Playlist::ELoadType InLoadType)
+			{
+				TSharedPtrTS<PlaylistLoadedMessage> p(new PlaylistLoadedMessage(PlayerResult, ConnectionInfo, InListType, InLoadType));
+				return p;
+			}
+
+			static const FString& Type()
+			{
+				static FString TypeName("PlaylistLoaded");
+				return TypeName;
+			}
+
+			virtual const FString& GetType() const
+			{
+				return Type();
+			}
+
+			const FErrorDetail& GetResult() const
+			{
+				return Result;
+			}
+
+			Playlist::EListType GetListType() const
+			{
+				return ListType;
+			}
+
+			Playlist::ELoadType GetLoadType() const
+			{
+				return LoadType;
+			}
+
+			const HTTP::FConnectionInfo& GetConnectionInfo() const
+			{
+				return ConnectionInfo;
+			}
+
+		private:
+			PlaylistLoadedMessage(const FErrorDetail& PlayerResult, const HTTP::FConnectionInfo* InConnectionInfo, Playlist::EListType InListType, Playlist::ELoadType InLoadType)
+				: Result(PlayerResult)
+				, ListType(InListType)
+				, LoadType(InLoadType)
+			{
+				if (InConnectionInfo)
+				{
+					// Have to make a dedicated copy of the connection info in order to get a copy of the retry info at this point in time.
+					ConnectionInfo.CopyFrom(*InConnectionInfo);
+				}
+			}
+			HTTP::FConnectionInfo	ConnectionInfo;
+			FErrorDetail			Result;
+			Playlist::EListType		ListType;
+			Playlist::ELoadType		LoadType;
+		};
+
+	};
+
+
+} // namespace Electra
+
+
diff --git a/Engine/Plugins/Media/ElectraPlayer/Source/ElectraPlayerRuntime/Private/Runtime/Player/mp4/ManifestMP4.cpp b/Engine/Plugins/Media/ElectraPlayer/Source/ElectraPlayerRuntime/Private/Runtime/Player/mp4/ManifestMP4.cpp
new file mode 100644
--- /dev/null
+++ b/Engine/Plugins/Media/ElectraPlayer/Source/ElectraPlayerRuntime/Private/Runtime/Player/mp4/ManifestMP4.cpp
@@ -0,0 +1,838 @@
+// Copyright Epic Games, Inc. All Rights Reserved.
+
+#include "PlayerCore.h"
+#include "Player/Manifest.h"
+#include "HTTP/HTTPManager.h"
+#include "Demuxer/ParserISO14496-12.h"
+#include "SynchronizedClock.h"
+#include "PlaylistReaderMP4.h"
+#include "Utilities/HashFunctions.h"
+#include "Utilities/TimeUtilities.h"
+#include "HAL/LowLevelMemTracker.h"
+#include "ElectraPlayerPrivate.h"
+#include "Player/AdaptiveStreamingPlayerResourceRequest.h"
+#include "Player/mp4/StreamReaderMP4.h"
+#include "Player/mp4/ManifestMP4.h"
+
+
+#define ERRCODE_MANIFEST_MP4_STARTSEGMENT_NOT_FOUND						1
+
+
+namespace Electra
+{
+
+//-----------------------------------------------------------------------------
+/**
+ * CTOR
+ */
+FManifestMP4Internal::FManifestMP4Internal(IPlayerSessionServices* InPlayerSessionServices)
+	: PlayerSessionServices(InPlayerSessionServices)
+{
+}
+
+
+//-----------------------------------------------------------------------------
+/**
+ * DTOR
+ */
+FManifestMP4Internal::~FManifestMP4Internal()
+{
+}
+
+
+//-----------------------------------------------------------------------------
+/**
+ * Builds the internal manifest from the mp4's moov box.
+ *
+ * @param MP4Parser
+ * @param URL
+ * @param InConnectionInfo
+ *
+ * @return
+ */
+FErrorDetail FManifestMP4Internal::Build(TSharedPtrTS<IParserISO14496_12> MP4Parser, const FString& URL, const HTTP::FConnectionInfo& InConnectionInfo)
+{
+	ConnectionInfo = InConnectionInfo;
+	MediaAsset = MakeSharedTS<FTimelineAssetMP4>();
+	return MediaAsset->Build(PlayerSessionServices, MP4Parser, URL);
+}
+
+
+//-----------------------------------------------------------------------------
+/**
+ * Logs a message.
+ *
+ * @param Level
+ * @param Message
+ */
+void FManifestMP4Internal::LogMessage(IInfoLog::ELevel Level, const FString& Message)
+{
+	if (PlayerSessionServices)
+	{
+		PlayerSessionServices->PostLog(Facility::EFacility::MP4Playlist, Level, Message);
+	}
+}
+
+
+//-----------------------------------------------------------------------------
+/**
+ * Returns the type of presentation.
+ * For a single mp4 file this is always VoD.
+ *
+ * @return
+ */
+IManifest::EType FManifestMP4Internal::GetPresentationType() const
+{
+	return IManifest::EType::OnDemand;
+}
+
+
+//-----------------------------------------------------------------------------
+/**
+ * Returns the media timeline object for this asset.
+ *
+ * @return
+ */
+TSharedPtrTS<IPlaybackAssetTimeline> FManifestMP4Internal::GetTimeline() const
+{
+	// Since an mp4 file is fixed and will not change the timeline is fixed.
+	// That's why we are inheriting from IPlaybackAssetTimeline and return ourselves.
+	TSharedPtrTS<const FManifestMP4Internal> This = SharedThis(this);
+	return StaticCastSharedPtr<IPlaybackAssetTimeline>(ConstCastSharedPtr<FManifestMP4Internal>(This));
+}
+
+
+//-----------------------------------------------------------------------------
+/**
+ * Returns the starting bitrate.
+ *
+ * This is merely informational and not strictly required.
+ * If fetching of the moov box provided us with the total size of the mp4 file
+ * we will use that divided by the duration.
+ *
+ * @return
+ */
+int64 FManifestMP4Internal::GetDefaultStartingBitrate() const
+{
+	FTimeValue dur = GetDuration();
+	if (ConnectionInfo.ContentLength > 0 && dur.IsValid() && dur > FTimeValue::GetZero())
+	{
+		return (int64)( ConnectionInfo.ContentLength * 8 / dur.GetAsSeconds() );
+	}
+	return 0;
+}
+
+
+//-----------------------------------------------------------------------------
+/**
+ * Returns stream metadata.
+ *
+ * @param OutMetadata
+ * @param StreamType
+ */
+void FManifestMP4Internal::GetStreamMetadata(TArray<FStreamMetadata>& OutMetadata, EStreamType StreamType) const
+{
+	if (MediaAsset.IsValid())
+	{
+		for(int32 i=0, iMax = MediaAsset->GetNumberOfAdaptationSets(StreamType); i<iMax; ++i)
+		{
+			TSharedPtrTS<IPlaybackAssetAdaptationSet> AdaptSet = MediaAsset->GetAdaptationSetByTypeAndIndex(StreamType, i);
+			if (AdaptSet.IsValid())
+			{
+				for(int32 j=0, jMax=AdaptSet->GetNumberOfRepresentations(); j<jMax; ++j)
+				{
+					TSharedPtrTS<FRepresentationMP4> Repr = StaticCastSharedPtr<FRepresentationMP4>(AdaptSet->GetRepresentationByIndex(j));
+					if (Repr.IsValid())
+					{
+						FStreamMetadata& meta = OutMetadata.AddDefaulted_GetRef();
+						meta.CodecInformation = Repr->GetCodecInformation();
+						LexFromString(meta.StreamUniqueID, *Repr->GetUniqueIdentifier());
+						meta.PlaylistID 	  = Repr->GetCDN();
+						meta.Bandwidth  	  = Repr->GetBitrate();
+						meta.LanguageCode     = AdaptSet->GetLanguage();
+					}
+				}
+			}
+		}
+	}
+}
+
+
+//-----------------------------------------------------------------------------
+/**
+ * Returns the minimum duration of content that must be buffered up before playback
+ * will begin. This is an arbitrary choice that could be controlled by a 'pdin' box.
+ *
+ * @return
+ */
+FTimeValue FManifestMP4Internal::GetMinBufferTime() const
+{
+	// NOTE: This could come from a 'pdin' (progressive download information) box, but those are rarely, if ever, set by any tool.
+	return FTimeValue().SetFromSeconds(2.0);
+}
+
+
+//-----------------------------------------------------------------------------
+/**
+ * Creates an instance of a stream reader to stream from the mp4 file.
+ *
+ * @return
+ */
+IStreamReader *FManifestMP4Internal::CreateStreamReaderHandler()
+{
+	return new FStreamReaderMP4;
+}
+
+
+//-----------------------------------------------------------------------------
+/**
+ * Returns the playback period for the given time.
+ *
+ * @param OutPlayPeriod
+ * @param StartPosition
+ * @param SearchType
+ *
+ * @return
+ */
+IManifest::FResult FManifestMP4Internal::FindPlayPeriod(TSharedPtrTS<IPlayPeriod>& OutPlayPeriod, const FPlayStartPosition& StartPosition, ESearchType SearchType)
+{
+	// FIXME: We could however check if the start position falls into the duration of the asset. Not sure why it wouldn't or why we would want to do that.
+	OutPlayPeriod = MakeSharedTS<FPlayPeriodMP4>(MediaAsset);
+	return IManifest::FResult(IManifest::FResult::EType::Found);
+}
+
+
+
+
+
+//-----------------------------------------------------------------------------
+/**
+ * Constructs a playback period.
+ *
+ * @param InMediaAsset
+ */
+FManifestMP4Internal::FPlayPeriodMP4::FPlayPeriodMP4(TSharedPtrTS<FManifestMP4Internal::FTimelineAssetMP4> InMediaAsset)
+	: MediaAsset(InMediaAsset)
+	, bIsReady(false)
+{
+}
+
+
+//-----------------------------------------------------------------------------
+/**
+ * Destroys a playback period.
+ */
+FManifestMP4Internal::FPlayPeriodMP4::~FPlayPeriodMP4()
+{
+}
+
+
+//-----------------------------------------------------------------------------
+/**
+ * Sets stream playback preferences for this playback period.
+ *
+ * @param InPreferences
+ */
+void FManifestMP4Internal::FPlayPeriodMP4::SetStreamPreferences(const FStreamPreferences& InPreferences)
+{
+	Preferences = InPreferences;
+}
+
+
+//-----------------------------------------------------------------------------
+/**
+ * Returns the ready state of this playback period.
+ *
+ * @return
+ */
+IManifest::IPlayPeriod::EReadyState FManifestMP4Internal::FPlayPeriodMP4::GetReadyState()
+{
+	return bIsReady ? IManifest::IPlayPeriod::EReadyState::IsReady : IManifest::IPlayPeriod::EReadyState::NotReady;
+}
+
+
+//-----------------------------------------------------------------------------
+/**
+ * Prepares the playback period for playback.
+ * With an mp4 file we are actually always ready for playback, but we say we're not
+ * one time to get here with any possible options.
+ *
+ * @param InOptions
+ */
+void FManifestMP4Internal::FPlayPeriodMP4::PrepareForPlay(const FParamDict& InOptions)
+{
+	Options = InOptions;
+	bIsReady = true;
+}
+
+
+//-----------------------------------------------------------------------------
+/**
+ * Returns the timeline media asset. We have a weak pointer to it only to
+ * prevent any cyclic locks, so we need to lock it first.
+ *
+ * @return
+ */
+TSharedPtrTS<ITimelineMediaAsset> FManifestMP4Internal::FPlayPeriodMP4::GetMediaAsset() const
+{
+	TSharedPtrTS<ITimelineMediaAsset> ma = MediaAsset.Pin();
+	return ma;
+}
+
+
+//-----------------------------------------------------------------------------
+/**
+ * Selects a particular stream (== internal track ID) for playback.
+ *
+ * @param AdaptationSet
+ * @param Representation
+ * @param PreferredCDN
+ */
+void FManifestMP4Internal::FPlayPeriodMP4::SelectStream(const TSharedPtrTS<IPlaybackAssetAdaptationSet>& AdaptationSet, const TSharedPtrTS<IPlaybackAssetRepresentation>& Representation, const FString& PreferredCDN)
+{
+	// Presently this method is only called by the ABR to switch between quality levels or CDNs.
+	// Since a single mp4 doesn't have different quality levels (technically it could, but we are concerning ourselves only with different bitrates and that doesn't apply since we are streaming
+	// the single file sequentially and selecting a different stream would not save any bandwidth so we don't bother) we ignore this for now.
+
+	// This may need an implementation when switching between different languages though.
+	// .....
+}
+
+
+//-----------------------------------------------------------------------------
+/**
+ * Creates the starting segment request to start playback with.
+ *
+ * @param OutSegment
+ * @param StartPosition
+ * @param SearchType
+ *
+ * @return
+ */
+IManifest::FResult FManifestMP4Internal::FPlayPeriodMP4::GetStartingSegment(TSharedPtrTS<IStreamSegment>& OutSegment, const FPlayStartPosition& StartPosition, ESearchType SearchType)
+{
+	TSharedPtrTS<FTimelineAssetMP4> ma = MediaAsset.Pin();
+	return ma.IsValid() ? ma->GetStartingSegment(OutSegment, StartPosition, SearchType, -1) : IManifest::FResult(IManifest::FResult::EType::NotFound);
+}
+
+
+//-----------------------------------------------------------------------------
+/**
+ * Sets up a starting segment request to loop playback to.
+ * The streams selected through SelectStream() will be used.
+ *
+ * @param OutSegment
+ * @param InOutLoopState
+ * @param InFinishedSegments
+ * @param StartPosition
+ * @param SearchType
+ *
+ * @return
+ */
+IManifest::FResult FManifestMP4Internal::FPlayPeriodMP4::GetLoopingSegment(TSharedPtrTS<IStreamSegment>& OutSegment, FPlayerLoopState& InOutLoopState, const TMultiMap<EStreamType, TSharedPtrTS<IStreamSegment>>& InFinishedSegments, const FPlayStartPosition& StartPosition, ESearchType SearchType)
+{
+	TSharedPtrTS<FTimelineAssetMP4> ma = MediaAsset.Pin();
+	return ma.IsValid() ? ma->GetLoopingSegment(OutSegment, InOutLoopState, InFinishedSegments, StartPosition, SearchType) : IManifest::FResult(IManifest::FResult::EType::NotFound);
+}
+
+
+//-----------------------------------------------------------------------------
+/**
+ * Creates the next segment request.
+ *
+ * @param OutSegment
+ * @param CurrentSegment
+ * @param InOptions
+ *
+ * @return
+ */
+IManifest::FResult FManifestMP4Internal::FPlayPeriodMP4::GetNextSegment(TSharedPtrTS<IStreamSegment>& OutSegment, TSharedPtrTS<const IStreamSegment> CurrentSegment, const FParamDict& InOptions)
+{
+	TSharedPtrTS<FTimelineAssetMP4> ma = MediaAsset.Pin();
+	return ma.IsValid() ? ma->GetNextSegment(OutSegment, CurrentSegment, InOptions) : IManifest::FResult(IManifest::FResult::EType::NotFound);
+}
+
+
+//-----------------------------------------------------------------------------
+/**
+ * Creates a segment retry request.
+ *
+ * @param OutSegment
+ * @param CurrentSegment
+ * @param InOptions
+ *
+ * @return
+ */
+IManifest::FResult FManifestMP4Internal::FPlayPeriodMP4::GetRetrySegment(TSharedPtrTS<IStreamSegment>& OutSegment, TSharedPtrTS<const IStreamSegment> CurrentSegment, const FParamDict& InOptions)
+{
+	TSharedPtrTS<FTimelineAssetMP4> ma = MediaAsset.Pin();
+	return ma.IsValid() ? ma->GetRetrySegment(OutSegment, CurrentSegment, InOptions) : IManifest::FResult(IManifest::FResult::EType::NotFound);
+}
+
+
+//-----------------------------------------------------------------------------
+/**
+ * Returns segment information for the next n segments.
+ *
+ * @param OutSegmentInformation
+ * @param OutAverageSegmentDuration
+ * @param CurrentSegment
+ * @param LookAheadTime
+ * @param AdaptationSet
+ * @param Representation
+ */
+void FManifestMP4Internal::FPlayPeriodMP4::GetSegmentInformation(TArray<FSegmentInformation>& OutSegmentInformation, FTimeValue& OutAverageSegmentDuration, TSharedPtrTS<const IStreamSegment> CurrentSegment, const FTimeValue& LookAheadTime, const TSharedPtrTS<IPlaybackAssetAdaptationSet>& AdaptationSet, const TSharedPtrTS<IPlaybackAssetRepresentation>& Representation)
+{
+	TSharedPtrTS<FTimelineAssetMP4> ma = MediaAsset.Pin();
+	if (ma.IsValid())
+	{
+		ma->GetSegmentInformation(OutSegmentInformation, OutAverageSegmentDuration, CurrentSegment, LookAheadTime, AdaptationSet, Representation);
+	}
+}
+
+
+
+
+
+//-----------------------------------------------------------------------------
+/**
+ * Builds the timeline asset.
+ *
+ * @param InPlayerSessionServices
+ * @param MP4Parser
+ * @param URL
+ *
+ * @return
+ */
+FErrorDetail FManifestMP4Internal::FTimelineAssetMP4::Build(IPlayerSessionServices* InPlayerSessionServices, TSharedPtrTS<IParserISO14496_12> MP4Parser, const FString& URL)
+{
+	PlayerSessionServices = InPlayerSessionServices;
+	MediaURL = URL;
+	// Go over the supported tracks and create an internal manifest-like structure for the player to work with.
+	for(int32 nTrack=0,nMaxTrack=MP4Parser->GetNumberOfTracks(); nTrack < nMaxTrack; ++nTrack)
+	{
+		const IParserISO14496_12::ITrack* Track = MP4Parser->GetTrackByIndex(nTrack);
+		if (Track)
+		{
+			// In an mp4 file we treat every track as a single adaptation set with one representation only.
+			// That's because by definition an adaptation set contains the same content at different bitrates and resolutions, but
+			// the type, language and codec has to be the same.
+			FErrorDetail err;
+			TSharedPtrTS<FAdaptationSetMP4> AdaptationSet = MakeSharedTS<FAdaptationSetMP4>();
+			err = AdaptationSet->CreateFrom(Track, URL);
+			if (err.IsOK())
+			{
+				// Add this track to the proper category.
+				switch(Track->GetCodecInformation().GetStreamType())
+				{
+					case EStreamType::Video:
+						VideoAdaptationSets.Add(AdaptationSet);
+						break;
+					case EStreamType::Audio:
+						AudioAdaptationSets.Add(AdaptationSet);
+						break;
+					default:
+						break;
+				}
+			}
+			else
+			{
+				return err;
+			}
+		}
+	}
+
+// FIXME: fragmented mp4's with a sidx and moof boxes!
+
+	// Hold on to the parsed MOOV box for future reference.
+	MoovBoxParser = MP4Parser;
+
+	return FErrorDetail();
+}
+
+
+void FManifestMP4Internal::FTimelineAssetMP4::LogMessage(IInfoLog::ELevel Level, const FString& Message)
+{
+	if (PlayerSessionServices)
+	{
+		PlayerSessionServices->PostLog(Facility::EFacility::MP4Playlist, Level, Message);
+	}
+}
+
+
+IManifest::FResult FManifestMP4Internal::FTimelineAssetMP4::GetStartingSegment(TSharedPtrTS<IStreamSegment>& OutSegment, const FPlayStartPosition& StartPosition, ESearchType SearchType, int64 AtAbsoluteFilePos)
+{
+// TODO: If there is a SIDX box we will look in there.
+
+	// Look at the actual tracks. If there is video search there first for a keyframe/IDR frame.
+	if (VideoAdaptationSets.Num())
+	{
+// TODO: use the selected track index if there are several tracks and we have a selected one.
+		TSharedPtrTS<FRepresentationMP4> Repr = StaticCastSharedPtr<FRepresentationMP4>(VideoAdaptationSets[0]->GetRepresentationByIndex(0));
+		if (Repr.IsValid())
+		{
+			int32 TrackID;
+			LexFromString(TrackID, *Repr->GetUniqueIdentifier());
+			const IParserISO14496_12::ITrack* Track = MoovBoxParser->GetTrackByTrackID(TrackID);
+			if (Track)
+			{
+				//
+				TSharedPtrTS<IParserISO14496_12::ITrackIterator> TrackIt(Track->CreateIterator());
+				IParserISO14496_12::ITrackIterator::ESearchMode SearchMode =
+					SearchType == ESearchType::After  || SearchType == ESearchType::StrictlyAfter  ? IParserISO14496_12::ITrackIterator::ESearchMode::After  :
+					SearchType == ESearchType::Before || SearchType == ESearchType::StrictlyBefore ? IParserISO14496_12::ITrackIterator::ESearchMode::Before : IParserISO14496_12::ITrackIterator::ESearchMode::Closest;
+
+				UEMediaError err;
+				FTimeValue firstTimestamp;
+				int64 firstByteOffset = 0;
+				if (AtAbsoluteFilePos < 0)
+				{
+					// Get the track duration. If the start time is outside the track this could be a deliberate seek past the end of the video track
+					// into a longer audio track. We do not want to snap to the last video IDR frame then.
+					FTimeFraction TrackDuration = Track->GetDuration();
+					if (TrackDuration.IsValid() && FTimeValue().SetFromTimeFraction(TrackDuration) <= StartPosition.Time && (SearchType == ESearchType::After || SearchType == ESearchType::StrictlyAfter || SearchType == ESearchType::Closest))
+					{
+						err = UEMEDIA_ERROR_END_OF_STREAM;
+					}
+					else
+					{
+						err = TrackIt->StartAtTime(StartPosition.Time, SearchMode);
+						if (err == UEMEDIA_ERROR_OK)
+						{
+							firstTimestamp.SetFromND(TrackIt->GetDTS(), TrackIt->GetTimescale());
+							firstByteOffset = TrackIt->GetSampleFileOffset();
+						}
+					}
+				}
+				else
+				{
+					TSharedPtr<IParserISO14496_12::IAllTrackIterator, ESPMode::ThreadSafe> AllTrackIterator = MoovBoxParser->CreateAllTrackIteratorByFilePos(AtAbsoluteFilePos);
+					const IParserISO14496_12::ITrackIterator* CurrentTrackIt = AllTrackIterator->Current();
+					if (CurrentTrackIt)
+					{
+						firstTimestamp.SetFromND(CurrentTrackIt->GetDTS(), CurrentTrackIt->GetTimescale());
+						firstByteOffset = CurrentTrackIt->GetSampleFileOffset();
+						err = UEMEDIA_ERROR_OK;
+					}
+					else
+					{
+						err = UEMEDIA_ERROR_END_OF_STREAM;
+					}
+				}
+				if (err == UEMEDIA_ERROR_OK)
+				{
+					// Time found. Set up the fragment request.
+					TSharedPtrTS<FStreamSegmentRequestMP4> req(new FStreamSegmentRequestMP4);
+					OutSegment = req;
+
+					req->MediaAsset 			= SharedThis(this);
+					req->PrimaryTrackIterator   = TrackIt;
+					req->FirstPTS   			= firstTimestamp;
+					req->PrimaryStreamType  	= EStreamType::Video;
+					req->FileStartOffset		= firstByteOffset;
+					req->FileEndOffset  		= -1;
+					req->Bitrate				= Repr->GetBitrate();
+					req->bStartingOnMOOF		= false;
+					req->bIsContinuationSegment = false;
+					req->SegmentDuration		= GetDuration() - req->FirstPTS;
+
+					// FIXME: this may need to add all additional tracks at some point if their individual IDs matter
+					if (AudioAdaptationSets.Num())
+					{
+						req->DependentStreams.AddDefaulted_GetRef().StreamType = EStreamType::Audio;
+					}
+
+					return IManifest::FResult(IManifest::FResult::EType::Found);
+				}
+				else if (err == UEMEDIA_ERROR_END_OF_STREAM)
+				{
+					// If there are no audio tracks we return an EOS request.
+					// Otherwise we search the audio tracks for a start position.
+					if (AudioAdaptationSets.Num() == 0)
+					{
+						TSharedPtrTS<FStreamSegmentRequestMP4> req(new FStreamSegmentRequestMP4);
+						OutSegment = req;
+						req->MediaAsset 			= SharedThis(this);
+						req->FirstPTS   			= FTimeValue::GetZero();
+						req->PrimaryStreamType  	= EStreamType::Video;
+						req->Bitrate				= Repr->GetBitrate();
+						req->bStartingOnMOOF		= false;
+						req->bIsContinuationSegment = false;
+						req->bAllTracksAtEOS		= true;
+						//req->DependentStreams.PushBack().StreamType = EStreamType::Audio;
+						return IManifest::FResult(IManifest::FResult::EType::Found);
+					}
+				}
+				else if (err == UEMEDIA_ERROR_INSUFFICIENT_DATA)
+				{
+					return IManifest::FResult(IManifest::FResult::EType::BeforeStart);
+				}
+				else
+				{
+					IManifest::FResult res(IManifest::FResult::EType::NotFound);
+					res.SetErrorDetail(FErrorDetail().SetError(err)
+									   .SetFacility(Facility::EFacility::MP4Playlist)
+									   .SetCode(ERRCODE_MANIFEST_MP4_STARTSEGMENT_NOT_FOUND)
+									   .SetMessage(FString::Printf(TEXT("Could not find start segment for time %lld"), (long long int)StartPosition.Time.GetAsHNS())));
+					return res;
+				}
+			}
+		}
+	}
+	// No video track(s). Are there audio tracks?
+	if (AudioAdaptationSets.Num())
+	{
+// TODO: use the selected track index if there are several tracks and we have a selected one.
+		TSharedPtrTS<FRepresentationMP4> Repr = StaticCastSharedPtr<FRepresentationMP4>(AudioAdaptationSets[0]->GetRepresentationByIndex(0));
+		if (Repr.IsValid())
+		{
+			int32 TrackID;
+			LexFromString(TrackID, *Repr->GetUniqueIdentifier());
+			const IParserISO14496_12::ITrack* Track = MoovBoxParser->GetTrackByTrackID(TrackID);
+			if (Track)
+			{
+				//
+				TSharedPtrTS<IParserISO14496_12::ITrackIterator> TrackIt(Track->CreateIterator());
+				IParserISO14496_12::ITrackIterator::ESearchMode SearchMode =
+					SearchType == ESearchType::After  || SearchType == ESearchType::StrictlyAfter  ? IParserISO14496_12::ITrackIterator::ESearchMode::After  :
+					SearchType == ESearchType::Before || SearchType == ESearchType::StrictlyBefore ? IParserISO14496_12::ITrackIterator::ESearchMode::Before : IParserISO14496_12::ITrackIterator::ESearchMode::Closest;
+
+				UEMediaError err;
+				FTimeValue firstTimestamp;
+				int64 firstByteOffset = 0;
+				if (AtAbsoluteFilePos < 0)
+				{
+					// Get the track duration. If the start time is outside the track this could be a deliberate seek past the end of the audio track
+					// for which we return EOS.
+					FTimeFraction TrackDuration = Track->GetDuration();
+					if (TrackDuration.IsValid() && FTimeValue().SetFromTimeFraction(TrackDuration) <= StartPosition.Time && (SearchType == ESearchType::After || SearchType == ESearchType::StrictlyAfter || SearchType == ESearchType::Closest))
+					{
+						err = UEMEDIA_ERROR_END_OF_STREAM;
+					}
+					else
+					{
+						err = TrackIt->StartAtTime(StartPosition.Time, SearchMode);
+						if (err == UEMEDIA_ERROR_OK)
+						{
+							firstTimestamp.SetFromND(TrackIt->GetDTS(), TrackIt->GetTimescale());
+							firstByteOffset = TrackIt->GetSampleFileOffset();
+						}
+					}
+				}
+				else
+				{
+					TSharedPtr<IParserISO14496_12::IAllTrackIterator, ESPMode::ThreadSafe> AllTrackIterator = MoovBoxParser->CreateAllTrackIteratorByFilePos(AtAbsoluteFilePos);
+					const IParserISO14496_12::ITrackIterator* CurrentTrackIt = AllTrackIterator->Current();
+					if (CurrentTrackIt)
+					{
+						firstTimestamp.SetFromND(CurrentTrackIt->GetDTS(), CurrentTrackIt->GetTimescale());
+						firstByteOffset = CurrentTrackIt->GetSampleFileOffset();
+						err = UEMEDIA_ERROR_OK;
+					}
+					else
+					{
+						err = UEMEDIA_ERROR_END_OF_STREAM;
+					}
+				}
+				if (err == UEMEDIA_ERROR_OK)
+				{
+					// Time found. Set up the fragment request.
+					TSharedPtrTS<FStreamSegmentRequestMP4> req(new FStreamSegmentRequestMP4);
+					OutSegment = req;
+
+					req->MediaAsset 			= SharedThis(this);
+					req->PrimaryTrackIterator   = TrackIt;
+					req->FirstPTS   			= firstTimestamp;
+					req->PrimaryStreamType  	= EStreamType::Audio;
+					req->FileStartOffset		= firstByteOffset;
+					req->FileEndOffset  		= -1;
+					req->Bitrate				= Repr->GetBitrate();
+					req->bStartingOnMOOF		= false;
+					req->bIsContinuationSegment = false;
+					req->SegmentDuration		= GetDuration() - req->FirstPTS;
+
+					// In case the video stream is shorter than audio we still need to add it as a dependent stream
+					// (if it exists) in case the video will loop back to a point where there is video.
+					if (VideoAdaptationSets.Num())
+					{
+						req->DependentStreams.AddDefaulted_GetRef().StreamType = EStreamType::Video;
+					}
+
+					// FIXME: there may be subtitle tracks here we need to add as dependent streams.
+
+					return IManifest::FResult(IManifest::FResult::EType::Found);
+				}
+				else if (err == UEMEDIA_ERROR_END_OF_STREAM)
+				{
+					// Regardless of there being a video stream or not we return an EOS request.
+					TSharedPtrTS<FStreamSegmentRequestMP4> req(new FStreamSegmentRequestMP4);
+					OutSegment = req;
+					req->MediaAsset 			= SharedThis(this);
+					req->FirstPTS   			= FTimeValue::GetZero();
+					req->PrimaryStreamType  	= EStreamType::Audio;
+					req->Bitrate				= Repr->GetBitrate();
+					req->bStartingOnMOOF		= false;
+					req->bIsContinuationSegment = false;
+					req->bAllTracksAtEOS		= true;
+					// But if there is a video track we add it as a dependent stream that is also at EOS.
+					if (VideoAdaptationSets.Num())
+					{
+						req->DependentStreams.AddDefaulted_GetRef().StreamType = EStreamType::Video;
+					}
+					return IManifest::FResult(IManifest::FResult::EType::Found);
+				}
+				else if (err == UEMEDIA_ERROR_INSUFFICIENT_DATA)
+				{
+					return IManifest::FResult(IManifest::FResult::EType::BeforeStart);
+				}
+				else
+				{
+					IManifest::FResult res(IManifest::FResult::EType::NotFound);
+					res.SetErrorDetail(FErrorDetail().SetError(err)
+									   .SetFacility(Facility::EFacility::MP4Playlist)
+									   .SetCode(ERRCODE_MANIFEST_MP4_STARTSEGMENT_NOT_FOUND)
+									   .SetMessage(FString::Printf(TEXT("Could not find start segment for time %lld"), (long long int)StartPosition.Time.GetAsHNS())));
+					return res;
+				}
+			}
+		}
+	}
+
+	return IManifest::FResult(IManifest::FResult::EType::NotFound).SetErrorDetail(FErrorDetail().SetError(UEMEDIA_ERROR_INSUFFICIENT_DATA)
+					   .SetFacility(Facility::EFacility::MP4Playlist)
+					   .SetCode(ERRCODE_MANIFEST_MP4_STARTSEGMENT_NOT_FOUND)
+					   .SetMessage(FString::Printf(TEXT("Could not find start segment for time %lld, no valid tracks"), (long long int)StartPosition.Time.GetAsHNS())));
+}
+
+IManifest::FResult FManifestMP4Internal::FTimelineAssetMP4::GetNextSegment(TSharedPtrTS<IStreamSegment>& OutSegment, TSharedPtrTS<const IStreamSegment> CurrentSegment, const FParamDict& Options)
+{
+	// At the moment we do not break the mp4 down into segments, so the one we had been streaming just now extended
+	// to the end of the asset which is now done.
+	return IManifest::FResult(IManifest::FResult::EType::PastEOS);
+}
+
+IManifest::FResult FManifestMP4Internal::FTimelineAssetMP4::GetRetrySegment(TSharedPtrTS<IStreamSegment>& OutSegment, TSharedPtrTS<const IStreamSegment> CurrentSegment, const FParamDict& Options)
+{
+	const FStreamSegmentRequestMP4* Request = static_cast<const FStreamSegmentRequestMP4*>(CurrentSegment.Get());
+	if (Request)
+	{
+		FPlayStartPosition dummyPos;
+		IManifest::FResult res = GetStartingSegment(OutSegment, dummyPos, ESearchType::Same, Request->CurrentIteratorBytePos);
+		if (res.GetType() == IManifest::FResult::EType::Found)
+		{
+			FStreamSegmentRequestMP4* RetryRequest = static_cast<FStreamSegmentRequestMP4*>(OutSegment.Get());
+			RetryRequest->PlayerLoopState   	 = Request->PlayerLoopState;
+			RetryRequest->bIsContinuationSegment = true;
+			RetryRequest->NumOverallRetries 	 = Request->NumOverallRetries + 1;
+			return res;
+		}
+	}
+	return IManifest::FResult(IManifest::FResult::EType::NotFound);
+}
+
+
+IManifest::FResult FManifestMP4Internal::FTimelineAssetMP4::GetLoopingSegment(TSharedPtrTS<IStreamSegment>& OutSegment, FPlayerLoopState& InOutLoopState, const TMultiMap<EStreamType, TSharedPtrTS<IStreamSegment>>& InFinishedSegments, const FPlayStartPosition& StartPosition, ESearchType SearchType)
+{
+	if (InFinishedSegments.Num())
+	{
+		auto It = InFinishedSegments.CreateConstIterator();
+		const FStreamSegmentRequestMP4* FinishedRequest = static_cast<const FStreamSegmentRequestMP4*>(It->Value.Get());
+		if (FinishedRequest)
+		{
+			IManifest::FResult res = GetStartingSegment(OutSegment, StartPosition, SearchType, -1);
+			if (res.GetType() == IManifest::FResult::EType::Found)
+			{
+				InOutLoopState.bLoopEnabled = true;
+				InOutLoopState.LoopBasetime = FinishedRequest->NextLargestExpectedTimestamp.IsValid() ? FinishedRequest->NextLargestExpectedTimestamp : FTimeValue::GetZero();
+				++InOutLoopState.LoopCount;
+				FStreamSegmentRequestMP4* LoopRequest = static_cast<FStreamSegmentRequestMP4*>(OutSegment.Get());
+				LoopRequest->PlayerLoopState = InOutLoopState;
+				return res;
+			}
+		}
+	}
+	// Return past EOS when we can't loop to indicate we're really done now.
+	return IManifest::FResult(IManifest::FResult::EType::PastEOS);
+}
+
+
+void FManifestMP4Internal::FTimelineAssetMP4::GetSegmentInformation(TArray<IManifest::IPlayPeriod::FSegmentInformation>& OutSegmentInformation, FTimeValue& OutAverageSegmentDuration, TSharedPtrTS<const IStreamSegment> CurrentSegment, const FTimeValue& LookAheadTime, const TSharedPtrTS<IPlaybackAssetAdaptationSet>& AdaptationSet, const TSharedPtrTS<IPlaybackAssetRepresentation>& Representation)
+{
+	// This is not expected to be called. And if it does we return a dummy entry.
+	OutAverageSegmentDuration.SetFromSeconds(60.0);
+	OutSegmentInformation.Empty();
+	IManifest::IPlayPeriod::FSegmentInformation& si = OutSegmentInformation.AddDefaulted_GetRef();
+	si.ByteSize = 1024 * 1024 * 1024;
+	si.Duration.SetFromSeconds(60.0);
+}
+
+TSharedPtrTS<IParserISO14496_12>	FManifestMP4Internal::FTimelineAssetMP4::GetMoovBoxParser()
+{
+	return MoovBoxParser;
+}
+
+
+
+FErrorDetail FManifestMP4Internal::FAdaptationSetMP4::CreateFrom(const IParserISO14496_12::ITrack* InTrack, const FString& URL)
+{
+	Representation = MakeSharedTS<FRepresentationMP4>();
+	FErrorDetail err;
+	err = Representation->CreateFrom(InTrack, URL);
+	if (err.IsOK())
+	{
+		CodecRFC6381	 = Representation->GetCodecInformation().GetCodecSpecifierRFC6381();
+		UniqueIdentifier = FString("adaptation.") + Representation->GetUniqueIdentifier();
+		Language		 = InTrack->GetLanguage();
+	}
+	return err;
+}
+
+
+FErrorDetail FManifestMP4Internal::FRepresentationMP4::CreateFrom(const IParserISO14496_12::ITrack* InTrack, const FString& URL)
+{
+	CodecInformation	 = InTrack->GetCodecInformation();
+	// Get the CSD
+	CodecSpecificData    = InTrack->GetCodecSpecificData();
+	CodecSpecificDataRAW = InTrack->GetCodecSpecificDataRAW();
+
+	// Since we are dealing with a track inside a multiplexed file there is no choice for CDNs.
+	// We set the URL as the CDN.
+	CDN = URL;
+
+	// The unique identifier will be the track ID inside the mp4.
+	// NOTE: This *MUST* be just a number since it gets parsed back out from a string into a number later! Do *NOT* prepend/append any string literals!!
+	UniqueIdentifier = LexToString(InTrack->GetID());
+
+	// Get bitrate from the average or max bitrate as stored in the track. If not stored it will be 0.
+	Bitrate = InTrack->GetBitrateInfo().AvgBitrate ? InTrack->GetBitrateInfo().AvgBitrate : InTrack->GetBitrateInfo().MaxBitrate;
+
+	// With no bitrate available we set some defaults. This is mainly to avoid a bitrate of 0 from being surfaced that would prevent
+	// events like the initial bitrate change that needs to transition away from 0 to something real.
+	if (Bitrate == 0)
+	{
+		switch(CodecInformation.GetStreamType())
+		{
+			case EStreamType::Video:
+				Bitrate = 1 * 1024 * 1024;
+				break;
+			case EStreamType::Audio:
+				Bitrate = 64 * 1024;
+				break;
+			case EStreamType::Subtitle:
+				Bitrate = 8 * 1024;
+				break;
+			default:
+				// Whatever it is, assume it's a low bitrate.
+				Bitrate = 32 * 1024;
+				break;
+		}
+	}
+
+	// Not a whole lot that could have gone wrong here.
+	return FErrorDetail();
+}
+
+
+} // namespace Electra
+
+
diff --git a/Engine/Plugins/Media/ElectraPlayer/Source/ElectraPlayerRuntime/Private/Runtime/Player/mp4/ManifestMP4.h b/Engine/Plugins/Media/ElectraPlayer/Source/ElectraPlayerRuntime/Private/Runtime/Player/mp4/ManifestMP4.h
new file mode 100644
--- /dev/null
+++ b/Engine/Plugins/Media/ElectraPlayer/Source/ElectraPlayerRuntime/Private/Runtime/Player/mp4/ManifestMP4.h
@@ -0,0 +1,258 @@
+// Copyright Epic Games, Inc. All Rights Reserved.
+
+#pragma once
+
+#include "PlayerCore.h"
+#include "Player/Manifest.h"
+#include "Player/PlaybackTimeline.h"
+#include "HTTP/HTTPManager.h"
+#include "Demuxer/ParserISO14496-12.h"
+#include "SynchronizedClock.h"
+#include "HAL/LowLevelMemTracker.h"
+#include "ElectraPlayerPrivate.h"
+#include "Player/PlayerStreamReader.h"
+
+
+
+namespace Electra
+{
+
+/**
+ * This class represents the internal "manifest" or "playlist" of an mp4 file.
+ * All metadata of the tracks it is composed of are maintained by this class.
+ */
+class FManifestMP4Internal : public IManifest, public IPlaybackAssetTimeline, public TSharedFromThis<FManifestMP4Internal, ESPMode::ThreadSafe>
+{
+public:
+	FManifestMP4Internal(IPlayerSessionServices* InPlayerSessionServices);
+	virtual ~FManifestMP4Internal();
+
+	FErrorDetail Build(TSharedPtrTS<IParserISO14496_12> MP4Parser, const FString& URL, const HTTP::FConnectionInfo& InConnectionInfo);
+
+
+	virtual EType GetPresentationType() const override;
+	virtual TSharedPtrTS<IPlaybackAssetTimeline> GetTimeline() const override;
+	virtual int64 GetDefaultStartingBitrate() const override;
+	virtual void GetStreamMetadata(TArray<FStreamMetadata>& OutMetadata, EStreamType StreamType) const override;
+	virtual FTimeValue GetMinBufferTime() const override;
+	virtual IStreamReader *CreateStreamReaderHandler() override;
+	virtual FResult FindPlayPeriod(TSharedPtrTS<IPlayPeriod>& OutPlayPeriod, const FPlayStartPosition& StartPosition, ESearchType SearchType) override;
+
+
+	class FRepresentationMP4 : public IPlaybackAssetRepresentation
+	{
+	public:
+		virtual ~FRepresentationMP4()
+		{ }
+
+		FErrorDetail CreateFrom(const IParserISO14496_12::ITrack* InTrack, const FString& URL);
+
+		virtual FString GetUniqueIdentifier() const override
+		{ return UniqueIdentifier; }
+		virtual const FStreamCodecInformation& GetCodecInformation() const override
+		{ return CodecInformation; }
+		virtual int32 GetBitrate() const override
+		{ return Bitrate; }
+		virtual FString GetCDN() const override
+		{ return CDN; }
+	private:
+		FStreamCodecInformation		CodecInformation;
+		FString						UniqueIdentifier;
+		FString						CDN;
+		int32						Bitrate;
+
+		TArray<uint8>				CodecSpecificData;
+		TArray<uint8>				CodecSpecificDataRAW;
+	};
+
+	class FAdaptationSetMP4 : public IPlaybackAssetAdaptationSet
+	{
+	public:
+		virtual ~FAdaptationSetMP4()
+		{ }
+
+		FErrorDetail CreateFrom(const IParserISO14496_12::ITrack* InTrack, const FString& URL);
+
+		virtual FString GetUniqueIdentifier() const override
+		{ return UniqueIdentifier; }
+		virtual FString GetListOfCodecs() const override
+		{ return CodecRFC6381; }
+		virtual FString GetLanguage() const override
+		{ return Language; }
+		virtual int32 GetNumberOfRepresentations() const override
+		{ return Representation.IsValid() ? 1 : 0; }
+		virtual TSharedPtrTS<IPlaybackAssetRepresentation> GetRepresentationByIndex(int32 RepresentationIndex) const override
+		{ return RepresentationIndex == 0 ? Representation : TSharedPtrTS<IPlaybackAssetRepresentation>(); }
+		virtual TSharedPtrTS<IPlaybackAssetRepresentation> GetRepresentationByUniqueIdentifier(const FString& InUniqueIdentifier) const override
+		{ return Representation.IsValid() && Representation->GetUniqueIdentifier() == InUniqueIdentifier ? Representation : TSharedPtrTS<IPlaybackAssetRepresentation>(); }
+	private:
+		TSharedPtrTS<FRepresentationMP4>		Representation;
+		FString									Language;
+		FString									CodecRFC6381;
+		FString									UniqueIdentifier;
+	};
+
+	class FTimelineAssetMP4 : public ITimelineMediaAsset, public TSharedFromThis<FTimelineAssetMP4, ESPMode::ThreadSafe>
+	{
+	public:
+		FTimelineAssetMP4()
+			: PlayerSessionServices(nullptr)
+		{ }
+
+		virtual ~FTimelineAssetMP4()
+		{ }
+
+		FErrorDetail Build(IPlayerSessionServices* InPlayerSessionServices, TSharedPtrTS<IParserISO14496_12> MP4Parser, const FString& URL);
+
+		virtual FTimeRange GetTimeRange() const override
+		{
+			FTimeRange tr;
+			tr.Start.SetToZero();
+			tr.End = GetDuration();
+			return tr;
+		}
+
+		virtual FTimeValue GetDuration() const override
+		{
+			if (MoovBoxParser.IsValid())
+			{
+				TMediaOptionalValue<FTimeFraction> movieDuration = MoovBoxParser->GetMovieDuration();
+				if (movieDuration.IsSet())
+				{
+					FTimeValue dur;
+					dur.SetFromTimeFraction(movieDuration.Value());
+					return dur;
+				}
+			}
+			return FTimeValue();
+		}
+
+		virtual FString GetAssetIdentifier() const override
+		{ return FString("mp4-asset.0"); }
+		virtual FString GetUniqueIdentifier() const override
+		{ return FString("mp4-media.0"); }
+		virtual int32 GetNumberOfAdaptationSets(EStreamType OfStreamType) const override
+		{
+			switch(OfStreamType)
+			{
+				case EStreamType::Video:
+					return VideoAdaptationSets.Num();
+				case EStreamType::Audio:
+					return AudioAdaptationSets.Num();
+				default:
+					return 0;
+			}
+		}
+		virtual TSharedPtrTS<IPlaybackAssetAdaptationSet> GetAdaptationSetByTypeAndIndex(EStreamType OfStreamType, int32 AdaptationSetIndex) const override
+		{
+			switch(OfStreamType)
+			{
+				case EStreamType::Video:
+					return AdaptationSetIndex < VideoAdaptationSets.Num() ? VideoAdaptationSets[AdaptationSetIndex] : TSharedPtrTS<IPlaybackAssetAdaptationSet>();
+				case EStreamType::Audio:
+					return AdaptationSetIndex < AudioAdaptationSets.Num() ? AudioAdaptationSets[AdaptationSetIndex] : TSharedPtrTS<IPlaybackAssetAdaptationSet>();
+				default:
+					return 0;
+			}
+			return TSharedPtrTS<IPlaybackAssetAdaptationSet>();
+		}
+
+
+		FResult GetStartingSegment(TSharedPtrTS<IStreamSegment>& OutSegment, const FPlayStartPosition& StartPosition, ESearchType SearchType, int64 AtAbsoluteFilePos);
+		FResult GetNextSegment(TSharedPtrTS<IStreamSegment>& OutSegment, TSharedPtrTS<const IStreamSegment> CurrentSegment, const FParamDict& Options);
+		FResult GetRetrySegment(TSharedPtrTS<IStreamSegment>& OutSegment, TSharedPtrTS<const IStreamSegment> CurrentSegment, const FParamDict& Options);
+		FResult GetLoopingSegment(TSharedPtrTS<IStreamSegment>& OutSegment, FPlayerLoopState& InOutLoopState, const TMultiMap<EStreamType, TSharedPtrTS<IStreamSegment>>& InFinishedSegments, const FPlayStartPosition& StartPosition, ESearchType SearchType);
+
+		void GetSegmentInformation(TArray<IManifest::IPlayPeriod::FSegmentInformation>& OutSegmentInformation, FTimeValue& OutAverageSegmentDuration, TSharedPtrTS<const IStreamSegment> CurrentSegment, const FTimeValue& LookAheadTime, const TSharedPtrTS<IPlaybackAssetAdaptationSet>& AdaptationSet, const TSharedPtrTS<IPlaybackAssetRepresentation>& Representation);
+		TSharedPtrTS<IParserISO14496_12>	GetMoovBoxParser();
+		const FString& GetMediaURL() const
+		{
+			return MediaURL;
+		}
+
+	private:
+		void LogMessage(IInfoLog::ELevel Level, const FString& Message);
+		IPlayerSessionServices* 						PlayerSessionServices;
+		FString											MediaURL;
+		TSharedPtrTS<IParserISO14496_12>				MoovBoxParser;
+		TArray<TSharedPtrTS<FAdaptationSetMP4>>			VideoAdaptationSets;
+		TArray<TSharedPtrTS<FAdaptationSetMP4>>			AudioAdaptationSets;
+	};
+
+
+	class FPlayPeriodMP4 : public IManifest::IPlayPeriod
+	{
+	public:
+		FPlayPeriodMP4(TSharedPtrTS<FTimelineAssetMP4> InMediaAsset);
+		virtual ~FPlayPeriodMP4();
+		virtual void SetStreamPreferences(const FStreamPreferences& Preferences) override;
+		virtual EReadyState GetReadyState() override;
+		virtual void PrepareForPlay(const FParamDict& Options) override;
+		virtual TSharedPtrTS<ITimelineMediaAsset> GetMediaAsset() const override;
+		virtual void SelectStream(const TSharedPtrTS<IPlaybackAssetAdaptationSet>& AdaptationSet, const TSharedPtrTS<IPlaybackAssetRepresentation>& Representation, const FString& PreferredCDN) override;
+		virtual FResult GetStartingSegment(TSharedPtrTS<IStreamSegment>& OutSegment, const FPlayStartPosition& StartPosition, ESearchType SearchType) override;
+		virtual FResult GetNextSegment(TSharedPtrTS<IStreamSegment>& OutSegment, TSharedPtrTS<const IStreamSegment> CurrentSegment, const FParamDict& Options) override;
+		virtual FResult GetRetrySegment(TSharedPtrTS<IStreamSegment>& OutSegment, TSharedPtrTS<const IStreamSegment> CurrentSegment, const FParamDict& Options) override;
+		virtual FResult GetLoopingSegment(TSharedPtrTS<IStreamSegment>& OutSegment, FPlayerLoopState& InOutLoopState, const TMultiMap<EStreamType, TSharedPtrTS<IStreamSegment>>& InFinishedSegments, const FPlayStartPosition& StartPosition, ESearchType SearchType) override;
+		virtual void GetSegmentInformation(TArray<FSegmentInformation>& OutSegmentInformation, FTimeValue& OutAverageSegmentDuration, TSharedPtrTS<const IStreamSegment> CurrentSegment, const FTimeValue& LookAheadTime, const TSharedPtrTS<IPlaybackAssetAdaptationSet>& AdaptationSet, const TSharedPtrTS<IPlaybackAssetRepresentation>& Representation) override;
+	private:
+		TWeakPtrTS<FTimelineAssetMP4>			MediaAsset;
+		FStreamPreferences						Preferences;
+		FParamDict								Options;
+		bool									bIsReady;
+	};
+
+	//-------------------------------------------------------------------------
+	// Methods from IPlaybackAssetTimeline
+	//
+	virtual FTimeValue GetAnchorTime() const override
+	{
+		return FTimeValue::GetZero();
+	}
+	virtual FTimeRange GetTotalTimeRange() const override
+	{
+		return MediaAsset.IsValid() ? MediaAsset->GetTimeRange() : FTimeRange();
+	}
+	virtual FTimeRange GetSeekableTimeRange() const override
+	{
+		FTimeRange tr = GetTotalTimeRange();
+		// FIXME: This would need to be the time of the last sync frame of video (if it exists) or audio.
+		//        For now it does not need to be precise since the nearest sync frame is searched for when seeking.
+		//        It doesn't seem to make a lot of sense to seek very close to the end and play from there so we
+		//        allow only to go as close to the end as we deem feasible.
+		tr.End -= FTimeValue().SetFromSeconds(2.0);
+		if (tr.End < FTimeValue::GetZero())
+		{
+			tr.End = FTimeValue::GetZero();
+		}
+		return tr;
+	}
+	virtual void GetSeekablePositions(TArray<FTimespan>& OutPositions) const override
+	{
+		// For the time being we do not return anything here as that would require to iterate the tracks.
+	}
+	virtual FTimeValue GetDuration() const override
+	{
+		return MediaAsset.IsValid() ? MediaAsset->GetDuration() : FTimeValue();
+	}
+	virtual int32 GetNumberOfMediaAssets() const override
+	{
+		return 1;
+	}
+	virtual TSharedPtrTS<ITimelineMediaAsset> GetMediaAssetByIndex(int32 MediaAssetIndex) const override
+	{
+		return MediaAsset;
+	}
+
+
+	void LogMessage(IInfoLog::ELevel Level, const FString& Message);
+
+	IPlayerSessionServices* 			PlayerSessionServices;
+	TSharedPtrTS<FTimelineAssetMP4>		MediaAsset;
+	HTTP::FConnectionInfo				ConnectionInfo;
+};
+
+
+} // namespace Electra
+
+
diff --git a/Engine/Plugins/Media/ElectraPlayer/Source/ElectraPlayerRuntime/Private/Runtime/Player/mp4/PlaylistReaderMP4.cpp b/Engine/Plugins/Media/ElectraPlayer/Source/ElectraPlayerRuntime/Private/Runtime/Player/mp4/PlaylistReaderMP4.cpp
new file mode 100644
--- /dev/null
+++ b/Engine/Plugins/Media/ElectraPlayer/Source/ElectraPlayerRuntime/Private/Runtime/Player/mp4/PlaylistReaderMP4.cpp
@@ -0,0 +1,506 @@
+// Copyright Epic Games, Inc. All Rights Reserved.
+
+#include "PlayerCore.h"
+#include "Player/Manifest.h"
+#include "HTTP/HTTPManager.h"
+#include "Demuxer/ParserISO14496-12.h"
+#include "SynchronizedClock.h"
+#include "PlaylistReaderMP4.h"
+#include "Utilities/HashFunctions.h"
+#include "Utilities/TimeUtilities.h"
+#include "HAL/LowLevelMemTracker.h"
+#include "ElectraPlayerPrivate.h"
+#include "Player/AdaptiveStreamingPlayerResourceRequest.h"
+#include "Player/PlayerStreamReader.h"
+#include "Player/mp4/ManifestMP4.h"
+
+
+#define ERRCODE_MP4_INVALID_FILE							1
+
+
+
+namespace Electra
+{
+
+/**
+ * This class is responsible for downloading the mp4 non-mdat boxes and parsing them.
+ */
+class FPlaylistReaderMP4 : public IPlaylistReaderMP4, public IParserISO14496_12::IReader, public IParserISO14496_12::IBoxCallback, public FMediaThread
+{
+public:
+	FPlaylistReaderMP4();
+	void Initialize(IPlayerSessionServices* PlayerSessionServices);
+
+	virtual ~FPlaylistReaderMP4();
+
+	/**
+	 * Returns the type of playlist format.
+	 * For this implementation it will be "mp4".
+	 *
+	 * @return "mp4" to indicate this is an mp4 file.
+	 */
+	virtual const FString& GetPlaylistType() const override
+	{
+		static FString Type("mp4");
+		return Type;
+	}
+
+	/**
+	 * Loads and parses the playlist.
+	 *
+	 * @param URL     URL of the playlist to load
+	 * @param Preferences
+	 *                User preferences for initial stream selection.
+	 * @param Options Options for the reader and parser
+	 */
+	virtual void LoadAndParse(const FString& URL, const FStreamPreferences& Preferences, const FParamDict& Options) override;
+
+	/**
+	 * Returns the URL from which the playlist was loaded (or supposed to be loaded).
+	 *
+	 * @return The playlist URL
+	 */
+	virtual FString GetURL() const override;
+
+	/**
+	 * Returns an interface to the manifest created from the loaded HLS playlists.
+	 *
+	 * @return A shared manifest interface pointer.
+	 */
+	virtual TSharedPtrTS<IManifest> GetManifest() override;
+
+private:
+	struct FReadBuffer
+	{
+		FReadBuffer()
+		{
+			Reset();
+		}
+		void Reset()
+		{
+			ReceiveBuffer.Reset();
+			ParsePos = 0;
+			bHasBeenAborted = false;
+		}
+		void Abort()
+		{
+			TSharedPtrTS<IElectraHttpManager::FReceiveBuffer>	CurrentBuffer = ReceiveBuffer;
+			if (CurrentBuffer.IsValid())
+			{
+				CurrentBuffer->Buffer.Abort();
+			}
+			bHasBeenAborted = true;
+		}
+		TSharedPtrTS<IElectraHttpManager::FReceiveBuffer>		ReceiveBuffer;
+		int64												ParsePos;
+		bool												bHasBeenAborted;
+	};
+
+	// Methods from IParserISO14496_12::IReader
+	virtual int64 ReadData(void* IntoBuffer, int64 NumBytesToRead) override;
+	virtual bool HasReachedEOF() const override;
+	virtual bool HasReadBeenAborted() const override;
+	virtual int64 GetCurrentOffset() const override;
+	// Methods from IParserISO14496_12::IBoxCallback
+	virtual IParserISO14496_12::IBoxCallback::EParseContinuation OnFoundBox(IParserISO14496_12::FBoxType Box, int64 BoxSizeInBytes, int64 FileDataOffset, int64 BoxDataOffset) override;
+
+	void Close();
+	void StartWorkerThread();
+	void StopWorkerThread();
+	void WorkerThread(void);
+
+	void PostError(const FString& Message, uint16 Code, UEMediaError Error = UEMEDIA_ERROR_OK);
+	void LogMessage(IInfoLog::ELevel Level, const FString& Message);
+
+	FErrorDetail CreateErrorAndLog(const FString& Message, uint16 Code, UEMediaError Error = UEMEDIA_ERROR_OK);
+
+	int32 HTTPProgressCallback(const IElectraHttpManager::FRequest* Request);
+	void HTTPCompletionCallback(const IElectraHttpManager::FRequest* Request);
+
+	bool HasErrored() const
+	{
+		return bHasErrored;
+	}
+
+	IPlayerSessionServices*									PlayerSessionServices;
+	FStreamPreferences										StreamPreferences;
+	FParamDict												Options;
+	FString													MasterPlaylistURL;
+	FMediaEvent												WorkerThreadQuitSignal;
+	bool													bIsWorkerThreadStarted;
+
+	TSharedPtrTS<IElectraHttpManager::FProgressListener>		ProgressListener;
+	FReadBuffer												ReadBuffer;
+	bool													bAbort;
+	bool													bHasErrored;
+
+	TSharedPtrTS<IParserISO14496_12>						MP4Parser;
+	bool													bFoundBoxFTYP;
+	bool													bFoundBoxMOOV;
+	bool													bFoundBoxSIDX;
+	bool													bFoundBoxMOOF;
+	bool													bFoundBoxMDAT;
+
+	TSharedPtrTS<FManifestMP4Internal>						Manifest;
+	FErrorDetail											LastErrorDetail;
+};
+
+
+/***************************************************************************************************************************************************/
+/***************************************************************************************************************************************************/
+/***************************************************************************************************************************************************/
+
+IPlaylistReader* IPlaylistReaderMP4::Create(IPlayerSessionServices* PlayerSessionServices)
+{
+	FPlaylistReaderMP4* PlaylistReader = new FPlaylistReaderMP4;
+	if (PlaylistReader)
+	{
+		PlaylistReader->Initialize(PlayerSessionServices);
+	}
+	return PlaylistReader;
+}
+
+/***************************************************************************************************************************************************/
+/***************************************************************************************************************************************************/
+/***************************************************************************************************************************************************/
+
+
+FPlaylistReaderMP4::FPlaylistReaderMP4()
+	: FMediaThread("ElectraPlayer::MP4 Playlist")
+	, PlayerSessionServices(nullptr)
+	, bIsWorkerThreadStarted(false)
+	, bAbort(false)
+	, bHasErrored(false)
+	, bFoundBoxFTYP(false)
+	, bFoundBoxMOOV(false)
+	, bFoundBoxSIDX(false)
+	, bFoundBoxMOOF(false)
+	, bFoundBoxMDAT(false)
+{
+}
+
+FPlaylistReaderMP4::~FPlaylistReaderMP4()
+{
+	Close();
+}
+
+FString FPlaylistReaderMP4::GetURL() const
+{
+	return MasterPlaylistURL;
+}
+
+TSharedPtrTS<IManifest> FPlaylistReaderMP4::GetManifest()
+{
+	return Manifest;
+}
+
+void FPlaylistReaderMP4::Initialize(IPlayerSessionServices* InPlayerSessionServices)
+{
+	PlayerSessionServices = InPlayerSessionServices;
+	ProgressListener = MakeSharedTS<IElectraHttpManager::FProgressListener>();
+	ProgressListener->CompletionDelegate = Electra::MakeDelegate(this, &FPlaylistReaderMP4::HTTPCompletionCallback);
+	ProgressListener->ProgressDelegate   = Electra::MakeDelegate(this, &FPlaylistReaderMP4::HTTPProgressCallback);
+}
+
+void FPlaylistReaderMP4::Close()
+{
+	ReadBuffer.Abort();
+	bAbort = true;
+
+	StopWorkerThread();
+}
+
+void FPlaylistReaderMP4::StartWorkerThread()
+{
+	check(!bIsWorkerThreadStarted);
+	ThreadStart(Electra::MakeDelegate(this, &FPlaylistReaderMP4::WorkerThread));
+	bIsWorkerThreadStarted = true;
+}
+
+void FPlaylistReaderMP4::StopWorkerThread()
+{
+	if (bIsWorkerThreadStarted)
+	{
+		WorkerThreadQuitSignal.Signal();
+		ThreadWaitDone();
+		ThreadReset();
+		bIsWorkerThreadStarted = false;
+	}
+}
+
+void FPlaylistReaderMP4::PostError(const FString& InMessage, uint16 InCode, UEMediaError InError)
+{
+	LastErrorDetail.Clear();
+	LastErrorDetail.SetError(InError != UEMEDIA_ERROR_OK ? InError : UEMEDIA_ERROR_DETAIL);
+	LastErrorDetail.SetFacility(Facility::EFacility::MP4PlaylistReader);
+	LastErrorDetail.SetCode(InCode);
+	LastErrorDetail.SetMessage(InMessage);
+	check(PlayerSessionServices);
+	if (PlayerSessionServices)
+	{
+		PlayerSessionServices->PostError(LastErrorDetail);
+	}
+}
+
+FErrorDetail FPlaylistReaderMP4::CreateErrorAndLog(const FString& InMessage, uint16 InCode, UEMediaError InError)
+{
+	FErrorDetail err;
+	err.SetError(InError != UEMEDIA_ERROR_OK ? InError : UEMEDIA_ERROR_DETAIL);
+	err.SetFacility(Facility::EFacility::MP4PlaylistReader);
+	err.SetCode(InCode);
+	err.SetMessage(InMessage);
+	check(PlayerSessionServices);
+	if (PlayerSessionServices)
+	{
+		PlayerSessionServices->PostLog(Facility::EFacility::MP4PlaylistReader, IInfoLog::ELevel::Error, err.GetPrintable());
+	}
+	return err;
+}
+
+void FPlaylistReaderMP4::LogMessage(IInfoLog::ELevel Level, const FString& Message)
+{
+	if (PlayerSessionServices)
+	{
+		PlayerSessionServices->PostLog(Facility::EFacility::MP4PlaylistReader, Level, Message);
+	}
+}
+
+void FPlaylistReaderMP4::LoadAndParse(const FString& URL, const FStreamPreferences& Preferences, const FParamDict& InOptions)
+{
+	StreamPreferences = Preferences;
+	Options 		  = InOptions;
+	MasterPlaylistURL = URL;
+	StartWorkerThread();
+}
+
+int32 FPlaylistReaderMP4::HTTPProgressCallback(const IElectraHttpManager::FRequest* Request)
+{
+	// Aborted?
+	return bAbort ? 1 : 0;
+}
+
+void FPlaylistReaderMP4::HTTPCompletionCallback(const IElectraHttpManager::FRequest* InRequest)
+{
+	bHasErrored = InRequest->ConnectionInfo.StatusInfo.ErrorDetail.IsError();
+}
+
+
+void FPlaylistReaderMP4::WorkerThread(void)
+{
+	LLM_SCOPE(ELLMTag::ElectraPlayer);
+	CSV_SCOPED_TIMING_STAT(ElectraPlayer, PlaylistReaderMP4_Worker);
+
+	ReadBuffer.Reset();
+	ReadBuffer.ReceiveBuffer = MakeSharedTS<IElectraHttpManager::FReceiveBuffer>();
+	ReadBuffer.ReceiveBuffer->Buffer.Reserve(1 << 20);
+	ReadBuffer.ReceiveBuffer->bEnableRingbuffer = true;
+
+	TSharedPtrTS<IElectraHttpManager::FRequest> HTTP(new IElectraHttpManager::FRequest);
+	HTTP->Parameters.URL = MasterPlaylistURL;
+	// Perform small reads of 256K ranges only. We just want to fetch the boxes up to the first 'mdat' and those should be fairly small.
+	HTTP->Parameters.SubRangeRequestSize = 256 << 10;
+	HTTP->ReceiveBuffer = ReadBuffer.ReceiveBuffer;
+	HTTP->ProgressListener = ProgressListener;
+	PlayerSessionServices->GetHTTPManager()->AddRequest(HTTP);
+
+	// Create the parser we need for parsing the "moov" box containing all the track information.
+	MP4Parser = IParserISO14496_12::CreateParser();
+
+	UEMediaError parseError = MP4Parser->ParseHeader(this, this, Options, PlayerSessionServices);
+
+	ProgressListener.Reset();
+	PlayerSessionServices->GetHTTPManager()->RemoveRequest(HTTP);
+
+	if (parseError != UEMEDIA_ERROR_ABORTED)
+	{
+		const HTTP::FConnectionInfo* ConnInfo = &HTTP->ConnectionInfo;
+
+		// Notify the download of the "master playlist". This indicates the download only, not the parsing thereof.
+		PlayerSessionServices->SendMessageToPlayer(IPlaylistReader::PlaylistDownloadMessage::Create(ConnInfo, Playlist::EListType::Master, Playlist::ELoadType::Initial));
+		// Notify that the "master playlist" has been parsed, successfully or not.
+		PlayerSessionServices->SendMessageToPlayer(IPlaylistReader::PlaylistLoadedMessage::Create(LastErrorDetail, ConnInfo, Playlist::EListType::Master, Playlist::ELoadType::Initial));
+
+		if (parseError == UEMEDIA_ERROR_OK)
+		{
+			// See that we have parsed all the boxes we need.
+			if (bFoundBoxFTYP && bFoundBoxMOOV)
+			{
+				// Prepare the tracks in the stream that are of a supported codec.
+				parseError = MP4Parser->PrepareTracks(TSharedPtrTS<const IParserISO14496_12>());
+				if (parseError == UEMEDIA_ERROR_OK)
+				{
+					Manifest = MakeSharedTS<FManifestMP4Internal>(PlayerSessionServices);
+					FErrorDetail err = Manifest->Build(MP4Parser, MasterPlaylistURL, *ConnInfo);
+
+					// Notify that the "variant playlists" are ready. There are no variants in an mp4, but this is the trigger that the playlists are all set up and are good to go now.
+					PlayerSessionServices->SendMessageToPlayer(IPlaylistReader::PlaylistLoadedMessage::Create(LastErrorDetail, ConnInfo, Playlist::EListType::Variant, Playlist::ELoadType::Initial));
+				}
+				else
+				{
+					PostError(FString::Printf(TEXT("Failed to parse tracks in mp4 \"%s\" with error %u"), *ConnInfo->EffectiveURL, parseError), ERRCODE_MP4_INVALID_FILE, UEMEDIA_ERROR_FORMAT_ERROR);
+				}
+			}
+			else
+			{
+				// No moov box usually means this is not a fast-start file.
+				PostError(FString::Printf(TEXT("No moov box found before mdat in \"%s\". Make sure the video is fast-startable."), *ConnInfo->EffectiveURL), ERRCODE_MP4_INVALID_FILE, UEMEDIA_ERROR_FORMAT_ERROR);
+			}
+		}
+		else
+		{
+			PostError(FString::Printf(TEXT("Failed to parse mp4 \"%s\" with error %u"), *ConnInfo->EffectiveURL, parseError), ERRCODE_MP4_INVALID_FILE, UEMEDIA_ERROR_FORMAT_ERROR);
+		}
+	}
+
+	HTTP.Reset();
+	ReadBuffer.Reset();
+
+	// This thread's work is done. We only wait for termination now.
+	WorkerThreadQuitSignal.Wait();
+}
+
+
+
+
+
+
+/**
+ * Read n bytes of data into the provided buffer.
+ *
+ * Reading must return the number of bytes asked to get, if necessary by blocking.
+ * If a read error prevents reading the number of bytes -1 must be returned.
+ *
+ * @param ToBuffer Buffer into which to store the data bytes. If nullptr is passed the data must be skipped over.
+ * @param NumBytes The number of bytes to read. Must not read more bytes and no less than requested.
+ * @return The number of bytes read or -1 on a read error.
+ */
+int64 FPlaylistReaderMP4::ReadData(void* ToBuffer, int64 NumBytes)
+{
+	FPODRingbuffer& SourceBuffer = ReadBuffer.ReceiveBuffer->Buffer;
+
+	uint8* OutputBuffer = (uint8*)ToBuffer;
+	// Do we have enough data in the ringbuffer to satisfy the read?
+	if (SourceBuffer.Num() >= NumBytes)
+	{
+		// Yes. Get the data and return.
+		int32 NumGot = SourceBuffer.PopData(OutputBuffer, NumBytes);
+		check(NumGot == NumBytes);
+		ReadBuffer.ParsePos += NumBytes;
+		return NumBytes;
+	}
+	else
+	{
+		// Do not have enough data yet or we want to read more than the ringbuffer can hold.
+		int32 NumBytesToGo = NumBytes;
+		while(NumBytesToGo > 0)
+		{
+			if (HasErrored() || SourceBuffer.WasAborted() || bAbort)
+			{
+				return -1;
+			}
+			// EOD?
+			if (SourceBuffer.IsEndOfData())
+			{
+				return 0;
+			}
+
+			// Get whatever amount of data is currently available to free up the buffer for receiving more data.
+			int32 NumGot = SourceBuffer.PopData(OutputBuffer, NumBytesToGo);
+			if ((NumBytesToGo -= NumGot) > 0)
+			{
+				if (OutputBuffer)
+				{
+					OutputBuffer += NumGot;
+				}
+				// Wait for data to arrive in the ringbuffer.
+				int32 WaitForBytes = NumBytesToGo > SourceBuffer.Capacity() ? SourceBuffer.Capacity() : NumBytesToGo;
+				SourceBuffer.WaitUntilSizeAvailable(WaitForBytes, 1000 * 100);
+			}
+		}
+		ReadBuffer.ParsePos += NumBytes;
+		return NumBytes;
+	}
+}
+
+/**
+ * Checks if the data source has reached the End Of File (EOF) and cannot provide any additional data.
+ *
+ * @return If EOF has been reached returns true, otherwise false.
+ */
+bool FPlaylistReaderMP4::HasReachedEOF() const
+{
+	FPODRingbuffer& SourceBuffer = ReadBuffer.ReceiveBuffer->Buffer;
+	return SourceBuffer.IsEndOfData();
+}
+
+/**
+ * Checks if reading of the file and therefor parsing has been aborted.
+ *
+ * @return true if reading/parsing has been aborted, false otherwise.
+ */
+bool FPlaylistReaderMP4::HasReadBeenAborted() const
+{
+	return bAbort || ReadBuffer.bHasBeenAborted;
+}
+
+/**
+ * Returns the current read offset.
+ *
+ * The first read offset is not necessarily zero. It could be anywhere inside the source.
+ *
+ * @return The current byte offset in the source.
+ */
+int64 FPlaylistReaderMP4::GetCurrentOffset() const
+{
+	return ReadBuffer.ParsePos;
+}
+
+
+IParserISO14496_12::IBoxCallback::EParseContinuation FPlaylistReaderMP4::OnFoundBox(IParserISO14496_12::FBoxType Box, int64 BoxSizeInBytes, int64 FileDataOffset, int64 BoxDataOffset)
+{
+	// The very first box must be an 'ftyp' box by definition. If this is not the case this is not a valid mp4 file.
+	if (FileDataOffset == 0 && Box != IParserISO14496_12::BoxType_ftyp)
+	{
+		PostError("Invalid mp4 file: first box is not 'ftyp'", ERRCODE_MP4_INVALID_FILE, UEMEDIA_ERROR_FORMAT_ERROR);
+		return IParserISO14496_12::IBoxCallback::EParseContinuation::Stop;
+	}
+
+	// Check which box is being parsed next.
+	switch(Box)
+	{
+		case IParserISO14496_12::BoxType_ftyp:
+		{
+			bFoundBoxFTYP = true;
+			return IParserISO14496_12::IBoxCallback::EParseContinuation::Continue;
+		}
+		case IParserISO14496_12::BoxType_moov:
+		{
+			bFoundBoxMOOV = true;
+			return IParserISO14496_12::IBoxCallback::EParseContinuation::Continue;
+		}
+		case IParserISO14496_12::BoxType_sidx:
+		{
+			bFoundBoxSIDX = true;
+			return IParserISO14496_12::IBoxCallback::EParseContinuation::Continue;
+		}
+		case IParserISO14496_12::BoxType_moof:
+		{
+			bFoundBoxMOOF = true;
+			return IParserISO14496_12::IBoxCallback::EParseContinuation::Stop;
+		}
+		case IParserISO14496_12::BoxType_mdat:
+		{
+			bFoundBoxMDAT = true;
+			return IParserISO14496_12::IBoxCallback::EParseContinuation::Stop;
+		}
+		default:
+		{
+			return IParserISO14496_12::IBoxCallback::EParseContinuation::Continue;
+		}
+	}
+}
+
+
+} // namespace Electra
+
+
diff --git a/Engine/Plugins/Media/ElectraPlayer/Source/ElectraPlayerRuntime/Private/Runtime/Player/mp4/StreamReaderMP4.cpp b/Engine/Plugins/Media/ElectraPlayer/Source/ElectraPlayerRuntime/Private/Runtime/Player/mp4/StreamReaderMP4.cpp
new file mode 100644
--- /dev/null
+++ b/Engine/Plugins/Media/ElectraPlayer/Source/ElectraPlayerRuntime/Private/Runtime/Player/mp4/StreamReaderMP4.cpp
@@ -0,0 +1,713 @@
+// Copyright Epic Games, Inc. All Rights Reserved.
+
+#include "PlayerCore.h"
+#include "Player/Manifest.h"
+#include "HTTP/HTTPManager.h"
+#include "Demuxer/ParserISO14496-12.h"
+#include "SynchronizedClock.h"
+#include "PlaylistReaderMP4.h"
+#include "Utilities/TimeUtilities.h"
+#include "HAL/LowLevelMemTracker.h"
+#include "ElectraPlayerPrivate.h"
+#include "Player/PlayerStreamReader.h"
+#include "Player/AdaptiveStreamingPlayerABR.h"
+#include "Player/mp4/ManifestMP4.h"
+#include "Player/mp4/StreamReaderMP4.h"
+
+
+
+
+
+namespace Electra
+{
+
+FStreamSegmentRequestMP4::FStreamSegmentRequestMP4()
+{
+	PrimaryStreamType      = EStreamType::Video;
+	FileStartOffset 	   = -1;
+	FileEndOffset   	   = -1;
+	Bitrate 			   = 0;
+	PlaybackSequenceID     = ~0U;
+	bStartingOnMOOF 	   = false;
+	bIsContinuationSegment = false;
+	bAllTracksAtEOS 	   = false;
+	CurrentIteratorBytePos = 0;
+	NumOverallRetries      = 0;
+}
+
+FStreamSegmentRequestMP4::~FStreamSegmentRequestMP4()
+{
+}
+
+void FStreamSegmentRequestMP4::SetPlaybackSequenceID(uint32 InPlaybackSequenceID)
+{
+	PlaybackSequenceID = InPlaybackSequenceID;
+}
+
+uint32 FStreamSegmentRequestMP4::GetPlaybackSequenceID() const
+{
+	return PlaybackSequenceID;
+}
+
+EStreamType FStreamSegmentRequestMP4::GetType() const
+{
+	return  PrimaryStreamType;
+}
+
+void FStreamSegmentRequestMP4::GetDependentStreams(TArray<FDependentStreams>& OutDependentStreams) const
+{
+	// Those are not "real" dependent streams in that they are multiplexed and do not need to be fetched from
+	// a different source. This merely indicates the types of non-primary streams we will be demuxing.
+	OutDependentStreams = DependentStreams;
+}
+
+void FStreamSegmentRequestMP4::GetEndedStreams(TArray<TSharedPtrTS<IStreamSegment>>& OutAlreadyEndedStreams)
+{
+	OutAlreadyEndedStreams.Empty();
+	if (bAllTracksAtEOS)
+	{
+		OutAlreadyEndedStreams.Push(SharedThis(this));
+		for(int32 i=0; i<DependentStreams.Num(); ++i)
+		{
+			FStreamSegmentRequestMP4* DepReq = new FStreamSegmentRequestMP4;
+			// Only need to set the stream type here.
+			DepReq->PrimaryStreamType = DependentStreams[i].StreamType;
+			TSharedPtrTS<IStreamSegment> p(DepReq);
+			OutAlreadyEndedStreams.Push(p);
+		}
+	}
+}
+
+FTimeValue FStreamSegmentRequestMP4::GetFirstPTS() const
+{
+	return FirstPTS;
+}
+
+int32 FStreamSegmentRequestMP4::GetQualityIndex() const
+{
+	// No quality choice here.
+	return 0;
+}
+
+int32 FStreamSegmentRequestMP4::GetBitrate() const
+{
+	return Bitrate;
+}
+
+void FStreamSegmentRequestMP4::GetDownloadStats(Metrics::FSegmentDownloadStats& OutStats) const
+{
+	OutStats = DownloadStats;
+}
+
+
+
+
+
+
+uint32	FStreamReaderMP4::UniqueDownloadID = 1;
+
+FStreamReaderMP4::FStreamReaderMP4()
+{
+	bIsStarted  		  = false;
+	bTerminate  		  = false;
+	bRequestCanceled	  = false;
+	bHasErrored 		  = false;
+}
+
+FStreamReaderMP4::~FStreamReaderMP4()
+{
+	Close();
+}
+
+UEMediaError FStreamReaderMP4::Create(IPlayerSessionServices* InPlayerSessionService, const CreateParam &InCreateParam)
+{
+	if (!InCreateParam.MemoryProvider ||
+		!InCreateParam.EventListener ||
+		!InCreateParam.PlayerSessionService)
+	{
+		return UEMEDIA_ERROR_BAD_ARGUMENTS;
+	}
+
+	Parameters = InCreateParam;
+	bTerminate = false;
+	bIsStarted = true;
+
+	ThreadSetPriority(InCreateParam.ReaderConfig.ThreadParam.Priority);
+	ThreadSetCoreAffinity(InCreateParam.ReaderConfig.ThreadParam.CoreAffinity);
+	ThreadSetStackSize(InCreateParam.ReaderConfig.ThreadParam.StackSize);
+	ThreadSetName("ElectraPlayer::MP4 streamer");
+	ThreadStart(Electra::MakeDelegate(this, &FStreamReaderMP4::WorkerThread));
+
+	return UEMEDIA_ERROR_OK;
+}
+
+void FStreamReaderMP4::Close()
+{
+	if (bIsStarted)
+	{
+		bIsStarted = false;
+
+		TSharedPtrTS<FStreamSegmentRequestMP4>	Request = CurrentRequest;
+		if (Request.IsValid())
+		{
+			// ...
+		}
+		CancelRequests();
+		bTerminate = true;
+		WorkSignal.Signal();
+		ThreadWaitDone();
+		ThreadReset();
+		CurrentRequest.Reset();
+	}
+}
+
+
+void FStreamReaderMP4::LogMessage(IInfoLog::ELevel Level, const FString& Message)
+{
+	if (Parameters.PlayerSessionService)
+	{
+		Parameters.PlayerSessionService->PostLog(Facility::EFacility::MP4StreamReader, Level, Message);
+	}
+}
+
+IStreamReader::EAddResult FStreamReaderMP4::AddRequest(uint32 CurrentPlaybackSequenceID, TSharedPtrTS<IStreamSegment> InRequest)
+{
+	TSharedPtrTS<FStreamSegmentRequestMP4>	Request = CurrentRequest;
+	if (Request.IsValid())
+	{
+		check(!"why is the handler busy??");
+		return IStreamReader::EAddResult::TryAgainLater;
+	}
+	Request = StaticCastSharedPtr<FStreamSegmentRequestMP4>(InRequest);
+	Request->SetPlaybackSequenceID(CurrentPlaybackSequenceID);
+	bRequestCanceled = false;
+	bHasErrored = false;
+	// Only add the request if it is not an all-EOS one!
+	if (!Request->bAllTracksAtEOS)
+	{
+		CurrentRequest = Request;
+		WorkSignal.Signal();
+	}
+	return EAddResult::Added;
+}
+
+void FStreamReaderMP4::CancelRequests()
+{
+	bRequestCanceled = true;
+	ReadBuffer.Abort();
+}
+
+void FStreamReaderMP4::PauseDownload()
+{
+	// Not implemented.
+}
+
+void FStreamReaderMP4::ResumeDownload()
+{
+	// Not implemented.
+}
+
+bool FStreamReaderMP4::HasBeenAborted() const
+{
+	return bRequestCanceled || ReadBuffer.bAbort;
+}
+
+bool FStreamReaderMP4::HasErrored() const
+{
+	return bHasErrored;
+}
+
+int32 FStreamReaderMP4::HTTPProgressCallback(const IElectraHttpManager::FRequest* InRequest)
+{
+	HTTPUpdateStats(MEDIAutcTime::Current(), InRequest);
+	// Aborted?
+	return HasBeenAborted() ? 1 : 0;
+}
+
+void FStreamReaderMP4::HTTPCompletionCallback(const IElectraHttpManager::FRequest* InRequest)
+{
+	HTTPUpdateStats(FTimeValue::GetInvalid(), InRequest);
+	bHasErrored = InRequest->ConnectionInfo.StatusInfo.ErrorDetail.IsError();
+	if (bHasErrored)
+	{
+		ReadBuffer.SetHasErrored();
+	}
+}
+
+void FStreamReaderMP4::HTTPUpdateStats(const FTimeValue& CurrentTime, const IElectraHttpManager::FRequest* Request)
+{
+	TSharedPtrTS<FStreamSegmentRequestMP4> SegmentRequest = CurrentRequest;
+	if (SegmentRequest.IsValid())
+	{
+		FMediaCriticalSection::ScopedLock lock(MetricUpdateLock);
+		SegmentRequest->ConnectionInfo = Request->ConnectionInfo;
+		// Update the current download stats which we report periodically to the ABR.
+		Metrics::FSegmentDownloadStats& ds = SegmentRequest->DownloadStats;
+		if (Request->ConnectionInfo.EffectiveURL.Len())
+		{
+			ds.URL 			  = Request->ConnectionInfo.EffectiveURL;
+		}
+		ds.HTTPStatusCode     = Request->ConnectionInfo.StatusInfo.HTTPStatus;
+		ds.TimeToFirstByte    = Request->ConnectionInfo.TimeUntilFirstByte;
+		ds.TimeToDownload     = ((CurrentTime.IsValid() ? CurrentTime : Request->ConnectionInfo.RequestEndTime) - Request->ConnectionInfo.RequestStartTime).GetAsSeconds();
+		ds.ByteSize 		  = Request->ConnectionInfo.ContentLength;
+		ds.NumBytesDownloaded = Request->ConnectionInfo.BytesReadSoFar;
+	}
+}
+
+void FStreamReaderMP4::WorkerThread()
+{
+	LLM_SCOPE(ELLMTag::ElectraPlayer);
+	while(!bTerminate)
+	{
+		WorkSignal.WaitAndReset();
+		if (bTerminate)
+		{
+			break;
+		}
+
+		TSharedPtrTS<FStreamSegmentRequestMP4>	Request = CurrentRequest;
+		if (!Request.IsValid())
+		{
+			continue;
+		}
+
+		FManifestMP4Internal::FTimelineAssetMP4* TimelineAsset = static_cast<FManifestMP4Internal::FTimelineAssetMP4*>(Request->MediaAsset.Get());
+
+		// Clear the active track map.
+		ActiveTrackMap.Reset();
+
+// FIXME: If looping to somewhere within the stream instead of the beginning at zero the offset here must be made relative to the DTS of the first sample we demux.
+//        Otherwise there would be a large jump ahead in time!
+		FTimeValue LoopTimestampOffset = Request->PlayerLoopState.LoopBasetime;
+		TSharedPtr<const FPlayerLoopState, ESPMode::ThreadSafe>		PlayerLoopState = MakeShared<const FPlayerLoopState, ESPMode::ThreadSafe>(Request->PlayerLoopState);
+
+
+		// Get the list of all the tracks that have been selected in the asset.
+		// This does not mean their data will be _used_ for playback, only that the track is usable by the player
+		// with regards to type and codec.
+		struct FPlaylistTrackMetadata
+		{
+			EStreamType		Type;
+			FString			Language;
+			FString			PeriodID;
+			FString			AdaptationSetID;
+			FString			RepresentationID;
+			FString			CDN;
+			int32			Bitrate;
+		};
+		TMap<uint32, FPlaylistTrackMetadata>	SelectedTrackMap;
+		const EStreamType TypesOfSupportedTracks[] = { EStreamType::Video, EStreamType::Audio, EStreamType::Subtitle };
+		for(int32 nStrType=0; nStrType<UE_ARRAY_COUNT(TypesOfSupportedTracks); ++nStrType)
+		{
+			int32 NumAdapt = Request->MediaAsset->GetNumberOfAdaptationSets(TypesOfSupportedTracks[nStrType]);
+			for(int32 nAdapt=0; nAdapt<NumAdapt; ++nAdapt)
+			{
+				TSharedPtrTS<IPlaybackAssetAdaptationSet> AdaptationSet = Request->MediaAsset->GetAdaptationSetByTypeAndIndex(TypesOfSupportedTracks[nStrType], nAdapt);
+				FString Language = AdaptationSet->GetLanguage();
+				FString AdaptID = AdaptationSet->GetUniqueIdentifier();
+				int32 NumRepr = AdaptationSet->GetNumberOfRepresentations();
+				for(int32 nRepr=0; nRepr<NumRepr; ++nRepr)
+				{
+					TSharedPtrTS<IPlaybackAssetRepresentation> Representation = AdaptationSet->GetRepresentationByIndex(nRepr);
+					// Note: By definition the representations unique identifier is a string of the numeric track ID and can thus be parsed back into a number.
+					FString ReprID = Representation->GetUniqueIdentifier();
+					uint32 TrackId;
+					LexFromString(TrackId, *ReprID); 
+					FPlaylistTrackMetadata tmd;
+					tmd.Type			 = TypesOfSupportedTracks[nStrType];
+					tmd.Language		 = Language;
+					tmd.PeriodID		 = Request->MediaAsset->GetUniqueIdentifier();
+					tmd.AdaptationSetID  = AdaptID;
+					tmd.RepresentationID = ReprID;
+					tmd.Bitrate 		 = Representation->GetBitrate();
+					tmd.CDN 			 = Representation->GetCDN();
+					SelectedTrackMap.Emplace(TrackId, tmd);
+				}
+			}
+		}
+
+
+		Metrics::FSegmentDownloadStats& ds = Request->DownloadStats;
+		ds.StatsID = FMediaInterlockedIncrement(UniqueDownloadID);
+
+		check(Request->PrimaryTrackIterator.IsValid() && Request->PrimaryTrackIterator->GetTrack());
+		uint32 PrimaryTrackID = Request->PrimaryTrackIterator->GetTrack()->GetID();
+		const FPlaylistTrackMetadata* PrimaryTrackMetadata = SelectedTrackMap.Find(PrimaryTrackID);
+		check(PrimaryTrackMetadata);
+		if (PrimaryTrackMetadata)
+		{
+			ds.MediaAssetID 	= PrimaryTrackMetadata->PeriodID;
+			ds.AdaptationSetID  = PrimaryTrackMetadata->AdaptationSetID;
+			ds.RepresentationID = PrimaryTrackMetadata->RepresentationID;
+			ds.Bitrate  		= PrimaryTrackMetadata->Bitrate;
+			ds.CDN  			= PrimaryTrackMetadata->CDN;
+		}
+
+		ds.FailureReason.Empty();
+		ds.bWasSuccessful      = true;
+		ds.bWasAborted  	   = false;
+		ds.bDidTimeout  	   = false;
+		ds.HTTPStatusCode      = 0;
+		ds.StreamType   	   = Request->GetType();
+		ds.SegmentType  	   = Metrics::ESegmentType::Media;
+		ds.PresentationTime    = (Request->FirstPTS + LoopTimestampOffset).GetAsSeconds();
+		ds.Duration 		   = Request->SegmentDuration.GetAsSeconds();
+		ds.DurationDownloaded  = 0.0;
+		ds.DurationDelivered   = 0.0;
+		ds.TimeToFirstByte     = 0.0;
+		ds.TimeToDownload      = 0.0;
+		ds.ByteSize 		   = -1;
+		ds.NumBytesDownloaded  = 0;
+		ds.ThroughputBps	   = 0;
+		ds.bInsertedFillerData = false;
+		ds.URL  			   = TimelineAsset->GetMediaURL();
+		ds.bIsMissingSegment   = false;
+		ds.bParseFailure	   = false;
+		ds.RetryNumber  	   = Request->NumOverallRetries;
+
+		Parameters.EventListener->OnFragmentOpen(Request);
+
+		TSharedPtrTS<IElectraHttpManager::FProgressListener>	ProgressListener;
+		ProgressListener = MakeSharedTS<IElectraHttpManager::FProgressListener>();
+		ProgressListener->CompletionDelegate = Electra::MakeDelegate(this, &FStreamReaderMP4::HTTPCompletionCallback);
+		ProgressListener->ProgressDelegate   = Electra::MakeDelegate(this, &FStreamReaderMP4::HTTPProgressCallback);
+
+		ReadBuffer.Reset();
+		ReadBuffer.ReceiveBuffer = MakeSharedTS<IElectraHttpManager::FReceiveBuffer>();
+		// Set the receive buffer to an okay-ish size. Too small and the file I/O may block too often and get too slow.
+		ReadBuffer.ReceiveBuffer->Buffer.Reserve(4 << 20);
+		ReadBuffer.ReceiveBuffer->bEnableRingbuffer = true;
+		ReadBuffer.SetCurrentPos(Request->FileStartOffset);
+		TSharedPtrTS<IElectraHttpManager::FRequest> HTTP(new IElectraHttpManager::FRequest);
+		HTTP->Parameters.URL				= TimelineAsset->GetMediaURL();
+		HTTP->Parameters.Range.Start		= Request->FileStartOffset;
+		HTTP->Parameters.Range.EndIncluding = Request->FileEndOffset;
+		HTTP->Parameters.SubRangeRequestSize = 1 << 20;
+		HTTP->ReceiveBuffer 				= ReadBuffer.ReceiveBuffer;
+		HTTP->ProgressListener  			= ProgressListener;
+		Parameters.PlayerSessionService->GetHTTPManager()->AddRequest(HTTP);
+
+
+		FTimeValue DurationSuccessfullyDelivered(FTimeValue::GetZero());
+		FTimeValue DurationSuccessfullyRead(FTimeValue::GetZero());
+		FTimeValue NextLargestExpectedTimestamp(FTimeValue::GetZero());
+		bool bDone = false;
+		TSharedPtr<IParserISO14496_12::IAllTrackIterator, ESPMode::ThreadSafe> AllTrackIterator = TimelineAsset->GetMoovBoxParser()->CreateAllTrackIteratorByFilePos(Request->FileStartOffset);
+		while(!bDone && !HasErrored() && !HasBeenAborted() && !bTerminate)
+		{
+			auto UpdateSelectedTrack = [&SelectedTrackMap](const IParserISO14496_12::ITrackIterator* trkIt, TMap<uint32, FSelectedTrackData>& ActiveTrks) -> FSelectedTrackData&
+			{
+				const IParserISO14496_12::ITrack* Track = trkIt->GetTrack();
+				uint32 tkid = Track->GetID();
+
+				// Check if this track ID is already in our map of active tracks.
+				FSelectedTrackData& st = ActiveTrks.FindOrAdd(tkid);
+				if (!st.StreamSourceInfo.IsValid())
+				{
+					auto meta = MakeShared<FStreamSourceInfo, ESPMode::ThreadSafe>();
+					meta->NumericTrackID = tkid;
+
+					// Check if this track is in the list of selected tracks.
+					const FPlaylistTrackMetadata* SelectedTrackMetadata = SelectedTrackMap.Find(tkid);
+					if (SelectedTrackMetadata)
+					{
+						st.bIsSelectedTrack = true;
+						st.StreamType = SelectedTrackMetadata->Type;
+						//meta->Role			 = SelectedTrackMetadata->;
+						meta->Language  		 = SelectedTrackMetadata->Language;
+						meta->PeriodID  		 = SelectedTrackMetadata->PeriodID;
+						meta->AdaptationSetID    = SelectedTrackMetadata->AdaptationSetID;
+						meta->RepresentationID   = SelectedTrackMetadata->RepresentationID;
+					}
+					st.StreamSourceInfo = MoveTemp(meta);
+				}
+				if (!st.CSD.IsValid())
+				{
+					TSharedPtrTS<FAccessUnit::CodecData> CSD(new FAccessUnit::CodecData);
+					CSD->CodecSpecificData = Track->GetCodecSpecificData();
+					CSD->RawCSD			= Track->GetCodecSpecificDataRAW();
+					CSD->ParsedInfo 		= Track->GetCodecInformation();
+					st.CSD = MoveTemp(CSD);
+				}
+				return st;
+			};
+
+			// Handle all the new tracks that have reached EOS while iterating. We do this first here to
+			// handle the tracks that hit EOS before reaching the intended start position.
+			TArray<const IParserISO14496_12::ITrackIterator*> TracksAtEOS;
+			AllTrackIterator->GetNewEOSTracks(TracksAtEOS);
+			AllTrackIterator->ClearNewEOSTracks();
+			for(int32 nTrk=0; nTrk<TracksAtEOS.Num(); ++nTrk)
+			{
+				const IParserISO14496_12::ITrackIterator* TrackIt = TracksAtEOS[nTrk];
+				FSelectedTrackData& SelectedTrack = UpdateSelectedTrack(TrackIt, ActiveTrackMap);
+				// Is this a track that is selected and we are interested in?
+				if (SelectedTrack.bIsSelectedTrack)
+				{
+					Parameters.EventListener->OnFragmentReachedEOS(SelectedTrack.StreamType, SelectedTrack.StreamSourceInfo);
+				}
+			}
+
+			// Handle current track iterator
+			const IParserISO14496_12::ITrackIterator* TrackIt = AllTrackIterator->Current();
+			if (TrackIt)
+			{
+				FSelectedTrackData& SelectedTrack = UpdateSelectedTrack(TrackIt, ActiveTrackMap);
+
+				// Get the sample properties
+				uint32 SampleNumber    = TrackIt->GetSampleNumber();
+				int64 DTS   		   = TrackIt->GetDTS();
+				int64 PTS   		   = TrackIt->GetPTS();
+				int64 Duration  	   = TrackIt->GetDuration();
+				uint32 Timescale	   = TrackIt->GetTimescale();
+				bool bIsSyncSample     = TrackIt->IsSyncSample();
+				int64 SampleSize	   = TrackIt->GetSampleSize();
+				int64 SampleFileOffset = TrackIt->GetSampleFileOffset();
+
+				// Remember at which file position we are currently at. In case of failure this is where we will retry.
+				Request->CurrentIteratorBytePos = SampleFileOffset;
+
+				// Do we need to skip over some data?
+				if (SampleFileOffset > ReadBuffer.GetCurrentPos())
+				{
+					int32 NumBytesToSkip = SampleFileOffset - ReadBuffer.GetCurrentPos();
+					int64 nr = ReadBuffer.ReadTo(nullptr, NumBytesToSkip);
+					if (nr != NumBytesToSkip)
+					{
+						bDone = true;
+						break;
+					}
+				}
+				else if (SampleFileOffset < ReadBuffer.GetCurrentPos())
+				{
+					ds.bParseFailure  = true;
+					Request->ConnectionInfo.StatusInfo.ErrorDetail.SetMessage(FString::Printf(TEXT("Segment parse error. Sample offset %lld for sample #%u in track %u is before the current read position at %lld"), (long long int)SampleFileOffset, SampleNumber, TrackIt->GetTrack()->GetID(), (long long int)ReadBuffer.GetCurrentPos()));
+					bHasErrored = true;
+					break;
+				}
+
+				// Do we read the sample because the track is selected or do we discard it?
+				if (SelectedTrack.bIsSelectedTrack)
+				{
+					// Is this a sync sample?
+					if (bIsSyncSample && !SelectedTrack.bGotKeyframe)
+					{
+						SelectedTrack.bGotKeyframe = true;
+					}
+					// Do we need to skip samples from this track until we reach a sync sample?
+					bool bSkipUntilSyncSample = !SelectedTrack.bGotKeyframe && !Request->bIsContinuationSegment;
+
+					// Create an access unit.
+					FAccessUnit *AccessUnit = FAccessUnit::Create(Parameters.MemoryProvider);
+					if (AccessUnit)
+					{
+						AccessUnit->ESType = SelectedTrack.StreamType;
+						AccessUnit->PTS.SetFromND(PTS, Timescale);
+						AccessUnit->DTS.SetFromND(DTS, Timescale);
+						AccessUnit->Duration.SetFromND(Duration, Timescale);
+						AccessUnit->AUSize = (uint32) SampleSize;
+						AccessUnit->AUCodecData = SelectedTrack.CSD;
+						AccessUnit->DropState = FAccessUnit::EDropState::None;
+						// If this is a continuation then we must not tag samples as being too early.
+						if (!Request->bIsContinuationSegment)
+						{
+							if (AccessUnit->DTS < Request->FirstPTS)
+							{
+								AccessUnit->DropState |= FAccessUnit::EDropState::DtsTooEarly;
+							}
+							if (AccessUnit->PTS < Request->FirstPTS)
+							{
+								AccessUnit->DropState |= FAccessUnit::EDropState::PtsTooEarly;
+							}
+						}
+						// FIXME: if we only want to read a partial segment we could set drop state based on the sample being 'too late'.
+
+						// Apply timestamp offsets for looping after checking the timestamp limits.
+						AccessUnit->PTS += LoopTimestampOffset;
+						AccessUnit->DTS += LoopTimestampOffset;
+
+						AccessUnit->bIsFirstInSequence = SelectedTrack.bIsFirstInSequence;
+						AccessUnit->bIsSyncSample = bIsSyncSample;
+						AccessUnit->bIsDummyData = false;
+						AccessUnit->AUData = AccessUnit->AllocatePayloadBuffer(AccessUnit->AUSize);
+
+						// Set the associated stream metadata
+						AccessUnit->StreamSourceInfo = SelectedTrack.StreamSourceInfo;
+						AccessUnit->PlayerLoopState = PlayerLoopState;
+
+						SelectedTrack.bIsFirstInSequence = false;
+
+						int64 nr = ReadBuffer.ReadTo(AccessUnit->AUData, (int32)SampleSize);
+						if (nr == SampleSize)
+						{
+							SelectedTrack.DurationSuccessfullyRead += AccessUnit->Duration;
+
+							//LogMessage(IInfoLog::ELevel::Info, FString::Printf("[%u] %4u: DTS=%lld PTS=%lld dur=%lld sync=%d; %lld bytes @ %lld", tkid, SampleNumber, (long long int)AccessUnit->mDTS.GetAsMicroseconds(), (long long int)AccessUnit->mPTS.GetAsMicroseconds(), (long long int)AccessUnit->mDuration.GetAsMicroseconds(), bIsSyncSample?1:0, (long long int)SampleSize, (long long int)SampleFileOffset));
+
+							// Keep track of the next expected sample PTS and remember the largest value of all tracks.
+							FTimeValue NextExpectedPTS = AccessUnit->PTS + AccessUnit->Duration;
+							if (NextExpectedPTS > NextLargestExpectedTimestamp)
+							{
+								NextLargestExpectedTimestamp = NextExpectedPTS;
+							}
+
+							bool bSentOff = false;
+							while(!bSentOff && !HasBeenAborted() && !bTerminate)
+							{
+								if (Parameters.EventListener->OnFragmentAccessUnitReceived(AccessUnit))
+								{
+									SelectedTrack.DurationSuccessfullyDelivered += AccessUnit->Duration;
+									bSentOff = true;
+									AccessUnit = nullptr;
+
+									// Since we have delivered this access unit, if we are detecting an error now we need to then
+									// retry on the _next_ AU and not this one again!
+									Request->CurrentIteratorBytePos = SampleFileOffset + SampleSize;
+								}
+								else
+								{
+									FMediaRunnable::SleepMicroseconds(1000 * 10);
+								}
+							}
+
+							// Release the AU if we still have it.
+							FAccessUnit::Release(AccessUnit);
+							AccessUnit = nullptr;
+
+							// For error handling, if we managed to get additional data we reset the retry count.
+							if (ds.RetryNumber && SelectedTrack.DurationSuccessfullyRead.GetAsSeconds() > 2.0)
+							{
+								ds.RetryNumber = 0;
+								Request->NumOverallRetries = 0;
+							}
+						}
+						else
+						{
+							// Did not get the number of bytes we needed. Either because of a read error or because we got aborted.
+							FAccessUnit::Release(AccessUnit);
+							AccessUnit = nullptr;
+							bDone = true;
+							break;
+						}
+					}
+					else
+					{
+						// TODO: Throw OOM error
+					}
+				}
+			}
+			else
+			{
+				break;
+			}
+			AllTrackIterator->Next();
+		}
+
+		// Remove the download request.
+		ProgressListener.Reset();
+		Parameters.PlayerSessionService->GetHTTPManager()->RemoveRequest(HTTP);
+		Request->ConnectionInfo = HTTP->ConnectionInfo;
+		HTTP.Reset();
+
+		// Remember the next largest timestamp from all tracks.
+		Request->NextLargestExpectedTimestamp = NextLargestExpectedTimestamp;
+
+
+		// Set downloaded and delivered duration from the primary track.
+		FSelectedTrackData& PrimaryTrack = ActiveTrackMap.FindOrAdd(PrimaryTrackID);
+		DurationSuccessfullyRead	  = PrimaryTrack.DurationSuccessfullyRead;
+		DurationSuccessfullyDelivered = PrimaryTrack.DurationSuccessfullyDelivered;
+
+		// Set up remaining download stat fields.
+// Note: currently commented out because of UE-88612.
+//       This must be reinstated once we set failure reasons in the loop above so they won't get replaced by this!
+//		if (ds.FailureReason.length() == 0)
+		{
+			ds.FailureReason = Request->ConnectionInfo.StatusInfo.ErrorDetail.GetMessage();
+		}
+//		if (bAbortedByABR)
+//		{
+//			// If aborted set the reason as the download failure.
+//			ds.FailureReason = ds.ABRState.ProgressDecision.Reason;
+//		}
+//		ds.bWasAborted  	  = bAbortedByABR;
+//		ds.bWasSuccessful     = !bHasErrored && !bAbortedByABR;
+		ds.bWasSuccessful     = !bHasErrored;
+		ds.URL  			  = Request->ConnectionInfo.EffectiveURL;
+		ds.HTTPStatusCode     = Request->ConnectionInfo.StatusInfo.HTTPStatus;
+		ds.DurationDownloaded = DurationSuccessfullyRead.GetAsSeconds();
+		ds.DurationDelivered  = DurationSuccessfullyDelivered.GetAsSeconds();
+		ds.TimeToFirstByte    = Request->ConnectionInfo.TimeUntilFirstByte;
+		ds.TimeToDownload     = (Request->ConnectionInfo.RequestEndTime - Request->ConnectionInfo.RequestStartTime).GetAsSeconds();
+		ds.ByteSize 		  = Request->ConnectionInfo.ContentLength;
+		ds.NumBytesDownloaded = Request->ConnectionInfo.BytesReadSoFar;
+		ds.ThroughputBps	  = Request->ConnectionInfo.Throughput.GetThroughput();
+		if (ds.ThroughputBps == 0)
+		{
+			ds.ThroughputBps = ds.TimeToDownload > 0.0 ? 8 * ds.NumBytesDownloaded / ds.TimeToDownload : 0;
+		}
+
+		ActiveTrackMap.Reset();
+
+		// Reset the current request so another one can be added immediately when we call OnFragmentClose()
+		CurrentRequest.Reset();
+		Parameters.PlayerSessionService->GetStreamSelector()->ReportDownloadEnd(ds);
+		Parameters.EventListener->OnFragmentClose(Request);
+	}
+}
+
+
+
+int32 FStreamReaderMP4::FReadBuffer::ReadTo(void* ToBuffer, int32 NumBytes)
+{
+	FPODRingbuffer& SourceBuffer = ReceiveBuffer->Buffer;
+
+	uint8* OutputBuffer = (uint8*)ToBuffer;
+	// Do we have enough data in the ringbuffer to satisfy the read?
+	if (SourceBuffer.Num() >= NumBytes)
+	{
+		// Yes. Get the data and return.
+		int32 NumGot = SourceBuffer.PopData(OutputBuffer, NumBytes);
+		check(NumGot == NumBytes);
+		CurrentPos += NumBytes;
+		return NumBytes;
+	}
+	else
+	{
+		// Do not have enough data yet or we want to read more than the ringbuffer can hold.
+		int32 NumBytesToGo = NumBytes;
+		while(NumBytesToGo > 0)
+		{
+			if (bHasErrored || SourceBuffer.WasAborted() || bAbort)
+			{
+				return -1;
+			}
+			// EOD?
+			if (SourceBuffer.IsEndOfData())
+			{
+				return 0;
+			}
+
+			// Get whatever amount of data is currently available to free up the buffer for receiving more data.
+			int32 NumGot = SourceBuffer.PopData(OutputBuffer, NumBytesToGo);
+			if ((NumBytesToGo -= NumGot) > 0)
+			{
+				if (OutputBuffer)
+				{
+					OutputBuffer += NumGot;
+				}
+				// Wait for data to arrive in the ringbuffer.
+				int32 WaitForBytes = NumBytesToGo > SourceBuffer.Capacity() ? SourceBuffer.Capacity() : NumBytesToGo;
+				SourceBuffer.WaitUntilSizeAvailable(WaitForBytes, 1000 * 100);
+			}
+		}
+		CurrentPos += NumBytes;
+		return NumBytes;
+	}
+}
+
+
+} // namespace Electra
+
+
diff --git a/Engine/Plugins/Media/ElectraPlayer/Source/ElectraPlayerRuntime/Private/Runtime/Player/mp4/StreamReaderMP4.h b/Engine/Plugins/Media/ElectraPlayer/Source/ElectraPlayerRuntime/Private/Runtime/Player/mp4/StreamReaderMP4.h
new file mode 100644
--- /dev/null
+++ b/Engine/Plugins/Media/ElectraPlayer/Source/ElectraPlayerRuntime/Private/Runtime/Player/mp4/StreamReaderMP4.h
@@ -0,0 +1,181 @@
+// Copyright Epic Games, Inc. All Rights Reserved.
+
+#pragma once
+
+#include "PlayerCore.h"
+#include "Player/Manifest.h"
+#include "HTTP/HTTPManager.h"
+#include "Demuxer/ParserISO14496-12.h"
+#include "SynchronizedClock.h"
+#include "PlaylistReaderMP4.h"
+#include "Utilities/HashFunctions.h"
+#include "Utilities/TimeUtilities.h"
+#include "HAL/LowLevelMemTracker.h"
+#include "ElectraPlayerPrivate.h"
+#include "Player/PlayerStreamReader.h"
+#include "StreamAccessUnitBuffer.h"
+
+
+
+
+namespace Electra
+{
+
+
+class FStreamSegmentRequestMP4 : public IStreamSegment
+{
+public:
+	FStreamSegmentRequestMP4();
+	virtual ~FStreamSegmentRequestMP4();
+
+	virtual void SetPlaybackSequenceID(uint32 InPlaybackSequenceID) override;
+	virtual uint32 GetPlaybackSequenceID() const override;
+
+	virtual EStreamType GetType() const override;
+
+	virtual void GetDependentStreams(TArray<FDependentStreams>& OutDependentStreams) const override;
+	virtual void GetEndedStreams(TArray<TSharedPtrTS<IStreamSegment>>& OutAlreadyEndedStreams) override;
+
+	//! Returns the first PTS value as indicated by the media timeline. This should correspond to the actual absolute PTS of the sample.
+	virtual FTimeValue GetFirstPTS() const override;
+
+	virtual int32 GetQualityIndex() const override;
+	virtual int32 GetBitrate() const override;
+
+	virtual void GetDownloadStats(Metrics::FSegmentDownloadStats& OutStats) const override;
+
+
+
+	TSharedPtrTS<ITimelineMediaAsset>					MediaAsset;					//!< The entire mp4 asset
+	TSharedPtrTS<IParserISO14496_12::ITrackIterator>	PrimaryTrackIterator;
+	TArray<FDependentStreams>							DependentStreams;
+
+	FTimeValue											FirstPTS;					//!< The PTS of the first sample
+	FTimeValue											SegmentDuration;
+	EStreamType											PrimaryStreamType;
+	int64												FileStartOffset;			//!< Where to start in the file
+	int64												FileEndOffset;				//!< Where to end in the file (for HTTP range GET requests)
+	uint32												PlaybackSequenceID;
+	int32												Bitrate;
+	bool												bStartingOnMOOF;			//!< when using the 'sidx' a segment is expected to start on a 'moof', otherwise inside an 'mdat'.
+	bool												bIsContinuationSegment;
+
+	bool												bAllTracksAtEOS;
+
+	FPlayerLoopState									PlayerLoopState;
+	int32												NumOverallRetries;			//!< Number of retries
+
+	Metrics::FSegmentDownloadStats						DownloadStats;
+	HTTP::FConnectionInfo								ConnectionInfo;
+	FTimeValue											NextLargestExpectedTimestamp;	//!< Largest timestamp of all samples (plus its duration) across all tracks.
+	int64												CurrentIteratorBytePos;
+};
+
+
+/**
+ * This class implements an interface to read from an mp4 file.
+ */
+class FStreamReaderMP4 : public IStreamReader, public FMediaThread
+{
+public:
+	FStreamReaderMP4();
+	virtual ~FStreamReaderMP4();
+	virtual UEMediaError Create(IPlayerSessionServices* PlayerSessionService, const CreateParam &createParam) override;
+	virtual void Close() override;
+	virtual EAddResult AddRequest(uint32 CurrentPlaybackSequenceID, TSharedPtrTS<IStreamSegment> Request) override;
+	virtual void CancelRequests() override;
+	virtual void PauseDownload() override;
+	virtual void ResumeDownload() override;
+
+private:
+	void WorkerThread();
+	void LogMessage(IInfoLog::ELevel Level, const FString& Message);
+	int32 HTTPProgressCallback(const IElectraHttpManager::FRequest* InRequest);
+	void HTTPCompletionCallback(const IElectraHttpManager::FRequest* InRequest);
+	void HTTPUpdateStats(const FTimeValue& CurrentTime, const IElectraHttpManager::FRequest* Request);
+
+	bool HasBeenAborted() const;
+	bool HasErrored() const;
+
+	struct FReadBuffer
+	{
+		FReadBuffer()
+		{
+			Reset();
+		}
+		void Reset()
+		{
+			ReceiveBuffer.Reset();
+			CurrentPos = 0;
+			bAbort = false;
+			bHasErrored = false;
+		}
+		void SetCurrentPos(int64 InPos)
+		{
+			CurrentPos = InPos;
+		}
+		int64 GetCurrentPos()
+		{
+			return CurrentPos;
+		}
+		void Abort()
+		{
+			bAbort = true;
+			TSharedPtrTS<IElectraHttpManager::FReceiveBuffer>	Buffer = ReceiveBuffer;
+			if (Buffer.IsValid())
+			{
+				Buffer->Buffer.Abort();
+			}
+		}
+		void SetHasErrored()
+		{
+			bHasErrored = true;
+		}
+		int32 ReadTo(void* ToBuffer, int32 NumBytes);
+		TSharedPtrTS<IElectraHttpManager::FReceiveBuffer>	ReceiveBuffer;
+		int64												CurrentPos;
+		bool												bAbort;
+		bool												bHasErrored;
+	};
+
+	struct FSelectedTrackData
+	{
+		FSelectedTrackData()
+			: StreamType(EStreamType::Video), bGotKeyframe(false), bIsSelectedTrack(false), bIsFirstInSequence(true)
+		{
+			DurationSuccessfullyRead.SetToZero();
+			DurationSuccessfullyDelivered.SetToZero();
+		}
+		TSharedPtrTS<FAccessUnit::CodecData>					CSD;
+		TSharedPtr<FStreamSourceInfo, ESPMode::ThreadSafe>		StreamSourceInfo;
+		EStreamType												StreamType;
+		bool													bGotKeyframe;
+		bool													bIsSelectedTrack;
+		bool													bIsFirstInSequence;
+		FTimeValue												DurationSuccessfullyRead;
+		FTimeValue												DurationSuccessfullyDelivered;
+	};
+
+	CreateParam									Parameters;
+	IPlayerSessionServices*						PlayerSessionServices;
+	bool										bIsStarted;
+	bool										bTerminate;
+	bool										bRequestCanceled;
+	bool										bHasErrored;
+	FErrorDetail								ErrorDetail;
+
+	TSharedPtrTS<FStreamSegmentRequestMP4>		CurrentRequest;
+	FMediaEvent									WorkSignal;
+	FReadBuffer									ReadBuffer;
+	FMediaCriticalSection						MetricUpdateLock;
+
+	TMap<uint32, FSelectedTrackData>			ActiveTrackMap;
+
+	static uint32								UniqueDownloadID;
+};
+
+
+
+} // namespace Electra
+
+
