ElectraPlayer: RELEASE. Moved from NFL to public. Note that this is done via a Delete and "Add New" operation as discussed.

diff --git a/Engine/Plugins/Media/ElectraPlayer/Source/ElectraPlayerRuntime/Private/Runtime/Renderer/RendererAudio.cpp b/Engine/Plugins/Media/ElectraPlayer/Source/ElectraPlayerRuntime/Private/Runtime/Renderer/RendererAudio.cpp
new file mode 100644
--- /dev/null
+++ b/Engine/Plugins/Media/ElectraPlayer/Source/ElectraPlayerRuntime/Private/Runtime/Renderer/RendererAudio.cpp
@@ -0,0 +1,257 @@
+// Copyright Epic Games, Inc. All Rights Reserved.
+
+#include "Core/MediaEventSignal.h"
+#include "Core/MediaInterlocked.h"
+
+#include "Renderer/RendererAudio.h"
+
+#include "Misc/Timespan.h"
+#include "ElectraPlayer.h"
+#include "ElectraPlayerPrivate.h"
+
+void FElectraRendererAudio::SampleReleasedToPool(FTimespan DurationToRelease)
+{
+	FPlatformAtomics::InterlockedDecrement(&NumOutputAudioBuffersInUse);
+	check(NumOutputAudioBuffersInUse >= 0 && NumOutputAudioBuffersInUse <= NumBuffers)
+}
+
+FElectraRendererAudio::SystemConfiguration::SystemConfiguration()
+{
+}
+
+bool FElectraRendererAudio::Startup(const FElectraRendererAudio::SystemConfiguration& Configuration)
+{
+	return true;
+}
+
+void FElectraRendererAudio::Shutdown(void)
+{
+}
+
+FElectraRendererAudio::FElectraRendererAudio(TSharedPtr<FElectraPlayer, ESPMode::ThreadSafe> InPlayer)
+	: Player(InPlayer)
+	, MaxBufferSize(0)
+	, NumOutputAudioBuffersInUse(0)
+	, NumBuffers(0)
+	, NumBuffersAcquiredForDecoder(0)
+{
+}
+
+FElectraRendererAudio::~FElectraRendererAudio()
+{
+	// We need to track this. Otherwise the audio buffers in flight may cause corruption...
+//	check(NumOutputAudioBuffersInUse == 0);
+//TRUE? THE POOL GONE IS CHECKED. THE RENDERER GONE, TOO... SOOO?
+}
+
+FElectraRendererAudio::OpenError FElectraRendererAudio::Open(const FElectraRendererAudio::InstanceConfiguration& Config)
+{
+	return OpenError::eOpen_Ok;
+}
+
+
+void FElectraRendererAudio::Close(void)
+{
+}
+
+
+void FElectraRendererAudio::DetachPlayer()
+{
+	Player.Reset();
+}
+
+
+const Electra::FParamDict& FElectraRendererAudio::GetBufferPoolProperties() const
+{
+	return BufferPoolProperties;
+}
+
+
+UEMediaError FElectraRendererAudio::CreateBufferPool(const Electra::FParamDict& Parameters)
+{
+	// Create a new pool and overwrite all "pending" values...
+	DecoderOutputPool.Reset();
+	NumOutputAudioBuffersInUse = 0;
+
+	const FVariantValue& variantNumBuffers = Parameters.GetValue("num_buffers");
+	if (!variantNumBuffers.IsType(FVariantValue::EDataType::TypeInt64))
+		return UEMEDIA_ERROR_BAD_ARGUMENTS;
+	NumBuffers = (int32)variantNumBuffers.GetInt64();
+
+	// Update dic for later query via GetBufferPoolProperties
+	BufferPoolProperties.Set("max_buffers", Electra::FVariantValue((int64)NumBuffers));
+
+	const FVariantValue& variantMaxBufferSize = Parameters.GetValue("max_buffer_size");
+	if (!variantMaxBufferSize.IsType(FVariantValue::EDataType::TypeInt64))
+		return UEMEDIA_ERROR_BAD_ARGUMENTS;
+
+	MaxBufferSize = (uint32)variantMaxBufferSize.GetInt64();
+
+	return UEMEDIA_ERROR_OK;
+}
+
+
+UEMediaError FElectraRendererAudio::AcquireBuffer(IBuffer*& OutBuffer, int32 TimeoutInMicroseconds, const Electra::FParamDict& InParameters)
+{
+	if (TimeoutInMicroseconds != 0)
+		return UEMEDIA_ERROR_BAD_ARGUMENTS;
+
+	// Check if we are inside the max buffer limit...
+	if (NumOutputAudioBuffersInUse >= NumBuffers)
+	{
+		return UEMEDIA_ERROR_INSUFFICIENT_DATA;
+	}
+
+	IAudioDecoderOutputPtr DecoderOutput = DecoderOutputPool.AcquireShared();
+	DecoderOutput->SetOwner(SharedThis(this));
+
+	FPlatformAtomics::InterlockedIncrement(&NumBuffersAcquiredForDecoder);
+	DecoderOutput->Reserve(MaxBufferSize);
+
+	FPlatformAtomics::InterlockedIncrement(&NumOutputAudioBuffersInUse);
+
+	// UE handles all with shared pointers...
+	FMediaBufferSharedPtrWrapper* MediaBufferSharedPtrWrapper = new FMediaBufferSharedPtrWrapper(DecoderOutput);
+	MediaBufferSharedPtrWrapper->BufferProperties.Set("size", FVariantValue((int64)DecoderOutput->GetReservedBufferBytes()));
+	MediaBufferSharedPtrWrapper->BufferProperties.Set("address", FVariantValue(const_cast<void*>(DecoderOutput->GetBuffer())));
+	OutBuffer = MediaBufferSharedPtrWrapper;
+
+	return UEMEDIA_ERROR_OK;
+}
+
+
+
+UEMediaError FElectraRendererAudio::ReturnBuffer(IBuffer* Buffer, bool bRender, const FParamDict& InSampleProperties)
+{
+	if (Buffer == nullptr)
+		return UEMEDIA_ERROR_BAD_ARGUMENTS;
+
+	FMediaBufferSharedPtrWrapper* MediaBufferSharedPtrWrapper = static_cast<FMediaBufferSharedPtrWrapper*>(Buffer);
+
+	if (bRender)
+	{
+		const FVariantValue& variantNumChannels = InSampleProperties.GetValue("num_channels");
+		if (!variantNumChannels.IsType(FVariantValue::EDataType::TypeInt64))
+			return UEMEDIA_ERROR_BAD_ARGUMENTS;
+
+		const FVariantValue& variantSampleRate = InSampleProperties.GetValue("sample_rate");
+		if (!variantSampleRate.IsType(FVariantValue::EDataType::TypeInt64))
+			return UEMEDIA_ERROR_BAD_ARGUMENTS;
+
+		const FVariantValue& variantBufferUsedBytes = InSampleProperties.GetValue("byte_size");
+		if (!variantBufferUsedBytes.IsType(FVariantValue::EDataType::TypeInt64))
+			return UEMEDIA_ERROR_BAD_ARGUMENTS;
+
+		const FVariantValue& variantPts = InSampleProperties.GetValue("pts");
+		if (!variantPts.IsType(FVariantValue::EDataType::TypeTimeValue))
+			return UEMEDIA_ERROR_BAD_ARGUMENTS;
+
+		const FVariantValue& variantDuration = InSampleProperties.GetValue("duration");
+		if (!variantDuration.IsType(FVariantValue::EDataType::TypeTimeValue))
+			return UEMEDIA_ERROR_BAD_ARGUMENTS;
+
+		IAudioDecoderOutputPtr DecoderOutput = MediaBufferSharedPtrWrapper->DecoderOutput;
+
+		uint32 InNumChannels = (uint32)variantNumChannels.GetInt64();
+		uint32 InSampleRate = (uint32)variantSampleRate.GetInt64();
+		uint32 InUsedBufferBytes = (uint32)variantBufferUsedBytes.GetInt64();
+
+		FTimespan InDuration = FTimespan(variantDuration.GetTimeValue().GetAsHNS());
+		FTimespan InPts = FTimespan(variantPts.GetTimeValue().GetAsHNS());
+
+		//UE_LOG(LogElectraPlayer, VeryVerbose, TEXT("-- FElectraRendererAudio::ReturnBuffer: Audio sample for time %s"), *InPts.ToString(TEXT("%h:%m:%s.%f")));
+
+		DecoderOutput->Initialize(IAudioDecoderOutput::ESampleFormat::Int16, InNumChannels, InSampleRate, InDuration, FDecoderTimeStamp(InPts, 0), InUsedBufferBytes);
+
+		// Push buffer to output queue...
+		if (TSharedPtr<FElectraPlayer, ESPMode::ThreadSafe> PinnedPlayer = Player.Pin())
+		{
+			PinnedPlayer->OnAudioDecoded(DecoderOutput);
+		}
+	}
+
+	// Render or not.. here we do not have any error. The decoder is done with this buffer
+	FPlatformAtomics::InterlockedDecrement(&NumBuffersAcquiredForDecoder);
+
+	// In all cases free the wrapper class...
+	delete MediaBufferSharedPtrWrapper;
+
+	return UEMEDIA_ERROR_OK;
+}
+
+void FElectraRendererAudio::SetRenderClock(const FTimespan& RenderTime)
+{
+	// Latest audio as time reference
+	Electra::FTimeValue AudioRenderTime;
+	AudioRenderTime.SetFromHNS(RenderTime.GetTicks());
+	RenderClock->SetCurrentTime(Electra::IMediaRenderClock::ERendererType::Audio, AudioRenderTime);
+	RenderClock->SetCurrentTime(Electra::IMediaRenderClock::ERendererType::Video, AudioRenderTime);
+}
+
+UEMediaError FElectraRendererAudio::ReleaseBufferPool()
+{
+	DecoderOutputPool.Reset();
+	return UEMEDIA_ERROR_OK;
+}
+
+bool FElectraRendererAudio::CanReceiveOutputFrames(uint64 NumFrames) const
+{
+	TSharedPtr<FElectraPlayer, ESPMode::ThreadSafe> PinnedPlayer = Player.Pin();
+	if (!PinnedPlayer.IsValid())
+	{
+		return false;
+	}
+	return PinnedPlayer->CanPresentAudioFrames(NumFrames);
+}
+
+/**
+ *
+ */
+void FElectraRendererAudio::SetRenderClock(TSharedPtr<Electra::IMediaRenderClock, ESPMode::ThreadSafe> InRenderClock)
+{
+	RenderClock = InRenderClock;
+}
+
+void FElectraRendererAudio::SetNextApproximatePresentationTime(const Electra::FTimeValue& NextApproxPTS)
+{
+}
+
+// Flushes all pending buffers not yet rendered
+UEMediaError FElectraRendererAudio::Flush(const Electra::FParamDict& InOptions)
+{
+	// [...]If there are still frames a decoder has not returned yet UEMEDIA_ERROR_INTERNAL should be returned.[...]
+	if (NumBuffersAcquiredForDecoder != 0)
+	{
+		return UEMEDIA_ERROR_INTERNAL;
+	}
+
+	if (TSharedPtr<FElectraPlayer, ESPMode::ThreadSafe> PinnedPlayer = Player.Pin())
+	{
+		PinnedPlayer->OnAudioFlush();
+	}
+
+	return UEMEDIA_ERROR_OK;
+}
+
+// Begins rendering of the first sample buffer
+void FElectraRendererAudio::StartRendering(const Electra::FParamDict& InOptions)
+{
+	//UE_LOG(LogElectraPlayer, Log, TEXT("-- FElectraRendererAudio::StartRendering (NumBuffers = %d)"), NumBuffers);
+	if (TSharedPtr<FElectraPlayer, ESPMode::ThreadSafe> PinnedPlayer = Player.Pin())
+	{
+		PinnedPlayer->OnAudioRenderingStarted();
+	}
+}
+
+
+// Stops rendering of sample buffers
+void FElectraRendererAudio::StopRendering(const Electra::FParamDict& InOptions)
+{
+	//UE_LOG(LogElectraPlayer, Log, TEXT("-- FElectraRendererAudio::StopRendering (NumBuffers = %d)"), NumBuffers);
+	if (TSharedPtr<FElectraPlayer, ESPMode::ThreadSafe> PinnedPlayer = Player.Pin())
+	{
+		PinnedPlayer->OnAudioRenderingStopped();
+	}
+}
+
+
diff --git a/Engine/Plugins/Media/ElectraPlayer/Source/ElectraPlayerRuntime/Private/Runtime/Renderer/RendererAudio.h b/Engine/Plugins/Media/ElectraPlayer/Source/ElectraPlayerRuntime/Private/Runtime/Renderer/RendererAudio.h
new file mode 100644
--- /dev/null
+++ b/Engine/Plugins/Media/ElectraPlayer/Source/ElectraPlayerRuntime/Private/Runtime/Renderer/RendererAudio.h
@@ -0,0 +1,140 @@
+// Copyright Epic Games, Inc. All Rights Reserved.
+
+#pragma once
+
+#include "Templates/SharedPointer.h"
+#include "Containers/Queue.h"
+
+#include "Renderer/RendererBase.h"
+#include "Core/MediaThreads.h"
+#include "PlayerTime.h"
+
+#include "MediaAudioDecoderOutput.h"
+
+struct FTimespan;
+class FElectraPlayer;
+
+class FElectraPlayerPlatformAudioDecoderOutputFactory
+{
+public:
+	static IAudioDecoderOutput* Create();
+};
+
+
+class FElectraRendererAudio : public Electra::IMediaRenderer, public TSharedFromThis<FElectraRendererAudio, ESPMode::ThreadSafe>
+{
+public:
+	struct SystemConfiguration
+	{
+		SystemConfiguration();
+	};
+	static bool Startup(const SystemConfiguration& Configuration);
+	static void Shutdown(void);
+
+	FElectraRendererAudio(TSharedPtr<FElectraPlayer, ESPMode::ThreadSafe> InPlayer);
+	virtual ~FElectraRendererAudio();
+
+	struct InstanceConfiguration
+	{
+		InstanceConfiguration()
+		{
+			// todo
+		}
+	};
+
+	enum OpenError
+	{
+		eOpen_Ok,
+		eOpen_Error,
+	};
+
+	OpenError Open(const InstanceConfiguration& Config);
+	void Close(void);
+
+	//-------------------------------------------------------------------------
+	// Methods for UE interface
+	//
+	void DetachPlayer();
+
+
+	//-------------------------------------------------------------------------
+	// Methods for Electra::IMediaRenderer
+	//
+
+	// Returns the properties of the buffer pool. Those properties should not change
+	virtual const Electra::FParamDict& GetBufferPoolProperties() const override;
+
+	// Create a buffer pool from where a decoder can get the block of memory to decode into.
+	virtual UEMediaError CreateBufferPool(const Electra::FParamDict& Parameters) override;
+
+	// Asks for a sample buffer from the buffer pool created previously through CreateBufferPool()
+	virtual UEMediaError AcquireBuffer(IBuffer*& OutBuffer, int32 TimeoutInMicroseconds, const Electra::FParamDict& InParameters) override;
+
+	// Releases the buffer for rendering and subsequent return to the buffer pool
+	virtual UEMediaError ReturnBuffer(IBuffer* Buffer, bool bRender, const Electra::FParamDict& InSampleProperties) override;
+
+	// Informs that the decoder is done with this pool. NO FREE!!!
+	virtual UEMediaError ReleaseBufferPool() override;
+
+	virtual bool CanReceiveOutputFrames(uint64 NumFrames) const override;
+
+	// Receives the render clock we need to update with the most recently rendered sample's timestamp.
+	virtual void SetRenderClock(TSharedPtr<Electra::IMediaRenderClock, ESPMode::ThreadSafe> RenderClock) override;
+
+	// Sets the next expected sample's approximate presentation time stamp
+	virtual void SetNextApproximatePresentationTime(const Electra::FTimeValue& NextApproxPTS) override;
+
+	// Flushes all pending buffers not yet rendered
+	virtual UEMediaError Flush(const Electra::FParamDict& InOptions) override;
+
+	// Begins rendering of the first sample buffer
+	virtual void StartRendering(const Electra::FParamDict& InOptions) override;
+
+	// Stops rendering of sample buffers
+	virtual void StopRendering(const Electra::FParamDict& InOptions) override;
+
+
+	//! Called when audio sample is returned to pool
+	virtual void SampleReleasedToPool(FTimespan DurationToRelease) override;
+
+	//! Set used clock to some external value
+	void SetRenderClock(const FTimespan& RenderTime);
+
+private:
+
+	TWeakPtr<FElectraPlayer, ESPMode::ThreadSafe> Player;
+	TDecoderOutputObjectPool<IAudioDecoderOutput, FElectraPlayerPlatformAudioDecoderOutputFactory> DecoderOutputPool;
+
+	uint32					 MaxBufferSize;
+	int32				     NumOutputAudioBuffersInUse;
+
+	class FMediaBufferSharedPtrWrapper : public Electra::IMediaRenderer::IBuffer
+	{
+	public:
+		explicit FMediaBufferSharedPtrWrapper(IAudioDecoderOutputPtr InDecoderOutput)
+			: DecoderOutput(InDecoderOutput)
+		{
+		}
+
+		~FMediaBufferSharedPtrWrapper()
+		{
+		}
+
+		virtual const Electra::FParamDict& GetBufferProperties() const override
+		{
+			return BufferProperties;
+		}
+		Electra::FParamDict BufferProperties;
+		IAudioDecoderOutputPtr DecoderOutput;
+	};
+
+	// ------------------------------------------------------------------------------------------------------
+
+	Electra::FParamDict							BufferPoolProperties;
+	int32										NumBuffers;
+	int32										NumBuffersAcquiredForDecoder;
+
+	TSharedPtr<Electra::IMediaRenderClock, ESPMode::ThreadSafe>	RenderClock;
+};
+
+
diff --git a/Engine/Plugins/Media/ElectraPlayer/Source/ElectraPlayerRuntime/Private/Runtime/Renderer/RendererBase.h b/Engine/Plugins/Media/ElectraPlayer/Source/ElectraPlayerRuntime/Private/Runtime/Renderer/RendererBase.h
new file mode 100644
--- /dev/null
+++ b/Engine/Plugins/Media/ElectraPlayer/Source/ElectraPlayerRuntime/Private/Runtime/Renderer/RendererBase.h
@@ -0,0 +1,216 @@
+// Copyright Epic Games, Inc. All Rights Reserved.
+
+#pragma once
+
+#include "PlayerCore.h"
+#include "PlayerTime.h"
+#include "ParameterDictionary.h"
+
+#include "MediaDecoderOutput.h"
+
+namespace Electra
+{
+
+	class IMediaRenderClock
+	{
+	public:
+		virtual ~IMediaRenderClock() = default;
+
+		enum class ERendererType
+		{
+			Video,
+			Audio,
+			Subtitles,
+		};
+
+		/**
+		 * Called by the renderer to set the time of the most sample that has been output most recently.
+		 *
+		 * @param ForRenderer
+		 *               Identifies the type of renderer from which the time is set.
+		 * @param CurrentRenderTime
+		 *               The time of the sample most recently output.
+		 */
+		virtual void SetCurrentTime(ERendererType ForRenderer, const FTimeValue& CurrentRenderTime) = 0;
+
+		/**
+		 * Gets the current _interpolated_ sample output time from the last
+		 * time SetCurrentTime() was called plus the elapsed time since then.
+		 *
+		 * @param FromRenderer
+		 *               Identifies the type of renderer for which to get the time.
+		 *
+		 * @return Interpolated render time
+		 */
+		virtual FTimeValue GetInterpolatedRenderTime(ERendererType FromRenderer) = 0;
+	};
+
+
+
+
+
+	class IMediaRenderer : public IDecoderOutputOwner
+	{
+	public:
+		virtual ~IMediaRenderer() = default;
+
+		/**
+		 * Buffer interface
+		 */
+		class IBuffer
+		{
+		public:
+			virtual ~IBuffer() = default;
+
+			/**
+			 * Returns the properties of the buffer.
+			 *
+			 * @return Const reference to the dictionary containing buffer properties.
+			 */
+			virtual const FParamDict& GetBufferProperties() const = 0;
+		};
+
+
+		//=================================================================================================================
+		// Methods called from both a decoder and the player from their respective threads.
+		//
+
+		/**
+		 * Returns the properties of the buffer pool. Those properties should not change
+		 * during the lifetime of the pool.
+		 * A mandatory property is the maximum number of buffers that can be obtained from
+		 * the pool ("max_buffers":int64).
+		 *
+		 * @return Const reference to the dictionary containing buffer pool properties.
+		 */
+		virtual const FParamDict& GetBufferPoolProperties() const = 0;
+
+
+
+		//=================================================================================================================
+		// Methods called from a decoder (and from within the decoder thread)
+		//
+
+		/**
+		 * Create a buffer pool from where a decoder can get the block of memory to decode into.
+		 *
+		 * @param Parameters Dictionary with create options depending on the type of
+		 *                   render buffers to create (video or audio)
+		 *
+		 * @return One of UEMEDIA_ERROR_OK, UEMEDIA_ERROR_BAD_ARGUMENTS or UEMEDIA_ERROR_OOM (or others if warranted).
+		 */
+		virtual UEMediaError CreateBufferPool(const FParamDict& Parameters) = 0;
+
+		/**
+		 * Asks for a sample buffer from the buffer pool created previously through CreateBufferPool().
+		 *
+		 * @param OutBuffer Receives the address of the returned IBuffer
+		 * @param TimeoutInMicroseconds
+		 *                  Maximum time to wait for a buffer to become available before giving up.
+		 *                  If 0 a "peek" if a buffer is available is to be performed and a buffer
+		 *                  returned if one is available.
+		 *                  No value indicates an infinite wait time!
+		 * @param InParameters
+		 *                  Optional parameters for buffer properties if different from those to CreateBufferPool().
+		 *
+		 * @return One of UEMEDIA_ERROR_OK, UEMEDIA_ERROR_BAD_ARGUMENTS or UEMEDIA_ERROR_INSUFFICIENT_DATA.
+		 *         Other values if warranted but those will cause a playback error.
+		 */
+		virtual UEMediaError AcquireBuffer(IBuffer*& OutBuffer, int32 TimeoutInMicroseconds, const FParamDict& InParameters) = 0;
+
+		/**
+		 * Releases the buffer for rendering and subsequent return to the buffer pool.
+		 *
+		 * @param Buffer  The buffer being returned to the renderer for rendering.
+		 * @param bRender true to render the contents of the buffer, false to discard them.
+		 * @param InSampleProperties
+		 *                Properties of the decoded sample (eg. width and height for video, sample rate for audio)
+		 *
+		 * @return Should be UEMEDIA_ERROR_OK only. Anything else will cause a playback error (like UEMEDIA_ERROR_BAD_ARGUMENTS for bad properties)
+		 */
+		virtual UEMediaError ReturnBuffer(IBuffer* Buffer, bool bRender, const FParamDict& InSampleProperties) = 0;
+
+		/**
+		 * Informs that the decoder is done with this pool. This only indicates no calls to
+		 * GetBuffer() and ReleaseBuffer() will be made by this decoder instance any more.
+		 * It is up to the renderer to decide whether or not to destroy the pool or keep it
+		 * around for re-use with a new decoder instance (the next CreateBufferPool() call).
+		 *
+		 * @return Should be UEMEDIA_ERROR_OK only. Anything else will cause a playback error.
+		 */
+		virtual UEMediaError ReleaseBufferPool() = 0;
+
+
+		/**
+		 * Returns if the renderer's output queues are ready to receive any output from a decode
+		 *
+		 * @return True if ready, false if decode should be delayed.
+		 */
+		virtual bool CanReceiveOutputFrames(uint64 NumFrames) const = 0;
+
+		//=================================================================================================================
+		// Methods called from the player (and from within the player thread)
+		//
+
+		/**
+		 * Sets the render clock which this renderer needs to update with the presentation time
+		 * of the most recently rendered sample.
+		 *
+		 * @param RenderClock
+		 *               Render clock to update with the most recent output sample time.
+		 */
+		virtual void SetRenderClock(TSharedPtr<IMediaRenderClock, ESPMode::ThreadSafe> RenderClock) = 0;
+
+		/**
+		 * Sets the next expected sample's approximate presentation time stamp.
+		 * This is called as data buffering begins at the indicated play start
+		 * position. The actual timestamp of the first sample will vary.
+		 * The renderer can ignore this value and rely only on the timestamp
+		 * of the decoded sample.
+		 *
+		 * @param NextApproxPTS
+		 *               Approximate timestamp at which playback is about to commence.
+		 */
+		virtual void SetNextApproximatePresentationTime(const FTimeValue& NextApproxPTS) = 0;
+
+		/**
+		 * Flushes all pending buffers not yet rendered.
+		 * There should be no outstanding buffers the decoder has not yet released.
+		 * Options may indicate to hold the current frame on-screen (for video) or
+		 * to tear it down.
+		 *
+		 * @param InOptions Options indicating flush operation, like keeping the current video frame on-screen.
+		 *
+		 * @return Should be UEMEDIA_ERROR_OK only.
+		 *         If there are still frames a decoder has not returned yet UEMEDIA_ERROR_INTERNAL should be returned.
+		 *         Any failure will cause a playback error.
+		 */
+		virtual UEMediaError Flush(const FParamDict& InOptions) = 0;
+
+		/**
+		 * Begins rendering of the first sample buffer.
+		 *
+		 * @param InOptions Reserved for future use.
+		 */
+		virtual void StartRendering(const FParamDict& InOptions) = 0;
+
+		/**
+		 * Stops rendering of sample buffers.
+		 *
+		 * @param InOptions Reserved for future use.
+		 */
+		virtual void StopRendering(const FParamDict& InOptions) = 0;
+
+		/**
+		 * Work any regular tasks on the output buffer pool
+		 *
+		 * @note Expected to be called as needed from (normally) the decoder worker thread
+		 */
+		virtual void TickOutputBufferPool() {};
+	};
+
+} // namespace Electra
+
+
+
+
diff --git a/Engine/Plugins/Media/ElectraPlayer/Source/ElectraPlayerRuntime/Private/Runtime/Renderer/RendererVideo.cpp b/Engine/Plugins/Media/ElectraPlayer/Source/ElectraPlayerRuntime/Private/Runtime/Renderer/RendererVideo.cpp
new file mode 100644
--- /dev/null
+++ b/Engine/Plugins/Media/ElectraPlayer/Source/ElectraPlayerRuntime/Private/Runtime/Renderer/RendererVideo.cpp
@@ -0,0 +1,307 @@
+// Copyright Epic Games, Inc. All Rights Reserved.
+
+#include "Misc/Build.h"
+#include <Core/MediaTypes.h>
+
+#include "Renderer/RendererVideo.h"
+
+#include "CoreGlobals.h"
+#include "Misc/Timespan.h"
+#include "HAL/UnrealMemory.h"
+
+#include "ElectraPlayer.h"
+#include "ElectraPlayerPrivate.h"
+
+DECLARE_DWORD_COUNTER_STAT(TEXT("Electra::MEDIArendererVideoUE_NumUsedOutputVideoSamples"), STAT_ElectraPlayer_MEDIArendererVideoUE_NumUsedOutputVideoSamples, STATGROUP_ElectraPlayer);
+
+// -----------------------------------------------------------------------------------------------------------------------
+
+void FElectraRendererVideo::SampleReleasedToPool(FTimespan durationToRelease)
+{
+	FPlatformAtomics::InterlockedDecrement(&NumOutputTexturesInUse);
+}
+
+
+FElectraRendererVideo::SystemConfiguration::SystemConfiguration()
+{
+}
+
+
+bool FElectraRendererVideo::Startup(const FElectraRendererVideo::SystemConfiguration& Configuration)
+{
+	return true;
+}
+
+void FElectraRendererVideo::Shutdown(void)
+{
+}
+
+
+
+FElectraRendererVideo::FElectraRendererVideo(TSharedPtr<FElectraPlayer, ESPMode::ThreadSafe> InPlayer)
+	: Player(InPlayer)
+	, NumOutputTexturesInUse(0)
+	, NumBuffers(0)
+	, NumBuffersAcquiredForDecoder(0)
+{
+}
+
+FElectraRendererVideo::~FElectraRendererVideo()
+{
+	// Manually delete all images in all queues. Well, the queues will are also deleted automatically.
+	QueueTickedAndWaitingForDecoder.Empty();
+}
+
+
+FElectraRendererVideo::OpenError FElectraRendererVideo::Open(const FElectraRendererVideo::InstanceConfiguration& Config)
+{
+	return OpenError::eOpen_Ok;
+}
+
+
+void FElectraRendererVideo::Close(void)
+{
+}
+
+
+void FElectraRendererVideo::DetachPlayer()
+{
+	Player.Reset();
+}
+
+
+void FElectraRendererVideo::PrepareForDecoderShutdown()
+{
+	DecoderOutputPool.PrepareForDecoderShutdown();
+}
+
+
+const Electra::FParamDict& FElectraRendererVideo::GetBufferPoolProperties() const
+{
+	return BufferPoolProperties;
+}
+
+
+/**
+ * Get some image from the media Pool.
+ * Please note that this "pool" is only used for reusing images. There is no upper limit in this pool.
+ * For this we have the separate counter
+ */
+void FElectraRendererVideo::AcquireFromPool(FVideoDecoderOutputPtr& DelayedImage)
+{
+	// Acquire a decoder output item
+	DelayedImage = DecoderOutputPool.AcquireShared();
+	DelayedImage->SetOwner(SharedThis(this));
+
+	FPlatformAtomics::InterlockedIncrement(&NumOutputTexturesInUse);
+}
+
+/**
+ *
+ */
+UEMediaError FElectraRendererVideo::CreateBufferPool(const Electra::FParamDict& Parameters)
+{
+	const FVariantValue& variantNumBuffers = Parameters.GetValue("num_buffers");
+	if (!variantNumBuffers.IsType(FVariantValue::EDataType::TypeInt64))
+		return UEMEDIA_ERROR_BAD_ARGUMENTS;
+	int32 RequestedNumBuffers = (int32)variantNumBuffers.GetInt64();
+
+	// Update dict for later query via GetBufferPoolProperties
+	BufferPoolProperties.Set("max_buffers", Electra::FVariantValue((int64)RequestedNumBuffers));
+
+	// Currently, we only handle enlargement of the buffer pool. If the size should shrink,
+	if (RequestedNumBuffers != NumBuffers)
+	{
+		// Preallocate the buffer structures, but do not do the actual memory allocation!
+		// This allows any AcquireBuffer command to immediately use a buffer from the queue here.
+		for (int32 i = NumBuffers; i < RequestedNumBuffers; ++i)
+		{
+			FVideoDecoderOutputPtr DelayedImage;
+			AcquireFromPool(DelayedImage);
+			QueueTickedAndWaitingForDecoder.Enqueue(DelayedImage);
+
+		}
+		NumBuffers = RequestedNumBuffers;
+	}
+	return UEMEDIA_ERROR_OK;
+}
+
+
+/**
+ * Asks for a sample buffer from the buffer pool created previously through CreateBufferPool
+ */
+UEMediaError FElectraRendererVideo::AcquireBuffer(IBuffer*& OutBuffer, int32 TimeoutInMicroseconds, const Electra::FParamDict& InParameters)
+{
+	if (TimeoutInMicroseconds != 0)
+	{
+		return UEMEDIA_ERROR_BAD_ARGUMENTS;
+	}
+
+	// Trigger removal of any old frames in the presentation queue of the player so we have all buffers we can have available
+	TSharedPtr<FElectraPlayer, ESPMode::ThreadSafe> PinnedPlayer = Player.Pin();
+	if (PinnedPlayer.IsValid())
+	{
+		PinnedPlayer->DropOldFramesFromPresentationQueue();
+	}
+
+	// Check if we have any buffer which we can hand out to the decoder...
+	if (QueueTickedAndWaitingForDecoder.IsEmpty())
+	{
+		// Check GLOBAL number of output textures in flight...
+		if (NumOutputTexturesInUse < NumBuffers)
+		{
+			// Allocate an image from the pool, but do NOT use it immediately because it COULD be still rendering...
+			FVideoDecoderOutputPtr DelayedImage;
+			// Acquire image and with information about how many render frames to wait before usage is possible
+			AcquireFromPool(DelayedImage);
+			QueueTickedAndWaitingForDecoder.Enqueue(DelayedImage);
+		}
+
+		if (QueueTickedAndWaitingForDecoder.IsEmpty())
+		{
+			return UEMEDIA_ERROR_INSUFFICIENT_DATA;
+		}
+	}
+
+	FVideoDecoderOutputPtr DelayedImage;
+	QueueTickedAndWaitingForDecoder.Dequeue(DelayedImage);
+	check(DelayedImage.IsValid());
+
+	FMediaBufferSharedPtrWrapper* MediaBufferSharedPtrWrapper = new FMediaBufferSharedPtrWrapper(DelayedImage);
+	check(MediaBufferSharedPtrWrapper);
+	MediaBufferSharedPtrWrapper->BufferProperties.Set("texture", FVariantValue(DelayedImage));
+	OutBuffer = MediaBufferSharedPtrWrapper;
+
+	FPlatformAtomics::InterlockedIncrement(&NumBuffersAcquiredForDecoder);
+
+	return UEMEDIA_ERROR_OK;
+}
+
+/**
+ * Releases the buffer for rendering and subsequent return to the buffer pool
+ */
+UEMediaError FElectraRendererVideo::ReturnBuffer(IBuffer* Buffer, bool bRender, const FParamDict& InSampleProperties)
+{
+	if (Buffer == nullptr)
+	{
+		return UEMEDIA_ERROR_BAD_ARGUMENTS;
+	}
+
+	FMediaBufferSharedPtrWrapper* MediaBufferSharedPtrWrapper = static_cast<FMediaBufferSharedPtrWrapper*>(Buffer);
+
+	if (bRender)
+	{
+		//OPT/CHANGE: Note that "MediaBufferSharedPtrWrapper->DecoderOutput->GetDict()" is the very same as InSampleProperties!
+		bool bIsDummyBuffer = InSampleProperties.GetValue("is_dummy").SafeGetBool(false);
+
+		// Put frame into output queue...
+		if (TSharedPtr<FElectraPlayer, ESPMode::ThreadSafe> PinnedPlayer = Player.Pin())
+		{
+			// This call pushes the image back to the game thread for processing...
+			PinnedPlayer->OnVideoDecoded(MediaBufferSharedPtrWrapper->DecoderOutput, bIsDummyBuffer);
+		}
+	}
+
+	// Render or not.. here we do not have any error. The decoder is done with this buffer
+	FPlatformAtomics::InterlockedDecrement(&NumBuffersAcquiredForDecoder);
+
+	// In all cases free the wrapper class...
+	delete MediaBufferSharedPtrWrapper;
+
+	return UEMEDIA_ERROR_OK;
+}
+
+UEMediaError FElectraRendererVideo::ReleaseBufferPool()
+{
+	DecoderOutputPool.Reset();
+	return UEMEDIA_ERROR_OK;
+}
+
+/**
+ *
+ */
+bool FElectraRendererVideo::CanReceiveOutputFrames(uint64 NumFrames) const
+{
+	TSharedPtr<FElectraPlayer, ESPMode::ThreadSafe> PinnedPlayer = Player.Pin();
+	if (!PinnedPlayer.IsValid())
+	{
+		return false;
+	}
+
+	return PinnedPlayer->CanPresentVideoFrames(NumFrames);
+}
+
+/**
+ *
+ */
+void FElectraRendererVideo::SetRenderClock(TSharedPtr<Electra::IMediaRenderClock, ESPMode::ThreadSafe> InRenderClock)
+{
+	RenderClock = InRenderClock;
+}
+
+/**
+ *
+ */
+void FElectraRendererVideo::SetNextApproximatePresentationTime(const Electra::FTimeValue& NextApproxPTS)
+{
+}
+
+/**
+ * Flushes all pending buffers not yet rendered
+ */
+UEMediaError FElectraRendererVideo::Flush(const Electra::FParamDict& InOptions)
+{
+	// [...]If there are still frames a decoder has not returned yet UEMEDIA_ERROR_INTERNAL should be returned.[...]
+	if (NumBuffersAcquiredForDecoder != 0)
+	{
+		return UEMEDIA_ERROR_INTERNAL;
+	}
+	// Tell the player to flush now as well.
+	if (TSharedPtr<FElectraPlayer, ESPMode::ThreadSafe> PinnedPlayer = Player.Pin())
+	{
+		PinnedPlayer->OnVideoFlush();
+	}
+	return UEMEDIA_ERROR_OK;
+}
+
+/**
+ * Begins rendering of the first sample buffer
+ */
+void FElectraRendererVideo::StartRendering(const Electra::FParamDict& InOptions)
+{
+	if (TSharedPtr<FElectraPlayer, ESPMode::ThreadSafe> PinnedPlayer = Player.Pin())
+	{
+		PinnedPlayer->OnVideoRenderingStarted();
+	}
+}
+
+/**
+ * Stops rendering of sample buffers
+*/
+void FElectraRendererVideo::StopRendering(const Electra::FParamDict& InOptions)
+{
+	if (TSharedPtr<FElectraPlayer, ESPMode::ThreadSafe> PinnedPlayer = Player.Pin())
+	{
+		PinnedPlayer->OnVideoRenderingStopped();
+	}
+}
+
+void FElectraRendererVideo::TickOutputBufferPool()
+{
+	DecoderOutputPool.Tick();
+}
+
+/**
+ * Called from game thread to update buffer management
+ */
+void FElectraRendererVideo::TickInput(FTimespan DeltaTime, FTimespan Timecode)
+{
+	LLM_SCOPE(ELLMTag::ElectraPlayer);
+	CSV_SCOPED_TIMING_STAT(ElectraPlayer, MEDIArendererVideoUE_Tick);
+
+	check(IsInGameThread());
+
+	// Keep stats up to date
+	SET_DWORD_STAT(STAT_ElectraPlayer_MEDIArendererVideoUE_NumUsedOutputVideoSamples, NumOutputTexturesInUse);
+}
+
diff --git a/Engine/Plugins/Media/ElectraPlayer/Source/ElectraPlayerRuntime/Private/Runtime/Renderer/RendererVideo.h b/Engine/Plugins/Media/ElectraPlayer/Source/ElectraPlayerRuntime/Private/Runtime/Renderer/RendererVideo.h
new file mode 100644
--- /dev/null
+++ b/Engine/Plugins/Media/ElectraPlayer/Source/ElectraPlayerRuntime/Private/Runtime/Renderer/RendererVideo.h
@@ -0,0 +1,153 @@
+// Copyright Epic Games, Inc. All Rights Reserved.
+
+#pragma once
+
+#include "Containers/Queue.h"
+
+#include "Renderer/RendererBase.h"
+
+#include "MediaVideoDecoderOutput.h"
+
+struct FTimespan;
+class FElectraPlayer;
+
+class FElectraPlayerPlatformVideoDecoderOutputFactory
+{
+public:
+	static FVideoDecoderOutput* Create();
+};
+
+
+class FElectraRendererVideo : public Electra::IMediaRenderer, public TSharedFromThis<FElectraRendererVideo, ESPMode::ThreadSafe>
+{
+public:
+	struct SystemConfiguration
+	{
+		SystemConfiguration();
+	};
+
+	static bool Startup(const SystemConfiguration& Configuration);
+	static void Shutdown(void);
+
+	FElectraRendererVideo(TSharedPtr<FElectraPlayer, ESPMode::ThreadSafe> InPlayer);
+	virtual ~FElectraRendererVideo();
+
+	struct InstanceConfiguration
+	{
+		InstanceConfiguration()
+		{
+			// todo
+		}
+	};
+
+	enum OpenError
+	{
+		eOpen_Ok,
+		eOpen_Error,
+	};
+
+	virtual OpenError Open(const InstanceConfiguration& Config);
+	void Close(void);
+
+	void PrepareForDecoderShutdown();
+
+	//-------------------------------------------------------------------------
+	// Methods for UE interface
+	//
+	void TickInput(FTimespan DeltaTime, FTimespan Timecode);
+	void DetachPlayer();
+
+
+	//-------------------------------------------------------------------------
+	// Methods for Electra::IMediaRenderer
+	//
+
+	// Returns the properties of the buffer pool. Those properties should not change
+	virtual const Electra::FParamDict& GetBufferPoolProperties() const override;
+
+	// Create a buffer pool from where a decoder can get the block of memory to decode into.
+	virtual UEMediaError CreateBufferPool(const Electra::FParamDict& Parameters) override;
+
+	// Asks for a sample buffer from the buffer pool created previously through CreateBufferPool()
+	virtual UEMediaError AcquireBuffer(IBuffer*& OutBuffer, int32 TimeoutInMicroseconds, const Electra::FParamDict& InParameters) override;
+
+	// Releases the buffer for rendering and subsequent return to the buffer pool
+	virtual UEMediaError ReturnBuffer(IBuffer* Buffer, bool bRender, const Electra::FParamDict& InSampleProperties) override;
+
+	// Informs that the decoder is done with this pool. NO FREE!!!
+	virtual UEMediaError ReleaseBufferPool() override;
+
+	virtual bool CanReceiveOutputFrames(uint64 NumFrames) const override;
+
+	// Receives the render clock we need to update with the most recently rendered sample's timestamp.
+	virtual void SetRenderClock(TSharedPtr<Electra::IMediaRenderClock, ESPMode::ThreadSafe> RenderClock) override;
+
+	// Sets the next expected sample's approximate presentation time stamp
+	virtual void SetNextApproximatePresentationTime(const Electra::FTimeValue& NextApproxPTS) override;
+
+	// Flushes all pending buffers not yet rendered
+	virtual UEMediaError Flush(const Electra::FParamDict& InOptions) override;
+
+	// Begins rendering of the first sample buffer
+	virtual void StartRendering(const Electra::FParamDict& InOptions) override;
+
+	// Stops rendering of sample buffers
+	virtual void StopRendering(const Electra::FParamDict& InOptions) override;
+
+	//! Tick any output buffer logic that might need regular updates
+	virtual void TickOutputBufferPool() override;
+
+	//! Called from FElectraPlayerVideoDecoderOutputXYZ when texture is returned to pool
+	virtual void SampleReleasedToPool(FTimespan DurationToRelease) override;
+
+private:
+
+
+	TWeakPtr<FElectraPlayer, ESPMode::ThreadSafe> Player;
+
+	TDecoderOutputObjectPool<FVideoDecoderOutput, FElectraPlayerPlatformVideoDecoderOutputFactory> DecoderOutputPool;
+	int32 NumOutputTexturesInUse;
+
+
+	class FMediaBufferSharedPtrWrapper : public Electra::IMediaRenderer::IBuffer
+	{
+	public:
+		explicit FMediaBufferSharedPtrWrapper(const FVideoDecoderOutputPtr& InDecoderOutput)
+			: DecoderOutput(InDecoderOutput)
+		{
+		}
+
+		~FMediaBufferSharedPtrWrapper()
+		{
+		}
+
+		virtual const Electra::FParamDict& GetBufferProperties() const override
+		{
+			return BufferProperties;
+		}
+		Electra::FParamDict		BufferProperties;
+		FVideoDecoderOutputPtr	DecoderOutput;
+	};
+
+
+
+
+	void AcquireFromPool(FVideoDecoderOutputPtr& DelayedImage);
+
+
+	// These queues are all single-producer, single-consumer which is the default mode...
+
+	TQueue<FVideoDecoderOutputPtr>					QueueTickedAndWaitingForDecoder;
+
+	// ------------------------------------------------------------------------------------------------------
+
+	Electra::FParamDict								BufferPoolProperties;
+
+	int32											NumBuffers;
+	int32											NumBuffersAcquiredForDecoder;
+
+	FTimespan										DurationOfPendingRenderBuffers;
+
+	TSharedPtr<Electra::IMediaRenderClock, ESPMode::ThreadSafe>	RenderClock;
+};
+
diff --git a/Engine/Plugins/Media/ElectraPlayer/Source/ElectraPlayerRuntime/Private/Runtime/StreamAccessUnitBuffer.cpp b/Engine/Plugins/Media/ElectraPlayer/Source/ElectraPlayerRuntime/Private/Runtime/StreamAccessUnitBuffer.cpp
new file mode 100644
--- /dev/null
+++ b/Engine/Plugins/Media/ElectraPlayer/Source/ElectraPlayerRuntime/Private/Runtime/StreamAccessUnitBuffer.cpp
@@ -0,0 +1,300 @@
+// Copyright Epic Games, Inc. All Rights Reserved.
+
+#include "StreamAccessUnitBuffer.h"
+
+
+namespace Electra
+{
+
+
+
+	FMultiTrackAccessUnitBuffer::FMultiTrackAccessUnitBuffer()
+	{
+		EmptyBuffer = CreateNewBuffer();
+		ActiveOutputID = -1;
+		LastPoppedBufferID = -1;
+		LastPoppedDTS.SetToInvalid();
+		LastPoppedPTS.SetToInvalid();
+		bAutoselectFirstTrack = false;
+		bIsDeselected = false;
+		bEndOfData = false;
+		bLastPushWasBlocked = false;
+	}
+
+	FMultiTrackAccessUnitBuffer::FAccessUnitBufferPtr FMultiTrackAccessUnitBuffer::CreateNewBuffer()
+	{
+		// All buffers here are internally unbounded as we need them to accept all incoming data.
+		// Limits are used on the active buffer only.
+		FAccessUnitBufferPtr NewBuffer = MakeShared<FAccessUnitBuffer, ESPMode::ThreadSafe>();
+		NewBuffer->CapacitySet(FAccessUnitBuffer::FConfiguration(1024 << 20, 3600.0));
+		return NewBuffer;
+	}
+
+	FMultiTrackAccessUnitBuffer::~FMultiTrackAccessUnitBuffer()
+	{
+		PurgeAll();
+	}
+
+	void FMultiTrackAccessUnitBuffer::CapacitySet(const FAccessUnitBuffer::FConfiguration& Config)
+	{
+		PrimaryBufferConfiguration = Config;
+	}
+
+	bool FMultiTrackAccessUnitBuffer::Push(FAccessUnit*& AU)
+	{
+		uint32 TrackID = 0;
+		if (AU->StreamSourceInfo.IsValid())
+		{
+			TrackID = AU->StreamSourceInfo->NumericTrackID;
+		}
+		// Do we have a buffer for this track ID yet?
+		AccessLock.Lock();
+		if (!TrackBuffers.Contains(TrackID))
+		{
+			TrackBuffers.Emplace(TrackID, CreateNewBuffer());
+		}
+		FAccessUnitBufferPtr TrackBuffer = TrackBuffers[TrackID];
+
+		if (bAutoselectFirstTrack)
+		{
+			bAutoselectFirstTrack = false;
+			ActiveOutputID = (int32)TrackID;
+		}
+
+		// Pushing data to either buffer means that there is data and we are not at EOD here any more.
+		// This does not necessarily mean that the selected track is not at EOD. Just that some track is not.
+		bEndOfData = false;
+		AccessLock.Unlock();
+		// If this is the active buffer or no active buffer has been set yet we use the configured buffer limit in pushing.
+		const FAccessUnitBuffer::FConfiguration* PushConfig = ActiveOutputID == -1 || ActiveOutputID == (int32)TrackID ? &PrimaryBufferConfiguration : nullptr;
+		bool bWasPushed = TrackBuffer->Push(AU, PushConfig);
+		bLastPushWasBlocked = TrackBuffer->WasLastPushBlocked();
+		return bWasPushed;
+	}
+
+	void FMultiTrackAccessUnitBuffer::PushEndOfDataFor(TSharedPtr<const FStreamSourceInfo, ESPMode::ThreadSafe> InStreamSourceInfo)
+	{
+		uint32 TrackID = 0;
+		if (InStreamSourceInfo.IsValid())
+		{
+			TrackID = InStreamSourceInfo->NumericTrackID;
+		}
+		// Do we have a buffer for this track ID yet?
+		AccessLock.Lock();
+		if (!TrackBuffers.Contains(TrackID))
+		{
+			TrackBuffers.Emplace(TrackID, CreateNewBuffer());
+		}
+		FAccessUnitBufferPtr TrackBuffer = TrackBuffers[TrackID];
+
+		if (bAutoselectFirstTrack)
+		{
+			bAutoselectFirstTrack = false;
+			ActiveOutputID = (int32)TrackID;
+		}
+
+		AccessLock.Unlock();
+		TrackBuffer->PushEndOfData();
+	}
+
+	void FMultiTrackAccessUnitBuffer::PushEndOfDataAll()
+	{
+		FMediaCriticalSection::ScopedLock lock(AccessLock);
+		bEndOfData = true;
+		// Push an end-of-data into all tracks.
+		for (auto& It : TrackBuffers)
+		{
+			It.Value->PushEndOfData();
+		}
+	}
+
+	void FMultiTrackAccessUnitBuffer::Flush()
+	{
+		FMediaCriticalSection::ScopedLock lock(AccessLock);
+		// Flush all existing buffers but keep them in the track map.
+		for (auto& It : TrackBuffers)
+		{
+			It.Value->Flush();
+		}
+		LastPoppedDTS.SetToInvalid();
+		LastPoppedPTS.SetToInvalid();
+		bEndOfData = false;
+		bLastPushWasBlocked = false;
+		LastPoppedBufferID = -1;
+	}
+
+	void FMultiTrackAccessUnitBuffer::PurgeAll()
+	{
+		FMediaCriticalSection::ScopedLock lock(AccessLock);
+		// Get rid of all buffers
+		TrackBuffers.Empty();
+		LastPoppedDTS.SetToInvalid();
+		LastPoppedPTS.SetToInvalid();
+		bEndOfData = false;
+		bLastPushWasBlocked = false;
+		LastPoppedBufferID = -1;
+	}
+
+	void FMultiTrackAccessUnitBuffer::AutoselectFirstTrack()
+	{
+		FMediaCriticalSection::ScopedLock lock(AccessLock);
+		bAutoselectFirstTrack = true;
+		bIsDeselected = false;
+	}
+
+
+	void FMultiTrackAccessUnitBuffer::SelectTrackByID(int32 TrackID)
+	{
+		FMediaCriticalSection::ScopedLock lock(AccessLock);
+		if (TrackID != ActiveOutputID)
+		{
+			ActiveOutputID = TrackID;
+		}
+		bAutoselectFirstTrack = false;
+		bIsDeselected = false;
+	}
+
+	void FMultiTrackAccessUnitBuffer::Deselect()
+	{
+		bIsDeselected = true;
+		// Note: We do not change ActiveOutputID or bAutoselectFirstTrack.
+		//       The buffer actually remains selected in order to track its fullness etc.
+		//       Deselected only means that the access unit may not get sent to the decoder
+		//       or be replaced by one that produces no output.
+	}
+
+	bool FMultiTrackAccessUnitBuffer::IsDeselected() const
+	{
+		return bIsDeselected;
+	}
+
+	FTimeValue FMultiTrackAccessUnitBuffer::GetLastPoppedPTS() const
+	{
+		FMediaCriticalSection::ScopedLock lock(AccessLock);
+		return LastPoppedPTS;
+	}
+
+
+	TSharedPtr<FAccessUnitBuffer, ESPMode::ThreadSafe> FMultiTrackAccessUnitBuffer::GetSelectedTrackBuffer()
+	{
+		FAccessUnitBufferPtr TrackBuffer;
+		AccessLock.Lock();
+		if (ActiveOutputID != -1 && TrackBuffers.Contains(ActiveOutputID))
+		{
+			TrackBuffer = TrackBuffers[ActiveOutputID];
+		}
+		else
+		{
+			TrackBuffer = EmptyBuffer;
+		}
+		AccessLock.Unlock();
+		return TrackBuffer;
+	}
+
+	TSharedPtr<const FAccessUnitBuffer, ESPMode::ThreadSafe> FMultiTrackAccessUnitBuffer::GetSelectedTrackBuffer() const
+	{
+		FAccessUnitBufferPtr TrackBuffer;
+		AccessLock.Lock();
+		if (ActiveOutputID != -1 && TrackBuffers.Contains(ActiveOutputID))
+		{
+			TrackBuffer = TrackBuffers[ActiveOutputID];
+		}
+		else
+		{
+			TrackBuffer = EmptyBuffer;
+		}
+		AccessLock.Unlock();
+		return TrackBuffer;
+	}
+
+	void FMultiTrackAccessUnitBuffer::GetStats(FAccessUnitBufferInfo& OutStats)
+	{
+		// Stats are always returned with our configuration.
+		TSharedPtr<const FAccessUnitBuffer, ESPMode::ThreadSafe> Buf = GetSelectedTrackBuffer();
+		Buf->GetStats(OutStats, &PrimaryBufferConfiguration);
+		if (Buf == EmptyBuffer)
+		{
+			OutStats.bEndOfData = bEndOfData;
+		}
+	}
+
+	bool FMultiTrackAccessUnitBuffer::Pop(FAccessUnit*& OutAU)
+	{
+		// Note: We assume the access lock is held by the caller!
+		TSharedPtr<FAccessUnitBuffer, ESPMode::ThreadSafe> Buf = FMultiTrackAccessUnitBuffer::GetSelectedTrackBuffer();
+		if (Buf->Num())
+		{
+			bool bDidPop = Buf->Pop(OutAU);
+			if (bDidPop && OutAU)
+			{
+				// Did we just pop from a different buffer than last time?
+				if (LastPoppedBufferID != -1 && LastPoppedBufferID != ActiveOutputID)
+				{
+					// FIXME: We may need to discard everything that is not tagged as a sync sample to ensure proper stream switching.
+					OutAU->bTrackChangeDiscontinuity = true;
+				}
+				// Remember from which buffer we popped.
+				LastPoppedBufferID = ActiveOutputID;
+				LastPoppedDTS = OutAU->DTS;
+				LastPoppedPTS = OutAU->PTS;
+				// With this AU being popped off now we will also pop off all AUs that are now obsolete in the other tracks.
+				for (auto& It : TrackBuffers)
+				{
+					if (It.Value != Buf)
+					{
+						FTimeValue PoppedDTS, PoppedPTS;
+						It.Value->DiscardUntil(LastPoppedDTS, LastPoppedPTS, PoppedDTS, PoppedPTS);
+					}
+				}
+			}
+			return bDidPop;
+		}
+		else
+		{
+			OutAU = nullptr;
+			return false;
+		}
+	}
+
+	void FMultiTrackAccessUnitBuffer::PopDiscardUntil(FTimeValue UntilTime)
+	{
+		// Note: We assume the access lock is held by the caller!
+		for (auto& It : TrackBuffers)
+		{
+			FAccessUnitBufferPtr Buf = It.Value;
+			FTimeValue PoppedDTS, PoppedPTS;
+			Buf->DiscardUntil(UntilTime, UntilTime, PoppedDTS, PoppedPTS);
+			if (PoppedDTS.IsValid() && LastPoppedDTS.IsValid() && PoppedDTS > LastPoppedDTS)
+			{
+				LastPoppedDTS = PoppedDTS;
+			}
+			if (PoppedPTS.IsValid() && LastPoppedPTS.IsValid() && PoppedPTS > LastPoppedPTS)
+			{
+				LastPoppedPTS = PoppedPTS;
+			}
+		}
+	}
+
+
+	bool FMultiTrackAccessUnitBuffer::IsEODFlagSet() const
+	{
+		// Note: We assume the access lock is held by the caller!
+		TSharedPtr<const FAccessUnitBuffer, ESPMode::ThreadSafe> Buf = GetSelectedTrackBuffer();
+		return Buf == EmptyBuffer ? bEndOfData : Buf->IsEODFlagSet();
+	}
+
+	int32 FMultiTrackAccessUnitBuffer::Num() const
+	{
+		// Note: We assume the access lock is held by the caller!
+		return (int32)GetSelectedTrackBuffer()->Num();
+	}
+
+	bool FMultiTrackAccessUnitBuffer::WasLastPushBlocked() const
+	{
+		return bLastPushWasBlocked;
+	}
+
+
+} // namespace Electra
+
+
diff --git a/Engine/Plugins/Media/ElectraPlayer/Source/ElectraPlayerRuntime/Private/Runtime/StreamAccessUnitBuffer.h b/Engine/Plugins/Media/ElectraPlayer/Source/ElectraPlayerRuntime/Private/Runtime/StreamAccessUnitBuffer.h
new file mode 100644
--- /dev/null
+++ b/Engine/Plugins/Media/ElectraPlayer/Source/ElectraPlayerRuntime/Private/Runtime/StreamAccessUnitBuffer.h
@@ -0,0 +1,714 @@
+// Copyright Epic Games, Inc. All Rights Reserved.
+
+#pragma once
+
+#include "PlayerCore.h"
+#include "PlayerTime.h"
+#include "StreamTypes.h"
+
+
+
+namespace Electra
+{
+
+
+	class IAccessUnitMemoryProvider
+	{
+	public:
+		enum class EDataType
+		{
+			AU,
+			Payload,
+			GenericData
+		};
+		virtual ~IAccessUnitMemoryProvider() = default;
+
+		virtual void* AUAllocate(EDataType type, SIZE_T size, SIZE_T alignment = 0) = 0;
+		virtual void AUDeallocate(EDataType type, void* pAddr) = 0;
+	};
+
+
+
+
+	struct FStreamSourceInfo
+	{
+		FStreamSourceInfo() : NumericTrackID(0) {}
+		FString		Language;
+		FString		Role;
+		FString		PeriodID;
+		FString		AdaptationSetID;
+		FString		RepresentationID;
+		uint32		NumericTrackID;
+	};
+
+	struct FAccessUnit
+	{
+		struct CodecData
+		{
+			TArray<uint8>					CodecSpecificData;
+			TArray<uint8>					RawCSD;
+			FStreamCodecInformation			ParsedInfo;
+		};
+
+
+		enum EDropState
+		{
+			None = 0,
+			DtsTooEarly = 1,
+			PtsTooEarly = 2,
+			DtsTooLate = 4,
+			PtsTooLate = 8
+		};
+
+		EStreamType					ESType;							//!< Type of elementary stream this is an access unit of.
+		FTimeValue					PTS;							//!< PTS
+		FTimeValue					DTS;							//!< DTS
+		FTimeValue					Duration;						//!< Duration
+		uint32						AUSize;							//!< Size of this access unit
+		void*						AUData;							//!< Access unit data
+		TSharedPtrTS<CodecData>		AUCodecData;					//!< If set, points to sideband data for this access unit.
+		uint8						DropState;						//!< If set this access unit is not to be rendered. If possible also not to be decoded.
+		bool						bIsFirstInSequence;				//!< true for the first AU in a segment
+		bool						bIsSyncSample;					//!< true if this is a sync sample (keyframe)
+		bool						bIsDummyData;					//!< True if the decoder must be reset before decoding this AU.
+		bool						bTrackChangeDiscontinuity;		//!< True if this is the first AU after a track change.
+
+		TSharedPtr<const FStreamSourceInfo, ESPMode::ThreadSafe>	StreamSourceInfo;
+		TSharedPtr<const FPlayerLoopState, ESPMode::ThreadSafe>		PlayerLoopState;
+
+		// Reserved for decoders.
+		// FIXME: move this out of here and have the decoders track what they're doing with the AU.
+		bool						bHasBeenPrepared;
+
+
+		static FAccessUnit* Create(IAccessUnitMemoryProvider* MemProvider)
+		{
+			void* NewBuffer = MemProvider->AUAllocate(IAccessUnitMemoryProvider::EDataType::AU, sizeof(FAccessUnit));
+			FAccessUnit* NewAU = new (NewBuffer) FAccessUnit;
+			NewAU->AUMemoryProvider = MemProvider;
+			return NewAU;
+		}
+
+		int64 TotalMemSize() const
+		{
+			return sizeof(FAccessUnit) + AUSize + (AUCodecData.IsValid() ? AUCodecData->CodecSpecificData.Num() : 0);
+		}
+
+		void AddRef()
+		{
+			FMediaInterlockedIncrement(RefCount);
+		}
+
+		void SetCodecSpecificData(const TArray<uint8>& Csd)
+		{
+			AUCodecData = MakeSharedTS<CodecData>();
+			AUCodecData->CodecSpecificData = Csd;
+		}
+
+		void SetCodecSpecificData(const TSharedPtrTS<CodecData>& Csd)
+		{
+			AUCodecData = Csd;
+		}
+
+		void* AllocatePayloadBuffer(SIZE_T Num)
+		{
+			check(AUMemoryProvider);
+			return AUMemoryProvider ? AUMemoryProvider->AUAllocate(IAccessUnitMemoryProvider::EDataType::Payload, Num) : nullptr;
+		}
+
+		void AdoptNewPayloadBuffer(void* Buffer, SIZE_T Num)
+		{
+			if (AUData)
+			{
+				AUMemoryProvider->AUDeallocate(IAccessUnitMemoryProvider::EDataType::Payload, AUData);
+			}
+			AUData = Buffer;
+			AUSize = Num;
+		}
+
+		static void Release(FAccessUnit* AccessUnit)
+		{
+			if (AccessUnit)
+			{
+				check(AccessUnit->RefCount > 0);
+				if (FMediaInterlockedDecrement(AccessUnit->RefCount) == 1)
+				{
+					IAccessUnitMemoryProvider* MemoryProvider = AccessUnit->AUMemoryProvider;
+					AccessUnit->~FAccessUnit();
+					MemoryProvider->AUDeallocate(IAccessUnitMemoryProvider::EDataType::AU, AccessUnit);
+				}
+			}
+		}
+
+	private:
+		IAccessUnitMemoryProvider*		AUMemoryProvider;			//!< Interface to use to delete the allocated AU
+		uint32							RefCount;
+
+		FAccessUnit()
+		{
+			RefCount = 1;
+			PTS.SetToInvalid();
+			DTS.SetToInvalid();
+			Duration.SetToInvalid();
+			ESType = EStreamType::Unsupported;
+			AUSize = 0;
+			AUData = nullptr;
+			AUMemoryProvider = nullptr;
+			DropState = EDropState::None;
+			bIsFirstInSequence = false;
+			bIsSyncSample = false;
+			bIsDummyData = false;
+			bTrackChangeDiscontinuity = false;
+			bHasBeenPrepared = false;
+		}
+
+		~FAccessUnit()
+		{
+			if (AUData)
+			{
+				check(AUMemoryProvider);
+				AUMemoryProvider->AUDeallocate(IAccessUnitMemoryProvider::EDataType::Payload, AUData);
+				AUData = nullptr;
+			}
+		}
+	};
+
+
+
+
+
+
+	struct FAccessUnitBufferInfo
+	{
+		FAccessUnitBufferInfo()
+		{
+			Clear();
+		}
+		void Clear()
+		{
+			FrontDTS.SetToInvalid();
+			PushedDuration.SetToZero();
+			PlayableDuration.SetToZero();
+			MaxDuration.SetToZero();
+			MaxDataSize = 0;
+			CurrentMemInUse = 0;
+			NumCurrentAccessUnits = 0;
+			bEndOfData = false;
+			bLastPushWasBlocked = false;
+		}
+
+		FTimeValue			FrontDTS;
+		FTimeValue			PushedDuration;
+		FTimeValue			PlayableDuration;
+		FTimeValue			MaxDuration;
+		int64				MaxDataSize;
+		int64				CurrentMemInUse;
+		int64				NumCurrentAccessUnits;
+		bool				bEndOfData;
+		bool				bLastPushWasBlocked;
+	};
+
+
+	/**
+	 * This class implements a decoder input data FIFO.
+	**/
+	class FAccessUnitBuffer
+	{
+	public:
+		struct FConfiguration
+		{
+			FConfiguration(int64 InMaxByteSize = 0, double InMaxSeconds = 0.0)
+				: MaxDataSize(InMaxByteSize)
+			{
+				MaxDuration.SetFromSeconds(InMaxSeconds);
+			}
+			FTimeValue	MaxDuration;
+			int64		MaxDataSize;
+		};
+
+		FAccessUnitBuffer()
+			: FrontDTS(FTimeValue::GetInvalid())
+			, PushedDuration(FTimeValue::GetZero())
+			, PlayableDuration(FTimeValue::GetZero())
+			, CurrentMemInUse(0)
+			, bEndOfData(false)
+			, bLastPushWasBlocked(false)
+		{
+		}
+
+		~FAccessUnitBuffer()
+		{
+			while (AccessUnits.Num())
+			{
+				FAccessUnit::Release(AccessUnits.Pop());
+			}
+		}
+
+		//! Returns the number of access units currently in the FIFO.
+		int64 Num() const
+		{
+			FMediaCriticalSection::ScopedLock Lock(AccessLock);
+			return AccessUnits.Num();
+		}
+
+		//! Returns the amount of memory currently allocated.
+		int64 AllocatedSize() const
+		{
+			FMediaCriticalSection::ScopedLock Lock(AccessLock);
+			return CurrentMemInUse;
+		}
+
+		//! Returns both the number of access units and the amount of allocated memory.
+		void SizeAndAllocation(int64& OutNumAUs, int64& OutAllocatedMem) const
+		{
+			FMediaCriticalSection::ScopedLock Lock(AccessLock);
+			OutNumAUs = AccessUnits.Num();
+			OutAllocatedMem = CurrentMemInUse;
+		}
+
+		//! Returns all vital statistics.
+		void GetStats(FAccessUnitBufferInfo& OutStats, const FConfiguration* Limit = nullptr) const
+		{
+			FMediaCriticalSection::ScopedLock Lock(AccessLock);
+			Limit = Limit ? Limit : &Config;
+			OutStats.FrontDTS = FrontDTS;
+			OutStats.PushedDuration = PushedDuration;
+			OutStats.PlayableDuration = PlayableDuration;
+			OutStats.MaxDuration = Limit->MaxDuration;
+			OutStats.MaxDataSize = Limit->MaxDataSize;
+			OutStats.CurrentMemInUse = CurrentMemInUse;
+			OutStats.NumCurrentAccessUnits = AccessUnits.Num();
+			OutStats.bEndOfData = bEndOfData;
+			OutStats.bLastPushWasBlocked = bLastPushWasBlocked;
+		}
+
+		//! Sets the maximum amount of memory/number of AUs this FIFO is allowed to allocate.
+		void CapacitySet(const FConfiguration& InConfig)
+		{
+			FMediaCriticalSection::ScopedLock Lock(AccessLock);
+			Flush();
+			Config = InConfig;
+		}
+
+		//! Adds an access unit to the FIFO. Returns true if successful, false if the FIFO has insufficient free space.
+		bool Push(FAccessUnit*& AU, const FConfiguration* Limit = nullptr)
+		{
+			FMediaCriticalSection::ScopedLock Lock(AccessLock);
+			// Pushing new data unconditionally clears the EOD flag even if the buffer is currently full.
+			// The attempt to push implies there will be more data.
+			bEndOfData = false;
+			Limit = Limit ? Limit : &Config;
+			if (CanPush(AU, Limit))
+			{
+				bLastPushWasBlocked = false;
+				AccessUnits.Push(AU);
+				int64 memSize = AU->TotalMemSize();
+				CurrentMemInUse += memSize;
+
+				if (AU->DropState == FAccessUnit::EDropState::None)
+				{
+					if (!FrontDTS.IsValid())
+					{
+						FrontDTS = AU->DTS;
+					}
+					PlayableDuration += AU->Duration;
+				}
+
+				PushedDuration += AU->Duration;
+				NumInSemaphore.Release();
+				return true;
+			}
+			else
+			{
+				bLastPushWasBlocked = true;
+				return false;
+			}
+		}
+
+		//! "Pushes" an end-of-data marker signaling that no further data will be pushed. May be called more than once. Flushing or pushing new data clears the flag.
+		void PushEndOfData()
+		{
+			bEndOfData = true;
+		}
+
+		//! Removes and returns the oldest access unit from the FIFO. Returns false if the FIFO is empty.
+		bool Pop(FAccessUnit*& OutAU)
+		{
+			FMediaCriticalSection::ScopedLock Lock(AccessLock);
+			if (Num())
+			{
+				OutAU = AccessUnits.Pop();
+				int64 nMemSize = OutAU->TotalMemSize();
+				CurrentMemInUse -= nMemSize;
+				NumInSemaphore.TryToObtain();
+
+				if (!AccessUnits.IsEmpty())
+				{
+					if (AccessUnits.FrontRef()->DropState == FAccessUnit::EDropState::None)
+					{
+						FrontDTS = AccessUnits.FrontRef()->DTS;
+					}
+					else
+					{
+						FrontDTS.SetToInvalid();
+						for (uint32 i = 1; i < AccessUnits.Num(); ++i)
+						{
+							if (AccessUnits[i]->DropState == FAccessUnit::EDropState::None)
+							{
+								FrontDTS = AccessUnits[i]->DTS;
+								break;
+							}
+						}
+					}
+					if (OutAU->DropState == FAccessUnit::EDropState::None)
+					{
+						PlayableDuration -= OutAU->Duration;
+					}
+					PushedDuration -= OutAU->Duration;
+				}
+				else
+				{
+					FrontDTS.SetToInvalid();
+					PlayableDuration.SetToZero();
+					PushedDuration.SetToZero();
+				}
+				return true;
+			}
+			else
+			{
+				OutAU = nullptr;
+				return false;
+			}
+		}
+
+		//!
+		bool Peek(FAccessUnit*& OutAU)
+		{
+			FMediaCriticalSection::ScopedLock Lock(AccessLock);
+			if (Num())
+			{
+				OutAU = AccessUnits.Front();
+				OutAU->AddRef();
+				return true;
+			}
+			else
+			{
+				OutAU = nullptr;
+				return false;
+			}
+		}
+
+		//! Discards data that has both its DTS and PTS less than the provided ones.
+		void DiscardUntil(FTimeValue NextValidDTS, FTimeValue NextValidPTS, FTimeValue& OutPoppedDTS, FTimeValue& OutPoppedPTS)
+		{
+			while (true)
+			{
+				FAccessUnit* NextAU = nullptr;
+				if (Peek(NextAU))
+				{
+					if (NextAU)
+					{
+						bool bDTS = NextValidDTS.IsValid() ? NextAU->DTS < NextValidDTS : true;
+						bool bPTS = NextValidPTS.IsValid() ? NextAU->PTS < NextValidPTS : true;
+						if (bDTS && bPTS)
+						{
+							OutPoppedDTS = NextAU->DTS;
+							OutPoppedPTS = NextAU->PTS;
+							Pop(NextAU);
+							FAccessUnit::Release(NextAU);
+						}
+						else
+						{
+							break;
+						}
+					}
+					else
+					{
+						break;
+					}
+				}
+				else
+				{
+					break;
+				}
+			}
+		}
+
+		//! Waits for data to arrive. Returns true if data is present. False if not and timeout expired.
+		bool WaitForData(int64 waitForMicroseconds = -1)
+		{
+			bool bHave = NumInSemaphore.Obtain(waitForMicroseconds);
+			if (bHave)
+			{
+				NumInSemaphore.Release();
+			}
+			return bHave;
+		}
+
+		//! Removes all elements from the FIFO
+		void Flush()
+		{
+			FMediaCriticalSection::ScopedLock Lock(AccessLock);
+			while (AccessUnits.Num())
+			{
+				FAccessUnit::Release(AccessUnits.Pop());
+				NumInSemaphore.TryToObtain();
+			}
+			CurrentMemInUse = 0;
+			bEndOfData = false;
+			bLastPushWasBlocked = false;
+			FrontDTS.SetToInvalid();
+			PushedDuration.SetToZero();
+			PlayableDuration.SetToZero();
+		}
+
+		//! Checks if the buffer has reached the end-of-data marker (marker is set and no more data is in the buffer).
+		bool IsEndOfData() const
+		{
+			FMediaCriticalSection::ScopedLock Lock(AccessLock);
+			return bEndOfData && AccessUnits.IsEmpty();
+		}
+
+		// Checks if the end-of-data flag has been set. There may still be data in the buffer though!
+		bool IsEODFlagSet() const
+		{
+			FMediaCriticalSection::ScopedLock Lock(AccessLock);
+			return bEndOfData;
+		}
+
+		// Was the last push blocked because the buffer limits were reached?
+		bool WasLastPushBlocked() const
+		{
+			return bLastPushWasBlocked;
+		}
+
+		// Helper class to lock the AU buffer
+		class FScopedLock : private TMediaNoncopyable<FScopedLock>
+		{
+		public:
+			explicit FScopedLock(const FAccessUnitBuffer& owner)
+				: mOwner(owner)
+			{
+				mOwner.AccessLock.Lock();
+			}
+			~FScopedLock()
+			{
+				mOwner.AccessLock.Unlock();
+			}
+		private:
+			FScopedLock();
+			const FAccessUnitBuffer& mOwner;
+		};
+
+	private:
+		//! Checks if an access unit can be pushed to the FIFO. Returns true if successful, false if the FIFO has insufficient free space.
+		bool CanPush(const FAccessUnit* AU, const FConfiguration* Limit = nullptr)
+		{
+			FMediaCriticalSection::ScopedLock Lock(AccessLock);
+			Limit = Limit ? Limit : &Config;
+			check(Limit->MaxDuration > FTimeValue::GetZero());
+			check(AU->Duration.IsValid() && !AU->Duration.IsInfinity());
+			// Memory ok?
+			if (AU->TotalMemSize() + CurrentMemInUse <= Limit->MaxDataSize)
+			{
+				// Max allowed duration ok?
+				return PlayableDuration.IsValid() && PlayableDuration + AU->Duration > Limit->MaxDuration ? false : true;
+			}
+			return false;
+		}
+
+		mutable FMediaCriticalSection							AccessLock;
+		FConfiguration								Config;
+		TMediaQueueDynamicNoLock<FAccessUnit*>		AccessUnits;
+		FMediaSemaphore								NumInSemaphore;
+		FTimeValue									FrontDTS;					//!< DTS of first AU in buffer
+		FTimeValue									PushedDuration;
+		FTimeValue									PlayableDuration;
+		int64										CurrentMemInUse;
+		bool										bEndOfData;
+		bool										bLastPushWasBlocked;
+	};
+
+
+
+	/**
+	 * A multi-track access unit buffer keeps access units from several tracks in individual buffers,
+	 * one of which is selected to return AUs to the decoder from. The other unselected tracks will
+	 * discard their AUs as the play position progresses.
+	 */
+	class FMultiTrackAccessUnitBuffer
+	{
+	public:
+		FMultiTrackAccessUnitBuffer();
+		~FMultiTrackAccessUnitBuffer();
+		void CapacitySet(const FAccessUnitBuffer::FConfiguration& Config);
+		void AutoselectFirstTrack();
+		void SelectTrackByID(int32 TrackID);
+		void Deselect();
+		//int32 GetSelectedTrackID() const;
+		bool Push(FAccessUnit*& AU);
+		void PushEndOfDataFor(TSharedPtr<const FStreamSourceInfo, ESPMode::ThreadSafe> InStreamSourceInfo);
+		void PushEndOfDataAll();
+		void Flush();
+		void GetStats(FAccessUnitBufferInfo& OutStats);
+		bool IsDeselected() const;
+		FTimeValue GetLastPoppedPTS() const;
+
+		// Helper class to lock the AU buffer
+		class FScopedLock : private TMediaNoncopyable<FScopedLock>
+		{
+		public:
+			explicit FScopedLock(const FMultiTrackAccessUnitBuffer& Self)
+				: LockedSelf(Self)
+			{
+				LockedSelf.AccessLock.Lock();
+			}
+			~FScopedLock()
+			{
+				LockedSelf.AccessLock.Unlock();
+			}
+		private:
+			FScopedLock();
+			const FMultiTrackAccessUnitBuffer& LockedSelf;
+		};
+
+		bool Pop(FAccessUnit*& OutAU);
+		void PopDiscardUntil(FTimeValue UntilTime);
+		bool IsEODFlagSet() const;
+		int32 Num() const;
+		bool WasLastPushBlocked() const;
+
+	private:
+		using FAccessUnitBufferPtr = TSharedPtr<FAccessUnitBuffer, ESPMode::ThreadSafe>;
+
+		void PurgeAll();
+
+		FAccessUnitBufferPtr CreateNewBuffer();
+		TSharedPtr<FAccessUnitBuffer, ESPMode::ThreadSafe> GetSelectedTrackBuffer();
+		TSharedPtr<const FAccessUnitBuffer, ESPMode::ThreadSafe> GetSelectedTrackBuffer() const;
+
+		mutable FMediaCriticalSection						AccessLock;
+		FAccessUnitBuffer::FConfiguration		PrimaryBufferConfiguration;			//!< Configuration of the primary (active) buffer. All other buffers are unbounded to ensure no stalling of the demuxer.
+		FAccessUnitBufferPtr					EmptyBuffer;						//!< An empty buffer
+		int32									ActiveOutputID;						//!< Track ID of the buffer that is feeding the decoder.
+		int32									LastPoppedBufferID;					//!< Buffer ID from which an AU was popped last.
+		TMap<uint32, FAccessUnitBufferPtr>		TrackBuffers;						//!< Map of track buffers. One per track ID.
+		FTimeValue								LastPoppedDTS;
+		FTimeValue								LastPoppedPTS;
+		bool									bAutoselectFirstTrack;
+		bool									bIsDeselected;
+		bool									bEndOfData;
+		bool									bLastPushWasBlocked;
+	};
+
+
+
+
+	/**
+	 * Base class for any decoder receiving data in "access units".
+	 * An access unit (AU) is considered a data packet that can be sent into a decoder without
+	 * the decoder stalling. For a h.264 decoder this is usually a single NALU.
+	**/
+	class IAccessUnitBufferInterface
+	{
+	public:
+		virtual ~IAccessUnitBufferInterface() = default;
+
+		enum class EAUpushResult
+		{
+			Ok,
+			Full,
+			Error,
+		};
+		//! Attempts to push an access unit to the decoder. Ownership of the access unit is transferred if the push is successful.
+		virtual EAUpushResult AUdataPushAU(FAccessUnit* AccessUnit) = 0;
+		//! Notifies the decoder that there will be no further access units.
+		virtual void AUdataPushEOD() = 0;
+		//! Instructs the decoder to flush all pending input and all already decoded output.
+		virtual void AUdataFlushEverything() = 0;
+	};
+
+
+
+	//---------------------------------------------------------------------------------------------------------------------
+	/**
+	 * A decoder input buffer listener callback to monitor the current state of decoder input buffer levels.
+	 *
+	 * The decoder will invoke this listener right before it wants to get an access unit from its input buffer,
+	 * whether the buffer already contains data or is empty.
+	**/
+	class IAccessUnitBufferListener
+	{
+	public:
+		virtual ~IAccessUnitBufferListener() = default;
+
+		struct FBufferStats
+		{
+			FBufferStats()
+			{
+				Clear();
+			}
+			void Clear()
+			{
+				NumAUsAvailable = 0;
+				NumBytesInBuffer = 0;
+				MaxBytesOfBuffer = 0;
+				bEODSignaled = false;
+				bEODReached = false;
+			}
+			int64	NumAUsAvailable;		//!< Number of AUs currently in the buffer.
+			int64	NumBytesInBuffer;		//!< Number of bytes currently in the buffer.
+			int64	MaxBytesOfBuffer;		//!< Maximum byte capacity this buffer is allowed to store.
+			bool	bEODSignaled;			//!< Set after PushEndOfData() has been called
+			bool	bEODReached;			//!< Set after PushEndOfData() has been called AND the last AU was taken from the buffer.
+		};
+
+		//! Called right before the decoder wants to get an access unit from its input buffer, regardless if it already has data or not.
+		virtual void DecoderInputNeeded(const FBufferStats& CurrentInputBufferStats) = 0;
+	};
+
+
+
+	//---------------------------------------------------------------------------------------------------------------------
+	/**
+	 * A decoder ready listener callback to monitor the decoder activity.
+	 *
+	 * The decoder will invoke this listener right before it needs a buffer from its output buffer prior to decoding.
+	 * This is called whether or not the output buffer has room for new decoded data or not.
+	**/
+	class IDecoderOutputBufferListener
+	{
+	public:
+		virtual ~IDecoderOutputBufferListener() = default;
+
+		struct FDecodeReadyStats
+		{
+			FDecodeReadyStats()
+			{
+				Clear();
+			}
+			void Clear()
+			{
+				ReadyDuration.SetToInvalid();
+				NumDecodedElementsReady = 0;
+				MaxDecodedElementsReady = 0;
+				NumElementsInDecoder = 0;
+				bOutputStalled = false;
+				bEODreached = false;
+			}
+			FTimeValue				ReadyDuration;				//!< Duration of the ready material
+			int64					NumDecodedElementsReady;	//!< Number of decoded elements ready for rendering
+			int64					MaxDecodedElementsReady;	//!< Maximum number of decoded elements.
+			int64					NumElementsInDecoder;		//!< Number of elements currently in the decoder pipeline
+			bool					bOutputStalled;				//!< true if the output is full and decoding is delayed until there's room again.
+			bool					bEODreached;				//!< true when the final decoded element has been passed on.
+		};
+
+		virtual void DecoderOutputReady(const FDecodeReadyStats& CurrentReadyStats) = 0;
+	};
+
+
+
+} // namespace Electra
+
+
diff --git a/Engine/Plugins/Media/ElectraPlayer/Source/ElectraPlayerRuntime/Private/Runtime/StreamDataBuffer.h b/Engine/Plugins/Media/ElectraPlayer/Source/ElectraPlayerRuntime/Private/Runtime/StreamDataBuffer.h
new file mode 100644
--- /dev/null
+++ b/Engine/Plugins/Media/ElectraPlayer/Source/ElectraPlayerRuntime/Private/Runtime/StreamDataBuffer.h
@@ -0,0 +1,402 @@
+// Copyright Epic Games, Inc. All Rights Reserved.
+
+#pragma once
+
+#include "PlayerCore.h"
+
+
+namespace Electra
+{
+
+	/**
+	 * Byte ringbuffer
+	**/
+	class FPODRingbuffer : private TMediaNoncopyable<FPODRingbuffer>
+	{
+	public:
+		FPODRingbuffer(int32 InNumBytes = 0)
+			: Data(nullptr)
+			, DataEnd(nullptr)
+			, WritePos(nullptr)
+			, ReadPos(nullptr)
+			, DataSize(0)
+			, NumIn(0)
+			, WaitingForSize(0)
+			, bEOD(false)
+			, bWasAborted(false)
+		{
+			Allocate(InNumBytes);
+			SizeAvailableSignal.Reset();
+		}
+		~FPODRingbuffer()
+		{
+			FMediaCriticalSection::ScopedLock Lock(AccessLock);
+			Deallocate();
+		}
+
+		//! Allocates buffer of the specified capacity, destroying any previous buffer.
+		bool Reserve(int32 InNumBytes)
+		{
+			FMediaCriticalSection::ScopedLock Lock(AccessLock);
+			Deallocate();
+			if (Allocate(InNumBytes))
+			{
+				Reset();
+				return true;
+			}
+			return false;
+		}
+
+		//! Enlarges the buffer to the new capacity, retaining the current content.
+		bool EnlargeTo(int32 InNewNumBytes)
+		{
+			FMediaCriticalSection::ScopedLock Lock(AccessLock);
+			// Do we need a bigger capacity?
+			if (InNewNumBytes > Capacity())
+			{
+				// Are we currently empty?
+				if (IsEmpty())
+				{
+					return Reserve(InNewNumBytes);
+				}
+				else
+				{
+					return InternalGrowTo(InNewNumBytes);
+				}
+			}
+			return true;
+		}
+
+		//! Clears the buffer
+		void Reset()
+		{
+			FMediaCriticalSection::ScopedLock Lock(AccessLock);
+			NumIn = 0;
+			WritePos = ReadPos = Data;
+			bEOD = false;
+			bWasAborted = false;
+			// Do not modify the "waiter" members. When waiting for data to arrive it needs to continue
+			// doing so even when we are resetting the buffer.
+				//WaitingForSize = 0;
+				//SizeAvailableSignal.Signal();
+		}
+
+		//! Returns the ringbuffer capacity
+		int32 Capacity() const
+		{
+			return DataSize;
+		}
+
+		//! Returns the number of bytes in the buffer (amount that can be popped)
+		int32 Num() const
+		{
+			return NumIn;
+		}
+
+		//! Returns the number of free bytes in the buffer (amount that can be pushed)
+		int32 Avail() const
+		{
+			return DataSize - Num();
+		}
+
+		//! Checks if the buffer is empty.
+		bool IsEmpty() const
+		{
+			return Num() == 0;
+		}
+
+		//! Checks if the buffer is full.
+		bool IsFull() const
+		{
+			return Avail() == 0;
+		}
+
+		//! Checks if the buffer has reached the end-of-data marker (marker is set and no more data is in the buffer).
+		bool IsEndOfData() const
+		{
+			return IsEmpty() && bEOD;
+		}
+
+		//! Checks if the end-of-data flag has been set. There may still be data in the buffer though!
+		bool GetEOD() const
+		{
+			return bEOD;
+		}
+
+		//! Waits until the specified number of bytes has arrived in the buffer.
+		//! NOTE: This method is somewhat dangerous in that there is no guarantee the required amount will ever arrive.
+		//!       You must also never wait for more data than is the capcacity of the buffer!
+		bool WaitUntilSizeAvailable(int32 SizeNeeded, int32 TimeoutMicroseconds)
+		{
+			// Only wait if not at EOD and more data than presently available is asked for.
+			// Otherwise return enough data to be present even if that is not actually the case.
+			if (!bEOD && SizeNeeded > Num())
+			{
+				AccessLock.Lock();
+				// Repeat the size check inside the mutex lock in case we enter this block
+				// while new data is being pushed from another thread.
+				if (SizeNeeded > Num())
+				{
+					SizeAvailableSignal.Reset();
+					WaitingForSize = SizeNeeded;
+				}
+				else
+				{
+					SizeAvailableSignal.Signal();
+					WaitingForSize = 0;
+				}
+				AccessLock.Unlock();
+				if (TimeoutMicroseconds > 0)
+				{
+					return SizeAvailableSignal.WaitTimeout(TimeoutMicroseconds);
+				}
+				else
+				{
+					// No infinite waiting by specifying negative timeouts!
+					check(TimeoutMicroseconds == 0);
+					return SizeAvailableSignal.IsSignaled();
+				}
+			}
+			return true;
+		}
+
+
+		//! Inserts elements into the buffer. Returns true if successful, false if there is no room.
+		bool PushData(const uint8* InData, int32 NumElements)
+		{
+			FMediaCriticalSection::ScopedLock Lock(AccessLock);
+			check(!bEOD);
+			// Zero elements can always be pushed...
+			if (NumElements == 0)
+			{
+				return true;
+			}
+			if (Avail() >= NumElements)
+			{
+				int32 NumBeforeWrap = DataEnd - WritePos;
+				if (NumElements <= NumBeforeWrap)
+				{
+					CopyData(WritePos, InData, NumElements);
+					if ((WritePos += NumElements) == DataEnd)
+					{
+						WritePos = Data;
+					}
+				}
+				else
+				{
+					CopyData(WritePos, InData, NumBeforeWrap);
+					CopyData(Data, InData + NumBeforeWrap, NumElements - NumBeforeWrap);
+					WritePos = Data + (NumElements - NumBeforeWrap);
+				}
+				NumIn += NumElements;
+				if (NumIn >= WaitingForSize)
+				{
+					SizeAvailableSignal.Signal();
+				}
+				return true;
+			}
+			else
+			{
+				return false;
+			}
+		}
+
+
+		//! "Opens" a push-block by returning 2 sets of pointers and block sizes where data to be pushed can be stored directly.
+		int32 PushBlockOpen(uint8*& Block1, int32& Block1Size, uint8*& Block2, int32& Block2Size)
+		{
+			FMediaCriticalSection::ScopedLock Lock(AccessLock);
+			int32 av = Avail();
+			int32 NumBeforeWrap = DataEnd - WritePos;
+			Block1 = WritePos;
+			if (av > NumBeforeWrap)
+			{
+				Block1Size = NumBeforeWrap;
+				Block2 = Data;
+				Block2Size = av - NumBeforeWrap;
+			}
+			else
+			{
+				Block1Size = av;
+				Block2 = nullptr;
+				Block2Size = 0;
+			}
+			return av;
+		}
+
+		//! "Closes" a previously opened push block and advances the write pointer by as many bytes as were actually stored.
+		void PushBlockClose(int32 NumElements)
+		{
+			if (NumElements)
+			{
+				FMediaCriticalSection::ScopedLock Lock(AccessLock);
+				if ((WritePos += NumElements) >= DataEnd)
+				{
+					WritePos -= DataSize;
+				}
+				NumIn += NumElements;
+				if (NumIn >= WaitingForSize)
+				{
+					SizeAvailableSignal.Signal();
+				}
+			}
+		}
+
+
+		//! "Pushes" an end-of-data marker signaling that no further data will be pushed. May be called more than once. Ringbuffer must be Reset() before next use.
+		void SetEOD()
+		{
+			FMediaCriticalSection::ScopedLock Lock(AccessLock);
+			bEOD = true;
+			// Signal that data is present to wake any waiters on WaitForData() even though there may be no data in the buffer anymore.
+			SizeAvailableSignal.Signal();
+		}
+
+		//! "Pops" data from the buffer to a destination. At most the specified number of bytes are popped, or fewer if not as many are available.
+		int32 PopData(uint8* OutData, int32 MaxElementsWanted)
+		{
+			FMediaCriticalSection::ScopedLock Lock(AccessLock);
+			int32 size = Num();
+			int32 NumBeforeWrap = DataEnd - ReadPos;
+
+			if (MaxElementsWanted > size)
+			{
+				MaxElementsWanted = size;
+			}
+			// Copy out or skip over?
+			if (OutData)
+			{
+				// Can copy out without wrap?
+				if (MaxElementsWanted <= NumBeforeWrap)
+				{
+					CopyData(OutData, ReadPos, MaxElementsWanted);
+				}
+				else
+				{
+					CopyData(OutData, ReadPos, NumBeforeWrap);
+					CopyData(OutData + NumBeforeWrap, Data, MaxElementsWanted - NumBeforeWrap);
+				}
+			}
+
+			if ((ReadPos += MaxElementsWanted) >= DataEnd)
+			{
+				ReadPos -= DataSize;
+			}
+			NumIn -= MaxElementsWanted;
+			return MaxElementsWanted;
+		}
+
+		void Lock()
+		{
+			AccessLock.Lock();
+		}
+
+		void Unlock()
+		{
+			AccessLock.Unlock();
+		}
+
+		int32 GetLinearReadSize() const
+		{
+			FMediaCriticalSection::ScopedLock Lock(AccessLock);
+			int32 NumBeforeWrap = DataEnd - ReadPos;
+			return NumBeforeWrap;
+		}
+
+		// NOTE: Must control Lock()/Unlock() externally!
+		const uint8* GetLinearReadData() const
+		{
+			return ReadPos;
+		}
+		uint8* GetLinearReadData()
+		{
+			return ReadPos;
+		}
+
+
+		void Abort()
+		{
+			FMediaCriticalSection::ScopedLock Lock(AccessLock);
+			bWasAborted = true;
+			SizeAvailableSignal.Signal();
+		}
+
+		bool WasAborted() const
+		{
+			return bWasAborted;
+		}
+
+	protected:
+		bool Allocate(int32 InSize)
+		{
+			if (InSize)
+			{
+				DataSize = InSize;
+				Data = (uint8*)FMemory::Malloc(InSize);
+				DataEnd = Data + DataSize;
+				WritePos = Data;
+				ReadPos = Data;
+				return Data != nullptr;
+			}
+			return true;
+		}
+
+		void Deallocate()
+		{
+			FMemory::Free(Data);
+			Data = nullptr;
+			DataEnd = nullptr;
+			DataSize = 0;
+			NumIn = 0;
+			WritePos = nullptr;
+			ReadPos = nullptr;
+		}
+
+		bool InternalGrowTo(int32 InNewNumBytes)
+		{
+			// Note: The access mutex must already be held here!
+			check(Data && InNewNumBytes);
+			// Resize the buffer
+			uint8* NewData = (uint8*)FMemory::Realloc(Data, InNewNumBytes);
+			if (NewData)
+			{
+				// Get current read & write offsets
+				int32 ReadOffset = ReadPos - Data;
+				int32 WriteOffset = WritePos - Data;
+				// Set the data pointers up in the new buffer.
+				DataSize = InNewNumBytes;
+				Data = NewData;
+				DataEnd = Data + DataSize;
+				WritePos = Data + WriteOffset;
+				ReadPos = Data + ReadOffset;
+				return true;
+			}
+			return false;
+		}
+
+		void CopyData(uint8* CopyTo, const uint8* CopyFrom, int32 NumElements)
+		{
+			if (NumElements && CopyTo && CopyFrom)
+			{
+				FMemory::Memcpy(CopyTo, CopyFrom, NumElements);
+			}
+		}
+
+		FMediaCriticalSection	AccessLock;
+		FMediaEvent				SizeAvailableSignal;	//!< signaled when WaitingForSize data is present
+		uint8*					Data;					//!< Base address of buffer
+		uint8*					DataEnd;				//!< End address of buffer
+		uint8*					WritePos;				//!< Current write position
+		uint8*					ReadPos;				//!< Current read position
+		int32					DataSize;				//!< Maximum number of bytes in the buffer
+		int32					NumIn;					//!< Current number of bytes in the buffer
+		int32					WaitingForSize;			//!< If waiting for a certain number of bytes to become available
+		bool					bEOD;					//!< true when the last packet of data was pushed into the buffer.
+		bool					bWasAborted;
+	};
+
+
+
+} // namespace Electra
+
+
diff --git a/Engine/Plugins/Media/ElectraPlayer/Source/ElectraPlayerRuntime/Private/Runtime/Utilities/TimeUtilities.h b/Engine/Plugins/Media/ElectraPlayer/Source/ElectraPlayerRuntime/Private/Runtime/Utilities/TimeUtilities.h
new file mode 100644
--- /dev/null
+++ b/Engine/Plugins/Media/ElectraPlayer/Source/ElectraPlayerRuntime/Private/Runtime/Utilities/TimeUtilities.h
@@ -0,0 +1,23 @@
+// Copyright Epic Games, Inc. All Rights Reserved.
+
+#pragma once
+
+#include "PlayerCore.h"
+
+
+namespace Electra
+{
+
+	namespace ISO8601
+	{
+		UEMediaError ParseDateTime(FTimeValue& OutTimeValue, const FString& DateTime);
+	}
+
+
+	namespace RFC7231
+	{
+		UEMediaError ParseDateTime(FTimeValue& OutTimeValue, const FString& DateTime);
+	}
+
+} // namespace Electra
+
diff --git a/Engine/Plugins/Media/ElectraPlayer/Source/ElectraPlayerRuntime/Private/Runtime/Utilities/URI.cpp b/Engine/Plugins/Media/ElectraPlayer/Source/ElectraPlayerRuntime/Private/Runtime/Utilities/URI.cpp
new file mode 100644
--- /dev/null
+++ b/Engine/Plugins/Media/ElectraPlayer/Source/ElectraPlayerRuntime/Private/Runtime/Utilities/URI.cpp
@@ -0,0 +1,404 @@
+// Copyright Epic Games, Inc. All Rights Reserved.
+
+// Unittests are currently not present within perforce but can be obtained from:
+// https://github.ol.epicgames.net/Eric-Schneller/urlparser/blob/master/parser_test.cpp
+
+#include "Utilities/URI.h"
+#include "Utilities/StringHelpers.h"
+
+namespace Electra
+{
+    namespace Utilities
+    {
+
+        const FString FURI::Format() const
+        {
+			static const FString kTextColon(TEXT(":"));
+			static const FString kTextQuestionMark(TEXT("?"));
+			static const FString kTextHash(TEXT("#"));
+			static const FString kTextDoubleSlash(TEXT("//"));
+            FString Result;
+
+            if (!Scheme.IsEmpty())
+            {
+				Result += Scheme + kTextColon;
+            }
+
+            if (HasAuthority())
+            {
+                Result += kTextDoubleSlash;
+                if (HasUserInfo())
+                {
+                    Result += UserInfo.Username;
+                    if (!UserInfo.Password.IsEmpty())
+                    {
+                        Result += kTextColon + UserInfo.Password;
+                    }
+					Result.AppendChar(TCHAR('@'));
+                }
+
+                Result += Host;
+            }
+
+            Result += Opaque;
+            Result += Path;
+
+            if (HasQuery)
+            {
+                Result += kTextQuestionMark + QueryString;
+            }
+
+            if (HasFragment)
+            {
+                Result += kTextHash + Fragment;
+            }
+
+            return Result;
+        }
+
+        const FURI FURI::Resolve(const FURI& InUri) const
+        {
+            FURI Target = *this;
+
+            if (!InUri.Scheme.IsEmpty())
+            {
+                Target.Host = InUri.Host;
+                Target.Scheme = InUri.Scheme;
+                Target.UserInfo = InUri.UserInfo;
+                Target.Path = CleanPath(InUri.Path);
+                Target.Opaque = InUri.Opaque;
+                Target.QueryString = InUri.QueryString;
+                Target.HasQuery = InUri.HasQuery;
+            }
+            else
+            {
+                if (InUri.HasAuthority())
+                {
+                    Target.UserInfo = InUri.UserInfo;
+                    Target.Host = InUri.Host;
+                    Target.QueryString = InUri.QueryString;
+                    Target.HasQuery = InUri.HasQuery;
+                    Target.Path = CleanPath(InUri.Path);
+                }
+                else
+                {
+                    if (InUri.Path.IsEmpty())
+                    {
+                        if (InUri.HasQuery)
+                        {
+                            Target.QueryString = InUri.QueryString;
+                            Target.HasQuery = true;
+                        }
+                    }
+                    else
+                    {
+						int32 FoundIndex;
+						InUri.Path.FindChar(TCHAR('/'), FoundIndex);
+                        if (FoundIndex == 0)
+                        {
+                            Target.Path = CleanPath(InUri.Path);
+                        }
+                        else
+                        {
+                            if (HasAuthority() && Path.IsEmpty())
+                            {
+								static const FString kTextSlash(TEXT("/"));
+                                Target.Path = CleanPath(kTextSlash + InUri.Path);
+                            }
+                            else
+                            {
+                                Target.Path = CleanPath(BasePath(Path) + InUri.Path);
+                            }
+                        }
+                        Target.QueryString = InUri.QueryString;
+                        Target.HasQuery = InUri.HasQuery;
+                    }
+                }
+            }
+
+            Target.HasFragment = InUri.HasFragment;
+            Target.Fragment = InUri.Fragment;
+
+            return Target;
+        }
+
+
+        FString parseScheme(StringHelpers::FStringIterator& it, FURIParseError& OutError)
+        {
+            FString scheme;
+            int start = 0;
+
+            while (it)
+            {
+                if ((TCHAR('a') <= *it && *it <= TCHAR('z')) || (TCHAR('A') <= *it && *it <= TCHAR('Z')))
+                {
+                }
+                else if ((TCHAR('0') <= *it && *it <= TCHAR('9')) || *it == TCHAR('-') || *it == TCHAR('+') || *it == TCHAR('.'))
+                {
+                    if (start == 0)
+                    {
+                        break;
+                    }
+                }
+                else if (*it == TCHAR(':'))
+                {
+                    if (start == 0)
+                    {
+                        OutError = FURIParseError::InvalidScheme;
+                    }
+
+                    ++it;
+                    return scheme;
+                }
+                else
+                {
+                    break;
+                }
+				scheme.AppendChar(FChar::ToLower(*it));
+                ++start;
+                ++it;
+            }
+
+            // Scheme is only found if there's a ":" and it's valid.
+            it -= start;
+            return TEXT("");
+        }
+
+        /**
+         * Splits a string using it's iterators based on the provided Split character at the first time it occurs. The
+         * first part of the string is passed into OutBefore, the second into OutAfter. The Split character doesn't
+         * occur in neither.
+         * @param it Will be iterated until end is reached.
+         * @param end
+         * @param OutBefore
+         * @param OutAfter
+         * @param Split The character used for splitting.
+         * @return Identifies if the Split character was found. This is necessary in case OutAfter is empty.
+         */
+        bool splitChar(StringHelpers::FStringIterator& it, FString& OutBefore, FString& OutAfter, const TCHAR Split)
+        {
+            bool bSplitFound = false;
+
+            while (it)
+            {
+                if (*it == Split && !bSplitFound)
+                {
+                    bSplitFound = true;
+                    ++it;
+                    continue;
+                }
+
+                if (bSplitFound)
+                {
+                    OutAfter += *it;
+                }
+                else
+                {
+                    OutBefore += *it;
+                }
+
+                ++it;
+            }
+
+            return bSplitFound;
+        }
+
+        FURIParseError parseAuthority(StringHelpers::FStringIterator& it, FString& OutHost, FURIUserInfo& OutInfo)
+        {
+            FString Part;
+            FString Info;
+
+            while (it)
+            {
+                if (*it == TCHAR('/'))
+                {
+                    break;
+                }
+
+                if (*it == TCHAR('@'))
+                {
+                    Info = Part;
+                    Part = TEXT("");
+                    ++it;
+                    continue;
+                }
+
+                Part += *it;
+                ++it;
+            }
+
+            OutHost = Part;
+
+            if (!Info.IsEmpty())
+            {
+                StringHelpers::FStringIterator InfoIt(Info);
+                splitChar(InfoIt, OutInfo.Username, OutInfo.Password, TCHAR(':'));
+            }
+
+            return FURIParseError::None;
+        }
+
+        FURI FURI::Parse(const FString& InUri)
+        {
+			StringHelpers::FStringIterator it(InUri);
+            FURI OutUrl;
+
+            FString beforeFragment;
+            OutUrl.HasFragment = splitChar(it, beforeFragment, OutUrl.Fragment, TCHAR('#'));
+            
+			StringHelpers::FStringIterator it2(beforeFragment);
+            OutUrl.Scheme = parseScheme(it2, OutUrl.ParseError);
+            if (OutUrl.ParseError != FURIParseError::None)
+            {
+                return OutUrl;
+            }
+
+            FString rest;
+            OutUrl.HasQuery = splitChar(it2, rest, OutUrl.QueryString, TCHAR('?'));
+
+			StringHelpers::FStringIterator it3(rest);
+            if (it3 && *it3 != TCHAR('/'))
+            {
+                if (!OutUrl.Scheme.IsEmpty())
+                {
+                    OutUrl.Opaque = rest;
+                    return OutUrl;
+                }
+
+				int32 colon, slash;
+				rest.FindChar(TCHAR(':'), colon);
+				rest.FindChar(TCHAR('/'), slash);
+                if (colon != INDEX_NONE && (slash == INDEX_NONE || colon < slash))
+                {
+                    OutUrl.ParseError = FURIParseError::PathSegmentContainsColon;
+                    return OutUrl;
+                }
+            }
+
+            if ((!OutUrl.Scheme.IsEmpty() || rest.Find(TEXT("///")) != 0) && rest.Find(TEXT("//")) == 0)
+            {
+                it3 += 2;
+                OutUrl.ParseError = parseAuthority(it3, OutUrl.Host, OutUrl.UserInfo);
+
+                if (OutUrl.ParseError != FURIParseError::None)
+                {
+                    return OutUrl;
+                }
+            }
+
+            while (it3)
+            {
+                OutUrl.Path += *it3;
+                ++it3;
+            }
+
+            return OutUrl;
+        }
+
+        FString FURI::BasePath(const FString& InPath)
+        {
+            FString Result;
+            FString Part;
+
+            StringHelpers::FStringIterator PathIt(InPath);
+
+            while (PathIt)
+            {
+                Part += *PathIt;
+
+                if (*PathIt == TCHAR('/'))
+                {
+                    Result += Part;
+                    Part = TEXT("");
+                }
+
+                ++PathIt;
+            }
+
+            return Result;
+        }
+
+        FString FURI::CleanPath(const FString& InPath)
+        {
+			TArray<FString> PathParts;
+            FString CurrentPath;
+
+            unsigned int Size = (unsigned int)InPath.Len();
+            unsigned int Position = 0;
+            bool bRooted = false;
+
+            if (Size > 0 && InPath[0] == TCHAR('/'))
+            {
+                bRooted = true;
+                ++Position;
+            }
+
+            for (; Position < Size; ++Position)
+            {
+                if (InPath[Position] == TCHAR('.') && (Position + 1 == Size || InPath[Position + 1] == TCHAR('/')) &&
+                    (Position == 0 || InPath[Position - 1] == TCHAR('/')))
+                {
+                    ++Position;
+                }
+                else if (InPath[Position] == TCHAR('.') && (Position + 1 < Size && InPath[Position + 1] == TCHAR('.') &&
+                                                     (Position + 2 == Size || InPath[Position + 2] == TCHAR('/')) &&
+                                                     (Position == 0 || InPath[Position - 1] == TCHAR('/'))))
+                {
+                    if (!PathParts.IsEmpty())
+                    {
+                        PathParts.Pop(false);
+                    }
+                    Position += 2;
+                }
+                else if (InPath[Position] != TCHAR('/'))
+                {
+                    CurrentPath += InPath[Position];
+                }
+                else if (!CurrentPath.IsEmpty())
+                {
+                    PathParts.Push(CurrentPath);
+                    CurrentPath = TEXT("");
+                }
+            }
+
+            if (!CurrentPath.IsEmpty())
+            {
+                PathParts.Push(CurrentPath);
+            }
+
+            FString Result;
+            if (bRooted)
+            {
+                Result.AppendChar('/');
+            }
+
+            for (int32 i = 0; i < PathParts.Num(); i++)
+            {
+                if (i != 0)
+                {
+                    Result.AppendChar(TCHAR('/'));
+                }
+                Result += PathParts[i];
+            }
+
+            // In case the last element was a positional element.
+            if (!Result.IsEmpty() && Result[Result.Len() - 1] != TCHAR('/'))
+            {
+                if ((Size > 2 && InPath[Size - 2] == TCHAR('/') && InPath[Size - 1] == TCHAR('.')) ||
+                    (Size > 3 && InPath[Size - 3] == TCHAR('/') && InPath[Size - 2] == TCHAR('.') && InPath[Size - 1] == TCHAR('.')))
+                {
+					Result.AppendChar(TCHAR('/'));
+                }
+                else if (Size > 2 && InPath[Size - 1] == TCHAR('/'))
+                {
+					Result.AppendChar(TCHAR('/'));
+                }
+            }
+
+            return Result;
+        }
+
+    }
+} // namespace Electra
+
+
diff --git a/Engine/Plugins/Media/ElectraPlayer/Source/ElectraPlayerRuntime/Private/Runtime/Utilities/URI.h b/Engine/Plugins/Media/ElectraPlayer/Source/ElectraPlayerRuntime/Private/Runtime/Utilities/URI.h
new file mode 100644
--- /dev/null
+++ b/Engine/Plugins/Media/ElectraPlayer/Source/ElectraPlayerRuntime/Private/Runtime/Utilities/URI.h
@@ -0,0 +1,162 @@
+// Copyright Epic Games, Inc. All Rights Reserved.
+
+#pragma once
+
+#include "PlayerCore.h"
+
+namespace Electra
+{
+
+    namespace Utilities
+    {
+
+        enum class FURIParseError
+        {
+            None,
+            InvalidScheme,
+            PathSegmentContainsColon,
+        };
+
+        struct FURIUserInfo
+        {
+            FString Username;
+            FString Password;
+        };
+
+        /**
+         * Implementation of a Uniform Resource Identifier implementation based on RFC 3986 which currently doesn't cover
+         * escaping / unescaping of paths, username, password and query variables.
+         */
+        class FURI
+        {
+        private:
+            FURIParseError ParseError = FURIParseError::None;
+        public:
+            FString Scheme;
+
+            /**
+             * Host which contains a port if present.
+             */
+            FString Host;
+            FString Path;
+
+            /**
+             * Identifies that a fragment is added to the URI even if it's empty.
+             */
+            bool HasFragment;
+
+            /**
+             * Encoded fragment which will be added after the #.
+             */
+            FString Fragment;
+
+            /**
+             * Opaque as of RFC 3986 is used for URI schemes where no hierarchy is visible. For example:
+             *
+             *   mailto:[REDACTED_EMAIL]
+             *   tel:+1-123-123-1234
+             *
+             * This is set if no authority or path was detected and can never be set together with Path.
+             */
+            FString Opaque;
+            FURIUserInfo UserInfo;
+
+            /**
+             * Identifies that a query is added to the URI even if it's empty.
+             */
+            bool HasQuery;
+
+            /**
+             * The query as an encoded string.
+             */
+            FString QueryString;
+
+            /**
+             * After using FURI::Parse this returns if the URI was successfully parsed.
+             * @return
+             */
+            inline const bool IsValid()
+            {
+                return ParseError == FURIParseError::None;
+            }
+
+            inline const FURIParseError GetParseError()
+            {
+                return ParseError;
+            }
+
+            /**
+             * Returns true if either Username or Password is available.
+             * @return
+             */
+            inline bool HasUserInfo() const
+            {
+                return !UserInfo.Password.IsEmpty() || !UserInfo.Username.IsEmpty();
+            }
+
+            /**
+             * Returns true if either a Host or Username/Password is available.
+             * @return
+             */
+            inline bool HasAuthority() const
+            {
+                return HasUserInfo() || !Host.IsEmpty();
+            }
+
+            /**
+             * Turns the URI into a string.
+             *
+             *   scheme ":" [ "//" authority ] ( path | opaque ) [ "?" query ] [ "#" fragment ]
+             *
+             * @return
+             */
+            const FString Format() const;
+
+            /**
+             * Resolves and returns a new absolute target url based on this FUrl instance using the provided InUrl as
+             * reference. Note that it's using the strict format of transformation, see
+             * (https://tools.ietf.org/html/rfc3986#section-5.2.2) for details.
+             * @param InUri The URI used for resolving the new URI, which can be relative.
+             * @return A new instance which contains the absolute resolved URI.
+             */
+            const FURI Resolve(const FURI& InUri) const;
+
+            /**
+             * Parses a string into a FURI. If the string doesn't represent a valid URL the resulting FURI.IsValid()
+             * will return false. To receive the proper error use FURI.GetParseError().
+             * @param InUri
+             * @return
+             */
+            static FURI Parse(const FString& InUri);
+
+            /**
+             * Returns the base path of an URI path.
+             *
+             * Example:
+             *  "/some/path" -> "/some/"
+             *  "/test/path/" -> "/test/path/"
+             *  "/" -> "/"
+             *  "test" -> ""
+             *  "relative/path" -> "relative"
+             *
+             * @param InPath
+             * @return
+             */
+            static FString BasePath(const FString& InPath);
+
+            /**
+             * Cleans up a path by removing "." and "..".
+             *
+             * Example:
+             *  "test/../path" -> "path"
+             *  "test/./path" -> "test/path"
+             *
+             * @param InPath
+             * @return
+             */
+            static FString CleanPath(const FString& InPath);
+        };
+    }
+} // namespace Electra
+
+
diff --git a/Engine/Plugins/Media/ElectraPlayer/Source/ElectraPlayerRuntime/Private/Runtime/Utilities/URLParser.cpp b/Engine/Plugins/Media/ElectraPlayer/Source/ElectraPlayerRuntime/Private/Runtime/Utilities/URLParser.cpp
new file mode 100644
--- /dev/null
+++ b/Engine/Plugins/Media/ElectraPlayer/Source/ElectraPlayerRuntime/Private/Runtime/Utilities/URLParser.cpp
@@ -0,0 +1,86 @@
+// Copyright Epic Games, Inc. All Rights Reserved.
+
+#include "Utilities/URLParser.h"
+#include "Utilities/URI.h"
+#include "Utilities/StringHelpers.h"
+
+
+namespace Electra
+{
+
+	class FURLParser : public IURLParser
+	{
+	public:
+		FURLParser();
+		virtual ~FURLParser();
+
+		virtual UEMediaError ParseURL(const FString& URL) override;
+		virtual FString GetPath() const override;
+		virtual void GetPathComponents(TArray<FString>& OutPathComponents) const override;
+
+		virtual FString ResolveWith(const FString& RelativeURL) const override;
+
+
+	private:
+		Utilities::FURI* URLHelper;
+	};
+
+
+	IURLParser* IURLParser::Create()
+	{
+		return new FURLParser;
+	}
+
+
+
+	FURLParser::FURLParser()
+		: URLHelper(nullptr)
+	{
+	}
+
+	FURLParser::~FURLParser()
+	{
+		delete URLHelper;
+	}
+
+	UEMediaError FURLParser::ParseURL(const FString& URL)
+	{
+		delete URLHelper;
+		URLHelper = new Utilities::FURI(Utilities::FURI::Parse(URL));
+		if (!URLHelper)
+		{
+			return UEMEDIA_ERROR_OOM;
+		}
+		return URLHelper->IsValid() ? UEMEDIA_ERROR_OK : UEMEDIA_ERROR_FORMAT_ERROR;
+	}
+
+
+	FString FURLParser::GetPath() const
+	{
+		return URLHelper ? URLHelper->Path : FString();
+	}
+
+	void FURLParser::GetPathComponents(TArray<FString>& OutPathComponents) const
+	{
+		OutPathComponents.Empty();
+		if (URLHelper)
+		{
+			static const FString kTextSlash(TEXT("/"));
+			StringHelpers::SplitByDelimiter(OutPathComponents, URLHelper->Path, kTextSlash);
+		}
+	}
+
+	FString FURLParser::ResolveWith(const FString& RelativeURL) const
+	{
+		if (URLHelper)
+		{
+			return URLHelper->Resolve(Utilities::FURI::Parse(RelativeURL)).Format();
+		}
+		return FString();
+	}
+
+
+
+} // namespace Electra
+
+
diff --git a/Engine/Plugins/Media/ElectraPlayer/Source/ElectraPlayerRuntime/Private/Runtime/Utilities/URLParser.h b/Engine/Plugins/Media/ElectraPlayer/Source/ElectraPlayerRuntime/Private/Runtime/Utilities/URLParser.h
new file mode 100644
--- /dev/null
+++ b/Engine/Plugins/Media/ElectraPlayer/Source/ElectraPlayerRuntime/Private/Runtime/Utilities/URLParser.h
@@ -0,0 +1,32 @@
+// Copyright Epic Games, Inc. All Rights Reserved.
+
+#pragma once
+
+#include "PlayerCore.h"
+
+
+namespace Electra
+{
+
+	class IURLParser
+	{
+	public:
+		static IURLParser* Create();
+		virtual ~IURLParser() = default;
+
+		virtual UEMediaError ParseURL(const FString& URL) = 0;
+
+		virtual FString GetPath() const = 0;
+		virtual void GetPathComponents(TArray<FString>& OutPathComponents) const = 0;
+
+		virtual FString ResolveWith(const FString& RelativeURL) const = 0;
+
+	protected:
+		IURLParser() = default;
+		IURLParser(const IURLParser&) = delete;
+	};
+
+
+} // namespace Electra
+
+
