ElectraPlayer: RELEASE. Moved from NFL to public. Note that this is done via a Delete and "Add New" operation as discussed.

diff --git a/Engine/Plugins/Media/ElectraPlayer/Source/ElectraPlayerRuntime/Private/Runtime/Utilities/UtilsMPEG.h b/Engine/Plugins/Media/ElectraPlayer/Source/ElectraPlayerRuntime/Private/Runtime/Utilities/UtilsMPEG.h
new file mode 100644
--- /dev/null
+++ b/Engine/Plugins/Media/ElectraPlayer/Source/ElectraPlayerRuntime/Private/Runtime/Utilities/UtilsMPEG.h
@@ -0,0 +1,120 @@
+// Copyright Epic Games, Inc. All Rights Reserved.
+
+#pragma once
+
+#include "PlayerCore.h"
+#include "BitDataStream.h"
+
+
+namespace Electra
+{
+	namespace MPEG
+	{
+
+		class FESDescriptor
+		{
+		public:
+			FESDescriptor();
+			void SetRawData(const void* Data, int64 Size);
+			const TArray<uint8>& GetRawData() const;
+
+			bool Parse();
+
+			const TArray<uint8>& GetCodecSpecificData() const;
+
+			uint32 GetBufferSize() const
+			{
+				return BufferSize;
+			}
+
+			uint32 GetMaxBitrate() const
+			{
+				return MaxBitrate;
+			}
+
+			uint32 GetAvgBitrate() const
+			{
+				return AvgBitrate;
+			}
+
+			// See http://mp4ra.org/#/object_types
+			enum class FObjectTypeID
+			{
+				Unknown = 0,
+				System_V1 = 1,
+				System_V2 = 2,
+				MPEG4_Video = 32,
+				MPEG4_AVC_SPS = 33,
+				MPEG4_AVC_PPS = 34,
+				MPEG4_Audio = 64,
+				MPEG2_SimpleVideo = 96,
+				MPEG2_MainVideo = 97,
+				MPEG2_SNRVideo = 98,
+				MPEG2_SpatialVideo = 99,
+				MPEG2_HighVideo = 100,
+				MPEG2_422Video = 101,
+				MPEG4_ADTS_Main = 102,
+				MPEG4_ADTS_LC = 103,
+				MPEG4_ADTS_SSR = 104,
+				MPEG2_ADTS = 105,
+				MPEG1_Video = 106,
+				MPEG1_ADTS = 107,
+				JPEG_Video = 108,
+				Private_Audio = 192,
+				Private_Video = 208,
+				PCM16LE_Audio = 224,
+				Vorbis_Audio = 225,
+				AC3_Audio = 226,
+				ALaw_Audio = 227,
+				MuLaw_Audio = 228,
+				G723_ADPCM_Audio = 229,
+				PCM16BE = 230,
+				YV12_Video = 240,
+				H264_Video = 241,
+				H263_Video = 242,
+				H261_Video = 243,
+			};
+			enum class FStreamType
+			{
+				Unknown = 0,
+				ObjectDescription = 1,
+				ClockReference = 2,
+				SceneDescription = 4,  	 //!< aka. visual description
+				Audio = 5,
+				MPEG7 = 6,
+				IPMP = 7,
+				OCI = 8,
+				MPEG_Java = 9,
+				User_Private = 32
+			};
+
+			FObjectTypeID GetObjectTypeID() const
+			{
+				return ObjectTypeID;
+			}
+
+			FStreamType GetStreamType() const
+			{
+				return StreamTypeID;
+			}
+
+		private:
+			TArray<uint8>						RawData;
+
+			TArray<uint8>						CodecSpecificData;
+			FObjectTypeID						ObjectTypeID;
+			FStreamType							StreamTypeID;
+			uint32								BufferSize;
+			uint32								MaxBitrate;
+			uint32								AvgBitrate;
+			uint16								ESID;
+			uint16								DependsOnStreamESID;
+			uint8								StreamPriority;
+			bool								bDependsOnStream;
+		};
+
+
+	} // namespace MPEG
+} // namespace Electra
+
+
diff --git a/Engine/Plugins/Media/ElectraPlayer/Source/ElectraPlayerRuntime/Private/Runtime/Utilities/UtilsMPEGAudio.cpp b/Engine/Plugins/Media/ElectraPlayer/Source/ElectraPlayerRuntime/Private/Runtime/Utilities/UtilsMPEGAudio.cpp
new file mode 100644
--- /dev/null
+++ b/Engine/Plugins/Media/ElectraPlayer/Source/ElectraPlayerRuntime/Private/Runtime/Utilities/UtilsMPEGAudio.cpp
@@ -0,0 +1,249 @@
+// Copyright Epic Games, Inc. All Rights Reserved.
+
+#include "Utilities/UtilsMPEGAudio.h"
+#include "BitDataStream.h"
+
+namespace Electra
+{
+	namespace MPEG
+	{
+
+		namespace AACParseHelper
+		{
+			static uint32 GIndexToSampleRate[16] =
+			{
+			96000, 88200, 64000, 48000, 44100, 32000, 24000, 22050, 16000, 12000, 11025, 8000, 0, 0, 0, 0
+			};
+
+			static uint32 GetAudioObjectType(FBitDataStream& Bitstream)
+			{
+				uint32 ObjectType = Bitstream.GetBits(5);
+				if (ObjectType == 31)
+				{
+					ObjectType = 32 + Bitstream.GetBits(6);
+				}
+				return ObjectType;
+			}
+			static uint32 GetSamplingRateIndex(FBitDataStream& Bitstream, uint32& Rate)
+			{
+				uint32 SamplingFrequencyIndex = Bitstream.GetBits(4);
+				if (SamplingFrequencyIndex == 15)
+				{
+					Rate = Bitstream.GetBits(24);
+				}
+				else
+				{
+					Rate = GIndexToSampleRate[SamplingFrequencyIndex];
+				}
+				return SamplingFrequencyIndex;
+			}
+			static void GetGASpecificConfig(FBitDataStream& Bitstream, uint32 SamplingFrequencyIndex, uint32 ChannelConfiguration, uint32 AudioObjectType)
+			{
+				int32 FrameLengthFlag = Bitstream.GetBits(1);
+				int32 DependsOnCoreDecoder = Bitstream.GetBits(1);
+				if (DependsOnCoreDecoder)
+				{
+					Bitstream.SkipBits(14);		// coreCoderDelay
+				}
+				int32 ExtensionFlag = Bitstream.GetBits(1);
+				if (ChannelConfiguration == 0)
+				{
+					// TODO:
+				}
+				if (AudioObjectType == 6 || AudioObjectType == 20)
+				{
+					Bitstream.SkipBits(3);	// layerNr
+				}
+				if (ExtensionFlag)
+				{
+					if (AudioObjectType == 22)
+					{
+						Bitstream.SkipBits(5);	// numOfSubFrame
+						Bitstream.SkipBits(11);	// layer_length
+					}
+					if (AudioObjectType == 17 || AudioObjectType == 19 || AudioObjectType == 20 || AudioObjectType == 23)
+					{
+						Bitstream.SkipBits(1);	// aacSectionDataResilienceFlag
+						Bitstream.SkipBits(1);	// aacScalefactorDataResilienceFlag;
+						Bitstream.SkipBits(1);	// aacSpectralDataResilienceFlag;
+					}
+					int32 ExtensionFlag3 = Bitstream.GetBits(1);
+					if (ExtensionFlag)
+					{
+						// TODO:
+					}
+				}
+			}
+
+		} // namespace AACParseHelper
+
+
+		FAACDecoderConfigurationRecord::FAACDecoderConfigurationRecord()
+		{
+			Reset();
+		}
+
+		void FAACDecoderConfigurationRecord::Reset()
+		{
+			SBRSignal = -1;
+			PSSignal = -1;
+			ChannelConfiguration = 0;
+			SamplingFrequencyIndex = 0;
+			SamplingRate = 0;
+			ExtSamplingFrequencyIndex = 0;
+			ExtSamplingFrequency = 0;
+			AOT = 0;
+			ExtAOT = 0;
+		}
+
+
+		const TArray<uint8>& FAACDecoderConfigurationRecord::GetCodecSpecificData() const
+		{
+			return CodecSpecificData;
+		}
+
+		bool FAACDecoderConfigurationRecord::ParseFrom(const void* Data, int64 Size)
+		{
+			CodecSpecificData.Empty();
+			CodecSpecificData.SetNumUninitialized(Size);
+			FMemory::Memcpy(CodecSpecificData.GetData(), Data, Size);
+
+			FBitDataStream bsp(Data, Size);
+			SBRSignal = -1;
+			PSSignal = -1;
+			SamplingRate = 0;
+			ExtSamplingFrequency = 0;
+			ExtSamplingFrequencyIndex = 0;
+			ExtAOT = 0;
+			AOT = AACParseHelper::GetAudioObjectType(bsp);
+			SamplingFrequencyIndex = AACParseHelper::GetSamplingRateIndex(bsp, SamplingRate);
+			ChannelConfiguration = bsp.GetBits(4);
+
+			if (AOT == 5 /*SBR*/ || AOT == 29 /*PS*/)
+			{
+				ExtAOT = AOT;
+				SBRSignal = 1;
+				if (AOT == 29)
+				{
+					PSSignal = 1;
+				}
+
+				ExtSamplingFrequencyIndex = AACParseHelper::GetSamplingRateIndex(bsp, ExtSamplingFrequency);
+				AOT = AACParseHelper::GetAudioObjectType(bsp);
+			}
+			// Handle supported AOT configs
+			if (AOT == 2 /*LC*/)		// Only LC for now
+			{
+				AACParseHelper::GetGASpecificConfig(bsp, SamplingFrequencyIndex, ChannelConfiguration, AOT);
+			}
+			// Would need to handle epConfig here now for a couple of AOT's we are NOT supporting
+			// ...
+
+			// Check for backward compatible SBR signaling.
+			if (ExtAOT != 5)
+			{
+				while (bsp.GetRemainingBits() > 15)
+				{
+					int32 syncExtensionType = bsp.PeekBits(11);
+					if (syncExtensionType == 0x2b7)
+					{
+						bsp.SkipBits(11);
+						ExtAOT = AACParseHelper::GetAudioObjectType(bsp);
+						if (ExtAOT == 5)
+						{
+							SBRSignal = bsp.GetBits(1);
+							if (SBRSignal)
+							{
+								ExtSamplingFrequencyIndex = AACParseHelper::GetSamplingRateIndex(bsp, ExtSamplingFrequency);
+							}
+						}
+						if (bsp.GetRemainingBits() > 11 && bsp.GetBits(11) == 0x548)
+							PSSignal = bsp.GetBits(1);
+						break;
+					}
+					else
+					{
+						bsp.SkipBits(1);
+					}
+				}
+			}
+
+			int64 remainingBits = bsp.GetRemainingBits();
+			MEDIA_UNUSED_VAR(remainingBits);
+			return true;
+		}
+
+
+
+
+
+	} // namespace MPEG
+} // namespace Electra
+
+#if 0
+
+struct ADTSheader
+{
+	// Refer to http://wiki.multimedia.cx/index.php?title=MPEG-4_Audio
+	unsigned mSyncWord : 12;		//!< All set
+	unsigned mVersion : 1;			//!< 0=MPEG-4, 1=MPEG-2
+	unsigned mLayer : 2;			//!< 0
+	unsigned mProtAbsent : 1;		//!< 0 if CRC present, 1 if CRC absent
+	unsigned mProfile : 2;			//!< MPEG-4 audio object type -1
+	unsigned mSmpFrqIndex : 4;		//!< MPEG-4 sampling frequency index (15 is forbidden)
+	unsigned mPrivateStream : 1;
+	unsigned mChannelCfg : 3;		//!< MPEG-4 channel configuration
+	unsigned mOriginality : 1;
+	unsigned mHome : 1;
+	unsigned mCopyrighted : 1;
+	unsigned mCopyrightStart : 1;
+	unsigned mFrameLength : 13;
+	unsigned mBufferFullness : 11;
+	unsigned mNumFrames : 2;
+	unsigned mCRCifPresent : 16;
+};
+
+static bool GetADTSHeader(ADTSheader& header, const void* pData, int64 nDataSize)
+{
+	if (nDataSize < 7)
+		return(false);
+
+	FBitDataStream bs(pData, nDataSize);
+	header.mSyncWord = bs.GetBits(12);
+	header.mVersion = bs.GetBits(1);
+	header.mLayer = bs.GetBits(2);
+	header.mProtAbsent = bs.GetBits(1);
+	header.mProfile = bs.GetBits(2);
+	header.mSmpFrqIndex = bs.GetBits(4);
+	header.mPrivateStream = bs.GetBits(1);
+	header.mChannelCfg = bs.GetBits(3);
+	header.mOriginality = bs.GetBits(1);
+	header.mHome = bs.GetBits(1);
+	header.mCopyrighted = bs.GetBits(1);
+	header.mCopyrightStart = bs.GetBits(1);
+	header.mFrameLength = bs.GetBits(13);
+	header.mBufferFullness = bs.GetBits(11);
+	header.mNumFrames = bs.GetBits(2);
+	if (nDataSize >= 9 && !header.mProtAbsent)
+		header.mCRCifPresent = bs.GetBits(16);
+	return(header.mSyncWord == 0xfff && header.mLayer == 0);
+}
+
+/*
+According to the Audio Coding Wiki, which is no longer available as of June, 2010, the ADIF format actually is just one header at the beginning of the AAC file.
+And the rest of the data are consecutive raw data blocks.
+This file format is meant for simple local storing purposes unlike ADTS or LATM which are meant for streaming AAC.
+The ADIF header is made up of the following Field name/Field size in bits/Comment (separated by semicolons):
+	adif_id/32/Always: "ADIF";
+	copyright_id_present/1/[none];
+	copyright_id/72/only if copyright_id_present == 1;
+	original_copy/1/[none];
+	home/1/[none];
+	bitstream_type/1/0: CBR, 1: VBR; bitrate/23/for CBR: bitrate, for VBR: peak bitrate, 0 means unknown;
+	num_program_config_elements/4/[none].
+	The next 2 fields come (num_program_config_elements+1) times: buffer_fullness/20/only if bitstream_type == 0; program_config_element/VAR/[none].
+*/
+
+
+#endif
+
diff --git a/Engine/Plugins/Media/ElectraPlayer/Source/ElectraPlayerRuntime/Private/Runtime/Utilities/UtilsMPEGAudio.h b/Engine/Plugins/Media/ElectraPlayer/Source/ElectraPlayerRuntime/Private/Runtime/Utilities/UtilsMPEGAudio.h
new file mode 100644
--- /dev/null
+++ b/Engine/Plugins/Media/ElectraPlayer/Source/ElectraPlayerRuntime/Private/Runtime/Utilities/UtilsMPEGAudio.h
@@ -0,0 +1,40 @@
+// Copyright Epic Games, Inc. All Rights Reserved.
+
+#pragma once
+
+#include "PlayerCore.h"
+#include "BitDataStream.h"
+
+namespace Electra
+{
+	namespace MPEG
+	{
+
+		class FAACDecoderConfigurationRecord
+		{
+		public:
+			FAACDecoderConfigurationRecord();
+
+			bool ParseFrom(const void* Data, int64 Size);
+			void Reset();
+			const TArray<uint8>& GetCodecSpecificData() const;
+
+			int32		SBRSignal;
+			int32		PSSignal;
+			uint32		ChannelConfiguration;
+			uint32		SamplingFrequencyIndex;
+			uint32		SamplingRate;
+			uint32		ExtSamplingFrequencyIndex;
+			uint32		ExtSamplingFrequency;
+			uint32		AOT;
+			uint32		ExtAOT;
+		private:
+			TArray<uint8>	CodecSpecificData;
+		};
+
+
+
+
+	} // namespace MPEG
+} // namespace Electra
+
diff --git a/Engine/Plugins/Media/ElectraPlayer/Source/ElectraPlayerRuntime/Private/Runtime/Utilities/UtilsMPEGVideo.cpp b/Engine/Plugins/Media/ElectraPlayer/Source/ElectraPlayerRuntime/Private/Runtime/Utilities/UtilsMPEGVideo.cpp
new file mode 100644
--- /dev/null
+++ b/Engine/Plugins/Media/ElectraPlayer/Source/ElectraPlayerRuntime/Private/Runtime/Utilities/UtilsMPEGVideo.cpp
@@ -0,0 +1,506 @@
+// Copyright Epic Games, Inc. All Rights Reserved.
+
+#include "PlayerCore.h"
+#include "Utilities/UtilsMPEGVideo.h"
+#include "BitDataStream.h"
+
+
+namespace Electra
+{
+	namespace MPEG
+	{
+
+		FAVCDecoderConfigurationRecord::FAVCDecoderConfigurationRecord()
+		{
+			ConfigurationVersion = 0;
+			AVCProfileIndication = 0;
+			ProfileCompatibility = 0;
+			AVCLevelIndication = 0;
+			NALUnitLength = 0;
+			ChromaFormat = 0;
+			BitDepthLumaMinus8 = 0;
+			BitDepthChromaMinus8 = 0;
+			bHaveAdditionalProfileIndication = false;
+		}
+
+		void FAVCDecoderConfigurationRecord::SetRawData(const void* Data, int64 Size)
+		{
+			RawData.Empty();
+			if (Size)
+			{
+				RawData.Reserve((uint32)Size);
+				RawData.SetNumUninitialized((uint32)Size);
+				FMemory::Memcpy(RawData.GetData(), Data, Size);
+			}
+		}
+
+		const TArray<uint8>& FAVCDecoderConfigurationRecord::GetRawData() const
+		{
+			return RawData;
+		}
+
+		const TArray<uint8>& FAVCDecoderConfigurationRecord::GetCodecSpecificData() const
+		{
+			return CodecSpecificData;
+		}
+
+		const TArray<uint8>& FAVCDecoderConfigurationRecord::GetCodecSpecificDataSPS() const
+		{
+			return CodecSpecificDataSPSOnly;
+		}
+
+		const TArray<uint8>& FAVCDecoderConfigurationRecord::GetCodecSpecificDataPPS() const
+		{
+			return CodecSpecificDataPPSOnly;
+		}
+
+		int32 FAVCDecoderConfigurationRecord::GetNumberOfSPS() const
+		{
+			return ParsedSPSs.Num();
+		}
+
+		const FISO14496_10_seq_parameter_set_data& FAVCDecoderConfigurationRecord::GetParsedSPS(int32 SpsIndex) const
+		{
+			check(SpsIndex < GetNumberOfSPS());
+			return ParsedSPSs[SpsIndex];
+		}
+
+
+		bool FAVCDecoderConfigurationRecord::Parse()
+		{
+			CodecSpecificData.Empty();
+
+			if (RawData.Num())
+			{
+				FByteReader BitReader(RawData.GetData(), RawData.Num());
+
+				if (!BitReader.ReadByte(ConfigurationVersion))
+				{
+					return false;
+				}
+				if (!BitReader.ReadByte(AVCProfileIndication))
+				{
+					return false;
+				}
+				if (!BitReader.ReadByte(ProfileCompatibility))
+				{
+					return false;
+				}
+				if (!BitReader.ReadByte(AVCLevelIndication))
+				{
+					return false;
+				}
+				if (!BitReader.ReadByte(NALUnitLength))
+				{
+					return false;
+				}
+				NALUnitLength = (NALUnitLength & 3) + 1;
+				uint8 nSPS = 0;
+				if (!BitReader.ReadByte(nSPS))
+				{
+					return false;
+				}
+				nSPS &= 31;
+				if (nSPS)
+				{
+					SequenceParameterSets.Reserve(nSPS);
+					ParsedSPSs.Reserve(nSPS);
+				}
+				int32 TotalSPSSize = 0;
+				for (int32 i = 0; i < nSPS; ++i)
+				{
+					uint16 spsLen = 0;
+					if (!BitReader.ReadByte(spsLen))
+					{
+						return false;
+					}
+					TArray<uint8>& sps = SequenceParameterSets.AddDefaulted_GetRef();
+					sps.Reserve(spsLen);
+					sps.SetNumUninitialized(spsLen);
+					if (!BitReader.ReadBytes(sps.GetData(), spsLen))
+					{
+						return false;
+					}
+					TotalSPSSize += 4 + spsLen;		// 4 because we always use 32 bit startcode and not NALUnitLength
+					ParseH264SPS(ParsedSPSs.AddDefaulted_GetRef(), sps.GetData(), sps.Num());
+				}
+				uint8 nPPS = 0;
+				if (!BitReader.ReadByte(nPPS))
+				{
+					return false;
+				}
+				if (nPPS)
+				{
+					PictureParameterSets.Reserve(nPPS);
+				}
+				int32 TotalPPSSize = 0;
+				for (int32 i = 0; i < nPPS; ++i)
+				{
+					uint16 ppsLen = 0;
+					if (!BitReader.ReadByte(ppsLen))
+					{
+						return false;
+					}
+					TArray<uint8>& pps = PictureParameterSets.AddDefaulted_GetRef();
+					pps.Reserve(ppsLen);
+					pps.SetNumUninitialized(ppsLen);
+					if (!BitReader.ReadBytes(pps.GetData(), ppsLen))
+					{
+						return false;
+					}
+					TotalPPSSize += 4 + ppsLen;		// 4 because we always use 32 bit startcode and not NALUnitLength
+				}
+
+				if (AVCProfileIndication == 100 || AVCProfileIndication == 110 || AVCProfileIndication == 122 || AVCProfileIndication == 144)
+				{
+					// At least according to the ISO 14496-15:2014 standard these values must appear.
+					// I do have however files that are of AVC profile 100 but omit these values.
+					// Therefore let's do a quick check if we can read at least 4 more bytes.
+					if (BitReader.BytesRemaining() >= 4)
+					{
+						if (!BitReader.ReadByte(ChromaFormat))
+						{
+							return false;
+						}
+						ChromaFormat &= 3;
+						if (!BitReader.ReadByte(BitDepthLumaMinus8))
+						{
+							return false;
+						}
+						BitDepthLumaMinus8 &= 7;
+						if (!BitReader.ReadByte(BitDepthChromaMinus8))
+						{
+							return false;
+						}
+						BitDepthChromaMinus8 &= 7;
+						if (!BitReader.ReadByte(nSPS))
+						{
+							return false;
+						}
+						if (nSPS)
+						{
+							SequenceParameterSetsExt.Reserve(nSPS);
+						}
+						for (int32 i = 0; i < nSPS; ++i)
+						{
+							uint16 spsLenExt = 0;
+							if (!BitReader.ReadByte(spsLenExt))
+							{
+								return false;
+							}
+							TArray<uint8>& spsExt = SequenceParameterSetsExt.AddDefaulted_GetRef();
+							spsExt.Reserve(spsLenExt);
+							spsExt.SetNumUninitialized(spsLenExt);
+							if (!BitReader.ReadBytes(spsExt.GetData(), spsLenExt))
+							{
+								return false;
+							}
+						}
+						bHaveAdditionalProfileIndication = true;
+					}
+				}
+
+				int32 TotalCSDSize = TotalSPSSize + TotalPPSSize;
+				if (TotalCSDSize)
+				{
+					CodecSpecificData.Reserve(TotalCSDSize);
+					CodecSpecificDataSPSOnly.Reserve(TotalSPSSize);
+					for (int32 i = 0; i < SequenceParameterSets.Num(); ++i)
+					{
+						CodecSpecificDataSPSOnly.Push(0);
+						CodecSpecificDataSPSOnly.Push(0);
+						CodecSpecificDataSPSOnly.Push(0);
+						CodecSpecificDataSPSOnly.Push(1);
+						CodecSpecificData.Push(0);
+						CodecSpecificData.Push(0);
+						CodecSpecificData.Push(0);
+						CodecSpecificData.Push(1);
+						for (int32 j = 0, jMax = SequenceParameterSets[i].Num(); j < jMax; ++j)
+						{
+							CodecSpecificDataSPSOnly.Push((SequenceParameterSets[i].GetData())[j]);
+							CodecSpecificData.Push((SequenceParameterSets[i].GetData())[j]);
+						}
+					}
+					CodecSpecificDataPPSOnly.Reserve(TotalPPSSize);
+					for (int32 i = 0; i < PictureParameterSets.Num(); ++i)
+					{
+						CodecSpecificDataPPSOnly.Push(0);
+						CodecSpecificDataPPSOnly.Push(0);
+						CodecSpecificDataPPSOnly.Push(0);
+						CodecSpecificDataPPSOnly.Push(1);
+						CodecSpecificData.Push(0);
+						CodecSpecificData.Push(0);
+						CodecSpecificData.Push(0);
+						CodecSpecificData.Push(1);
+						for (int32 j = 0, jMax = PictureParameterSets[i].Num(); j < jMax; ++j)
+						{
+							CodecSpecificDataPPSOnly.Push((PictureParameterSets[i].GetData())[j]);
+							CodecSpecificData.Push((PictureParameterSets[i].GetData())[j]);
+						}
+					}
+				}
+
+				return true;
+			}
+			return false;
+		}
+
+
+
+
+		static int32 FindStartCode(const uint8* InData, SIZE_T InDataSize, int32& NALUnitLength)
+		{
+			for (const uint8* Data = InData; InDataSize >= 3; ++Data, --InDataSize)
+			{
+				if (Data[0] == 0 && Data[1] == 0 && (Data[2] == 1 || (InDataSize >= 4 && Data[2] == 0 && Data[3] == 1)))
+				{
+					NALUnitLength = Data[2] ? 3 : 4;
+					return Data - reinterpret_cast<const uint8*>(InData);
+				}
+			}
+			NALUnitLength = -1;
+			return -1;
+		}
+
+
+
+		void ParseBitstreamForNALUs(TArray<FNaluInfo>& outNALUs, const void* InBitstream, uint64 InBitstreamLength)
+		{
+			outNALUs.Reset();
+
+			uint64 Pos = 0;
+			uint64 BytesToGo = InBitstreamLength;
+			const uint8* BitstreamData = (const uint8*)InBitstream;
+			while (1)
+			{
+				int32 UnitLength = 0, StartCodePos = FindStartCode(BitstreamData, BytesToGo, UnitLength);
+				if (StartCodePos >= 0)
+				{
+					if (outNALUs.Num())
+					{
+						outNALUs.Last().Size = StartCodePos;
+						outNALUs.Last().Type = *BitstreamData;
+					}
+					FNaluInfo n;
+					n.Offset = Pos + StartCodePos;
+					n.Size = 0;
+					n.Type = 0;
+					n.UnitLength = UnitLength;
+					outNALUs.Push(n);
+					BitstreamData = Electra::AdvancePointer(BitstreamData, StartCodePos + UnitLength);
+					Pos += StartCodePos + UnitLength;
+					BytesToGo -= StartCodePos + UnitLength;
+				}
+				else
+				{
+					if (outNALUs.Num())
+					{
+						outNALUs.Last().Size = BytesToGo;
+						outNALUs.Last().Type = *BitstreamData;
+					}
+					break;
+				}
+			}
+
+		}
+
+
+
+		int32 EBSPtoRBSP(uint8* OutBuf, const uint8* InBuf, int32 NumBytesIn)
+		{
+			uint8* OutBase = OutBuf;
+			while (NumBytesIn-- > 0)
+			{
+				uint8 b = *InBuf++;
+				*OutBuf++ = b;
+				if (b == 0)
+				{
+					if (NumBytesIn > 1)
+					{
+						if (InBuf[0] == 0x00 && InBuf[1] == 0x03)
+						{
+							*OutBuf++ = 0x00;
+							InBuf += 2;
+							NumBytesIn -= 2;
+						}
+					}
+				}
+			}
+			return OutBuf - OutBase;
+		}
+
+
+
+		template <typename T>
+		struct FScopedDataPtr
+		{
+			FScopedDataPtr(void* Addr)
+				: Data(static_cast<T*>(Addr))
+			{
+			}
+			~FScopedDataPtr()
+			{
+				FMemory::Free(static_cast<void*>(Data));
+			}
+			operator T* ()
+			{
+				return Data;
+			}
+			T* Data;
+		};
+
+
+
+
+		bool ParseH264SPS(FISO14496_10_seq_parameter_set_data& OutSPS, const void* Data, int32 Size)
+		{
+			struct FSyntaxElement
+			{
+				static uint32 ue_v(FBitDataStream& BitStream)
+				{
+					int32 lz = -1;
+					for (uint32 b = 0; b == 0; ++lz)
+					{
+						b = BitStream.GetBits(1);
+					}
+					if (lz)
+					{
+						return ((1 << lz) | BitStream.GetBits(lz)) - 1;
+					}
+					return 0;
+				}
+				static int32 se_v(FBitDataStream& BitStream)
+				{
+					uint32 c = ue_v(BitStream);
+					return c & 1 ? int32((c + 1) >> 1) : -int32((c + 1) >> 1);
+				}
+			};
+
+			// SPS is usually an EBSP so we need to strip it down.
+			FScopedDataPtr<uint8> RBSP(FMemory::Malloc(Size));
+			int32 RBSPsize = EBSPtoRBSP(RBSP, static_cast<const uint8*>(Data), Size);
+			FBitDataStream BitReader(RBSP, RBSPsize);
+
+			FMemory::Memzero(OutSPS);
+
+			uint8 nalUnitType = (uint8)BitReader.GetBits(8);
+			if ((nalUnitType & 0x1f) != 0x7)	// SPS NALU?
+			{
+				return false;
+			}
+
+			OutSPS.profile_idc = (uint8)BitReader.GetBits(8);
+			OutSPS.constraint_set0_flag = (uint8)BitReader.GetBits(1);
+			OutSPS.constraint_set1_flag = (uint8)BitReader.GetBits(1);
+			OutSPS.constraint_set2_flag = (uint8)BitReader.GetBits(1);
+			OutSPS.constraint_set3_flag = (uint8)BitReader.GetBits(1);
+			OutSPS.constraint_set4_flag = (uint8)BitReader.GetBits(1);
+			OutSPS.constraint_set5_flag = (uint8)BitReader.GetBits(1);
+			BitReader.SkipBits(2);
+			OutSPS.level_idc = (uint8)BitReader.GetBits(8);
+			OutSPS.seq_parameter_set_id = FSyntaxElement::ue_v(BitReader);
+			if (OutSPS.profile_idc == 100 || OutSPS.profile_idc == 110 || OutSPS.profile_idc == 122 || OutSPS.profile_idc == 244 ||
+				OutSPS.profile_idc == 44 || OutSPS.profile_idc == 83 || OutSPS.profile_idc == 86 || OutSPS.profile_idc == 118 || OutSPS.profile_idc == 128)
+			{
+				OutSPS.chroma_format_idc = FSyntaxElement::ue_v(BitReader);
+				if (OutSPS.chroma_format_idc == 3)
+					OutSPS.separate_colour_plane_flag = (uint8)BitReader.GetBits(1);
+				OutSPS.bit_depth_luma_minus8 = FSyntaxElement::ue_v(BitReader);
+				OutSPS.bit_depth_chroma_minus8 = FSyntaxElement::ue_v(BitReader);
+				OutSPS.qpprime_y_zero_transform_bypass_flag = (uint8)BitReader.GetBits(1);
+				OutSPS.seq_scaling_matrix_present_flag = (uint8)BitReader.GetBits(1);
+				if (OutSPS.seq_scaling_matrix_present_flag)
+				{
+					check(!"oh no!");
+					return false;
+				}
+			}
+			OutSPS.log2_max_frame_num_minus4 = FSyntaxElement::ue_v(BitReader);
+			OutSPS.pic_order_cnt_type = FSyntaxElement::ue_v(BitReader);
+			if (OutSPS.pic_order_cnt_type == 0)
+			{
+				OutSPS.log2_max_pic_order_cnt_lsb_minus4 = FSyntaxElement::ue_v(BitReader);
+			}
+			else if (OutSPS.pic_order_cnt_type == 1)
+			{
+				OutSPS.delta_pic_order_always_zero_flag = FSyntaxElement::ue_v(BitReader);
+				OutSPS.offset_for_non_ref_pic = FSyntaxElement::se_v(BitReader);
+				OutSPS.offset_for_top_to_bottom_field = FSyntaxElement::se_v(BitReader);
+				OutSPS.num_ref_frames_in_pic_order_cnt_cycle = FSyntaxElement::ue_v(BitReader);
+				for (uint32 i = 0; i < OutSPS.num_ref_frames_in_pic_order_cnt_cycle; ++i)
+				{
+					FSyntaxElement::se_v(BitReader);		// discard
+				}
+			}
+			OutSPS.max_num_ref_frames = FSyntaxElement::ue_v(BitReader);
+			OutSPS.gaps_in_frame_num_value_allowed_flag = (uint8)BitReader.GetBits(1);
+			OutSPS.pic_width_in_mbs_minus1 = FSyntaxElement::ue_v(BitReader);
+			OutSPS.pic_height_in_map_units_minus1 = FSyntaxElement::ue_v(BitReader);
+			OutSPS.frame_mbs_only_flag = (uint8)BitReader.GetBits(1);
+			if (!OutSPS.frame_mbs_only_flag)
+			{
+				OutSPS.mb_adaptive_frame_field_flag = (uint8)BitReader.GetBits(1);
+			}
+			OutSPS.direct_8x8_inference_flag = (uint8)BitReader.GetBits(1);
+			OutSPS.frame_cropping_flag = (uint8)BitReader.GetBits(1);
+			if (OutSPS.frame_cropping_flag)
+			{
+				OutSPS.frame_crop_left_offset = FSyntaxElement::ue_v(BitReader);
+				OutSPS.frame_crop_right_offset = FSyntaxElement::ue_v(BitReader);
+				OutSPS.frame_crop_top_offset = FSyntaxElement::ue_v(BitReader);
+				OutSPS.frame_crop_bottom_offset = FSyntaxElement::ue_v(BitReader);
+			}
+			OutSPS.vui_parameters_present_flag = (uint8)BitReader.GetBits(1);
+			if (OutSPS.vui_parameters_present_flag)
+			{
+				OutSPS.aspect_ratio_info_present_flag = (uint8)BitReader.GetBits(1);
+				if (OutSPS.aspect_ratio_info_present_flag)
+				{
+					OutSPS.aspect_ratio_idc = (uint8)BitReader.GetBits(8);
+					if (OutSPS.aspect_ratio_idc == 255)
+					{
+						OutSPS.sar_width = (uint16)BitReader.GetBits(16);
+						OutSPS.sar_height = (uint16)BitReader.GetBits(16);
+					}
+				}
+				OutSPS.overscan_info_present_flag = (uint8)BitReader.GetBits(1);
+				if (OutSPS.overscan_info_present_flag)
+				{
+					OutSPS.overscan_appropriate_flag = (uint8)BitReader.GetBits(1);
+				}
+				OutSPS.video_signal_type_present_flag = (uint8)BitReader.GetBits(1);
+				if (OutSPS.video_signal_type_present_flag)
+				{
+					OutSPS.video_format = (uint8)BitReader.GetBits(3);
+					OutSPS.video_full_range_flag = (uint8)BitReader.GetBits(1);
+					OutSPS.colour_description_present_flag = (uint8)BitReader.GetBits(1);
+					if (OutSPS.colour_description_present_flag)
+					{
+						OutSPS.colour_primaries = (uint8)BitReader.GetBits(8);
+						OutSPS.transfer_characteristics = (uint8)BitReader.GetBits(8);
+						OutSPS.matrix_coefficients = (uint8)BitReader.GetBits(8);
+					}
+				}
+				OutSPS.chroma_loc_info_present_flag = (uint8)BitReader.GetBits(1);
+				if (OutSPS.chroma_loc_info_present_flag)
+				{
+					OutSPS.chroma_sample_loc_type_top_field = FSyntaxElement::ue_v(BitReader);
+					OutSPS.chroma_sample_loc_type_bottom_field = FSyntaxElement::ue_v(BitReader);
+				}
+				OutSPS.timing_info_present_flag = (uint8)BitReader.GetBits(1);
+				if (OutSPS.timing_info_present_flag)
+				{
+					OutSPS.num_units_in_tick = BitReader.GetBits(32);
+					OutSPS.time_scale = BitReader.GetBits(32);
+					OutSPS.fixed_frame_rate_flag = (uint8)BitReader.GetBits(1);
+				}
+				// The remainder is of no interest to us at the moment.
+			}
+			return true;
+		}
+
+
+
+	} // namespace MPEG
+} // namespace Electra
+
diff --git a/Engine/Plugins/Media/ElectraPlayer/Source/ElectraPlayerRuntime/Private/Runtime/Utilities/UtilsMPEGVideo.h b/Engine/Plugins/Media/ElectraPlayer/Source/ElectraPlayerRuntime/Private/Runtime/Utilities/UtilsMPEGVideo.h
new file mode 100644
--- /dev/null
+++ b/Engine/Plugins/Media/ElectraPlayer/Source/ElectraPlayerRuntime/Private/Runtime/Utilities/UtilsMPEGVideo.h
@@ -0,0 +1,210 @@
+// Copyright Epic Games, Inc. All Rights Reserved.
+
+#pragma once
+
+
+#include "PlayerCore.h"
+
+#include "PlayerTime.h"
+#include "BitDataStream.h"
+
+
+
+namespace Electra
+{
+	namespace MPEG
+	{
+
+		/**
+		 * Sequence parameter set as per ISO/IEC 14496-10:2012, section 7.3.2.1.1
+		 */
+		struct FISO14496_10_seq_parameter_set_data
+		{
+			uint8	profile_idc;
+			uint8	constraint_set0_flag;
+			uint8	constraint_set1_flag;
+			uint8	constraint_set2_flag;
+			uint8	constraint_set3_flag;
+			uint8	constraint_set4_flag;
+			uint8	constraint_set5_flag;
+			uint8	level_idc;
+			uint32	seq_parameter_set_id;
+			uint32	chroma_format_idc;
+			uint8	separate_colour_plane_flag;
+			uint32	bit_depth_luma_minus8;
+			uint32	bit_depth_chroma_minus8;
+			uint8	qpprime_y_zero_transform_bypass_flag;
+			uint8	seq_scaling_matrix_present_flag;
+			uint32	log2_max_frame_num_minus4;
+			uint32	pic_order_cnt_type;
+			uint32	log2_max_pic_order_cnt_lsb_minus4;
+			uint32	delta_pic_order_always_zero_flag;
+			int32	offset_for_non_ref_pic;
+			int32	offset_for_top_to_bottom_field;
+			uint32	num_ref_frames_in_pic_order_cnt_cycle;
+			uint32	max_num_ref_frames;
+			uint8	gaps_in_frame_num_value_allowed_flag;
+			uint32	pic_width_in_mbs_minus1;
+			uint32	pic_height_in_map_units_minus1;
+			uint8	frame_mbs_only_flag;
+			uint8	mb_adaptive_frame_field_flag;
+			uint8	direct_8x8_inference_flag;
+			uint8	frame_cropping_flag;
+			uint32	frame_crop_left_offset;
+			uint32	frame_crop_right_offset;
+			uint32	frame_crop_top_offset;
+			uint32	frame_crop_bottom_offset;
+			uint8	vui_parameters_present_flag;
+			uint8	aspect_ratio_info_present_flag;
+			uint8	aspect_ratio_idc;
+			uint16	sar_width;
+			uint16	sar_height;
+			uint8	overscan_info_present_flag;
+			uint8	overscan_appropriate_flag;
+			uint8	video_signal_type_present_flag;
+			uint8	video_format;
+			uint8	video_full_range_flag;
+			uint8	colour_description_present_flag;
+			uint8	colour_primaries;
+			uint8	transfer_characteristics;
+			uint8	matrix_coefficients;
+			uint8	chroma_loc_info_present_flag;
+			uint32	chroma_sample_loc_type_top_field;
+			uint32	chroma_sample_loc_type_bottom_field;
+			uint8	timing_info_present_flag;
+			uint32	num_units_in_tick;
+			uint32	time_scale;
+			uint8	fixed_frame_rate_flag;
+
+			int32 GetWidth() const
+			{
+				return (pic_width_in_mbs_minus1 + 1) * 16;
+			}
+
+			int32 GetHeight() const
+			{
+				return (pic_height_in_map_units_minus1 + 1) * 16;
+			}
+
+			void GetCrop(int32& Left, int32& Right, int32& Top, int32& Bottom) const
+			{
+				if (frame_cropping_flag)
+				{
+					// The scaling factors are determined by the chroma_format_idc (see ISO/IEC 14496-10, table 6.1)
+					// For our purposes this will be 1, so the sub width/height are 2.
+					const int32 CropUnitX = 2;
+					const int32 CropUnitY = 2;
+					Left = (uint16)frame_crop_left_offset * CropUnitX;
+					Right = (uint16)frame_crop_right_offset * CropUnitX;
+					Top = (uint16)frame_crop_top_offset * CropUnitY;
+					Bottom = (uint16)frame_crop_bottom_offset * CropUnitY;
+				}
+				else
+				{
+					Left = Right = Top = Bottom = 0;
+				}
+			}
+
+			void GetAspect(int32& SarW, int32& SarH) const
+			{
+				if (vui_parameters_present_flag && aspect_ratio_info_present_flag)
+				{
+					switch (aspect_ratio_idc)
+					{
+						default:	SarW = SarH = 0;		break;
+						case 0:		SarW = SarH = 0;		break;
+						case 1:		SarW = SarH = 1;		break;
+						case 2:		SarW = 12; SarH = 11;	break;
+						case 3:		SarW = 10; SarH = 11;	break;
+						case 4:		SarW = 16; SarH = 11;	break;
+						case 5:		SarW = 40; SarH = 33;	break;
+						case 6:		SarW = 24; SarH = 11;	break;
+						case 7:		SarW = 20; SarH = 11;	break;
+						case 8:		SarW = 32; SarH = 11;	break;
+						case 9:		SarW = 80; SarH = 33;	break;
+						case 10:	SarW = 18; SarH = 11;	break;
+						case 11:	SarW = 15; SarH = 11;	break;
+						case 12:	SarW = 64; SarH = 33;	break;
+						case 13:	SarW = 160; SarH = 99;	break;
+						case 14:	SarW = 4; SarH = 3;		break;
+						case 15:	SarW = 3; SarH = 2;		break;
+						case 16:	SarW = 2; SarH = 1;		break;
+						case 255:	SarW = sar_width; SarH = sar_height;	break;
+					}
+				}
+				else
+				{
+					SarW = SarH = 1;
+				}
+			}
+			FTimeFraction GetTiming() const
+			{
+				if (vui_parameters_present_flag && timing_info_present_flag)
+				{
+					return FTimeFraction(time_scale, num_units_in_tick * 2);
+				}
+				return FTimeFraction();
+			}
+		};
+
+		/**
+		 * ISO/IEC 14496-15:2014
+		 */
+		class FAVCDecoderConfigurationRecord
+		{
+		public:
+			FAVCDecoderConfigurationRecord();
+			void SetRawData(const void* Data, int64 Size);
+			const TArray<uint8>& GetRawData() const;
+
+			bool Parse();
+
+			const TArray<uint8>& GetCodecSpecificData() const;
+			const TArray<uint8>& GetCodecSpecificDataSPS() const;
+			const TArray<uint8>& GetCodecSpecificDataPPS() const;
+
+			int32 GetNumberOfSPS() const;
+			const FISO14496_10_seq_parameter_set_data& GetParsedSPS(int32 SpsIndex) const;
+
+		private:
+			TArray<uint8>													RawData;
+			TArray<uint8>													CodecSpecificData;
+			TArray<uint8>													CodecSpecificDataSPSOnly;
+			TArray<uint8>													CodecSpecificDataPPSOnly;
+			TArray<TArray<uint8>>											SequenceParameterSets;
+			TArray<TArray<uint8>>											PictureParameterSets;
+			TArray<TArray<uint8>>											SequenceParameterSetsExt;
+			uint8															ConfigurationVersion;
+			uint8															AVCProfileIndication;
+			uint8															ProfileCompatibility;
+			uint8															AVCLevelIndication;
+			uint8															NALUnitLength;
+			uint8															ChromaFormat;
+			uint8															BitDepthLumaMinus8;
+			uint8															BitDepthChromaMinus8;
+			bool															bHaveAdditionalProfileIndication;
+
+			TArray<FISO14496_10_seq_parameter_set_data>						ParsedSPSs;
+		};
+
+
+
+		//! Parses a H.264 (ISO/IEC 14496-10) bitstream for NALUs.
+		struct FNaluInfo
+		{
+			uint64		Offset;
+			uint64		Size;
+			uint8		Type;
+			uint8		UnitLength;
+		};
+		void ParseBitstreamForNALUs(TArray<FNaluInfo>& outNALUs, const void* InBitstream, uint64 InBitstreamLength);
+
+
+		//! Parses a H.264 (ISO/IEC 14496-10) SPS NALU.
+		bool ParseH264SPS(FISO14496_10_seq_parameter_set_data& OutSPS, const void* Data, int32 Size);
+
+
+	} // namespace MPEG
+} // namespace Electra
+
+
diff --git a/Engine/Plugins/Media/ElectraPlayer/Source/ElectraPlayerRuntime/Private/Tests/Parser.spec.cpp b/Engine/Plugins/Media/ElectraPlayer/Source/ElectraPlayerRuntime/Private/Tests/Parser.spec.cpp
new file mode 100644
--- /dev/null
+++ b/Engine/Plugins/Media/ElectraPlayer/Source/ElectraPlayerRuntime/Private/Tests/Parser.spec.cpp
@@ -0,0 +1,309 @@
+// Copyright Epic Games, Inc. All Rights Reserved.
+
+#include "Misc/AutomationTest.h"
+#include "Player/HLS/Parser.h"
+#include "Player/HLS/Tags.h"
+#include "Player/HLS/EpicTags.h"
+#include "Player/HLS/LHLSTags.h"
+#include "CoreMinimal.h"
+
+BEGIN_DEFINE_SPEC(FParserSpec, "ElectraPlayer.Unit", EAutomationTestFlags::ProductFilter | EAutomationTestFlags::ApplicationContextMask)
+END_DEFINE_SPEC(FParserSpec)
+
+void FParserSpec::Define()
+{
+	Describe("FParser", [this]()
+	{
+		It("Should parse variant playlists", [this]()
+		{
+			Electra::HLSPlaylistParser::FParser Parser;
+			Electra::HLSPlaylistParser::FPlaylist Playlist;
+
+			FString PlaylistContent = "#EXTM3U\n#EXT-X-VERSION:6\n"
+				"#EXT-X-TARGETDURATION:2\n"
+				"#EXT-X-DISCONTINUITY\n"
+				"#EXTINF:2.8444,test_title\n"
+				"/test_1.mp4\n"
+				"#EXTINF:2.0000,\n"
+				"/test_2.mp4\n"
+				"#EXT-X-ENDLIST";
+			Electra::HLSPlaylistParser::EPlaylistError Error = Parser.Parse(PlaylistContent, Playlist);
+			TestEqual("ErrorReturn", Error, Electra::HLSPlaylistParser::EPlaylistError::None);
+			TestEqual("SegmentNum", Playlist.GetSegments().Num(), 2);
+
+			TestTrue("FirstSegmentDiscontinuity", Playlist.GetSegments()[0].ContainsTag(Electra::HLSPlaylistParser::ExtXDiscontinuity));
+			TestTrue("FirstSegmentExtinf", Playlist.GetSegments()[0].ContainsTag(Electra::HLSPlaylistParser::ExtINF));
+			TestFalse("SecondSegmentDiscontinuity", Playlist.GetSegments()[1].ContainsTag(Electra::HLSPlaylistParser::ExtXDiscontinuity));
+			TestTrue("SecondSegmentExtinf", Playlist.GetSegments()[1].ContainsTag(Electra::HLSPlaylistParser::ExtINF));
+		});
+
+		It("Should skip custom attribute", [this]()
+		{
+			Electra::HLSPlaylistParser::FParser Parser;
+			Electra::HLSPlaylistParser::FPlaylist Playlist;
+
+			TestEqual("ErrorReturn", Parser.Parse(
+				"#EXTM3U\n#EXT-X-VERSION:6\n"
+				"#EXT-X-TARGETDURATION:20\n"
+				"#EXTINF:20000,test_title\n"
+				"#EXT-X-DATERANGE:ID=\"unique\",START-DATE=\"1234\",X-TEST=1\n"
+				"/test.mp4",
+				Playlist),
+				Electra::HLSPlaylistParser::EPlaylistError::None);
+		});
+
+		It("Should inherit tags if option is set", [this]()
+		{
+			Electra::HLSPlaylistParser::FParser Parser;
+			Electra::HLSPlaylistParser::FPlaylist Playlist;
+			Parser.bInheritTags = true;
+
+			TestEqual("ErrorReturn", Parser.Parse(
+				"#EXTM3U\n"
+				"#EXT-X-VERSION:6\n"
+				"#EXT-X-TARGETDURATION:20\n"
+				"#EXTINF:20000,test_title\n"
+				"#EXT-X-KEY:METHOD=NONE\n"
+				"#EXTINF:20000,test_x\n"
+				"/test_1.mp4\n"
+				"#EXTINF:20000,test_x\n"
+				"/test_2.mp4\n",
+				Playlist),
+				Electra::HLSPlaylistParser::EPlaylistError::None);
+
+			TestEqual("SegmentNumber", Playlist.GetSegments().Num(), 2);
+			TestTrue("FirstSegmentKey", Playlist.GetSegments()[0].ContainsTag(Electra::HLSPlaylistParser::ExtXKey));
+			TestTrue("SecondSegmentKey", Playlist.GetSegments()[1].ContainsTag(Electra::HLSPlaylistParser::ExtXKey));
+		});
+
+		It("Should support quoted spaces in attributes", [this]()
+		{
+			Electra::HLSPlaylistParser::FParser Parser;
+			Electra::HLSPlaylistParser::FPlaylist Playlist;
+
+			TestEqual("ErrorReturn", Parser.Parse(
+				"#EXTM3U\n#EXT-X-VERSION:6\n#EXT-X-MEDIA:TYPE=AUDIO,GROUP-ID=\"aac\",NAME=\"English Test\",DEFAULT=YES,AUTOSELECT=YES,LANGUAGE=\"en\",URI=\"main/english-audio.m3u8\"",
+				Playlist),
+				Electra::HLSPlaylistParser::EPlaylistError::None);
+
+			FString Name;
+			FString ExpectedName = "English Test";
+			TestEqual("GetTagAttribute", Playlist.GetTagAttributeValue(Electra::HLSPlaylistParser::ExtXMedia, "NAME", Name), Electra::HLSPlaylistParser::EPlaylistError::None);
+			TestEqual("Name", Name, ExpectedName);
+		});
+
+		It("Should keep none predefined attributes", [this]()
+		{
+			Electra::HLSPlaylistParser::FParser Parser;
+			Electra::HLSPlaylistParser::FPlaylist Playlist;
+
+			Parser.Configure(Electra::HLSPlaylistParser::Epic::TagMap);
+
+			FString PlaylistContent = "#EXTM3U\n"
+				"#EXT-X-VERSION:3\n"
+				"#EXT-X-ALLOW-CACHE:NO\n"
+				"#EXT-X-EPICGAMES-CUSTOM:FIRST=1,SECOND=test,THRID=\"with spaces\"\n"
+				"#EXT-X-STREAM-INF:BANDWIDTH=1280000,AVERAGE-BANDWIDTH=1000000\n"
+				"http://example.com/low.m3u8\n";
+
+			TestEqual("ErrorReturn", Parser.Parse(PlaylistContent, Playlist), Electra::HLSPlaylistParser::EPlaylistError::None);
+
+			FString TestValue;
+			FString ExpectedValue = "1";
+
+			TestEqual("PlaylistSize", Playlist.GetPlaylists().Num(), 1);
+
+			if (Playlist.GetPlaylists().Num() == 1)
+			{
+				Electra::HLSPlaylistParser::FMediaPlaylist MediaPlaylist = Playlist.GetPlaylists()[0];
+				TestEqual("TagAttr", MediaPlaylist.GetTagAttributeValue(Electra::HLSPlaylistParser::Epic::ExtXEpicGamesCustom, "FIRST", TestValue), Electra::HLSPlaylistParser::EPlaylistError::None);
+				TestEqual("TestValue", TestValue, ExpectedValue);
+			}
+		});
+
+		It("Should support multiple tags of the same type", [this]()
+		{
+			Electra::HLSPlaylistParser::FParser Parser;
+			Electra::HLSPlaylistParser::FPlaylist Playlist;
+
+			FString PlaylistContent = "#EXTM3U\n"
+				"#EXT-X-VERSION:3\n"
+				"#EXT-X-ALLOW-CACHE:NO\n"
+				"#EXT-X-MEDIA:TYPE=AUDIO,GROUP-ID=\"aac\",NAME=\"English Test\",DEFAULT=YES,AUTOSELECT=YES,LANGUAGE=\"en\",URI=\"main/english-audio.m3u8\"\n"
+				"#EXT-X-MEDIA:TYPE=AUDIO,GROUP-ID=\"aac\",NAME=\"Spain Test\",DEFAULT=YES,AUTOSELECT=NO,LANGUAGE=\"es\",URI=\"main/spain-audio.m3u8\"\n"
+				"#EXT-X-STREAM-INF:BANDWIDTH=1280000,AVERAGE-BANDWIDTH=1000000\n"
+				"http://example.com/low.m3u8\n";
+
+			TestEqual("ErrorReturn", Parser.Parse(PlaylistContent, Playlist), Electra::HLSPlaylistParser::EPlaylistError::None);
+			TestEqual("TagNum", Playlist.TagNum(Electra::HLSPlaylistParser::ExtXMedia), 2);
+
+			FString OutValue;
+			FString ExpectValue = "en";
+			TestEqual("ErrorFirstTag", Playlist.GetTagAttributeValue(Electra::HLSPlaylistParser::ExtXMedia, "LANGUAGE", OutValue), Electra::HLSPlaylistParser::EPlaylistError::None);
+			TestEqual("FirstTagLanguage", OutValue, ExpectValue);
+
+			ExpectValue = "es";
+			TestEqual("ErrorSecondTag", Playlist.GetTagAttributeValue(Electra::HLSPlaylistParser::ExtXMedia, "LANGUAGE", OutValue, 1), Electra::HLSPlaylistParser::EPlaylistError::None);
+			TestEqual("SecondTagLanguage", OutValue, ExpectValue);
+		});
+
+		It("Should return a proper error message", [this]()
+		{
+			Electra::HLSPlaylistParser::FParser Parser;
+			Electra::HLSPlaylistParser::FPlaylist Playlist;
+
+			TestEqual("ErrorReturn", Parser.Parse("#EXTM3U\n#EXT-X-VERSION:6\n#EXTINF:20000,test_title\n#EXT-X-KEY:URI=test\n/test.mp4",
+				Playlist),
+				Electra::HLSPlaylistParser::EPlaylistError::MissingRequiredAttribute);
+
+			FString ExpectedVal = "invalid token in line: 4 character: 19";
+			TestEqual("ErrorMessage", Parser.GetLastErrorMessage(), ExpectedVal);
+		});
+
+		It("Should not modify the configure target", [this]()
+		{
+			Electra::HLSPlaylistParser::FParser Parser;
+			Electra::HLSPlaylistParser::FPlaylist Playlist;
+
+			TMap<FString, Electra::HLSPlaylistParser::FMediaTag> NewTags = {
+				{"EXT-X-TEST", Electra::HLSPlaylistParser::FMediaTag("EXT-X-TEST", Electra::HLSPlaylistParser::TargetPlaylistURL)},
+			};
+
+			TestEqual("OriginalSize", NewTags.Num(), 1);
+			Parser.Configure(NewTags);
+			TestEqual("AfterConfigureSize", NewTags.Num(), 1);
+		});
+
+		It("Should parse prefetch tag as segment", [this]()
+		{
+			Electra::HLSPlaylistParser::FParser Parser;
+			Electra::HLSPlaylistParser::FPlaylist Playlist;
+			Parser.Configure(Electra::HLSPlaylistParser::LHLS::TagMap);
+
+			TestEqual("ErrorReturn", Parser.Parse(
+				"#EXTM3U\n"
+				"#EXT-X-VERSION:6\n"
+				"#EXT-X-TARGETDURATION:2.00000\n"
+				"#EXTINF:2.00000,test_title\n"
+				"/test_1.mp4\n"
+				"#EXT-X-PREFETCH-DISCONTINUITY\n"
+				"#EXT-X-PREFETCH:./test_2.mp4\n",
+				Playlist),
+				Electra::HLSPlaylistParser::EPlaylistError::None);
+
+			TestEqual("SegmentNumber", Playlist.GetSegments().Num(), 2);
+
+			if (Playlist.GetSegments().Num() == 2)
+			{
+				// First segment doesn't have a prefetch tag.
+				Electra::HLSPlaylistParser::FMediaSegment FirstSegment = Playlist.GetSegments()[0];
+				TestFalse("FirstNoPrefetch", FirstSegment.ContainsTag(Electra::HLSPlaylistParser::LHLS::ExtXPrefetch));
+
+				// Second segment has the prefetch tag.
+				Electra::HLSPlaylistParser::FMediaSegment Segment = Playlist.GetSegments()[1];
+				TestTrue("HasPrefetch", Segment.ContainsTag(Electra::HLSPlaylistParser::LHLS::ExtXPrefetch));
+				TestTrue("HasDiscontinuity", Segment.ContainsTag(Electra::HLSPlaylistParser::LHLS::ExtXPrefetchDiscontinuity));
+
+				FString ExpectUrl = "./test_2.mp4";
+				TestEqual("PrefetchUrl", Segment.URL, ExpectUrl);
+			}
+		});
+
+		It("Should support a prefetch on an empty playlist", [this]()
+		{
+			Electra::HLSPlaylistParser::FParser Parser;
+			Electra::HLSPlaylistParser::FPlaylist Playlist;
+			Parser.Configure(Electra::HLSPlaylistParser::LHLS::TagMap);
+
+			TestEqual("ErrorReturn", Parser.Parse(
+				"#EXTM3U\n"
+				"#EXT-X-TARGETDURATION:2.000000\n"
+				"#EXT-X-PREFETCH-DISCONTINUITY\n"
+				"#EXT-X-PREFETCH:./test_1.mp4\n",
+				Playlist),
+				Electra::HLSPlaylistParser::EPlaylistError::None);
+			TestEqual("SegmentNumber", Playlist.GetSegments().Num(), 1);
+
+			if (Playlist.GetSegments().Num() == 1)
+			{
+				Electra::HLSPlaylistParser::FMediaSegment Segment = Playlist.GetSegments()[0];
+				TestTrue("HasPrefetch", Segment.ContainsTag(Electra::HLSPlaylistParser::LHLS::ExtXPrefetch));
+				TestTrue("HasDiscontinuity", Segment.ContainsTag(Electra::HLSPlaylistParser::LHLS::ExtXPrefetchDiscontinuity));
+			}
+		});
+
+		Describe("FTagContainer", [this]()
+		{
+
+			TArray<Electra::HLSPlaylistParser::FTagContent> FirstItemArr;
+			FirstItemArr.Init({ "testdsd" }, 1);
+
+			TArray<Electra::HLSPlaylistParser::FTagContent> SecondItemArr;
+			SecondItemArr.Init({ "", {
+					{"ATTR", "strval"},
+					{"INT", "1"},
+					{"DOUBLE", "2.04"},
+					{"RESOLUTION", "1024x768"},
+					{"HEX", "0xFF"},
+					{"DURATION", "90000,test_title"},
+					{"INVALID_INT", "A"},
+					{"INVALID_DOUBLE", "D"},
+			} }, 1);
+
+			TArray<Electra::HLSPlaylistParser::FTagContent> ThirdItemArr;
+			TArray<Electra::HLSPlaylistParser::FTagContent> ForthItemArr;
+			ThirdItemArr.Init({ "test" }, 1);
+			ForthItemArr.Add({ "1" });
+			ForthItemArr.Add({ "2" });
+
+			TMap<FString, TArray<Electra::HLSPlaylistParser::FTagContent>> mergedItems = {
+					{"EXT-X-BLUB", FirstItemArr},
+			};
+
+			TMap<FString, TArray<Electra::HLSPlaylistParser::FTagContent>> items = {
+					{"EXT-X-TEST", ThirdItemArr},
+					{"EXT-X-ATTR", SecondItemArr},
+					{"EXT-X-SEC", ForthItemArr},
+			};
+			items.Append(mergedItems);
+
+			Electra::HLSPlaylistParser::TTagContainer Container(items);
+
+			It("Should find tags based on the name", [this, Container]()
+			{
+				TestTrue("ContainsTag(EXT-X-TEST)", Container.ContainsTag("EXT-X-TEST"));
+			});
+
+			It("Should return Tag string values", [this, Container]()
+			{
+				FString TestValue;
+				FString ExpectedValue = "test";
+
+				TestEqual("ErrorReturn", Container.GetTagValue("EXT-X-TEST", TestValue), Electra::HLSPlaylistParser::EPlaylistError::None);
+				TestEqual("Value", TestValue, ExpectedValue);
+			});
+
+			It("Should return Tag Attribute values", [this, Container]()
+			{
+				FString TestValue;
+				FString ExpectedValue = "strval";
+
+				TestEqual("ErrorReturn", Container.GetTagAttributeValue("EXT-X-ATTR", "ATTR", TestValue), Electra::HLSPlaylistParser::EPlaylistError::None);
+				TestEqual("Value", TestValue, ExpectedValue);
+			});
+
+			It("Should return multiple Tag values", [this, Container]()
+			{
+				FString TestValue;
+				FString ExpectedValue = "1";
+
+				TestEqual("ErrorReturnFirst", Container.GetTagValue("EXT-X-SEC", TestValue), Electra::HLSPlaylistParser::EPlaylistError::None);
+				TestEqual("FirstValue", TestValue, ExpectedValue);
+
+				ExpectedValue = "2";
+				TestEqual("ErrorReturnSecond", Container.GetTagValue("EXT-X-SEC", TestValue, 1), Electra::HLSPlaylistParser::EPlaylistError::None);
+				TestEqual("SecondValue", TestValue, ExpectedValue);
+			});
+		});
+	});
+}
diff --git a/Engine/Plugins/Media/ElectraPlayer/Source/ElectraPlayerRuntime/Private/Windows/ElectraPlayerPrivate_Platform.h b/Engine/Plugins/Media/ElectraPlayer/Source/ElectraPlayerRuntime/Private/Windows/ElectraPlayerPrivate_Platform.h
new file mode 100644
--- /dev/null
+++ b/Engine/Plugins/Media/ElectraPlayer/Source/ElectraPlayerRuntime/Private/Windows/ElectraPlayerPrivate_Platform.h
@@ -0,0 +1,109 @@
+// Copyright Epic Games, Inc. All Rights Reserved.
+#pragma once
+
+#include "Windows/AllowWindowsPlatformTypes.h"
+
+#pragma warning(push)
+#pragma warning(disable : 4005)	// macro redefinition
+ #	include <d3d11.h>
+#	include <d3d10.h> // for ID3D10Multithread
+#pragma warning(pop)
+
+#include "Windows/WindowsHWrapper.h"
+
+inline const FString GetComErrorDescription(HRESULT Res)
+{
+	const uint32 BufSize = 4096;
+	WIDECHAR buffer[4096];
+	if (::FormatMessageW(FORMAT_MESSAGE_FROM_SYSTEM,
+		nullptr,
+		Res,
+		MAKELANGID(LANG_NEUTRAL, SUBLANG_NEUTRAL),
+		buffer,
+		sizeof(buffer) / sizeof(*buffer),
+		nullptr))
+	{
+		return buffer;
+	}
+	else
+	{
+		return TEXT("[cannot find error description]");
+	}
+}
+
+#define VERIFY_HR(FNcall, Msg, What)	\
+res = FNcall;							\
+if (FAILED(res))						\
+{										\
+	PostError(res, Msg, What);			\
+	return false;						\
+}
+
+#include "Windows/HideWindowsPlatformTypes.h"
+
+// a convenience macro to deal with COM result codes
+// is designed to be used in functions returning `bool`: CHECK_HR(COM_function_call());
+#define CHECK_HR(COM_call)\
+	{\
+		HRESULT Res = COM_call;\
+		if (FAILED(Res))\
+		{\
+			UE_LOG(LogElectraPlayer, Error, TEXT("`" #COM_call "` failed: 0x%X - %s"), Res, *GetComErrorDescription(Res));\
+			return false;\
+		}\
+	}
+
+THIRD_PARTY_INCLUDES_START
+#include "mfobjects.h"
+#include "mfapi.h"
+THIRD_PARTY_INCLUDES_END
+
+#include <d3d9.h>
+#include <dxva2api.h>
+#include <DxErr.h>
+
+#define CHECK_HR_DX9(DX9_call)\
+	{\
+		HRESULT Res = DX9_call;\
+		if (FAILED(Res))\
+		{\
+			UE_LOG(LogElectraPlayer, Error, TEXT("`" #DX9_call "` failed 0x%X: %s - %s"), Res, DXGetErrorString(Res), DXGetErrorDescription(Res));\
+			return false;\
+		}\
+	}
+
+
+
+namespace Electra
+{
+
+	struct FDXDeviceInfo
+	{
+		enum class ED3DVersion
+		{
+			VersionUnknown,
+			Version9Win7,
+			Version11Win8,
+			Version11XB1,
+			Version12Win10
+		};
+
+		ED3DVersion								DxVersion;
+		TRefCountPtr<ID3D11Device>				DxDevice;
+		TRefCountPtr<ID3D11DeviceContext>		DxDeviceContext;
+		TRefCountPtr<IMFDXGIDeviceManager>		DxDeviceManager;
+
+		TRefCountPtr<IDirect3D9>				Dx9;
+		TRefCountPtr<IDirect3DDevice9>			Dx9Device;
+		TRefCountPtr<IDirect3DDeviceManager9>	Dx9DeviceManager;
+		TRefCountPtr<ID3D11Device>				RenderingDx11Device;
+	
+		static TUniquePtr<FDXDeviceInfo>		s_DXDeviceInfo;
+	};
+
+
+	bool IsWindows8Plus();
+	bool IsWindows7Plus();
+
+} // namespace Electra
+
diff --git a/Engine/Plugins/Media/ElectraPlayer/Source/ElectraPlayerRuntime/Public/IElectraPlayerInterface.h b/Engine/Plugins/Media/ElectraPlayer/Source/ElectraPlayerRuntime/Public/IElectraPlayerInterface.h
new file mode 100644
--- /dev/null
+++ b/Engine/Plugins/Media/ElectraPlayer/Source/ElectraPlayerRuntime/Public/IElectraPlayerInterface.h
@@ -0,0 +1,215 @@
+// Copyright Epic Games, Inc. All Rights Reserved.
+#pragma once
+
+#include "MediaDecoderOutput.h"
+#include "IAnalyticsProviderET.h"
+#include "ParameterDictionary.h"
+
+#include "ElectraPlayerMisc.h"
+#include "ElectraPlayerResourceDelegate.h"
+
+class FVideoDecoderOutput;
+using FVideoDecoderOutputPtr = TSharedPtr<FVideoDecoderOutput, ESPMode::ThreadSafe>;
+class IAudioDecoderOutput;
+using IAudioDecoderOutputPtr = TSharedPtr<IAudioDecoderOutput, ESPMode::ThreadSafe>;
+
+// ---------------------------------------------------------------------------------------------
+
+DECLARE_MULTICAST_DELEGATE_TwoParams(FElectraPlayerSendAnalyticMetricsDelegate, const TSharedPtr<IAnalyticsProviderET>& /*AnalyticsProvider*/, const FGuid& /*PlayerGuid*/);
+DECLARE_MULTICAST_DELEGATE_OneParam(FElectraPlayerSendAnalyticMetricsPerMinuteDelegate, const TSharedPtr<IAnalyticsProviderET>& /*AnalyticsProvider*/);
+DECLARE_MULTICAST_DELEGATE_TwoParams(FElectraPlayerReportVideoStreamingErrorDelegate, const FGuid& /*PlayerGuid*/, const FString& /*LastError*/);
+DECLARE_MULTICAST_DELEGATE_FourParams(FElectraPlayerReportSubtitlesMetricsDelegate, const FGuid& /*PlayerGuid*/, const FString& /*URL*/, double /*ResponseTime*/, const FString& /*LastError*/);
+
+// ---------------------------------------------------------------------------------------------
+
+class IElectraPlayerAdapterDelegate : public TSharedFromThis<IElectraPlayerAdapterDelegate, ESPMode::ThreadSafe>
+{
+public:
+	virtual ~IElectraPlayerAdapterDelegate() {}
+
+	enum class EOptionType
+	{
+		MaxVerticalStreamResolution = 0,
+		MaxBandwidthForStreaming,
+		PlayListData,
+		LicenseKeyData,
+		PlaystartPosFromSeekPositions,
+		Android_VideoSurface,
+	};
+	virtual Electra::FVariantValue QueryOptions(EOptionType Type, const Electra::FVariantValue& Param = Electra::FVariantValue()) = 0;
+
+	enum class EPlayerEvent
+	{
+		MediaBuffering = 0,
+		MediaClosed,
+		MediaConnecting,
+		MediaOpened,
+		MediaOpenFailed,
+		PlaybackEndReached,
+		PlaybackResumed,
+		PlaybackSuspended,
+		SeekCompleted,
+		TracksChanged,
+
+		Internal_Start,
+		Internal_PurgeVideoSamplesHint = Internal_Start,
+		Internal_ResetForDiscontinuity,
+		Internal_RenderClockStart,
+		Internal_RenderClockStop,
+		Internal_VideoSamplesAvailable,
+		Internal_VideoSamplesUnavailable,
+		Internal_AudioSamplesAvailable,
+		Internal_AudioSamplesUnavailable
+	};
+	virtual void SendMediaEvent(EPlayerEvent Event) = 0;
+	virtual void OnVideoFlush() = 0;
+	virtual void OnAudioFlush() = 0;
+	virtual void PresentVideoFrame(const FVideoDecoderOutputPtr& InVideoFrame) = 0;
+	virtual void PresentAudioFrame(const IAudioDecoderOutputPtr& InAudioFrame) = 0;
+	virtual bool CanReceiveVideoSamples(int32 NumFrames) = 0;
+	virtual bool CanReceiveAudioSamples(int32 NumFrames) = 0;
+	virtual void PrepareForDecoderShutdown() = 0;
+
+	virtual FString GetVideoAdapterName() const = 0;
+
+	virtual IElectraPlayerResourceDelegate* GetResourceDelegate() const = 0;
+};
+
+// Container class to be passed through options as shared pointer to allow passing any non-standard ref-counted entities to the player
+class IOptionPointerValueContainer
+{
+public:
+	virtual ~IOptionPointerValueContainer() {}
+	virtual void* GetPointer() const = 0;
+};
+
+class IElectraPlayerInterface : public TSharedFromThis<IElectraPlayerInterface, ESPMode::ThreadSafe>
+{
+public:
+	virtual ~IElectraPlayerInterface() {}
+
+	virtual FString GetUrl() const = 0;
+
+	virtual void SetGuid(const FGuid& Guid) = 0;
+
+	// -------- PlayerAdapter (Plugin/Native) API
+
+	struct FPlaystartOptions
+	{
+		TOptional<FTimespan>		TimeOffset;
+		TOptional<int32>			AudioTrackIndex;
+		TOptional<int32>			MaxVerticalStreamResolution;
+		TOptional<int32>			MaxBandwidthForStreaming;
+	};
+
+	virtual bool OpenInternal(const FString& Url, const Electra::FParamDict& PlayerOptions, const FPlaystartOptions& InPlaystartOptions) = 0;
+	virtual void CloseInternal(bool bKillAfterClose) = 0;
+
+	virtual void Tick(FTimespan DeltaTime, FTimespan Timecode) = 0;
+
+	virtual bool IsKillAfterCloseAllowed() const = 0;
+
+	enum class EPlayerState
+	{
+		Closed = 0,
+		Error,
+		Paused,
+		Playing,
+		Preparing,
+		Stopped
+	};
+
+	enum class EPlayerStatus
+	{
+		None = 0x0,
+		Buffering = 0x1,
+		Connecting = 0x2
+	};
+
+	virtual EPlayerState GetState() const = 0;
+	virtual EPlayerStatus GetStatus() const = 0;
+
+	virtual bool IsLooping() const = 0;
+	virtual bool SetLooping(bool bLooping) = 0;
+
+	virtual FTimespan GetTime() const = 0;
+	virtual FTimespan GetDuration() const = 0;
+
+	virtual bool IsLive() const = 0;
+	virtual FTimespan GetSeekableDuration() const = 0;
+
+	virtual float GetRate() const = 0;
+	virtual bool SetRate(float Rate) = 0;
+
+	virtual bool Seek(const FTimespan& Time) = 0;
+
+	struct FAudioTrackFormat
+	{
+		uint32 BitsPerSample;
+		uint32 NumChannels;
+		uint32 SampleRate;
+		FString TypeName;
+	};
+	struct FVideoTrackFormat
+	{
+		FIntPoint Dim;
+		float FrameRate;
+		TRange<float> FrameRates;
+		FString TypeName;
+	};
+	virtual bool GetAudioTrackFormat(int32 TrackIndex, int32 FormatIndex, FAudioTrackFormat& OutFormat) const = 0;
+	virtual bool GetVideoTrackFormat(int32 TrackIndex, int32 FormatIndex, FVideoTrackFormat& OutFormat) const = 0;
+
+	enum class EPlayerTrackType
+	{
+		Audio = 0,
+		Caption,
+		Metadata,
+		Script,
+		Subtitle,
+		Text,
+		Video
+	};
+	virtual int32 GetNumTracks(EPlayerTrackType TrackType) const = 0;
+	virtual int32 GetNumTrackFormats(EPlayerTrackType TrackType, int32 TrackIndex) const = 0;
+	virtual int32 GetSelectedTrack(EPlayerTrackType TrackType) const = 0;
+	virtual FText GetTrackDisplayName(EPlayerTrackType TrackType, int32 TrackIndex) const = 0;
+	virtual int32 GetTrackFormat(EPlayerTrackType TrackType, int32 TrackIndex) const = 0;
+	virtual FString GetTrackLanguage(EPlayerTrackType TrackType, int32 TrackIndex) const = 0;
+	virtual FString GetTrackName(EPlayerTrackType TrackType, int32 TrackIndex) const = 0;
+	virtual bool SelectTrack(EPlayerTrackType TrackType, int32 TrackIndex) = 0;
+
+	virtual void NotifyOfOptionChange() = 0;
+
+	enum {
+		ResourceFlags_Decoder = 1 << 0,
+		ResourceFlags_OutputBuffers = 1 << 1,
+
+		ResourceFlags_All = (1 << 2) - 1,
+		ResourceFlags_Any = ResourceFlags_All,
+	};
+
+	class IAsyncResourceReleaseNotifyContainer
+	{
+	public:
+		virtual ~IAsyncResourceReleaseNotifyContainer() {}
+		virtual void Signal(uint32 ResourceFlags) = 0;
+	};
+	virtual void SetAsyncResourceReleaseNotification(IAsyncResourceReleaseNotifyContainer* AsyncResourceReleaseNotification) = 0;
+};
+
+class ELECTRAPLAYERRUNTIME_API FElectraPlayerRuntimeFactory
+{
+public:
+	static IElectraPlayerInterface* CreatePlayer(const TSharedPtr<IElectraPlayerAdapterDelegate, ESPMode::ThreadSafe>& AdapterDelegate,
+												FElectraPlayerSendAnalyticMetricsDelegate& InSendAnalyticMetricsDelegate,
+												FElectraPlayerSendAnalyticMetricsPerMinuteDelegate& InSendAnalyticMetricsPerMinuteDelegate,
+												FElectraPlayerReportVideoStreamingErrorDelegate& InReportVideoStreamingErrorDelegate,
+												FElectraPlayerReportSubtitlesMetricsDelegate& InReportSubtitlesFileMetricsDelegate);
+};
+
+class ELECTRAPLAYERRUNTIME_API FElectraPlayerPlatform
+{
+public:
+	static bool StartupPlatformResources(const Electra::FParamDict & Params = Electra::FParamDict());
+};
diff --git a/Engine/Plugins/Media/ElectraUtil/ElectraUtil.uplugin b/Engine/Plugins/Media/ElectraUtil/ElectraUtil.uplugin
new file mode 100644
--- /dev/null
+++ b/Engine/Plugins/Media/ElectraUtil/ElectraUtil.uplugin
@@ -0,0 +1,33 @@
+{
+	"FileVersion" : 1,
+	"Version" : 0,
+	"VersionName" : "0.1",
+	"FriendlyName" : "Electra Player Utilities",
+	"Description" : "Reusable Base Components for Electra Player Media Playback",
+	"Category" : "Media Players",
+	"CreatedBy" : "Epic Games, Inc.",
+	"CreatedByURL" : "http://epicgames.com",
+	"DocsURL" : "https://docs.unrealengine.com/en-US/Engine/MediaFramework/Overview",
+	"MarketplaceURL" : "",
+	"SupportURL" : "",
+	"EnabledByDefault" : false,
+	"CanContainContent" : false,
+	"IsBetaVersion" : true,
+	"Installed" : false,
+	"Modules": [
+		{
+			"Name": "ElectraBase",
+			"Type": "RuntimeNoCommandlet",
+			"LoadingPhase": "PostEngineInit",
+			"WhitelistPlatforms": [ "Win64", "Mac", "IOS", "Android" ],
+			"BlacklistTargets": [ "Server" ]
+		},
+		{
+			"Name": "ElectraSamples",
+			"Type": "RuntimeNoCommandlet",
+			"LoadingPhase": "PostEngineInit",
+			"WhitelistPlatforms": [ "Win64", "Mac", "IOS", "Android" ],
+			"BlacklistTargets": [ "Server" ]
+		}
+	]
+}
\ No newline at end of file
diff --git a/Engine/Plugins/Media/ElectraUtil/Source/ElectraBase/Public/MediaAudioDecoderOutput.h b/Engine/Plugins/Media/ElectraUtil/Source/ElectraBase/Public/MediaAudioDecoderOutput.h
new file mode 100644
--- /dev/null
+++ b/Engine/Plugins/Media/ElectraUtil/Source/ElectraBase/Public/MediaAudioDecoderOutput.h
@@ -0,0 +1,39 @@
+// Copyright Epic Games, Inc. All Rights Reserved.
+#pragma once
+
+#include "MediaDecoderOutput.h"
+
+class IAudioDecoderOutput : public IDecoderOutput
+{
+public:
+	enum class ESampleFormat
+	{
+		Undefined = 0,
+		Double,
+		Float,
+		Int8,
+		Int16,
+		Int32
+	};
+
+	virtual ~IAudioDecoderOutput() {}
+
+	virtual void Reserve(uint32 InBufferSizeBytes) = 0;
+
+	virtual void Initialize(ESampleFormat InFormat, uint32 InNumChannels, uint32 InSampleRate, FTimespan InDuration, const FDecoderTimeStamp& InPts, uint32 InBufferSizeBytes) = 0;
+
+	virtual const void* GetBuffer() const = 0;
+	virtual uint32 GetUsedBufferBytes() const = 0;
+	virtual uint32 GetReservedBufferBytes() const = 0;
+
+	virtual uint32 GetChannels() const = 0;
+	virtual uint32 GetFrames() const = 0;
+	virtual ESampleFormat GetFormat() const = 0;
+
+	virtual FDecoderTimeStamp GetTime() const = 0;
+	virtual FTimespan GetDuration() const = 0;
+
+	virtual uint32 GetSampleRate() const = 0;
+};
+
+using IAudioDecoderOutputPtr = TSharedPtr<IAudioDecoderOutput, ESPMode::ThreadSafe>;
diff --git a/Engine/Plugins/Media/ElectraUtil/Source/ElectraBase/Public/MediaDecoderOutput.h b/Engine/Plugins/Media/ElectraUtil/Source/ElectraBase/Public/MediaDecoderOutput.h
new file mode 100644
--- /dev/null
+++ b/Engine/Plugins/Media/ElectraUtil/Source/ElectraBase/Public/MediaDecoderOutput.h
@@ -0,0 +1,291 @@
+// Copyright Epic Games, Inc. All Rights Reserved.
+#pragma once
+
+#include "Misc/Timespan.h"
+#include "Math/IntPoint.h"
+#include "Templates/SharedPointer.h"
+#include "Templates/RefCounting.h"
+#include "Containers/Queue.h"
+#include "Misc/ScopeLock.h"
+
+struct FDecoderTimeStamp
+{
+	FDecoderTimeStamp() {}
+	FDecoderTimeStamp(FTimespan InTime, int64 InSequenceIndex) : Time(InTime), SequenceIndex(InSequenceIndex) {}
+
+	FTimespan Time;
+	int64 SequenceIndex;
+};
+
+
+class IDecoderOutputPoolable : public TSharedFromThis<IDecoderOutputPoolable, ESPMode::ThreadSafe>
+{
+public:
+	virtual void InitializePoolable() { }
+	virtual void ShutdownPoolable() { }
+	virtual bool IsReadyForReuse() { return true; }
+
+public:
+	virtual ~IDecoderOutputPoolable() { }
+};
+
+
+template<typename ObjectType> class TElectraPoolDefaultObjectFactory
+{
+public:
+	static ObjectType* Create() { return new ObjectType; }
+};
+
+template<typename ObjectType, typename ObjectFactory = TElectraPoolDefaultObjectFactory<ObjectType>>
+class TDecoderOutputObjectPool
+{
+	static_assert(TPointerIsConvertibleFromTo<ObjectType, IDecoderOutputPoolable>::Value, "Poolable objects must implement the IDecoderOutputPoolable interface.");
+
+	/** Object pool storage. */
+	class TStorage
+	{
+	public:
+		TStorage(ObjectFactory* InObjectFactoryInstance)
+			: ObjectFactoryInstance(InObjectFactoryInstance)
+		{}
+
+		/** Destructor. */
+		~TStorage()
+		{
+			Reserve(0);
+			ObjectType* Object;
+			while (WaitReadyForReuse.Dequeue(Object))
+			{
+				Object->ShutdownPoolable();
+				delete Object;
+			}
+		}
+
+	public:
+
+		/** Acquire an object from the pool. */
+		ObjectType* Acquire()
+		{
+			ObjectType* Result = nullptr;
+			{
+				FScopeLock Lock(&CriticalSection);
+
+				if (Pool.Num() > 0)
+				{
+					Result = Pool.Pop(false);
+				}
+				else
+				{
+					if (WaitReadyForReuse.Peek(Result))
+					{
+						if (Result->IsReadyForReuse())
+						{
+							WaitReadyForReuse.Pop();
+							Result->ShutdownPoolable();
+						}
+						else
+						{
+							Result = nullptr;
+						}
+					}
+				}
+			}
+
+			if (Result == nullptr)
+			{
+				Result = ObjectFactoryInstance->Create();
+			}
+
+			Result->InitializePoolable();
+
+			return Result;
+		}
+
+		/** Get the number of objects stored. */
+		int32 Num() const
+		{
+			FScopeLock Lock(&CriticalSection);
+			return Pool.Num();
+		}
+
+		/** Return the given object to the pool. */
+		void Release(ObjectType* Object)
+		{
+			if (Object == nullptr)
+			{
+				return;
+			}
+
+			FScopeLock Lock(&CriticalSection);
+
+			if (Object->IsReadyForReuse())
+			{
+				Object->ShutdownPoolable();
+				Pool.Push(Object);
+			}
+			else
+			{
+				WaitReadyForReuse.Enqueue(Object);
+			}
+		}
+
+		/** Reserve the specified number of objects. */
+		void Reserve(uint32 NumObjects)
+		{
+			FScopeLock Lock(&CriticalSection);
+
+			while (NumObjects < (uint32)Pool.Num())
+			{
+				delete Pool.Pop(false);
+			}
+
+			while (NumObjects > (uint32)Pool.Num())
+			{
+				Pool.Push(ObjectFactoryInstance->Create());
+			}
+		}
+
+		/** Regular tick call */
+		void Tick()
+		{
+			if (WaitReadyForReuse.IsEmpty())
+			{
+				// Conservative early out to avoid CS: we will get any items missed the next time around
+				return;
+			}
+
+			FScopeLock Lock(&CriticalSection);
+
+			ObjectType* Object;
+			while (WaitReadyForReuse.Peek(Object))
+			{
+				if (!Object->IsReadyForReuse())
+				{
+					break;
+				}
+
+				Object->ShutdownPoolable();
+				Pool.Push(Object);
+				WaitReadyForReuse.Pop();
+			}
+		}
+
+		void PrepareForDecoderShutdown()
+		{
+			// [...]
+		}
+
+	private:
+
+		/** Critical section for synchronizing access to the free list. */
+		mutable FCriticalSection CriticalSection;
+
+		/** List of unused objects. */
+		TArray<ObjectType*> Pool;
+
+		/** List of unused objects, waiting for reuse-ability. */
+		TQueue<ObjectType*> WaitReadyForReuse;
+
+		/** Object factory instance (nullptr by default) */
+		ObjectFactory* ObjectFactoryInstance;
+	};
+
+
+	/** Deleter for pooled objects. */
+	class TDeleter
+	{
+	public:
+
+		/** Create and initialize a new instance. */
+		TDeleter(const TSharedRef<TStorage, ESPMode::ThreadSafe>& InStorage)
+			: StoragePtr(InStorage)
+		{ }
+
+		/** Function operator to execute deleter. */
+		void operator()(ObjectType* ObjectToDelete)
+		{
+			TSharedPtr<TStorage, ESPMode::ThreadSafe> PinnedStorage = StoragePtr.Pin();
+
+			if (PinnedStorage.IsValid())
+			{
+				PinnedStorage->Release(ObjectToDelete);
+			}
+			else
+			{
+				delete ObjectToDelete;
+			}
+		}
+
+	private:
+
+		/** Weak pointer to object pool storage. */
+		TWeakPtr<TStorage, ESPMode::ThreadSafe> StoragePtr;
+	};
+
+public:
+
+	TDecoderOutputObjectPool(ObjectFactory* ObjectFactoryInstance = nullptr)
+		: Storage(MakeShareable(new TStorage(ObjectFactoryInstance)))
+	{ }
+
+	TDecoderOutputObjectPool(uint32 NumReserve, ObjectFactory* ObjectFactoryInstance = nullptr)
+		: Storage(MakeShareable(new TStorage(ObjectFactoryInstance)))
+	{
+		Storage->Reserve(NumReserve);
+	}
+
+public:
+
+	TSharedRef<ObjectType, ESPMode::ThreadSafe> AcquireShared()
+	{
+		ObjectType* Object = Storage->Acquire();
+		check(Object != nullptr);
+
+		return MakeShareable(Object, TDeleter(Storage));
+	}
+
+	int32 Num() const
+	{
+		return Storage->Num();
+	}
+
+	void Release(ObjectType* Object)
+	{
+		Storage->Release(Object);
+	}
+
+	void Reset(uint32 NumObjects = 0)
+	{
+		Storage->Reserve(NumObjects);
+	}
+
+	void Tick()
+	{
+		Storage->Tick();
+	}
+
+	void PrepareForDecoderShutdown()
+	{
+		Storage->PrepareForDecoderShutdown();
+	}
+
+private:
+
+	/** Storage for pooled objects. */
+	TSharedRef<TStorage, ESPMode::ThreadSafe> Storage;
+};
+
+
+class IDecoderOutputOwner
+{
+public:
+	virtual void SampleReleasedToPool(FTimespan durationToRelease) = 0;
+};
+
+class IDecoderOutput : public IDecoderOutputPoolable
+{
+public:
+	virtual void SetOwner(const TSharedPtr<IDecoderOutputOwner, ESPMode::ThreadSafe>& Renderer) = 0;
+};
+
+
diff --git a/Engine/Plugins/Media/ElectraUtil/Source/ElectraBase/Public/MediaVideoDecoderOutput.h b/Engine/Plugins/Media/ElectraUtil/Source/ElectraBase/Public/MediaVideoDecoderOutput.h
new file mode 100644
--- /dev/null
+++ b/Engine/Plugins/Media/ElectraUtil/Source/ElectraBase/Public/MediaVideoDecoderOutput.h
@@ -0,0 +1,186 @@
+// Copyright Epic Games, Inc. All Rights Reserved.
+#pragma once
+
+#include "Misc/Timespan.h"
+#include "Math/IntPoint.h"
+#include "Templates/SharedPointer.h"
+#include "Templates/RefCounting.h"
+#include "PixelFormat.h"
+
+#include "MediaDecoderOutput.h"
+#include "ParameterDictionary.h"
+
+// -------------------------------------------------------------------------------------------------------------------------------
+
+struct FVideoDecoderCropInfo
+{
+	FVideoDecoderCropInfo() : CropLeft(0), CropTop(0), CropRight(0), CropBottom(0) {}
+	FVideoDecoderCropInfo(uint32 InCropLeft, uint32 InCropTop, uint32 InCropRight, uint32 InCropBottom) : CropLeft(InCropLeft), CropTop(InCropRight), CropRight(InCropRight), CropBottom(InCropBottom) {}
+
+	bool HasCropping() const { return CropLeft != 0 || CropRight != 0 || CropTop != 0 || CropBottom != 0;  }
+
+	uint32 CropLeft;
+	uint32 CropTop;
+	uint32 CropRight;
+	uint32 CropBottom;
+};
+
+enum class EVideoOrientation
+{
+	Original = 0,
+	CW90,
+	CW180,
+	CW270,
+};
+
+class IVideoDecoderTexture : public TSharedFromThis<IVideoDecoderTexture, ESPMode::ThreadSafe>
+{
+public:
+	virtual ~IVideoDecoderTexture() {}
+protected:
+	IVideoDecoderTexture() {}
+};
+
+
+class FVideoDecoderOutput : public IDecoderOutput
+{
+public:
+	virtual ~FVideoDecoderOutput() {}
+
+	const Electra::FParamDict& GetDict() const { return *ParamDict; }
+
+	FDecoderTimeStamp GetTime() const
+	{
+		if (!ParamDict)
+		{
+			return FDecoderTimeStamp(FTimespan::Zero(), 0);
+		}
+		return FDecoderTimeStamp(FTimespan(ParamDict->GetValue("pts").GetTimeValue().GetAsHNS()), 0);
+	}
+
+	FTimespan GetDuration() const
+	{
+		if (!ParamDict)
+		{
+			return 0;
+		}
+		return FTimespan(ParamDict->GetValue("duration").GetTimeValue().GetAsHNS());
+	}
+
+	FIntPoint GetOutputDim() const
+	{
+		if (!ParamDict)
+		{
+			return FIntPoint::ZeroValue;
+		}
+		if (!(Cached.Flags & FCached::Valid_OutputDim))
+		{
+			Cached.OutputDim.X = (int32)ParamDict->GetValue("width").SafeGetInt64(0);
+			Cached.OutputDim.Y = (int32)ParamDict->GetValue("height").SafeGetInt64(0);
+			FPlatformMisc::MemoryBarrier();
+			Cached.Flags |= FCached::Valid_OutputDim;
+		}
+		return Cached.OutputDim;
+	}
+
+	double GetAspectRatio() const
+	{
+		if (!ParamDict)
+		{
+			return 1.0;
+		}
+		const FIntPoint Dim = GetOutputDim();
+		return ((double)Dim.X / (double)Dim.Y) * ParamDict->GetValue("aspect_ratio").SafeGetDouble(1.0);
+	}
+
+	FVideoDecoderCropInfo GetCropInfo() const
+	{
+		if (!ParamDict)
+		{
+			return FVideoDecoderCropInfo();
+		}
+		if (!(Cached.Flags & FCached::Valid_CropInfo))
+		{ 
+			Cached.CropInfo.CropLeft = (int32)ParamDict->GetValue("crop_left").SafeGetInt64(0);
+			Cached.CropInfo.CropTop = (int32)ParamDict->GetValue("crop_top").SafeGetInt64(0);
+			Cached.CropInfo.CropRight = (int32)ParamDict->GetValue("crop_right").SafeGetInt64(0);
+			Cached.CropInfo.CropBottom = (int32)ParamDict->GetValue("crop_bottom").SafeGetInt64(0);
+			FPlatformMisc::MemoryBarrier();
+			Cached.Flags |= FCached::Valid_CropInfo;
+		}
+		return Cached.CropInfo;
+	}
+
+	virtual FIntPoint GetDim() const
+	{
+		if (!ParamDict)
+		{
+			return FIntPoint::ZeroValue;
+		}
+		const FIntPoint OutputDim = GetOutputDim();
+		const FVideoDecoderCropInfo CropInfo = GetCropInfo();
+
+		FIntPoint Dim;
+		Dim.X = OutputDim.X + (CropInfo.CropLeft + CropInfo.CropRight);
+		Dim.Y = OutputDim.Y + (CropInfo.CropTop + CropInfo.CropBottom);
+		return Dim;
+	}
+
+	EPixelFormat GetFormat() const
+	{
+		if (!ParamDict)
+		{
+			return EPixelFormat::PF_Unknown;
+		}
+		return (EPixelFormat)ParamDict->GetValue("pixelfmt").SafeGetInt64((int64)EPixelFormat::PF_Unknown);
+	}
+
+	EVideoOrientation GetOrientation() const
+	{
+		if (!ParamDict)
+		{
+			return EVideoOrientation::Original;
+		}
+		if (!(Cached.Flags & FCached::Valid_Orientation))
+		{
+			Cached.Orientation = (EVideoOrientation)ParamDict->GetValue("orientation").SafeGetInt64((int64)EVideoOrientation::Original);
+			FPlatformMisc::MemoryBarrier();
+			Cached.Flags |= FCached::Valid_Orientation;
+		}
+		return Cached.Orientation;
+	}
+
+protected:
+	FVideoDecoderOutput()
+	{
+		Cached.Flags = 0;
+		Cached.OutputDim = FIntPoint::ZeroValue;
+		Cached.Orientation = EVideoOrientation::Original;
+	}
+
+	void Initialize(Electra::FParamDict* InParamDict)
+	{
+		Cached.Flags = 0;
+		ParamDict.Reset(InParamDict);
+	}
+
+private:
+	TUniquePtr<Electra::FParamDict> ParamDict;
+
+	mutable struct FCached
+	{
+		enum
+		{
+			Valid_CropInfo = 1 << 0,
+			Valid_OutputDim = 1 << 1,
+			Valid_Orientation = 1 << 2,
+		};
+
+		uint32 Flags;
+		FVideoDecoderCropInfo CropInfo;
+		FIntPoint OutputDim;
+		EVideoOrientation Orientation;
+	} Cached;
+};
+
+using FVideoDecoderOutputPtr = TSharedPtr<FVideoDecoderOutput, ESPMode::ThreadSafe>;
diff --git a/Engine/Plugins/Media/ElectraUtil/Source/ElectraBase/Public/MediaVideoDecoderOutputAndroid.h b/Engine/Plugins/Media/ElectraUtil/Source/ElectraBase/Public/MediaVideoDecoderOutputAndroid.h
new file mode 100644
--- /dev/null
+++ b/Engine/Plugins/Media/ElectraUtil/Source/ElectraBase/Public/MediaVideoDecoderOutputAndroid.h
@@ -0,0 +1,24 @@
+// Copyright Epic Games, Inc. All Rights Reserved.
+#pragma once
+
+#include "MediaVideoDecoderOutput.h"
+
+#include "Templates/RefCounting.h"
+#include "Containers/Array.h"
+#include "Pixelformat.h"
+
+class FVideoDecoderOutputAndroid : public FVideoDecoderOutput
+{
+public:
+	enum class EOutputType
+	{
+		Unknown = 0,
+		DirectToSurfaceAsView,	//!< Output surface is assumed to be GUI element
+		DirectToSurfaceAsQueue,	//!< Output surface is assumed to be queue for texture output
+	};
+
+	virtual EOutputType GetOutputType() const = 0;
+
+	virtual void ReleaseToSurface() const = 0;
+};
+
diff --git a/Engine/Plugins/Media/ElectraUtil/Source/ElectraBase/Public/MediaVideoDecoderOutputApple.h b/Engine/Plugins/Media/ElectraUtil/Source/ElectraBase/Public/MediaVideoDecoderOutputApple.h
new file mode 100644
--- /dev/null
+++ b/Engine/Plugins/Media/ElectraUtil/Source/ElectraBase/Public/MediaVideoDecoderOutputApple.h
@@ -0,0 +1,16 @@
+// Copyright Epic Games, Inc. All Rights Reserved.
+#pragma once
+
+#include "MediaVideoDecoderOutput.h"
+
+#include "Templates/RefCounting.h"
+#include "Containers/Array.h"
+#include "Pixelformat.h"
+
+class FVideoDecoderOutputApple : public FVideoDecoderOutput
+{
+public:
+	virtual uint32 GetStride() const = 0;
+	virtual CVImageBufferRef GetImageBuffer() const = 0;
+};
+
diff --git a/Engine/Plugins/Media/ElectraUtil/Source/ElectraBase/Public/MediaVideoDecoderOutputPC.h b/Engine/Plugins/Media/ElectraUtil/Source/ElectraBase/Public/MediaVideoDecoderOutputPC.h
new file mode 100644
--- /dev/null
+++ b/Engine/Plugins/Media/ElectraUtil/Source/ElectraBase/Public/MediaVideoDecoderOutputPC.h
@@ -0,0 +1,46 @@
+// Copyright Epic Games, Inc. All Rights Reserved.
+#pragma once
+
+#include "MediaVideoDecoderOutput.h"
+
+#include "Templates/RefCounting.h"
+#include "Containers/Array.h"
+#include "Pixelformat.h"
+
+#include "Windows/AllowWindowsPlatformTypes.h"
+
+#pragma warning(push)
+#pragma warning(disable : 4005)	// macro redefinition
+# include <d3d11.h>
+#pragma warning(pop)
+
+THIRD_PARTY_INCLUDES_START
+#include "mfobjects.h"
+#include "mfapi.h"
+THIRD_PARTY_INCLUDES_END
+
+class FVideoDecoderOutputPC : public FVideoDecoderOutput
+{
+public:
+	enum class EOutputType
+	{
+		Unknown = 0,
+		SoftwareWin7,			// SW decode into buffer
+		SoftwareWin8Plus,		// SW decode into DX11 texture
+		HardwareWin8Plus,		// HW decode into shared DX11 texture
+		HardwareDX9_DX12,		// HW decode into buffer
+	};
+
+	virtual EOutputType GetOutputType() const = 0;
+
+	virtual TRefCountPtr<IMFSample> GetMFSample() const = 0;
+
+	virtual const TArray<uint8> & GetBuffer() const = 0;
+	virtual uint32 GetStride() const = 0;
+
+	virtual TRefCountPtr<ID3D11Texture2D> GetTexture() const = 0;
+	virtual TRefCountPtr<ID3D11Device> GetDevice() const = 0;
+};
+
+#include "Windows/HideWindowsPlatformTypes.h"
+
diff --git a/Engine/Plugins/Media/ElectraUtil/Source/ElectraBase/Public/ParameterDictionary.h b/Engine/Plugins/Media/ElectraUtil/Source/ElectraBase/Public/ParameterDictionary.h
new file mode 100644
--- /dev/null
+++ b/Engine/Plugins/Media/ElectraUtil/Source/ElectraBase/Public/ParameterDictionary.h
@@ -0,0 +1,167 @@
+// Copyright Epic Games, Inc. All Rights Reserved.
+
+#pragma once
+
+#include "CoreMinimal.h"
+#include "PlayerTime.h"
+#include "Templates/TypeCompatibleBytes.h"
+#include "Containers/Array.h"
+
+
+namespace Electra
+{
+
+class ELECTRABASE_API FVariantValue
+{
+	class FSharedPtrHolderBase
+	{
+	public:
+		virtual ~FSharedPtrHolderBase() {}
+		virtual void SetValueOn(FSharedPtrHolderBase *Dst) const = 0;
+	};
+
+	template<typename T> class TSharedPtrHolder : public FSharedPtrHolderBase
+	{
+		using SharedPtrType = TSharedPtr<T, ESPMode::ThreadSafe>;
+
+	public:
+		TSharedPtrHolder(const SharedPtrType & InPointer)
+			: Pointer(InPointer)
+		{}
+
+		SharedPtrType	Pointer;
+
+		virtual void SetValueOn(FSharedPtrHolderBase *Dst) const override
+		{
+			new(Dst) TSharedPtrHolder<T>(reinterpret_cast<const SharedPtrType&>(Pointer));
+		}
+	};
+
+public:
+	FVariantValue();
+	~FVariantValue();
+	FVariantValue(const FVariantValue& rhs);
+	FVariantValue& operator=(const FVariantValue&rhs);
+
+	explicit FVariantValue(const FString& StringValue);
+	explicit FVariantValue(double DoubleValue);
+	explicit FVariantValue(int64 Int64Value);
+	explicit FVariantValue(bool BoolValue);
+	explicit FVariantValue(const FTimeValue& TimeValue);
+	explicit FVariantValue(void* PointerValue);
+	template<typename T> explicit FVariantValue(const TSharedPtr<T, ESPMode::ThreadSafe>& PointerValue)
+	: DataType(EDataType::TypeUninitialized)
+	{
+		Set(PointerValue);
+	}
+
+	FVariantValue& Set(const FString& StringValue);
+	FVariantValue& Set(double DoubleValue);
+	FVariantValue& Set(int64 Int64Value);
+	FVariantValue& Set(bool BoolValue);
+	FVariantValue& Set(const FTimeValue& TimeValue);
+	FVariantValue& Set(void* PointerValue);
+	template <typename T> FVariantValue& Set(const TSharedPtr<T, ESPMode::ThreadSafe>& PointerValue)
+	{
+		Clear();
+		new(&DataBuffer) TSharedPtrHolder<T>(PointerValue);
+		DataType = EDataType::TypeSharedPointer;
+		return *this;
+	}
+
+	// Returns variant value. Type *must* match. Otherwise an empty/zero value is returned.
+	const FString& GetFString() const;
+	const double& GetDouble() const;
+	const int64& GetInt64() const;
+	const bool& GetBool() const;
+	const FTimeValue& GetTimeValue() const;
+	void* const & GetPointer() const;
+	template<typename T> TSharedPtr<T, ESPMode::ThreadSafe> GetSharedPointer() const
+	{
+		check(DataType == EDataType::TypeSharedPointer);
+		if (DataType == EDataType::TypeSharedPointer)
+		{
+			const TSharedPtrHolder<T>* Pointer = reinterpret_cast<const TSharedPtrHolder<T>*>(&DataBuffer);
+			return Pointer->Pointer;
+		}
+		else
+		{
+			return TSharedPtr<T, ESPMode::ThreadSafe>();
+		}
+	}
+
+	// Returns variant value. If type does not match the specified default will be returned.
+	const FString& SafeGetFString(const FString& Default = FString()) const;
+	double SafeGetDouble(double Default=0.0) const;
+	int64 SafeGetInt64(int64 Default=0) const;
+	bool SafeGetBool(bool Default=false) const;
+	FTimeValue SafeGetTimeValue(const FTimeValue& Default=FTimeValue()) const;
+	void* SafeGetPointer(void* Default=nullptr) const;
+
+	enum class EDataType
+	{
+		TypeUninitialized,
+		TypeFString,
+		TypeDouble,
+		TypeInt64,
+		TypeBoolean,
+		TypeTimeValue,
+		TypeVoidPointer,
+		TypeSharedPointer,
+	};
+	EDataType GetDataType() const
+	{
+		return DataType;
+	}
+
+	bool IsValid() const
+	{
+		return DataType != EDataType::TypeUninitialized;
+	}
+
+	bool IsType(EDataType type) const
+	{
+		return DataType == type;
+	}
+
+private:
+
+	union FUnionLayout
+	{
+		uint8 MemSizeFString[sizeof(FString)];
+		uint8 MemSizeDouble[sizeof(double)];
+		uint8 MemSizeInt64[sizeof(int64)];
+		uint8 MemSizeVoidPtr[sizeof(void*)];
+		uint8 MemSizeBoolean[sizeof(bool)];
+		uint8 MemSizeTimeValue[sizeof(FTimeValue)];
+		uint8 MemSizeSharedPtrValue[sizeof(TSharedPtrHolder<uint8>)];
+	};
+
+	void Clear();
+	void CopyInternal(const FVariantValue& FromOther);
+
+	TAlignedBytes<sizeof(FUnionLayout), 16> DataBuffer;
+	EDataType	DataType;
+};
+
+
+
+class ELECTRABASE_API FParamDict
+{
+public:
+	FParamDict();
+	~FParamDict();
+	void Clear();
+	void Set(const FString& Key, const FVariantValue& Value);
+	bool HaveKey(const FString& Key) const;
+	const FVariantValue& GetValue(const FString& Key) const;
+
+	void GetKeysStartingWith(const FString& StartsWith, TArray<FString>& Keys);
+private:
+	TMap<FString, FVariantValue>		Dictionary;
+};
+
+
+} // namespace Electra
+
+
diff --git a/Engine/Plugins/Media/ElectraUtil/Source/ElectraBase/Public/PlayerTime.h b/Engine/Plugins/Media/ElectraUtil/Source/ElectraBase/Public/PlayerTime.h
new file mode 100644
--- /dev/null
+++ b/Engine/Plugins/Media/ElectraUtil/Source/ElectraBase/Public/PlayerTime.h
@@ -0,0 +1,546 @@
+// Copyright Epic Games, Inc. All Rights Reserved.
+
+#pragma once
+
+#include "CoreMinimal.h"
+#include <limits>
+
+namespace Electra
+{
+class FTimeFraction;
+
+/**
+ * Keeps a time value in hundred nanoseconds (HNS).
+ */
+class FTimeValue
+{
+public:
+	static FTimeValue GetInvalid()
+	{
+		static FTimeValue kInvalid;
+		return kInvalid;
+	}
+
+	static FTimeValue GetZero()
+	{
+		static FTimeValue kZero((int64)0);
+		return kZero;
+	}
+
+	static FTimeValue GetPositiveInfinity()
+	{
+		static FTimeValue kPosInf(std::numeric_limits<double>::infinity());
+		return kPosInf;
+	}
+
+	FTimeValue()
+		: HNS(0)
+		, bIsValid(false)
+		, bIsInfinity(false)
+	{
+	}
+
+	FTimeValue(const FTimeValue& rhs)
+		: HNS(rhs.HNS)
+		, bIsValid(rhs.bIsValid)
+		, bIsInfinity(rhs.bIsInfinity)
+	{
+	}
+
+	FTimeValue& operator=(const FTimeValue& rhs)
+	{
+		HNS 		= rhs.HNS;
+		bIsValid	= rhs.bIsValid;
+		bIsInfinity = rhs.bIsInfinity;
+		return *this;
+	}
+
+	bool IsValid() const
+	{
+		return bIsValid;
+	}
+
+	bool IsInfinity() const
+	{
+		return bIsInfinity;
+	}
+
+	double GetAsSeconds(double DefaultIfInvalid=0.0) const
+	{
+		return bIsValid ? (bIsInfinity ? (HNS >= 0 ? std::numeric_limits<double>::infinity() : -std::numeric_limits<double>::infinity()) : HNS / 10000000.0) : DefaultIfInvalid;
+	}
+
+	int64 GetAsMilliseconds(int64 DefaultIfInvalid=0) const
+	{
+		return bIsValid ? (bIsInfinity ? (HNS >= 0 ? 0x7fffffffffffffffLL : -0x7fffffffffffffffLL) : HNS / 10000) : DefaultIfInvalid;
+	}
+
+	int64 GetAsMicroseconds(int64 DefaultIfInvalid=0) const
+	{
+		return bIsValid ? (bIsInfinity ? (HNS >= 0 ? 0x7fffffffffffffffLL : -0x7fffffffffffffffLL) : HNS / 10) : DefaultIfInvalid;
+	}
+
+	int64 GetAsHNS(int64 DefaultIfInvalid=0) const
+	{
+		return bIsValid ? (bIsInfinity ? (HNS >= 0 ? 0x7fffffffffffffffLL : -0x7fffffffffffffffLL) : HNS) : DefaultIfInvalid;
+	}
+
+	int64 GetAs90kHz(int64 DefaultIfInvalid=0) const
+	{
+		return bIsValid ? (bIsInfinity ? (HNS >= 0 ? 0x7fffffffffffffffLL : -0x7fffffffffffffffLL) : HNS * 9 / 1000) : DefaultIfInvalid;
+	}
+
+	//! Returns this time value in a custom timebase. Requires internal bigint conversion and is therefor SLOW!
+	int64 GetAsTimebase(uint32 CustomTimebase) const;
+
+	FTimeValue& SetToInvalid()
+	{
+		HNS 		= 0;
+		bIsValid	= false;
+		bIsInfinity = false;
+		return *this;
+	}
+
+	FTimeValue& SetToZero()
+	{
+		HNS 		= 0;
+		bIsValid	= true;
+		bIsInfinity = false;
+		return *this;
+	}
+
+	FTimeValue& SetToPositiveInfinity()
+	{
+		HNS 		= 0x7fffffffffffffffLL;
+		bIsValid	= true;
+		bIsInfinity = true;
+		return *this;
+	}
+
+	FTimeValue& SetFromSeconds(double Seconds)
+	{
+		if ((bIsInfinity = (Seconds == std::numeric_limits<double>::infinity())) == true)
+		{
+			HNS = 0x7fffffffffffffffLL;
+			bIsValid = true;
+		}
+		else
+		{
+			if ((bIsValid = Seconds >= -922337203685.0 && Seconds <= 922337203685.0) == true)
+			{
+				HNS = (int64)(Seconds * 10000000.0);
+			}
+			else
+			{
+				check(!"Value cannot be represented!");
+				HNS = 0;
+			}
+		}
+		return *this;
+	}
+
+
+	FTimeValue& SetFromMilliseconds(int64 Milliseconds)
+	{
+		bIsInfinity = false;
+		if ((bIsValid = Milliseconds >= -922337203685477 && Milliseconds <= 922337203685477) == true)
+		{
+			HNS = Milliseconds * 10000;
+		}
+		else
+		{
+			check(!"Value cannot be represented!");
+			HNS = 0;
+		}
+		return *this;
+	}
+
+	FTimeValue& SetFromMicroseconds(int64 Microseconds)
+	{
+		bIsInfinity = false;
+		if ((bIsValid = Microseconds >= -922337203685477580 && Microseconds <= 922337203685477580) == true)
+		{
+			HNS = Microseconds * 10;
+		}
+		else
+		{
+			check(!"Value cannot be represented!");
+			HNS = 0;
+		}
+		return *this;
+	}
+
+	FTimeValue& SetFromHNS(int64 InHNS)
+	{
+		HNS 		= InHNS;
+		bIsValid	= true;
+		bIsInfinity = false;
+		return *this;
+	}
+
+	FTimeValue& SetFrom90kHz(int64 Ticks)
+	{
+		HNS 		= Ticks * 1000 / 9;
+		bIsValid	= true;
+		bIsInfinity = false;
+		return *this;
+	}
+
+	FTimeValue& SetFromND(int64 Numerator, uint32 Denominator)
+	{
+		if (Denominator != 0)
+		{
+			if (Denominator == 10000000)
+			{
+				HNS 		= Numerator;
+				bIsValid	= true;
+				bIsInfinity = false;
+			}
+			else if (Numerator >= -922337203685LL && Numerator <= 922337203685LL)
+			{
+				HNS 		= Numerator * 10000000 / Denominator;
+				bIsValid	= true;
+				bIsInfinity = false;
+			}
+			else
+			{
+				// Need a bigint here!
+				check(!"TODO");
+				HNS 		= (int64)(Numerator * (10000000.0 / (double)Denominator));
+				bIsValid	= true;
+				bIsInfinity = false;
+			}
+		}
+		else
+		{
+			HNS 		= Numerator>=0 ? 0x7fffffffffffffffLL : -0x7fffffffffffffffLL;
+			bIsValid	= true;
+			bIsInfinity = true;
+		}
+		return *this;
+	}
+
+
+	FTimeValue& SetFromTimeFraction(const FTimeFraction& TimeFraction);
+
+
+	FTimeValue& AdvanceBySeconds(double Seconds)
+	{
+		FTimeValue tv;
+		tv.SetFromSeconds(Seconds);
+		return operator += (tv);
+	}
+
+
+	bool operator == (const FTimeValue& rhs) const
+	{
+		return (!bIsValid && !rhs.bIsValid) || (bIsValid == rhs.bIsValid && bIsInfinity == rhs.bIsInfinity && HNS == rhs.HNS);
+	}
+
+	bool operator != (const FTimeValue& rhs) const
+	{
+		return !(*this == rhs);
+	}
+
+	bool operator < (const FTimeValue& rhs) const
+	{
+		if (bIsValid && rhs.bIsValid)
+		{
+			if (!bIsInfinity)
+			{
+				return !rhs.bIsInfinity ? HNS < rhs.HNS : rhs.HNS > 0;
+			}
+			else
+			{
+				return rhs.bIsInfinity ? HNS < rhs.HNS : HNS < 0;
+			}
+		}
+		return false;
+	}
+
+	bool operator <= (const FTimeValue& rhs) const
+	{
+		return (!bIsValid || !rhs.bIsValid) ? false : (*this < rhs || *this == rhs);
+	}
+
+	bool operator > (const FTimeValue &rhs) const
+	{
+		return (!bIsValid || !rhs.bIsValid) ? false : !(*this <= rhs);
+	}
+
+	bool operator >= (const FTimeValue &rhs) const
+	{
+		return (!bIsValid || !rhs.bIsValid) ? false : !(*this < rhs);
+	}
+
+	inline FTimeValue& operator += (const FTimeValue& rhs)
+	{
+		if (bIsValid)
+		{
+			if (rhs.bIsValid)
+			{
+				if (!bIsInfinity && !rhs.bIsInfinity)
+				{
+					HNS += rhs.HNS;
+				}
+				else
+				{
+					SetToInvalid();
+				}
+			}
+			else
+			{
+				SetToInvalid();
+			}
+		}
+		return *this;
+	}
+
+	inline FTimeValue& operator -= (const FTimeValue& rhs)
+	{
+		if (bIsValid)
+		{
+			if (rhs.bIsValid)
+			{
+				if (!bIsInfinity && !rhs.bIsInfinity)
+				{
+					HNS -= rhs.HNS;
+				}
+				else
+				{
+					SetToInvalid();
+				}
+			}
+			else
+			{
+				SetToInvalid();
+			}
+		}
+		return *this;
+	}
+
+	inline FTimeValue& operator /= (int32 Scale)
+	{
+		if (bIsValid && !bIsInfinity)
+		{
+			if (Scale)
+			{
+				HNS /= Scale;
+			}
+			else
+			{
+				SetToPositiveInfinity();
+			}
+		}
+		return *this;
+	}
+
+	inline FTimeValue& operator *= (int32 Scale)
+	{
+		if (bIsValid && !bIsInfinity)
+		{
+		// FIXME: what if this overflows?
+			HNS *= Scale;
+		}
+		return *this;
+	}
+
+	inline FTimeValue operator + (const FTimeValue& rhs) const
+	{
+		FTimeValue Result;
+		if (bIsValid && rhs.bIsValid)
+		{
+			if (!bIsInfinity && !rhs.bIsInfinity)
+			{
+				Result.HNS = HNS + rhs.HNS;
+				Result.bIsValid = true;
+			}
+		}
+		return Result;
+	}
+
+	inline FTimeValue operator - (const FTimeValue& rhs) const
+	{
+		FTimeValue Result;
+		if (bIsValid && rhs.bIsValid)
+		{
+			if (!bIsInfinity && !rhs.bIsInfinity)
+			{
+				Result.HNS = HNS - rhs.HNS;
+				Result.bIsValid = true;
+			}
+		}
+		return Result;
+	}
+
+	inline FTimeValue operator << (int32 Shift) const
+	{
+		FTimeValue Result(*this);
+		if (bIsValid && !bIsInfinity)
+		{
+			Result.HNS <<= Shift;
+		}
+		return Result;
+	}
+
+	inline FTimeValue operator >> (int32 Shift) const
+	{
+		FTimeValue Result(*this);
+		if (bIsValid && !bIsInfinity)
+		{
+			Result.HNS >>= Shift;
+		}
+		return Result;
+	}
+
+	inline FTimeValue operator * (int32 Scale) const
+	{
+		FTimeValue Result(*this);
+		if (bIsValid && !bIsInfinity)
+		{
+			Result.HNS *= Scale;
+		}
+		return Result;
+	}
+
+	inline FTimeValue operator / (int32 Scale) const
+	{
+		FTimeValue Result(*this);
+		if (bIsValid && !bIsInfinity)
+		{
+			Result.HNS /= Scale;
+		}
+		return(Result);
+	}
+
+private:
+	explicit FTimeValue(int64 InHNS) : HNS(InHNS), bIsValid(true), bIsInfinity(false)
+	{
+	}
+	explicit FTimeValue(double Seconds)
+	{
+		SetFromSeconds(Seconds);
+	}
+
+	int64	HNS;
+	bool	bIsValid;
+	bool	bIsInfinity;
+};
+
+
+
+struct FTimeRange
+{
+	void Reset()
+	{
+		Start.SetToInvalid();
+		End.SetToInvalid();
+	}
+	bool IsValid() const
+	{
+		return Start.IsValid() && End.IsValid();
+	}
+	FTimeValue	Start;
+	FTimeValue	End;
+};
+
+
+
+/**
+ * Keeps a time value as a fractional.
+ */
+class FTimeFraction
+{
+public:
+	static const FTimeFraction& GetInvalid()
+	{
+		static FTimeFraction Invalid;
+		return Invalid;
+	}
+
+	static const FTimeFraction& GetZero()
+	{
+		static FTimeFraction Zero(0, 1);
+		return Zero;
+	}
+
+	FTimeFraction() : Numerator(0), Denominator(0), bIsValid(false)
+	{
+	}
+
+	FTimeFraction(int64 n, uint32 d) : Numerator(n), Denominator(d), bIsValid(true)
+	{
+	}
+
+	FTimeFraction(const FTimeFraction& rhs) : Numerator(rhs.Numerator), Denominator(rhs.Denominator), bIsValid(rhs.bIsValid)
+	{
+	}
+
+	FTimeFraction& operator=(const FTimeFraction& rhs)
+	{
+		Numerator   = rhs.Numerator;
+		Denominator = rhs.Denominator;
+		bIsValid	= rhs.bIsValid;
+		return *this;
+	}
+
+	bool IsValid() const
+	{
+		return bIsValid;
+	}
+
+	bool IsPositiveInfinity() const
+	{
+		return bIsValid && Denominator==0 && Numerator>=0;
+	}
+
+	int64 GetNumerator() const
+	{
+		return Numerator;
+	}
+
+	uint32 GetDenominator() const
+	{
+		return Denominator;
+	}
+
+	double GetAsDouble() const
+	{
+		return Numerator / (double)Denominator;
+	}
+
+	FTimeFraction& SetFromND(int64 InNumerator, uint32 InDenominator)
+	{
+		Numerator   = InNumerator;
+		Denominator = InDenominator;
+		bIsValid	= true;
+		return *this;
+	}
+
+	FTimeFraction& SetFromFloatString(const FString& In);
+
+private:
+	int64	Numerator;
+	uint32	Denominator;
+	bool	bIsValid;
+};
+
+
+}
+
+
+
+/**
+ * System wallclock time (UTC)
+**/
+class MEDIAutcTime
+{
+public:
+	static Electra::FTimeValue Current()
+	{
+		return Electra::FTimeValue().SetFromMilliseconds(CurrentMSec());
+	}
+	static int64 CurrentMSec();
+};
+
+
diff --git a/Engine/Plugins/Media/ElectraUtil/Source/ElectraSamples/ElectraSamples.Build.cs b/Engine/Plugins/Media/ElectraUtil/Source/ElectraSamples/ElectraSamples.Build.cs
new file mode 100644
--- /dev/null
+++ b/Engine/Plugins/Media/ElectraUtil/Source/ElectraSamples/ElectraSamples.Build.cs
@@ -0,0 +1,94 @@
+// Copyright Epic Games, Inc. All Rights Reserved.
+using EpicGames.Core;
+using System.IO;
+
+namespace UnrealBuildTool.Rules
+{
+	public class ElectraSamples: ModuleRules
+	{
+		public ElectraSamples(ReadOnlyTargetRules Target) : base(Target)
+		{
+
+            bLegalToDistributeObjectCode = true;
+            PCHUsage = PCHUsageMode.NoPCHs;
+
+            DynamicallyLoadedModuleNames.AddRange(
+				new string[] {
+					"Media",
+				});
+
+			PrivateDependencyModuleNames.AddRange(
+				new string[] {
+					"Core",
+					"CoreUObject",
+					"MediaUtils",
+					"RenderCore",
+					"RHI",
+					"RHICore",
+					"ElectraBase",
+				});
+
+			PrivateIncludePathModuleNames.AddRange(
+				new string[] {
+					"Media",
+                });
+
+			if (Target.bCompileAgainstEngine)
+			{
+				PrivateDependencyModuleNames.Add("Engine");
+			}
+
+            if (Target.Platform.IsInGroup(UnrealPlatformGroup.Windows))
+            {
+                string DirectXSDKDir = Target.UEThirdPartySourceDirectory + "Windows/DirectX";
+                PublicSystemIncludePaths.Add(DirectXSDKDir + "/include");
+
+				if (Target.Platform == UnrealTargetPlatform.Win64)
+                {
+                    DirectXSDKDir += "/Lib/x64/";
+                }
+                else if (Target.Platform == UnrealTargetPlatform.Win32)
+                {
+                    DirectXSDKDir += "/Lib/x86/";
+                }
+
+				if (Target.Platform == UnrealTargetPlatform.Win64 || Target.Platform == UnrealTargetPlatform.Win32)
+				{
+					AddEngineThirdPartyPrivateStaticDependencies(Target, "DX9");
+
+					PublicAdditionalLibraries.AddRange(new string[] {
+						DirectXSDKDir + "dxerr.lib",
+					});
+
+					PrivateDependencyModuleNames.Add("D3D11RHI");
+
+					PrivateDefinitions.Add("ELECTRA_SUPPORT_PREWIN8");
+					PrivateDefinitions.Add("ELECTRA_HAVE_DX11");
+				}
+
+				PublicSystemLibraries.AddRange(new string[] {
+                    "strmiids.lib",
+                    "legacy_stdio_definitions.lib",
+                    "Dxva2.lib",
+                });
+
+				PublicIncludePaths.Add("$(ModuleDir)/Public/Windows");
+			}
+			else if (Target.Platform == UnrealTargetPlatform.Mac || Target.Platform == UnrealTargetPlatform.IOS)
+            {
+                PrivateDependencyModuleNames.Add("MetalRHI");
+				PublicIncludePaths.Add("$(ModuleDir)/Public/Apple");
+			}
+			else if (Target.IsInPlatformGroup(UnrealPlatformGroup.Android))
+			{
+				PublicIncludePaths.Add("$(ModuleDir)/Public/Android");
+				string PluginPath = Utils.MakePathRelativeTo(ModuleDirectory, Target.RelativeEnginePath);
+				AdditionalPropertiesForReceipt.Add("AndroidPlugin", Path.Combine(PluginPath, "ElectraSamples_UPL.xml"));
+			}
+			else if (Target.Platform == UnrealTargetPlatform.Linux)
+			{
+				PublicIncludePaths.Add("$(ModuleDir)/Public/GStreamer");
+			}
+		}
+	}
+}
diff --git a/Engine/Plugins/Media/ElectraUtil/Source/ElectraSamples/Private/ElectraSamplesModule.cpp b/Engine/Plugins/Media/ElectraUtil/Source/ElectraSamples/Private/ElectraSamplesModule.cpp
new file mode 100644
--- /dev/null
+++ b/Engine/Plugins/Media/ElectraUtil/Source/ElectraSamples/Private/ElectraSamplesModule.cpp
@@ -0,0 +1,41 @@
+// Copyright Epic Games, Inc. All Rights Reserved.
+
+#include "ElectraSamplesModule.h"
+
+#include "CoreMinimal.h"
+#include "Modules/ModuleManager.h"
+#include "IElectraSamplesModule.h"
+
+#include "IMediaTextureSample.h"
+#include "MediaVideoDecoderOutput.h"
+
+#define LOCTEXT_NAMESPACE "ElectraBaseModule"
+
+DEFINE_LOG_CATEGORY(LogElectraSamples);
+
+// -----------------------------------------------------------------------------------------------------------------------------------
+
+class FElectraSamplesModule: public IElectraSamplesModule
+{
+public:
+public:
+	// IModuleInterface interface
+
+	void StartupModule() override
+	{
+		static_assert((int32)EMediaOrientation::Original == (int32)EVideoOrientation::Original, "check alignment of both enums");
+		static_assert((int32)EMediaOrientation::CW90 == (int32)EVideoOrientation::CW90, "check alignment of both enums");
+	}
+
+	void ShutdownModule() override
+	{
+	}
+
+private:
+};
+
+IMPLEMENT_MODULE(FElectraSamplesModule, ElectraSamples);
+
+#undef LOCTEXT_NAMESPACE
+
+
diff --git a/Engine/Plugins/Media/ElectraUtil/Source/ElectraSamples/Private/ElectraSamplesTextureManager.cpp b/Engine/Plugins/Media/ElectraUtil/Source/ElectraSamples/Private/ElectraSamplesTextureManager.cpp
new file mode 100644
--- /dev/null
+++ b/Engine/Plugins/Media/ElectraUtil/Source/ElectraSamples/Private/ElectraSamplesTextureManager.cpp
@@ -0,0 +1,78 @@
+// Copyright Epic Games, Inc. All Rights Reserved.
+
+#include "ElectraSamplesTextureManager.h"
+#include "misc/ScopeLock.h"
+
+FElectraSamplesTextureManager::FElectraSamplesTextureManager()
+{
+}
+
+FElectraSamplesTextureManager::~FElectraSamplesTextureManager()
+{
+}
+
+void FElectraSamplesTextureManager::CleanupMap()
+{
+	TArray<void*> InvalidKeys;
+	InvalidKeys.Reserve(KnownTextures.Num());
+	for (const TPair<void*, TWeakPtr<IVideoDecoderTexture, ESPMode::ThreadSafe>>& Item : KnownTextures)
+	{
+		if (!Item.Value.IsValid())
+		{
+			InvalidKeys.Add(Item.Key);
+		}
+	}
+	for (void* Key : InvalidKeys)
+	{
+		KnownTextures.Remove(Key);
+	}
+}
+
+TSharedPtr<IVideoDecoderTexture, ESPMode::ThreadSafe> FElectraSamplesTextureManager::CreateTexture(const FIntPoint& Dim, EPixelFormat Fmt)
+{
+	FScopeLock Lock(&AccessCS);
+
+	CleanupMap();
+	TSharedPtr<IVideoDecoderTexture, ESPMode::ThreadSafe> NewTexture = PlatformCreateTexture(Dim, Fmt);
+	if (NewTexture.IsValid())
+	{
+		KnownTextures.Add(PlatformGetTexturePlatform(NewTexture.Get()), NewTexture);
+	}
+	return NewTexture;
+}
+
+void FElectraSamplesTextureManager::PlatformAddSharedTextureRef(TSharedPtr<IVideoDecoderTexture, ESPMode::ThreadSafe> InTextureRef)
+{
+	FScopeLock Lock(&AccessCS);
+	if (InTextureRef.IsValid())
+	{
+		// Only add once if it's not already there.
+		KnownTextures.FindOrAdd(PlatformGetTexturePlatform(InTextureRef.Get()), InTextureRef);
+	}
+}
+
+void FElectraSamplesTextureManager::PlatformRemoveSharedTextureRef(TSharedPtr<IVideoDecoderTexture, ESPMode::ThreadSafe> InTextureRef)
+{
+	FScopeLock Lock(&AccessCS);
+	if (InTextureRef.IsValid())
+	{
+		KnownTextures.Remove(PlatformGetTexturePlatform(InTextureRef.Get()));
+	}
+}
+
+
+FTexture2DRHIRef FElectraSamplesTextureManager::GetRHITextureFromPlatformTexture(void* PlatformTexture)
+{
+	FScopeLock Lock(&AccessCS);
+
+	CleanupMap();
+	TWeakPtr<IVideoDecoderTexture, ESPMode::ThreadSafe> *Texture = KnownTextures.Find(PlatformTexture);
+	if (Texture)
+	{
+		if (TSharedPtr<IVideoDecoderTexture, ESPMode::ThreadSafe> PinnedTexture = Texture->Pin())
+		{
+			return PlatformGetTextureRHI(PinnedTexture.Get());
+		}
+	}
+	return FTexture2DRHIRef();
+}
