ElectraPlayer: allowing for master playlists to not specify an audio-only variant if audio renditions exist; allowing for whitespace in attribute list; removed vendor specific mp4 boxes

#rnx

diff --git a/Engine/Plugins/Media/ElectraPlayer/Source/ElectraPlayerRuntime/Private/Runtime/Demuxer/ParserISO14496-12.cpp b/Engine/Plugins/Media/ElectraPlayer/Source/ElectraPlayerRuntime/Private/Runtime/Demuxer/ParserISO14496-12.cpp
--- a/Engine/Plugins/Media/ElectraPlayer/Source/ElectraPlayerRuntime/Private/Runtime/Demuxer/ParserISO14496-12.cpp
+++ b/Engine/Plugins/Media/ElectraPlayer/Source/ElectraPlayerRuntime/Private/Runtime/Demuxer/ParserISO14496-12.cpp
@@ -19,6 +19,10 @@
 #define MEDIA_DEBUG_HAS_BOX_NAMES 0
 #endif
 
+DECLARE_LOG_CATEGORY_EXTERN(LogElectraMP4Parser, Log, All);
+DEFINE_LOG_CATEGORY(LogElectraMP4Parser);
+
+
 namespace Electra
 {
 	namespace Utils
@@ -184,7 +188,6 @@ namespace Electra
 			, CurrentHandlerBox(nullptr)
 			, CurrentMediaHandlerBox(nullptr)
 			, PlayerSession(nullptr)
-			, bParsingVendorSpecific(false)
 			, NumTotalBoxesParsed(0)
 		{
 		}
@@ -287,22 +290,20 @@ namespace Electra
 			return BoxNestingLevel;
 		}
 
-		FMP4Box* RootBox;				//!< A root box that is not an actual file box but a container representing the file itself.
-		FMP4BoxReader* BoxReader;				//!< Instance of the box reader we were given. This is not ours so we must not delete it!
-		int32						BoxNestingLevel;		//!< Box tree depth
+		FMP4Box* RootBox;					//!< A root box that is not an actual file box but a container representing the file itself.
+		FMP4BoxReader* BoxReader;			//!< Instance of the box reader we were given. This is not ours so we must not delete it!
+		int32 BoxNestingLevel;				//!< Box tree depth
 
 		FMP4Box* CurrentMoovBox;			//!< moov being parsed at the moment.
 		FMP4Box* CurrentMoofBox;			//!< moof being parsed at the moment.
-		FMP4Box* CurrentTrackBox;		//!< trak/traf being parsed at the moment.
-		FMP4Box* CurrentHandlerBox;		//!< hdlr being parsed at the moment.
+		FMP4Box* CurrentTrackBox;			//!< trak/traf being parsed at the moment.
+		FMP4Box* CurrentHandlerBox;			//!< hdlr being parsed at the moment.
 		FMP4Box* CurrentMediaHandlerBox;	//!< media handler being parsed at the moment ('vmhd', 'smhd', 'sthd', nmhd')
 
 
-		FParamDict					Options;
+		FParamDict Options;
 		IPlayerSessionServices* PlayerSession;
-		bool						bParsingVendorSpecific;
-		int32						NumTotalBoxesParsed;
-
+		int32 NumTotalBoxesParsed;
 	};
 
 
@@ -536,16 +537,6 @@ namespace Electra
 		static const IParserISO14496_12::FBoxType kBox_tfdt = MAKE_BOX_ATOM('t', 'f', 'd', 't');
 		static const IParserISO14496_12::FBoxType kBox_trun = MAKE_BOX_ATOM('t', 'r', 'u', 'n');
 
-		/**
-		 * Vendor specific custom boxes: IBM Ustream
-		*/
-		static const IParserISO14496_12::FBoxType kBox_ustm = MAKE_BOX_ATOM('u', 's', 't', 'm');
-		static const IParserISO14496_12::FBoxType kBox_fhst = MAKE_BOX_ATOM('f', 'h', 's', 't');
-		static const IParserISO14496_12::FBoxType kBox_fmov = MAKE_BOX_ATOM('f', 'm', 'o', 'v');
-		static const IParserISO14496_12::FBoxType kBox_idta = MAKE_BOX_ATOM('i', 'd', 't', 'a');
-		static const IParserISO14496_12::FBoxType kBox_hash = MAKE_BOX_ATOM('h', 'a', 's', 'h');
-
-
 		/**
 		 * Track handler types
 		 */
@@ -3138,7 +3129,7 @@ namespace Electra
 #if MEDIA_DEBUG_HAS_BOX_NAMES
 				char boxName[4];
 				*((uint32*)&boxName) = MEDIA_TO_BIG_ENDIAN(BoxType);
-				//UE_LOG(LogTemp, Log, TEXT("Next box: '%c%c%c%c' %lld B @ %lld"), boxName[0]?boxName[0]:'?', boxName[1]?boxName[1]:'?', boxName[2]?boxName[2]:'?', boxName[3]?boxName[3]:'?', (long long int)BoxSize, (long long int)BoxStartOffset);
+				//UE_LOG(LogElectraMP4Parser, Log, TEXT("Next box: '%c%c%c%c' %lld B @ %lld"), boxName[0]?boxName[0]:'?', boxName[1]?boxName[1]:'?', boxName[2]?boxName[2]:'?', boxName[3]?boxName[3]:'?', (long long int)BoxSize, (long long int)BoxStartOffset);
 #endif
 
 				int64 BoxDataOffset = BoxReader->GetCurrentReadOffset();
@@ -3317,15 +3308,7 @@ namespace Electra
 							break;
 
 						case FMP4Box::kBox_free:
-							if (NumTotalBoxesParsed == 0 && !bParsingVendorSpecific && Options.GetValue(IParserISO14496_12::OptionKeyParseFirstFreeBox).GetBool())
-							{
-								NextBox = new FMP4BoxBasic(BoxType, BoxSize, BoxStartOffset, BoxDataOffset, false);
-								bParsingVendorSpecific = true;
-							}
-							else
-							{
-								NextBox = new FMP4BoxIgnored(BoxType, BoxSize, BoxStartOffset, BoxDataOffset, true);
-							}
+							NextBox = new FMP4BoxIgnored(BoxType, BoxSize, BoxStartOffset, BoxDataOffset, true);
 							break;
 
 						case FMP4Box::kBox_uuid:
@@ -3334,27 +3317,9 @@ namespace Electra
 							NextBox = new FMP4BoxIgnored(BoxType, BoxSize, BoxStartOffset, BoxDataOffset, true);
 							break;
 
-							// Vendor specific (IBM Ustream)
-						case FMP4Box::kBox_ustm:
-						case FMP4Box::kBox_fmov:
-						case FMP4Box::kBox_idta:
-							if (bParsingVendorSpecific)
-							{
-								NextBox = new FMP4BoxBasic(BoxType, BoxSize, BoxStartOffset, BoxDataOffset, false);
-							}
-							else
-							{
-								NextBox = new FMP4BoxIgnored(BoxType, BoxSize, BoxStartOffset, BoxDataOffset, true);
-							}
-							break;
-						case FMP4Box::kBox_fhst:
-						case FMP4Box::kBox_hash:
-							NextBox = new FMP4BoxIgnored(BoxType, BoxSize, BoxStartOffset, BoxDataOffset, true);
-							break;
-
 						default:
 #if MEDIA_DEBUG_HAS_BOX_NAMES
-							ensureAlwaysMsgf(0, TEXT("Ignoring mp4 box '%c%c%c%c'"), boxName[0], boxName[1], boxName[2], boxName[3]);
+							UE_LOG(LogElectraMP4Parser, Log, TEXT("Ignoring mp4 box: '%c%c%c%c' %lld B @ %lld"), boxName[0]?boxName[0]:'?', boxName[1]?boxName[1]:'?', boxName[2]?boxName[2]:'?', boxName[3]?boxName[3]:'?', (long long int)BoxSize, (long long int)BoxStartOffset);
 #endif
 							NextBox = new FMP4BoxIgnored(BoxType, BoxSize, BoxStartOffset, BoxDataOffset, true);
 							break;
@@ -5086,9 +5051,6 @@ namespace Electra
 
 	/*********************************************************************************************************************/
 
-	const FString IParserISO14496_12::OptionKeyParseFirstFreeBox(TEXT("ISO14496-12:parseFreeBox"));					//!< Option to not skip 'free' box but instead attempt to parse their content
-
-
 	TSharedPtrTS<IParserISO14496_12> IParserISO14496_12::CreateParser()
 	{
 		return MakeSharedTS<FParserISO14496_12>();
diff --git a/Engine/Plugins/Media/ElectraPlayer/Source/ElectraPlayerRuntime/Private/Runtime/Demuxer/ParserISO14496-12.h b/Engine/Plugins/Media/ElectraPlayer/Source/ElectraPlayerRuntime/Private/Runtime/Demuxer/ParserISO14496-12.h
--- a/Engine/Plugins/Media/ElectraPlayer/Source/ElectraPlayerRuntime/Private/Runtime/Demuxer/ParserISO14496-12.h
+++ b/Engine/Plugins/Media/ElectraPlayer/Source/ElectraPlayerRuntime/Private/Runtime/Demuxer/ParserISO14496-12.h
@@ -127,9 +127,6 @@ namespace Electra
 
 		static TSharedPtrTS<IParserISO14496_12> CreateParser();
 
-		static const FString 	OptionKeyParseFirstFreeBox;			// Option to not skip 'free' box but instead attempt to parse their content
-
-
 		/**
 		 * Parses the header boxes (all non-MDAT boxes).
 		 */
diff --git a/Engine/Plugins/Media/ElectraPlayer/Source/ElectraPlayerRuntime/Private/Runtime/Player/HLS/ManifestBuilderHLS.cpp b/Engine/Plugins/Media/ElectraPlayer/Source/ElectraPlayerRuntime/Private/Runtime/Player/HLS/ManifestBuilderHLS.cpp
--- a/Engine/Plugins/Media/ElectraPlayer/Source/ElectraPlayerRuntime/Private/Runtime/Player/HLS/ManifestBuilderHLS.cpp
+++ b/Engine/Plugins/Media/ElectraPlayer/Source/ElectraPlayerRuntime/Private/Runtime/Player/HLS/ManifestBuilderHLS.cpp
@@ -579,6 +579,8 @@ FErrorDetail FManifestBuilderHLS::SetupRenditions(FManifestHLSInternal* Manifest
 		Rendition->Internal.UniqueID = FMediaInterlockedIncrement(NextUniqueID) + 1;
 
 		// Add to appropriate rendition group and to the lookup map.
+		// Note: Adding to a map/multimap breaks the order in which the representations appear in the master playlist.
+		//       If we need the order later we can use the rendition internal unique ID.
 		if (Type == TEXT("VIDEO"))
 		{
 			Rendition->Internal.bHasVideo = true;
@@ -608,7 +610,7 @@ FErrorDetail FManifestBuilderHLS::SetupRenditions(FManifestHLSInternal* Manifest
 
 	// Create adaptation sets for the audio rendition groups
 // FIXME: We will need to do this for the other rendition types (VIDEO, SUBTITLES and CLOSED-CAPTIONS as well)
-	for (TMultiMap<FString, TSharedPtrTS<FManifestHLSInternal::FRendition>>::TConstIterator It = Manifest->AudioRenditions.CreateConstIterator(); It; ++It)
+	for(TMultiMap<FString, TSharedPtrTS<FManifestHLSInternal::FRendition>>::TConstIterator It = Manifest->AudioRenditions.CreateConstIterator(); It; ++It)
 	{
 		HLSTimeline::FPlaybackAssetAdaptationSet* Adapt = static_cast<HLSTimeline::FPlaybackAssetAdaptationSet*>(Asset->GetAdaptationSetByTypeAndUniqueIdentifier(EStreamType::Audio, It.Key()).Get());
 		if (!Adapt)
@@ -1055,20 +1057,19 @@ FErrorDetail FManifestBuilderHLS::SetupVariants(FManifestHLSInternal* Manifest,
 
 						// Get the range of renditions in this group and find those that do not have a dedicated playlist.
 						// We let those point back to this variant stream since that is what we effectively have to use.
-						TArray< const TSharedPtrTS<FManifestHLSInternal::FRendition>* > RenditionRange;
+						TArray<TSharedPtrTS<FManifestHLSInternal::FRendition>*> RenditionRange;
 						Manifest->AudioRenditions.MultiFindPointer(vs->AudioGroupID, RenditionRange);
 						if (!RenditionRange.IsEmpty())
 						{
 							for(int32 ii=0; ii < RenditionRange.Num(); ++ii)
 							{
-								const TSharedPtrTS<FManifestHLSInternal::FRendition>& Rendition = *RenditionRange[ii];
+								TSharedPtrTS<FManifestHLSInternal::FRendition>& Rendition = *RenditionRange[ii];
 								if (Rendition->URI.Len() == 0)
 								{
 									// When the rendition has no dedicated URL then it is merely informational and this audio-only variant is the stream itself to use.
-									Repr->UniqueIdentifier = LexToString(StreamMetaData.StreamUniqueID);
-									Repr->CodecInformation = StreamMetaData.CodecInformation;
-									Repr->CDN   		   = StreamMetaData.PlaylistID;
-									Repr->Bitrate   	   = StreamMetaData.Bandwidth;
+									Rendition->URI = vs->GetURL();
+									Repr->CDN = StreamMetaData.PlaylistID;
+									Repr->Bitrate = StreamMetaData.Bandwidth;
 									vs->Internal.AdaptationSetUniqueID  = Adapt->UniqueIdentifier;
 									vs->Internal.RepresentationUniqueID = Repr->UniqueIdentifier;
 									vs->Internal.CDN					= StreamMetaData.PlaylistID;
@@ -1083,9 +1084,77 @@ FErrorDetail FManifestBuilderHLS::SetupVariants(FManifestHLSInternal* Manifest,
 		}
 	}
 
+
+	// Create fake audio-only variant streams from audio renditions. These are what is exposed as audio tracks that can be chosen from.
+	// Do this only when there are any variant streams. If there are only renditions for whatever reason without any reference to them
+	// do not do this.
+	bool bCreateFakeAudioVariants = Manifest->VariantStreams.Num() || Manifest->AudioOnlyStreams.Num();
+	if (bCreateFakeAudioVariants)
+	{
+		// Remove audio only streams that are already specified but reference a rendition group to avoid duplication.
+		for(int32 i=Manifest->AudioOnlyStreams.Num()-1; i>=0; --i)
+		{
+			if (!Manifest->AudioOnlyStreams[i]->AudioGroupID.IsEmpty())
+			{
+				Manifest->PlaylistIDMap.Remove(Manifest->AudioOnlyStreams[i]->Internal.UniqueID);
+				Manifest->AudioOnlyStreams.RemoveAt(i);
+				Manifest->StreamMetadataAudio.RemoveAt(i);
+			}
+		}
+		// Set up variants from the renditions in the audio adaptation sets.
+		for(int32 i=0; i<Asset->AudioAdaptationSets.Num(); ++i)
+		{
+			const HLSTimeline::FPlaybackAssetAdaptationSet* Adapt = static_cast<const HLSTimeline::FPlaybackAssetAdaptationSet*>(Asset->GetAdaptationSetByTypeAndIndex(EStreamType::Audio, i).Get());
+			for(int32 j=0; j<Adapt->GetNumberOfRepresentations(); ++j)
+			{
+				const HLSTimeline::FPlaybackAssetRepresentation* Repr = static_cast<const HLSTimeline::FPlaybackAssetRepresentation*>(Adapt->GetRepresentationByIndex(j).Get());
+				int32 ReprID = 0;
+				LexFromString(ReprID, *Repr->GetUniqueIdentifier());
+				for(TMultiMap<FString, TSharedPtrTS<FManifestHLSInternal::FRendition>>::TConstIterator It = Manifest->AudioRenditions.CreateConstIterator(); It; ++It)
+				{
+					TSharedPtrTS<FManifestHLSInternal::FRendition> Rendition = It.Value();
+					if (Rendition->Internal.UniqueID == ReprID)
+					{
+						TSharedPtrTS<FManifestHLSInternal::FVariantStream> vs = MakeSharedTS<FManifestHLSInternal::FVariantStream>();
+						vs->URI = Rendition->GetURL();
+						vs->AudioGroupID = Adapt->GetUniqueIdentifier();
+						vs->StreamCodecInformationList.Push(Repr->GetCodecInformation());
+						vs->Internal.bHasAudio = true;
+						vs->Internal.UniqueID = FMediaInterlockedIncrement(NextUniqueID) + 1;
+
+						FStreamMetadata StreamMetaData;
+						StreamMetaData.Bandwidth	  = 128000;
+						StreamMetaData.StreamUniqueID = vs->Internal.UniqueID;
+						StreamMetaData.PlaylistID     = UrlBuilder->ResolveWith(vs->URI);
+						StreamMetaData.LanguageCode   = Rendition->Language;
+						for(int32 k=0; k<vs->StreamCodecInformationList.Num(); ++k)
+						{
+							if (vs->StreamCodecInformationList[k].IsAudioCodec())
+							{
+								StreamMetaData.CodecInformation = vs->StreamCodecInformationList[k];
+								break;
+							}
+						}
+						Manifest->StreamMetadataAudio.Push(StreamMetaData);
+
+						vs->Internal.AdaptationSetUniqueID  = Adapt->UniqueIdentifier;
+						vs->Internal.RepresentationUniqueID = Repr->UniqueIdentifier;
+						vs->Internal.CDN					= StreamMetaData.PlaylistID;
+						Manifest->AudioOnlyStreams.Push(vs);
+
+						Manifest->PlaylistIDMap.Add(vs->Internal.UniqueID, vs);
+
+						break;
+					}
+				}
+			}
+		}
+	}
+
+
 	// Index the stream quality level map. Lower quality = lower index
 	int32 QualityIndex = 0;
-	for (TMap<int32, int32>::TIterator It = Manifest->BandwidthToQualityIndex.CreateIterator(); It; ++It)
+	for(TMap<int32, int32>::TIterator It = Manifest->BandwidthToQualityIndex.CreateIterator(); It; ++It)
 	{
 		It.Value() = QualityIndex++;
 	}
@@ -1314,7 +1383,7 @@ UEMediaError FManifestBuilderHLS::UpdateFailedInitialPlaylistLoadRequest(FPlayli
 	}
 
 	// Audio rendition?
-	for (TMultiMap<FString, TSharedPtrTS<FManifestHLSInternal::FRendition>>::TConstIterator It = Manifest->AudioRenditions.CreateConstIterator(); It; ++It)
+	for(TMultiMap<FString, TSharedPtrTS<FManifestHLSInternal::FRendition>>::TConstIterator It = Manifest->AudioRenditions.CreateConstIterator(); It; ++It)
 	{
 		TSharedPtrTS<FManifestHLSInternal::FRendition> AudioRendition = It.Value();
 		if (AudioRendition->Internal.UniqueID == InOutFailedRequest.InternalUniqueID)
@@ -1360,7 +1429,7 @@ void FManifestBuilderHLS::SetVariantPlaylistFailure(TSharedPtrTS<FManifestHLSInt
 	InHLSPlaylist->LockPlaylists();
 
 	TWeakPtrTS<FManifestHLSInternal::FPlaylistBase>* PlaylistID = InHLSPlaylist->PlaylistIDMap.Find(SourceRequest.InternalUniqueID);
-	if(PlaylistID != nullptr)
+	if (PlaylistID != nullptr)
 	{
 		TSharedPtrTS<FManifestHLSInternal::FPlaylistBase> Playlist = PlaylistID->Pin();
 		if (Playlist.IsValid())
@@ -1402,7 +1471,7 @@ FErrorDetail FManifestBuilderHLS::UpdateFromVariantPlaylist(TSharedPtrTS<FManife
 	TSharedPtrTS<FManifestHLSInternal::FPlaylistBase> Playlist;
 	InOutHLSPlaylist->LockPlaylists();
 	TWeakPtrTS<FManifestHLSInternal::FPlaylistBase>* PlaylistID = InOutHLSPlaylist->PlaylistIDMap.Find(SourceRequest.InternalUniqueID);
-	if(PlaylistID != nullptr)
+	if (PlaylistID != nullptr)
 	{
 		Playlist = PlaylistID->Pin();
 	}
@@ -1965,7 +2034,7 @@ void FManifestBuilderHLS::UpdateManifestMetadataFromStream(FManifestHLSInternal*
 				StreamIDsToConsider.Add(Manifest->VariantStreams[i]->Internal.UniqueID);
 			}
 		}
-		for (TMultiMap<FString, TSharedPtrTS<FManifestHLSInternal::FRendition>>::TConstIterator It = Manifest->AudioRenditions.CreateConstIterator(); It; ++It)
+		for(TMultiMap<FString, TSharedPtrTS<FManifestHLSInternal::FRendition>>::TConstIterator It = Manifest->AudioRenditions.CreateConstIterator(); It; ++It)
 		{
 			if (It.Value()->Internal.LoadState == FManifestHLSInternal::FPlaylistBase::FInternal::ELoadState::Loaded)
 			{
diff --git a/Engine/Plugins/Media/ElectraPlayer/Source/ElectraPlayerRuntime/Private/Runtime/Player/HLS/Parser.cpp b/Engine/Plugins/Media/ElectraPlayer/Source/ElectraPlayerRuntime/Private/Runtime/Player/HLS/Parser.cpp
--- a/Engine/Plugins/Media/ElectraPlayer/Source/ElectraPlayerRuntime/Private/Runtime/Player/HLS/Parser.cpp
+++ b/Engine/Plugins/Media/ElectraPlayer/Source/ElectraPlayerRuntime/Private/Runtime/Player/HLS/Parser.cpp
@@ -228,12 +228,10 @@ namespace Electra
             {
                 Error = ParseToken(ContentIterator, OutValue, TAttributeSeparation | TLineBreak | TEOF);
             }
-			// Note: cast required to silence static code analyzer (isdigit() expects an int, apparently need to cast char to unsigned first to prevent sign extension)
             else if (FChar::IsDigit(*ContentIterator) || *ContentIterator == TCHAR('-'))
             {
                 Error = ParseToken(ContentIterator, OutValue, TAttributeSeparation | TLineBreak | TEOF);
             }
-			// Note: cast required to silence static code analyzer (isalpha() expects an int, apparently need to cast char to unsigned first to prevent sign extension)
             else if (FChar::IsAlpha(*ContentIterator))
             {
                 Error = ParseToken(ContentIterator, OutValue, TAttributeSeparation | TLineBreak | TEOF);
@@ -250,6 +248,12 @@ namespace Electra
         {
             while (ContentIterator && *ContentIterator != CLineBreak && *ContentIterator != TComment)
             {
+				// Swallow spaces that may appear before the separating ,
+				while (*ContentIterator && *ContentIterator == CSpace)
+				{
+                    ++ContentIterator;
+				}
+
                 if (*ContentIterator == CAttributeSeparation)
                 {
                     ++ContentIterator;
@@ -259,6 +263,12 @@ namespace Electra
                     }
                 }
 
+				// And also swallow spaces that may appear after the separating ,
+				while (*ContentIterator && *ContentIterator == CSpace)
+				{
+                    ++ContentIterator;
+				}
+
                 FString AttributeKey;
                 EPlaylistError Error = ParseToken(ContentIterator, AttributeKey, TAttributeDeclaration);
                 if (Error != EPlaylistError::None)
@@ -363,7 +373,7 @@ namespace Electra
                         break;
                     }
                     bFirstLine = false;
-                    // Skip the hole line.
+                    // Skip the whole line.
                     Error = SkipToken(ContentIterator, TLineBreak | TEOF | TIgnore);
                     if (Error != EPlaylistError::None)
                     {
diff --git a/Engine/Plugins/Media/ElectraPlayer/Source/ElectraPlayerRuntime/Private/Runtime/Player/HLS/Token.h b/Engine/Plugins/Media/ElectraPlayer/Source/ElectraPlayerRuntime/Private/Runtime/Player/HLS/Token.h
--- a/Engine/Plugins/Media/ElectraPlayer/Source/ElectraPlayerRuntime/Private/Runtime/Player/HLS/Token.h
+++ b/Engine/Plugins/Media/ElectraPlayer/Source/ElectraPlayerRuntime/Private/Runtime/Player/HLS/Token.h
@@ -10,6 +10,7 @@ namespace Electra
         const TCHAR CColumn = TCHAR(':');
         const TCHAR CLineBreak = TCHAR('\n');
         const TCHAR CAttributeSeparation = TCHAR(',');
+        const TCHAR CSpace = TCHAR(' ');
 
         typedef unsigned int Token;
 
