Electra: ignoring provided thread priorities and affinities since they are all set to normal anyway. Also ignoring the stack size since a few threads are setting them too small for a check() or ensureAlwaysMsgf() to be able to generate the log message without crashing. Using defaults instead.

#rnx

diff --git a/Engine/Plugins/Media/ElectraPlayer/Source/ElectraPlayerRuntime/Private/Runtime/Core/MediaThreads.cpp b/Engine/Plugins/Media/ElectraPlayer/Source/ElectraPlayerRuntime/Private/Runtime/Core/MediaThreads.cpp
--- a/Engine/Plugins/Media/ElectraPlayer/Source/ElectraPlayerRuntime/Private/Runtime/Core/MediaThreads.cpp
+++ b/Engine/Plugins/Media/ElectraPlayer/Source/ElectraPlayerRuntime/Private/Runtime/Core/MediaThreads.cpp
@@ -84,7 +84,9 @@ void FMediaRunnable::SetDoneSignal(FMediaEvent* InDoneSignal)
  */
 void FMediaRunnable::StartInternal()
 {
-	MediaThreadRunnable = FRunnableThread::Create(this, *ThreadName, StackSize, ThreadPriority, InitialCoreAffinity);
+	// We ignore stack size, priority and affinity here on purpose. Prio and affinity are all set to defaults and stack sizes are
+	// often set too small so we rather go with defaults.
+	MediaThreadRunnable = FRunnableThread::Create(this, *ThreadName);
 	checkf(MediaThreadRunnable, TEXT("Could not create RunnableThread!"));
 	if (MediaThreadRunnable)
 	{
