Electra: removed an uncritical assert when more than one video track is in the file

#rnx

diff --git a/Engine/Plugins/Media/ElectraPlayer/Source/ElectraPlayerRuntime/Private/Runtime/Player/AdaptiveStreamingPlayer.cpp b/Engine/Plugins/Media/ElectraPlayer/Source/ElectraPlayerRuntime/Private/Runtime/Player/AdaptiveStreamingPlayer.cpp
--- a/Engine/Plugins/Media/ElectraPlayer/Source/ElectraPlayerRuntime/Private/Runtime/Player/AdaptiveStreamingPlayer.cpp
+++ b/Engine/Plugins/Media/ElectraPlayer/Source/ElectraPlayerRuntime/Private/Runtime/Player/AdaptiveStreamingPlayer.cpp
@@ -2414,9 +2414,9 @@ bool FAdaptiveStreamingPlayer::FindMatchingStreamInfo(FStreamCodecInformation& O
 	TSharedPtrTS<ITimelineMediaAsset> Asset = CurrentPlayPeriod->GetMediaAsset();
 	if (Asset.IsValid())
 	{
-		check(Asset->GetNumberOfAdaptationSets(EStreamType::Video) <= 1); // what if this is more than one?
 		if (Asset->GetNumberOfAdaptationSets(EStreamType::Video) > 0)
 		{
+			// What if this is more than one?
 			TSharedPtrTS<IPlaybackAssetAdaptationSet> VideoSet = Asset->GetAdaptationSetByTypeAndIndex(EStreamType::Video, 0);
 			check(VideoSet.IsValid());
 			if (VideoSet.IsValid())
