Electra: added parsing of scaling matrices to the h.264 SPS parser

#rnx

diff --git a/Engine/Plugins/Media/ElectraPlayer/Source/ElectraPlayerRuntime/Private/Runtime/Utilities/UtilsMPEGVideo.cpp b/Engine/Plugins/Media/ElectraPlayer/Source/ElectraPlayerRuntime/Private/Runtime/Utilities/UtilsMPEGVideo.cpp
--- a/Engine/Plugins/Media/ElectraPlayer/Source/ElectraPlayerRuntime/Private/Runtime/Utilities/UtilsMPEGVideo.cpp
+++ b/Engine/Plugins/Media/ElectraPlayer/Source/ElectraPlayerRuntime/Private/Runtime/Utilities/UtilsMPEGVideo.cpp
@@ -107,7 +107,7 @@ namespace Electra
 					ParsedSPSs.Reserve(nSPS);
 				}
 				int32 TotalSPSSize = 0;
-				for (int32 i = 0; i < nSPS; ++i)
+				for(int32 i = 0; i < nSPS; ++i)
 				{
 					uint16 spsLen = 0;
 					if (!BitReader.ReadByte(spsLen))
@@ -134,7 +134,7 @@ namespace Electra
 					PictureParameterSets.Reserve(nPPS);
 				}
 				int32 TotalPPSSize = 0;
-				for (int32 i = 0; i < nPPS; ++i)
+				for(int32 i = 0; i < nPPS; ++i)
 				{
 					uint16 ppsLen = 0;
 					if (!BitReader.ReadByte(ppsLen))
@@ -181,7 +181,7 @@ namespace Electra
 						{
 							SequenceParameterSetsExt.Reserve(nSPS);
 						}
-						for (int32 i = 0; i < nSPS; ++i)
+						for(int32 i = 0; i < nSPS; ++i)
 						{
 							uint16 spsLenExt = 0;
 							if (!BitReader.ReadByte(spsLenExt))
@@ -205,7 +205,7 @@ namespace Electra
 				{
 					CodecSpecificData.Reserve(TotalCSDSize);
 					CodecSpecificDataSPSOnly.Reserve(TotalSPSSize);
-					for (int32 i = 0; i < SequenceParameterSets.Num(); ++i)
+					for(int32 i = 0; i < SequenceParameterSets.Num(); ++i)
 					{
 						CodecSpecificDataSPSOnly.Push(0);
 						CodecSpecificDataSPSOnly.Push(0);
@@ -215,14 +215,14 @@ namespace Electra
 						CodecSpecificData.Push(0);
 						CodecSpecificData.Push(0);
 						CodecSpecificData.Push(1);
-						for (int32 j = 0, jMax = SequenceParameterSets[i].Num(); j < jMax; ++j)
+						for(int32 j = 0, jMax = SequenceParameterSets[i].Num(); j < jMax; ++j)
 						{
 							CodecSpecificDataSPSOnly.Push((SequenceParameterSets[i].GetData())[j]);
 							CodecSpecificData.Push((SequenceParameterSets[i].GetData())[j]);
 						}
 					}
 					CodecSpecificDataPPSOnly.Reserve(TotalPPSSize);
-					for (int32 i = 0; i < PictureParameterSets.Num(); ++i)
+					for(int32 i = 0; i < PictureParameterSets.Num(); ++i)
 					{
 						CodecSpecificDataPPSOnly.Push(0);
 						CodecSpecificDataPPSOnly.Push(0);
@@ -232,7 +232,7 @@ namespace Electra
 						CodecSpecificData.Push(0);
 						CodecSpecificData.Push(0);
 						CodecSpecificData.Push(1);
-						for (int32 j = 0, jMax = PictureParameterSets[i].Num(); j < jMax; ++j)
+						for(int32 j = 0, jMax = PictureParameterSets[i].Num(); j < jMax; ++j)
 						{
 							CodecSpecificDataPPSOnly.Push((PictureParameterSets[i].GetData())[j]);
 							CodecSpecificData.Push((PictureParameterSets[i].GetData())[j]);
@@ -250,7 +250,7 @@ namespace Electra
 
 		static int32 FindStartCode(const uint8* InData, SIZE_T InDataSize, int32& NALUnitLength)
 		{
-			for (const uint8* Data = InData; InDataSize >= 3; ++Data, --InDataSize)
+			for(const uint8* Data = InData; InDataSize >= 3; ++Data, --InDataSize)
 			{
 				if (Data[0] == 0 && Data[1] == 0 && (Data[2] == 1 || (InDataSize >= 4 && Data[2] == 0 && Data[3] == 1)))
 				{
@@ -271,7 +271,7 @@ namespace Electra
 			uint64 Pos = 0;
 			uint64 BytesToGo = InBitstreamLength;
 			const uint8* BitstreamData = (const uint8*)InBitstream;
-			while (1)
+			while(1)
 			{
 				int32 UnitLength = 0, StartCodePos = FindStartCode(BitstreamData, BytesToGo, UnitLength);
 				if (StartCodePos >= 0)
@@ -309,7 +309,7 @@ namespace Electra
 		int32 EBSPtoRBSP(uint8* OutBuf, const uint8* InBuf, int32 NumBytesIn)
 		{
 			uint8* OutBase = OutBuf;
-			while (NumBytesIn-- > 0)
+			while(NumBytesIn-- > 0)
 			{
 				uint8 b = *InBuf++;
 				*OutBuf++ = b;
@@ -359,7 +359,7 @@ namespace Electra
 				static uint32 ue_v(FBitDataStream& BitStream)
 				{
 					int32 lz = -1;
-					for (uint32 b = 0; b == 0; ++lz)
+					for(uint32 b = 0; b == 0; ++lz)
 					{
 						b = BitStream.GetBits(1);
 					}
@@ -404,15 +404,50 @@ namespace Electra
 			{
 				OutSPS.chroma_format_idc = FSyntaxElement::ue_v(BitReader);
 				if (OutSPS.chroma_format_idc == 3)
+				{
 					OutSPS.separate_colour_plane_flag = (uint8)BitReader.GetBits(1);
+				}
 				OutSPS.bit_depth_luma_minus8 = FSyntaxElement::ue_v(BitReader);
 				OutSPS.bit_depth_chroma_minus8 = FSyntaxElement::ue_v(BitReader);
 				OutSPS.qpprime_y_zero_transform_bypass_flag = (uint8)BitReader.GetBits(1);
 				OutSPS.seq_scaling_matrix_present_flag = (uint8)BitReader.GetBits(1);
 				if (OutSPS.seq_scaling_matrix_present_flag)
 				{
-					check(!"oh no!");
-					return false;
+					auto scaling_list = [&BitReader](int32* scalingList, int32 sizeOfScalingList, bool& useDefaultScalingMatrixFlag) -> void
+					{
+						int32 lastScale = 8;
+						int32 nextScale = 8;
+						for(int32 j=0; j<sizeOfScalingList; ++j)
+						{
+							if (nextScale)
+							{
+								int32 delta_scale = FSyntaxElement::se_v(BitReader);
+								nextScale = (lastScale + delta_scale + 256) % 256;
+								useDefaultScalingMatrixFlag = (j == 0 && nextScale == 0);
+							}
+							scalingList[j] = (nextScale == 0) ? lastScale : nextScale;
+							lastScale = scalingList[j];
+						}
+					};
+
+					// Skip over the scaling matrices.
+					int32 dummyScalingMatrix[64] = {0};
+					bool bDummyDefaultScalingMatrixFlag = false;
+					for(int32 i=0, iMax=OutSPS.chroma_format_idc!=3?8:12; i<iMax; ++i)
+					{
+						uint8 seq_scaling_list_present_flag = (uint8)BitReader.GetBits(1);
+						if (seq_scaling_list_present_flag)
+						{
+							if (i < 6)
+							{
+								scaling_list(dummyScalingMatrix, 16, bDummyDefaultScalingMatrixFlag);
+							}
+							else
+							{
+								scaling_list(dummyScalingMatrix, 64, bDummyDefaultScalingMatrixFlag);
+							}
+						}
+					}
 				}
 			}
 			OutSPS.log2_max_frame_num_minus4 = FSyntaxElement::ue_v(BitReader);
@@ -427,7 +462,7 @@ namespace Electra
 				OutSPS.offset_for_non_ref_pic = FSyntaxElement::se_v(BitReader);
 				OutSPS.offset_for_top_to_bottom_field = FSyntaxElement::se_v(BitReader);
 				OutSPS.num_ref_frames_in_pic_order_cnt_cycle = FSyntaxElement::ue_v(BitReader);
-				for (uint32 i = 0; i < OutSPS.num_ref_frames_in_pic_order_cnt_cycle; ++i)
+				for(uint32 i = 0; i < OutSPS.num_ref_frames_in_pic_order_cnt_cycle; ++i)
 				{
 					FSyntaxElement::se_v(BitReader);		// discard
 				}
