Electra: added flag to look for sync samples when setting up the track iterator and added a simpler convenience method to start iterating at the first sample. This is needed for chunked fmp4 segments where each sample is carried in its own moof+mdat box and the sample is not necessarily a sync sample.

#rnx

diff --git a/Engine/Plugins/Media/ElectraPlayer/Source/ElectraPlayerRuntime/Private/Runtime/Demuxer/ParserISO14496-12.cpp b/Engine/Plugins/Media/ElectraPlayer/Source/ElectraPlayerRuntime/Private/Runtime/Demuxer/ParserISO14496-12.cpp
--- a/Engine/Plugins/Media/ElectraPlayer/Source/ElectraPlayerRuntime/Private/Runtime/Demuxer/ParserISO14496-12.cpp
+++ b/Engine/Plugins/Media/ElectraPlayer/Source/ElectraPlayerRuntime/Private/Runtime/Demuxer/ParserISO14496-12.cpp
@@ -96,7 +96,7 @@ namespace Electra
 				return UEMEDIA_ERROR_OK;
 			}
 			uint8 FirstByte = 0;
-			for (uint32 i = 0; i < MaxBytes; ++i)
+			for(uint32 i = 0; i < MaxBytes; ++i)
 			{
 				uint8 NextChar;
 				UEMediaError Error = Read(NextChar);
@@ -330,7 +330,7 @@ namespace Electra
 
 		virtual ~FMP4Box()
 		{
-			for (int32 i = 0, iMax = ChildBoxes.Num(); i < iMax; ++i)
+			for(int32 i = 0, iMax = ChildBoxes.Num(); i < iMax; ++i)
 			{
 				delete ChildBoxes[i];
 			}
@@ -386,7 +386,7 @@ namespace Electra
 		const FMP4Box* FindBox(IParserISO14496_12::FBoxType BoxTypeIn, int32 MaxDepth = 32) const
 		{
 			// First sweep, check children
-			for (int32 i = 0, iMax = ChildBoxes.Num(); i < iMax; ++i)
+			for(int32 i = 0, iMax = ChildBoxes.Num(); i < iMax; ++i)
 			{
 				if (ChildBoxes[i]->GetType() == BoxTypeIn)
 				{
@@ -396,7 +396,7 @@ namespace Electra
 			// Second sweep, check children recursively
 			if (MaxDepth > 0)
 			{
-				for (int32 i = 0, iMax = ChildBoxes.Num(); i < iMax; ++i)
+				for(int32 i = 0, iMax = ChildBoxes.Num(); i < iMax; ++i)
 				{
 					const FMP4Box* Box = ChildBoxes[i]->FindBox(BoxTypeIn, MaxDepth - 1);
 					if (Box)
@@ -411,7 +411,7 @@ namespace Electra
 		void GetAllBoxInstances(TArray<const FMP4Box*>& OutAllBoxes, IParserISO14496_12::FBoxType BoxTypeIn) const
 		{
 			// First sweep, check children
-			for (int32 i = 0, iMax = ChildBoxes.Num(); i < iMax; ++i)
+			for(int32 i = 0, iMax = ChildBoxes.Num(); i < iMax; ++i)
 			{
 				if (ChildBoxes[i]->GetType() == BoxTypeIn)
 				{
@@ -474,7 +474,7 @@ namespace Electra
 
 		virtual UEMediaError ReadAndParseAttributes(FMP4ParseInfo* ParseInfo) = 0;
 
-		// This whole constexpr thing doesn't work...
+		// constexpr doesn't work here
 	//	#define MAKE_BOX_ATOM(a,b,c,d) (IParserISO14496_12::FBoxType)((uint32)a << 24) | ((uint32)b << 16) | ((uint32)c << 8) | ((uint32)d)
 
 		/**
@@ -579,7 +579,7 @@ namespace Electra
 			int64	StartOffset;
 		};
 
-		FMP4Box* ParentBox;
+		FMP4Box*						ParentBox;
 		TArray<FMP4Box*>				ChildBoxes;
 		IParserISO14496_12::FBoxType	BoxType;
 		int64							BoxSize;
@@ -721,7 +721,7 @@ namespace Electra
 				if (NumCompatibleBrands)
 				{
 					CompatibleBrands.Reserve(NumCompatibleBrands);
-					for (uint32 i = 0; i < NumCompatibleBrands; ++i)
+					for(uint32 i = 0; i < NumCompatibleBrands; ++i)
 					{
 						uint32 CompatibleBrand = 0;
 						RETURN_IF_ERROR(ParseInfo->Reader()->Read(CompatibleBrand));
@@ -831,7 +831,7 @@ namespace Electra
 				Duration = (Value32 == kAllOnes32) ? kAllOnes64 : Value32;
 			}
 			// Read and ignore: rate, volume, reserved, reserved, matrix, pre_defined (old QuickTime values) and next_track_ID
-			for (int32 i = 0; i < 20; ++i)
+			for(int32 i = 0; i < 20; ++i)
 			{
 				RETURN_IF_ERROR(ParseInfo->Reader()->Read(Value32));
 			}
@@ -956,7 +956,7 @@ namespace Electra
 			RETURN_IF_ERROR(ParseInfo->Reader()->Read(Value16));				// volume
 			RETURN_IF_ERROR(ParseInfo->Reader()->Read(Value16));				// reserved
 			// Skip matrix
-			for (int32 i = 0; i < 9; ++i)
+			for(int32 i = 0; i < 9; ++i)
 			{
 				RETURN_IF_ERROR(ParseInfo->Reader()->Read(Value32));
 			}
@@ -1075,7 +1075,7 @@ namespace Electra
 
 		const FMP4BoxTREX* FindTREXForTrackID(uint32 TrackID) const
 		{
-			for (int32 i = 0, iMax = GetNumberOfChildren(); i < iMax; ++i)
+			for(int32 i = 0, iMax = GetNumberOfChildren(); i < iMax; ++i)
 			{
 				const FMP4Box* Box = GetChildBox(i);
 				if (Box->GetType() == FMP4Box::kBox_trex)
@@ -1280,7 +1280,7 @@ namespace Electra
 				int64	MediaTime = 0;
 				int16	MediaRateInteger = 0;
 				int16	MediaRateFraction = 0;
-				for (uint32 i = 0; i < NumEntries; ++i)
+				for(uint32 i = 0; i < NumEntries; ++i)
 				{
 					RETURN_IF_ERROR(ParseInfo->Reader()->Read(SegmentDuration));	// segment_duration
 					RETURN_IF_ERROR(ParseInfo->Reader()->Read(MediaTime));			// media_time
@@ -1299,7 +1299,7 @@ namespace Electra
 				int32	MediaTime = 0;
 				int16	MediaRateInteger = 0;
 				int16	MediaRateFraction = 0;
-				for (uint32 i = 0; i < NumEntries; ++i)
+				for(uint32 i = 0; i < NumEntries; ++i)
 				{
 					RETURN_IF_ERROR(ParseInfo->Reader()->Read(SegmentDuration));	// segment_duration
 					RETURN_IF_ERROR(ParseInfo->Reader()->Read(MediaTime));			// media_time
@@ -1458,7 +1458,7 @@ namespace Electra
 			UEMediaError Error = UEMEDIA_ERROR_OK;
 			RETURN_IF_ERROR(FMP4BoxFull::ReadAndParseAttributes(ParseInfo));
 			RETURN_IF_ERROR(ParseInfo->Reader()->Read(NumEntries));					// entry_count
-			for (uint32 i = 0; i < NumEntries; ++i)
+			for(uint32 i = 0; i < NumEntries; ++i)
 			{
 				RETURN_IF_ERROR(ParseInfo->ReadAndParseNextBox(this));
 			}
@@ -2047,7 +2047,7 @@ namespace Electra
 			// Check for media handler first.
 			if (ParseInfo->GetCurrentMediaHandlerBox())
 			{
-				switch (ParseInfo->GetCurrentMediaHandlerBox()->GetType())
+				switch(ParseInfo->GetCurrentMediaHandlerBox()->GetType())
 				{
 					case FMP4Box::kBox_vmhd:
 						ExpectedSampleType = EExpectedSampleType::ESTVideo;
@@ -2066,7 +2066,7 @@ namespace Electra
 			else
 			{
 				const FMP4BoxHDLR* CurrentHandler = static_cast<FMP4BoxHDLR*>(ParseInfo->GetCurrentHandlerBox());
-				switch (CurrentHandler->GetHandlerType())
+				switch(CurrentHandler->GetHandlerType())
 				{
 					case FMP4Box::kTrack_vide:
 						ExpectedSampleType = EExpectedSampleType::ESTVideo;
@@ -2083,7 +2083,7 @@ namespace Electra
 			// Check if this is a supported sample format.
 			bIsSupportedFormat = ExpectedSampleType != EExpectedSampleType::ESTIgnored;
 
-			for (uint32 i = 0; i < EntryCount; ++i)
+			for(uint32 i = 0; i < EntryCount; ++i)
 			{
 				IParserISO14496_12::FBoxType	ChildBoxType;
 				int64							ChildBoxSize;
@@ -2097,7 +2097,7 @@ namespace Electra
 #endif
 				int64 BoxDataOffset = ParseInfo->Reader()->GetCurrentReadOffset();
 				FMP4Box* NextBox = nullptr;
-				switch (ExpectedSampleType)
+				switch(ExpectedSampleType)
 				{
 					case EExpectedSampleType::ESTVideo:
 						NextBox = new FMP4BoxVisualSampleEntry(ChildBoxType, ChildBoxSize, BoxStartOffset, BoxDataOffset, true);
@@ -2210,7 +2210,7 @@ namespace Electra
 				//        in media timescale units.
 				Entries.Reserve(NumEntries);
 				uint32 c = 0, d = 0;
-				for (uint32 i = 0; i < NumEntries; ++i)
+				for(uint32 i = 0; i < NumEntries; ++i)
 				{
 					RETURN_IF_ERROR(ParseInfo->Reader()->Read(c));					// sample_count
 					RETURN_IF_ERROR(ParseInfo->Reader()->Read(d));					// sample_delta
@@ -2295,7 +2295,7 @@ namespace Electra
 				Entries.Reserve(NumEntries);
 				uint32 c = 0;
 				uint32 o = 0;
-				for (uint32 i = 0; i < NumEntries; ++i)
+				for(uint32 i = 0; i < NumEntries; ++i)
 				{
 					// NOTE: We always read the offset as unsigned as if the box is always version 0.
 					RETURN_IF_ERROR(ParseInfo->Reader()->Read(c));					// sample_count
@@ -2356,7 +2356,7 @@ namespace Electra
 			{
 				Entries.Reserve(NumEntries);
 				uint32 n = 0;
-				for (uint32 i = 0; i < NumEntries; ++i)
+				for(uint32 i = 0; i < NumEntries; ++i)
 				{
 					RETURN_IF_ERROR(ParseInfo->Reader()->Read(n));					// sample_number
 					Entries.Push(n);
@@ -2425,7 +2425,7 @@ namespace Electra
 				//        the sample interleaving strategy of the muxer and the relative sample durations.
 				Entries.Reserve(NumEntries);
 				uint32 FirstChunk = 0, SamplesPerChunk = 0, DescriptionIndex = 0;
-				for (uint32 i = 0; i < NumEntries; ++i)
+				for(uint32 i = 0; i < NumEntries; ++i)
 				{
 					RETURN_IF_ERROR(ParseInfo->Reader()->Read(FirstChunk));			// first_chunk
 					RETURN_IF_ERROR(ParseInfo->Reader()->Read(SamplesPerChunk));	// samples_per_chunk
@@ -2488,7 +2488,7 @@ namespace Electra
 			{
 				Entries.Reserve(NumEntries);
 				uint32 n = 0;
-				for (uint32 i = 0; i < NumEntries; ++i)
+				for(uint32 i = 0; i < NumEntries; ++i)
 				{
 					RETURN_IF_ERROR(ParseInfo->Reader()->Read(n));					// entry_size
 					Entries.Push(n);
@@ -2547,7 +2547,7 @@ namespace Electra
 				{
 					Entries32.Reserve(NumEntries);
 					uint32 n = 0;
-					for (uint32 i = 0; i < NumEntries; ++i)
+					for(uint32 i = 0; i < NumEntries; ++i)
 					{
 						RETURN_IF_ERROR(ParseInfo->Reader()->Read(n));				// chunk_offset
 						Entries32.Push(n);
@@ -2557,7 +2557,7 @@ namespace Electra
 				{
 					Entries64.Reserve(NumEntries);
 					uint64 n = 0;
-					for (uint32 i = 0; i < NumEntries; ++i)
+					for(uint32 i = 0; i < NumEntries; ++i)
 					{
 						RETURN_IF_ERROR(ParseInfo->Reader()->Read(n));				// chunk_offset
 						Entries64.Push(n);
@@ -2635,7 +2635,7 @@ namespace Electra
 			if (ReferenceCount)
 			{
 				Entries.Reserve(ReferenceCount);
-				for (uint32 i = 0; i < ReferenceCount; ++i)
+				for(uint32 i = 0; i < ReferenceCount; ++i)
 				{
 					FEntry& e = Entries.AddDefaulted_GetRef();
 					e.ReferenceTypeAndSize = 0;
@@ -2971,7 +2971,7 @@ namespace Electra
 					SampleCompositionTimeOffsets.SetNumUninitialized(SampleCount);
 				}
 
-				for (uint32 i = 0; i < SampleCount; ++i)
+				for(uint32 i = 0; i < SampleCount; ++i)
 				{
 					if (Flags & 0x100)
 					{
@@ -2997,7 +2997,7 @@ namespace Electra
 							//        our SampleCompositionTimeOffsets to be an int64 table.
 							if (Value32 > 0x7fffffffU)
 							{
-								ensureAlwaysMsgf(0, TEXT("WARNING: ElectraPlayer/FMP4BoxTRUN::ReadAndParseAttributes(): Version 0 time value cannot be represented as a signed value"));
+								UE_LOG(LogElectraMP4Parser, Error, TEXT("WARNING: ElectraPlayer/FMP4BoxTRUN::ReadAndParseAttributes(): Version 0 time value cannot be represented as a signed value"));
 								return UEMEDIA_ERROR_FORMAT_ERROR;
 							}
 							SampleCompositionTimeOffsets[i] = (int32)Value32;
@@ -3109,7 +3109,7 @@ namespace Electra
 		}
 		else
 		{
-			while (1)
+			while(1)
 			{
 				// Check before reading anything if we reached the end of the file.
 				if (BoxReader->IsAtEOF())
@@ -3171,7 +3171,7 @@ namespace Electra
 				{
 					// Create a box of the appropriate type.
 					FMP4Box* NextBox = nullptr;
-					switch (BoxType)
+					switch(BoxType)
 					{
 						case FMP4Box::kBox_moov:
 							check(GetCurrentMoovBox() == nullptr);
@@ -3362,7 +3362,7 @@ namespace Electra
 								// We need to remember the handler for the continued parsing of the 'stsd' box, but only if this is a handler
 								// of a type we support.
 								FMP4BoxHDLR* HdlrBox = static_cast<FMP4BoxHDLR*>(NextBox);
-								switch (HdlrBox->GetHandlerType())
+								switch(HdlrBox->GetHandlerType())
 								{
 									case FMP4Box::kTrack_vide:
 									case FMP4Box::kTrack_soun:
@@ -3457,7 +3457,8 @@ namespace Electra
 			virtual int64 GetBaseMediaDecodeTime() const override;
 			virtual bool IsAtEOS() const override;
 
-			virtual UEMediaError StartAtTime(const FTimeValue& AtTime, ESearchMode SearchMode) override;
+			virtual UEMediaError StartAtTime(const FTimeValue& AtTime, ESearchMode SearchMode, bool bNeedSyncSample) override;
+			virtual UEMediaError StartAtFirst(bool bNeedSyncSample) override;
 			virtual UEMediaError Next() override;
 
 			virtual uint32 GetSampleNumber() const override;
@@ -3644,7 +3645,7 @@ namespace Electra
 
 			FTrack* GetTrackByID(uint32 TrackID) const
 			{
-				for (int32 i = 0; i < TrackList.Num(); ++i)
+				for(int32 i = 0; i < TrackList.Num(); ++i)
 				{
 					if (TrackList[i]->GetID() == TrackID)
 					{
@@ -3664,7 +3665,7 @@ namespace Electra
 			}
 
 		private:
-			TMediaOptionalValue<FTimeFraction>		MovieDuration;
+			TMediaOptionalValue<FTimeFraction>	MovieDuration;
 			TArray<FTrack*>						TrackList;
 		};
 
@@ -3693,7 +3694,7 @@ namespace Electra
 
 	void FParserISO14496_12::FTrackInfo::DeleteTrackList()
 	{
-		for (int32 i = 0, iMax = TrackList.Num(); i < iMax; ++i)
+		for(int32 i = 0, iMax = TrackList.Num(); i < iMax; ++i)
 		{
 			delete TrackList[i];
 		}
@@ -3775,7 +3776,7 @@ namespace Electra
 	const TArray<uint8>& FParserISO14496_12::FTrack::GetCodecSpecificData() const
 	{
 		static TArray<uint8> EmptyCSD;
-		switch (CodecInformation.GetCodec())
+		switch(CodecInformation.GetCodec())
 		{
 			case FStreamCodecInformation::ECodec::H264:
 			{
@@ -3796,7 +3797,7 @@ namespace Electra
 	const TArray<uint8>& FParserISO14496_12::FTrack::GetCodecSpecificDataRAW() const
 	{
 		static TArray<uint8> EmptyCSD;
-		switch (CodecInformation.GetCodec())
+		switch(CodecInformation.GetCodec())
 		{
 			case FStreamCodecInformation::ECodec::H264:
 			{
@@ -4150,7 +4151,7 @@ namespace Electra
 			if (RemainingSamplesInTRUN == 0)
 			{
 				// Because a run can be empty we go into a loop here to find the next run that is not empty.
-				while (1)
+				while(1)
 				{
 					// Another TRUN box to move into?
 					if (InTRUNIndex + 1 >= (uint32)Track->TRUNBoxes.Num())
@@ -4332,8 +4333,24 @@ namespace Electra
 		}
 	}
 
+	UEMediaError FParserISO14496_12::FTrackIterator::StartAtFirst(bool bNeedSyncSample)
+	{
+		UEMediaError err = StartAtFirst();
+		if (err == UEMEDIA_ERROR_OK)
+		{
+			for(; err == UEMEDIA_ERROR_OK; err = Next())
+			{
+				if (!bNeedSyncSample || (bNeedSyncSample && IsSyncSample()))
+				{
+					return UEMEDIA_ERROR_OK;
+				}
+			}
+			err = UEMEDIA_ERROR_INSUFFICIENT_DATA;
+		}
+		return err;
+	}
 
-	UEMediaError FParserISO14496_12::FTrackIterator::StartAtTime(const FTimeValue& AtTime, FParserISO14496_12::ITrackIterator::ESearchMode SearchMode)
+	UEMediaError FParserISO14496_12::FTrackIterator::StartAtTime(const FTimeValue& AtTime, FParserISO14496_12::ITrackIterator::ESearchMode SearchMode, bool bNeedSyncSample)
 	{
 		// Note: This code is perhaps not very optimal in that we're doing a linear search here. For unfragmented files we could look at the stss box
 		//       to locate the sync sample and get its dts/pts relatively quickly from the stts/ctts boxes.
@@ -4349,9 +4366,9 @@ namespace Electra
 			// Get the local media time in its timescale units.
 			int64 localTime = AtTime.GetAsTimebase(GetTimescale());
 
-			for (; err == UEMEDIA_ERROR_OK; err = Next())
+			for(; err == UEMEDIA_ERROR_OK; err = Next())
 			{
-				if (IsSyncSample())
+				if (!bNeedSyncSample || (bNeedSyncSample && IsSyncSample()))
 				{
 					if (!Best.IsValid())
 					{
@@ -4546,10 +4563,10 @@ namespace Electra
 			// Is this a fragmented stream?
 			const FMP4BoxMVEX* MVEXBox = static_cast<const FMP4BoxMVEX*>(MOOVBox->FindBox(FMP4Box::kBox_mvex, 0));
 
-			for (int32 i = 0, iMax = MOOVBox->GetNumberOfChildren(); i < iMax; ++i)
+			for(int32 i = 0, iMax = MOOVBox->GetNumberOfChildren(); i < iMax; ++i)
 			{
 				const FMP4Box* Box = MOOVBox->GetChildBox(i);
-				switch (Box->GetType())
+				switch(Box->GetType())
 				{
 					case FMP4Box::kBox_trak:
 					{
@@ -4597,10 +4614,10 @@ namespace Electra
 						}
 
 						// Get the sample description boxes from the 'stbl'
-						for (int32 j = 0, jMax = Track->STBLBox->GetNumberOfChildren(); j < jMax; ++j)
+						for(int32 j = 0, jMax = Track->STBLBox->GetNumberOfChildren(); j < jMax; ++j)
 						{
 							const FMP4Box* SampleInfoBox = Track->STBLBox->GetChildBox(j);
-							switch (SampleInfoBox->GetType())
+							switch(SampleInfoBox->GetType())
 							{
 								case FMP4Box::kBox_stts:
 									Track->STTSBox = static_cast<const FMP4BoxSTTS*>(SampleInfoBox);
@@ -4638,7 +4655,7 @@ namespace Electra
 						}
 
 						bool bIsSupported = true;
-						switch (Track->STSDBox->GetChildBox(0)->GetType())
+						switch(Track->STSDBox->GetChildBox(0)->GetType())
 						{
 							case FMP4Box::kSample_avc1:
 							{
@@ -4648,10 +4665,10 @@ namespace Electra
 								{
 									// There may be several entries. Usually there is an 'avcC' (required) and optional boxes like 'pasp', 'btrt', 'clap', etc.
 									bool bGotVideoFormat = false;
-									for (int32 j = 0, jMax = VisualSampleEntry->GetNumberOfChildren(); j < jMax; ++j)
+									for(int32 j = 0, jMax = VisualSampleEntry->GetNumberOfChildren(); j < jMax; ++j)
 									{
 										const FMP4Box* SampleEntry = VisualSampleEntry->GetChildBox(j);
-										switch (SampleEntry->GetType())
+										switch(SampleEntry->GetType())
 										{
 											case FMP4Box::kBox_avcC:
 											{
@@ -4735,10 +4752,10 @@ namespace Electra
 								{
 									// There may be several entries. Usually there is the required sample format box (eg 'mp4a') and optional boxes like 'btrt'
 									bool bGotAudioFormat = false;
-									for (int32 j = 0, jMax = AudioSampleEntry->GetNumberOfChildren(); j < jMax; ++j)
+									for(int32 j = 0, jMax = AudioSampleEntry->GetNumberOfChildren(); j < jMax; ++j)
 									{
 										const FMP4Box* SampleEntry = AudioSampleEntry->GetChildBox(j);
-										switch (SampleEntry->GetType())
+										switch(SampleEntry->GetType())
 										{
 											case FMP4Box::kBox_esds:
 											{
@@ -4857,10 +4874,10 @@ namespace Electra
 		const FMP4Box* MOOFBox = ParsedData->GetCurrentMoofBox();
 		if (MOOFBox)
 		{
-			for (int32 i = 0, iMax = MOOFBox->GetNumberOfChildren(); i < iMax; ++i)
+			for(int32 i = 0, iMax = MOOFBox->GetNumberOfChildren(); i < iMax; ++i)
 			{
 				const FMP4Box* Box = MOOFBox->GetChildBox(i);
-				switch (Box->GetType())
+				switch(Box->GetType())
 				{
 					case FMP4Box::kBox_traf:
 					{
@@ -4886,7 +4903,7 @@ namespace Electra
 									// Locate all TRUN boxes.
 									TArray<const FMP4Box*> AllTRUNBoxes;
 									Box->GetAllBoxInstances(AllTRUNBoxes, FMP4Box::kBox_trun);
-									for (int32 nTruns = 0; nTruns < AllTRUNBoxes.Num(); ++nTruns)
+									for(int32 nTruns = 0; nTruns < AllTRUNBoxes.Num(); ++nTruns)
 									{
 										Track->TRUNBoxes.Add(static_cast<const FMP4BoxTRUN*>(AllTRUNBoxes[nTruns]));
 									}
@@ -4971,7 +4988,7 @@ namespace Electra
 			}
 			CurrentIterator.Reset();
 			int64 LowestFilePos = TNumericLimits<int64>::Max();
-			for (int32 nTrk = 0; nTrk < TrackIterators.Num(); ++nTrk)
+			for(int32 nTrk = 0; nTrk < TrackIterators.Num(); ++nTrk)
 			{
 				if (!TrackIterators[nTrk]->IsAtEOS() && TrackIterators[nTrk]->GetSampleFileOffset() < LowestFilePos)
 				{
@@ -5000,7 +5017,7 @@ namespace Electra
 	void FParserISO14496_12::FAllTrackIterator::GetAllIterators(TArray<const ITrackIterator*>& OutIterators) const
 	{
 		OutIterators.Empty();
-		for (int32 nTrk = 0; nTrk < TrackIterators.Num(); ++nTrk)
+		for(int32 nTrk = 0; nTrk < TrackIterators.Num(); ++nTrk)
 		{
 			OutIterators.Add(TrackIterators[nTrk].Get());
 		}
@@ -5012,7 +5029,7 @@ namespace Electra
 		TSharedPtr<FParserISO14496_12::FAllTrackIterator, ESPMode::ThreadSafe> ti = MakeShared<FParserISO14496_12::FAllTrackIterator, ESPMode::ThreadSafe>();
 
 		int64 LowestFilePos = TNumericLimits<int64>::Max();
-		for (int32 nTrk = 0, nTrkMax = GetNumberOfTracks(); nTrk < nTrkMax; ++nTrk)
+		for(int32 nTrk = 0, nTrkMax = GetNumberOfTracks(); nTrk < nTrkMax; ++nTrk)
 		{
 			const FTrack* Track = ParsedTrackInfo->GetTrackByIndex(nTrk);
 			check(Track);
@@ -5022,7 +5039,7 @@ namespace Electra
 			UEMediaError err = TrkIt->StartAtFirst();
 			if (err == UEMEDIA_ERROR_OK)
 			{
-				while (TrkIt->GetSampleFileOffset() < InFromFilePos)
+				while(TrkIt->GetSampleFileOffset() < InFromFilePos)
 				{
 					err = TrkIt->Next();
 					if (err != UEMEDIA_ERROR_OK)
diff --git a/Engine/Plugins/Media/ElectraPlayer/Source/ElectraPlayerRuntime/Private/Runtime/Demuxer/ParserISO14496-12.h b/Engine/Plugins/Media/ElectraPlayer/Source/ElectraPlayerRuntime/Private/Runtime/Demuxer/ParserISO14496-12.h
--- a/Engine/Plugins/Media/ElectraPlayer/Source/ElectraPlayerRuntime/Private/Runtime/Demuxer/ParserISO14496-12.h
+++ b/Engine/Plugins/Media/ElectraPlayer/Source/ElectraPlayerRuntime/Private/Runtime/Demuxer/ParserISO14496-12.h
@@ -70,12 +70,12 @@ namespace Electra
 		};
 
 		/** Box type is a 32 bit value in an mp4 file. */
-		typedef uint32  FBoxType;
+		typedef uint32 FBoxType;
 
 		/**
 		 * Commonly referenced boxes
 		 */
-		 // This whole constexpr thing doesn't work...
+		 // constexpr doesn't work here
 			 /*
 				 static inline constexpr FBoxType MakeBoxType(unsigned char v1, unsigned char v2, unsigned char v3, unsigned char v4)
 				 {
@@ -84,6 +84,7 @@ namespace Electra
 			 */
 #define MAKE_BOX_ATOM(a,b,c,d) (IParserISO14496_12::FBoxType)((uint32)a << 24) | ((uint32)b << 16) | ((uint32)c << 8) | ((uint32)d)
 		static const FBoxType BoxType_ftyp = MAKE_BOX_ATOM('f', 't', 'y', 'p');
+		static const FBoxType BoxType_styp = MAKE_BOX_ATOM('s', 't', 'y', 'p');
 		static const FBoxType BoxType_moov = MAKE_BOX_ATOM('m', 'o', 'o', 'v');
 		static const FBoxType BoxType_sidx = MAKE_BOX_ATOM('s', 'i', 'd', 'x');
 		static const FBoxType BoxType_moof = MAKE_BOX_ATOM('m', 'o', 'o', 'f');
@@ -138,7 +139,7 @@ namespace Electra
 		/*******************************************************************************************************************/
 		class ITrack;
 
-		virtual UEMediaError PrepareTracks(TSharedPtrTS<const IParserISO14496_12>	OptionalMP4InitSegment) = 0;
+		virtual UEMediaError PrepareTracks(TSharedPtrTS<const IParserISO14496_12> OptionalMP4InitSegment) = 0;
 
 		virtual TMediaOptionalValue<FTimeFraction> GetMovieDuration() const = 0;
 
@@ -157,7 +158,8 @@ namespace Electra
 				Closest
 			};
 
-			virtual UEMediaError StartAtTime(const FTimeValue& AtTime, ESearchMode SearchMode) = 0;
+			virtual UEMediaError StartAtTime(const FTimeValue& AtTime, ESearchMode SearchMode, bool bNeedSyncSample) = 0;
+			virtual UEMediaError StartAtFirst(bool bNeedSyncSample) = 0;
 			virtual UEMediaError Next() = 0;
 
 			virtual bool IsAtEOS() const = 0;
@@ -227,8 +229,4 @@ namespace Electra
 
 	};
 
-
-
 } // namespace Electra
-
-
diff --git a/Engine/Plugins/Media/ElectraPlayer/Source/ElectraPlayerRuntime/Private/Runtime/Player/HLS/StreamReaderHLSfmp4.cpp b/Engine/Plugins/Media/ElectraPlayer/Source/ElectraPlayerRuntime/Private/Runtime/Player/HLS/StreamReaderHLSfmp4.cpp
--- a/Engine/Plugins/Media/ElectraPlayer/Source/ElectraPlayerRuntime/Private/Runtime/Player/HLS/StreamReaderHLSfmp4.cpp
+++ b/Engine/Plugins/Media/ElectraPlayer/Source/ElectraPlayerRuntime/Private/Runtime/Player/HLS/StreamReaderHLSfmp4.cpp
@@ -959,9 +959,9 @@ void FStreamReaderHLSfmp4::FStreamHandler::HandleRequest()
 									FTimeValue DTS;
 									FTimeValue PTS;
 									FTimeValue Duration;
-									for(Error = TrackIterator->StartAtTime(FTimeValue::GetZero(), IParserISO14496_12::ITrackIterator::ESearchMode::After); Error == UEMEDIA_ERROR_OK; Error = TrackIterator->Next())
+									for(Error = TrackIterator->StartAtFirst(false); Error == UEMEDIA_ERROR_OK; Error = TrackIterator->Next())
 									{
-										// Get the DTS and PTS. Those are 0-based in a fragment and offeet by the base media decode time of the fragment.
+										// Get the DTS and PTS. Those are 0-based in a fragment and offset by the base media decode time of the fragment.
 										DTS.SetFromND(TrackIterator->GetDTS(), TrackIterator->GetTimescale());
 										PTS.SetFromND(TrackIterator->GetPTS(), TrackIterator->GetTimescale());
 
@@ -1063,6 +1063,13 @@ void FStreamReaderHLSfmp4::FStreamHandler::HandleRequest()
 									}
 									delete TrackIterator;
 
+									if (Error != UEMEDIA_ERROR_OK && Error != UEMEDIA_ERROR_END_OF_STREAM)
+									{
+										// error iterating
+										LogMessage(IInfoLog::ELevel::Error, FString::Printf(TEXT("Failed to iterate over segment \"%s\""), *Request->URL));
+										bHasErrored = true;
+									}
+
 									// Check if we are done or if there is additional data that needs parsing, like more moof boxes.
 									if (HasReadBeenAborted() || HasReachedEOF())
 									{
diff --git a/Engine/Plugins/Media/ElectraPlayer/Source/ElectraPlayerRuntime/Private/Runtime/Player/mp4/ManifestMP4.cpp b/Engine/Plugins/Media/ElectraPlayer/Source/ElectraPlayerRuntime/Private/Runtime/Player/mp4/ManifestMP4.cpp
--- a/Engine/Plugins/Media/ElectraPlayer/Source/ElectraPlayerRuntime/Private/Runtime/Player/mp4/ManifestMP4.cpp
+++ b/Engine/Plugins/Media/ElectraPlayer/Source/ElectraPlayerRuntime/Private/Runtime/Player/mp4/ManifestMP4.cpp
@@ -496,7 +496,7 @@ IManifest::FResult FManifestMP4Internal::FTimelineAssetMP4::GetStartingSegment(T
 					}
 					else
 					{
-						err = TrackIt->StartAtTime(StartPosition.Time, SearchMode);
+						err = TrackIt->StartAtTime(StartPosition.Time, SearchMode, true);
 						if (err == UEMEDIA_ERROR_OK)
 						{
 							firstTimestamp.SetFromND(TrackIt->GetDTS(), TrackIt->GetTimescale());
@@ -611,7 +611,7 @@ IManifest::FResult FManifestMP4Internal::FTimelineAssetMP4::GetStartingSegment(T
 					}
 					else
 					{
-						err = TrackIt->StartAtTime(StartPosition.Time, SearchMode);
+						err = TrackIt->StartAtTime(StartPosition.Time, SearchMode, true);
 						if (err == UEMEDIA_ERROR_OK)
 						{
 							firstTimestamp.SetFromND(TrackIt->GetDTS(), TrackIt->GetTimescale());
