ElectraPlayer: streaming of an mp4 file is now artificially segmented to load smaller blocks of data that are easier to retry on failure and also do not keep the connection active for overly long periods of time; added basic support for mp4 edit lists to shift timestamps

#rnx

diff --git a/Engine/Plugins/Media/ElectraPlayer/Source/ElectraPlayerRuntime/Private/Runtime/Demuxer/ParserISO14496-12.cpp b/Engine/Plugins/Media/ElectraPlayer/Source/ElectraPlayerRuntime/Private/Runtime/Demuxer/ParserISO14496-12.cpp
--- a/Engine/Plugins/Media/ElectraPlayer/Source/ElectraPlayerRuntime/Private/Runtime/Demuxer/ParserISO14496-12.cpp
+++ b/Engine/Plugins/Media/ElectraPlayer/Source/ElectraPlayerRuntime/Private/Runtime/Demuxer/ParserISO14496-12.cpp
@@ -1135,7 +1135,7 @@ namespace Electra
 
 		FString GetLanguage() const
 		{
-			return FString(Language, FMEDIA_STATIC_ARRAY_COUNT(Language));
+			return FString(FMEDIA_STATIC_ARRAY_COUNT(Language), Language);
 		}
 
 	private:
@@ -1261,6 +1261,17 @@ namespace Electra
 			int16		MediaRateFraction;
 		};
 
+		int32 GetNumberOfEntries() const
+		{
+			return Entries.Num();
+		}
+
+		const FEntry& GetEntry(int32 EntryIndex) const
+		{
+			return Entries[EntryIndex];
+		}
+
+
 	private:
 		FMP4BoxELST() = delete;
 		FMP4BoxELST(const FMP4BoxELST&) = delete;
@@ -2171,7 +2182,7 @@ namespace Electra
 		{
 		}
 
-		struct Entry
+		struct FEntry
 		{
 			uint32		SampleCount;
 			uint32		SampleDelta;
@@ -2187,7 +2198,7 @@ namespace Electra
 			return Entries.Num();
 		}
 
-		const Entry& GetEntry(int32 EntryIndex) const
+		const FEntry& GetEntry(int32 EntryIndex) const
 		{
 			return Entries[EntryIndex];
 		}
@@ -2214,7 +2225,7 @@ namespace Electra
 				{
 					RETURN_IF_ERROR(ParseInfo->Reader()->Read(c));					// sample_count
 					RETURN_IF_ERROR(ParseInfo->Reader()->Read(d));					// sample_delta
-					Entry& e = Entries.AddDefaulted_GetRef();
+					FEntry& e = Entries.AddDefaulted_GetRef();
 					e.SampleCount = c;
 					e.SampleDelta = d;
 					// Update total counts for convenience.
@@ -2226,7 +2237,7 @@ namespace Electra
 		}
 
 	private:
-		TArray<Entry>					Entries;
+		TArray<FEntry>					Entries;
 		int64							NumTotalSamples;
 		int64							TotalDuration;
 	};
@@ -2249,7 +2260,7 @@ namespace Electra
 		{
 		}
 
-		struct Entry
+		struct FEntry
 		{
 			uint32		SampleCount;
 			int64		SampleOffset;
@@ -2260,9 +2271,9 @@ namespace Electra
 			return Entries.Num();
 		}
 
-		const Entry GetEntry(int32 EntryIndex) const
+		const FEntry GetEntry(int32 EntryIndex) const
 		{
-			Entry e;
+			FEntry e;
 			e.SampleCount = Entries[EntryIndex].SampleCount;
 			// For box version 0 return the unsigned value, otherwise convert the unsigned to signed (first to 32, then to 64 bits)
 			e.SampleOffset = Version == 0 ? (int64)Entries[EntryIndex].SampleOffset : (int64)((int32)Entries[EntryIndex].SampleOffset);
@@ -2386,7 +2397,7 @@ namespace Electra
 		{
 		}
 
-		struct Entry
+		struct FEntry
 		{
 			uint32		FirstChunk;
 			uint32		SamplesPerChunk;
@@ -2398,7 +2409,7 @@ namespace Electra
 			return Entries.Num();
 		}
 
-		const Entry& GetEntry(int32 EntryIndex) const
+		const FEntry& GetEntry(int32 EntryIndex) const
 		{
 			return Entries[EntryIndex];
 		}
@@ -2430,7 +2441,7 @@ namespace Electra
 					RETURN_IF_ERROR(ParseInfo->Reader()->Read(FirstChunk));			// first_chunk
 					RETURN_IF_ERROR(ParseInfo->Reader()->Read(SamplesPerChunk));	// samples_per_chunk
 					RETURN_IF_ERROR(ParseInfo->Reader()->Read(DescriptionIndex));	// sample_description_index
-					Entry& e = Entries.AddDefaulted_GetRef();
+					FEntry& e = Entries.AddDefaulted_GetRef();
 					e.FirstChunk = FirstChunk;
 					e.SamplesPerChunk = SamplesPerChunk;
 					e.SampleDescriptionIndex = DescriptionIndex;
@@ -2440,7 +2451,7 @@ namespace Electra
 		}
 
 	private:
-		TArray<Entry>	Entries;
+		TArray<FEntry>	Entries;
 	};
 
 
@@ -3469,8 +3480,12 @@ namespace Electra
 			virtual bool IsSyncSample() const override;
 			virtual int64 GetSampleSize() const override;
 			virtual int64 GetSampleFileOffset() const override;
+			virtual int64 GetRawDTS() const override;
+			virtual int64 GetRawPTS() const override;
+			virtual int64 GetCompositionTimeEdit() const override;
+			virtual int64 GetEmptyEditOffset() const override;
 
-			UEMediaError StartAtFirst();
+			UEMediaError StartAtFirstInteral();
 			void SetTrack(const FTrack* InTrack);
 			void SetOptions(const FParamDict& InOptions);
 			bool IsValid() const
@@ -3504,6 +3519,8 @@ namespace Electra
 					CTTSIndex = Other.CTTSIndex;
 					STSSNextSyncSampleNum = Other.STSSNextSyncSampleNum;
 					STSSIndex = Other.STSSIndex;
+					EmptyEditDurationInMediaTimeUnits = Other.EmptyEditDurationInMediaTimeUnits;
+					CompositionTimeEditOffset = Other.CompositionTimeEditOffset;
 					DataOffset = Other.DataOffset;
 					SampleDTS = Other.SampleDTS;
 					SamplePTS = Other.SamplePTS;
@@ -3517,7 +3534,7 @@ namespace Electra
 			}
 
 			FParamDict			Options;
-			const FTrack* Track;
+			const FTrack*		Track;
 
 			bool				bIsFragmented;
 
@@ -3546,6 +3563,8 @@ namespace Electra
 			int32				STSSIndex;
 
 			// Used in all cases
+			int64				EmptyEditDurationInMediaTimeUnits;
+			int64				CompositionTimeEditOffset;
 			int64				DataOffset;
 			int64				SampleDTS;
 			int64				SamplePTS;
@@ -3578,10 +3597,10 @@ namespace Electra
 			const FMP4BoxMVHD* MVHDBox;
 			const FMP4BoxELST* ELSTBox;
 			const FMP4BoxTKHD* TKHDBox;
-			const FMP4Box* MDIABox;
+			const FMP4Box*	   MDIABox;
 			const FMP4BoxMDHD* MDHDBox;
 			const FMP4BoxHDLR* HDLRBox;
-			const FMP4Box* STBLBox;
+			const FMP4Box*	   STBLBox;
 			const FMP4BoxSTSD* STSDBox;
 			const FMP4BoxSTTS* STTSBox;
 			const FMP4BoxCTTS* CTTSBox;
@@ -3590,7 +3609,7 @@ namespace Electra
 			const FMP4BoxSTCO* STCOBox;
 			const FMP4BoxSTSS* STSSBox;
 			// Fragmented track
-			const FMP4Box* MOOFBox;
+			const FMP4Box*	   MOOFBox;
 			const FMP4BoxTREX* TREXBox;
 			const FMP4BoxTFHD* TFHDBox;
 			const FMP4BoxTFDT* TFDTBox;
@@ -3599,6 +3618,9 @@ namespace Electra
 
 			FStreamCodecInformation					CodecInformation;
 
+			int64									EmptyEditDurationInMediaTimeUnits;
+			int64									CompositionTimeEditOffset;
+
 			MPEG::FAVCDecoderConfigurationRecord	CodecSpecificDataAVC;
 			MPEG::FESDescriptor						CodecSpecificDataMP4A;
 			FBitrateInfo							BitrateInfo;
@@ -3727,8 +3749,11 @@ namespace Electra
 		, TREXBox(nullptr)
 		, TFHDBox(nullptr)
 		, TFDTBox(nullptr)
+		, EmptyEditDurationInMediaTimeUnits(0)
+		, CompositionTimeEditOffset(0)
 	{
 	}
+
 	uint32 FParserISO14496_12::FTrack::GetID() const
 	{
 		return TKHDBox ? TKHDBox->GetTrackID() : 0;
@@ -3861,6 +3886,8 @@ namespace Electra
 		CTTSIndex = 0;
 		STSSNextSyncSampleNum = 0;
 		STSSIndex = 0;
+		EmptyEditDurationInMediaTimeUnits = 0;
+		CompositionTimeEditOffset = 0;
 		DataOffset = 0;
 		SampleDTS = 0;
 		SamplePTS = 0;
@@ -3909,7 +3936,7 @@ namespace Electra
 	}
 
 
-	UEMediaError FParserISO14496_12::FTrackIterator::StartAtFirst()
+	UEMediaError FParserISO14496_12::FTrackIterator::StartAtFirstInteral()
 	{
 		// NOTE: All boxes that are referenced here have been checked to exist so accessing them is safe.
 
@@ -3922,6 +3949,10 @@ namespace Electra
 		bEOS = false;
 		bIsFragmented = Track->TFHDBox != nullptr;
 
+		// Get the relevant 'elst' values over from the track for convenience.
+		CompositionTimeEditOffset = Track->CompositionTimeEditOffset;
+		EmptyEditDurationInMediaTimeUnits = Track->EmptyEditDurationInMediaTimeUnits;
+
 		if (bIsFragmented)
 		{
 			CurrentSampleNumber = 0;
@@ -4025,7 +4056,7 @@ namespace Electra
 			// This is because of the interleaving of samples from the contained tracks. How many times the sample size can be added
 			// is determined by the stsc box.
 			STSCIndex = 0;
-			const FMP4BoxSTSC::Entry& stsc0 = Track->STSCBox->GetEntry(0);
+			const FMP4BoxSTSC::FEntry& stsc0 = Track->STSCBox->GetEntry(0);
 			// Establish the file data offset for the first sample.
 			CurrentChunkIndex = stsc0.FirstChunk;
 			DataOffset = Track->STCOBox->GetEntry(CurrentChunkIndex - 1);
@@ -4068,8 +4099,6 @@ namespace Electra
 
 			// Set up the timestamps
 			int64 BaseMediaDecodeTime = 0;
-			// NOTE: There is the unfortunate issue with an edit list which can be used to shift the composition time among other things.
-			//       For the time being we ignore edit lists.
 			SampleDTS = 0;
 			SamplePTS = SampleDTS + CTTSOffset;
 
@@ -4254,7 +4283,7 @@ namespace Electra
 					// Switch to next chunk.
 					++STSCIndex;
 
-					const FMP4BoxSTSC::Entry& stsc = Track->STSCBox->GetEntry(STSCIndex);
+					const FMP4BoxSTSC::FEntry& stsc = Track->STSCBox->GetEntry(STSCIndex);
 					check(CurrentChunkIndex == stsc.FirstChunk);
 					DataOffset = Track->STCOBox->GetEntry(CurrentChunkIndex - 1);
 					NumSamplesInChunk = stsc.SamplesPerChunk;
@@ -4335,7 +4364,7 @@ namespace Electra
 
 	UEMediaError FParserISO14496_12::FTrackIterator::StartAtFirst(bool bNeedSyncSample)
 	{
-		UEMediaError err = StartAtFirst();
+		UEMediaError err = StartAtFirstInteral();
 		if (err == UEMEDIA_ERROR_OK)
 		{
 			for(; err == UEMEDIA_ERROR_OK; err = Next())
@@ -4358,7 +4387,7 @@ namespace Electra
 		//       Otherwise for fragmented files it's tricky to get at the sample flags to find the sync samples unless we would assume there is only
 		//       a single sync sample at the start of the fragment.
 		// So for now we take the shortcut to just iterate.
-		UEMediaError err = StartAtFirst();
+		UEMediaError err = StartAtFirstInteral();
 		if (err == UEMEDIA_ERROR_OK)
 		{
 			FTrackIterator Best;
@@ -4434,12 +4463,12 @@ namespace Electra
 
 	int64 FParserISO14496_12::FTrackIterator::GetDTS() const
 	{
-		return SampleDTS;
+		return SampleDTS + EmptyEditDurationInMediaTimeUnits;
 	}
 
 	int64 FParserISO14496_12::FTrackIterator::GetPTS() const
 	{
-		return SamplePTS;
+		return SamplePTS + EmptyEditDurationInMediaTimeUnits - CompositionTimeEditOffset;
 	}
 
 	int64 FParserISO14496_12::FTrackIterator::GetDuration() const
@@ -4467,6 +4496,25 @@ namespace Electra
 		return DataOffset;
 	}
 
+	int64 FParserISO14496_12::FTrackIterator::GetRawDTS() const
+	{
+		return SampleDTS;
+	}
+
+	int64 FParserISO14496_12::FTrackIterator::GetRawPTS() const
+	{
+		return SamplePTS;
+	}
+
+	int64 FParserISO14496_12::FTrackIterator::GetCompositionTimeEdit() const
+	{
+		return CompositionTimeEditOffset;
+	}
+
+	int64 FParserISO14496_12::FTrackIterator::GetEmptyEditOffset() const
+	{
+		return EmptyEditDurationInMediaTimeUnits;
+	}
 
 
 	/***************************************************************************************************************************************************/
@@ -4593,6 +4641,61 @@ namespace Electra
 						}
 						// TODO: Check the dinf->dref->url box to not reference any external media!
 
+						// Get the composition time offset and the empty edit from the edit list. Those are necessary to correct the PTS values from
+						// positive values in the 'ctts' box as well as general mapping of the content.
+						if (Track->ELSTBox)
+						{
+							bool bHaveEmptyEdit = false;
+							bool bHaveEmptyEditWarned = false;
+							bool bHaveOffsetEdit = false;
+							bool bHaveOffsetEditWarned = false;
+							for(int32 nEntry=0; nEntry<Track->ELSTBox->GetNumberOfEntries(); ++nEntry)
+							{
+								const FMP4BoxELST::FEntry& ee = Track->ELSTBox->GetEntry(nEntry);
+								// As per 8.6.6.1 the entries we are interested in have a rate of 1, so we use that to narrow down the entries to look at.
+								if (ee.MediaRateInteger == 1 && ee.MediaRateFraction == 0)
+								{
+									// Empty edit?
+									if (ee.MediaTime == -1)
+									{
+										// An empty edit means to insert blank, nonexistent material into the timeline.
+										// Essentially this means that everything from the media track will come after this blank and the best way
+										// to do this is to add the duration of the blank to the timestamps.
+										// However, this means that there will be no _real_ samples for this blank duration and thus nothing in
+										// any of the receive buffers. This should not be a huge blank period since we expect data to arrive in all buffers.
+										if (!bHaveEmptyEdit)
+										{
+											// We can only have a single empty edit.
+											bHaveEmptyEdit = true;
+											Track->EmptyEditDurationInMediaTimeUnits = Track->MDHDBox->GetTimescale() == Timescale ? ee.SegmentDuration : FTimeFraction(ee.SegmentDuration, Timescale).GetAsTimebase(Track->MDHDBox->GetTimescale());
+										}
+										else if (!bHaveEmptyEditWarned)
+										{
+											// Emit a warning.
+											bHaveEmptyEditWarned = true;
+											UE_LOG(LogElectraMP4Parser, Warning, TEXT("Found another empty edit in 'elst' box, ignoring"));
+										}
+									}
+									else if (ee.MediaTime >= 0)
+									{
+										// This defines the part of the media that is mapped to the track timeline's 0 point.
+										// This is commonly used to adjust the composition time offset and we treat it as such.
+										if (!bHaveOffsetEdit)
+										{
+											bHaveOffsetEdit = true;
+											Track->CompositionTimeEditOffset = ee.MediaTime;
+										}
+										else if (!bHaveOffsetEditWarned)
+										{
+											bHaveOffsetEditWarned = true;
+											UE_LOG(LogElectraMP4Parser, Warning, TEXT("Found another media edit in 'elst' box, ignoring"));
+										}
+									}
+								}
+							}
+						}
+
+
 						// Fragmented track?
 						if (MVEXBox)
 						{
@@ -4882,7 +4985,6 @@ namespace Electra
 					case FMP4Box::kBox_traf:
 					{
 						// FIXME: for encrypted content we need to also handle: sbgp, sgpd, subs, saiz, saio
-						// FIXME: there can theoretically be several trun boxes in sequence.
 
 						// Get the tfhd box.
 						const FMP4BoxTFHD* TFHDBox = static_cast<const FMP4BoxTFHD*>(Box->FindBox(FMP4Box::kBox_tfhd, 0));
@@ -5036,7 +5138,7 @@ namespace Electra
 			FTrackIterator* TrkIt = static_cast<FTrackIterator*>(Track->CreateIterator());
 			TSharedPtr<ITrackIterator, ESPMode::ThreadSafe> SafeTrkIt(TrkIt);
 			ti->TrackIterators.Add(SafeTrkIt);
-			UEMediaError err = TrkIt->StartAtFirst();
+			UEMediaError err = TrkIt->StartAtFirstInteral();
 			if (err == UEMEDIA_ERROR_OK)
 			{
 				while(TrkIt->GetSampleFileOffset() < InFromFilePos)
diff --git a/Engine/Plugins/Media/ElectraPlayer/Source/ElectraPlayerRuntime/Private/Runtime/Demuxer/ParserISO14496-12.h b/Engine/Plugins/Media/ElectraPlayer/Source/ElectraPlayerRuntime/Private/Runtime/Demuxer/ParserISO14496-12.h
--- a/Engine/Plugins/Media/ElectraPlayer/Source/ElectraPlayerRuntime/Private/Runtime/Demuxer/ParserISO14496-12.h
+++ b/Engine/Plugins/Media/ElectraPlayer/Source/ElectraPlayerRuntime/Private/Runtime/Demuxer/ParserISO14496-12.h
@@ -173,6 +173,10 @@ namespace Electra
 			virtual bool IsSyncSample() const = 0;
 			virtual int64 GetSampleSize() const = 0;
 			virtual int64 GetSampleFileOffset() const = 0;
+			virtual int64 GetRawDTS() const = 0;
+			virtual int64 GetRawPTS() const = 0;
+			virtual int64 GetCompositionTimeEdit() const = 0;
+			virtual int64 GetEmptyEditOffset() const = 0;
 		};
 
 		class ITrack
diff --git a/Engine/Plugins/Media/ElectraPlayer/Source/ElectraPlayerRuntime/Private/Runtime/HTTP/HTTPManager.cpp b/Engine/Plugins/Media/ElectraPlayer/Source/ElectraPlayerRuntime/Private/Runtime/HTTP/HTTPManager.cpp
--- a/Engine/Plugins/Media/ElectraPlayer/Source/ElectraPlayerRuntime/Private/Runtime/HTTP/HTTPManager.cpp
+++ b/Engine/Plugins/Media/ElectraPlayer/Source/ElectraPlayerRuntime/Private/Runtime/HTTP/HTTPManager.cpp
@@ -551,8 +551,13 @@ namespace Electra
 			}
 			else
 			{
-				// Give a warning to the console about this circumstance.
-				UE_LOG(LogElectraHTTPManager, Warning, TEXT("Receive buffer set to ring buffer mode but no sub range request size specified!"));
+				// Give a warning to the console about this circumstance if the original request is not ranged.
+				// If it is a range request that goes up to the end of the file we cannot say with certainty if this is a problem
+				// since we do not know how large the response is going to be. It could be a few bytes, it could be gigabytes...
+				if (!Handle->ActiveResponse.Range.IsSet())
+				{
+					UE_LOG(LogElectraHTTPManager, Warning, TEXT("Receive buffer set to ring buffer mode but no sub range request size specified!"));
+				}
 			}
 		}
 		// Set the possibly adjusted range.
diff --git a/Engine/Plugins/Media/ElectraPlayer/Source/ElectraPlayerRuntime/Private/Runtime/Player/mp4/ManifestMP4.cpp b/Engine/Plugins/Media/ElectraPlayer/Source/ElectraPlayerRuntime/Private/Runtime/Player/mp4/ManifestMP4.cpp
--- a/Engine/Plugins/Media/ElectraPlayer/Source/ElectraPlayerRuntime/Private/Runtime/Player/mp4/ManifestMP4.cpp
+++ b/Engine/Plugins/Media/ElectraPlayer/Source/ElectraPlayerRuntime/Private/Runtime/Player/mp4/ManifestMP4.cpp
@@ -15,7 +15,7 @@
 #include "Player/mp4/ManifestMP4.h"
 
 
-#define ERRCODE_MANIFEST_MP4_STARTSEGMENT_NOT_FOUND						1
+#define ERRCODE_MANIFEST_MP4_STARTSEGMENT_NOT_FOUND		1
 
 
 namespace Electra
@@ -460,6 +460,69 @@ void FManifestMP4Internal::FTimelineAssetMP4::LogMessage(IInfoLog::ELevel Level,
 }
 
 
+void FManifestMP4Internal::FTimelineAssetMP4::LimitSegmentDownloadSize(TSharedPtrTS<IStreamSegment>& InOutSegment, TSharedPtr<IParserISO14496_12::IAllTrackIterator, ESPMode::ThreadSafe> InAllTrackIterator)
+{
+	// Limit the segment download size.
+	// This helps with downloads that might otherwise take too long or keep the connection open for too long (when downloading a large mp4 from start to finish).
+	const int64 MaxSegmentSize = 4 * 1024 * 1024;
+	if (MaxSegmentSize > 0)
+	{
+		if (InOutSegment.IsValid())
+		{
+			FStreamSegmentRequestMP4* Request = static_cast<FStreamSegmentRequestMP4*>(InOutSegment.Get());
+			const int64 StartOffset = Request->FileStartOffset;
+			TSharedPtr<IParserISO14496_12::IAllTrackIterator, ESPMode::ThreadSafe> AllTrackIterator = InAllTrackIterator;
+			if (!AllTrackIterator.IsValid())
+			{
+				AllTrackIterator = MoovBoxParser->CreateAllTrackIteratorByFilePos(StartOffset);
+			}
+			bool bFirst = true;
+			uint32 TrackId = ~0U;
+			uint32 TrackTimeScale = 0;
+			int64 TrackDur = 0;
+			int64 LastTrackOffset = -1;
+			int64 LastSampleSize = 0;
+			while(AllTrackIterator.IsValid())
+			{
+				const IParserISO14496_12::ITrackIterator* CurrentTrackIt = AllTrackIterator->Current();
+				if (CurrentTrackIt)
+				{
+					LastTrackOffset = CurrentTrackIt->GetSampleFileOffset();
+					LastSampleSize = CurrentTrackIt->GetSampleSize();
+					if (bFirst)
+					{
+						bFirst = false;
+						TrackId = CurrentTrackIt->GetTrack()->GetID();
+						TrackTimeScale = CurrentTrackIt->GetTimescale();
+					}
+					if (TrackId == CurrentTrackIt->GetTrack()->GetID())
+					{
+						TrackDur += CurrentTrackIt->GetDuration();
+					}
+					int64 CurrentTrackOffset = LastTrackOffset;
+					if (CurrentTrackOffset - StartOffset >= MaxSegmentSize)
+					{
+						// Limit reached.
+						Request->FileEndOffset = CurrentTrackOffset - 1;
+						Request->SegmentInternalSize = CurrentTrackOffset - StartOffset;
+						Request->SegmentDuration.SetFromND(TrackDur, TrackTimeScale);
+						Request->bIsLastSegment = false;
+						break;
+					}
+					AllTrackIterator->Next();
+				}
+				else
+				{
+					// Done iterating
+					Request->SegmentInternalSize = LastTrackOffset + LastSampleSize - StartOffset;
+					break;
+				}
+			}
+		}
+	}
+}
+
+
 IManifest::FResult FManifestMP4Internal::FTimelineAssetMP4::GetStartingSegment(TSharedPtrTS<IStreamSegment>& OutSegment, const FPlayStartPosition& StartPosition, ESearchType SearchType, int64 AtAbsoluteFilePos)
 {
 // TODO: If there is a SIDX box we will look in there.
@@ -534,6 +597,8 @@ IManifest::FResult FManifestMP4Internal::FTimelineAssetMP4::GetStartingSegment(T
 					req->Bitrate				= Repr->GetBitrate();
 					req->bStartingOnMOOF		= false;
 					req->bIsContinuationSegment = false;
+					req->bIsFirstSegment		= true;
+					req->bIsLastSegment			= true;
 					req->SegmentDuration		= GetDuration() - req->FirstPTS;
 
 					// FIXME: this may need to add all additional tracks at some point if their individual IDs matter
@@ -542,6 +607,7 @@ IManifest::FResult FManifestMP4Internal::FTimelineAssetMP4::GetStartingSegment(T
 						req->DependentStreams.AddDefaulted_GetRef().StreamType = EStreamType::Audio;
 					}
 
+					LimitSegmentDownloadSize(OutSegment, nullptr);
 					return IManifest::FResult(IManifest::FResult::EType::Found);
 				}
 				else if (err == UEMEDIA_ERROR_END_OF_STREAM)
@@ -558,6 +624,8 @@ IManifest::FResult FManifestMP4Internal::FTimelineAssetMP4::GetStartingSegment(T
 						req->Bitrate				= Repr->GetBitrate();
 						req->bStartingOnMOOF		= false;
 						req->bIsContinuationSegment = false;
+						req->bIsFirstSegment		= false;
+						req->bIsLastSegment			= true;
 						req->bAllTracksAtEOS		= true;
 						//req->DependentStreams.PushBack().StreamType = EStreamType::Audio;
 						return IManifest::FResult(IManifest::FResult::EType::Found);
@@ -649,6 +717,8 @@ IManifest::FResult FManifestMP4Internal::FTimelineAssetMP4::GetStartingSegment(T
 					req->Bitrate				= Repr->GetBitrate();
 					req->bStartingOnMOOF		= false;
 					req->bIsContinuationSegment = false;
+					req->bIsFirstSegment		= true;
+					req->bIsLastSegment			= true;
 					req->SegmentDuration		= GetDuration() - req->FirstPTS;
 
 					// In case the video stream is shorter than audio we still need to add it as a dependent stream
@@ -660,6 +730,7 @@ IManifest::FResult FManifestMP4Internal::FTimelineAssetMP4::GetStartingSegment(T
 
 					// FIXME: there may be subtitle tracks here we need to add as dependent streams.
 
+					LimitSegmentDownloadSize(OutSegment, nullptr);
 					return IManifest::FResult(IManifest::FResult::EType::Found);
 				}
 				else if (err == UEMEDIA_ERROR_END_OF_STREAM)
@@ -673,6 +744,8 @@ IManifest::FResult FManifestMP4Internal::FTimelineAssetMP4::GetStartingSegment(T
 					req->Bitrate				= Repr->GetBitrate();
 					req->bStartingOnMOOF		= false;
 					req->bIsContinuationSegment = false;
+					req->bIsFirstSegment		= false;
+					req->bIsLastSegment			= true;
 					req->bAllTracksAtEOS		= true;
 					// But if there is a video track we add it as a dependent stream that is also at EOS.
 					if (VideoAdaptationSets.Num())
@@ -706,8 +779,23 @@ IManifest::FResult FManifestMP4Internal::FTimelineAssetMP4::GetStartingSegment(T
 
 IManifest::FResult FManifestMP4Internal::FTimelineAssetMP4::GetNextSegment(TSharedPtrTS<IStreamSegment>& OutSegment, TSharedPtrTS<const IStreamSegment> CurrentSegment, const FParamDict& Options)
 {
-	// At the moment we do not break the mp4 down into segments, so the one we had been streaming just now extended
-	// to the end of the asset which is now done.
+	const FStreamSegmentRequestMP4* Request = static_cast<const FStreamSegmentRequestMP4*>(CurrentSegment.Get());
+	if (Request)
+	{
+		// Check if the current request did not already go up to the end of the stream. If so there is no next segment.
+		if (Request->FileEndOffset >= 0)
+		{
+			FPlayStartPosition dummyPos;
+			IManifest::FResult res = GetStartingSegment(OutSegment, dummyPos, ESearchType::Same, Request->FileEndOffset + 1);
+			if (res.GetType() == IManifest::FResult::EType::Found)
+			{
+				FStreamSegmentRequestMP4* NextRequest = static_cast<FStreamSegmentRequestMP4*>(OutSegment.Get());
+				NextRequest->bIsContinuationSegment = true;
+				NextRequest->bIsFirstSegment		= false;
+				return res;
+			}
+		}
+	}
 	return IManifest::FResult(IManifest::FResult::EType::PastEOS);
 }
 
diff --git a/Engine/Plugins/Media/ElectraPlayer/Source/ElectraPlayerRuntime/Private/Runtime/Player/mp4/ManifestMP4.h b/Engine/Plugins/Media/ElectraPlayer/Source/ElectraPlayerRuntime/Private/Runtime/Player/mp4/ManifestMP4.h
--- a/Engine/Plugins/Media/ElectraPlayer/Source/ElectraPlayerRuntime/Private/Runtime/Player/mp4/ManifestMP4.h
+++ b/Engine/Plugins/Media/ElectraPlayer/Source/ElectraPlayerRuntime/Private/Runtime/Player/mp4/ManifestMP4.h
@@ -29,13 +29,12 @@ public:
 
 	FErrorDetail Build(TSharedPtrTS<IParserISO14496_12> MP4Parser, const FString& URL, const HTTP::FConnectionInfo& InConnectionInfo);
 
-
 	virtual EType GetPresentationType() const override;
 	virtual TSharedPtrTS<IPlaybackAssetTimeline> GetTimeline() const override;
 	virtual int64 GetDefaultStartingBitrate() const override;
 	virtual void GetStreamMetadata(TArray<FStreamMetadata>& OutMetadata, EStreamType StreamType) const override;
 	virtual FTimeValue GetMinBufferTime() const override;
-	virtual IStreamReader *CreateStreamReaderHandler() override;
+	virtual IStreamReader* CreateStreamReaderHandler() override;
 	virtual FResult FindPlayPeriod(TSharedPtrTS<IPlayPeriod>& OutPlayPeriod, const FPlayStartPosition& StartPosition, ESearchType SearchType) override;
 
 
@@ -158,6 +157,8 @@ public:
 		}
 
 
+		void LimitSegmentDownloadSize(TSharedPtrTS<IStreamSegment>& InOutSegment, TSharedPtr<IParserISO14496_12::IAllTrackIterator, ESPMode::ThreadSafe> AllTrackIterator);
+
 		FResult GetStartingSegment(TSharedPtrTS<IStreamSegment>& OutSegment, const FPlayStartPosition& StartPosition, ESearchType SearchType, int64 AtAbsoluteFilePos);
 		FResult GetNextSegment(TSharedPtrTS<IStreamSegment>& OutSegment, TSharedPtrTS<const IStreamSegment> CurrentSegment, const FParamDict& Options);
 		FResult GetRetrySegment(TSharedPtrTS<IStreamSegment>& OutSegment, TSharedPtrTS<const IStreamSegment> CurrentSegment, const FParamDict& Options);
diff --git a/Engine/Plugins/Media/ElectraPlayer/Source/ElectraPlayerRuntime/Private/Runtime/Player/mp4/PlaylistReaderMP4.cpp b/Engine/Plugins/Media/ElectraPlayer/Source/ElectraPlayerRuntime/Private/Runtime/Player/mp4/PlaylistReaderMP4.cpp
--- a/Engine/Plugins/Media/ElectraPlayer/Source/ElectraPlayerRuntime/Private/Runtime/Player/mp4/PlaylistReaderMP4.cpp
+++ b/Engine/Plugins/Media/ElectraPlayer/Source/ElectraPlayerRuntime/Private/Runtime/Player/mp4/PlaylistReaderMP4.cpp
@@ -63,7 +63,7 @@ public:
 	virtual FString GetURL() const override;
 
 	/**
-	 * Returns an interface to the manifest created from the loaded HLS playlists.
+	 * Returns an interface to the manifest created from the loaded mp4 playlists.
 	 *
 	 * @return A shared manifest interface pointer.
 	 */
@@ -91,7 +91,7 @@ private:
 			}
 			bHasBeenAborted = true;
 		}
-		TSharedPtrTS<IElectraHttpManager::FReceiveBuffer>		ReceiveBuffer;
+		TSharedPtrTS<IElectraHttpManager::FReceiveBuffer>	ReceiveBuffer;
 		int64												ParsePos;
 		bool												bHasBeenAborted;
 	};
@@ -129,7 +129,7 @@ private:
 	FMediaEvent												WorkerThreadQuitSignal;
 	bool													bIsWorkerThreadStarted;
 
-	TSharedPtrTS<IElectraHttpManager::FProgressListener>		ProgressListener;
+	TSharedPtrTS<IElectraHttpManager::FProgressListener>	ProgressListener;
 	FReadBuffer												ReadBuffer;
 	bool													bAbort;
 	bool													bHasErrored;
@@ -458,7 +458,7 @@ int64 FPlaylistReaderMP4::GetCurrentOffset() const
 
 IParserISO14496_12::IBoxCallback::EParseContinuation FPlaylistReaderMP4::OnFoundBox(IParserISO14496_12::FBoxType Box, int64 BoxSizeInBytes, int64 FileDataOffset, int64 BoxDataOffset)
 {
-	// The very first box must be an 'ftyp' box by definition. If this is not the case this is not a valid mp4 file.
+	// We require the very first box to be an 'ftyp' box.
 	if (FileDataOffset == 0 && Box != IParserISO14496_12::BoxType_ftyp)
 	{
 		PostError("Invalid mp4 file: first box is not 'ftyp'", ERRCODE_MP4_INVALID_FILE, UEMEDIA_ERROR_FORMAT_ERROR);
diff --git a/Engine/Plugins/Media/ElectraPlayer/Source/ElectraPlayerRuntime/Private/Runtime/Player/mp4/StreamReaderMP4.cpp b/Engine/Plugins/Media/ElectraPlayer/Source/ElectraPlayerRuntime/Private/Runtime/Player/mp4/StreamReaderMP4.cpp
--- a/Engine/Plugins/Media/ElectraPlayer/Source/ElectraPlayerRuntime/Private/Runtime/Player/mp4/StreamReaderMP4.cpp
+++ b/Engine/Plugins/Media/ElectraPlayer/Source/ElectraPlayerRuntime/Private/Runtime/Player/mp4/StreamReaderMP4.cpp
@@ -26,10 +26,13 @@ FStreamSegmentRequestMP4::FStreamSegmentRequestMP4()
 	PrimaryStreamType      = EStreamType::Video;
 	FileStartOffset 	   = -1;
 	FileEndOffset   	   = -1;
+	SegmentInternalSize	   = -1;
 	Bitrate 			   = 0;
 	PlaybackSequenceID     = ~0U;
 	bStartingOnMOOF 	   = false;
 	bIsContinuationSegment = false;
+	bIsFirstSegment		   = false;
+	bIsLastSegment		   = false;
 	bAllTracksAtEOS 	   = false;
 	CurrentIteratorBytePos = 0;
 	NumOverallRetries      = 0;
@@ -383,7 +386,32 @@ void FStreamReaderMP4::WorkerThread()
 		HTTP->Parameters.URL				= TimelineAsset->GetMediaURL();
 		HTTP->Parameters.Range.Start		= Request->FileStartOffset;
 		HTTP->Parameters.Range.EndIncluding = Request->FileEndOffset;
-		HTTP->Parameters.SubRangeRequestSize = 1 << 20;
+		// Explicit range?
+		int64 NumRequestedBytes = HTTP->Parameters.Range.GetNumberOfBytes();
+		int32 SubRequestSize = 0;
+		if (NumRequestedBytes > 0)
+		{
+			if (Request->bIsFirstSegment && !Request->bIsLastSegment)
+			{
+				SubRequestSize = 512 << 10;
+			}
+			else if (Request->bIsFirstSegment && Request->bIsLastSegment)
+			{
+				SubRequestSize = 2 << 20;
+			}
+		}
+		else
+		{
+			if (Request->SegmentInternalSize < 0)
+			{
+				SubRequestSize = 2 << 20;
+			}
+		}
+		if (SubRequestSize)
+		{
+			HTTP->Parameters.SubRangeRequestSize = SubRequestSize;
+		}
+
 		HTTP->ReceiveBuffer 				= ReadBuffer.ReceiveBuffer;
 		HTTP->ProgressListener  			= ProgressListener;
 		Parameters.PlayerSessionService->GetHTTPManager()->AddRequest(HTTP);
@@ -426,8 +454,8 @@ void FStreamReaderMP4::WorkerThread()
 				{
 					TSharedPtrTS<FAccessUnit::CodecData> CSD(new FAccessUnit::CodecData);
 					CSD->CodecSpecificData = Track->GetCodecSpecificData();
-					CSD->RawCSD			= Track->GetCodecSpecificDataRAW();
-					CSD->ParsedInfo 		= Track->GetCodecInformation();
+					CSD->RawCSD			   = Track->GetCodecSpecificDataRAW();
+					CSD->ParsedInfo		   = Track->GetCodecInformation();
 					st.CSD = MoveTemp(CSD);
 				}
 				return st;
diff --git a/Engine/Plugins/Media/ElectraPlayer/Source/ElectraPlayerRuntime/Private/Runtime/Player/mp4/StreamReaderMP4.h b/Engine/Plugins/Media/ElectraPlayer/Source/ElectraPlayerRuntime/Private/Runtime/Player/mp4/StreamReaderMP4.h
--- a/Engine/Plugins/Media/ElectraPlayer/Source/ElectraPlayerRuntime/Private/Runtime/Player/mp4/StreamReaderMP4.h
+++ b/Engine/Plugins/Media/ElectraPlayer/Source/ElectraPlayerRuntime/Private/Runtime/Player/mp4/StreamReaderMP4.h
@@ -55,10 +55,13 @@ public:
 	EStreamType											PrimaryStreamType;
 	int64												FileStartOffset;			//!< Where to start in the file
 	int64												FileEndOffset;				//!< Where to end in the file (for HTTP range GET requests)
+	int64												SegmentInternalSize;		//!< Size of the segment as defined by internal structures.
 	uint32												PlaybackSequenceID;
 	int32												Bitrate;
 	bool												bStartingOnMOOF;			//!< when using the 'sidx' a segment is expected to start on a 'moof', otherwise inside an 'mdat'.
-	bool												bIsContinuationSegment;
+	bool												bIsContinuationSegment;		//!< true if this segment continues where the previous left off and no sync samples should be expected.
+	bool												bIsFirstSegment;			//!< true if this segment is the first to start with or the first after a seek.
+	bool												bIsLastSegment;				//!< true if this segment is the last.
 
 	bool												bAllTracksAtEOS;
 
diff --git a/Engine/Plugins/Media/ElectraPlayer/Source/ElectraPlayerRuntime/Private/Runtime/PlayerTime.cpp b/Engine/Plugins/Media/ElectraPlayer/Source/ElectraPlayerRuntime/Private/Runtime/PlayerTime.cpp
--- a/Engine/Plugins/Media/ElectraPlayer/Source/ElectraPlayerRuntime/Private/Runtime/PlayerTime.cpp
+++ b/Engine/Plugins/Media/ElectraPlayer/Source/ElectraPlayerRuntime/Private/Runtime/PlayerTime.cpp
@@ -87,5 +87,34 @@ namespace Electra
 	}
 
 
+	int64 FTimeFraction::GetAsTimebase(uint32 CustomTimebase) const
+	{
+		// Some shortcuts
+		if (!bIsValid)
+		{
+			return 0;
+		}
+		else if (Numerator == 0)
+		{
+			return 0;
+		}
+		// Infinity?
+		else if (Denominator == 0 || CustomTimebase == 0)
+		{
+			return Numerator >= 0 ? 0x7fffffffffffffffLL : -0x7fffffffffffffffLL;
+		}
+
+		TBigInt<128> n(Numerator);
+		TBigInt<128> d(Denominator);
+		TBigInt<128> s(CustomTimebase);
+
+		n *= s;
+		n /= d;
+
+		int64 r = n.ToInt();
+		return r;
+	}
+
+
 }
 
diff --git a/Engine/Plugins/Media/ElectraUtil/Source/ElectraBase/Public/PlayerTime.h b/Engine/Plugins/Media/ElectraUtil/Source/ElectraBase/Public/PlayerTime.h
--- a/Engine/Plugins/Media/ElectraUtil/Source/ElectraBase/Public/PlayerTime.h
+++ b/Engine/Plugins/Media/ElectraUtil/Source/ElectraBase/Public/PlayerTime.h
@@ -509,6 +509,9 @@ public:
 		return Numerator / (double)Denominator;
 	}
 
+	//! Returns this time value in a custom timebase. Requires internal bigint conversion and is therefor SLOW!
+	int64 GetAsTimebase(uint32 CustomTimebase) const;
+
 	FTimeFraction& SetFromND(int64 InNumerator, uint32 InDenominator)
 	{
 		Numerator   = InNumerator;
