ElectraPlayer: Fixed bug with downloading license file into not initialized license cache.

#rnx

diff --git a/Engine/Plugins/Media/ElectraPlayer/Source/ElectraPlayerRuntime/Private/Runtime/Player/HLS/LicenseKeyCacheHLS.cpp b/Engine/Plugins/Media/ElectraPlayer/Source/ElectraPlayerRuntime/Private/Runtime/Player/HLS/LicenseKeyCacheHLS.cpp
--- a/Engine/Plugins/Media/ElectraPlayer/Source/ElectraPlayerRuntime/Private/Runtime/Player/HLS/LicenseKeyCacheHLS.cpp
+++ b/Engine/Plugins/Media/ElectraPlayer/Source/ElectraPlayerRuntime/Private/Runtime/Player/HLS/LicenseKeyCacheHLS.cpp
@@ -51,6 +51,8 @@ ILicenseKeyCacheHLS* ILicenseKeyCacheHLS::Create(IPlayerSessionServices* Session
 FLicenseKeyCacheHLS::FLicenseKeyCacheHLS(IPlayerSessionServices* InSessionServices, const FParamDict& InOptions)
 	: Options(InOptions)
 	, SessionServices(InSessionServices)
+	, Cache(kMaxCacheEntries)
+
 {
 }
 
