ElectraPlayer: Fixed a double event signal trigger where the 2nd might use a bad object after having triggered the first time.

#rnx

diff --git a/Engine/Plugins/Media/ElectraPlayer/Source/ElectraPlayerRuntime/Private/Runtime/HTTP/HTTPManager.cpp b/Engine/Plugins/Media/ElectraPlayer/Source/ElectraPlayerRuntime/Private/Runtime/HTTP/HTTPManager.cpp
--- a/Engine/Plugins/Media/ElectraPlayer/Source/ElectraPlayerRuntime/Private/Runtime/HTTP/HTTPManager.cpp
+++ b/Engine/Plugins/Media/ElectraPlayer/Source/ElectraPlayerRuntime/Private/Runtime/HTTP/HTTPManager.cpp
@@ -289,60 +289,16 @@ namespace Electra
 
 		struct FRemoveRequest
 		{
-			FRemoveRequest()
-				: WaitingEvent(nullptr)
-			{
-			}
-			~FRemoveRequest()
+			void SignalDone()
 			{
 				if (WaitingEvent)
 				{
 					WaitingEvent->Signal();
+					WaitingEvent = nullptr;
 				}
 			}
-			FRemoveRequest(const FRemoveRequest& other)
-				: Request(other.Request), WaitingEvent(other.WaitingEvent)
-			{
-			}
-
-			FRemoveRequest(FRemoveRequest&& other)
-				: Request(other.Request), WaitingEvent(other.WaitingEvent)
-			{
-				other.Request.Reset();
-				other.WaitingEvent = nullptr;
-			}
-			FRemoveRequest& operator=(FRemoveRequest&& other)
-			{
-				if (this != &other)
-				{
-					if (WaitingEvent)
-					{
-						WaitingEvent->Signal();
-					}
-					Request = other.Request;
-					WaitingEvent = other.WaitingEvent;
-
-					other.Request.Reset();
-					other.WaitingEvent = nullptr;
-				}
-				return *this;
-			}
-			FRemoveRequest& operator=(const FRemoveRequest& other)
-			{
-				if (this != &other)
-				{
-					if (WaitingEvent)
-					{
-						WaitingEvent->Signal();
-					}
-					Request = other.Request;
-					WaitingEvent = other.WaitingEvent;
-				}
-				return *this;
-			}
-
-			TSharedPtrTS<FRequest>			Request;
-			FMediaEvent* WaitingEvent;
+			TSharedPtrTS<FRequest> Request;
+			FMediaEvent* WaitingEvent = nullptr;
 		};
 
 		FCriticalSection												Lock;
@@ -625,32 +581,34 @@ namespace Electra
 		// Remove pending requests
 		while(!RequestsToRemove.IsEmpty())
 		{
-			TSharedPtrTS<FRequest>	Request = RequestsToRemove.Peek()->Request;
-
-			// Is this an active request?
-			for(TMap<FHandle*, TSharedPtrTS<FRequest>>::TIterator It = ActiveRequests.CreateIterator(); It; ++It)
+			FRemoveRequest Next;
+			if (RequestsToRemove.Dequeue(Next))
 			{
-				if (It.Value() == Request)
+				TSharedPtrTS<FRequest>	Request = Next.Request;
+
+				// Is this an active request?
+				for(TMap<FHandle*, TSharedPtrTS<FRequest>>::TIterator It = ActiveRequests.CreateIterator(); It; ++It)
 				{
-					FHandle* Handle = It.Key();
-					It.RemoveCurrent();
-					delete Handle;
-					break;
+					if (It.Value() == Request)
+					{
+						FHandle* Handle = It.Key();
+						It.RemoveCurrent();
+						delete Handle;
+						break;
+					}
 				}
-			}
 
-			// Removing an unfinished transfer that has not errored means it was aborted.
-			if (!Request->ConnectionInfo.bHasFinished && !Request->ConnectionInfo.StatusInfo.ErrorDetail.IsError())
-			{
-				Request->ConnectionInfo.bWasAborted = true;
-			}
-			if (!Request->ConnectionInfo.RequestEndTime.IsValid())
-			{
-				Request->ConnectionInfo.RequestEndTime = Now;
+				// Removing an unfinished transfer that has not errored means it was aborted.
+				if (!Request->ConnectionInfo.bHasFinished && !Request->ConnectionInfo.StatusInfo.ErrorDetail.IsError())
+				{
+					Request->ConnectionInfo.bWasAborted = true;
+				}
+				if (!Request->ConnectionInfo.RequestEndTime.IsValid())
+				{
+					Request->ConnectionInfo.RequestEndTime = Now;
+				}
+				Next.SignalDone();
 			}
-
-			// Pop the remove request. This will release the waiting caller.
-			RequestsToRemove.Pop();
 		}
 	}
 
@@ -822,7 +780,14 @@ namespace Electra
 	void FElectraHttpManager::RemoveAllRequests()
 	{
 		RequestsToAdd.Empty();
-		RequestsToRemove.Empty();
+		while(!RequestsToRemove.IsEmpty())
+		{
+			FRemoveRequest Next;
+			if (RequestsToRemove.Dequeue(Next))
+			{
+				Next.SignalDone();
+			}
+		}
 		RequestsCompleted.Empty();
 		while(!ActiveRequests.IsEmpty())
 		{
