ElectraPlayer: Removed a check() when an unsupported HLS playlist is encountered. Reworded the log message slightly for clarity.

#rnx

diff --git a/Engine/Plugins/Media/ElectraPlayer/Source/ElectraPlayerRuntime/Private/Runtime/Player/HLS/ManifestBuilderHLS.cpp b/Engine/Plugins/Media/ElectraPlayer/Source/ElectraPlayerRuntime/Private/Runtime/Player/HLS/ManifestBuilderHLS.cpp
--- a/Engine/Plugins/Media/ElectraPlayer/Source/ElectraPlayerRuntime/Private/Runtime/Player/HLS/ManifestBuilderHLS.cpp
+++ b/Engine/Plugins/Media/ElectraPlayer/Source/ElectraPlayerRuntime/Private/Runtime/Player/HLS/ManifestBuilderHLS.cpp
@@ -826,11 +826,11 @@ FErrorDetail FManifestBuilderHLS::SetupVariants(FManifestHLSInternal* Manifest,
 		}
 		else
 		{
-			LogMessage(IInfoLog::ELevel::Warning, FString::Printf(TEXT("EXT-X-STREAM-INF CODECS not specified. Assuming H.264 HIGH profile level 4.2 and LC-AAC audio.")));
-
 			// With either resolution or framerate specified we assume this is video.
 			if (bHaveResolution || FrameRate.Len())
 			{
+				LogMessage(IInfoLog::ELevel::Warning, FString::Printf(TEXT("EXT-X-STREAM-INF CODECS not specified. Assuming H.264 HIGH profile level 4.2 and LC-AAC audio.")));
+
 				FStreamCodecInformation si;
 				// Set the custom attributes with the codec information. This sets all extra options regardless of codec type!
 				si.GetExtras() = CompanyCustomExtraOptions;
@@ -869,10 +869,15 @@ FErrorDetail FManifestBuilderHLS::SetupVariants(FManifestHLSInternal* Manifest,
 				vs->StreamCodecInformationList.Push(si);
 				bHasAudio = true;
 			}
+			// If there is still neither video or audio this is probably an unsupported legacy master playlist that lists only bandwidth
+			// with no additional attributes and is thus very likely to also use MPEG2-TS segments that are not supported anyway.
+			if (!bHasVideo && !bHasAudio)
+			{
+				return CreateErrorAndLog(FString::Printf(TEXT("Unsupported variant playlist type. Neither CODECS, RESOLUTION or AUDIO is specified.")), ERRCODE_HLS_BUILDER_UNSUPPORTED_FEATURE, UEMEDIA_ERROR_FORMAT_ERROR);
+			}
 		}
 
 		// Usable stream?
-		check(bHasVideo || bHasAudio);
 		if ((bHasVideo || bHasAudio) && vs.IsValid())
 		{
 			// Check if the stream can be used on this platform.
@@ -1175,7 +1180,7 @@ FErrorDetail FManifestBuilderHLS::GetInitialPlaylistLoadRequests(TArray<FPlaylis
 	if (Manifest->VariantStreams.Num() == 0 && Manifest->AudioOnlyStreams.Num() == 0)
 	{
 		// NOTE: There might be a some text/caption-only variant but is this really a valid scenario?
-		return CreateErrorAndLog(FString::Printf(TEXT("No variant playlists in master playlist")), ERRCODE_HLS_BUILDER_NO_VARIANTS_IN_MASTER_PLAYLIST, UEMEDIA_ERROR_FORMAT_ERROR);
+		return CreateErrorAndLog(FString::Printf(TEXT("No usable variant playlists in master playlist")), ERRCODE_HLS_BUILDER_NO_VARIANTS_IN_MASTER_PLAYLIST, UEMEDIA_ERROR_FORMAT_ERROR);
 	}
 
 	TSharedPtrTS<FManifestHLSInternal::FVariantStream> StartingVariantStream;
