ElectraPlayer: Added MPEG DASH Live playback. Currently only for single period presentations.

#rnx

diff --git a/Engine/Plugins/Media/ElectraPlayer/Source/ElectraPlayerRuntime/Private/Runtime/HTTP/HTTPManager.cpp b/Engine/Plugins/Media/ElectraPlayer/Source/ElectraPlayerRuntime/Private/Runtime/HTTP/HTTPManager.cpp
--- a/Engine/Plugins/Media/ElectraPlayer/Source/ElectraPlayerRuntime/Private/Runtime/HTTP/HTTPManager.cpp
+++ b/Engine/Plugins/Media/ElectraPlayer/Source/ElectraPlayerRuntime/Private/Runtime/HTTP/HTTPManager.cpp
@@ -830,7 +830,7 @@ namespace Electra
 			}
 		}
 		RequestsCompleted.Empty();
-		while(!ActiveRequests.IsEmpty())
+		while(ActiveRequests.Num() != 0)
 		{
 			TMap<FHandle*, TSharedPtrTS<FRequest>>::TIterator It = ActiveRequests.CreateIterator();
 			FHandle* Handle = It.Key();
@@ -948,7 +948,7 @@ namespace Electra
 
 					FHttpResponsePtr Response = EvData.Response;
 					const int32 /*EHttpResponseCodes::Type*/ ResponseCode = Response ? Response->GetResponseCode() : EHttpResponseCodes::Unknown;
-					const bool bHTTPResponseSuccess = ResponseCode >= 200 && ResponseCode <= 299;
+					const bool bHTTPResponseSuccess = (ResponseCode >= 200 && ResponseCode <= 299) || ResponseCode == 304;
 
 					//FPlatformMisc::LocalPrint(*FString::Printf(TEXT("(%d, %d, %f; %d, %d, %d, %d)\n"), Status, ResponseCode, ElapsedTime, EvData.bComplete, EvData.bConnectedSuccessfully, EvData.BytesReceived, EvData.BytesSent));
 
@@ -1073,12 +1073,11 @@ namespace Electra
 					if (EvData.bComplete && Response.IsValid())
 					{
 						/*
-											// Get the effective URL. This may however be the original URL that was requested.
-											ci.EffectiveURL.assign(TCHAR_TO_ANSI(*Response->GetURL()));
+							We would like to get the effective URL after any redirections.
+							There doesn't seem to be a way right now to get it.
 
-									Taken out because this does not work on iOS and Mac. The engine HTTP module implementation there return an *empty* URL
-									for some reason which is very bad since user code makes use of the effective URL.
-									In absence of the actual effective URL (after redirection) we set the URL we made the request with.
+							UE_LOG(LogElectraHTTPManager, Log, TEXT("Effective URL: %s"), *Response->GetURL());
+							UE_LOG(LogElectraHTTPManager, Log, TEXT("Request   URL: %s"), *Request->Parameters.URL);
 						*/
 						ci.EffectiveURL = Request->Parameters.URL;
 						ci.StatusInfo.HTTPStatus = ResponseCode;
diff --git a/Engine/Plugins/Media/ElectraPlayer/Source/ElectraPlayerRuntime/Private/Runtime/Player/AdaptivePlayerOptionKeynames.h b/Engine/Plugins/Media/ElectraPlayer/Source/ElectraPlayerRuntime/Private/Runtime/Player/AdaptivePlayerOptionKeynames.h
new file mode 100644
--- /dev/null
+++ b/Engine/Plugins/Media/ElectraPlayer/Source/ElectraPlayerRuntime/Private/Runtime/Player/AdaptivePlayerOptionKeynames.h
@@ -0,0 +1,14 @@
+// Copyright Epic Games, Inc. All Rights Reserved.
+
+#pragma once
+
+#include "PlayerCore.h"
+
+namespace Electra
+{
+	const TCHAR* const OptionKeyInitialBitrate = TEXT("initial_bitrate");							//!< (int64) value indicating the bitrate to start with.
+	const TCHAR* const OptionKeyLiveSeekableEndOffset = TEXT("seekable_range_live_end_offset");		//!< (FTimeValue) value specifying how many seconds away from the Live media timeline the seekable range should end.
+
+} // namespace Electra
+
+
diff --git a/Engine/Plugins/Media/ElectraPlayer/Source/ElectraPlayerRuntime/Private/Runtime/Player/AdaptiveStreamingPlayer.cpp b/Engine/Plugins/Media/ElectraPlayer/Source/ElectraPlayerRuntime/Private/Runtime/Player/AdaptiveStreamingPlayer.cpp
--- a/Engine/Plugins/Media/ElectraPlayer/Source/ElectraPlayerRuntime/Private/Runtime/Player/AdaptiveStreamingPlayer.cpp
+++ b/Engine/Plugins/Media/ElectraPlayer/Source/ElectraPlayerRuntime/Private/Runtime/Player/AdaptiveStreamingPlayer.cpp
@@ -7,6 +7,7 @@
 #include "StreamAccessUnitBuffer.h"
 #include "Player/AdaptiveStreamingPlayerABR.h"
 #include "Player/PlayerLicenseKey.h"
+#include "Player/AdaptivePlayerOptionKeynames.h"
 
 #include "HAL/LowLevelMemTracker.h"
 
@@ -1177,7 +1178,7 @@ void FAdaptiveStreamingPlayer::WorkerThreadFN()
 						case EStreamType::Audio:
 							InitialStreamSelectionAud = MakeShared<FStreamMetadata, ESPMode::ThreadSafe>(msg.Data.TrackSelection.StreamMetadata);
 							// For media with multiplexed tracks we can switch to a different track pretty much immediately.
-							if (ManifestReader.IsValid() && ManifestReader->GetPlaylistType() == "mp4")
+							if (ManifestReader.IsValid() && ManifestReader->GetPlaylistType() == TEXT("mp4"))
 							{
 								// Only select audio for now.
 								if (InitialStreamSelectionAud.IsValid())
@@ -1203,7 +1204,7 @@ void FAdaptiveStreamingPlayer::WorkerThreadFN()
 						case EStreamType::Audio:
 							InitialStreamSelectionAud.Reset();
 							// For media with multiplexed tracks we can switch to a different track pretty much immediately.
-							if (ManifestReader.IsValid() && ManifestReader->GetPlaylistType() == "mp4")
+							if (ManifestReader.IsValid() && ManifestReader->GetPlaylistType() == TEXT("mp4"))
 							{
 								MultiStreamBufferAud.Deselect();
 							}
@@ -1324,6 +1325,8 @@ void FAdaptiveStreamingPlayer::WorkerThreadFN()
 			CheckForErrors();
 			// Handle pending media segment requests.
 			HandlePendingMediaSegmentRequests();
+			// Handle changes in metadata, like timeline changes or track availability.
+			HandleMetadataChanges();
 			// Handle buffer level changes
 			HandleNewBufferedData();
 			// Handle new output data (finish prerolling)
@@ -1380,6 +1383,30 @@ void FAdaptiveStreamingPlayer::HandlePlayStateChanges()
 }
 
 
+//-----------------------------------------------------------------------------
+/**
+ * Handle changes in metadata, like timeline changes or track availability.
+ */
+void FAdaptiveStreamingPlayer::HandleMetadataChanges()
+{
+	// The timeline can change dynamically. Refresh it on occasion.
+	if (Manifest.IsValid() && ManifestReader.IsValid() && ManifestReader->GetPlaylistType() == TEXT("dash"))
+	{
+		PlaybackState.SetSeekableRange(Manifest->GetSeekableTimeRange());
+		PlaybackState.SetTimelineRange(Manifest->GetTotalTimeRange());
+		PlaybackState.SetDuration(Manifest->GetDuration());
+/*
+		TArray<FStreamMetadata> VideoStreamMetadata;
+		TArray<FStreamMetadata> AudioStreamMetadata;
+		pPresentation->GetStreamMetadata(VideoStreamMetadata, EStreamType::Video);
+		pPresentation->GetStreamMetadata(AudioStreamMetadata, EStreamType::Audio);
+		PlaybackState.SetStreamMetadata(VideoStreamMetadata, AudioStreamMetadata);
+		PlaybackState.SetHaveMetadata(true);
+*/
+	}
+}
+
+
 //-----------------------------------------------------------------------------
 /**
  * Handles a player session message sent by one of the player sub-systems.
@@ -1414,7 +1441,7 @@ void FAdaptiveStreamingPlayer::HandleSessionMessage(TSharedPtrTS<IPlayerMessage>
 			{
 				if (!Result.WasAborted())
 				{
-					if (pMsg->GetListType() == Playlist::EListType::Master)
+					if (pMsg->GetListType() == Playlist::EListType::Master && pMsg->GetLoadType() == Playlist::ELoadType::Initial)
 					{
 						DispatchEvent(FMetricEvent::ReportReceivedMasterPlaylist(pMsg->GetConnectionInfo().EffectiveURL));
 					}
@@ -1544,7 +1571,7 @@ void FAdaptiveStreamingPlayer::HandleNewBufferedData()
 			// When we are dealing with a single multiplexed stream and one buffer is blocked then essentially all buffers
 			// must be considered blocked since demuxing cannot continue.
 			// FIXME: do this more elegant somehow
-			if (ManifestReader.IsValid() && ManifestReader->GetPlaylistType() == "mp4")
+			if (ManifestReader.IsValid() && ManifestReader->GetPlaylistType() == TEXT("mp4"))
 			{
 				if (VideoBufferStats.StreamBuffer.bLastPushWasBlocked ||
 					AudioBufferStats.StreamBuffer.bLastPushWasBlocked)
@@ -1756,7 +1783,7 @@ void FAdaptiveStreamingPlayer::HandlePendingMediaSegmentRequests()
 		return;
 	}
 
-	FTimeValue	CurrentTime 			= MEDIAutcTime::Current();
+	FTimeValue CurrentTime = MEDIAutcTime::Current();
 
 	// Is there a start request?
 	if (PendingStartRequest.IsValid())
@@ -1779,16 +1806,21 @@ void FAdaptiveStreamingPlayer::HandlePendingMediaSegmentRequests()
 					// If the starting time has not been set we now check if we are dealing with a VoD or a Live stream and choose a starting point.
 					if (!PendingStartRequest->StartAt.Time.IsValid())
 					{
-						FTimeRange Seekable = CurrentTimeline->GetSeekableTimeRange();
-						if (Manifest->GetPresentationType() == IManifest::EType::OnDemand)
-						{
-							check(Seekable.Start.IsValid());
-							PendingStartRequest->StartAt.Time = Seekable.Start;
-						}
-						else
+						// Use the presentation provided start time, if it has one.
+						PendingStartRequest->StartAt.Time = Manifest->GetDefaultStartTime();
+						if (!PendingStartRequest->StartAt.Time.IsValid())
 						{
-							check(Seekable.End.IsValid());
-							PendingStartRequest->StartAt.Time = Seekable.End;
+							FTimeRange Seekable = Manifest->GetSeekableTimeRange();
+							if (Manifest->GetPresentationType() == IManifest::EType::OnDemand)
+							{
+								check(Seekable.Start.IsValid());
+								PendingStartRequest->StartAt.Time = Seekable.Start;
+							}
+							else
+							{
+								check(Seekable.End.IsValid());
+								PendingStartRequest->StartAt.Time = Seekable.End;
+							}
 						}
 					}
 
@@ -1850,7 +1882,7 @@ void FAdaptiveStreamingPlayer::HandlePendingMediaSegmentRequests()
 						StreamSelector->SetCurrentPlaybackPeriod(CurrentPlayPeriod);
 
 						// Single stream mp4 files cannot switch streams. Failures need to retry the same file.
-						if (ManifestReader.IsValid() && ManifestReader->GetPlaylistType() == "mp4")
+						if (ManifestReader.IsValid() && ManifestReader->GetPlaylistType() == TEXT("mp4"))
 						{
 							StreamSelector->SetCanSwitchToAlternateStreams(false);
 						}
@@ -1907,9 +1939,15 @@ void FAdaptiveStreamingPlayer::HandlePendingMediaSegmentRequests()
 									break;
 								}
 								case IManifest::FResult::EType::NotFound:
+								case IManifest::FResult::EType::NotLoaded:
+								{
+									// Reset the current play period and start over with the initial period selection.
+									PostLog(Facility::EFacility::Player, IInfoLog::ELevel::Info, FString::Printf(TEXT("Period start segments not found. Reselecting start period.")));
+									CurrentPlayPeriod.Reset();
+									break;
+								}
 								case IManifest::FResult::EType::BeforeStart:
 								case IManifest::FResult::EType::PastEOS:
-								case IManifest::FResult::EType::NotLoaded:
 								{
 									FErrorDetail err;
 									err.SetFacility(Facility::EFacility::Player);
@@ -2042,8 +2080,26 @@ void FAdaptiveStreamingPlayer::HandlePendingMediaSegmentRequests()
 		else if (CurrentPlayPeriod.IsValid())
 		{
 			// Evaluate ABR to select the next stream quality.
-			FTimeValue ActionDelay;
-			IAdaptiveStreamSelector::ESegmentAction Action = StreamSelector->SelectSuitableStreams(ActionDelay, FinishedReq.Request);
+			IAdaptiveStreamSelector::ESegmentAction Action;
+			if (FinishedReq.PreviousEvalResult.IsSet())
+			{
+				Action = FinishedReq.PreviousEvalResult.Value();
+				FinishedReq.PreviousEvalResult.Reset();
+			}
+			else
+			{
+				FTimeValue ActionDelay(FTimeValue::GetZero());
+				Action = StreamSelector->SelectSuitableStreams(ActionDelay, FinishedReq.Request);
+				// If the selected choice is to be executed after a certain time take note of
+				// the choice and add the request back to the pending queue.
+				if (ActionDelay > FTimeValue::GetZero())
+				{
+					FinishedReq.AtTime = CurrentTime + ActionDelay;
+					FinishedReq.PreviousEvalResult.Set(Action);
+					NextPendingSegmentRequests.Enqueue(MoveTemp(FinishedReq));
+					continue;
+				}
+			}
 			if (Action == IAdaptiveStreamSelector::ESegmentAction::Fail)
 			{
 				FErrorDetail err;
@@ -2155,10 +2211,10 @@ void FAdaptiveStreamingPlayer::HandlePendingMediaSegmentRequests()
 			// Looping enabled?
 			if (CurrentLoopParam.bEnableLooping)
 			{
-				if (CurrentTimeline.IsValid() && CurrentPlayPeriod.IsValid())
+				if (Manifest.IsValid() && CurrentPlayPeriod.IsValid())
 				{
 					FPlayStartPosition	StartAt;
-					FTimeRange			Seekable = CurrentTimeline->GetSeekableTimeRange();
+					FTimeRange			Seekable = Manifest->GetSeekableTimeRange();
 					StartAt.Time = Seekable.Start;
 
 					TSharedPtrTS<IStreamSegment> FirstSegmentRequest;
@@ -2629,7 +2685,7 @@ bool FAdaptiveStreamingPlayer::InternalStartAt(const FSeekParam& NewPosition)
 	int64 StartingBitrate = StreamSelector->GetAverageBandwidth();
 	if (StartingBitrate <= 0)
 	{
-		StartingBitrate = PlayerOptions.GetValue("initial_bitrate").SafeGetInt64(0);
+		StartingBitrate = PlayerOptions.GetValue(OptionKeyInitialBitrate).SafeGetInt64(0);
 	}
 	if (StartingBitrate <= 0)
 	{
@@ -2655,7 +2711,7 @@ bool FAdaptiveStreamingPlayer::InternalStartAt(const FSeekParam& NewPosition)
 	UpdateDataAvailabilityState(DataAvailabilityStateTxt, Metrics::FDataAvailabilityChange::EAvailability::DataNotAvailable);
 
 	// For an mp4 file we expect several audio tracks to be available and select the one matching the selected index, if there is one.
-	if (ManifestReader.IsValid() && ManifestReader->GetPlaylistType() == "mp4")
+	if (ManifestReader.IsValid() && ManifestReader->GetPlaylistType() == TEXT("mp4"))
 	{
 		// Only select audio for now.
 		if (InitialStreamSelectionAud.IsValid())
diff --git a/Engine/Plugins/Media/ElectraPlayer/Source/ElectraPlayerRuntime/Private/Runtime/Player/AdaptiveStreamingPlayerABR.cpp b/Engine/Plugins/Media/ElectraPlayer/Source/ElectraPlayerRuntime/Private/Runtime/Player/AdaptiveStreamingPlayerABR.cpp
--- a/Engine/Plugins/Media/ElectraPlayer/Source/ElectraPlayerRuntime/Private/Runtime/Player/AdaptiveStreamingPlayerABR.cpp
+++ b/Engine/Plugins/Media/ElectraPlayer/Source/ElectraPlayerRuntime/Private/Runtime/Player/AdaptiveStreamingPlayerABR.cpp
@@ -33,7 +33,7 @@ namespace Electra
 		}
 
 		//! Clears everything
-		void Clear(void)
+		void Clear()
 		{
 			FMediaCriticalSection::ScopedLock lock(CriticalSection);
 			SimpleMovingAverage.Clear();
@@ -67,21 +67,23 @@ namespace Electra
 		T GetLastSample() const
 		{
 			FMediaCriticalSection::ScopedLock lock(CriticalSection);
-			return(LastSample);
+			return LastSample;
 		}
 
 		//! Returns the simple moving average from all the collected history samples.
-		T GetSMA(void) const
+		T GetSMA() const
 		{
 			FMediaCriticalSection::ScopedLock lock(CriticalSection);
 			if (SimpleMovingAverage.Num())
 			{
 				T sum = 0;
-				for (SIZE_T i = 0; i < SimpleMovingAverage.Num(); ++i)
+				for(SIZE_T i=0; i<SimpleMovingAverage.Num(); ++i)
+				{
 					sum += SimpleMovingAverage[i];
-				return(T(sum / SimpleMovingAverage.Num()));
+				}
+				return T(sum / SimpleMovingAverage.Num());
 			}
-			return(T(0));
+			return T(0);
 		}
 
 	private:
@@ -139,8 +141,6 @@ namespace Electra
 
 		virtual FABRDownloadProgressDecision ReportDownloadProgress(const Metrics::FSegmentDownloadStats& SegmentDownloadStats) override;
 		virtual void ReportDownloadEnd(const Metrics::FSegmentDownloadStats& SegmentDownloadStats) override;
-
-
 		virtual void ReportBufferingStart(Metrics::EBufferingReason BufferingReason) override;
 		virtual void ReportBufferingEnd(Metrics::EBufferingReason BufferingReason) override;
 		virtual void ReportOpenSource(const FString& URL) override {}
@@ -186,9 +186,6 @@ namespace Electra
 				Metrics::FSegmentDownloadStats					LastDownloadStats;
 			};
 
-			TSharedPtrTS<IPlaybackAssetAdaptationSet>		AdaptationSet;
-			TSharedPtrTS<IPlaybackAssetRepresentation>		Representation;
-
 			FString											AdaptationSetUniqueID;
 			FString											RepresentationUniqueID;
 
@@ -212,7 +209,7 @@ namespace Electra
 		void SimulateDownload(FSimulationResult& sim, const TArray<IManifest::IPlayPeriod::FSegmentInformation>& NextSegments, double secondsToSimulate, int64 atBPS);
 
 
-		IPlayerSessionServices* PlayerSessionServices;
+		IPlayerSessionServices*								PlayerSessionServices;
 		FConfiguration										Config;
 
 		FMediaCriticalSection								AccessMutex;
@@ -379,24 +376,24 @@ namespace Electra
 					{
 						CurrentMediaAssetID = Asset->GetUniqueIdentifier();
 						CurrentVideoAdaptationSetID = Adapt->GetUniqueIdentifier();
-						for (int32 i = 0, iMax = Adapt->GetNumberOfRepresentations(); i < iMax; ++i)
+						for(int32 i=0, iMax=Adapt->GetNumberOfRepresentations(); i<iMax; ++i)
 						{
+							TSharedPtrTS<IPlaybackAssetRepresentation> Repr = Adapt->GetRepresentationByIndex(i);
+
 							TSharedPtrTS<FStreamInformation> si = MakeSharedTS<FStreamInformation>();
-							si->AdaptationSet = Adapt;
-							si->Representation = Adapt->GetRepresentationByIndex(i);
-							si->AdaptationSetUniqueID = si->AdaptationSet->GetUniqueIdentifier();
-							si->RepresentationUniqueID = si->Representation->GetUniqueIdentifier();
-							si->Bitrate = si->Representation->GetBitrate();
-							if (si->Representation->GetCodecInformation().IsVideoCodec())
+							si->AdaptationSetUniqueID = Adapt->GetUniqueIdentifier();
+							si->RepresentationUniqueID = Repr->GetUniqueIdentifier();
+							si->Bitrate = Repr->GetBitrate();
+							if (Repr->GetCodecInformation().IsVideoCodec())
 							{
-								si->Resolution = si->Representation->GetCodecInformation().GetResolution();
+								si->Resolution = Repr->GetCodecInformation().GetResolution();
 							}
 							StreamInformationVideo.Push(si);
 						}
 						// Sort the representations by ascending bitrate
 						StreamInformationVideo.Sort([](const TSharedPtrTS<FStreamInformation>& a, const TSharedPtrTS<FStreamInformation>& b)
 						{
-							return a->Representation->GetBitrate() < b->Representation->GetBitrate();
+							return a->Bitrate < b->Bitrate;
 						});
 					}
 				}
@@ -408,26 +405,19 @@ namespace Electra
 					{
 						CurrentMediaAssetID = Asset->GetUniqueIdentifier();
 						CurrentAudioAdaptationSetID = Adapt->GetUniqueIdentifier();
-						for (int32 i = 0, iMax = Adapt->GetNumberOfRepresentations(); i < iMax; ++i)
+						for(int32 i=0, iMax=Adapt->GetNumberOfRepresentations(); i<iMax; ++i)
 						{
+							TSharedPtrTS<IPlaybackAssetRepresentation> Repr = Adapt->GetRepresentationByIndex(i);
 							TSharedPtrTS<FStreamInformation> si = MakeSharedTS<FStreamInformation>();
-							si->AdaptationSet = Adapt;
-							si->Representation = Adapt->GetRepresentationByIndex(i);
-							si->AdaptationSetUniqueID = si->AdaptationSet->GetUniqueIdentifier();
-							si->RepresentationUniqueID = si->Representation->GetUniqueIdentifier();
-							si->Bitrate = si->Representation->GetBitrate();
-							/*
-							if (si->Representation->GetCodecInformation().IsAudioCodec())
-							{
-								// ...
-							}
-							*/
+							si->AdaptationSetUniqueID = Adapt->GetUniqueIdentifier();
+							si->RepresentationUniqueID = Repr->GetUniqueIdentifier();
+							si->Bitrate = Repr->GetBitrate();
 							StreamInformationAudio.Push(si);
 						}
 						// Sort the representations by ascending bitrate
 						StreamInformationAudio.Sort([](const TSharedPtrTS<FStreamInformation>& a, const TSharedPtrTS<FStreamInformation>& b)
 						{
-							return a->Representation->GetBitrate() > b->Representation->GetBitrate();
+							return a->Bitrate > b->Bitrate;
 						});
 					}
 				}
@@ -510,7 +500,7 @@ namespace Electra
 	{
 		FMediaCriticalSection::ScopedLock lock(AccessMutex);
 		double lastBandwidth = AverageBandwidth.GetLastSample();
-		return((int64)lastBandwidth);
+		return (int64)lastBandwidth;
 	}
 
 
@@ -524,7 +514,7 @@ namespace Electra
 	{
 		FMediaCriticalSection::ScopedLock lock(AccessMutex);
 		const double averageBandwidth = AverageBandwidth.GetSMA();
-		return((int64)averageBandwidth);
+		return (int64)averageBandwidth;
 	}
 
 	//-----------------------------------------------------------------------------
@@ -731,7 +721,7 @@ namespace Electra
 		}
 		if (StreamInfos)
 		{
-			for (int32 i = 0, iMax = (int32)StreamInfos->Num(); i < iMax; ++i)
+			for(int32 i=0, iMax=StreamInfos->Num(); i<iMax; ++i)
 			{
 				if (FromDownloadStats.AdaptationSetID == (*StreamInfos)[i]->AdaptationSetUniqueID &&
 					FromDownloadStats.RepresentationID == (*StreamInfos)[i]->RepresentationUniqueID &&
@@ -812,14 +802,14 @@ namespace Electra
 		FTimeValue TimeNow = PlayerSessionServices->GetSynchronizedUTCTime()->GetTime();
 		// Make all streams available again if their blacklist time has expired
 		AccessMutex.Lock();
-		for (int32 i = 0, iMax = StreamInformationVideo.Num(); i < iMax; ++i)
+		for(int32 i=0, iMax=StreamInformationVideo.Num(); i<iMax; ++i)
 		{
 			if (StreamInformationVideo[i]->Health.BecomesAvailableAgainAtUTC <= TimeNow)
 			{
 				StreamInformationVideo[i]->Health.BecomesAvailableAgainAtUTC.SetToZero();
 			}
 		}
-		for (int32 i = 0, iMax = StreamInformationAudio.Num(); i < iMax; ++i)
+		for(int32 i=0, iMax=StreamInformationAudio.Num(); i<iMax; ++i)
 		{
 			if (StreamInformationAudio[i]->Health.BecomesAvailableAgainAtUTC <= TimeNow)
 			{
@@ -844,6 +834,17 @@ namespace Electra
 				if (!Stats.bWasSuccessful)
 				{
 					// Was not successful. Figure out what to do now.
+					// In the case where the segment had an availability window set the stream reader was waiting for
+					// to be entered and we got a 404 back, the server did not manage to publish the new segment in time.
+					// We try this again after a slight delay.
+					if (Stats.AvailibilityDelay && Utils::AbsoluteValue(Stats.AvailibilityDelay) < 0.5)
+					{
+						//LogMessage(IInfoLog::ELevel::Warning, FString::Printf(TEXT("Segment not available at announced time. Trying again.")));
+						CurrentPlayPeriod->IncreaseSegmentFetchDelay(FTimeValue(FTimeValue::MillisecondsToHNS(100)));
+						OutDelay.SetFromMilliseconds(500);
+						return ESegmentAction::Retry;
+					}
+
 					// Too many failures already? It's unlikely that there is a way that would magically fix itself now.
 					if (Stats.RetryNumber >= 3)
 					{
@@ -983,10 +984,10 @@ namespace Electra
 			{
 				PossibleRepresentations.Push(StreamInformationVideo[0]);
 			}
-			for (int32 nStr = 1; nStr < StreamInformationVideo.Num(); ++nStr)
+			for(int32 nStr=1; nStr<StreamInformationVideo.Num(); ++nStr)
 			{
 				// Check if bitrate and resolution are acceptable
-				if (StreamInformationVideo[nStr]->Representation->GetBitrate() <= BandwidthCeiling &&
+				if (StreamInformationVideo[nStr]->Bitrate <= BandwidthCeiling &&
 					!StreamInformationVideo[nStr]->Resolution.ExceedsLimit(MaxStreamResolution))
 				{
 					FSimulationResult sim;
@@ -994,7 +995,7 @@ namespace Electra
 					// Get information on the next segments.
 					TArray<IManifest::IPlayPeriod::FSegmentInformation> NextSegments;
 					FTimeValue AverageSegmentDuration;
-					CurrentPlayPeriod->GetSegmentInformation(NextSegments, AverageSegmentDuration, CurrentSegment, FTimeValue().SetFromSeconds(simulationDuration), StreamInformationVideo[nStr]->AdaptationSet, StreamInformationVideo[nStr]->Representation);
+					CurrentPlayPeriod->GetSegmentInformation(NextSegments, AverageSegmentDuration, CurrentSegment, FTimeValue().SetFromSeconds(simulationDuration), StreamInformationVideo[nStr]->AdaptationSetUniqueID, StreamInformationVideo[nStr]->RepresentationUniqueID);
 
 					// Long term prediction: Simulate download (at a user-scaled bandwidth factor)
 					SimulateDownload(sim, NextSegments, simulationDuration, simulationBandwidthLongterm);
@@ -1022,7 +1023,7 @@ namespace Electra
 					bool bIsFeasible = (bCond1 || bCond2) && bCond3;
 
 					// If there is an enforced initial bitrate then any stream that is within that bitrate is feasible by default.
-					if (ForcedInitialBandwidth && StreamInformationVideo[nStr]->Representation->GetBitrate() <= ForcedInitialBandwidth)
+					if (ForcedInitialBandwidth && StreamInformationVideo[nStr]->Bitrate <= ForcedInitialBandwidth)
 					{
 						bIsFeasible = true;
 					}
@@ -1041,7 +1042,7 @@ namespace Electra
 		//else if (StreamType == EStreamType::Audio)
 		else
 		{
-			for (int32 nStr = 0; nStr < StreamInformationAudio.Num(); ++nStr)
+			for(int32 nStr=0; nStr<StreamInformationAudio.Num(); ++nStr)
 			{
 				PossibleRepresentations.Push(StreamInformationAudio[nStr]);
 			}
@@ -1055,13 +1056,13 @@ namespace Electra
 			if (BlacklistedExternally.Num())
 			{
 				TArray<TSharedPtrTS<FStreamInformation>> RemainingCandidates;
-				for (int32 j = 0, jMax = PossibleRepresentations.Num(); j < jMax; ++j)
+				for(int32 j=0, jMax=PossibleRepresentations.Num(); j<jMax; ++j)
 				{
 					bool bStillGood = true;
-					for (int32 i = 0, iMax = BlacklistedExternally.Num(); i < iMax; ++i)
+					for(int32 i=0, iMax=BlacklistedExternally.Num(); i<iMax; ++i)
 					{
-						if (PossibleRepresentations[j]->AdaptationSet->GetUniqueIdentifier() == BlacklistedExternally[i].AdaptationSetUniqueID &&
-							PossibleRepresentations[j]->Representation->GetUniqueIdentifier() == BlacklistedExternally[i].RepresentationUniqueID)
+						if (PossibleRepresentations[j]->AdaptationSetUniqueID == BlacklistedExternally[i].AdaptationSetUniqueID &&
+							PossibleRepresentations[j]->RepresentationUniqueID == BlacklistedExternally[i].RepresentationUniqueID)
 						{
 							bStillGood = false;
 							break;
@@ -1079,7 +1080,7 @@ namespace Electra
 			//if (1)
 			{
 				TArray<TSharedPtrTS<FStreamInformation>> RemainingCandidates;
-				for (int32 j = 0, jMax = PossibleRepresentations.Num(); j < jMax; ++j)
+				for(int32 j=0, jMax=PossibleRepresentations.Num(); j<jMax; ++j)
 				{
 					bool bStillGood = PossibleRepresentations[j]->Health.BecomesAvailableAgainAtUTC <= TimeNow;
 					if (bStillGood)
@@ -1094,7 +1095,7 @@ namespace Electra
 			if (PossibleRepresentations.Num())
 			{
 				const TSharedPtrTS<FStreamInformation>& Candidate = PossibleRepresentations.Top();
-				CurrentPlayPeriod->SelectStream(Candidate->AdaptationSet, Candidate->Representation);
+				CurrentPlayPeriod->SelectStream(Candidate->AdaptationSetUniqueID, Candidate->RepresentationUniqueID);
 
 				if (bRetryIfPossible)
 				{
diff --git a/Engine/Plugins/Media/ElectraPlayer/Source/ElectraPlayerRuntime/Private/Runtime/Player/AdaptiveStreamingPlayerInternal.h b/Engine/Plugins/Media/ElectraPlayer/Source/ElectraPlayerRuntime/Private/Runtime/Player/AdaptiveStreamingPlayerInternal.h
--- a/Engine/Plugins/Media/ElectraPlayer/Source/ElectraPlayerRuntime/Private/Runtime/Player/AdaptiveStreamingPlayerInternal.h
+++ b/Engine/Plugins/Media/ElectraPlayer/Source/ElectraPlayerRuntime/Private/Runtime/Player/AdaptiveStreamingPlayerInternal.h
@@ -1181,6 +1181,7 @@ private:
 	void HandlePendingMediaSegmentRequests();
 	void HandleDeselectedBuffers();
 	void HandleDecoderChanges();
+	void HandleMetadataChanges();
 
 	void CheckForStreamEnd();
 
@@ -1287,7 +1288,6 @@ private:
 
 	TMediaQueueDynamic<TSharedPtrTS<FErrorDetail>>						ErrorQueue;
 
-	TSharedPtrTS<IPlaybackAssetTimeline>								CurrentTimeline;
 	TSharedPtrTS<IManifest>												Manifest;
 	TSharedPtrTS<IPlaylistReader>										ManifestReader;
 
@@ -1358,6 +1358,7 @@ private:
 	{
 		TSharedPtrTS<IStreamSegment>		Request;
 		FTimeValue							AtTime;
+		TMediaOptionalValue<IAdaptiveStreamSelector::ESegmentAction> PreviousEvalResult;
 	};
 
 	TSharedPtrTS<IManifest::IPlayPeriod>								CurrentPlayPeriod;
diff --git a/Engine/Plugins/Media/ElectraPlayer/Source/ElectraPlayerRuntime/Private/Runtime/Player/AdaptiveStreamingPlayerManifest.cpp b/Engine/Plugins/Media/ElectraPlayer/Source/ElectraPlayerRuntime/Private/Runtime/Player/AdaptiveStreamingPlayerManifest.cpp
--- a/Engine/Plugins/Media/ElectraPlayer/Source/ElectraPlayerRuntime/Private/Runtime/Player/AdaptiveStreamingPlayerManifest.cpp
+++ b/Engine/Plugins/Media/ElectraPlayer/Source/ElectraPlayerRuntime/Private/Runtime/Player/AdaptiveStreamingPlayerManifest.cpp
@@ -112,19 +112,19 @@ void FAdaptiveStreamingPlayer::InternalLoadManifest(const FString& URL, const FS
 			{
 				ManifestReader = IPlaylistReaderHLS::Create(this);
 				DispatchEvent(FMetricEvent::ReportOpenSource(URL));
-				ManifestReader->LoadAndParse(URL, StreamPreferences, PlayerOptions);
+				ManifestReader->LoadAndParse(URL);
 			}
 			else if (mimeType == Playlist::MIMETypeMP4)
 			{
 				ManifestReader = IPlaylistReaderMP4::Create(this);
 				DispatchEvent(FMetricEvent::ReportOpenSource(URL));
-				ManifestReader->LoadAndParse(URL, StreamPreferences, PlayerOptions);
+				ManifestReader->LoadAndParse(URL);
 			}
 			else if (mimeType == Playlist::MIMETypeDASH)
 			{
 				ManifestReader = IPlaylistReaderDASH::Create(this);
 				DispatchEvent(FMetricEvent::ReportOpenSource(URL));
-				ManifestReader->LoadAndParse(URL, StreamPreferences, PlayerOptions);
+				ManifestReader->LoadAndParse(URL);
 			}
 			else
 			{
@@ -162,89 +162,28 @@ bool FAdaptiveStreamingPlayer::SelectManifest()
 	if (ManifestReader)
 	{
 		check(Manifest == nullptr);
-		if (ManifestReader->GetPlaylistType() == TEXT("hls"))
+		if (ManifestReader->GetPlaylistType() == TEXT("hls") ||
+			ManifestReader->GetPlaylistType() == TEXT("mp4") ||
+			ManifestReader->GetPlaylistType() == TEXT("dash"))
 		{
 			TArray<FTimespan> SeekablePositions;
-			TSharedPtrTS<IManifest> pPresentation = ManifestReader->GetManifest();
-			check(pPresentation.IsValid());
+			TSharedPtrTS<IManifest> NewPresentation = ManifestReader->GetManifest();
+			check(NewPresentation.IsValid());
 
-			CurrentTimeline = pPresentation->GetTimeline();
-			PlaybackState.SetSeekableRange(CurrentTimeline->GetSeekableTimeRange());
-			CurrentTimeline->GetSeekablePositions(SeekablePositions);
+			PlaybackState.SetSeekableRange(NewPresentation->GetSeekableTimeRange());
+			NewPresentation->GetSeekablePositions(SeekablePositions);
 			PlaybackState.SetSeekablePositions(SeekablePositions);
-			PlaybackState.SetTimelineRange(CurrentTimeline->GetTotalTimeRange());
-			PlaybackState.SetDuration(CurrentTimeline->GetDuration());
+			PlaybackState.SetTimelineRange(NewPresentation->GetTotalTimeRange());
+			PlaybackState.SetDuration(NewPresentation->GetDuration());
 
 			TArray<FStreamMetadata> VideoStreamMetadata;
 			TArray<FStreamMetadata> AudioStreamMetadata;
-			pPresentation->GetStreamMetadata(VideoStreamMetadata, EStreamType::Video);
-			pPresentation->GetStreamMetadata(AudioStreamMetadata, EStreamType::Audio);
+			NewPresentation->GetStreamMetadata(VideoStreamMetadata, EStreamType::Video);
+			NewPresentation->GetStreamMetadata(AudioStreamMetadata, EStreamType::Audio);
 			PlaybackState.SetStreamMetadata(VideoStreamMetadata, AudioStreamMetadata);
 			PlaybackState.SetHaveMetadata(true);
 
-			Manifest = pPresentation;
-
-			CurrentState = EPlayerState::eState_Ready;
-
-			double minBufTimeMPD = Manifest->GetMinBufferTime().GetAsSeconds();
-			PlayerConfig.InitialBufferMinTimeAvailBeforePlayback = Utils::Min(minBufTimeMPD * PlayerConfig.ScaleMPDInitialBufferMinTimeBeforePlayback,   PlayerConfig.InitialBufferMinTimeAvailBeforePlayback);
-			PlayerConfig.SeekBufferMinTimeAvailBeforePlayback    = Utils::Min(minBufTimeMPD * PlayerConfig.ScaleMPDSeekBufferMinTimeAvailBeforePlayback, PlayerConfig.SeekBufferMinTimeAvailBeforePlayback);
-			PlayerConfig.RebufferMinTimeAvailBeforePlayback 	 = Utils::Min(minBufTimeMPD * PlayerConfig.ScaleMPDRebufferMinTimeAvailBeforePlayback,   PlayerConfig.RebufferMinTimeAvailBeforePlayback);
-
-			return true;
-		}
-		else if (ManifestReader->GetPlaylistType() == TEXT("mp4"))
-		{
-			TArray<FTimespan> SeekablePositions;
-			TSharedPtrTS<IManifest> pPresentation = ManifestReader->GetManifest();
-			check(pPresentation.IsValid());
-
-			CurrentTimeline = pPresentation->GetTimeline();
-			PlaybackState.SetSeekableRange(CurrentTimeline->GetSeekableTimeRange());
-			CurrentTimeline->GetSeekablePositions(SeekablePositions);
-			PlaybackState.SetSeekablePositions(SeekablePositions);
-			PlaybackState.SetTimelineRange(CurrentTimeline->GetTotalTimeRange());
-			PlaybackState.SetDuration(CurrentTimeline->GetDuration());
-
-			TArray<FStreamMetadata> VideoStreamMetadata;
-			TArray<FStreamMetadata> AudioStreamMetadata;
-			pPresentation->GetStreamMetadata(VideoStreamMetadata, EStreamType::Video);
-			pPresentation->GetStreamMetadata(AudioStreamMetadata, EStreamType::Audio);
-			PlaybackState.SetStreamMetadata(VideoStreamMetadata, AudioStreamMetadata);
-			PlaybackState.SetHaveMetadata(true);
-
-			Manifest = pPresentation;
-
-			CurrentState = EPlayerState::eState_Ready;
-
-			double minBufTimeMPD = Manifest->GetMinBufferTime().GetAsSeconds();
-			PlayerConfig.InitialBufferMinTimeAvailBeforePlayback = Utils::Min(minBufTimeMPD * PlayerConfig.ScaleMPDInitialBufferMinTimeBeforePlayback,   PlayerConfig.InitialBufferMinTimeAvailBeforePlayback);
-			PlayerConfig.SeekBufferMinTimeAvailBeforePlayback    = Utils::Min(minBufTimeMPD * PlayerConfig.ScaleMPDSeekBufferMinTimeAvailBeforePlayback, PlayerConfig.SeekBufferMinTimeAvailBeforePlayback);
-			PlayerConfig.RebufferMinTimeAvailBeforePlayback 	 = Utils::Min(minBufTimeMPD * PlayerConfig.ScaleMPDRebufferMinTimeAvailBeforePlayback,   PlayerConfig.RebufferMinTimeAvailBeforePlayback);
-
-			return true;
-		}
-		else if (ManifestReader->GetPlaylistType() == TEXT("dash"))
-		{
-			TArray<FTimespan> SeekablePositions;
-			TSharedPtrTS<IManifest> pPresentation = ManifestReader->GetManifest();
-			check(pPresentation.IsValid());
-
-			CurrentTimeline = pPresentation->GetTimeline();
-			PlaybackState.SetSeekableRange(CurrentTimeline->GetSeekableTimeRange());
-			CurrentTimeline->GetSeekablePositions(SeekablePositions);
-			PlaybackState.SetSeekablePositions(SeekablePositions);
-			PlaybackState.SetTimelineRange(CurrentTimeline->GetTotalTimeRange());
-			PlaybackState.SetDuration(CurrentTimeline->GetDuration());
-
-			TArray<FStreamMetadata> VideoStreamMetadata;
-			TArray<FStreamMetadata> AudioStreamMetadata;
-			pPresentation->GetStreamMetadata(VideoStreamMetadata, EStreamType::Video);
-			pPresentation->GetStreamMetadata(AudioStreamMetadata, EStreamType::Audio);
-			PlaybackState.SetStreamMetadata(VideoStreamMetadata, AudioStreamMetadata);
-			PlaybackState.SetHaveMetadata(true);
-
-			Manifest = pPresentation;
+			Manifest = NewPresentation;
 
 			CurrentState = EPlayerState::eState_Ready;
 
@@ -271,17 +210,15 @@ bool FAdaptiveStreamingPlayer::SelectManifest()
 
 void FAdaptiveStreamingPlayer::UpdateManifest()
 {
-	check(Manifest != nullptr);
+	check(Manifest.IsValid());
 
 	TArray<FTimespan> SeekablePositions;
 
-	// Update the current timeline.
-	CurrentTimeline = Manifest->GetTimeline();
-	PlaybackState.SetSeekableRange(CurrentTimeline->GetSeekableTimeRange());
-	CurrentTimeline->GetSeekablePositions(SeekablePositions);
+	PlaybackState.SetSeekableRange(Manifest->GetSeekableTimeRange());
+	Manifest->GetSeekablePositions(SeekablePositions);
 	PlaybackState.SetSeekablePositions(SeekablePositions);
-	PlaybackState.SetTimelineRange(CurrentTimeline->GetTotalTimeRange());
-	PlaybackState.SetDuration(CurrentTimeline->GetDuration());
+	PlaybackState.SetTimelineRange(Manifest->GetTotalTimeRange());
+	PlaybackState.SetDuration(Manifest->GetDuration());
 
 	TArray<FStreamMetadata> VideoStreamMetadata;
 	TArray<FStreamMetadata> AudioStreamMetadata;
diff --git a/Engine/Plugins/Media/ElectraPlayer/Source/ElectraPlayerRuntime/Private/Runtime/Player/AdaptiveStreamingPlayerMetrics.h b/Engine/Plugins/Media/ElectraPlayer/Source/ElectraPlayerRuntime/Private/Runtime/Player/AdaptiveStreamingPlayerMetrics.h
--- a/Engine/Plugins/Media/ElectraPlayer/Source/ElectraPlayerRuntime/Private/Runtime/Player/AdaptiveStreamingPlayerMetrics.h
+++ b/Engine/Plugins/Media/ElectraPlayer/Source/ElectraPlayerRuntime/Private/Runtime/Player/AdaptiveStreamingPlayerMetrics.h
@@ -146,6 +146,7 @@ namespace Metrics
 		// Outputs from stream reader
 		uint32			StatsID = 0;						//!< ID uniquely identifying this download
 		FString			FailureReason;						//!< Human readable failure reason. Only for display purposes.
+		double			AvailibilityDelay = 0.0;			//!< Time the download had to wait for the segment to enter its availability window.
 		double			DurationDownloaded = 0.0;			//!< Duration of content successfully downloaded. May be less than Duration in case of errors.
 		double			DurationDelivered = 0.0;			//!< Duration of content delivered to buffer. If larger than DurationDownloaded indicates dummy data was inserted into buffer.
 		double			TimeToFirstByte = 0.0;				//!< Time in seconds until first data byte was received
diff --git a/Engine/Plugins/Media/ElectraPlayer/Source/ElectraPlayerRuntime/Private/Runtime/Player/DASH/MPDElementsDASH.h b/Engine/Plugins/Media/ElectraPlayer/Source/ElectraPlayerRuntime/Private/Runtime/Player/DASH/MPDElementsDASH.h
--- a/Engine/Plugins/Media/ElectraPlayer/Source/ElectraPlayerRuntime/Private/Runtime/Player/DASH/MPDElementsDASH.h
+++ b/Engine/Plugins/Media/ElectraPlayer/Source/ElectraPlayerRuntime/Private/Runtime/Player/DASH/MPDElementsDASH.h
@@ -642,7 +642,7 @@ public:
 	const FTimeValue& GetTimeShiftBufferDepth() const						{ return TimeShiftBufferDepth; }
 	const FString& GetIndexRange() const									{ return IndexRange; }
 	bool GetIndexRangeExact() const											{ return bIndexRangeExact; }
-	const FTimeValue GetAvailabilityTimeOffset() const						{ return AvailabilityTimeOffset; }
+	const FTimeValue& GetAvailabilityTimeOffset() const						{ return AvailabilityTimeOffset; }
 	bool GetAvailabilityTimeComplete() const								{ return bAvailabilityTimeComplete; }
 
 protected:
@@ -1655,16 +1655,38 @@ public:
 
 	virtual FString GetDocumentURL() const override			{ return DocumentURL; }
 	void SetDocumentURL(const FString& InDocumentURL)		{ DocumentURL = InDocumentURL; }
-	const FTimeValue& GetFetchTime() const					{ return FetchTime; }
-	void SetFetchTime(const FTimeValue& InFetchTime)		{ FetchTime = InFetchTime; }
+	const FTimeValue& GetFetchTime() const					{ FScopeLock Lock(&UpdateLock); return FetchTime; }
+	void SetFetchTime(const FTimeValue& InFetchTime)		{ FScopeLock Lock(&UpdateLock); FetchTime = InFetchTime; }
+	FString GetETag() const									{ return ETag; }
+	void SetETag(const FString& InETag)						{ ETag = InETag; }
+
+	// Methods to manipulate the presentation type.
+	void LockAccess()											{ UpdateLock.Lock(); }
+	void UnlockAccess()											{ UpdateLock.Unlock(); }
+	void SetType(const FString& InType)							{ Type = InType; }
+	void SetPublishTime(const FTimeValue& InPUBT)				{ PublishTime = InPUBT; }
+	void SetAvailabilityStartTime(const FTimeValue& InAST)		{ AvailabilityStartTime = InAST; }
+	void SetAvailabilityEndTime(const FTimeValue& InAET)		{ AvailabilityEndTime = InAET; }
+	void SetMediaPresentationDuration(const FTimeValue& InPD)	{ MediaPresentationDuration = InPD; }
+	void SetMinimumUpdatePeriod(const FTimeValue& InMUP)		{ MinimumUpdatePeriod = InMUP; }
+	void SetTimeShiftBufferDepth(const FTimeValue& InTSB)		{ TimeShiftBufferDepth = InTSB; }
+	void SetSuggestedPresentationDelay(const FTimeValue& InSPD)	{ SuggestedPresentationDelay = InSPD; }
 
 private:
 	virtual bool ProcessElement(FManifestParserDASH* Builder, const TCHAR* ElementName, const TCHAR* ElementData, int32 XmlFileLineNumber) override;
 	virtual bool ProcessAttribute(FManifestParserDASH* Builder, const TCHAR* AttributeName, const TCHAR* AttributeValue) override;
 	virtual int32 ReplaceChildElementWithRemoteEntities(TSharedPtrTS<IDashMPDElement> Element, const FDashMPD_RootEntities& NewRootEntities, int64 OldResolveID, int64 NewResolveID) override;
 
+	// The MPD is not expected to get modified after being built except for potential XLink changes
+	// and extending the validity if an MPD update returned a 304 on a conditional GET.
+	// For these occurrences we need to lock access to the document on occasion.
+	mutable FCriticalSection UpdateLock;
+
 	// The URL from which this MPD was loaded. Provided for convenience.
 	FString DocumentURL;
+
+	// Optional ETag for conditional GET update requests.
+	FString ETag;
 	
 	// Time the MPD load request was started.
 	FTimeValue FetchTime;
diff --git a/Engine/Plugins/Media/ElectraPlayer/Source/ElectraPlayerRuntime/Private/Runtime/Player/DASH/ManifestBuilderDASH.cpp b/Engine/Plugins/Media/ElectraPlayer/Source/ElectraPlayerRuntime/Private/Runtime/Player/DASH/ManifestBuilderDASH.cpp
--- a/Engine/Plugins/Media/ElectraPlayer/Source/ElectraPlayerRuntime/Private/Runtime/Player/DASH/ManifestBuilderDASH.cpp
+++ b/Engine/Plugins/Media/ElectraPlayer/Source/ElectraPlayerRuntime/Private/Runtime/Player/DASH/ManifestBuilderDASH.cpp
@@ -7,6 +7,8 @@
 #include "Player/PlayerSessionServices.h"
 #include "Player/PlayerStreamFilter.h"
 #include "Player/PlayerEntityCache.h"
+#include "Player/DASH/OptionKeynamesDASH.h"
+#include "Player/AdaptivePlayerOptionKeynames.h"
 
 #include "ManifestBuilderDASH.h"
 #include "ManifestParserDASH.h"
@@ -38,11 +40,19 @@
 #define ERRCODE_DASH_MPD_BUILDER_BAD_PERIOD_DURATION				203
 
 
+
 namespace Electra
 {
+#ifdef LOG_DASHINT_CTORDTOR
+DECLARE_LOG_CATEGORY_EXTERN(LogElectraDASHMPD, Log, All);
+DEFINE_LOG_CATEGORY(LogElectraDASHMPD);
+#endif
 
 namespace
 {
+	const TCHAR* const Custom_EpicStaticStart = TEXT("EpicStaticStart");
+	const TCHAR* const Custom_EpicDynamicStart = TEXT("EpicDynamicStart");
+
 	const TCHAR* const XLinkActuateOnLoad = TEXT("onLoad");
 	const TCHAR* const XLinkActuateOnRequest = TEXT("onRequest");
 	const TCHAR* const XLinkResolveToZero = TEXT("urn:mpeg:dash:resolve-to-zero:2013");
@@ -452,6 +462,14 @@ namespace DASHUrlHelpers
 				}
 			}
 
+			// In case parameter replacement failed we could have sequences of double ampersands (&&) or a last/only & char.
+			while(FinalQueryString.ReplaceInline(TEXT("&&"), TEXT("&")))
+			{ }
+			if (FinalQueryString.Len() && FinalQueryString[FinalQueryString.Len()-1] == TCHAR('&'))
+			{
+				FinalQueryString.LeftChopInline(1);
+			}
+
 			// Where does the output go? URL query param or HTTP request header?
 			if (uq->GetExtendedUrlInfoType() == FDashMPD_UrlQueryInfoType::EExtendedUrlInfoType::ExtUrlQueryInfo)
 			{
@@ -504,8 +522,10 @@ namespace DASHUrlHelpers
 	 * Returns true if an absolute URL could be generated, false if not.
 	 * Since the MPD URL had to be an absolute URL this cannot actually fail.
 	 */
-	bool BuildAbsoluteElementURL(FString& OutURL, const FString& DocumentURL, const TArray<TSharedPtrTS<const FDashMPD_BaseURLType>>& BaseURLs, const FString& InElementURL)
+	bool BuildAbsoluteElementURL(FString& OutURL, FTimeValue& ATO, const FString& DocumentURL, const TArray<TSharedPtrTS<const FDashMPD_BaseURLType>>& BaseURLs, const FString& InElementURL)
 	{
+		FTimeValue SumOfUrlATOs(FTimeValue::GetZero());
+		ATO.SetToZero();
 		FURL_RFC3986 UrlParser;
 		FString ElementURL(InElementURL);
 		// If the element URL is empty it is specified as the first entry in the BaseURL array.
@@ -514,10 +534,11 @@ namespace DASHUrlHelpers
 		int32 nBase = 0;
 		if (ElementURL.IsEmpty())
 		{
-			// Must not be empty for this!
+			// If there is no element URL and BaseURLs this is probably for an MPD update in which case we return the original document URL.
 			if (BaseURLs.Num() == 0)
 			{
-				return false;
+				OutURL = DocumentURL;
+				return true;
 			}
 			ElementURL = BaseURLs[0]->GetURL();
 			++nBase;
@@ -534,7 +555,9 @@ namespace DASHUrlHelpers
 				{
 					if (nBase < BaseURLs.Num())
 					{
-						UrlParser.ResolveAgainst(BaseURLs[nBase++]->GetURL());
+						UrlParser.ResolveAgainst(BaseURLs[nBase]->GetURL());
+						SumOfUrlATOs += BaseURLs[nBase]->GetAvailabilityTimeOffset();
+						++nBase;
 					}
 					else
 					{
@@ -544,6 +567,7 @@ namespace DASHUrlHelpers
 				}
 			}
 			OutURL = UrlParser.Get();
+			ATO = SumOfUrlATOs;
 			return UrlParser.IsAbsolute();
 		}
 		return false;
@@ -668,19 +692,9 @@ class FManifestBuilderDASH : public IManifestBuilderDASH
 public:
 	FManifestBuilderDASH(IPlayerSessionServices* InPlayerSessionServices);
 	virtual ~FManifestBuilderDASH() = default;
+	
+	virtual FErrorDetail BuildFromMPD(TSharedPtrTS<FManifestDASHInternal>& OutMPD, TCHAR* InOutMPDXML, const FString& EffectiveURL, const FTimeValue& FetchTime, const FString& ETag) override;
 
-	/**
-	 * Builds a new internal manifest from a DASH MPD
-	 *
-	 * @param OutMPD
-	 * @param InOutMPDXML
-	 * @param SourceRequest
-	 * @param Preferences
-	 * @param Options
-	 *
-	 * @return
-	 */
-	virtual FErrorDetail BuildFromMPD(TSharedPtrTS<FManifestDASHInternal>& OutMPD, TCHAR* InOutMPDXML, TSharedPtrTS<FMPDLoadRequestDASH> SourceRequest, const FStreamPreferences& Preferences, const FParamDict& Options) override;
 private:
 	ELECTRA_IMPL_DEFAULT_ERROR_METHODS(DASHMPDBuilder);
 
@@ -701,7 +715,7 @@ FManifestBuilderDASH::FManifestBuilderDASH(IPlayerSessionServices* InPlayerSessi
 {
 }
 
-FErrorDetail FManifestBuilderDASH::BuildFromMPD(TSharedPtrTS<FManifestDASHInternal>& OutMPD, TCHAR* InOutMPDXML, TSharedPtrTS<FMPDLoadRequestDASH> SourceRequest, const FStreamPreferences& Preferences, const FParamDict& Options)
+FErrorDetail FManifestBuilderDASH::BuildFromMPD(TSharedPtrTS<FManifestDASHInternal>& OutMPD, TCHAR* InOutMPDXML, const FString& EffectiveURL, const FTimeValue& FetchTime, const FString& ETag)
 {
 	TSharedPtrTS<FDashMPD_MPDType> NewMPD;
 	TArray<TWeakPtrTS<IDashMPDElement>> XLinkElements;
@@ -717,8 +731,9 @@ FErrorDetail FManifestBuilderDASH::BuildFromMPD(TSharedPtrTS<FManifestDASHIntern
 		return CreateErrorAndLog(PlayerSessionServices, FString::Printf(TEXT("No root <MPD> element found")), ERRCODE_DASH_MPD_BUILDER_MISSING_REQUIRED_ELEMENT);
 	}
 	NewMPD = RootEntities.MPDs[0];
-	NewMPD->SetDocumentURL(SourceRequest->URL);
-	NewMPD->SetFetchTime(SourceRequest->GetConnectionInfo()->RequestStartTime);
+	NewMPD->SetDocumentURL(EffectiveURL);
+	NewMPD->SetFetchTime(FetchTime);
+	NewMPD->SetETag(ETag);
 
 	// Check if this MPD uses a profile we can handle. (If our array is empty we claim to handle anything)
 	bool bHasUsableProfile = UE_ARRAY_COUNT(SupportedProfiles) == 0;
@@ -775,7 +790,7 @@ FErrorDetail FManifestBuilderDASH::BuildFromMPD(TSharedPtrTS<FManifestDASHIntern
 	}
 
 	TSharedPtrTS<FManifestDASHInternal> NewManifest = MakeSharedTS<FManifestDASHInternal>();
-	Error = NewManifest->Build(PlayerSessionServices, NewMPD, MoveTemp(XLinkElements), Preferences, Options);
+	Error = NewManifest->Build(PlayerSessionServices, NewMPD, MoveTemp(XLinkElements));
 	if (Error.IsOK() || Error.IsTryAgain())
 	{
 		OutMPD = MoveTemp(NewManifest);
@@ -785,7 +800,7 @@ FErrorDetail FManifestBuilderDASH::BuildFromMPD(TSharedPtrTS<FManifestDASHIntern
 
 
 
-FErrorDetail FManifestDASHInternal::PrepareRemoteElementLoadRequest(IPlayerSessionServices* PlayerSessionServices, TArray<TWeakPtrTS<FMPDLoadRequestDASH>>& OutRemoteElementLoadRequests, TWeakPtrTS<IDashMPDElement> InElementWithXLink, int64 RequestID)
+FErrorDetail FManifestDASHInternal::PrepareRemoteElementLoadRequest(TArray<TWeakPtrTS<FMPDLoadRequestDASH>>& OutRemoteElementLoadRequests, TWeakPtrTS<IDashMPDElement> InElementWithXLink, int64 RequestID)
 {
 	const TCHAR* PreferredServiceLocation = nullptr;
 
@@ -818,6 +833,7 @@ FErrorDetail FManifestDASHInternal::PrepareRemoteElementLoadRequest(IPlayerSessi
 
 		FString URL;
 		FString RequestHeader;
+		FTimeValue UrlATO;
 		// Get the xlink:href
 		FString XlinkHRef = xl.GetHref();
 		if (!XlinkHRef.Equals(XLinkResolveToZero))
@@ -830,7 +846,7 @@ FErrorDetail FManifestDASHInternal::PrepareRemoteElementLoadRequest(IPlayerSessi
 			{
 				TArray<TSharedPtrTS<const FDashMPD_BaseURLType>> OutBaseURLs;
 				DASHUrlHelpers::GetAllHierarchyBaseURLs(PlayerSessionServices, OutBaseURLs, XElem->GetParentElement(), PreferredServiceLocation);
-				if (!DASHUrlHelpers::BuildAbsoluteElementURL(URL, MPDRoot->GetDocumentURL(), OutBaseURLs, XlinkHRef))
+				if (!DASHUrlHelpers::BuildAbsoluteElementURL(URL, UrlATO, MPDRoot->GetDocumentURL(), OutBaseURLs, XlinkHRef))
 				{
 					// Not resolving to an absolute URL is very unlikely as we had to load the MPD itself from somewhere.
 					return CreateErrorAndLog(PlayerSessionServices, FString::Printf(TEXT("xlink:href did not resolve to an absolute URL")), ERRCODE_DASH_MPD_BUILDER_URL_FAILED_TO_RESOLVE);
@@ -857,7 +873,7 @@ FErrorDetail FManifestDASHInternal::PrepareRemoteElementLoadRequest(IPlayerSessi
 		LoadReq->URL = URL;
 		if (RequestHeader.Len())
 		{
-			LoadReq->Headers.Emplace(HTTP::FHTTPHeader({TEXT("MPEG-DASH-Param"), RequestHeader}));
+			LoadReq->Headers.Emplace(HTTP::FHTTPHeader({DASH::HTTPHeaderOptionName, RequestHeader}));
 		}
 		LoadReq->XLinkElement = XElem;
 
@@ -872,16 +888,73 @@ FErrorDetail FManifestDASHInternal::PrepareRemoteElementLoadRequest(IPlayerSessi
 }
 
 
+void FManifestDASHInternal::TransformIntoEpicEvent()
+{
+	// For this to work the presentation must be 'static'
+	if (PresentationType == EPresentationType::Static)
+	{
+		FString Time;
+		bool bDynamic = false;
+		for(int32 i=0,iMax=URLFragmentComponents.Num(); i<iMax; ++i)
+		{
+			bool bStaticStart = URLFragmentComponents[i].Name.Equals(Custom_EpicStaticStart);
+			bool bDynamicStart = URLFragmentComponents[i].Name.Equals(Custom_EpicDynamicStart);
+			if (bStaticStart || bDynamicStart)
+			{
+				bDynamic = bDynamicStart;
+				Time = URLFragmentComponents[i].Value;
+				break;
+			}
+		}
+		if (!Time.IsEmpty())
+		{
+			// Get the event start time. This is either a Posix time in seconds since the Epoch (1/1/1970) or the special
+			// word 'now' optionally followed by a value to be added or subtracted from now.
+			FTimeValue Start;
+			if (Time.StartsWith(TEXT("now")))
+			{
+				Time.RightChopInline(3);
+				Start = PlayerSessionServices->GetSynchronizedUTCTime()->GetTime();
+			}
+			if (!Time.IsEmpty())
+			{
+				FTimeValue Posix = FTimeValue().SetFromTimeFraction(FTimeFraction().SetFromFloatString(Time));
+				if (Posix.IsValid())
+				{
+					if (Start.IsValid())
+					{
+						Start += Posix;
+					}
+					else
+					{
+						Start = Posix;
+					}
+				}
+			}
+			if (Start.IsValid())
+			{
+				MPDRoot->SetAvailabilityStartTime(Start);
+				if (bDynamic)
+				{
+					MPDRoot->SetPublishTime(Start);
+					MPDRoot->SetType(TEXT("dynamic"));
+					PresentationType = EPresentationType::Dynamic;
+					bIsEventType = true;
+				}
+			}
+		}
+	}
+}
+
+
 
-FErrorDetail FManifestDASHInternal::Build(IPlayerSessionServices* PlayerSessionServices, TSharedPtr<FDashMPD_MPDType, ESPMode::ThreadSafe> InMPDRoot, TArray<TWeakPtrTS<IDashMPDElement>> InXLinkElements, const FStreamPreferences& InPreferences, const FParamDict& InOptions)
+FErrorDetail FManifestDASHInternal::Build(IPlayerSessionServices* InPlayerSessionServices, TSharedPtr<FDashMPD_MPDType, ESPMode::ThreadSafe> InMPDRoot, TArray<TWeakPtrTS<IDashMPDElement>> InXLinkElements)
 {
+	PlayerSessionServices = InPlayerSessionServices;
 	RemoteElementsToResolve = MoveTemp(InXLinkElements);
 
 	FErrorDetail Error;
 
-	Preferences = InPreferences;
-	Options = InOptions;
-
 	MPDRoot = InMPDRoot;
 	
 	// What type of presentation is this?
@@ -921,7 +994,7 @@ FErrorDetail FManifestDASHInternal::Build(IPlayerSessionServices* PlayerSessionS
 	// As long as this list is not empty we are not done with the initial MPD setup.
 	for(int32 nRemoteElement=0; nRemoteElement<RemoteElementsToResolve.Num(); ++nRemoteElement)
 	{
-		Error = PrepareRemoteElementLoadRequest(PlayerSessionServices, PendingRemoteElementLoadRequests, RemoteElementsToResolve[nRemoteElement], 1);
+		Error = PrepareRemoteElementLoadRequest(PendingRemoteElementLoadRequests, RemoteElementsToResolve[nRemoteElement], 1);
 		if (Error.IsSet())
 		{
 			return Error;
@@ -931,10 +1004,10 @@ FErrorDetail FManifestDASHInternal::Build(IPlayerSessionServices* PlayerSessionS
 	{
 		return Error.SetTryAgain();
 	}
-	return BuildAfterInitialRemoteElementDownload(PlayerSessionServices);
+	return BuildAfterInitialRemoteElementDownload();
 }
 
-FErrorDetail FManifestDASHInternal::BuildAfterInitialRemoteElementDownload(IPlayerSessionServices* PlayerSessionServices)
+FErrorDetail FManifestDASHInternal::BuildAfterInitialRemoteElementDownload()
 {
 	FErrorDetail Error;
 
@@ -954,7 +1027,7 @@ FErrorDetail FManifestDASHInternal::BuildAfterInitialRemoteElementDownload(IPlay
 		// In this case we make up an artificial name.
 		if (p->ID.IsEmpty())
 		{
-			p->ID = FString::Printf(TEXT("$unnamed.%d$"), nPeriod);
+			p->ID = FString::Printf(TEXT("$period.%d$"), nPeriod);
 		}
 
 		p->Start = MPDperiods[nPeriod]->GetStart();
@@ -973,6 +1046,10 @@ FErrorDetail FManifestDASHInternal::BuildAfterInitialRemoteElementDownload(IPlay
 			}
 			else if (PresentationType == EPresentationType::Dynamic)
 			{
+				if (nPeriod == 0)
+				{
+					LogMessage(PlayerSessionServices, IInfoLog::ELevel::Info, FString::Printf(TEXT("The first period \"%s\" of this dynamic presentation has no start. Is this intended?"), *p->ID));
+				}
 				p->bIsEarlyPeriod = true;
 			}
 			else
@@ -1101,13 +1178,17 @@ FErrorDetail FManifestDASHInternal::BuildAfterInitialRemoteElementDownload(IPlay
 }
 
 
-FErrorDetail FManifestDASHInternal::PreparePeriodAdaptationSets(IPlayerSessionServices* PlayerSessionServices, TSharedPtrTS<FPeriod> Period, bool bRequestXlink)
+void FManifestDASHInternal::PreparePeriodAdaptationSets(TSharedPtrTS<FPeriod> Period, bool bRequestXlink)
 {
-	FErrorDetail Error;
 	check(Period.IsValid());
 	TSharedPtrTS<FDashMPD_PeriodType> MPDPeriod = Period->Period.Pin();
 	if (MPDPeriod.IsValid())
 	{
+		// If already prepared just return.
+		if (Period->GetHasBeenPrepared() && !bRequestXlink)
+		{
+			return;
+		}
 		Period->AdaptationSets.Empty();
 		const TArray<TSharedPtrTS<FDashMPD_AdaptationSetType>>& MPDAdaptationSets = MPDPeriod->GetAdaptationSets();
 		for(int32 nAdapt=0, nAdaptMax=MPDAdaptationSets.Num(); nAdapt<nAdaptMax; ++nAdapt)
@@ -1286,6 +1367,9 @@ FErrorDetail FManifestDASHInternal::PreparePeriodAdaptationSets(IPlayerSessionSe
 					continue;
 				}
 
+				Representation->ID = MPDRepresentation->GetID();
+				Representation->Bitrate = MPDRepresentation->GetBandwidth();
+
 				// Propagate the language code from the AdaptationSet into the codec info
 				Representation->CodecInfo.SetStreamLanguageCode(AdaptationSet->Language);
 
@@ -1449,8 +1533,8 @@ FErrorDetail FManifestDASHInternal::PreparePeriodAdaptationSets(IPlayerSessionSe
 			}
 
 		}
+		Period->SetHasBeenPrepared(true);
 	}
-	return Error;
 }
 
 
@@ -1458,7 +1542,7 @@ FErrorDetail FManifestDASHInternal::PreparePeriodAdaptationSets(IPlayerSessionSe
 /**
  * Resolves an xlink request made by the initial load of the MPD.
  */
-FErrorDetail FManifestDASHInternal::ResolveInitialRemoteElementRequest(IPlayerSessionServices* PlayerSessionServices, TSharedPtrTS<FMPDLoadRequestDASH> RequestResponse, FString XMLResponse, bool bSuccess)
+FErrorDetail FManifestDASHInternal::ResolveInitialRemoteElementRequest(TSharedPtrTS<FMPDLoadRequestDASH> RequestResponse, FString XMLResponse, bool bSuccess)
 {
 	// Because this is intended solely for initial MPD entities we do not need to worry about anyone accessing our internal structures
 	// while we update them. The player proper has not been informed yet that the initial manifest is ready for use.
@@ -1586,6 +1670,481 @@ int32 FManifestDASHInternal::ReplaceElementWithRemoteEntities(TSharedPtrTS<IDash
 }
 
 
+FTimeValue FManifestDASHInternal::GetAnchorTime() const
+{
+	return MPDRoot->GetAvailabilityStartTime().IsValid() ? MPDRoot->GetAvailabilityStartTime() : FTimeValue::GetZero();
+}
+
+FTimeRange FManifestDASHInternal::GetTotalTimeRange() const
+{
+	if (GetPresentationType() == FManifestDASHInternal::EPresentationType::Static)
+	{
+		FTimeValue Anchor = GetAnchorTime();
+		TotalTimeRange.Start = Anchor + GetPeriods()[0]->GetStart();
+		TotalTimeRange.End = TotalTimeRange.Start + GetDuration();
+	}
+	else
+	{
+		bool bIsUpdating = AreUpdatesExpected();
+		FTimeValue ast = GetAnchorTime();
+		FTimeValue Now = PlayerSessionServices->GetSynchronizedUTCTime()->GetTime();
+		FTimeValue LastEnd = GetLastPeriodEndTime();
+		TotalTimeRange.End = bIsUpdating ? Now : LastEnd;
+		if (TotalTimeRange.End.IsValid())
+		{
+			if (LastEnd.IsValid() && TotalTimeRange.End > LastEnd)
+			{
+				TotalTimeRange.End = LastEnd;
+			}
+
+			FTimeValue tsb = MPDRoot->GetTimeShiftBufferDepth();
+			if (tsb.IsValid())
+			{
+				TotalTimeRange.Start = TotalTimeRange.End - tsb;
+			}
+			else
+			{
+				TotalTimeRange.Start = ast;
+			}
+			FTimeValue PST = ast + GetPeriods()[0]->GetStart();
+			if (TotalTimeRange.Start < PST)
+			{
+				TotalTimeRange.Start = PST;
+			}
+			if (TotalTimeRange.End < TotalTimeRange.Start)
+			{
+				TotalTimeRange.End = TotalTimeRange.Start;
+			}
+		}
+		else
+		{
+			TotalTimeRange.Start = TotalTimeRange.End = ast;
+		}
+	}
+	return TotalTimeRange;
+}
+
+
+FTimeValue FManifestDASHInternal::CalculateDistanceToLiveEdge() const
+{
+	// Check if there is a user provided value. If there is it takes precedence over everything else.
+	FTimeValue Distance = PlayerSessionServices->GetOptions().GetValue(OptionKeyLiveSeekableEndOffset).SafeGetTimeValue(FTimeValue());
+	// If not set use the MPD@suggestedPresentationDelay
+	if (!Distance.IsValid())
+	{
+		Distance = MPDRoot->GetSuggestedPresentationDelay();
+	}
+	// If not set see if there is an MPD@maxSegmentDuration and use that.
+	if (!Distance.IsValid())
+	{
+		// Keep a distance of 3 segment durations.
+		// The DASH-IF-IOP recommendation is 2 to 4 SDURATION. See section 4.3.3.2.2.
+		Distance = MPDRoot->GetMaxSegmentDuration() * 3;
+		// And no less than 4 seconds.
+		if (Distance.GetAsMilliseconds(4000) < 4000)
+		{
+			Distance.SetFromMilliseconds(4000);
+		}
+	}
+	// If still not valid use MPD@minBufferTime. That one is mandatory to be present.
+	if (!Distance.IsValid())
+	{
+		// Also stay 3 durations away. This could be a really short or a really large value now.
+		// We make no adjustments here.
+		Distance = MPDRoot->GetMinBufferTime() * 3;
+	}
+	// We do clamp this to the MPD@timeShiftBufferDepth however to ensure it doesn't get bigger than that.
+	if (Distance > MPDRoot->GetTimeShiftBufferDepth())
+	{
+		Distance = MPDRoot->GetTimeShiftBufferDepth();
+	}
+	check(Distance.IsValid());
+	return Distance;
+}
+
+bool FManifestDASHInternal::UsesAST() const
+{
+	return MPDRoot->GetAvailabilityStartTime().IsValid();
+}
+
+bool FManifestDASHInternal::IsStaticType() const
+{
+	return GetPresentationType() == FManifestDASHInternal::EPresentationType::Static;
+}
+
+
+FTimeValue FManifestDASHInternal::GetMinimumUpdatePeriod() const
+{
+	if (GetPresentationType() == FManifestDASHInternal::EPresentationType::Dynamic && MPDRoot.IsValid())
+	{
+		return MPDRoot->GetMinimumUpdatePeriod();
+	}
+	return FTimeValue::GetInvalid();
+}
+
+bool FManifestDASHInternal::AreUpdatesExpected() const
+{
+	// Returns true if the MPD is subject to updates.
+	if (GetPresentationType() == FManifestDASHInternal::EPresentationType::Static)
+	{
+		return false;
+	}
+	return MPDRoot.IsValid() && MPDRoot->GetMinimumUpdatePeriod().IsValid();
+}
+
+
+FTimeValue FManifestDASHInternal::GetAvailabilityEndTime() const
+{
+	return GetMPDRoot()->GetAvailabilityEndTime();
+}
+
+FTimeValue FManifestDASHInternal::GetTimeshiftBufferDepth() const
+{
+	FTimeValue TSB = GetMPDRoot()->GetTimeShiftBufferDepth();
+	if (!TSB.IsValid())
+	{
+		// Section 5.3.1.2
+		// "When not present, the value is infinite. This value of the attribute is undefined if the type attribute is equal to 'static'."
+		TSB = GetPresentationType() == FManifestDASHInternal::EPresentationType::Static ? FTimeValue::GetZero() : FTimeValue::GetPositiveInfinity();
+	}
+	return TSB;
+}
+
+
+FTimeValue FManifestDASHInternal::GetLastPeriodEndTime() const
+{
+	// As per Annex A.3.2
+	FTimeValue ast = GetAnchorTime();
+	FTimeValue FetchTime = MPDRoot->GetFetchTime();
+	FTimeValue MediaPresentationDuration = FTimeValue::GetPositiveInfinity();
+	if (MPDRoot->GetMediaPresentationDuration().IsValid())
+	{
+		MediaPresentationDuration = MPDRoot->GetMediaPresentationDuration();
+	}
+	else if (Periods.Last()->GetDuration().IsValid())
+	{
+		MediaPresentationDuration = Periods.Last()->GetStart() + Periods.Last()->GetDuration() - Periods[0]->GetStart();
+	}
+	else if (Periods.Last()->GetEnd().IsValid())
+	{
+		MediaPresentationDuration = Periods.Last()->GetEnd() - Periods[0]->GetStart();
+	}
+
+	FTimeValue End = ast + MediaPresentationDuration;
+	if (MPDRoot->GetMinimumUpdatePeriod().IsValid())
+	{
+		FTimeValue CheckTime = FetchTime + MPDRoot->GetMinimumUpdatePeriod();
+		return CheckTime < End ? CheckTime : End;
+	}
+	else
+	{
+		return End;
+	}
+}
+
+
+FTimeValue FManifestDASHInternal::GetMPDValidityEndTime() const
+{
+	// Here we return the time until which the MPD is valid.
+	FTimeValue ast = GetAnchorTime();
+	FTimeValue FetchTime = MPDRoot->GetFetchTime();
+	if (MPDRoot->GetMinimumUpdatePeriod().IsValid())
+	{
+		return FetchTime + MPDRoot->GetMinimumUpdatePeriod();
+	}
+	return FetchTime;
+}
+
+
+FTimeRange FManifestDASHInternal::GetSeekableTimeRange() const
+{
+	const FTimeValue FixedSeekEndDistance(0.0);
+	if (GetPresentationType() == FManifestDASHInternal::EPresentationType::Static)
+	{
+		FTimeValue Anchor = GetAnchorTime();
+		SeekableTimeRange.Start = Anchor + GetPeriods()[0]->GetStart();
+		// FIXME: 10 seconds is an arbitrary value. We do not know the actual segment duration of the very last segment
+		//        so the only recourse right now is to only allow seeking up to some sensible point before the end.
+		SeekableTimeRange.End = Anchor + GetPeriods().Last()->GetEnd() - FixedSeekEndDistance;
+		if (SeekableTimeRange.End < SeekableTimeRange.Start)
+		{
+			SeekableTimeRange.End = SeekableTimeRange.Start;
+		}
+	}
+	else
+	{
+		// We can only rely on the MPD up to the time it was fetched. While it could be documenting
+		// future segments already (a pre-existing presentation made available over time) we would not
+		// be able to fetch them anyway on account of their availability window not being valid yet.
+		// Typically we would want to play some distance away from the bleeding Live edge.
+		bool bIsUpdating = AreUpdatesExpected();
+		FTimeValue Distance = bIsUpdating ? CalculateDistanceToLiveEdge() : FixedSeekEndDistance;
+		FTimeValue Now = PlayerSessionServices->GetSynchronizedUTCTime()->GetTime();
+		FTimeValue LastEnd = GetLastPeriodEndTime();
+		FTimeValue CurrentEnd = Now;
+		if (LastEnd.IsValid() && CurrentEnd > LastEnd)
+		{
+			CurrentEnd = LastEnd;
+		}
+		FTimeValue ast = GetAnchorTime();
+		if (CurrentEnd.IsValid())
+		{
+			SeekableTimeRange.End = CurrentEnd - Distance;
+
+			FTimeValue tsb = MPDRoot->GetTimeShiftBufferDepth();
+			if (tsb.IsValid())
+			{
+				SeekableTimeRange.Start = CurrentEnd - tsb;
+			}
+			else
+			{
+				SeekableTimeRange.Start = ast;
+			}
+			/*
+				The start must be covered by a Period. See: https://dashif-documents.azurewebsites.net/Guidelines-TimingModel/master/Guidelines-TimingModel.html#timing-timeshift
+					"Clients SHALL NOT allow seeking into regions of the time shift buffer that are not covered by periods, 
+					regardless of whether such regions are before or after the periods described by the MPD."
+			*/
+			FTimeValue PST = ast + GetPeriods()[0]->GetStart();
+			if (SeekableTimeRange.Start < PST)
+			{
+				SeekableTimeRange.Start = PST;
+			}
+			if (SeekableTimeRange.End < SeekableTimeRange.Start)
+			{
+				SeekableTimeRange.End = SeekableTimeRange.Start;
+			}
+			else if (SeekableTimeRange.End > LastEnd - Distance)
+			{
+				SeekableTimeRange.End = LastEnd - Distance;
+			}
+		}
+		else
+		{
+			SeekableTimeRange.Start = SeekableTimeRange.End = ast;
+		}
+	}
+	return SeekableTimeRange;
+}
+
+void FManifestDASHInternal::GetSeekablePositions(TArray<FTimespan>& OutPositions) const
+{
+	FTimeValue Anchor = GetAnchorTime();
+	for(int32 i=0; i<Periods.Num(); ++i)
+	{
+		if (!Periods[i]->GetIsEarlyPeriod())
+		{
+			if (GetPresentationType() == FManifestDASHInternal::EPresentationType::Static)
+			{
+				// The beginning of a period is a seekable position.
+				OutPositions.Emplace(FTimespan((Periods[i]->GetStart() + Anchor).GetAsHNS()));
+			}
+			else
+			{
+				/*
+					If the period is within the availability window we could add times.
+					But as with static presentations, without knowing the period contents there is nothing
+					we can really add here.
+				*/
+			}
+		}
+	}
+}
+
+FTimeValue FManifestDASHInternal::GetDuration() const
+{
+	// An MPD@mediaPresentationDuration is authoritative in both static and dynamic presentations.
+	if (MPDRoot->GetMediaPresentationDuration().IsValid())
+	{
+		return MPDRoot->GetMediaPresentationDuration();
+	}
+	else if (GetPresentationType() == FManifestDASHInternal::EPresentationType::Static)
+	{
+		if (MPDRoot->GetMediaPresentationDuration().IsValid())
+		{
+			return MPDRoot->GetMediaPresentationDuration();
+		}
+		else
+		{
+			// In a static period there cannot be any early periods so we can just use the difference between first and last.
+			return Periods.Last()->GetEnd() - Periods[0]->GetStart();
+		}
+	}
+	else
+	{
+		if (!MPDRoot->GetMinimumUpdatePeriod().IsValid() && (Periods.Last()->GetDuration().IsValid() || Periods.Last()->GetEnd().IsValid()))
+		{
+			return (Periods.Last()->GetEnd().IsValid() ? Periods.Last()->GetEnd() : Periods.Last()->GetStart() + Periods.Last()->GetDuration()) - Periods[0]->GetStart();
+		}
+		return FTimeValue::GetPositiveInfinity();
+	}
+}
+
+
+FTimeRange FManifestDASHInternal::GetPlayTimesFromURI() const
+{
+	FTimeRange FromTo;
+	
+	// We are interested in the 't' and 'period' fragment values here.
+	FString Time, PeriodID;
+	for(int32 i=0,iMax=URLFragmentComponents.Num(); i<iMax; ++i)
+	{
+		if (URLFragmentComponents[i].Name.Equals(TEXT("t")))
+		{
+			Time = URLFragmentComponents[i].Value;
+		}
+		else if (URLFragmentComponents[i].Name.Equals(TEXT("period")))
+		{
+			PeriodID = URLFragmentComponents[i].Value;
+		}
+	}
+	if (Time.IsEmpty() && PeriodID.IsEmpty())
+	{
+		return FromTo;
+	}
+
+	FTimeRange Seekable = GetSeekableTimeRange();
+	// Is the time specified as a POSIX time?
+	TArray<FString> TimeRange;
+	const TCHAR* const TimeDelimiter = TEXT(",");
+	if (Time.StartsWith(TEXT("posix")))
+	{
+		Time.RightChopInline(6);
+		Time.ParseIntoArray(TimeRange, TimeDelimiter, false);
+		// There needs to be a start time in the range. If there is only an end time we do nothing.
+		if (TimeRange.Num() && !TimeRange[0].IsEmpty())
+		{
+			// Is the start time the special time 'now'?
+			if (TimeRange[0].Equals(TEXT("now")))
+			{
+				FTimeValue Now = PlayerSessionServices->GetSynchronizedUTCTime()->GetTime();
+				/*
+				FTimeValue LastEnd = GetLastPeriodEndTime();
+				if (LastEnd.IsValid() && Now > LastEnd)
+				{
+					LastEnd = Now;
+				}
+				*/
+				FromTo.Start = Now;
+			}
+			else
+			{
+				int64 s = 0;
+				LexFromString(s, *TimeRange[0]);
+				FromTo.Start.SetFromSeconds(s);
+			}
+		}
+	}
+	else
+	{
+		// If there is no period specified then the period is the one with the earliest start time
+		if (PeriodID.IsEmpty())
+		{
+			FromTo.Start = Periods[0]->GetStart();
+		}
+		else
+		{
+			// Look for the named period.
+			for(int32 i=0; i<Periods.Num(); ++i)
+			{
+				if (Periods[i]->GetID().Equals(PeriodID))
+				{
+					FromTo.Start = Periods[i]->GetStart();
+					break;
+				}
+			}
+			// If the named period wasn't found use the first one.
+			if (!FromTo.Start.IsValid())
+			{
+				FromTo.Start = Periods[0]->GetStart();
+			}
+		}
+
+		FromTo.Start += GetAnchorTime();
+		// If there is no t specified we are done, otherwise we need to parse it.
+		if (!Time.IsEmpty())
+		{
+			// We need to parse out the 't' and add it to the period.
+			Time.ParseIntoArray(TimeRange, TimeDelimiter, false);
+			// There needs to be a start time in the range. If there is only an end time we do nothing.
+			if (TimeRange.Num() && !TimeRange[0].IsEmpty())
+			{
+				// When there is no POSIX time these values here are NPT
+				FTimeValue Offset;
+				if (RFC2326::ParseNPTTime(Offset, TimeRange[0]))
+				{
+					FromTo.Start += Offset;
+				}
+			}
+		}
+	}
+	// Need to clamp this into the seekable range to prevent any issues.
+	if (FromTo.Start.IsValid())
+	{
+		if (Seekable.Start.IsValid() && FromTo.Start < Seekable.Start)
+		{
+			FromTo.Start = Seekable.Start;
+		}
+		/*
+		else if (Seekable.End.IsValid() && FromTo.Start > Seekable.End)
+		{
+			FromTo.Start = Seekable.End;
+		}
+		*/
+	}
+	return FromTo;
+}
+
+FTimeValue FManifestDASHInternal::GetDefaultStartTime() const
+{
+	FTimeRange FromTo = GetPlayTimesFromURI();
+	return FromTo.Start;
+}
+
+
+#ifdef LOG_DASHINT_CTORDTOR
+FManifestDASHInternal::~FManifestDASHInternal()
+{
+	UE_LOG(LogElectraDASHMPD, Log, TEXT("~FManifestDASHInternal(%p)"), this);
+}
+
+FManifestDASHInternal::FManifestDASHInternal()
+{
+	UE_LOG(LogElectraDASHMPD, Log, TEXT("FManifestDASHInternal(%p)"), this);
+}
+
+FManifestDASHInternal::FPeriod::FPeriod()
+{
+	UE_LOG(LogElectraDASHMPD, Log, TEXT("FManifestDASHInternal::FPeriod(%p)"), this);
+}
+
+FManifestDASHInternal::FPeriod::~FPeriod()
+{
+	UE_LOG(LogElectraDASHMPD, Log, TEXT("FManifestDASHInternal::~FPeriod(%p)"), this);
+}
+
+FManifestDASHInternal::FAdaptationSet::FAdaptationSet()
+{
+	UE_LOG(LogElectraDASHMPD, Log, TEXT("FManifestDASHInternal::FAdaptationSet(%p)"), this);
+}
+
+FManifestDASHInternal::FAdaptationSet::~FAdaptationSet()
+{
+	UE_LOG(LogElectraDASHMPD, Log, TEXT("FManifestDASHInternal::~FAdaptationSet(%p)"), this);
+}
+
+FManifestDASHInternal::FRepresentation::FRepresentation()
+{
+	UE_LOG(LogElectraDASHMPD, Log, TEXT("FManifestDASHInternal::FRepresentation(%p)"), this);
+}
+
+FManifestDASHInternal::FRepresentation::~FRepresentation()
+{
+	UE_LOG(LogElectraDASHMPD, Log, TEXT("FManifestDASHInternal::~FRepresentation(%p)"), this);
+}
+
+
+#endif
 
 
 } // namespace Electra
diff --git a/Engine/Plugins/Media/ElectraPlayer/Source/ElectraPlayerRuntime/Private/Runtime/Player/DASH/ManifestBuilderDASH.h b/Engine/Plugins/Media/ElectraPlayer/Source/ElectraPlayerRuntime/Private/Runtime/Player/DASH/ManifestBuilderDASH.h
--- a/Engine/Plugins/Media/ElectraPlayer/Source/ElectraPlayerRuntime/Private/Runtime/Player/DASH/ManifestBuilderDASH.h
+++ b/Engine/Plugins/Media/ElectraPlayer/Source/ElectraPlayerRuntime/Private/Runtime/Player/DASH/ManifestBuilderDASH.h
@@ -26,13 +26,13 @@ public:
 	 *
 	 * @param OutMPD
 	 * @param InOutMPDXML
-	 * @param Request
-	 * @param Preferences
-	 * @param Options
+	 * @param EffectiveURL
+	 * @param FetchTime
+	 * @param ETag
 	 *
 	 * @return
 	 */
-	virtual FErrorDetail BuildFromMPD(TSharedPtrTS<FManifestDASHInternal>& OutMPD, TCHAR* InOutMPDXML, TSharedPtrTS<FMPDLoadRequestDASH> Request, const FStreamPreferences& Preferences, const FParamDict& Options) = 0;
+	virtual FErrorDetail BuildFromMPD(TSharedPtrTS<FManifestDASHInternal>& OutMPD, TCHAR* InOutMPDXML, const FString& EffectiveURL, const FTimeValue& FetchTime, const FString& ETag) = 0;
 };
 
 
