ElectraPlayer: Added MPEG DASH Live playback. Currently only for single period presentations.

diff --git a/Engine/Plugins/Media/ElectraPlayer/Source/ElectraPlayerRuntime/Private/Runtime/Player/DASH/ManifestDASH.cpp b/Engine/Plugins/Media/ElectraPlayer/Source/ElectraPlayerRuntime/Private/Runtime/Player/DASH/ManifestDASH.cpp
--- a/Engine/Plugins/Media/ElectraPlayer/Source/ElectraPlayerRuntime/Private/Runtime/Player/DASH/ManifestDASH.cpp
+++ b/Engine/Plugins/Media/ElectraPlayer/Source/ElectraPlayerRuntime/Private/Runtime/Player/DASH/ManifestDASH.cpp
@@ -12,8 +12,11 @@
 #include "Utilities/Utilities.h"
 #include "Utilities/StringHelpers.h"
 #include "Utilities/URLParser.h"
+#include "Utilities/TimeUtilities.h"
 #include "Player/DASH/OptionKeynamesDASH.h"
 #include "Player/PlayerEntityCache.h"
+#include "Player/AdaptivePlayerOptionKeynames.h"
+
 
 namespace Electra
 {
@@ -73,183 +76,6 @@ namespace DashUtils
 }
 
 
-
-class FDASHTimeline : public IPlaybackAssetTimeline
-{
-public:
-	FDASHTimeline(TSharedPtrTS<FManifestDASHInternal> InInternalManifest);
-
-	virtual ~FDASHTimeline();
-	virtual FTimeValue GetAnchorTime() const override;
-	virtual FTimeRange GetTotalTimeRange() const override;
-	virtual FTimeRange GetSeekableTimeRange() const override;
-	virtual void GetSeekablePositions(TArray<FTimespan>& OutPositions) const override;
-	virtual FTimeValue GetDuration() const override;
-	virtual int32 GetNumberOfMediaAssets() const override;
-	virtual TSharedPtrTS<ITimelineMediaAsset> GetMediaAssetByIndex(int32 MediaAssetIndex) const override;
-
-	TSharedPtrTS<FManifestDASHInternal> GetManifest()
-	{
-		return Manifest;
-	}
-private:
-	TSharedPtrTS<FManifestDASHInternal> Manifest;
-	mutable FTimeValue					AnchorTime;
-	mutable FTimeRange					TotalTimeRange;
-	mutable FTimeRange					SeekableTimeRange;
-};
-
-
-FDASHTimeline::FDASHTimeline(TSharedPtrTS<FManifestDASHInternal> InInternalManifest)
-	: Manifest(InInternalManifest)
-{
-}
-
-FDASHTimeline::~FDASHTimeline()
-{
-}
-
-FTimeValue FDASHTimeline::GetAnchorTime() const
-{
-	if (!AnchorTime.IsValid())
-	{
-		TSharedPtrTS<const FDashMPD_MPDType> MPDRoot = Manifest->GetMPDRoot();
-		AnchorTime = MPDRoot->GetAvailabilityStartTime().IsValid() ? MPDRoot->GetAvailabilityStartTime() : FTimeValue::GetZero();
-	}
-	return AnchorTime;
-}
-
-FTimeRange FDASHTimeline::GetTotalTimeRange() const
-{
-	if (!TotalTimeRange.IsValid())
-	{
-		if (Manifest->GetPresentationType() == FManifestDASHInternal::EPresentationType::Static)
-		{
-			FTimeValue Anchor = GetAnchorTime();
-			TotalTimeRange.Start = Anchor + Manifest->GetPeriods()[0]->GetStart();
-			TotalTimeRange.End = TotalTimeRange.Start + GetDuration();
-		}
-		else
-		{
-// TODO:
-			check(!"TODO");
-		}
-	}
-	return TotalTimeRange;
-}
-
-
-FTimeRange FDASHTimeline::GetSeekableTimeRange() const
-{
-	if (!SeekableTimeRange.IsValid())
-	{
-		if (Manifest->GetPresentationType() == FManifestDASHInternal::EPresentationType::Static)
-		{
-			FTimeValue Anchor = GetAnchorTime();
-			SeekableTimeRange.Start = Anchor + Manifest->GetPeriods()[0]->GetStart();
-			// FIXME: 10 seconds is an arbitrary value. We do not know the actual segment duration of the very last segment
-			//        so the only recourse right now is to only allow seeking up to some sensible point before the end.
-			SeekableTimeRange.End = Anchor + Manifest->GetPeriods().Last()->GetEnd() - FTimeValue().SetFromSeconds(10.0);
-			if (SeekableTimeRange.End < SeekableTimeRange.Start)
-			{
-				SeekableTimeRange.End = SeekableTimeRange.Start;
-			}
-		}
-		else
-		{
-// TODO:
-			check(!"TODO");
-		}
-	}
-	return SeekableTimeRange;
-}
-
-void FDASHTimeline::GetSeekablePositions(TArray<FTimespan>& OutPositions) const
-{
-	auto Periods = Manifest->GetPeriods();
-	FTimeValue Anchor = GetAnchorTime();
-	for(int32 i=0; i<Periods.Num(); ++i)
-	{
-		if (!Periods[i]->GetIsEarlyPeriod())
-		{
-			if (Manifest->GetPresentationType() == FManifestDASHInternal::EPresentationType::Static)
-			{
-				// The beginning of a period is a seekable position.
-				OutPositions.Emplace(FTimespan((Periods[i]->GetStart() + Anchor).GetAsHNS()));
-			}
-			else
-			{
-// TODO:
-				/*
-					In a dynamic presentation we need to consider the timeshiftBufferDepth and suggestedPresentationDelay
-					with regards to the MPD fetch time.
-				*/
-				check(!"TODO");
-			}
-		}
-	}
-}
-
-FTimeValue FDASHTimeline::GetDuration() const
-{
-	if (Manifest->GetPresentationType() == FManifestDASHInternal::EPresentationType::Static)
-	{
-		if (Manifest->GetMPDRoot()->GetMediaPresentationDuration().IsValid())
-		{
-			return Manifest->GetMPDRoot()->GetMediaPresentationDuration();
-		}
-		else
-		{
-			auto Periods = Manifest->GetPeriods();
-			for(int32 i=Periods.Num()-1; i>=0; --i)
-			{
-				if (!Periods[i]->GetIsEarlyPeriod())
-				{
-					return Periods[i]->GetEnd() - Periods[0]->GetStart();
-				}
-			}
-			return FTimeValue::GetInvalid();
-		}
-	}
-	else
-	{
-		// This is not necessarily infinity. A dynamic presentation can still have a predefined duration
-		// of pre-existing content that is made available in a dynamic way.
-// TODO:
-		check(!"TODO");
-		return FTimeValue::GetPositiveInfinity();
-	}
-}
-
-
-/**
-	* Returns the number of media assets on the playback timeline.
-	*
-	* @return Number of media assets on the playback timeline.
-	*/
-int32 FDASHTimeline::GetNumberOfMediaAssets() const
-{
-	check(!"TODO");
-	return 0;
-}
-
-/**
-	* Returns a media asset from the timeline by its index.
-	* Please note that media assets are not sorted by their time range on the timeline.
-	* This is done on purpose to keep the order of assets as they appear in the manifest.
-	*
-	* @param MediaAssetIndex
-	*               Index (0 to GetNumberOfMediaAssets()-1) of the asset to get.
-	*
-	* @return Shared pointer to the requested media asset.
-	*/
-TSharedPtrTS<ITimelineMediaAsset> FDASHTimeline::GetMediaAssetByIndex(int32 MediaAssetIndex) const
-{
-	check(!"TODO");
-	return nullptr;
-}
-
-
 /***************************************************************************************************************************************************/
 /***************************************************************************************************************************************************/
 /***************************************************************************************************************************************************/
@@ -257,10 +83,9 @@ TSharedPtrTS<ITimelineMediaAsset> FDASHTimeline::GetMediaAssetByIndex(int32 Medi
 class FDASHPlayPeriod : public IManifest::IPlayPeriod
 {
 public:
-	FDASHPlayPeriod(IPlayerSessionServices* InPlayerSessionServices, TSharedPtrTS<FManifestDASHInternal::FPeriod> SelectedPeriod, TSharedPtrTS<FDASHTimeline> InTimeline)
+	FDASHPlayPeriod(IPlayerSessionServices* InPlayerSessionServices, const FString& SelectedPeriodID)
 		: PlayerSessionServices(InPlayerSessionServices)
-		, InternalPeriod(MoveTemp(SelectedPeriod))
-		, Timeline(MoveTemp(InTimeline))
+		, PeriodID(SelectedPeriodID)
 	{
 	}
 
@@ -275,31 +100,35 @@ public:
 	virtual EReadyState GetReadyState() override;
 	virtual void PrepareForPlay(const FParamDict& Options) override;
 	virtual TSharedPtrTS<ITimelineMediaAsset> GetMediaAsset() const override;
-	virtual void SelectStream(const TSharedPtrTS<IPlaybackAssetAdaptationSet>& AdaptationSet, const TSharedPtrTS<IPlaybackAssetRepresentation>& Representation) override;
+	virtual void SelectStream(const FString& AdaptationSetID, const FString& RepresentationID) override;
 	virtual IManifest::FResult GetStartingSegment(TSharedPtrTS<IStreamSegment>& OutSegment, const FPlayStartPosition& StartPosition, IManifest::ESearchType SearchType) override;
 	virtual IManifest::FResult GetLoopingSegment(TSharedPtrTS<IStreamSegment>& OutSegment, FPlayerLoopState& InOutLoopState, const TMultiMap<EStreamType, TSharedPtrTS<IStreamSegment>>& InFinishedSegments, const FPlayStartPosition& StartPosition, IManifest::ESearchType SearchType) override;
 	virtual IManifest::FResult GetNextSegment(TSharedPtrTS<IStreamSegment>& OutSegment, TSharedPtrTS<const IStreamSegment> CurrentSegment) override;
 	virtual IManifest::FResult GetRetrySegment(TSharedPtrTS<IStreamSegment>& OutSegment, TSharedPtrTS<const IStreamSegment> CurrentSegment, bool bReplaceWithFillerData) override;
-	virtual void GetSegmentInformation(TArray<FSegmentInformation>& OutSegmentInformation, FTimeValue& OutAverageSegmentDuration, TSharedPtrTS<const IStreamSegment> CurrentSegment, const FTimeValue& LookAheadTime, const TSharedPtrTS<IPlaybackAssetAdaptationSet>& AdaptationSet, const TSharedPtrTS<IPlaybackAssetRepresentation>& Representation) override;
+	virtual void IncreaseSegmentFetchDelay(const FTimeValue& IncreaseAmount) override;
+	virtual void GetSegmentInformation(TArray<FSegmentInformation>& OutSegmentInformation, FTimeValue& OutAverageSegmentDuration, TSharedPtrTS<const IStreamSegment> CurrentSegment, const FTimeValue& LookAheadTime, const FString& AdaptationSetID, const FString& RepresentationID) override;
 
 private:
 	ELECTRA_IMPL_DEFAULT_ERROR_METHODS(DASHManifest);
 
+	TSharedPtrTS<FManifestDASHInternal> GetCurrentManifest() const;
+
 	TSharedPtrTS<IPlaybackAssetRepresentation> GetRepresentationFromAdaptationByMaxBandwidth(TWeakPtrTS<IPlaybackAssetAdaptationSet> InAdaptationSet, int32 NotExceedingBandwidth);
 	IManifest::FResult GetNextOrRetrySegment(TSharedPtrTS<IStreamSegment>& OutSegment, TSharedPtrTS<const IStreamSegment> InCurrentSegment, bool bForRetry);
 
-
 	IPlayerSessionServices* PlayerSessionServices = nullptr;
-	TWeakPtrTS<FManifestDASHInternal::FPeriod> InternalPeriod;
-	TWeakPtrTS<FDASHTimeline> Timeline;
 	EReadyState ReadyState = EReadyState::NotReady;
 	FStreamPreferences StreamPreferences;
 
-	TWeakPtrTS<IPlaybackAssetAdaptationSet> ActiveVideoAdaptationSet;
-	TWeakPtrTS<IPlaybackAssetAdaptationSet> ActiveAudioAdaptationSet;
+	FTimeValue SegmentFetchDelay = FTimeValue::GetZero();
+
+	FString PeriodID;
 
-	TWeakPtrTS<IPlaybackAssetRepresentation> ActiveVideoRepresentation;
-	TWeakPtrTS<IPlaybackAssetRepresentation> ActiveAudioRepresentation;
+	FString ActiveVideoAdaptationSetID;
+	FString ActiveAudioAdaptationSetID;
+
+	FString ActiveVideoRepresentationID;
+	FString ActiveAudioRepresentationID;
 };
 
 /***************************************************************************************************************************************************/
@@ -307,43 +136,31 @@ private:
 /***************************************************************************************************************************************************/
 
 
-TSharedPtrTS<FManifestDASH> FManifestDASH::Create(IPlayerSessionServices* SessionServices, const FParamDict& Options, TWeakPtrTS<FPlaylistReaderDASH> PlaylistReader, TSharedPtrTS<FManifestDASHInternal> Manifest)
+TSharedPtrTS<FManifestDASH> FManifestDASH::Create(IPlayerSessionServices* SessionServices, TSharedPtrTS<FManifestDASHInternal> Manifest)
 {
-	FManifestDASH* m = new FManifestDASH(SessionServices, Options, PlaylistReader, Manifest);
+	FManifestDASH* m = new FManifestDASH(SessionServices, Manifest);
 	return MakeShareable<FManifestDASH>(m);
 }
 
-FManifestDASH::FManifestDASH(IPlayerSessionServices* InSessionServices, const FParamDict& InOptions, TWeakPtrTS<FPlaylistReaderDASH> InPlaylistReader, TSharedPtrTS<FManifestDASHInternal> InManifest)
-	: Options(InOptions)
-	, InternalManifest(InManifest)
-	, PlayerSessionServices(InSessionServices)
-	, PlaylistReader(InPlaylistReader)
+FManifestDASH::FManifestDASH(IPlayerSessionServices* InSessionServices, TSharedPtrTS<FManifestDASHInternal> InManifest)
+	: PlayerSessionServices(InSessionServices)
+	, CurrentManifest(InManifest)
 {
 	bHaveCurrentMetadata = false;
 }
 
-
 FManifestDASH::~FManifestDASH()
 {
 }
 
-void FManifestDASH::UpdateTimeline()
+void FManifestDASH::UpdateInternalManifest(TSharedPtrTS<FManifestDASHInternal> UpdatedManifest)
 {
-	TSharedPtrTS<FManifestDASHInternal> m = InternalManifest.Pin();
-	if (m.IsValid())
-	{
-		CurrentTimeline = MakeSharedTS<FDASHTimeline>(m);
-	}
-	else
-	{
-		CurrentTimeline.Reset();
-	}
+	CurrentManifest = UpdatedManifest;
 }
 
-
 IManifest::EType FManifestDASH::GetPresentationType() const
 {
-	TSharedPtrTS<FManifestDASHInternal> Manifest(InternalManifest.Pin());
+	TSharedPtrTS<FManifestDASHInternal> Manifest(CurrentManifest);
 	if (Manifest.IsValid())
 	{
 		return Manifest->GetPresentationType() == FManifestDASHInternal::EPresentationType::Static ? IManifest::EType::OnDemand : IManifest::EType::Live;
@@ -351,9 +168,43 @@ IManifest::EType FManifestDASH::GetPresentationType() const
 	return IManifest::EType::OnDemand;
 }
 
-TSharedPtrTS<IPlaybackAssetTimeline> FManifestDASH::GetTimeline() const
+FTimeValue FManifestDASH::GetAnchorTime() const
+{
+	TSharedPtrTS<FManifestDASHInternal> Manifest(CurrentManifest);
+	return Manifest.IsValid() ? Manifest->GetAnchorTime() : FTimeValue();
+}
+
+FTimeRange FManifestDASH::GetTotalTimeRange() const
+{
+	TSharedPtrTS<FManifestDASHInternal> Manifest(CurrentManifest);
+	return Manifest.IsValid() ? Manifest->GetTotalTimeRange() : FTimeRange();
+}
+
+FTimeRange FManifestDASH::GetSeekableTimeRange() const
 {
-	return CurrentTimeline;
+	TSharedPtrTS<FManifestDASHInternal> Manifest(CurrentManifest);
+	return Manifest.IsValid() ? Manifest->GetSeekableTimeRange() : FTimeRange();
+}
+
+void FManifestDASH::GetSeekablePositions(TArray<FTimespan>& OutPositions) const
+{
+	TSharedPtrTS<FManifestDASHInternal> Manifest(CurrentManifest);
+	if (Manifest.IsValid())
+	{
+		Manifest->GetSeekablePositions(OutPositions);
+	}
+}
+
+FTimeValue FManifestDASH::GetDuration() const
+{
+	TSharedPtrTS<FManifestDASHInternal> Manifest(CurrentManifest);
+	return Manifest.IsValid() ? Manifest->GetDuration() : FTimeValue();
+}
+
+FTimeValue FManifestDASH::GetDefaultStartTime() const
+{
+	TSharedPtrTS<FManifestDASHInternal> Manifest(CurrentManifest);
+	return Manifest.IsValid() ? Manifest->GetDefaultStartTime() : FTimeValue();
 }
 
 int64 FManifestDASH::GetDefaultStartingBitrate() const
@@ -363,7 +214,7 @@ int64 FManifestDASH::GetDefaultStartingBitrate() const
 
 FTimeValue FManifestDASH::GetMinBufferTime() const
 {
-	TSharedPtrTS<FManifestDASHInternal> Manifest(InternalManifest.Pin());
+	TSharedPtrTS<FManifestDASHInternal> Manifest(CurrentManifest);
 	if (Manifest.IsValid())
 	{
 		TSharedPtrTS<const FDashMPD_MPDType> MPDRoot = Manifest->GetMPDRoot();
@@ -374,61 +225,57 @@ FTimeValue FManifestDASH::GetMinBufferTime() const
 
 void FManifestDASH::GetStreamMetadata(TArray<FStreamMetadata>& OutMetadata, EStreamType StreamType) const
 {
-	if (CurrentTimeline.IsValid())
+	if (!bHaveCurrentMetadata)
 	{
-		if (!bHaveCurrentMetadata)
+		bHaveCurrentMetadata = true;
+		TSharedPtrTS<FManifestDASHInternal> Manifest(CurrentManifest);
+		if (Manifest.IsValid() && Manifest->GetPeriods().Num())
 		{
-			bHaveCurrentMetadata = true;
-			TSharedPtrTS<FManifestDASHInternal> Manifest = CurrentTimeline->GetManifest();
-			check(Manifest.IsValid());
-			if (Manifest->GetPeriods().Num())
+			Manifest->PreparePeriodAdaptationSets(Manifest->GetPeriods()[0], false);
+			const TArray<TSharedPtrTS<FManifestDASHInternal::FAdaptationSet>>& AdaptationSets = Manifest->GetPeriods()[0]->GetAdaptationSets();
+			for(int32 nA=0; nA<AdaptationSets.Num(); ++nA)
 			{
-				Manifest->PreparePeriodAdaptationSets(PlayerSessionServices, Manifest->GetPeriods()[0], false);
-				const TArray<TSharedPtrTS<FManifestDASHInternal::FAdaptationSet>>& AdaptationSets = Manifest->GetPeriods()[0]->GetAdaptationSets();
-				for(int32 nA=0; nA<AdaptationSets.Num(); ++nA)
+				FStreamMetadata sm;
+				sm.CodecInformation = AdaptationSets[nA]->GetCodec();
+				sm.PlaylistID = Manifest->GetPeriods()[0]->GetID();
+				sm.Bandwidth = AdaptationSets[nA]->GetMaxBandwidth();
+				sm.StreamUniqueID = 0;
+				sm.LanguageCode = AdaptationSets[nA]->GetLanguage();
+				switch(AdaptationSets[nA]->GetCodec().GetStreamType())
 				{
-					FStreamMetadata sm;
-					sm.CodecInformation = AdaptationSets[nA]->GetCodec();
-					sm.PlaylistID = Manifest->GetPeriods()[0]->GetID();
-					sm.Bandwidth = AdaptationSets[nA]->GetMaxBandwidth();
-					sm.StreamUniqueID = 0;
-					sm.LanguageCode = AdaptationSets[nA]->GetLanguage();
-					switch(AdaptationSets[nA]->GetCodec().GetStreamType())
+					case EStreamType::Video:
 					{
-						case EStreamType::Video:
-						{
-							CurrentMetadataVideo.Emplace(MoveTemp(sm));
-							break;
-						}
-						case EStreamType::Audio:
-						{
-							CurrentMetadataAudio.Emplace(MoveTemp(sm));
-							break;
-						}
-						case EStreamType::Subtitle:
-						{
-							CurrentMetadataSubtitle.Emplace(MoveTemp(sm));
-							break;
-						}
+						CurrentMetadataVideo.Emplace(MoveTemp(sm));
+						break;
+					}
+					case EStreamType::Audio:
+					{
+						CurrentMetadataAudio.Emplace(MoveTemp(sm));
+						break;
+					}
+					case EStreamType::Subtitle:
+					{
+						CurrentMetadataSubtitle.Emplace(MoveTemp(sm));
+						break;
 					}
 				}
 			}
 		}
+	}
 
-		// At present we return metadata from the first period only as every period can have totally different
-		// number of streams and even codecs. There is no commonality between periods.
-		if (StreamType == EStreamType::Video)
-		{
-			OutMetadata = CurrentMetadataVideo;
-		}
-		else if (StreamType == EStreamType::Audio)
-		{
-			OutMetadata = CurrentMetadataAudio;
-		}
-		else if (StreamType == EStreamType::Subtitle)
-		{
-			OutMetadata = CurrentMetadataSubtitle;
-		}
+	// At present we return metadata from the first period only as every period can have totally different
+	// number of streams and even codecs. There is no commonality between periods.
+	if (StreamType == EStreamType::Video)
+	{
+		OutMetadata = CurrentMetadataVideo;
+	}
+	else if (StreamType == EStreamType::Audio)
+	{
+		OutMetadata = CurrentMetadataAudio;
+	}
+	else if (StreamType == EStreamType::Subtitle)
+	{
+		OutMetadata = CurrentMetadataSubtitle;
 	}
 }
 
@@ -437,7 +284,6 @@ void FManifestDASH::UpdateDynamicRefetchCounter()
 	++CurrentPeriodAndAdaptationXLinkResolveID;
 }
 
-
 IStreamReader* FManifestDASH::CreateStreamReaderHandler()
 {
 	return new FStreamReaderFMP4DASH;
@@ -445,18 +291,16 @@ IStreamReader* FManifestDASH::CreateStreamReaderHandler()
 
 IManifest::FResult FManifestDASH::FindPlayPeriod(TSharedPtrTS<IPlayPeriod>& OutPlayPeriod, const FPlayStartPosition& StartPosition, ESearchType SearchType)
 {
-	TSharedPtrTS<FDASHTimeline> Timeline = CurrentTimeline;
-	// When there is no current timeline we return NotFound.
-	if (!Timeline.IsValid())
+	TSharedPtrTS<FManifestDASHInternal> Manifest(CurrentManifest);
+	if (!Manifest.IsValid())
 	{
-		return IManifest::FResult();
+		return IManifest::FResult(IManifest::FResult::EType::NotLoaded);
 	}
 
-	TSharedPtrTS<FManifestDASHInternal> Manifest = Timeline->GetManifest();
 	FTimeValue StartTime = StartPosition.Time;
 
 	// All time values we communicate to the outside - and therefor get from the outside - are offset by the availabilityStartTime.
-	StartTime -= Timeline->GetAnchorTime();
+	StartTime -= Manifest->GetAnchorTime();
 
 	// Find the period into which the start time falls.
 	TSharedPtrTS<FManifestDASHInternal::FPeriod> SelectedPeriod;
@@ -477,8 +321,6 @@ IManifest::FResult FManifestDASH::FindPlayPeriod(TSharedPtrTS<IPlayPeriod>& OutP
 				if (StartTime < PeriodEndTime)
 				{
 					FTimeValue DiffToNextPeriod = nPeriod + 1 < Periods.Num() && !Periods[nPeriod+1]->GetIsEarlyPeriod() ? Periods[nPeriod + 1]->GetStart() - StartTime : FTimeValue::GetPositiveInfinity();
-					//FTimeValue DiffToPrevPeriod = nPeriod ? StartTime - Periods[nPeriod - 1]->GetEnd() : FTimeValue::GetPositiveInfinity();
-					//FTimeValue DiffToStart = StartTime - Periods[nPeriod]->GetStart();
 					FTimeValue DiffToEnd = PeriodEndTime - StartTime;
 					switch(SearchType)
 					{
@@ -486,7 +328,7 @@ IManifest::FResult FManifestDASH::FindPlayPeriod(TSharedPtrTS<IPlayPeriod>& OutP
 						{
 							// Closest only looks at the distance to the following period. It will never take us to the preceeding period since it is unlikely
 							// the intention (for closest) is to start playback at the end of a period only to transition into this one.
-							if (DiffToEnd.IsValid() && DiffToNextPeriod <= DiffToEnd)
+							if (DiffToEnd.IsValid() && DiffToNextPeriod <= DiffToEnd && nPeriod + 1 < Periods.Num())
 							{
 								SelectedPeriod = Periods[nPeriod + 1];
 							}
@@ -547,7 +389,7 @@ IManifest::FResult FManifestDASH::FindPlayPeriod(TSharedPtrTS<IPlayPeriod>& OutP
 				}
 			}
 			// Wrap the period in an externally accessible interface.
-			TSharedPtrTS<FDASHPlayPeriod> PlayPeriod = MakeSharedTS<FDASHPlayPeriod>(PlayerSessionServices, SelectedPeriod, Timeline);
+			TSharedPtrTS<FDASHPlayPeriod> PlayPeriod = MakeSharedTS<FDASHPlayPeriod>(PlayerSessionServices, SelectedPeriod->GetUniqueIdentifier());
 			OutPlayPeriod = PlayPeriod;
 			return IManifest::FResult(IManifest::FResult::EType::Found);
 		}
@@ -567,6 +409,18 @@ IManifest::FResult FManifestDASH::FindPlayPeriod(TSharedPtrTS<IPlayPeriod>& OutP
 /***************************************************************************************************************************************************/
 /***************************************************************************************************************************************************/
 
+TSharedPtrTS<FManifestDASHInternal> FDASHPlayPeriod::GetCurrentManifest() const
+{
+	TSharedPtrTS<IPlaylistReader> ManifestReader = PlayerSessionServices->GetManifestReader();
+	if (ManifestReader.IsValid())
+	{
+		check(ManifestReader->GetPlaylistType().Equals(TEXT("dash")));
+		IPlaylistReaderDASH* Reader = static_cast<IPlaylistReaderDASH*>(ManifestReader.Get());
+		return Reader->GetCurrentMPD();
+	}
+	return nullptr;
+}
+
 void FDASHPlayPeriod::SetStreamPreferences(const FStreamPreferences& Preferences)
 {
 	StreamPreferences = Preferences;
@@ -579,24 +433,28 @@ IManifest::IPlayPeriod::EReadyState FDASHPlayPeriod::GetReadyState()
 
 void FDASHPlayPeriod::PrepareForPlay(const FParamDict& Options)
 {
-	TSharedPtrTS<FManifestDASHInternal::FPeriod> Period = InternalPeriod.Pin();
+	TSharedPtrTS<FManifestDASHInternal> Manifest = GetCurrentManifest();
+	TSharedPtrTS<FManifestDASHInternal::FPeriod> Period = Manifest.IsValid() ? Manifest->GetPeriodByUniqueID(PeriodID) : nullptr;
 	if (!Period.IsValid())
 	{
-		LogMessage(PlayerSessionServices, IInfoLog::ELevel::Verbose, FString::Printf(TEXT("Period in PrepareToPlay() is no longer valid. Must reselect.")));
+		LogMessage(PlayerSessionServices, IInfoLog::ELevel::Verbose, FString::Printf(TEXT("Period in PrepareToPlay() is not or no longer valid. Must reselect.")));
 		ReadyState = EReadyState::MustReselect;
 		return;
 	}
 
+	// Prepare the adaptation sets and periods.
+	Manifest->PreparePeriodAdaptationSets(Period, false);
+	
 	// We need to select one adaptation set per stream type we wish to play.
 	int32 NumVideoAdaptationSets = Period->GetNumberOfAdaptationSets(EStreamType::Video);
 	if (NumVideoAdaptationSets > 0)
 	{
 		// For now pick the first one!
 		TSharedPtrTS<IPlaybackAssetAdaptationSet> VideoAS = Period->GetAdaptationSetByTypeAndIndex(EStreamType::Video, 0);
-		ActiveVideoAdaptationSet = VideoAS;
+		ActiveVideoAdaptationSetID = VideoAS->GetUniqueIdentifier();
 
 		TSharedPtrTS<IPlaybackAssetRepresentation> VideoRepr = GetRepresentationFromAdaptationByMaxBandwidth(VideoAS, 2*1000*1000);
-		ActiveVideoRepresentation = VideoRepr;
+		ActiveVideoRepresentationID = VideoRepr->GetUniqueIdentifier();
 	}
 
 	int32 NumAudioAdaptationSets = Period->GetNumberOfAdaptationSets(EStreamType::Audio);
@@ -604,10 +462,10 @@ void FDASHPlayPeriod::PrepareForPlay(const FParamDict& Options)
 	{
 		// For now pick the first one!
 		TSharedPtrTS<IPlaybackAssetAdaptationSet> AudioAS = Period->GetAdaptationSetByTypeAndIndex(EStreamType::Audio, 0);
-		ActiveAudioAdaptationSet = AudioAS;
+		ActiveAudioAdaptationSetID = AudioAS->GetUniqueIdentifier();
 
 		TSharedPtrTS<IPlaybackAssetRepresentation> AudioRepr = GetRepresentationFromAdaptationByMaxBandwidth(AudioAS, 128 * 1000);
-		ActiveAudioRepresentation = AudioRepr;
+		ActiveAudioRepresentationID = AudioRepr->GetUniqueIdentifier();
 	}
 
 //	ReadyState = EReadyState::Preparing;
@@ -616,19 +474,28 @@ void FDASHPlayPeriod::PrepareForPlay(const FParamDict& Options)
 
 TSharedPtrTS<ITimelineMediaAsset> FDASHPlayPeriod::GetMediaAsset() const
 {
-	return InternalPeriod.Pin();
+	TSharedPtrTS<FManifestDASHInternal> Manifest = GetCurrentManifest();
+	TSharedPtrTS<FManifestDASHInternal::FPeriod> Period = Manifest.IsValid() ? Manifest->GetPeriodByUniqueID(PeriodID) : nullptr;
+	if (Period.IsValid())
+	{
+		// Returning the asset typically means the caller wants to access the adaptation sets and representations.
+		// Prepare them if necessary (if they already are this method returns immediately).
+		Manifest->PreparePeriodAdaptationSets(Period, false);
+		return Period;
+	}
+	return nullptr;
 }
 
-void FDASHPlayPeriod::SelectStream(const TSharedPtrTS<IPlaybackAssetAdaptationSet>& AdaptationSet, const TSharedPtrTS<IPlaybackAssetRepresentation>& Representation)
+void FDASHPlayPeriod::SelectStream(const FString& AdaptationSetID, const FString& RepresentationID)
 {
 	// The ABR must not try to switch adaptation sets at the moment. As such the adaptation set passed in must be one of the already active ones.
-	if (AdaptationSet == ActiveVideoAdaptationSet)
+	if (AdaptationSetID == ActiveVideoAdaptationSetID)
 	{
-		ActiveVideoRepresentation = Representation;
+		ActiveVideoRepresentationID = RepresentationID;
 	}
-	else if (AdaptationSet == ActiveAudioAdaptationSet)
+	else if (AdaptationSetID == ActiveAudioAdaptationSetID)
 	{
-		ActiveAudioRepresentation = Representation;
+		ActiveAudioRepresentationID = RepresentationID;
 	}
 	else
 	{
@@ -638,20 +505,23 @@ void FDASHPlayPeriod::SelectStream(const TSharedPtrTS<IPlaybackAssetAdaptationSe
 
 IManifest::FResult FDASHPlayPeriod::GetStartingSegment(TSharedPtrTS<IStreamSegment>& OutSegment, const FPlayStartPosition& StartPosition, IManifest::ESearchType SearchType)
 {
-	TSharedPtrTS<FManifestDASHInternal::FPeriod> Period = InternalPeriod.Pin();
-	if (!Period.IsValid())
+	TSharedPtrTS<FManifestDASHInternal> Manifest = GetCurrentManifest();
+	if (!Manifest.IsValid())
 	{
-		return IManifest::FResult(IManifest::FResult::EType::NotLoaded).SetErrorDetail(FErrorDetail().SetMessage("Period to locate start segment in has disappeared"));
+		return IManifest::FResult(IManifest::FResult::EType::NotLoaded).SetErrorDetail(FErrorDetail().SetMessage("The manifest to locate the period start segment in has disappeared"));
 	}
-	TSharedPtrTS<FDASHTimeline> TL = Timeline.Pin();
-	if (!TL.IsValid())
+	TSharedPtrTS<FManifestDASHInternal::FPeriod> Period = Manifest->GetPeriodByUniqueID(PeriodID);
+	if (!Period.IsValid())
 	{
-		return IManifest::FResult(IManifest::FResult::EType::NotLoaded).SetErrorDetail(FErrorDetail().SetMessage("Timeline the period to locate start segment in has disappeared"));
+		// If the period has suddenly disappeared there must have been an MPD update that removed it. This is extremely rare but possible.
+		return IManifest::FResult(IManifest::FResult::EType::NotFound).SetErrorDetail(FErrorDetail().SetMessage("Period to locate start segment in has disappeared"));
 	}
+	Manifest->PreparePeriodAdaptationSets(Period, false);
 
+	FTimeValue AST = Manifest->GetAnchorTime();
 	FTimeValue StartTime = StartPosition.Time;
 	// All time values we communicate to the outside - and therefor get from the outside - are offset by the availabilityStartTime.
-	StartTime -= TL->GetAnchorTime();
+	StartTime -= AST;
 
 	// Due to the way we have been searching for the period it is possible for the start time to fall (slightly) outside the actual times.
 	if (StartTime < Period->GetStart())
@@ -665,45 +535,56 @@ IManifest::FResult FDASHPlayPeriod::GetStartingSegment(TSharedPtrTS<IStreamSegme
 	// We are searching for a time local to the period so we need to subtract the period start time.
 	StartTime -= Period->GetStart();
 
+	bool bUsesAST = Manifest->UsesAST();
+	bool bIsStaticType = Manifest->IsStaticType() || Manifest->IsEventType();
 
 	// Create a segment request to which the individual stream segment requests will add themselves as
 	// dependent streams. This is a special case for playback start.
 	TSharedPtrTS<FStreamSegmentRequestFMP4DASH> StartSegmentRequest = MakeSharedTS<FStreamSegmentRequestFMP4DASH>();
 	StartSegmentRequest->bIsInitialStartRequest = true;
 
-
 	struct FSelectedStream
 	{
-		EStreamType										StreamType;
-		TSharedPtrTS<IPlaybackAssetRepresentation>		Representation;
-		TSharedPtrTS<IPlaybackAssetAdaptationSet>		AdaptationSet;
+		EStreamType	StreamType;
+		FString		RepresentationID;
+		FString		AdaptationSetID;
 	};
 	TArray<FSelectedStream> ActiveSelection;
-	if (ActiveVideoRepresentation.IsValid())
+	if (!ActiveVideoRepresentationID.IsEmpty())
 	{
-		ActiveSelection.Emplace(FSelectedStream({EStreamType::Video, ActiveVideoRepresentation.Pin(), ActiveVideoAdaptationSet.Pin()}));
+		ActiveSelection.Emplace(FSelectedStream({EStreamType::Video, ActiveVideoRepresentationID, ActiveVideoAdaptationSetID}));
 	}
-	if (ActiveAudioRepresentation.IsValid())
+	if (!ActiveAudioRepresentationID.IsEmpty())
 	{
-		ActiveSelection.Emplace(FSelectedStream({EStreamType::Audio, ActiveAudioRepresentation.Pin(), ActiveAudioAdaptationSet.Pin()}));
+		ActiveSelection.Emplace(FSelectedStream({EStreamType::Audio, ActiveAudioRepresentationID, ActiveAudioAdaptationSetID}));
 	}
 
-
 	bool bDidAdjustStartTime = false;
 	bool bTryAgainLater = false;
 	bool bAnyStreamAtEOS = false;
 	bool bAllStreamsAtEOS = true;
 	for(int32 i=0; i<ActiveSelection.Num(); ++i)
 	{
-		if (ActiveSelection[i].AdaptationSet.IsValid() && ActiveSelection[i].Representation.IsValid())
+		if (ActiveSelection[i].AdaptationSetID.Len() && ActiveSelection[i].RepresentationID.Len())
 		{
-			FManifestDASHInternal::FRepresentation* Repr = static_cast<FManifestDASHInternal::FRepresentation*>(ActiveSelection[i].Representation.Get());
+			TSharedPtrTS<FManifestDASHInternal::FAdaptationSet> Adapt = Period->GetAdaptationSetByUniqueID(ActiveSelection[i].AdaptationSetID);
+			TSharedPtrTS<FManifestDASHInternal::FRepresentation> Repr = Adapt.IsValid()? Adapt->GetRepresentationByUniqueID(ActiveSelection[i].RepresentationID) : nullptr;
+			if (!Repr.IsValid())
+			{
+				// If the AdaptationSet or the Representation has suddenly disappeared there must have been an MPD update that removed it, which is illegal because the
+				// Period itself is still there (checked for above).
+				return IManifest::FResult(IManifest::FResult::EType::NotFound).SetErrorDetail(FErrorDetail().SetMessage("Period no longer has the selected AdaptationSet or Representation."));
+			}
 
 			FManifestDASHInternal::FSegmentInformation SegmentInfo;
 			FManifestDASHInternal::FSegmentSearchOption SearchOpt;
 			TArray<TWeakPtrTS<FMPDLoadRequestDASH>> RemoteElementLoadRequests;
 			SearchOpt.PeriodLocalTime = StartTime;
 			SearchOpt.PeriodDuration = Period->GetDuration();
+			if (!SearchOpt.PeriodDuration.IsValid() || SearchOpt.PeriodDuration.IsPositiveInfinity())
+			{
+				SearchOpt.PeriodDuration = Manifest->GetLastPeriodEndTime() - AST;
+			}
 			SearchOpt.SearchType = SearchType;
 			FManifestDASHInternal::FRepresentation::ESearchResult SearchResult = Repr->FindSegment(PlayerSessionServices, SegmentInfo, RemoteElementLoadRequests, SearchOpt);
 			if (SearchResult == FManifestDASHInternal::FRepresentation::ESearchResult::NeedElement)
@@ -722,11 +603,16 @@ IManifest::FResult FDASHPlayPeriod::GetStartingSegment(TSharedPtrTS<IStreamSegme
 			{
 				TSharedPtrTS<FStreamSegmentRequestFMP4DASH> SegmentRequest = MakeSharedTS<FStreamSegmentRequestFMP4DASH>();
 				SegmentRequest->StreamType = ActiveSelection[i].StreamType;
-				SegmentRequest->Representation = ActiveSelection[i].Representation;
-				SegmentRequest->AdaptationSet = ActiveSelection[i].AdaptationSet;
+				SegmentRequest->Representation = Repr;
+				SegmentRequest->AdaptationSet = Adapt;
 				SegmentRequest->Period = Period;
 				SegmentRequest->PeriodStart = Period->GetStart();
-				SegmentRequest->AvailabilityStartTime = TL->GetAnchorTime();
+				SegmentRequest->AST = AST;
+				if (bUsesAST)
+				{
+					SegmentRequest->ASAST = SegmentInfo.CalculateASAST(AST, Period->GetStart(), bIsStaticType) + SegmentFetchDelay;
+					SegmentRequest->SAET = SegmentInfo.CalculateSAET(AST, Period->GetStart(), Manifest->GetAvailabilityEndTime(), Manifest->GetTimeshiftBufferDepth(), bIsStaticType);
+				}
 				SegmentRequest->Segment = MoveTemp(SegmentInfo);
 				SegmentRequest->bIsEOSSegment = true;
 				StartSegmentRequest->DependentStreams.Emplace(MoveTemp(SegmentRequest));
@@ -766,11 +652,16 @@ IManifest::FResult FDASHPlayPeriod::GetStartingSegment(TSharedPtrTS<IStreamSegme
 
 				TSharedPtrTS<FStreamSegmentRequestFMP4DASH> SegmentRequest = MakeSharedTS<FStreamSegmentRequestFMP4DASH>();
 				SegmentRequest->StreamType = ActiveSelection[i].StreamType;
-				SegmentRequest->Representation = ActiveSelection[i].Representation;
-				SegmentRequest->AdaptationSet = ActiveSelection[i].AdaptationSet;
+				SegmentRequest->Representation = Repr;
+				SegmentRequest->AdaptationSet = Adapt;
 				SegmentRequest->Period = Period;
 				SegmentRequest->PeriodStart = Period->GetStart();
-				SegmentRequest->AvailabilityStartTime = TL->GetAnchorTime();
+				SegmentRequest->AST = AST;
+				if (bUsesAST)
+				{
+					SegmentRequest->ASAST = SegmentInfo.CalculateASAST(AST, Period->GetStart(), bIsStaticType) + SegmentFetchDelay;
+					SegmentRequest->SAET = SegmentInfo.CalculateSAET(AST, Period->GetStart(), Manifest->GetAvailabilityEndTime(), Manifest->GetTimeshiftBufferDepth(), bIsStaticType);
+				}
 				// If the segment is known to be missing we need to instead insert filler data.
 				if (SegmentInfo.bIsMissing)
 				{
@@ -782,19 +673,25 @@ IManifest::FResult FDASHPlayPeriod::GetStartingSegment(TSharedPtrTS<IStreamSegme
 				// the playback position to. If not valid yet update it with the current stream values.
 				if (!StartSegmentRequest->GetFirstPTS().IsValid())
 				{
-					StartSegmentRequest->AvailabilityStartTime = SegmentRequest->AvailabilityStartTime;
+					StartSegmentRequest->AST = SegmentRequest->AST;
 					StartSegmentRequest->AdditionalAdjustmentTime = SegmentRequest->AdditionalAdjustmentTime;
 					StartSegmentRequest->PeriodStart = SegmentRequest->PeriodStart;
 					StartSegmentRequest->Segment = SegmentRequest->Segment;
 				}
 
+				// Similarly the start segment request might need to look at a segment availability window.
+				if (bUsesAST && !StartSegmentRequest->ASAST.IsValid())
+				{
+					StartSegmentRequest->ASAST = SegmentRequest->ASAST;
+					StartSegmentRequest->SAET = SegmentRequest->SAET;
+				}
+
 				StartSegmentRequest->DependentStreams.Emplace(MoveTemp(SegmentRequest));
 				bAllStreamsAtEOS = false;
 			}
 			else
 			{
-				check(!"Unhandled search result!");
-				return IManifest::FResult();
+				return IManifest::FResult(IManifest::FResult::EType::NotFound).SetErrorDetail(FErrorDetail().SetMessage("Unhandled search result!"));
 			}
 		}
 	}
@@ -811,7 +708,7 @@ IManifest::FResult FDASHPlayPeriod::GetStartingSegment(TSharedPtrTS<IStreamSegme
 		return IManifest::FResult(IManifest::FResult::EType::PastEOS);
 	}
 
-	// Done. Note that there is no 'NotFound', 'BeforeStart' or 'NotLoaded'.
+	// Done.
 	OutSegment = MoveTemp(StartSegmentRequest);
 	return IManifest::FResult(IManifest::FResult::EType::Found);
 }
@@ -819,39 +716,42 @@ IManifest::FResult FDASHPlayPeriod::GetStartingSegment(TSharedPtrTS<IStreamSegme
 
 IManifest::FResult FDASHPlayPeriod::GetNextOrRetrySegment(TSharedPtrTS<IStreamSegment>& OutSegment, TSharedPtrTS<const IStreamSegment> InCurrentSegment, bool bForRetry)
 {
-	// Make sure we still got the same MPD and everything.
-	// These will likely become invalid when the MPD updates!
-	check(Timeline.Pin());
-	check(InternalPeriod.Pin());
-	//
-	check(InCurrentSegment.IsValid());
 	TSharedPtrTS<const FStreamSegmentRequestFMP4DASH> Current = StaticCastSharedPtr<const FStreamSegmentRequestFMP4DASH>(InCurrentSegment);
-	check(Current->bIsInitialStartRequest == false);
+	if (Current->bIsInitialStartRequest)
+	{
+		return IManifest::FResult(IManifest::FResult::EType::NotFound).SetErrorDetail(FErrorDetail().SetMessage("The next segment cannot be located for the initial start request, only for an actual media request!"));
+	}
 
-	TSharedPtrTS<IPlaybackAssetRepresentation> ActiveRepresentationByType;
-	TSharedPtrTS<IPlaybackAssetAdaptationSet> ActiveAdaptationSetByType;
+	TSharedPtrTS<FManifestDASHInternal> Manifest = GetCurrentManifest();
+	TSharedPtrTS<FManifestDASHInternal::FPeriod> Period = Manifest->GetPeriodByUniqueID(PeriodID);
+	if (!Period.IsValid())
+	{
+		return IManifest::FResult(IManifest::FResult::EType::NotLoaded).SetErrorDetail(FErrorDetail().SetMessage("Period to locate start segment in has disappeared"));
+	}
+	Manifest->PreparePeriodAdaptationSets(Period, false);
+
+	FString ActiveRepresentationIDByType;
+	FString ActiveAdaptationSetIDByType;
 	switch(Current->GetType())
 	{
 		case EStreamType::Video:
-			ActiveRepresentationByType = ActiveVideoRepresentation.Pin();
-			ActiveAdaptationSetByType = ActiveVideoAdaptationSet.Pin();
+			ActiveAdaptationSetIDByType = ActiveVideoAdaptationSetID;
+			ActiveRepresentationIDByType = ActiveVideoRepresentationID;
 			break;
 		case EStreamType::Audio:
-			ActiveRepresentationByType = ActiveAudioRepresentation.Pin();
-			ActiveAdaptationSetByType = ActiveAudioAdaptationSet.Pin();
+			ActiveAdaptationSetIDByType = ActiveAudioAdaptationSetID;
+			ActiveRepresentationIDByType = ActiveAudioRepresentationID;
 			break;
 		default:
 			break;
 	}
-	if (!ActiveAdaptationSetByType.IsValid() || !ActiveRepresentationByType.IsValid())
+	TSharedPtrTS<FManifestDASHInternal::FAdaptationSet> Adapt = Period->GetAdaptationSetByUniqueID(ActiveAdaptationSetIDByType);
+	TSharedPtrTS<FManifestDASHInternal::FRepresentation> Repr = Adapt.IsValid()? Adapt->GetRepresentationByUniqueID(ActiveRepresentationIDByType) : nullptr;
+	if (!Repr.IsValid())
 	{
 		return IManifest::FResult(IManifest::FResult::EType::NotFound).SetErrorDetail(FErrorDetail().SetMessage("No active stream found to get next segment for"));
 	}
 
-	TSharedPtrTS<FDASHTimeline> TL = Timeline.Pin();
-	TSharedPtrTS<FManifestDASHInternal::FPeriod> Period = InternalPeriod.Pin();
-
-	FManifestDASHInternal::FRepresentation* Repr = static_cast<FManifestDASHInternal::FRepresentation*>(ActiveRepresentationByType.Get());
 	FManifestDASHInternal::FSegmentInformation SegmentInfo;
 	FManifestDASHInternal::FSegmentSearchOption SearchOpt;
 	TArray<TWeakPtrTS<FMPDLoadRequestDASH>> RemoteElementLoadRequests;
@@ -870,6 +770,14 @@ IManifest::FResult FDASHPlayPeriod::GetNextOrRetrySegment(TSharedPtrTS<IStreamSe
 		SearchOpt.PeriodDuration = Period->GetDuration();
 		SearchOpt.SearchType = IManifest::ESearchType::Closest;
 	}
+	FTimeValue AST = Manifest->GetAnchorTime();
+	bool bUsesAST = Manifest->UsesAST();
+	bool bIsStaticType = Manifest->IsStaticType() || Manifest->IsEventType();
+	if (!SearchOpt.PeriodDuration.IsValid() || SearchOpt.PeriodDuration.IsPositiveInfinity())
+	{
+		SearchOpt.PeriodDuration = Manifest->GetLastPeriodEndTime() - AST;
+	}
+
 	FManifestDASHInternal::FRepresentation::ESearchResult SearchResult = Repr->FindSegment(PlayerSessionServices, SegmentInfo, RemoteElementLoadRequests, SearchOpt);
 	if (SearchResult == FManifestDASHInternal::FRepresentation::ESearchResult::NeedElement)
 	{
@@ -885,7 +793,24 @@ IManifest::FResult FDASHPlayPeriod::GetNextOrRetrySegment(TSharedPtrTS<IStreamSe
 	}
 	else if (SearchResult == FManifestDASHInternal::FRepresentation::ESearchResult::PastEOS)
 	{
-		return IManifest::FResult(IManifest::FResult::EType::PastEOS);
+		// We may have reached the end of a period or the last segment of an ongoing Live presentation.
+		// Need to figure out which it is.
+		if (!Manifest->AreUpdatesExpected())
+		{
+			// No updates of the manifest means this period is over. Try moving onto the next if there is one.
+			return IManifest::FResult(IManifest::FResult::EType::PastEOS);
+		}
+		// Could be the end or a period. Is there a regular period following?
+		if (Manifest->HasFollowingRegularPeriod(Period))
+		{
+			// Yes, so we can move onto the next period.
+			return IManifest::FResult(IManifest::FResult::EType::PastEOS);
+		}
+
+		IPlaylistReaderDASH* Reader = static_cast<IPlaylistReaderDASH*>(PlayerSessionServices->GetManifestReader().Get());
+		check(Reader->GetPlaylistType().Equals(TEXT("dash")));
+		Reader->RequestMPDUpdate(false);
+		return IManifest::FResult().RetryAfterMilliseconds(250);
 	}
 	else if (SearchResult == FManifestDASHInternal::FRepresentation::ESearchResult::Gone)
 	{
@@ -903,11 +828,16 @@ IManifest::FResult FDASHPlayPeriod::GetNextOrRetrySegment(TSharedPtrTS<IStreamSe
 		TSharedPtrTS<FStreamSegmentRequestFMP4DASH> SegmentRequest = MakeSharedTS<FStreamSegmentRequestFMP4DASH>();
 		SegmentRequest->PlayerLoopState = Current->PlayerLoopState;
 		SegmentRequest->StreamType = Current->GetType();
-		SegmentRequest->Representation = ActiveRepresentationByType;
-		SegmentRequest->AdaptationSet = ActiveAdaptationSetByType;
+		SegmentRequest->Representation = Repr;
+		SegmentRequest->AdaptationSet = Adapt;
 		SegmentRequest->Period = Period;
 		SegmentRequest->PeriodStart = Period->GetStart();
-		SegmentRequest->AvailabilityStartTime = TL->GetAnchorTime();
+		SegmentRequest->AST = Manifest->GetAnchorTime();
+		if (bUsesAST)
+		{
+			SegmentRequest->ASAST = SegmentInfo.CalculateASAST(AST, Period->GetStart(), bIsStaticType) + SegmentFetchDelay;
+			SegmentRequest->SAET = SegmentInfo.CalculateSAET(AST, Period->GetStart(), Manifest->GetAvailabilityEndTime(), Manifest->GetTimeshiftBufferDepth(), bIsStaticType);
+		}
 		// If the segment is known to be missing we need to instead insert filler data.
 		if (SegmentInfo.bIsMissing)
 		{
@@ -927,20 +857,29 @@ IManifest::FResult FDASHPlayPeriod::GetNextOrRetrySegment(TSharedPtrTS<IStreamSe
 			SegmentRequest->NumOverallRetries = Current->NumOverallRetries + 1;
 		}
 
+		// If we stayed on the same representation and the stream reader has already warned about a timescale
+		// mismatch then we take on the warning flag to reduce console spam.
+		if (SegmentRequest->Representation == Current->Representation)
+		{
+			SegmentRequest->bWarnedAboutTimescale = Current->bWarnedAboutTimescale;
+		}
+
 		OutSegment = MoveTemp(SegmentRequest);
 		return IManifest::FResult(IManifest::FResult::EType::Found);
 	}
 	else
 	{
-		check(!"Unhandled search result!");
-		return IManifest::FResult();
+		return IManifest::FResult(IManifest::FResult::EType::NotFound).SetErrorDetail(FErrorDetail().SetMessage("Unhandled search result!"));
 	}
 }
 
 
 IManifest::FResult FDASHPlayPeriod::GetNextSegment(TSharedPtrTS<IStreamSegment>& OutSegment, TSharedPtrTS<const IStreamSegment> InCurrentSegment)
 {
-	check(InCurrentSegment.IsValid());
+	if (!InCurrentSegment.IsValid())
+	{
+		return IManifest::FResult(IManifest::FResult::EType::NotFound).SetErrorDetail(FErrorDetail().SetMessage("There is no current segment to locate the next one for!"));
+	}
 	// Did the stream reader see a 'lmsg' brand on this segment?
 	// If so then this stream has ended and there will not be a next segment.
 	const FStreamSegmentRequestFMP4DASH* CurrentRequest = static_cast<const FStreamSegmentRequestFMP4DASH*>(InCurrentSegment.Get());
@@ -953,6 +892,10 @@ IManifest::FResult FDASHPlayPeriod::GetNextSegment(TSharedPtrTS<IStreamSegment>&
 
 IManifest::FResult FDASHPlayPeriod::GetRetrySegment(TSharedPtrTS<IStreamSegment>& OutSegment, TSharedPtrTS<const IStreamSegment> InCurrentSegment, bool bReplaceWithFillerData)
 {
+	if (!InCurrentSegment.IsValid())
+	{
+		return IManifest::FResult(IManifest::FResult::EType::NotFound).SetErrorDetail(FErrorDetail().SetMessage("There is no current segment to locate a retry segment for!"));
+	}
 	// To insert filler data we can use the current request over again.
 	if (bReplaceWithFillerData)
 	{
@@ -965,25 +908,45 @@ IManifest::FResult FDASHPlayPeriod::GetRetrySegment(TSharedPtrTS<IStreamSegment>
 		OutSegment = NewRequest;
 		return IManifest::FResult(IManifest::FResult::EType::Found);
 	}
-
 	return GetNextOrRetrySegment(OutSegment, InCurrentSegment, true);
 }
 
 
 IManifest::FResult FDASHPlayPeriod::GetLoopingSegment(TSharedPtrTS<IStreamSegment>& OutSegment, FPlayerLoopState& InOutLoopState, const TMultiMap<EStreamType, TSharedPtrTS<IStreamSegment>>& InFinishedSegments, const FPlayStartPosition& StartPosition, IManifest::ESearchType SearchType)
 {
-	check(!"TODO");
-	return IManifest::FResult();
+	return IManifest::FResult(IManifest::FResult::EType::NotFound).SetErrorDetail(FErrorDetail().SetMessage("Looping is not currently supported!"));
+}
+
+
+void FDASHPlayPeriod::IncreaseSegmentFetchDelay(const FTimeValue& IncreaseAmount)
+{
+// TODO: when this starts to get too large we should trigger a resync of our clock with the server time.
+	SegmentFetchDelay += IncreaseAmount;
 }
 
 
-void FDASHPlayPeriod::GetSegmentInformation(TArray<FSegmentInformation>& OutSegmentInformation, FTimeValue& OutAverageSegmentDuration, TSharedPtrTS<const IStreamSegment> CurrentSegment, const FTimeValue& LookAheadTime, const TSharedPtrTS<IPlaybackAssetAdaptationSet>& AdaptationSet, const TSharedPtrTS<IPlaybackAssetRepresentation>& Representation)
+void FDASHPlayPeriod::GetSegmentInformation(TArray<FSegmentInformation>& OutSegmentInformation, FTimeValue& OutAverageSegmentDuration, TSharedPtrTS<const IStreamSegment> CurrentSegment, const FTimeValue& LookAheadTime, const FString& AdaptationSetID, const FString& RepresentationID)
 {
 	OutSegmentInformation.Empty();
 	OutAverageSegmentDuration.SetToInvalid();
-	if (Representation.IsValid() && AdaptationSet.IsValid())
+
+	TSharedPtrTS<FManifestDASHInternal> Manifest = GetCurrentManifest();
+	if (Manifest.IsValid())
 	{
-		static_cast<FManifestDASHInternal::FRepresentation*>(Representation.Get())->GetSegmentInformation(OutSegmentInformation, OutAverageSegmentDuration, CurrentSegment, LookAheadTime, AdaptationSet);
+		TSharedPtrTS<FManifestDASHInternal::FPeriod> Period = Manifest->GetPeriodByUniqueID(PeriodID);
+		if (Period.IsValid())
+		{
+			Manifest->PreparePeriodAdaptationSets(Period, false);
+			TSharedPtrTS<FManifestDASHInternal::FAdaptationSet> AdaptationSet = Period->GetAdaptationSetByUniqueID(AdaptationSetID);
+			if (AdaptationSet.IsValid())
+			{
+				TSharedPtrTS<FManifestDASHInternal::FRepresentation> Repr = AdaptationSet->GetRepresentationByUniqueID(RepresentationID);
+				if (Repr.IsValid())
+				{
+					Repr->GetSegmentInformation(OutSegmentInformation, OutAverageSegmentDuration, CurrentSegment, LookAheadTime, AdaptationSet);
+				}
+			}
+		}
 	}
 }
 
@@ -1061,11 +1024,39 @@ T GetAttribute(const TArray<TSharedPtrTS<FDashMPD_SegmentTemplateType>>& Arr, Ge
 
 #define GET_ATTR(Array, GetVal, IsValid, Default) GetAttribute(Array, [](const auto& e){return e->GetVal;}, [](const auto& v){return v.IsValid;}, Default)
 
+
+FTimeValue CalculateSegmentAvailabilityTimeOffset(const TArray<TSharedPtrTS<FDashMPD_SegmentBaseType>>& Arr)
+{
+	FTimeValue Sum(FTimeValue::GetZero());
+	for(int32 i=0; i<Arr.Num(); ++i)
+	{
+		FTimeValue v = Arr[i]->GetAvailabilityTimeOffset();
+		if (v.IsValid())
+		{
+			Sum += v;
+		}
+	}
+	return Sum;
+}
+FTimeValue CalculateSegmentAvailabilityTimeOffset(const TArray<TSharedPtrTS<FDashMPD_SegmentTemplateType>>& Arr)
+{
+	FTimeValue Sum(FTimeValue::GetZero());
+	for(int32 i=0; i<Arr.Num(); ++i)
+	{
+		FTimeValue v = Arr[i]->GetAvailabilityTimeOffset();
+		if (v.IsValid())
+		{
+			Sum += v;
+		}
+	}
+	return Sum;
+}
+
 }
 
 
 
-FManifestDASHInternal::FRepresentation::ESearchResult FManifestDASHInternal::FRepresentation::PrepareSegmentIndex(IPlayerSessionServices* PlayerSessionServices, const TArray<TSharedPtrTS<FDashMPD_SegmentBaseType>>& SegmentBase, TArray<TWeakPtrTS<FMPDLoadRequestDASH>>& OutRemoteElementLoadRequests)
+FManifestDASHInternal::FRepresentation::ESearchResult FManifestDASHInternal::FRepresentation::PrepareSegmentIndex(IPlayerSessionServices* InPlayerSessionServices, const TArray<TSharedPtrTS<FDashMPD_SegmentBaseType>>& SegmentBase, TArray<TWeakPtrTS<FMPDLoadRequestDASH>>& OutRemoteElementLoadRequests)
 {
 	// If the segment index has been requested and is still pending, return right away.
 	if (PendingSegmentIndexLoadRequest.IsValid())
@@ -1092,29 +1083,40 @@ FManifestDASHInternal::FRepresentation::ESearchResult FManifestDASHInternal::FRe
 				IndexRange = RepresentationIndex->GetRange();
 				if (!RepresentationIndex->GetSourceURL().IsEmpty())
 				{
-					LogMessage(PlayerSessionServices, IInfoLog::ELevel::Warning, FString::Printf(TEXT("Ignoring <RepresentationIndex> element URL present for <SegmentBase> of representation \"%s\""), *MPDRepresentation->GetID()));
+					LogMessage(InPlayerSessionServices, IInfoLog::ELevel::Warning, FString::Printf(TEXT("Ignoring <RepresentationIndex> element URL present for <SegmentBase> of representation \"%s\""), *MPDRepresentation->GetID()));
 				}
 			}
 		}
 		if (!IndexRange.IsEmpty())
 		{
+			IElectraHttpManager::FParams::FRange r;
+			r.Set(IndexRange);
+			if (r.IsSet())
+			{
+				check(r.GetStart() >= 0);
+				check(r.GetEndIncluding() > 0);
+				SegmentIndexRangeStart = r.GetStart();
+				SegmentIndexRangeSize = r.GetEndIncluding() + 1 - SegmentIndexRangeStart;
+			}
+
 			FString PreferredServiceLocation;
-			DashUtils::GetPlayerOption(PlayerSessionServices, PreferredServiceLocation, DASH::OptionKey_CurrentCDN, FString());
+			DashUtils::GetPlayerOption(InPlayerSessionServices, PreferredServiceLocation, DASH::OptionKey_CurrentCDN, FString());
 			// Get the relevant <BaseURL> elements from the hierarchy.
 			TArray<TSharedPtrTS<const FDashMPD_BaseURLType>> OutBaseURLs;
-			DASHUrlHelpers::GetAllHierarchyBaseURLs(PlayerSessionServices, OutBaseURLs, MPDRepresentation, *PreferredServiceLocation);
+			DASHUrlHelpers::GetAllHierarchyBaseURLs(InPlayerSessionServices, OutBaseURLs, MPDRepresentation, *PreferredServiceLocation);
 			if (OutBaseURLs.Num() == 0)
 			{
-				PostError(PlayerSessionServices, FString::Printf(TEXT("Representation \"%s\" has no <BaseURL> element on any hierarchy level!"), *MPDRepresentation->GetID()), ERRCODE_DASH_MPD_BAD_REPRESENTATION);
+				PostError(InPlayerSessionServices, FString::Printf(TEXT("Representation \"%s\" has no <BaseURL> element on any hierarchy level!"), *MPDRepresentation->GetID()), ERRCODE_DASH_MPD_BAD_REPRESENTATION);
 				bIsUsable = false;
 				return FManifestDASHInternal::FRepresentation::ESearchResult::BadType;
 			}
 			// Generate the absolute URL
 			FString DocumentURL = MPDRepresentation->GetDocumentURL();
 			FString URL, RequestHeader;
-			if (!DASHUrlHelpers::BuildAbsoluteElementURL(URL, DocumentURL, OutBaseURLs, FString()))
+			FTimeValue UrlATO;
+			if (!DASHUrlHelpers::BuildAbsoluteElementURL(URL, UrlATO, DocumentURL, OutBaseURLs, FString()))
 			{
-				PostError(PlayerSessionServices, FString::Printf(TEXT("Representation \"%s\" failed to resolve URL to an absolute URL!"), *MPDRepresentation->GetID()), ERRCODE_DASH_MPD_BAD_REPRESENTATION);
+				PostError(InPlayerSessionServices, FString::Printf(TEXT("Representation \"%s\" failed to resolve URL to an absolute URL!"), *MPDRepresentation->GetID()), ERRCODE_DASH_MPD_BAD_REPRESENTATION);
 				bIsUsable = false;
 				return FManifestDASHInternal::FRepresentation::ESearchResult::BadType;
 			}
@@ -1122,10 +1124,10 @@ FManifestDASHInternal::FRepresentation::ESearchResult FManifestDASHInternal::FRe
 			// The URL query might need to be changed. Look for the UrlQuery properties.
 			TArray<TSharedPtrTS<FDashMPD_UrlQueryInfoType>> UrlQueries;
 			DASHUrlHelpers::GetAllHierarchyUrlQueries(UrlQueries, MPDRepresentation, DASHUrlHelpers::EUrlQueryRequestType::Segment, true);
-			FErrorDetail Error = DASHUrlHelpers::ApplyUrlQueries(PlayerSessionServices, DocumentURL, URL, RequestHeader, UrlQueries);
+			FErrorDetail Error = DASHUrlHelpers::ApplyUrlQueries(InPlayerSessionServices, DocumentURL, URL, RequestHeader, UrlQueries);
 			if (Error.IsSet())
 			{
-				PostError(PlayerSessionServices, Error);
+				PostError(InPlayerSessionServices, Error);
 				bIsUsable = false;
 				return FManifestDASHInternal::FRepresentation::ESearchResult::BadType;
 			}
@@ -1133,7 +1135,7 @@ FManifestDASHInternal::FRepresentation::ESearchResult FManifestDASHInternal::FRe
 
 			// Check with entity cache if the index has been retrieved before.
 			IPlayerEntityCache::FCacheItem CachedItem;
-			if (PlayerSessionServices->GetEntityCache()->GetCachedEntity(CachedItem, URL, IndexRange))
+			if (InPlayerSessionServices->GetEntityCache()->GetCachedEntity(CachedItem, URL, IndexRange))
 			{
 				// Already cached. Use it.
 				SegmentIndex = CachedItem.Parsed14496_12Data;
@@ -1147,9 +1149,9 @@ FManifestDASHInternal::FRepresentation::ESearchResult FManifestDASHInternal::FRe
 				LoadReq->Range = IndexRange;
 				if (RequestHeader.Len())
 				{
-					LoadReq->Headers.Emplace(HTTP::FHTTPHeader({TEXT("MPEG-DASH-Param"), RequestHeader}));
+					LoadReq->Headers.Emplace(HTTP::FHTTPHeader({DASH::HTTPHeaderOptionName, RequestHeader}));
 				}
-				LoadReq->PlayerSessionServices = PlayerSessionServices;
+				LoadReq->PlayerSessionServices = InPlayerSessionServices;
 				LoadReq->XLinkElement = MPDRepresentation;
 				LoadReq->CompleteCallback.BindThreadSafeSP(AsShared(), &FManifestDASHInternal::FRepresentation::SegmentIndexDownloadComplete);
 				OutRemoteElementLoadRequests.Emplace(LoadReq);
@@ -1162,7 +1164,7 @@ FManifestDASHInternal::FRepresentation::ESearchResult FManifestDASHInternal::FRe
 }
 
 
-bool FManifestDASHInternal::FRepresentation::PrepareDownloadURLs(IPlayerSessionServices* PlayerSessionServices, FSegmentInformation& InOutSegmentInfo, const TArray<TSharedPtrTS<FDashMPD_SegmentBaseType>>& SegmentBase)
+bool FManifestDASHInternal::FRepresentation::PrepareDownloadURLs(IPlayerSessionServices* InPlayerSessionServices, FSegmentInformation& InOutSegmentInfo, const TArray<TSharedPtrTS<FDashMPD_SegmentBaseType>>& SegmentBase)
 {
 	TSharedPtrTS<FDashMPD_RepresentationType> MPDRepresentation = Representation.Pin();
 
@@ -1173,36 +1175,37 @@ bool FManifestDASHInternal::FRepresentation::PrepareDownloadURLs(IPlayerSessionS
 		InOutSegmentInfo.InitializationURL.Range = Initialization->GetRange();
 		if (!Initialization->GetSourceURL().IsEmpty())
 		{
-			LogMessage(PlayerSessionServices, IInfoLog::ELevel::Warning, FString::Printf(TEXT("Ignoring <Initialization> element URL present for <SegmentBase> of representation \"%s\""), *MPDRepresentation->GetID()));
+			LogMessage(InPlayerSessionServices, IInfoLog::ELevel::Warning, FString::Printf(TEXT("Ignoring <Initialization> element URL present for <SegmentBase> of representation \"%s\""), *MPDRepresentation->GetID()));
 		}
 	}
 
 	FString PreferredServiceLocation;
-	DashUtils::GetPlayerOption(PlayerSessionServices, PreferredServiceLocation, DASH::OptionKey_CurrentCDN, FString());
+	DashUtils::GetPlayerOption(InPlayerSessionServices, PreferredServiceLocation, DASH::OptionKey_CurrentCDN, FString());
 	// Get the relevant <BaseURL> elements from the hierarchy.
 	TArray<TSharedPtrTS<const FDashMPD_BaseURLType>> OutBaseURLs;
-	DASHUrlHelpers::GetAllHierarchyBaseURLs(PlayerSessionServices, OutBaseURLs, MPDRepresentation, *PreferredServiceLocation);
+	DASHUrlHelpers::GetAllHierarchyBaseURLs(InPlayerSessionServices, OutBaseURLs, MPDRepresentation, *PreferredServiceLocation);
 	if (OutBaseURLs.Num() == 0)
 	{
-		PostError(PlayerSessionServices, FString::Printf(TEXT("Representation \"%s\" has no <BaseURL> element on any hierarchy level!"), *MPDRepresentation->GetID()), ERRCODE_DASH_MPD_BAD_REPRESENTATION);
+		PostError(InPlayerSessionServices, FString::Printf(TEXT("Representation \"%s\" has no <BaseURL> element on any hierarchy level!"), *MPDRepresentation->GetID()), ERRCODE_DASH_MPD_BAD_REPRESENTATION);
 		return false;
 	}
 	// Generate the absolute URL
 	FString DocumentURL = MPDRepresentation->GetDocumentURL();
 	FString URL, RequestHeader;
-	if (!DASHUrlHelpers::BuildAbsoluteElementURL(URL, DocumentURL, OutBaseURLs, FString()))
+	FTimeValue UrlATO;
+	if (!DASHUrlHelpers::BuildAbsoluteElementURL(URL, UrlATO, DocumentURL, OutBaseURLs, FString()))
 	{
-		PostError(PlayerSessionServices, FString::Printf(TEXT("Representation \"%s\" failed to resolve URL to an absolute URL!"), *MPDRepresentation->GetID()), ERRCODE_DASH_MPD_BAD_REPRESENTATION);
+		PostError(InPlayerSessionServices, FString::Printf(TEXT("Representation \"%s\" failed to resolve URL to an absolute URL!"), *MPDRepresentation->GetID()), ERRCODE_DASH_MPD_BAD_REPRESENTATION);
 		return false;
 	}
 
 	// The URL query might need to be changed. Look for the UrlQuery properties.
 	TArray<TSharedPtrTS<FDashMPD_UrlQueryInfoType>> UrlQueries;
 	DASHUrlHelpers::GetAllHierarchyUrlQueries(UrlQueries, MPDRepresentation, DASHUrlHelpers::EUrlQueryRequestType::Segment, true);
-	FErrorDetail Error = DASHUrlHelpers::ApplyUrlQueries(PlayerSessionServices, DocumentURL, URL, RequestHeader, UrlQueries);
+	FErrorDetail Error = DASHUrlHelpers::ApplyUrlQueries(InPlayerSessionServices, DocumentURL, URL, RequestHeader, UrlQueries);
 	if (Error.IsSet())
 	{
-		PostError(PlayerSessionServices, Error);
+		PostError(InPlayerSessionServices, Error);
 		return false;
 	}
 
@@ -1224,11 +1227,13 @@ bool FManifestDASHInternal::FRepresentation::PrepareDownloadURLs(IPlayerSessionS
 	InOutSegmentInfo.MediaURL.CDN = PreferredServiceLocation;
 	InOutSegmentInfo.MediaURL.CustomHeader = RequestHeader;
 
+	// Add any availabilityTimeOffset from the <BaseURL>.
+	InOutSegmentInfo.ATO += UrlATO;
 
 	return true;
 }
 
-bool FManifestDASHInternal::FRepresentation::PrepareDownloadURLs(IPlayerSessionServices* PlayerSessionServices, FSegmentInformation& InOutSegmentInfo, const TArray<TSharedPtrTS<FDashMPD_SegmentTemplateType>>& SegmentTemplate)
+bool FManifestDASHInternal::FRepresentation::PrepareDownloadURLs(IPlayerSessionServices* InPlayerSessionServices, FSegmentInformation& InOutSegmentInfo, const TArray<TSharedPtrTS<FDashMPD_SegmentTemplateType>>& SegmentTemplate)
 {
 	TSharedPtrTS<FDashMPD_RepresentationType> MPDRepresentation = Representation.Pin();
 
@@ -1237,7 +1242,7 @@ bool FManifestDASHInternal::FRepresentation::PrepareDownloadURLs(IPlayerSessionS
 	FString MediaTemplate = GET_ATTR(SegmentTemplate, GetMediaTemplate(), Len(), FString());
 	if (MediaTemplate.Len() == 0)
 	{
-		PostError(PlayerSessionServices, FString::Printf(TEXT("Representation \"%s\" provides no media template!"), *MPDRepresentation->GetID()), ERRCODE_DASH_MPD_BAD_REPRESENTATION);
+		PostError(InPlayerSessionServices, FString::Printf(TEXT("Representation \"%s\" provides no media template!"), *MPDRepresentation->GetID()), ERRCODE_DASH_MPD_BAD_REPRESENTATION);
 		return false;
 	}
 	// Get the initialization template string. If this is not specified try any <Initialization> elements.
@@ -1250,7 +1255,7 @@ bool FManifestDASHInternal::FRepresentation::PrepareDownloadURLs(IPlayerSessionS
 			InOutSegmentInfo.InitializationURL.Range = Initialization->GetRange();
 			if (Initialization->GetSourceURL().IsEmpty())
 			{
-				PostError(PlayerSessionServices, FString::Printf(TEXT("Representation \"%s\" provides no initialization segment!"), *MPDRepresentation->GetID()), ERRCODE_DASH_MPD_BAD_REPRESENTATION);
+				PostError(InPlayerSessionServices, FString::Printf(TEXT("Representation \"%s\" provides no initialization segment!"), *MPDRepresentation->GetID()), ERRCODE_DASH_MPD_BAD_REPRESENTATION);
 				return false;
 			}
 			// Note: This URL should probably not be using any template strings but I can't find any evidence for this, so just treat it as a template string as well.
@@ -1264,10 +1269,10 @@ bool FManifestDASHInternal::FRepresentation::PrepareDownloadURLs(IPlayerSessionS
 
 	// Get the preferred CDN and the <BaseURL> and <UrlQueryInfo> elements affecting the URL assembly.
 	FString PreferredServiceLocation;
-	DashUtils::GetPlayerOption(PlayerSessionServices, PreferredServiceLocation, DASH::OptionKey_CurrentCDN, FString());
+	DashUtils::GetPlayerOption(InPlayerSessionServices, PreferredServiceLocation, DASH::OptionKey_CurrentCDN, FString());
 	// Get the relevant <BaseURL> elements from the hierarchy.
 	TArray<TSharedPtrTS<const FDashMPD_BaseURLType>> OutBaseURLs;
-	DASHUrlHelpers::GetAllHierarchyBaseURLs(PlayerSessionServices, OutBaseURLs, MPDRepresentation, *PreferredServiceLocation);
+	DASHUrlHelpers::GetAllHierarchyBaseURLs(InPlayerSessionServices, OutBaseURLs, MPDRepresentation, *PreferredServiceLocation);
 	// The URL query might need to be changed. Look for the UrlQuery properties.
 	TArray<TSharedPtrTS<FDashMPD_UrlQueryInfoType>> UrlQueries;
 	DASHUrlHelpers::GetAllHierarchyUrlQueries(UrlQueries, MPDRepresentation, DASHUrlHelpers::EUrlQueryRequestType::Segment, true);
@@ -1275,29 +1280,31 @@ bool FManifestDASHInternal::FRepresentation::PrepareDownloadURLs(IPlayerSessionS
 	// Generate the absolute media URL
 	FString DocumentURL = MPDRepresentation->GetDocumentURL();
 	FString MediaURL, MediaRequestHeader;
-	if (!DASHUrlHelpers::BuildAbsoluteElementURL(MediaURL, DocumentURL, OutBaseURLs, MediaTemplateURL))
+	FTimeValue MediaUrlATO;
+	if (!DASHUrlHelpers::BuildAbsoluteElementURL(MediaURL, MediaUrlATO, DocumentURL, OutBaseURLs, MediaTemplateURL))
 	{
-		PostError(PlayerSessionServices, FString::Printf(TEXT("Representation \"%s\" failed to resolve media segment URL to an absolute URL!"), *MPDRepresentation->GetID()), ERRCODE_DASH_MPD_BAD_REPRESENTATION);
+		PostError(InPlayerSessionServices, FString::Printf(TEXT("Representation \"%s\" failed to resolve media segment URL to an absolute URL!"), *MPDRepresentation->GetID()), ERRCODE_DASH_MPD_BAD_REPRESENTATION);
 		return false;
 	}
-	FErrorDetail Error = DASHUrlHelpers::ApplyUrlQueries(PlayerSessionServices, DocumentURL, MediaURL, MediaRequestHeader, UrlQueries);
+	FErrorDetail Error = DASHUrlHelpers::ApplyUrlQueries(InPlayerSessionServices, DocumentURL, MediaURL, MediaRequestHeader, UrlQueries);
 	if (Error.IsSet())
 	{
-		PostError(PlayerSessionServices, Error);
+		PostError(InPlayerSessionServices, Error);
 		return false;
 	}
 
 	// And also generate the absolute init segment URL
 	FString InitURL, InitRequestHeader;
-	if (!DASHUrlHelpers::BuildAbsoluteElementURL(InitURL, DocumentURL, OutBaseURLs, InitTemplateURL))
+	FTimeValue InitUrlATO;
+	if (!DASHUrlHelpers::BuildAbsoluteElementURL(InitURL, InitUrlATO, DocumentURL, OutBaseURLs, InitTemplateURL))
 	{
-		PostError(PlayerSessionServices, FString::Printf(TEXT("Representation \"%s\" failed to resolve init segment URL to an absolute URL!"), *MPDRepresentation->GetID()), ERRCODE_DASH_MPD_BAD_REPRESENTATION);
+		PostError(InPlayerSessionServices, FString::Printf(TEXT("Representation \"%s\" failed to resolve init segment URL to an absolute URL!"), *MPDRepresentation->GetID()), ERRCODE_DASH_MPD_BAD_REPRESENTATION);
 		return false;
 	}
-	Error = DASHUrlHelpers::ApplyUrlQueries(PlayerSessionServices, DocumentURL, InitURL, InitRequestHeader, UrlQueries);
+	Error = DASHUrlHelpers::ApplyUrlQueries(InPlayerSessionServices, DocumentURL, InitURL, InitRequestHeader, UrlQueries);
 	if (Error.IsSet())
 	{
-		PostError(PlayerSessionServices, Error);
+		PostError(InPlayerSessionServices, Error);
 		return false;
 	}
 
@@ -1316,6 +1323,10 @@ bool FManifestDASHInternal::FRepresentation::PrepareDownloadURLs(IPlayerSessionS
 	InOutSegmentInfo.MediaURL.CDN = PreferredServiceLocation;
 	InOutSegmentInfo.MediaURL.CustomHeader = MediaRequestHeader;
 
+	// Add any availabilityTimeOffset from the <BaseURL>. We use the media ATO since we request the init segment at the same time
+	// as the media segment and the init segment needs to be available at the same time anyway.
+	InOutSegmentInfo.ATO += MediaUrlATO;
+
 	return true;
 }
 
@@ -1415,7 +1426,7 @@ FString FManifestDASHInternal::FRepresentation::ApplyTemplateStrings(FString Tem
 
 
 
-FManifestDASHInternal::FRepresentation::ESearchResult FManifestDASHInternal::FRepresentation::FindSegment(IPlayerSessionServices* PlayerSessionServices, FSegmentInformation& OutSegmentInfo, TArray<TWeakPtrTS<FMPDLoadRequestDASH>>& OutRemoteElementLoadRequests, const FSegmentSearchOption& InSearchOptions)
+FManifestDASHInternal::FRepresentation::ESearchResult FManifestDASHInternal::FRepresentation::FindSegment(IPlayerSessionServices* InPlayerSessionServices, FSegmentInformation& OutSegmentInfo, TArray<TWeakPtrTS<FMPDLoadRequestDASH>>& OutRemoteElementLoadRequests, const FSegmentSearchOption& InSearchOptions)
 {
 	/*
 		Note: We use the DASH-IF-IOP specification and timing model. This is more strict than the general DASH standard and removes ambiguities
@@ -1460,7 +1471,7 @@ FManifestDASHInternal::FRepresentation::ESearchResult FManifestDASHInternal::FRe
 	// On representation level there can be at most one of the others.
 	if (SegmentBase[0].IsValid() && SegmentTemplate[0].IsValid())
 	{
-		PostError(PlayerSessionServices, FString::Printf(TEXT("Representation \"%s\" must have only one of <SegmentBase> or <SegmentTemplate>!"), *MPDRepresentation->GetID()), ERRCODE_DASH_MPD_BAD_REPRESENTATION);
+		PostError(InPlayerSessionServices, FString::Printf(TEXT("Representation \"%s\" must have only one of <SegmentBase> or <SegmentTemplate>!"), *MPDRepresentation->GetID()), ERRCODE_DASH_MPD_BAD_REPRESENTATION);
 		bIsUsable = false;
 		return FManifestDASHInternal::FRepresentation::ESearchResult::BadType;
 	}
@@ -1470,7 +1481,7 @@ FManifestDASHInternal::FRepresentation::ESearchResult FManifestDASHInternal::FRe
 		// Again, there can be at most one of the others.
 		if (SegmentBase[1].IsValid() && SegmentTemplate[1].IsValid())
 		{
-			PostError(PlayerSessionServices, FString::Printf(TEXT("Representation \"%s\" must only inherit one of <SegmentBase> or <SegmentTemplate> from enclosing AdaptationSet!"), *MPDRepresentation->GetID()), ERRCODE_DASH_MPD_BAD_REPRESENTATION);
+			PostError(InPlayerSessionServices, FString::Printf(TEXT("Representation \"%s\" must only inherit one of <SegmentBase> or <SegmentTemplate> from enclosing AdaptationSet!"), *MPDRepresentation->GetID()), ERRCODE_DASH_MPD_BAD_REPRESENTATION);
 			bIsUsable = false;
 			return FManifestDASHInternal::FRepresentation::ESearchResult::BadType;
 		}
@@ -1480,7 +1491,7 @@ FManifestDASHInternal::FRepresentation::ESearchResult FManifestDASHInternal::FRe
 			// Again, there can be at most one of the others.
 			if (SegmentBase[2].IsValid() && SegmentTemplate[2].IsValid())
 			{
-				PostError(PlayerSessionServices, FString::Printf(TEXT("Representation \"%s\" must only inherit one of <SegmentBase> or <SegmentTemplate> from enclosing Period!"), *MPDRepresentation->GetID()), ERRCODE_DASH_MPD_BAD_REPRESENTATION);
+				PostError(InPlayerSessionServices, FString::Printf(TEXT("Representation \"%s\" must only inherit one of <SegmentBase> or <SegmentTemplate> from enclosing Period!"), *MPDRepresentation->GetID()), ERRCODE_DASH_MPD_BAD_REPRESENTATION);
 				bIsUsable = false;
 				return FManifestDASHInternal::FRepresentation::ESearchResult::BadType;
 			}
@@ -1492,14 +1503,14 @@ FManifestDASHInternal::FRepresentation::ESearchResult FManifestDASHInternal::FRe
 	// Nothing? Bad MPD.
 	if (SegmentBase.Num() == 0 && SegmentTemplate.Num() == 0)
 	{
-		PostError(PlayerSessionServices, FString::Printf(TEXT("Representation \"%s\" does not have one of <SegmentBase> or <SegmentTemplate> anywhere in the MPD hierarchy!"), *MPDRepresentation->GetID()), ERRCODE_DASH_MPD_BAD_REPRESENTATION);
+		PostError(InPlayerSessionServices, FString::Printf(TEXT("Representation \"%s\" does not have one of <SegmentBase> or <SegmentTemplate> anywhere in the MPD hierarchy!"), *MPDRepresentation->GetID()), ERRCODE_DASH_MPD_BAD_REPRESENTATION);
 		bIsUsable = false;
 		return FManifestDASHInternal::FRepresentation::ESearchResult::BadType;
 	}
 
 	if (SegmentBase.Num())
 	{
-		return FindSegment_Base(PlayerSessionServices, OutSegmentInfo, OutRemoteElementLoadRequests, InSearchOptions, MPDRepresentation, SegmentBase);
+		return FindSegment_Base(InPlayerSessionServices, OutSegmentInfo, OutRemoteElementLoadRequests, InSearchOptions, MPDRepresentation, SegmentBase);
 	}
 	else
 	{
@@ -1507,26 +1518,26 @@ FManifestDASHInternal::FRepresentation::ESearchResult FManifestDASHInternal::FRe
 		TSharedPtrTS<FDashMPD_SegmentTimelineType> SegmentTimeline = GET_ATTR(SegmentTemplate, GetSegmentTimeline(), IsValid(), TSharedPtrTS<FDashMPD_SegmentTimelineType>());
 		if (SegmentTimeline.IsValid())
 		{
-			return FindSegment_Timeline(PlayerSessionServices, OutSegmentInfo, OutRemoteElementLoadRequests, InSearchOptions, MPDRepresentation, SegmentTemplate, SegmentTimeline);
+			return FindSegment_Timeline(InPlayerSessionServices, OutSegmentInfo, OutRemoteElementLoadRequests, InSearchOptions, MPDRepresentation, SegmentTemplate, SegmentTimeline);
 		}
 		else
 		{
-			return FindSegment_Template(PlayerSessionServices, OutSegmentInfo, OutRemoteElementLoadRequests, InSearchOptions, MPDRepresentation, SegmentTemplate);
+			return FindSegment_Template(InPlayerSessionServices, OutSegmentInfo, OutRemoteElementLoadRequests, InSearchOptions, MPDRepresentation, SegmentTemplate);
 		}
 	}
 }
 
 
-FManifestDASHInternal::FRepresentation::ESearchResult FManifestDASHInternal::FRepresentation::FindSegment_Base(IPlayerSessionServices* PlayerSessionServices, FSegmentInformation& OutSegmentInfo, TArray<TWeakPtrTS<FMPDLoadRequestDASH>>& OutRemoteElementLoadRequests, const FSegmentSearchOption& InSearchOptions, const TSharedPtrTS<FDashMPD_RepresentationType>& MPDRepresentation, const TArray<TSharedPtrTS<FDashMPD_SegmentBaseType>>& SegmentBase)
+FManifestDASHInternal::FRepresentation::ESearchResult FManifestDASHInternal::FRepresentation::FindSegment_Base(IPlayerSessionServices* InPlayerSessionServices, FSegmentInformation& OutSegmentInfo, TArray<TWeakPtrTS<FMPDLoadRequestDASH>>& OutRemoteElementLoadRequests, const FSegmentSearchOption& InSearchOptions, const TSharedPtrTS<FDashMPD_RepresentationType>& MPDRepresentation, const TArray<TSharedPtrTS<FDashMPD_SegmentBaseType>>& SegmentBase)
 {
-	ESearchResult SegIndexResult = PrepareSegmentIndex(PlayerSessionServices, SegmentBase, OutRemoteElementLoadRequests);
+	ESearchResult SegIndexResult = PrepareSegmentIndex(InPlayerSessionServices, SegmentBase, OutRemoteElementLoadRequests);
 	if (SegIndexResult != ESearchResult::Found)
 	{
 		return SegIndexResult;
 	}
 	if (!SegmentIndex.IsValid())
 	{
-		PostError(PlayerSessionServices, FString::Printf(TEXT("A segment index is required for Representation \"%s\""), *MPDRepresentation->GetID()), ERRCODE_DASH_MPD_BAD_REPRESENTATION);
+		PostError(InPlayerSessionServices, FString::Printf(TEXT("A segment index is required for Representation \"%s\""), *MPDRepresentation->GetID()), ERRCODE_DASH_MPD_BAD_REPRESENTATION);
 		bIsUsable = false;
 		return FManifestDASHInternal::FRepresentation::ESearchResult::BadType;
 	}
@@ -1535,7 +1546,7 @@ FManifestDASHInternal::FRepresentation::ESearchResult FManifestDASHInternal::FRe
 	uint32 SidxTimescale = Sidx->GetTimescale();
 	if (SidxTimescale == 0)
 	{
-		PostError(PlayerSessionServices, FString::Printf(TEXT("Timescale of segment index for Representation \"%s\" is invalid!"), *MPDRepresentation->GetID()), ERRCODE_DASH_MPD_BAD_REPRESENTATION);
+		PostError(InPlayerSessionServices, FString::Printf(TEXT("Timescale of segment index for Representation \"%s\" is invalid!"), *MPDRepresentation->GetID()), ERRCODE_DASH_MPD_BAD_REPRESENTATION);
 		bIsUsable = false;
 		return FManifestDASHInternal::FRepresentation::ESearchResult::BadType;
 	}
@@ -1545,10 +1556,12 @@ FManifestDASHInternal::FRepresentation::ESearchResult FManifestDASHInternal::FRe
 	uint64 PTO = GET_ATTR(SegmentBase, GetPresentationTimeOffset(), IsSet(), TMediaOptionalValue<uint64>(0)).Value();
 	uint32 MPDTimescale = GET_ATTR(SegmentBase, GetTimescale(), IsSet(), TMediaOptionalValue<uint32>(1)).Value();
 	// Since the PTO is specified in the timescale as given in the MPD the timescales of the MPD and the segment index should better match!
-	if (PTO && MPDTimescale != SidxTimescale)
+	if (PTO && MPDTimescale != SidxTimescale && !bWarnedAboutTimescale)
 	{
-		LogMessage(PlayerSessionServices, IInfoLog::ELevel::Warning, FString::Printf(TEXT("Representation timescale (%u) in MPD is not equal to timescale used in the segment index (%u) for Representation \"%s\"."), MPDTimescale, SidxTimescale, *MPDRepresentation->GetID()));
+		bWarnedAboutTimescale = true;
+		LogMessage(InPlayerSessionServices, IInfoLog::ELevel::Warning, FString::Printf(TEXT("Representation timescale (%u) in MPD is not equal to timescale used in the segment index (%u) for Representation \"%s\"."), MPDTimescale, SidxTimescale, *MPDRepresentation->GetID()));
 	}
+	FTimeValue ATO = CalculateSegmentAvailabilityTimeOffset(SegmentBase);
 
 	// Convert the local media search time to the timescale of the segment index.
 	// The PTO (presentation time offset) which maps the internal media time to the zero point of the period must be included as well.
@@ -1580,29 +1593,30 @@ FManifestDASHInternal::FRepresentation::ESearchResult FManifestDASHInternal::FRe
 	for(CurrentN=0; CurrentN<EndNumber; ++CurrentN)
 	{
 		const IParserISO14496_12::ISegmentIndex::FEntry& SegmentInfo = Sidx->GetEntry(CurrentN);
+
 		// We do not support hierarchical segment indices!
 		if (SegmentInfo.IsReferenceType)
 		{
-			PostError(PlayerSessionServices, FString::Printf(TEXT("Segment index for Representation \"%s\" must directly reference the media, not another index!"), *MPDRepresentation->GetID()), ERRCODE_DASH_MPD_BAD_REPRESENTATION);
+			PostError(InPlayerSessionServices, FString::Printf(TEXT("Segment index for Representation \"%s\" must directly reference the media, not another index!"), *MPDRepresentation->GetID()), ERRCODE_DASH_MPD_BAD_REPRESENTATION);
 			bIsUsable = false;
 			return FManifestDASHInternal::FRepresentation::ESearchResult::BadType;
 		}
 		else if (SegmentInfo.StartsWithSAP == 0)
 		{
-			PostError(PlayerSessionServices, FString::Printf(TEXT("Segment index for Representation \"%s\" must have starts_with_sap set!"), *MPDRepresentation->GetID()), ERRCODE_DASH_MPD_BAD_REPRESENTATION);
+			PostError(InPlayerSessionServices, FString::Printf(TEXT("Segment index for Representation \"%s\" must have starts_with_sap set!"), *MPDRepresentation->GetID()), ERRCODE_DASH_MPD_BAD_REPRESENTATION);
 			bIsUsable = false;
 			return FManifestDASHInternal::FRepresentation::ESearchResult::BadType;
 		}
 		// We require segments to begin with SAP type 1 or 2 (preferably 1)
 		else if (SegmentInfo.SAPType != 1 && SegmentInfo.SAPType != 2)
 		{
-			PostError(PlayerSessionServices, FString::Printf(TEXT("Segment index for Representation \"%s\" must have SAP_type 1 or 2 only!"), *MPDRepresentation->GetID()), ERRCODE_DASH_MPD_BAD_REPRESENTATION);
+			PostError(InPlayerSessionServices, FString::Printf(TEXT("Segment index for Representation \"%s\" must have SAP_type 1 or 2 only!"), *MPDRepresentation->GetID()), ERRCODE_DASH_MPD_BAD_REPRESENTATION);
 			bIsUsable = false;
 			return FManifestDASHInternal::FRepresentation::ESearchResult::BadType;
 		}
 		else if (SegmentInfo.SAPDeltaTime != 0)
 		{
-			LogMessage(PlayerSessionServices, IInfoLog::ELevel::Warning, FString::Printf(TEXT("Segment index for Representation \"%s\" uses SAP_delta_time. This may result in incorrect decoding."), *MPDRepresentation->GetID()));
+			LogMessage(InPlayerSessionServices, IInfoLog::ELevel::Warning, FString::Printf(TEXT("Segment index for Representation \"%s\" uses SAP_delta_time. This may result in incorrect decoding."), *MPDRepresentation->GetID()));
 			bIsUsable = false;
 			return FManifestDASHInternal::FRepresentation::ESearchResult::BadType;
 		}
@@ -1670,7 +1684,7 @@ FManifestDASHInternal::FRepresentation::ESearchResult FManifestDASHInternal::FRe
 			}
 			else
 			{
-				PostError(PlayerSessionServices, FString::Printf(TEXT("Unsupported segment search mode!")), ERRCODE_DASH_MPD_INTERNAL);
+				PostError(InPlayerSessionServices, FString::Printf(TEXT("Unsupported segment search mode!")), ERRCODE_DASH_MPD_INTERNAL);
 				bIsUsable = false;
 				return FManifestDASHInternal::FRepresentation::ESearchResult::BadType;
 			}
@@ -1695,7 +1709,8 @@ FManifestDASHInternal::FRepresentation::ESearchResult FManifestDASHInternal::FRe
 		OutSegmentInfo.MediaLocalFirstAUTime = MediaLocalSearchTime;
 		OutSegmentInfo.MediaLocalLastAUTime = MediaLocalPeriodEnd;
 		OutSegmentInfo.Timescale = SidxTimescale;
-		if (!PrepareDownloadURLs(PlayerSessionServices, OutSegmentInfo, SegmentBase))
+		OutSegmentInfo.ATO = ATO;
+		if (!PrepareDownloadURLs(InPlayerSessionServices, OutSegmentInfo, SegmentBase))
 		{
 			bIsUsable = false;
 			return FManifestDASHInternal::FRepresentation::ESearchResult::BadType;
@@ -1712,32 +1727,35 @@ FManifestDASHInternal::FRepresentation::ESearchResult FManifestDASHInternal::FRe
 
 }
 
-FManifestDASHInternal::FRepresentation::ESearchResult FManifestDASHInternal::FRepresentation::FindSegment_Template(IPlayerSessionServices* PlayerSessionServices, FSegmentInformation& OutSegmentInfo, TArray<TWeakPtrTS<FMPDLoadRequestDASH>>& OutRemoteElementLoadRequests, const FSegmentSearchOption& InSearchOptions, const TSharedPtrTS<FDashMPD_RepresentationType>& MPDRepresentation, const TArray<TSharedPtrTS<FDashMPD_SegmentTemplateType>>& SegmentTemplate)
+FManifestDASHInternal::FRepresentation::ESearchResult FManifestDASHInternal::FRepresentation::FindSegment_Template(IPlayerSessionServices* InPlayerSessionServices, FSegmentInformation& OutSegmentInfo, TArray<TWeakPtrTS<FMPDLoadRequestDASH>>& OutRemoteElementLoadRequests, const FSegmentSearchOption& InSearchOptions, const TSharedPtrTS<FDashMPD_RepresentationType>& MPDRepresentation, const TArray<TSharedPtrTS<FDashMPD_SegmentTemplateType>>& SegmentTemplate)
 {
 	uint64 PTO = GET_ATTR(SegmentTemplate, GetPresentationTimeOffset(), IsSet(), TMediaOptionalValue<uint64>(0)).Value();
 	uint32 MPDTimescale = GET_ATTR(SegmentTemplate, GetTimescale(), IsSet(), TMediaOptionalValue<uint32>(1)).Value();
 	uint32 StartNumber = GET_ATTR(SegmentTemplate, GetStartNumber(), IsSet(), TMediaOptionalValue<uint32>(1)).Value();
-	uint32 EndNumber = GET_ATTR(SegmentTemplate, GetEndNumber(), IsSet(), TMediaOptionalValue<uint32>(~0U)).Value();
+	TMediaOptionalValue<uint32> EndNumber = GET_ATTR(SegmentTemplate, GetEndNumber(), IsSet(), TMediaOptionalValue<uint32>());
 	TMediaOptionalValue<uint32> Duration = GET_ATTR(SegmentTemplate, GetDuration(), IsSet(), TMediaOptionalValue<uint32>());
 	TMediaOptionalValue<int32> EptDelta = GET_ATTR(SegmentTemplate, GetEptDelta(), IsSet(), TMediaOptionalValue<int32>());
+	FTimeValue ATO = CalculateSegmentAvailabilityTimeOffset(SegmentTemplate);
+
 
 	// The timescale should in all likelihood not be 1. While certainly allowed an accuracy of only one second is more likely to be
 	// an oversight when building the MPD.
 	if (MPDTimescale == 0)
 	{
-		PostError(PlayerSessionServices, FString::Printf(TEXT("Timescale for Representation \"%s\" is invalid!"), *MPDRepresentation->GetID()), ERRCODE_DASH_MPD_BAD_REPRESENTATION);
+		PostError(InPlayerSessionServices, FString::Printf(TEXT("Timescale for Representation \"%s\" is invalid!"), *MPDRepresentation->GetID()), ERRCODE_DASH_MPD_BAD_REPRESENTATION);
 		bIsUsable = false;
 		return FManifestDASHInternal::FRepresentation::ESearchResult::BadType;
 	}
-	else if (MPDTimescale == 1)
+	else if (MPDTimescale == 1 && !bWarnedAboutTimescale)
 	{
-		LogMessage(PlayerSessionServices, IInfoLog::ELevel::Warning, FString::Printf(TEXT("Timescale for Representation \"%s\" is given as 1. Is this intended?"), *MPDRepresentation->GetID()));
+		bWarnedAboutTimescale = true;
+		LogMessage(InPlayerSessionServices, IInfoLog::ELevel::Warning, FString::Printf(TEXT("Timescale for Representation \"%s\" is given as 1. Is this intended?"), *MPDRepresentation->GetID()));
 	}
 
 	// There needs to be a segment duration here.
 	if (!Duration.IsSet() || Duration.Value() == 0)
 	{
-		PostError(PlayerSessionServices, FString::Printf(TEXT("Representation \"%s\" has no valid segment duration!"), *MPDRepresentation->GetID()), ERRCODE_DASH_MPD_BAD_REPRESENTATION);
+		PostError(InPlayerSessionServices, FString::Printf(TEXT("Representation \"%s\" has no valid segment duration!"), *MPDRepresentation->GetID()), ERRCODE_DASH_MPD_BAD_REPRESENTATION);
 		bIsUsable = false;
 		return FManifestDASHInternal::FRepresentation::ESearchResult::BadType;
 	}
@@ -1754,13 +1772,26 @@ FManifestDASHInternal::FRepresentation::ESearchResult FManifestDASHInternal::FRe
 		MediaLocalSearchTime = 0;
 	}
 
-	int64 MediaLocalPeriodDuration = InSearchOptions.PeriodDuration.IsValid() && !InSearchOptions.PeriodDuration.IsInfinity() ? InSearchOptions.PeriodDuration.GetAsTimebase(MPDTimescale) - EPTdelta : TNumericLimits<int64>::Max();
-	uint32 MaxSegmentsInPeriod = (MediaLocalPeriodDuration + SegmentDuration - 1) / SegmentDuration;
+	int64 MediaLocalPeriodDuration;
+	uint32 MaxSegmentsInPeriod;
+	if (InSearchOptions.PeriodDuration.IsValid() && !InSearchOptions.PeriodDuration.IsInfinity())
+	{
+		MediaLocalPeriodDuration = InSearchOptions.PeriodDuration.GetAsTimebase(MPDTimescale) - EPTdelta;
+		MaxSegmentsInPeriod = (MediaLocalPeriodDuration + SegmentDuration - 1) / SegmentDuration;
+	}
+	else
+	{
+		MediaLocalPeriodDuration = TNumericLimits<int64>::Max() - PTO;
+		MaxSegmentsInPeriod = TNumericLimits<uint32>::Max();
+	}
 	// Clamp against the number of segments described by EndNumber.
 	// The assumption is that end number is inclusive, so @startNumber == @endNumber means there is 1 segment.
-	if (MaxSegmentsInPeriod > (EndNumber - StartNumber)+1)
+	if (EndNumber.IsSet())
 	{
-		MaxSegmentsInPeriod = EndNumber - StartNumber + 1;
+		if ((int64)MaxSegmentsInPeriod > (int64)EndNumber.Value() - StartNumber + 1)
+		{
+			MaxSegmentsInPeriod = (uint32)((int64)EndNumber.Value() - StartNumber + 1);
+		}
 	}
 
 	// Now we calculate the number of the segment the search time falls into.
@@ -1810,7 +1841,7 @@ FManifestDASHInternal::FRepresentation::ESearchResult FManifestDASHInternal::FRe
 	}
 	else
 	{
-		PostError(PlayerSessionServices, FString::Printf(TEXT("Unsupported segment search mode!")), ERRCODE_DASH_MPD_INTERNAL);
+		PostError(InPlayerSessionServices, FString::Printf(TEXT("Unsupported segment search mode!")), ERRCODE_DASH_MPD_INTERNAL);
 		bIsUsable = false;
 		return FManifestDASHInternal::FRepresentation::ESearchResult::BadType;
 	}
@@ -1821,7 +1852,7 @@ FManifestDASHInternal::FRepresentation::ESearchResult FManifestDASHInternal::FRe
 		return FManifestDASHInternal::FRepresentation::ESearchResult::PastEOS;
 	}
 
-	OutSegmentInfo.Time = PTO + EPTdelta + SegmentNum * SegmentDuration;
+	OutSegmentInfo.Time = PTO + EPTdelta + SegmentNum * (int64)SegmentDuration;
 	OutSegmentInfo.PTO = PTO;
 	OutSegmentInfo.EPTdelta = EPTdelta;
 	OutSegmentInfo.Duration = SegmentDuration;
@@ -1830,7 +1861,8 @@ FManifestDASHInternal::FRepresentation::ESearchResult FManifestDASHInternal::FRe
 	OutSegmentInfo.MediaLocalLastAUTime = MediaLocalPeriodDuration + PTO;
 	OutSegmentInfo.Timescale = MPDTimescale;
 	OutSegmentInfo.bMayBeMissing = SegmentNum + 1 >= MaxSegmentsInPeriod;
-	if (!PrepareDownloadURLs(PlayerSessionServices, OutSegmentInfo, SegmentTemplate))
+	OutSegmentInfo.ATO = ATO;
+	if (!PrepareDownloadURLs(InPlayerSessionServices, OutSegmentInfo, SegmentTemplate))
 	{
 		bIsUsable = false;
 		return FManifestDASHInternal::FRepresentation::ESearchResult::BadType;
@@ -1841,19 +1873,19 @@ FManifestDASHInternal::FRepresentation::ESearchResult FManifestDASHInternal::FRe
 	}
 }
 
-FManifestDASHInternal::FRepresentation::ESearchResult FManifestDASHInternal::FRepresentation::FindSegment_Timeline(IPlayerSessionServices* PlayerSessionServices, FSegmentInformation& OutSegmentInfo, TArray<TWeakPtrTS<FMPDLoadRequestDASH>>& OutRemoteElementLoadRequests, const FSegmentSearchOption& InSearchOptions, const TSharedPtrTS<FDashMPD_RepresentationType>& MPDRepresentation, const TArray<TSharedPtrTS<FDashMPD_SegmentTemplateType>>& SegmentTemplate, const TSharedPtrTS<FDashMPD_SegmentTimelineType>& SegmentTimeline)
+FManifestDASHInternal::FRepresentation::ESearchResult FManifestDASHInternal::FRepresentation::FindSegment_Timeline(IPlayerSessionServices* InPlayerSessionServices, FSegmentInformation& OutSegmentInfo, TArray<TWeakPtrTS<FMPDLoadRequestDASH>>& OutRemoteElementLoadRequests, const FSegmentSearchOption& InSearchOptions, const TSharedPtrTS<FDashMPD_RepresentationType>& MPDRepresentation, const TArray<TSharedPtrTS<FDashMPD_SegmentTemplateType>>& SegmentTemplate, const TSharedPtrTS<FDashMPD_SegmentTimelineType>& SegmentTimeline)
 {
 	// Segment timeline must not be empty.
 	auto Selements = SegmentTimeline->GetS_Elements();
 	if (Selements.Num() == 0)
 	{
-		PostError(PlayerSessionServices, FString::Printf(TEXT("Representation \"%s\" has an empty <SegmentTimeline>!"), *MPDRepresentation->GetID()), ERRCODE_DASH_MPD_BAD_REPRESENTATION);
+		PostError(InPlayerSessionServices, FString::Printf(TEXT("Representation \"%s\" has an empty <SegmentTimeline>!"), *MPDRepresentation->GetID()), ERRCODE_DASH_MPD_BAD_REPRESENTATION);
 		bIsUsable = false;
 		return FManifestDASHInternal::FRepresentation::ESearchResult::BadType;
 	}
 	else if (!Selements[0].HaveD)
 	{
-		PostError(PlayerSessionServices, FString::Printf(TEXT("Representation \"%s\" <SegmentTimeline> does not have mandatory 'd' element!"), *MPDRepresentation->GetID()), ERRCODE_DASH_MPD_BAD_REPRESENTATION);
+		PostError(InPlayerSessionServices, FString::Printf(TEXT("Representation \"%s\" <SegmentTimeline> does not have mandatory 'd' element!"), *MPDRepresentation->GetID()), ERRCODE_DASH_MPD_BAD_REPRESENTATION);
 		bIsUsable = false;
 		return FManifestDASHInternal::FRepresentation::ESearchResult::BadType;
 	}
@@ -1863,18 +1895,20 @@ FManifestDASHInternal::FRepresentation::ESearchResult FManifestDASHInternal::FRe
 	uint32 StartNumber = GET_ATTR(SegmentTemplate, GetStartNumber(), IsSet(), TMediaOptionalValue<uint32>(1)).Value();
 	uint32 EndNumber = GET_ATTR(SegmentTemplate, GetEndNumber(), IsSet(), TMediaOptionalValue<uint32>(~0U)).Value();
 	//TMediaOptionalValue<uint32> Duration = GET_ATTR(SegmentTemplate, GetDuration(), IsSet(), TMediaOptionalValue<uint32>());
+	FTimeValue ATO = CalculateSegmentAvailabilityTimeOffset(SegmentTemplate);
 
 	// The timescale should in all likelihood not be 1. While certainly allowed an accuracy of only one second is more likely to be
 	// an oversight when building the MPD.
 	if (MPDTimescale == 0)
 	{
-		PostError(PlayerSessionServices, FString::Printf(TEXT("Timescale for Representation \"%s\" is invalid!"), *MPDRepresentation->GetID()), ERRCODE_DASH_MPD_BAD_REPRESENTATION);
+		PostError(InPlayerSessionServices, FString::Printf(TEXT("Timescale for Representation \"%s\" is invalid!"), *MPDRepresentation->GetID()), ERRCODE_DASH_MPD_BAD_REPRESENTATION);
 		bIsUsable = false;
 		return FManifestDASHInternal::FRepresentation::ESearchResult::BadType;
 	}
-	else if (MPDTimescale == 1)
+	else if (MPDTimescale == 1 && !bWarnedAboutTimescale)
 	{
-		LogMessage(PlayerSessionServices, IInfoLog::ELevel::Warning, FString::Printf(TEXT("Timescale for Representation \"%s\" is given as 1. Is this intended?"), *MPDRepresentation->GetID()));
+		bWarnedAboutTimescale = true;
+		LogMessage(InPlayerSessionServices, IInfoLog::ELevel::Warning, FString::Printf(TEXT("Timescale for Representation \"%s\" is given as 1. Is this intended?"), *MPDRepresentation->GetID()));
 	}
 
 	// Get the period local time into media local timescale and add the PTO.
@@ -1908,7 +1942,7 @@ FManifestDASHInternal::FRepresentation::ESearchResult FManifestDASHInternal::FRe
 		if (MissingContentDuration > 0.1 && !bWarnedAboutTimelineStartGap)
 		{
 			bWarnedAboutTimelineStartGap = true;
-			LogMessage(PlayerSessionServices, IInfoLog::ELevel::Warning, FString::Printf(TEXT("Representation \"%s\" <SegmentTimeline> starts with %.3f seconds of missing content that will be skipped over and might lead to playback issues."), *MPDRepresentation->GetID(), MissingContentDuration));
+			LogMessage(InPlayerSessionServices, IInfoLog::ELevel::Warning, FString::Printf(TEXT("Representation \"%s\" <SegmentTimeline> starts with %.3f seconds of missing content that will be skipped over and might lead to playback issues."), *MPDRepresentation->GetID(), MissingContentDuration));
 		}
 		bFound = true;
 	}
@@ -1921,7 +1955,7 @@ FManifestDASHInternal::FRepresentation::ESearchResult FManifestDASHInternal::FRe
 		{
 			if (!Selements[nIndex].HaveD)
 			{
-				PostError(PlayerSessionServices, FString::Printf(TEXT("Representation \"%s\" <SegmentTimeline> does not have mandatory 'd' element!"), *MPDRepresentation->GetID()), ERRCODE_DASH_MPD_BAD_REPRESENTATION);
+				PostError(InPlayerSessionServices, FString::Printf(TEXT("Representation \"%s\" <SegmentTimeline> does not have mandatory 'd' element!"), *MPDRepresentation->GetID()), ERRCODE_DASH_MPD_BAD_REPRESENTATION);
 				bIsUsable = false;
 				return FManifestDASHInternal::FRepresentation::ESearchResult::BadType;
 			}
@@ -1933,7 +1967,7 @@ FManifestDASHInternal::FRepresentation::ESearchResult FManifestDASHInternal::FRe
 
 			if (CurrentD == 0)
 			{
-				PostError(PlayerSessionServices, FString::Printf(TEXT("Representation \"%s\" <SegmentTimeline> has an entry with 'd'=0, which is invalid."), *MPDRepresentation->GetID()), ERRCODE_DASH_MPD_BAD_REPRESENTATION);
+				PostError(InPlayerSessionServices, FString::Printf(TEXT("Representation \"%s\" <SegmentTimeline> has an entry with 'd'=0, which is invalid."), *MPDRepresentation->GetID()), ERRCODE_DASH_MPD_BAD_REPRESENTATION);
 				bIsUsable = false;
 				return FManifestDASHInternal::FRepresentation::ESearchResult::BadType;
 			}
@@ -1942,7 +1976,7 @@ FManifestDASHInternal::FRepresentation::ESearchResult FManifestDASHInternal::FRe
 			if (CurrentN > MAX_uint32 && !bWarnedAboutTimelineNumberOverflow)
 			{
 				bWarnedAboutTimelineNumberOverflow = true;
-				LogMessage(PlayerSessionServices, IInfoLog::ELevel::Warning, FString::Printf(TEXT("Representation \"%s\" <SegmentTimeline> 'n' value exceeds unsignedInt (32 bits)."), *MPDRepresentation->GetID()));
+				LogMessage(InPlayerSessionServices, IInfoLog::ELevel::Warning, FString::Printf(TEXT("Representation \"%s\" <SegmentTimeline> 'n' value exceeds unsignedInt (32 bits)."), *MPDRepresentation->GetID()));
 			}
 
 			// Warn if explicit numbering results in a gap or overlap. We do nothing besides warn about this.
@@ -1951,7 +1985,7 @@ FManifestDASHInternal::FRepresentation::ESearchResult FManifestDASHInternal::FRe
 				if (!bWarnedAboutInconsistentNumbering)
 				{
 					bWarnedAboutInconsistentNumbering = true;
-					LogMessage(PlayerSessionServices, IInfoLog::ELevel::Warning, FString::Printf(TEXT("Representation \"%s\" <SegmentTimeline> 'n' value %lld is not the expected %d. This may cause playback issues"), *MPDRepresentation->GetID(), (long long int)CurrentN, (long long int)PreviousN+1));
+					LogMessage(InPlayerSessionServices, IInfoLog::ELevel::Warning, FString::Printf(TEXT("Representation \"%s\" <SegmentTimeline> 'n' value %lld is not the expected %d. This may cause playback issues"), *MPDRepresentation->GetID(), (long long int)CurrentN, (long long int)PreviousN+1));
 				}
 			}
 
@@ -1974,7 +2008,7 @@ FManifestDASHInternal::FRepresentation::ESearchResult FManifestDASHInternal::FRe
 						if (!bWarnedAboutTimelineOverlap)
 						{
 							bWarnedAboutTimelineOverlap = true;
-							LogMessage(PlayerSessionServices, IInfoLog::ELevel::Warning, FString::Printf(TEXT("Representation \"%s\" <SegmentTimeline> 't' value %lld overlaps with preceeding segment (ends at %lld). This may cause playback issues"), *MPDRepresentation->GetID(), (long long int)CurrentT, (long long int)ExpectedT));
+							LogMessage(InPlayerSessionServices, IInfoLog::ELevel::Warning, FString::Printf(TEXT("Representation \"%s\" <SegmentTimeline> 't' value %lld overlaps with preceeding segment (ends at %lld). This may cause playback issues"), *MPDRepresentation->GetID(), (long long int)CurrentT, (long long int)ExpectedT));
 						}
 					}
 					else
@@ -2009,7 +2043,7 @@ FManifestDASHInternal::FRepresentation::ESearchResult FManifestDASHInternal::FRe
 						if (!bWarnedAboutTimelineNoTAfterNegativeR)
 						{
 							bWarnedAboutTimelineNoTAfterNegativeR = true;
-							LogMessage(PlayerSessionServices, IInfoLog::ELevel::Warning, FString::Printf(TEXT("Representation \"%s\" <SegmentTimeline> element following after a 'r'=-1 repeat count does not have a new 't' value!"), *MPDRepresentation->GetID()));
+							LogMessage(InPlayerSessionServices, IInfoLog::ELevel::Warning, FString::Printf(TEXT("Representation \"%s\" <SegmentTimeline> element following after a 'r'=-1 repeat count does not have a new 't' value!"), *MPDRepresentation->GetID()));
 						}
 					}
 					else
@@ -2022,14 +2056,14 @@ FManifestDASHInternal::FRepresentation::ESearchResult FManifestDASHInternal::FRe
 
 				if (EndTime == TNumericLimits<int64>::Max())
 				{
-					PostError(PlayerSessionServices, FString::Printf(TEXT("Representation \"%s\" <SegmentTimeline> repeats until infinity as last period is open-ended which is not currently supported."), *MPDRepresentation->GetID()), ERRCODE_DASH_MPD_INTERNAL);
+					PostError(InPlayerSessionServices, FString::Printf(TEXT("Representation \"%s\" <SegmentTimeline> repeats until infinity as last period is open-ended which is not currently supported."), *MPDRepresentation->GetID()), ERRCODE_DASH_MPD_INTERNAL);
 					bIsUsable = false;
 					return FManifestDASHInternal::FRepresentation::ESearchResult::BadType;
 				}
 
 				if (CurrentR < 0)
 				{
-					PostError(PlayerSessionServices, FString::Printf(TEXT("Representation \"%s\" <SegmentTimeline> repeat count of -1 failed to resolved to a positive value."), *MPDRepresentation->GetID()), ERRCODE_DASH_MPD_INTERNAL);
+					PostError(InPlayerSessionServices, FString::Printf(TEXT("Representation \"%s\" <SegmentTimeline> repeat count of -1 failed to resolved to a positive value."), *MPDRepresentation->GetID()), ERRCODE_DASH_MPD_INTERNAL);
 					bIsUsable = false;
 					return FManifestDASHInternal::FRepresentation::ESearchResult::BadType;
 				}
@@ -2043,7 +2077,7 @@ FManifestDASHInternal::FRepresentation::ESearchResult FManifestDASHInternal::FRe
 					// If this segment consists of subsegments we fail. This is not currently supported.
 					if (Selements[nIndex].HaveK)
 					{
-						PostError(PlayerSessionServices, FString::Printf(TEXT("Representation \"%s\" <SegmentTimeline> uses 'k' element which is not currently supported!"), *MPDRepresentation->GetID()), ERRCODE_DASH_MPD_BAD_REPRESENTATION);
+						PostError(InPlayerSessionServices, FString::Printf(TEXT("Representation \"%s\" <SegmentTimeline> uses 'k' element which is not currently supported!"), *MPDRepresentation->GetID()), ERRCODE_DASH_MPD_BAD_REPRESENTATION);
 						bIsUsable = false;
 						return FManifestDASHInternal::FRepresentation::ESearchResult::BadType;
 					}
@@ -2102,7 +2136,7 @@ FManifestDASHInternal::FRepresentation::ESearchResult FManifestDASHInternal::FRe
 					}
 					else
 					{
-						PostError(PlayerSessionServices, FString::Printf(TEXT("Unsupported segment search mode!")), ERRCODE_DASH_MPD_INTERNAL);
+						PostError(InPlayerSessionServices, FString::Printf(TEXT("Unsupported segment search mode!")), ERRCODE_DASH_MPD_INTERNAL);
 						bIsUsable = false;
 						return FManifestDASHInternal::FRepresentation::ESearchResult::BadType;
 					}
@@ -2120,6 +2154,18 @@ FManifestDASHInternal::FRepresentation::ESearchResult FManifestDASHInternal::FRe
 				++CurrentN;
 				--CurrentR;
 			}
+
+			// If the search time falls into the last segment we will not have found it above.
+			if (!bFound && CurrentT >= MediaLocalSearchTime && nIndex+1 == Selements.Num())
+			{
+				if (InSearchOptions.SearchType == IManifest::ESearchType::Closest || InSearchOptions.SearchType == IManifest::ESearchType::Same || InSearchOptions.SearchType == IManifest::ESearchType::Before || InSearchOptions.SearchType == IManifest::ESearchType::StrictlyBefore)
+				{
+					--CurrentN;
+					CurrentD = PreviousD;
+					CurrentT = PreviousT;
+					bFound = true;
+				}
+			}
 		}
 	}
 
@@ -2138,9 +2184,10 @@ FManifestDASHInternal::FRepresentation::ESearchResult FManifestDASHInternal::FRe
 		{
 			OutSegmentInfo.bMayBeMissing = true;
 			OutSegmentInfo.bIsMissing = true;
-			LogMessage(PlayerSessionServices, IInfoLog::ELevel::Warning, FString::Printf(TEXT("Representation \"%s\" <SegmentTimeline> gap encountered for needed 't' value of %lld. Replacing with an empty filler segment."), *MPDRepresentation->GetID(), (long long int)CurrentT));
+			LogMessage(InPlayerSessionServices, IInfoLog::ELevel::Warning, FString::Printf(TEXT("Representation \"%s\" <SegmentTimeline> gap encountered for needed 't' value of %lld. Replacing with an empty filler segment."), *MPDRepresentation->GetID(), (long long int)CurrentT));
 		}
-		if (!PrepareDownloadURLs(PlayerSessionServices, OutSegmentInfo, SegmentTemplate))
+		OutSegmentInfo.ATO = ATO;
+		if (!PrepareDownloadURLs(InPlayerSessionServices, OutSegmentInfo, SegmentTemplate))
 		{
 			bIsUsable = false;
 			return FManifestDASHInternal::FRepresentation::ESearchResult::BadType;
@@ -2306,15 +2353,6 @@ void FManifestDASHInternal::FRepresentation::SegmentIndexDownloadComplete(TShare
 			{
 				if (Index->PrepareTracks(nullptr) == UEMEDIA_ERROR_OK && Index->GetNumberOfSegmentIndices() > 0)
 				{
-					IElectraHttpManager::FParams::FRange r;
-					r.Set(LoadRequest->Range);
-					if (r.IsSet())
-					{
-						check(r.GetStart() >= 0);
-						check(r.GetEndIncluding() > 0);
-						SegmentIndexRangeStart = r.GetStart();
-						SegmentIndexRangeSize = r.GetEndIncluding() + 1 - SegmentIndexRangeStart;
-					}
 					SegmentIndex = MoveTemp(Index);
 					bOk = true;
 					// Add this to the entity cache in case it needs to be retrieved again.
