ElectraPlayer: Added MPEG DASH Live playback. Currently only for single period presentations.

diff --git a/Engine/Plugins/Media/ElectraPlayer/Source/ElectraPlayerRuntime/Private/Runtime/Player/DASH/ManifestDASH.h b/Engine/Plugins/Media/ElectraPlayer/Source/ElectraPlayerRuntime/Private/Runtime/Player/DASH/ManifestDASH.h
--- a/Engine/Plugins/Media/ElectraPlayer/Source/ElectraPlayerRuntime/Private/Runtime/Player/DASH/ManifestDASH.h
+++ b/Engine/Plugins/Media/ElectraPlayer/Source/ElectraPlayerRuntime/Private/Runtime/Player/DASH/ManifestDASH.h
@@ -14,12 +14,17 @@ class FDASHTimeline;
 class FManifestDASH : public IManifest
 {
 public:
-	static TSharedPtrTS<FManifestDASH> Create(IPlayerSessionServices* SessionServices, const FParamDict& Options, TWeakPtrTS<FPlaylistReaderDASH> PlaylistReader, TSharedPtrTS<FManifestDASHInternal> Manifest);
-	void UpdateTimeline();
+	static TSharedPtrTS<FManifestDASH> Create(IPlayerSessionServices* SessionServices, TSharedPtrTS<FManifestDASHInternal> Manifest);
+	void UpdateInternalManifest(TSharedPtrTS<FManifestDASHInternal> UpdatedManifest);
 
 	virtual ~FManifestDASH();
 	virtual EType GetPresentationType() const override;
-	virtual TSharedPtrTS<IPlaybackAssetTimeline> GetTimeline() const override;
+	virtual FTimeValue GetAnchorTime() const override;
+	virtual FTimeRange GetTotalTimeRange() const override;
+	virtual FTimeRange GetSeekableTimeRange() const override;
+	virtual void GetSeekablePositions(TArray<FTimespan>& OutPositions) const override;
+	virtual FTimeValue GetDuration() const override;
+	virtual FTimeValue GetDefaultStartTime() const override;
 	virtual int64 GetDefaultStartingBitrate() const override;
 	virtual FTimeValue GetMinBufferTime() const override;
 	virtual void GetStreamMetadata(TArray<FStreamMetadata>& OutMetadata, EStreamType StreamType) const override;
@@ -31,13 +36,10 @@ public:
 private:
 	ELECTRA_IMPL_DEFAULT_ERROR_METHODS(DASHManifest);
 
-	FManifestDASH(IPlayerSessionServices* SessionServices, const FParamDict& InOptions, TWeakPtrTS<FPlaylistReaderDASH> PlaylistReader, TSharedPtrTS<FManifestDASHInternal> Manifest);
+	FManifestDASH(IPlayerSessionServices* SessionServices, TSharedPtrTS<FManifestDASHInternal> Manifest);
 
-	FParamDict									Options;
-	TWeakPtrTS<FManifestDASHInternal>			InternalManifest;
 	IPlayerSessionServices* 					PlayerSessionServices = nullptr;
-	TWeakPtrTS<FPlaylistReaderDASH>				PlaylistReader;
-	TSharedPtrTS<FDASHTimeline>					CurrentTimeline;
+	TSharedPtrTS<FManifestDASHInternal>			CurrentManifest;
 	mutable TArray<FStreamMetadata>				CurrentMetadataVideo;
 	mutable TArray<FStreamMetadata>				CurrentMetadataAudio;
 	mutable TArray<FStreamMetadata>				CurrentMetadataSubtitle;
diff --git a/Engine/Plugins/Media/ElectraPlayer/Source/ElectraPlayerRuntime/Private/Runtime/Player/DASH/OptionKeynamesDASH.h b/Engine/Plugins/Media/ElectraPlayer/Source/ElectraPlayerRuntime/Private/Runtime/Player/DASH/OptionKeynamesDASH.h
--- a/Engine/Plugins/Media/ElectraPlayer/Source/ElectraPlayerRuntime/Private/Runtime/Player/DASH/OptionKeynamesDASH.h
+++ b/Engine/Plugins/Media/ElectraPlayer/Source/ElectraPlayerRuntime/Private/Runtime/Player/DASH/OptionKeynamesDASH.h
@@ -9,8 +9,18 @@ namespace Electra
 namespace DASH
 {
 
+const TCHAR* const OptionKeyMPDLoadConnectTimeout = TEXT("mpd_connection_timeout");				//!< (FTimeValue) value specifying connection timeout fetching the MPD
+const TCHAR* const OptionKeyMPDLoadNoDataTimeout = TEXT("mpd_nodata_timeout");					//!< (FTimeValue) value specifying no-data timeout fetching the MPD
+const TCHAR* const OptionKeyMPDReloadConnectTimeout = TEXT("mpd_update_connection_timeout");	//!< (FTimeValue) value specifying connection timeout fetching the MPD repeatedly
+const TCHAR* const OptionKeyMPDReloadNoDataTimeout = TEXT("mpd_update_connection_timeout");		//!< (FTimeValue) value specifying no-data timeout fetching the MPD repeatedly
+
+const TCHAR* const OptionKey_MinTimeBetweenMPDUpdates = TEXT("dash:min_mpd_update_interval");	//!< (FTimeValue) value limiting the time between MPD updates unless these are required updates.
+
 const TCHAR* const OptionKey_CurrentCDN = TEXT("dash:current_cdn");			//!< A string giving the currently selected CDN to be used with <BaseURL@serviceLocation> elements.
 
+
+const TCHAR* const HTTPHeaderOptionName = TEXT("MPEG-DASH-Param");			//!< HTTP request header name carrying optional DASH specific parameters.
+
 }
 
 } // namespace Electra
diff --git a/Engine/Plugins/Media/ElectraPlayer/Source/ElectraPlayerRuntime/Private/Runtime/Player/DASH/PlaylistReaderDASH.cpp b/Engine/Plugins/Media/ElectraPlayer/Source/ElectraPlayerRuntime/Private/Runtime/Player/DASH/PlaylistReaderDASH.cpp
--- a/Engine/Plugins/Media/ElectraPlayer/Source/ElectraPlayerRuntime/Private/Runtime/Player/DASH/PlaylistReaderDASH.cpp
+++ b/Engine/Plugins/Media/ElectraPlayer/Source/ElectraPlayerRuntime/Private/Runtime/Player/DASH/PlaylistReaderDASH.cpp
@@ -10,10 +10,14 @@
 #include "ManifestDASH.h"
 #include "Utilities/HashFunctions.h"
 #include "Utilities/TimeUtilities.h"
+#include "Utilities/URLParser.h"
 #include "HAL/LowLevelMemTracker.h"
 #include "ElectraPlayerPrivate.h"
 #include "Player/AdaptiveStreamingPlayerResourceRequest.h"
 #include "Player/PlayerEntityCache.h"
+#include "Player/DASH/OptionKeynamesDASH.h"
+
+#include "Misc/DateTime.h"
 
 #define ERRCODE_DASH_PARSER_ERROR							1
 #define ERRCODE_DASH_MPD_DOWNLOAD_FAILED					2
@@ -22,16 +26,11 @@
 #define ERRCODE_DASH_MPD_REMOTE_ENTITY_FAILED				5
 
 
+//#define DO_NOT_PERFORM_CONDITIONAL_GET
 
 namespace Electra
 {
 
-const FString IPlaylistReaderDASH::OptionKeyMPDLoadConnectTimeout(TEXT("mpd_connection_timeout"));				//!< (FTimeValue) value specifying connection timeout fetching the MPD
-const FString IPlaylistReaderDASH::OptionKeyMPDLoadNoDataTimeout(TEXT("mpd_nodata_timeout"));					//!< (FTimeValue) value specifying no-data timeout fetching the MPD
-const FString IPlaylistReaderDASH::OptionKeyMPDReloadConnectTimeout(TEXT("mpd_update_connection_timeout"));		//!< (FTimeValue) value specifying connection timeout fetching the MPD repeatedly
-const FString IPlaylistReaderDASH::OptionKeyMPDReloadNoDataTimeout(TEXT("mpd_update_connection_timeout"));		//!< (FTimeValue) value specifying no-data timeout fetching the MPD repeatedly
-
-
 /**
  * This class is responsible for downloading a DASH MPD and parsing it.
  */
@@ -43,51 +42,20 @@ public:
 
 	virtual ~FPlaylistReaderDASH();
 	virtual void Close() override;
-
-	/**
-	 * Returns the type of playlist format.
-	 * For this implementation it will be "dash".
-	 *
-	 * @return "dash" to indicate this is a DASH MPD.
-	 */
 	virtual const FString& GetPlaylistType() const override
 	{
 		static FString Type("dash");
 		return Type;
 	}
-
-	/**
-	 * Loads and parses the MPD.
-	 *
-	 * @param URL     URL of the MPD to load
-	 * @param Preferences
-	 *                User preferences for initial stream selection.
-	 * @param Options Options for the reader and parser
-	 */
-	virtual void LoadAndParse(const FString& URL, const FStreamPreferences& Preferences, const FParamDict& Options) override;
-
-	/**
-	 * Returns the URL from which the MPD was loaded (or supposed to be loaded).
-	 *
-	 * @return The MPD URL
-	 */
+	virtual void LoadAndParse(const FString& URL) override;
 	virtual FString GetURL() const override;
-
-	/**
-	 * Returns an interface to the manifest created from the loaded MPD.
-	 *
-	 * @return A shared manifest interface pointer.
-	 */
 	virtual TSharedPtrTS<IManifest> GetManifest() override;
-
-	/**
-	 * Requests updating of a dynamic DASH MPD or an XLINK item.
-	 *
-	 * @param LoadRequest
-	 */
-	virtual void RequestMPDUpdate(const FMPDLoadRequestDASH& LoadRequest) override;
-
 	virtual void AddElementLoadRequests(const TArray<TWeakPtrTS<FMPDLoadRequestDASH>>& RemoteElementLoadRequests) override;
+	virtual void RequestMPDUpdate(bool bForcedUpdate) override;
+	virtual TSharedPtrTS<FManifestDASHInternal> GetCurrentMPD() override
+	{
+		return Manifest;
+	}
 
 private:
 	using FResourceLoadRequestPtr = TSharedPtrTS<FMPDLoadRequestDASH>;
@@ -98,11 +66,11 @@ private:
 
 	void StartWorkerThread();
 	void StopWorkerThread();
-	void WorkerThread(void);
+	void WorkerThread();
 	void ExecutePendingRequests(const FTimeValue& TimeNow);
 	void HandleCompletedRequests(const FTimeValue& TimeNow);
 	void HandleStaticRequestCompletions(const FTimeValue& TimeNow);
-	void CheckForMPDUpdate(const FTimeValue& TimeNow);
+	void CheckForMPDUpdate();
 
 	void PostError(const FErrorDetail& Error);
 	void PostError(const FString& Message, uint16 Code, UEMediaError Error = UEMEDIA_ERROR_OK);
@@ -117,22 +85,34 @@ private:
 	void EnqueueInitialXLinkRequests();
 
 	void ManifestDownloadCompleted(FResourceLoadRequestPtr Request, bool bSuccess);
+	void ManifestUpdateDownloadCompleted(FResourceLoadRequestPtr Request, bool bSuccess);
 	void InitialMPDXLinkElementDownloadCompleted(FResourceLoadRequestPtr Request, bool bSuccess);
 
 
-	IPlayerSessionServices*									PlayerSessionServices;
-	FStreamPreferences										StreamPreferences;
-	FParamDict												Options;
+	IPlayerSessionServices*									PlayerSessionServices = nullptr;
 	FString													MPDURL;
+	FString													Fragment;
 	FMediaSemaphore											WorkerThreadSignal;
-	bool													bIsWorkerThreadStarted;
-	volatile bool											bTerminateWorkerThread;
+	bool													bIsWorkerThreadStarted = false;
+	volatile bool											bTerminateWorkerThread = false;
 
 	FCriticalSection										RequestsLock;
 	TArray<FResourceLoadRequestPtr>							PendingRequests;
 	TArray<FResourceLoadRequestPtr>							ActiveRequests;
 	TArray<FResourceLoadRequestPtr>							CompletedRequests;
+	FTimeValue												MinTimeBetweenUpdates;
+	FTimeValue												NextMPDUpdateTime;
+	FTimeValue												MostRecentMPDUpdateTime;
+	bool													bIsMPDUpdateInProgress = false;
+	enum class EUpdateRequestType
+	{
+		None,
+		Regular,
+		Forced
+	};
+	EUpdateRequestType										UpdateRequested = EUpdateRequestType::None;
 
+	TUniquePtr<IManifestBuilderDASH>						Builder;
 
 	TSharedPtrTS<FManifestDASHInternal>						Manifest;
 	FErrorDetail											LastErrorDetail;
@@ -162,9 +142,6 @@ TSharedPtrTS<IPlaylistReader> IPlaylistReaderDASH::Create(IPlayerSessionServices
 
 FPlaylistReaderDASH::FPlaylistReaderDASH()
 	: FMediaThread("ElectraPlayer::DASH MPD")
-	, PlayerSessionServices(nullptr)
-	, bIsWorkerThreadStarted(false)
-	, bTerminateWorkerThread(false)
 {
 }
 
@@ -256,20 +233,102 @@ void FPlaylistReaderDASH::LogMessage(IInfoLog::ELevel Level, const FString& Mess
 	}
 }
 
-void FPlaylistReaderDASH::LoadAndParse(const FString& URL, const FStreamPreferences& Preferences, const FParamDict& InOptions)
+void FPlaylistReaderDASH::LoadAndParse(const FString& URL)
 {
-	StreamPreferences = Preferences;
-	Options 		  = InOptions;
 	MPDURL = URL;
 	StartWorkerThread();
 }
 
-void FPlaylistReaderDASH::RequestMPDUpdate(const FMPDLoadRequestDASH& LoadRequest)
-{
-}
 
-void FPlaylistReaderDASH::CheckForMPDUpdate(const FTimeValue& TimeNow)
+void FPlaylistReaderDASH::CheckForMPDUpdate()
 {
+	RequestsLock.Lock();
+	FTimeValue Next = NextMPDUpdateTime;
+	bool bRequestNow = false;
+	bool bIsForced = false;
+	if (UpdateRequested == EUpdateRequestType::Forced)
+	{
+		bIsForced = true;
+		Next.SetToZero();
+	}
+	else if (UpdateRequested == EUpdateRequestType::Regular)
+	{
+		if (!Next.IsValid())
+		{
+			Next.SetToZero();
+		}
+	}
+	if (!bIsMPDUpdateInProgress && Next.IsValid())
+	{
+		// Get the time now and do not pass it in from the worker thread as the clock could have been
+		// resynchronized to the server time in the meantime.
+		FTimeValue Now = PlayerSessionServices->GetSynchronizedUTCTime()->GetTime();
+		// Limit the frequency of the updates.
+		bool bIsPossible = !MostRecentMPDUpdateTime.IsValid() || Now - MostRecentMPDUpdateTime > MinTimeBetweenUpdates;
+		if (Next < Now && bIsPossible)
+		{
+			bRequestNow = true;
+			NextMPDUpdateTime.SetToInvalid();
+		}
+	}
+	UpdateRequested = EUpdateRequestType::None;
+	RequestsLock.Unlock();
+	if (bRequestNow)
+	{
+		TSharedPtrTS<const FDashMPD_MPDType> MPDRoot = Manifest.IsValid() ? Manifest->GetMPDRoot() : nullptr;
+		if (MPDRoot.IsValid())
+		{
+			FString NewMPDLocation;
+			// See if there are any <Location> elements on the current MPD that tell us from where to get the update.
+			// As per A.11 no <BaseURL> elements affect the construction of the new document URL. At most, if it is
+			// a relative URL it will be resolved against the current document URL.
+			const TArray<TSharedPtrTS<FDashMPD_OtherType>>& Locations = MPDRoot->GetLocations();
+			if (Locations.Num())
+			{
+				// There can be several new locations, but the elements have no DASH defined attributes to differentiate them.
+				// They are merely elements containing a URL. In absence of any preference we just use the first one.
+				NewMPDLocation = Locations[0]->GetData();
+			}
+
+			TArray<TSharedPtrTS<const FDashMPD_BaseURLType>> OutBaseURLs;
+			FString URL, RequestHeader;
+			FTimeValue UrlATO;
+			DASHUrlHelpers::BuildAbsoluteElementURL(URL, UrlATO, MPDRoot->GetDocumentURL(), OutBaseURLs, NewMPDLocation);
+
+			// The URL query might need to be changed. Look for the UrlQuery properties.
+			TArray<TSharedPtrTS<FDashMPD_UrlQueryInfoType>> UrlQueries;
+			DASHUrlHelpers::GetAllHierarchyUrlQueries(UrlQueries, MPDRoot, DASHUrlHelpers::EUrlQueryRequestType::Mpd, false);
+			FErrorDetail Error = DASHUrlHelpers::ApplyUrlQueries(PlayerSessionServices, MPDRoot->GetDocumentURL(), URL, RequestHeader, UrlQueries);
+			if (Error.IsOK())
+			{
+				FString ETag = MPDRoot->GetETag();
+
+				TSharedPtrTS<FMPDLoadRequestDASH> PlaylistLoadRequest = MakeSharedTS<FMPDLoadRequestDASH>();
+				PlaylistLoadRequest->LoadType = FMPDLoadRequestDASH::ELoadType::MPDUpdate;
+				PlaylistLoadRequest->URL = URL;
+				if (RequestHeader.Len())
+				{
+					PlaylistLoadRequest->Headers.Emplace(HTTP::FHTTPHeader({DASH::HTTPHeaderOptionName, RequestHeader}));
+				}
+#ifndef DO_NOT_PERFORM_CONDITIONAL_GET
+				// As per the standard we should perform conditional GET requests.
+				if (ETag.Len())
+				{
+					PlaylistLoadRequest->Headers.Emplace(HTTP::FHTTPHeader({TEXT("If-None-Match"), ETag}));
+				}
+#endif
+				PlaylistLoadRequest->CompleteCallback.BindThreadSafeSP(AsShared(), &FPlaylistReaderDASH::ManifestUpdateDownloadCompleted);
+				
+				bIsMPDUpdateInProgress = true;
+				EnqueueResourceRequest(MoveTemp(PlaylistLoadRequest));
+			}
+			else
+			{
+				// Failed to build the URL.
+				PostError(Error);
+			}
+		}
+	}
 }
 
 
@@ -339,17 +398,35 @@ void FPlaylistReaderDASH::AddElementLoadRequests(const TArray<TWeakPtrTS<FMPDLoa
 	}
 }
 
+void FPlaylistReaderDASH::RequestMPDUpdate(bool bForcedUpdate)
+{
+	FScopeLock lock(&RequestsLock);
+	if (bForcedUpdate)
+	{
+		UpdateRequested = EUpdateRequestType::Forced;
+	}
+	else if (UpdateRequested == EUpdateRequestType::None)
+	{
+		UpdateRequested = EUpdateRequestType::Regular;
+	}
+}
 
 
-void FPlaylistReaderDASH::WorkerThread(void)
+void FPlaylistReaderDASH::WorkerThread()
 {
 	LLM_SCOPE(ELLMTag::ElectraPlayer);
 	CSV_SCOPED_TIMING_STAT(ElectraPlayer, MPDReaderDASH_Worker);
 
+	// Get the minimum MPD update time limit.
+	MinTimeBetweenUpdates = PlayerSessionServices->GetOptions().GetValue(DASH::OptionKey_MinTimeBetweenMPDUpdates).SafeGetTimeValue(FTimeValue().SetFromMilliseconds(1000));
+
 	// Setup the playlist load request for the master playlist.
 	TSharedPtrTS<FMPDLoadRequestDASH> PlaylistLoadRequest = MakeSharedTS<FMPDLoadRequestDASH>();
 	PlaylistLoadRequest->LoadType = FMPDLoadRequestDASH::ELoadType::MPD;
-	PlaylistLoadRequest->URL = MPDURL;
+	FURL_RFC3986 UrlParser;
+	UrlParser.Parse(MPDURL);
+	PlaylistLoadRequest->URL = UrlParser.Get(true, false);
+	Fragment = UrlParser.GetFragment();
 	PlaylistLoadRequest->CompleteCallback.BindThreadSafeSP(AsShared(), &FPlaylistReaderDASH::ManifestDownloadCompleted);
 	EnqueueResourceRequest(MoveTemp(PlaylistLoadRequest));
 
@@ -367,7 +444,7 @@ void FPlaylistReaderDASH::WorkerThread(void)
 		ExecutePendingRequests(Now);
 
 		// Check if the MPD must be updated.
-		CheckForMPDUpdate(Now);
+		CheckForMPDUpdate();
 	}
 
 	// Cleanup!
@@ -383,6 +460,7 @@ void FPlaylistReaderDASH::WorkerThread(void)
 	ActiveRequests.Empty();
 	CompletedRequests.Empty();
 	RequestsLock.Unlock();
+	Builder.Reset();
 }
 
 
@@ -425,18 +503,19 @@ void FPlaylistReaderDASH::ExecutePendingRequests(const FTimeValue& TimeNow)
 
 void FPlaylistReaderDASH::SetupRequestTimeouts(FResourceLoadRequestPtr InRequest)
 {
+	const FParamDict& Options = PlayerSessionServices->GetOptions();
 	switch(InRequest->LoadType)
 	{
 		case FMPDLoadRequestDASH::ELoadType::MPD:
 		{
-			InRequest->Request->ConnectionTimeout(Options.GetValue(IPlaylistReaderDASH::OptionKeyMPDLoadConnectTimeout).SafeGetTimeValue(FTimeValue().SetFromMilliseconds(1000 * 8)));
-			InRequest->Request->NoDataTimeout(Options.GetValue(IPlaylistReaderDASH::OptionKeyMPDLoadNoDataTimeout).SafeGetTimeValue(FTimeValue().SetFromMilliseconds(1000 * 5)));
+			InRequest->Request->ConnectionTimeout(Options.GetValue(DASH::OptionKeyMPDLoadConnectTimeout).SafeGetTimeValue(FTimeValue().SetFromMilliseconds(1000 * 8)));
+			InRequest->Request->NoDataTimeout(Options.GetValue(DASH::OptionKeyMPDLoadNoDataTimeout).SafeGetTimeValue(FTimeValue().SetFromMilliseconds(1000 * 5)));
 			break;
 		}
 		case FMPDLoadRequestDASH::ELoadType::MPDUpdate:
 		{
-			InRequest->Request->ConnectionTimeout(Options.GetValue(IPlaylistReaderDASH::OptionKeyMPDReloadConnectTimeout).SafeGetTimeValue(FTimeValue().SetFromMilliseconds(1000 * 2)));
-			InRequest->Request->NoDataTimeout(Options.GetValue(IPlaylistReaderDASH::OptionKeyMPDReloadNoDataTimeout).SafeGetTimeValue(FTimeValue().SetFromMilliseconds(1000 * 2)));
+			InRequest->Request->ConnectionTimeout(Options.GetValue(DASH::OptionKeyMPDReloadConnectTimeout).SafeGetTimeValue(FTimeValue().SetFromMilliseconds(1000 * 2)));
+			InRequest->Request->NoDataTimeout(Options.GetValue(DASH::OptionKeyMPDReloadNoDataTimeout).SafeGetTimeValue(FTimeValue().SetFromMilliseconds(1000 * 2)));
 			break;
 		}
 		case FMPDLoadRequestDASH::ELoadType::Segment:
@@ -650,25 +729,32 @@ void FPlaylistReaderDASH::ManifestDownloadCompleted(FResourceLoadRequestPtr Requ
 		// Get the Date header from the response and set our clock to this time.
 		// Do this once only with the date from the master playlist.
 		const HTTP::FConnectionInfo* ConnInfo = Request->GetConnectionInfo();
+		FString ETag;
+		FTimeValue FetchTime;
+		FString EffectiveURL = ConnInfo->EffectiveURL;
 		if (ConnInfo)
 		{
 			for(int32 i=0; i<ConnInfo->ResponseHeaders.Num(); ++i)
 			{
 				if (ConnInfo->ResponseHeaders[i].Header == "Date")
 				{
-					// How old is the response already?
-					FTimeValue ResponseAge = MEDIAutcTime::Current() - ConnInfo->RequestStartTime - FTimeValue().SetFromSeconds(ConnInfo->TimeUntilFirstByte);
 					// Parse the header
 					FTimeValue DateFromHeader;
 					UEMediaError DateParseError = RFC7231::ParseDateTime(DateFromHeader, ConnInfo->ResponseHeaders[i].Value);
 					if (DateParseError == UEMEDIA_ERROR_OK && DateFromHeader.IsValid())
 					{
-						PlayerSessionServices->GetSynchronizedUTCTime()->SetTime(DateFromHeader + ResponseAge);
+						PlayerSessionServices->GetSynchronizedUTCTime()->SetTime(ConnInfo->RequestStartTime, DateFromHeader);
 					}
-					break;
+					// The MPD 'FetchTime' is the time at which the request to get the MPD was made to the server.
+					FetchTime = PlayerSessionServices->GetSynchronizedUTCTime()->MapToSyncTime(ConnInfo->RequestStartTime);
+				}
+				else if (ConnInfo->ResponseHeaders[i].Header == "ETag")
+				{
+					ETag = ConnInfo->ResponseHeaders[i].Value;
 				}
 			}
 		}
+		MostRecentMPDUpdateTime = FetchTime;
 
 		PlayerSessionServices->SendMessageToPlayer(IPlaylistReader::PlaylistDownloadMessage::Create(ConnInfo, Playlist::EListType::Master, Playlist::ELoadType::Initial));
 
@@ -676,23 +762,56 @@ void FPlaylistReaderDASH::ManifestDownloadCompleted(FResourceLoadRequestPtr Requ
 		LastErrorDetail = GetXMLResponseString(XML, Request);
 		if (LastErrorDetail.IsOK())
 		{
-			IManifestBuilderDASH* Builder = IManifestBuilderDASH::Create(PlayerSessionServices);
+			Builder.Reset(IManifestBuilderDASH::Create(PlayerSessionServices));
 			TSharedPtrTS<FManifestDASHInternal> NewManifest;
-			LastErrorDetail = Builder->BuildFromMPD(NewManifest, XML.GetCharArray().GetData(), Request, StreamPreferences, Options);
-			Manifest = NewManifest;
-			PlayerManifest = FManifestDASH::Create(PlayerSessionServices, Options, AsShared(), Manifest);
-			// Notify that the "master playlist" has been parsed, successfully or not. If we still need to resolve remote entities this is not an error!
-			PlayerSessionServices->SendMessageToPlayer(IPlaylistReader::PlaylistLoadedMessage::Create(!LastErrorDetail.IsTryAgain() ? LastErrorDetail : FErrorDetail(), ConnInfo, Playlist::EListType::Master, Playlist::ELoadType::Initial));
-			if (LastErrorDetail.IsOK())
+			LastErrorDetail = Builder->BuildFromMPD(NewManifest, XML.GetCharArray().GetData(), EffectiveURL, FetchTime, ETag);
+			if (LastErrorDetail.IsOK() || LastErrorDetail.IsTryAgain())
 			{
-				PlayerManifest->UpdateTimeline();
-				// Notify that the "variant playlists" are ready. There are no variants in DASH, but this is the trigger that the playlists are all set up and are good to go now.
-				PlayerSessionServices->SendMessageToPlayer(IPlaylistReader::PlaylistLoadedMessage::Create(LastErrorDetail, ConnInfo, Playlist::EListType::Variant, Playlist::ELoadType::Initial));
-			}
-			// "try again" is returned when there are remote entities that need to be resolved first.
-			else if (LastErrorDetail.IsTryAgain())
-			{
-				EnqueueInitialXLinkRequests();
+				if (NewManifest.IsValid())
+				{
+					// Parse the URL fragment into its components. For a DASH URL the fragment is constructed like a query string with & delimited key/value pairs.
+					TArray<FURL_RFC3986::FQueryParam> URLFragmentComponents;
+					FURL_RFC3986::GetQueryParams(URLFragmentComponents, Fragment, false);	// The fragment is already URL escaped, so no need to do it again.
+					NewManifest->SetURLFragmentComponents(MoveTemp(URLFragmentComponents));
+					// If the URL has a special fragment part to turn this presentation into an event, do so.
+					// Note: This MAY require the client clock to be synced to the server IF the special keyword 'now' is used.
+					NewManifest->TransformIntoEpicEvent();
+				}
+				Manifest = NewManifest;
+				PlayerManifest = FManifestDASH::Create(PlayerSessionServices, Manifest);
+
+				// Check if the MPD defines an @minimumUpdatePeriod
+				FTimeValue mup = Manifest->GetMinimumUpdatePeriod();
+				// If the update time is zero then updates happen just in time when segments are required or through
+				// an inband event stream. Either way, we do not need to update periodically.
+				if (mup.IsValid() && mup > FTimeValue::GetZero())
+				{
+					// Warn if MUP is really small
+					if (mup.GetAsMilliseconds() < 1000)
+					{
+						LogMessage(IInfoLog::ELevel::Warning, FString::Printf(TEXT("MPD@minimumUpdatePeriod is set to a really small value of %lld msec. This could be a performance issue."), (long long int)mup.GetAsMilliseconds()));
+					}
+					RequestsLock.Lock();
+					NextMPDUpdateTime = MostRecentMPDUpdateTime + mup;
+					RequestsLock.Unlock();
+				}
+
+				// Notify that the "master playlist" has been parsed, successfully or not. If we still need to resolve remote entities this is not an error!
+				PlayerSessionServices->SendMessageToPlayer(IPlaylistReader::PlaylistLoadedMessage::Create(!LastErrorDetail.IsTryAgain() ? LastErrorDetail : FErrorDetail(), ConnInfo, Playlist::EListType::Master, Playlist::ELoadType::Initial));
+				if (LastErrorDetail.IsOK())
+				{
+					// Notify that the "variant playlists" are ready. There are no variants in DASH, but this is the trigger that the playlists are all set up and are good to go now.
+					PlayerSessionServices->SendMessageToPlayer(IPlaylistReader::PlaylistLoadedMessage::Create(LastErrorDetail, ConnInfo, Playlist::EListType::Variant, Playlist::ELoadType::Initial));
+				}
+				// "try again" is returned when there are remote entities that need to be resolved first.
+				else if (LastErrorDetail.IsTryAgain())
+				{
+					EnqueueInitialXLinkRequests();
+				}
+				else
+				{
+					PostError(LastErrorDetail);
+				}
 			}
 			else
 			{
@@ -710,6 +829,119 @@ void FPlaylistReaderDASH::ManifestDownloadCompleted(FResourceLoadRequestPtr Requ
 	}
 }
 
+
+void FPlaylistReaderDASH::ManifestUpdateDownloadCompleted(FResourceLoadRequestPtr Request, bool bSuccess)
+{
+	bIsMPDUpdateInProgress = false;
+
+	if (bSuccess)
+	{
+		if (!Request->Request.IsValid() || !Request->Request->GetResponseBuffer().IsValid())
+		{
+			return;
+		}
+		
+		const HTTP::FConnectionInfo* ConnInfo = Request->GetConnectionInfo();
+		FString ETag;
+		FString EffectiveURL;
+		FTimeValue FetchTime;
+		if (ConnInfo)
+		{
+			EffectiveURL = ConnInfo->EffectiveURL;
+			FetchTime = PlayerSessionServices->GetSynchronizedUTCTime()->MapToSyncTime(ConnInfo->RequestStartTime);
+			for(int32 i=0; i<ConnInfo->ResponseHeaders.Num(); ++i)
+			{
+				if (ConnInfo->ResponseHeaders[i].Header == "ETag")
+				{
+					ETag = ConnInfo->ResponseHeaders[i].Value;
+				}
+			}
+		}
+		MostRecentMPDUpdateTime = FetchTime;
+
+		// If this was a "304 - Not modified" (RFC 7323) there is no new MPD at the moment.
+		if (ConnInfo->StatusInfo.HTTPStatus == 304)
+		{
+			// It does however extend the validity of the MPD.
+			Manifest->GetMPDRoot()->SetFetchTime(FetchTime);
+
+			// The minimum update period still holds.
+			FTimeValue mup = Manifest->GetMinimumUpdatePeriod();
+			if (mup.IsValid() && mup > FTimeValue::GetZero())
+			{
+				RequestsLock.Lock();
+				NextMPDUpdateTime = MostRecentMPDUpdateTime + mup;
+				RequestsLock.Unlock();
+			}
+		}
+		else
+		{
+			PlayerSessionServices->SendMessageToPlayer(IPlaylistReader::PlaylistDownloadMessage::Create(ConnInfo, Playlist::EListType::Master, Playlist::ELoadType::Update));
+
+			FString XML;
+			LastErrorDetail = GetXMLResponseString(XML, Request);
+			if (LastErrorDetail.IsOK())
+			{
+				TSharedPtrTS<FManifestDASHInternal> NewManifest;
+				LastErrorDetail = Builder->BuildFromMPD(NewManifest, XML.GetCharArray().GetData(), EffectiveURL, FetchTime, ETag);
+				if (LastErrorDetail.IsOK() || LastErrorDetail.IsTryAgain())
+				{
+					// Copy over the initial document URL fragments.
+					if (NewManifest.IsValid())
+					{
+						NewManifest->SetURLFragmentComponents(Manifest->GetURLFragmentComponents());
+					}
+
+					// Switch over tp the new manifest.
+					Manifest = NewManifest;
+					// Also update in the external manifest we handed out to the player.
+					PlayerManifest->UpdateInternalManifest(Manifest);
+
+					// Check if the new MPD also defines an @minimumUpdatePeriod
+					FTimeValue mup = Manifest->GetMinimumUpdatePeriod();
+					if (mup.IsValid() && mup > FTimeValue::GetZero())
+					{
+						if (mup.GetAsMilliseconds() < 1000)
+						{
+							LogMessage(IInfoLog::ELevel::Warning, FString::Printf(TEXT("MPD@minimumUpdatePeriod is set to a really small value of %lld msec. This could be a performance issue."), (long long int)mup.GetAsMilliseconds()));
+						}
+						RequestsLock.Lock();
+						NextMPDUpdateTime = MostRecentMPDUpdateTime + mup;
+						RequestsLock.Unlock();
+					}
+
+					PlayerSessionServices->SendMessageToPlayer(IPlaylistReader::PlaylistLoadedMessage::Create(!LastErrorDetail.IsTryAgain() ? LastErrorDetail : FErrorDetail(), ConnInfo, Playlist::EListType::Master, Playlist::ELoadType::Update));
+					if (LastErrorDetail.IsOK())
+					{
+					}
+					// "try again" is returned when there are remote entities that need to be resolved first.
+					else if (LastErrorDetail.IsTryAgain())
+					{
+						EnqueueInitialXLinkRequests();
+					}
+					else
+					{
+						PostError(LastErrorDetail);
+					}
+				}
+				else
+				{
+					PostError(LastErrorDetail);
+				}
+			}
+			else
+			{
+				PostError(LastErrorDetail);
+			}
+		}
+	}
+	else
+	{
+		PostError(FString::Printf(TEXT("Failed to download MPD \"%s\" (%s)"), *Request->URL, *Request->GetErrorDetail()), ERRCODE_DASH_MPD_DOWNLOAD_FAILED, UEMEDIA_ERROR_READ_ERROR);
+	}
+}
+
+
 void FPlaylistReaderDASH::InitialMPDXLinkElementDownloadCompleted(FResourceLoadRequestPtr Request, bool bSuccess)
 {
 	// Is the manifest for which this request was made still of interest?
@@ -718,16 +950,12 @@ void FPlaylistReaderDASH::InitialMPDXLinkElementDownloadCompleted(FResourceLoadR
 	{
 		FString XML;
 		LastErrorDetail = GetXMLResponseString(XML, Request);
-		LastErrorDetail = ForManifest->ResolveInitialRemoteElementRequest(PlayerSessionServices, Request, MoveTemp(XML), bSuccess);
+		LastErrorDetail = ForManifest->ResolveInitialRemoteElementRequest(Request, MoveTemp(XML), bSuccess);
 
 		// Now check if all pending requests have finished
 		if (LastErrorDetail.IsOK())
 		{
-			LastErrorDetail = ForManifest->BuildAfterInitialRemoteElementDownload(PlayerSessionServices);
-			if (LastErrorDetail.IsOK())
-			{
-				PlayerManifest->UpdateTimeline();
-			}
+			LastErrorDetail = ForManifest->BuildAfterInitialRemoteElementDownload();
 			// Notify that the "variant playlists" are ready. There are no variants in DASH, but this is the trigger that the playlists are all set up and are good to go now.
 			PlayerSessionServices->SendMessageToPlayer(IPlaylistReader::PlaylistLoadedMessage::Create(LastErrorDetail, nullptr, Playlist::EListType::Variant, Playlist::ELoadType::Initial));
 		}
diff --git a/Engine/Plugins/Media/ElectraPlayer/Source/ElectraPlayerRuntime/Private/Runtime/Player/DASH/PlaylistReaderDASH.h b/Engine/Plugins/Media/ElectraPlayer/Source/ElectraPlayerRuntime/Private/Runtime/Player/DASH/PlaylistReaderDASH.h
--- a/Engine/Plugins/Media/ElectraPlayer/Source/ElectraPlayerRuntime/Private/Runtime/Player/DASH/PlaylistReaderDASH.h
+++ b/Engine/Plugins/Media/ElectraPlayer/Source/ElectraPlayerRuntime/Private/Runtime/Player/DASH/PlaylistReaderDASH.h
@@ -11,6 +11,7 @@ namespace Electra
 {
 
 struct FMPDLoadRequestDASH;
+class FManifestDASHInternal;
 
 
 
@@ -21,20 +22,12 @@ public:
 
 	virtual ~IPlaylistReaderDASH() = default;
 
-	static const FString OptionKeyMPDLoadConnectTimeout;		//!< (FTimeValue) value specifying connection timeout fetching the MPD
-	static const FString OptionKeyMPDLoadNoDataTimeout;			//!< (FTimeValue) value specifying no-data timeout fetching the MPD
-	static const FString OptionKeyMPDReloadConnectTimeout;		//!< (FTimeValue) value specifying connection timeout fetching the MPD repeatedly
-	static const FString OptionKeyMPDReloadNoDataTimeout;		//!< (FTimeValue) value specifying no-data timeout fetching the MPD repeatedly
-
 	/**
 	 * Loads and parses the MPD.
 	 *
 	 * @param URL     URL of the MPD to load
-	 * @param Preferences
-	 *                User preferences for initial stream selection.
-	 * @param Options Options
 	 */
-	virtual void LoadAndParse(const FString& URL, const FStreamPreferences& Preferences, const FParamDict& Options) = 0;
+	virtual void LoadAndParse(const FString& URL) = 0;
 
 	/**
 	 * Returns the URL from which the MPD was loaded (or supposed to be loaded).
@@ -43,19 +36,22 @@ public:
 	 */
 	virtual FString GetURL() const = 0;
 
-
 	/**
-	 * Requests loading of an updated MPD (or an XLINK item)
-	 *
-	 * @param LoadRequest
+	 * Adds load requests for additional elements.
+	 * 
+	 * @param RemoteElementLoadRequests
 	 */
-	virtual void RequestMPDUpdate(const FMPDLoadRequestDASH& LoadRequest) = 0;
+	virtual void AddElementLoadRequests(const TArray<TWeakPtrTS<FMPDLoadRequestDASH>>& RemoteElementLoadRequests) = 0;
 
 	/**
-	 * Adds load requests for additional elements.
+	 * Requests an MPD update.
 	 */
-	virtual void AddElementLoadRequests(const TArray<TWeakPtrTS<FMPDLoadRequestDASH>>& RemoteElementLoadRequests) = 0;
+	virtual void RequestMPDUpdate(bool bForcedUpdate) = 0;
 
+	/**
+	 * Access the current internal MPD
+	 */
+	virtual TSharedPtrTS<FManifestDASHInternal> GetCurrentMPD() = 0;
 };
 
 
diff --git a/Engine/Plugins/Media/ElectraPlayer/Source/ElectraPlayerRuntime/Private/Runtime/Player/DASH/PlaylistReaderDASH_Internal.h b/Engine/Plugins/Media/ElectraPlayer/Source/ElectraPlayerRuntime/Private/Runtime/Player/DASH/PlaylistReaderDASH_Internal.h
--- a/Engine/Plugins/Media/ElectraPlayer/Source/ElectraPlayerRuntime/Private/Runtime/Player/DASH/PlaylistReaderDASH_Internal.h
+++ b/Engine/Plugins/Media/ElectraPlayer/Source/ElectraPlayerRuntime/Private/Runtime/Player/DASH/PlaylistReaderDASH_Internal.h
@@ -8,6 +8,10 @@
 #include "MPDElementsDASH.h"
 
 #include "Player/AdaptiveStreamingPlayerResourceRequest.h"
+#include "Utilities/URLParser.h"
+
+
+//#define LOG_DASHINT_CTORDTOR
 
 namespace Electra
 {
@@ -102,7 +106,7 @@ namespace DASHUrlHelpers
 	};
 	void GetAllHierarchyUrlQueries(TArray<TSharedPtrTS<FDashMPD_UrlQueryInfoType>>& OutUrlQueries, TSharedPtrTS<const IDashMPDElement> StartingAt, EUrlQueryRequestType ForRequestType, bool bInclude2014);
 	FErrorDetail ApplyUrlQueries(IPlayerSessionServices* PlayerSessionServices, const FString& InMPDUrl, FString& InOutURL, FString& OutRequestHeader, const TArray<TSharedPtrTS<FDashMPD_UrlQueryInfoType>>& UrlQueries);
-	bool BuildAbsoluteElementURL(FString& OutURL, const FString& DocumentURL, const TArray<TSharedPtrTS<const FDashMPD_BaseURLType>>& BaseURLs, const FString& ElementURL);
+	bool BuildAbsoluteElementURL(FString& OutURL, FTimeValue& ATO, const FString& DocumentURL, const TArray<TSharedPtrTS<const FDashMPD_BaseURLType>>& BaseURLs, const FString& ElementURL);
 	FString ApplyAnnexEByteRange(FString InURL, FString InRange, const TArray<TSharedPtrTS<const FDashMPD_BaseURLType>>& OutBaseURLs);
 }
 
@@ -129,7 +133,7 @@ public:
 		};
 		FURL InitializationURL;
 		FURL MediaURL;
-		FTimeValue AvailabilityTimeOffset;			//!< ATO, if applicable.
+		FTimeValue ATO;
 		int64 Time = 0;								//!< Time value T in timescale units without PTO, EPT or ELST applied
 		int64 PTO = 0;								//!< PresentationTimeOffset
 		int64 EPTdelta = 0;
@@ -144,14 +148,50 @@ public:
 		bool bMayBeMissing = false;					//!< true if the last segment in <SegmentTemplate> that might not exist.
 		bool bIsMissing = false;					//!< Set to true if known to be missing.
 		bool bSawLMSG = false;						//!< Will be set to true by the stream reader if the 'lmsg' brand was found.
+
+		FTimeValue CalculateASAST(const FTimeValue& AST, const FTimeValue& PeriodStart, bool bIsStatic)
+		{
+			if (bIsStatic)
+			{
+				return AST;
+			}
+			else
+			{
+				if (ATO < FTimeValue::GetPositiveInfinity())
+				{
+					return AST + PeriodStart + FTimeValue(Time - PTO - EPTdelta + Duration, Timescale) - ATO;
+				}
+				// ATO of infinity means the segment is always available, so we return zero time as earliest UTC time.
+				return FTimeValue::GetZero();
+			}
+		}
+		FTimeValue CalculateSAET(const FTimeValue& AST, const FTimeValue& PeriodStart, const FTimeValue& MPDAET, const FTimeValue& TSB, bool bIsStatic)
+		{
+			if (bIsStatic)
+			{
+				// If the MPD has a global availabilityEndTime then that is the end time of the segment as well.
+				if (MPDAET.IsValid())
+				{
+					return MPDAET;
+				}
+				return FTimeValue::GetPositiveInfinity();
+			}
+			else
+			{
+				// If the MPD has a global availabilityEndTime then that is the end time of the segment as well.
+				if (MPDAET.IsValid())
+				{
+					return MPDAET;
+				}
+				return AST + PeriodStart + FTimeValue(Time - PTO - EPTdelta + Duration * 2, Timescale) + (TSB.IsValid() ? TSB : FTimeValue::GetPositiveInfinity());
+			}
+		}
 	};
 
 	struct FSegmentSearchOption
 	{
-		FTimeValue PeriodLocalTime;
-		FTimeValue PeriodDuration;
-		int64 Number = -1;							//!< Explicit segment number
-		int64 Time = -1;							//!< Explicit segment time
+		FTimeValue PeriodLocalTime;					//!< Time local in the period to search a segment for.
+		FTimeValue PeriodDuration;					//!< Duration of the period. Needed to determine the number of segments in the period.
 		IManifest::ESearchType SearchType = IManifest::ESearchType::Closest;
 		int64 RequestID = 0;						//!< Sequential request ID across all segments during playback, needed to re-resolve potential UrlQueryInfo xlinks.
 	};
@@ -159,9 +199,13 @@ public:
 	class FRepresentation : public IPlaybackAssetRepresentation, public TSharedFromThis<FRepresentation, ESPMode::ThreadSafe>
 	{
 	public:
-		FRepresentation()
-		{ }
-		virtual ~FRepresentation() = default;
+		#ifdef LOG_DASHINT_CTORDTOR
+			FRepresentation();
+			virtual ~FRepresentation();
+		#else
+			FRepresentation() = default;
+			virtual ~FRepresentation() = default;
+		#endif
 
 		enum ESearchResult
 		{
@@ -182,8 +226,7 @@ public:
 		//
 		virtual FString GetUniqueIdentifier() const override
 		{
-			TSharedPtrTS<FDashMPD_RepresentationType> MPDRepresentation = Representation.Pin();
-			return MPDRepresentation.IsValid() ? MPDRepresentation->GetID() : FString(TEXT(" invalid name "));
+			return ID;
 		}
 		virtual const FStreamCodecInformation& GetCodecInformation() const override
 		{
@@ -191,13 +234,7 @@ public:
 		}
 		virtual int32 GetBitrate() const override
 		{
-			TSharedPtrTS<FDashMPD_RepresentationType> MPDRepresentation = Representation.Pin();
-			if (MPDRepresentation.IsValid())
-			{
-				return (int32) MPDRepresentation->GetBandwidth();
-			}
-			// Return max possible value to make sure this will never get chosen.
-			return TNumericLimits<int32>::Max();
+			return Bitrate;
 		}
 		virtual int32 GetQualityIndex() const override
 		{
@@ -223,6 +260,8 @@ public:
 		friend class FManifestDASHInternal;
 		TWeakPtrTS<FDashMPD_RepresentationType> Representation;
 		FStreamCodecInformation CodecInfo;
+		FString ID;
+		int32 Bitrate = MAX_int32;
 		int32 QualityIndex = -1;
 		bool bIsUsable = false;
 		bool bIsEnabled = true;
@@ -231,6 +270,7 @@ public:
 		bool bWarnedAboutTimelineNumberOverflow = false;
 		bool bWarnedAboutInconsistentNumbering = false;
 		bool bWarnedAboutTimelineOverlap = false;
+		bool bWarnedAboutTimescale = false;
 		//
 		bool bNeedsSegmentIndex = true;
 		TSharedPtrTS<const IParserISO14496_12> SegmentIndex;
@@ -242,9 +282,13 @@ public:
 	class FAdaptationSet : public IPlaybackAssetAdaptationSet
 	{
 	public:
-		FAdaptationSet()
-		{ }
-		virtual ~FAdaptationSet() = default;
+		#ifdef LOG_DASHINT_CTORDTOR
+			FAdaptationSet();
+			virtual ~FAdaptationSet();
+		#else
+			FAdaptationSet() = default;
+			virtual ~FAdaptationSet() = default;
+		#endif
 
 		const FStreamCodecInformation& GetCodec() const								{ return Codec; }
 		const TArray<TSharedPtrTS<FRepresentation>>& GetRepresentations() const		{ return Representations; }
@@ -255,6 +299,18 @@ public:
 		int32 GetMaxBandwidth() const												{ return MaxBandwidth; }
 		bool GetIsUsable() const													{ return bIsUsable; }
 
+		TSharedPtrTS<FRepresentation> GetRepresentationByUniqueID(const FString& UniqueIdentifier) const
+		{
+			for(int32 i=0; i<Representations.Num(); ++i)
+			{
+				if (Representations[i]->GetUniqueIdentifier().Equals(UniqueIdentifier))
+				{
+					return Representations[i];
+				}
+			}
+			return nullptr;
+		}
+
 		//----------------------------------------------
 		// Methods from IPlaybackAssetAdaptationSet
 		//
@@ -284,14 +340,7 @@ public:
 		}
 		virtual TSharedPtrTS<IPlaybackAssetRepresentation> GetRepresentationByUniqueIdentifier(const FString& UniqueIdentifier) const override
 		{
-			for(int32 i=0; i<Representations.Num(); ++i)
-			{
-				if (Representations[i]->GetUniqueIdentifier().Equals(UniqueIdentifier))
-				{
-					return Representations[i];
-				}
-			}
-			return nullptr;
+			return GetRepresentationByUniqueID(UniqueIdentifier);
 		}
 	private:
 		friend class FManifestDASHInternal;
@@ -311,9 +360,22 @@ public:
 	class FPeriod : public ITimelineMediaAsset
 	{
 	public:
-		FPeriod()
-		{ }
-		virtual ~FPeriod() = default;
+		#ifdef LOG_DASHINT_CTORDTOR
+			FPeriod();
+			virtual ~FPeriod();
+		#else
+			FPeriod() = default;
+			virtual ~FPeriod() = default;
+		#endif
+
+		bool GetHasBeenPrepared() const
+		{
+			return bHasBeenPrepared;
+		}
+		void SetHasBeenPrepared(bool bPrepared)
+		{
+			bHasBeenPrepared = bPrepared;
+		}
 
 		const FString& GetID() const { return ID; }
 		const FTimeValue& GetStart() const { return Start; }
@@ -321,6 +383,28 @@ public:
 		bool GetIsEarlyPeriod() const { return bIsEarlyPeriod; }
 		const TArray<TSharedPtrTS<FAdaptationSet>>& GetAdaptationSets() const { return AdaptationSets; }
 
+		TSharedPtrTS<FAdaptationSet> GetAdaptationSetByUniqueID(const FString& UniqueID) const
+		{
+			for(int32 i=0; i<AdaptationSets.Num(); ++i)
+			{
+				if (AdaptationSets[i]->GetUniqueIdentifier().Equals(UniqueID))
+				{
+					return AdaptationSets[i];
+				}
+			}
+			return nullptr;
+		}
+/*
+		TSharedPtrTS<FRepresentation> GetRepresentationFromAdaptationSetIDs(const FString& AdaptationSetID, const FString& RepresentationID) const
+		{
+			TSharedPtrTS<FAdaptationSet> Adapt = GetAdaptationSetByUniqueID(AdaptationSetID);
+			if (Adapt.IsValid())
+			{
+				return Adapt->GetRepresentationByUniqueID(RepresentationID);
+			}
+			return nullptr;
+		}
+*/
 		TSharedPtrTS<FDashMPD_PeriodType> GetMPDPeriod() { return Period.Pin(); }
 
 		//----------------------------------------------
@@ -373,65 +457,151 @@ public:
 
 	private:
 		friend class FManifestDASHInternal;
+		TArray<TSharedPtrTS<FAdaptationSet>> AdaptationSets;
 		TWeakPtrTS<FDashMPD_PeriodType> Period;
 		FString ID;
 		FTimeValue Start;
 		FTimeValue End;
 		FTimeValue Duration;
 		bool bIsEarlyPeriod = false;
-		TArray<TSharedPtrTS<FAdaptationSet>> AdaptationSets;
+		bool bHasBeenPrepared = false;
 	};
 
-	FErrorDetail Build(IPlayerSessionServices* InPlayerSessionServices, TSharedPtr<FDashMPD_MPDType, ESPMode::ThreadSafe> InMPDRoot, TArray<TWeakPtrTS<IDashMPDElement>> InXLinkElements, const FStreamPreferences& InPreferences, const FParamDict& InOptions);
+	FErrorDetail Build(IPlayerSessionServices* InPlayerSessionServices, TSharedPtr<FDashMPD_MPDType, ESPMode::ThreadSafe> InMPDRoot, TArray<TWeakPtrTS<IDashMPDElement>> InXLinkElements);
 
-	FErrorDetail BuildAfterInitialRemoteElementDownload(IPlayerSessionServices* InPlayerSessionServices);
+	FErrorDetail BuildAfterInitialRemoteElementDownload();
 
 	void GetRemoteElementLoadRequests(TArray<TWeakPtrTS<FMPDLoadRequestDASH>>& OutRemoteElementLoadRequests)
 	{
 		OutRemoteElementLoadRequests = PendingRemoteElementLoadRequests;
 	}
 
-	FErrorDetail ResolveInitialRemoteElementRequest(IPlayerSessionServices* PlayerSessionServices, TSharedPtrTS<FMPDLoadRequestDASH> RequestResponse, FString XMLResponse, bool bSuccess);
+	FErrorDetail ResolveInitialRemoteElementRequest(TSharedPtrTS<FMPDLoadRequestDASH> RequestResponse, FString XMLResponse, bool bSuccess);
 
 	EPresentationType GetPresentationType() const
 	{
 		return PresentationType;
 	}
 
+	bool IsEventType() const
+	{
+		return bIsEventType;
+	}
+
 	const TArray<TSharedPtrTS<FPeriod>>& GetPeriods() const
 	{
 		return Periods;
 	}
 
+	TSharedPtrTS<FManifestDASHInternal::FPeriod> GetPeriodByUniqueID(const FString& InUniqueID) const
+	{
+		for(int32 i=0, iMax=Periods.Num(); i<iMax; ++i)
+		{
+			if (Periods[i]->GetUniqueIdentifier().Equals(InUniqueID))
+			{
+				return Periods[i];
+			}
+		}
+		return nullptr;
+	}
+
+	bool HasFollowingRegularPeriod(TSharedPtrTS<FPeriod> InPeriod) const
+	{
+		for(int32 i=0, iMax=Periods.Num(); i<iMax; ++i)
+		{
+			if (Periods[i] == InPeriod)
+			{
+				// Is there a following period and is it a regular period?
+				if (i+1 < iMax && !Periods[i+1]->GetIsEarlyPeriod())
+				{
+					return true;
+				}
+				return false;
+			}
+		}
+		return false;
+	}
+
 	TSharedPtrTS<const FDashMPD_MPDType> GetMPDRoot() const
 	{
 		return MPDRoot;
 	}
 
-	FErrorDetail PreparePeriodAdaptationSets(IPlayerSessionServices* PlayerSessionServices, TSharedPtrTS<FPeriod> InPeriod, bool bRequestXlink);
+	TSharedPtrTS<FDashMPD_MPDType> GetMPDRoot()
+	{
+		return MPDRoot;
+	}
+
+	void PreparePeriodAdaptationSets(TSharedPtrTS<FPeriod> InPeriod, bool bRequestXlink);
+
+	const TArray<FURL_RFC3986::FQueryParam>& GetURLFragmentComponents() const
+	{
+		return URLFragmentComponents;
+	}
+	
+	void SetURLFragmentComponents(TArray<FURL_RFC3986::FQueryParam> InURLFragmentComponents)
+	{
+		URLFragmentComponents = MoveTemp(InURLFragmentComponents);
+	}
+
+	const TArray<FURL_RFC3986::FQueryParam>& GetURLFragmentComponents()
+	{
+		return URLFragmentComponents;
+	}
+
+	void TransformIntoEpicEvent();
 
 	ELECTRA_IMPL_DEFAULT_ERROR_METHODS(DASHManifest);
 
+
+	FTimeValue GetAnchorTime() const;
+	FTimeRange GetTotalTimeRange() const;
+	FTimeRange GetSeekableTimeRange() const;
+	void GetSeekablePositions(TArray<FTimespan>& OutPositions) const;
+	FTimeValue GetDuration() const;
+	FTimeValue GetDefaultStartTime() const;
+
+	FTimeValue GetMPDValidityEndTime() const;
+	FTimeValue GetLastPeriodEndTime() const;
+	FTimeValue GetMinimumUpdatePeriod() const;
+	bool AreUpdatesExpected() const;
+	bool IsStaticType() const;
+	bool UsesAST() const;
+	FTimeValue GetAvailabilityEndTime() const;
+	FTimeValue GetTimeshiftBufferDepth() const;
+
+#ifdef LOG_DASHINT_CTORDTOR
+	FManifestDASHInternal();
+	~FManifestDASHInternal();
+#endif
+
 private:
-	FErrorDetail PrepareRemoteElementLoadRequest(IPlayerSessionServices* PlayerSessionServices, TArray<TWeakPtrTS<FMPDLoadRequestDASH>>& OutRemoteElementLoadRequests, TWeakPtrTS<IDashMPDElement> ElementWithXLink, int64 RequestID);
+	FErrorDetail PrepareRemoteElementLoadRequest(TArray<TWeakPtrTS<FMPDLoadRequestDASH>>& OutRemoteElementLoadRequests, TWeakPtrTS<IDashMPDElement> ElementWithXLink, int64 RequestID);
 
 	int32 ReplaceElementWithRemoteEntities(TSharedPtrTS<IDashMPDElement> Element, const FDashMPD_RootEntities& NewRootEntities, int64 OldResolveID, int64 NewResolveID);
 
+	FTimeValue CalculateDistanceToLiveEdge() const;
+	FTimeRange GetPlayTimesFromURI() const;
+
+	IPlayerSessionServices* PlayerSessionServices = nullptr;
 
 	TArray<TWeakPtrTS<IDashMPDElement>> RemoteElementsToResolve;
 	TArray<TWeakPtrTS<FMPDLoadRequestDASH>> PendingRemoteElementLoadRequests;
 
-	FStreamPreferences Preferences;
-	FParamDict Options;
-
 	// The parsed MPD.
 	TSharedPtrTS<FDashMPD_MPDType> MPDRoot;
 
+	// The MPD URL fragment components
+	TArray<FURL_RFC3986::FQueryParam> URLFragmentComponents;
+
 	// Type of the presentation.
 	EPresentationType PresentationType;
+	bool bIsEventType = false;
 
 	TArray<TSharedPtrTS<FPeriod>> Periods;
 
+	mutable FTimeRange TotalTimeRange;
+	mutable FTimeRange SeekableTimeRange;
 };
 
 
diff --git a/Engine/Plugins/Media/ElectraPlayer/Source/ElectraPlayerRuntime/Private/Runtime/Player/DASH/StreamReaderFMP4DASH.cpp b/Engine/Plugins/Media/ElectraPlayer/Source/ElectraPlayerRuntime/Private/Runtime/Player/DASH/StreamReaderFMP4DASH.cpp
--- a/Engine/Plugins/Media/ElectraPlayer/Source/ElectraPlayerRuntime/Private/Runtime/Player/DASH/StreamReaderFMP4DASH.cpp
+++ b/Engine/Plugins/Media/ElectraPlayer/Source/ElectraPlayerRuntime/Private/Runtime/Player/DASH/StreamReaderFMP4DASH.cpp
@@ -11,6 +11,8 @@
 #include "Player/AdaptiveStreamingPlayerABR.h"
 #include "Player/PlayerEntityCache.h"
 #include "Player/DASH/PlaylistReaderDASH.h"
+#include "Player/DASH/OptionKeynamesDASH.h"
+#include "Utilities/Utilities.h"
 
 #define INTERNAL_ERROR_INIT_SEGMENT_DOWNLOAD_ERROR					1
 #define INTERNAL_ERROR_INIT_SEGMENT_PARSE_ERROR						2
@@ -70,7 +72,7 @@ void FStreamSegmentRequestFMP4DASH::GetEndedStreams(TArray<TSharedPtrTS<IStreamS
 
 FTimeValue FStreamSegmentRequestFMP4DASH::GetFirstPTS() const
 {
-	return AvailabilityStartTime + AdditionalAdjustmentTime + PeriodStart + FTimeValue(Segment.Time - Segment.PTO, Segment.Timescale);
+	return AST + AdditionalAdjustmentTime + PeriodStart + FTimeValue(Segment.Time - Segment.PTO, Segment.Timescale);
 }
 
 int32 FStreamSegmentRequestFMP4DASH::GetQualityIndex() const
@@ -364,7 +366,7 @@ FErrorDetail FStreamReaderFMP4DASH::FStreamHandler::GetInitSegment(TSharedPtrTS<
 	LoadReq->Range = Request->Segment.InitializationURL.Range;
 	if (Request->Segment.InitializationURL.CustomHeader.Len())
 	{
-		LoadReq->Headers.Emplace(HTTP::FHTTPHeader({TEXT("MPEG-DASH-Param"), Request->Segment.InitializationURL.CustomHeader}));
+		LoadReq->Headers.Emplace(HTTP::FHTTPHeader({DASH::HTTPHeaderOptionName, Request->Segment.InitializationURL.CustomHeader}));
 	}
 	LoadReq->PlayerSessionServices = PlayerSessionService;
 	LoadReq->CompleteCallback.BindLambda([LoadResult](TSharedPtrTS<FMPDLoadRequestDASH> LoadRequest, bool bSuccess)
@@ -501,6 +503,37 @@ void FStreamReaderFMP4DASH::FStreamHandler::HandleRequest()
 	ds.RetryNumber  	= Request->NumOverallRetries;
 	ds.ABRState.Reset();
 
+
+	// Does this request have an availability window?
+	if (Request->ASAST.IsValid())
+	{
+		FTimeValue Now = PlayerSessionService->GetSynchronizedUTCTime()->GetTime();
+#if 0
+		if (Now < Request->ASAST)
+		{
+			LogMessage(IInfoLog::ELevel::Info, FString::Printf(TEXT("waiting fr ASAST %lld"), (long long int) (Request->ASAST - Now).GetAsMilliseconds() ));
+		}
+#endif
+		// Availability start time not reached yet?
+		while(Now < Request->ASAST)
+		{
+			if (HasReadBeenAborted())
+			{
+				CurrentRequest.Reset();
+				return;
+			}
+			const int64 SleepIntervalMS = 100;
+			FMediaRunnable::SleepMilliseconds((uint32) Utils::Max(Utils::Min((Request->ASAST - Now).GetAsMilliseconds(), SleepIntervalMS), (int64)0));
+			Now = PlayerSessionService->GetSynchronizedUTCTime()->GetTime();
+		}
+
+		// Already expired?
+		if (Now > Request->SAET)
+		{
+			// TODO:
+		}
+	}
+
 	// Clear internal work variables.
 	bHasErrored 		   = false;
 	bAbortedByABR   	   = false;
@@ -512,7 +545,7 @@ void FStreamReaderFMP4DASH::FStreamHandler::HandleRequest()
 
 	FTimeValue NextExpectedDTS;
 	FTimeValue LastKnownAUDuration;
-	FTimeValue TimeOffset = Request->PeriodStart + Request->AvailabilityStartTime + Request->AdditionalAdjustmentTime;
+	FTimeValue TimeOffset = Request->PeriodStart + Request->AST + Request->AdditionalAdjustmentTime;
 
 	// Get the init segment if there is one. Either gets it from the entity cache or requests it now and adds it to the cache.
 	InitSegmentError = GetInitSegment(MP4InitSegment, Request);
@@ -552,7 +585,7 @@ void FStreamReaderFMP4DASH::FStreamHandler::HandleRequest()
 			HTTP->Parameters.AcceptEncoding.Set(TEXT("identity"));
 			if (Request->Segment.MediaURL.CustomHeader.Len())
 			{
-				HTTP->Parameters.RequestHeaders.Emplace(HTTP::FHTTPHeader({TEXT("MPEG-DASH-Param"), Request->Segment.MediaURL.CustomHeader}));
+				HTTP->Parameters.RequestHeaders.Emplace(HTTP::FHTTPHeader({DASH::HTTPHeaderOptionName, Request->Segment.MediaURL.CustomHeader}));
 			}
 
 			ProgressReportCount = 0;
@@ -589,9 +622,20 @@ void FStreamReaderFMP4DASH::FStreamHandler::HandleRequest()
 								IParserISO14496_12::ITrackIterator* TrackIterator = Track->CreateIterator();
 								TrackTimescale = TrackIterator->GetTimescale();
 
+								int64 MediaLocalFirstAUTime = Request->Segment.MediaLocalFirstAUTime;
+								int64 MediaLocalLastAUTime = Request->Segment.MediaLocalLastAUTime;
+
 								if (TrackTimescale != Request->Segment.Timescale)
 								{
-									LogMessage(IInfoLog::ELevel::Warning, FString::Printf(TEXT("Track timescale %u differs from timescale of %u in MPD or segment index. This may cause playback problems!"), TrackTimescale, Request->Segment.Timescale));
+									// Need to rescale the AU times from the MPD timescale to the media timescale.
+									MediaLocalFirstAUTime = FTimeFraction(Request->Segment.MediaLocalFirstAUTime, Request->Segment.Timescale).GetAsTimebase(TrackTimescale);
+									MediaLocalLastAUTime = FTimeFraction(Request->Segment.MediaLocalLastAUTime, Request->Segment.Timescale).GetAsTimebase(TrackTimescale);
+
+									if (!Request->bWarnedAboutTimescale)
+									{
+										Request->bWarnedAboutTimescale = true;
+										LogMessage(IInfoLog::ELevel::Warning, FString::Printf(TEXT("Track timescale %u differs from timescale of %u in MPD or segment index. This may cause playback problems!"), TrackTimescale, Request->Segment.Timescale));
+									}
 								}
 
 								for(Error = TrackIterator->StartAtFirst(false); Error == UEMEDIA_ERROR_OK; Error = TrackIterator->Next())
@@ -614,32 +658,38 @@ void FStreamReaderFMP4DASH::FStreamHandler::HandleRequest()
 
 									AccessUnit->DropState = FAccessUnit::EDropState::None;
 									// Is this AU entirely before the time we want?
-									int64 Overlap = AUPTS + AUDuration - Request->Segment.MediaLocalFirstAUTime;
+									int64 Overlap = AUPTS + AUDuration - MediaLocalFirstAUTime;
 									if (Overlap <= 0)
 									{
 										AccessUnit->DropState |= FAccessUnit::EDropState::PtsTooEarly;
 										AccessUnit->DropState |= FAccessUnit::EDropState::DtsTooEarly;
 									}
 									// Partial overlap?
-									else if (AUPTS < Request->Segment.MediaLocalFirstAUTime)
+									else if (AUPTS < MediaLocalFirstAUTime)
 									{
 										// Use it, but remember the amount of overlap.
 										AccessUnit->OverlapAdjust.SetFromND(Overlap, TrackTimescale);
 									}
-									// Entirely past the time allowed?
-									if (AUPTS >= Request->Segment.MediaLocalLastAUTime)
-									{
-										AccessUnit->DropState |= FAccessUnit::EDropState::PtsTooLate;
-										AccessUnit->DropState |= FAccessUnit::EDropState::DtsTooLate;
-									}
-									else if (AUPTS + AUDuration >= Request->Segment.MediaLocalLastAUTime)
-									{
-										// Note that we cannot stop reading here. For one we need to continue in case there are additional tracks here
-										// that have not yet finished and we also need to make it all the way to the end in case there are additional
-										// boxes that we have to read (like 'emsg' or 'lmsg').
-										AccessUnit->OverlapAdjust.SetFromND(AUPTS + AUDuration - Request->Segment.MediaLocalLastAUTime, TrackTimescale);
-									}
-
+									/*
+										Taken out for now since varying segment durations and possible timescale conversions
+										can lead to inaccuracies. We need a flag saying this is the last segment of the period
+										and only when that is set and the period has a valid duration should we mark the
+										overlapping AUs. Unless the period has a duration this should never be set.
+						
+										// Entirely past the time allowed?
+										if (AUPTS >= MediaLocalLastAUTime)
+										{
+											AccessUnit->DropState |= FAccessUnit::EDropState::PtsTooLate;
+											AccessUnit->DropState |= FAccessUnit::EDropState::DtsTooLate;
+										}
+										else if (AUPTS + AUDuration >= MediaLocalLastAUTime)
+										{
+											// Note that we cannot stop reading here. For one we need to continue in case there are additional tracks here
+											// that have not yet finished and we also need to make it all the way to the end in case there are additional
+											// boxes that we have to read (like 'emsg' or 'lmsg').
+											AccessUnit->OverlapAdjust.SetFromND(AUPTS + AUDuration - MediaLocalLastAUTime, TrackTimescale);
+										}
+									*/
 									FTimeValue Duration(TrackIterator->GetDuration(), TrackTimescale);
 									AccessUnit->Duration = Duration;
 
@@ -979,6 +1029,17 @@ void FStreamReaderFMP4DASH::FStreamHandler::HandleRequest()
 		CurrentRequest->Segment.bIsMissing = true;
 	}
 
+	// If we had to wait for the segment to become available and we got a 404 back we might have been trying to fetch
+	// the segment before the server made it available.
+	if (Request->ASAST.IsValid() && (ds.HTTPStatusCode == 404 || ds.HTTPStatusCode == 416))
+	{
+		FTimeValue Now = PlayerSessionService->GetSynchronizedUTCTime()->GetTime();
+		if ((ds.AvailibilityDelay = (Request->ASAST - Now).GetAsSeconds()) == 0.0)
+		{
+			// In the extremely unlikely event this comes out to zero exactly set a small value so the ABR knows there was a delay.
+			ds.AvailibilityDelay = -0.01;
+		}
+	}
 
 	StreamSelector->ReportDownloadEnd(ds);
 
diff --git a/Engine/Plugins/Media/ElectraPlayer/Source/ElectraPlayerRuntime/Private/Runtime/Player/DASH/StreamReaderFMP4DASH.h b/Engine/Plugins/Media/ElectraPlayer/Source/ElectraPlayerRuntime/Private/Runtime/Player/DASH/StreamReaderFMP4DASH.h
--- a/Engine/Plugins/Media/ElectraPlayer/Source/ElectraPlayerRuntime/Private/Runtime/Player/DASH/StreamReaderFMP4DASH.h
+++ b/Engine/Plugins/Media/ElectraPlayer/Source/ElectraPlayerRuntime/Private/Runtime/Player/DASH/StreamReaderFMP4DASH.h
@@ -44,10 +44,14 @@ public:
 	bool													bIsEOSSegment = false;								//!< true if this is not an actual request but a stream-has-already-ended request.
 	bool													bIsInitialStartRequest = false;						//!< true if this is the initial playback start request.
 	FTimeValue												PeriodStart;										//!< Value to add to all DTS & PTS to map them into the Period timeline
-	FTimeValue												AvailabilityStartTime = FTimeValue::GetZero();		//!< Value of AST to add to all time to generate wallclock time
+	FTimeValue												AST = FTimeValue::GetZero();						//!< Value of AST to add to all time to generate wallclock time
 	FTimeValue												AdditionalAdjustmentTime = FTimeValue::GetZero();	//!< Sum of any other time corrections
 	bool													bInsertFillerData = false;							//!< true to insert empty access units into the buffer instead of reading actual data.
 
+	// UTC wallclock times during which this segment can be fetched;
+	FTimeValue												ASAST;
+	FTimeValue												SAET;
+
 	// Internal work variables
 	FPlayerLoopState										PlayerLoopState;
 	int32													NumOverallRetries = 0;								//!< Number of retries for this _segment_ across all possible quality levels and CDNs.
@@ -56,6 +60,8 @@ public:
 	Metrics::FSegmentDownloadStats							DownloadStats;
 	HTTP::FConnectionInfo									ConnectionInfo;
 	FTimeValue												NextLargestExpectedTimestamp;						//!< Largest timestamp of all samples (plus its duration) across all tracks.
+
+	bool													bWarnedAboutTimescale = false;
 };
 
 
diff --git a/Engine/Plugins/Media/ElectraPlayer/Source/ElectraPlayerRuntime/Private/Runtime/Player/HLS/InitSegmentCacheHLS.cpp b/Engine/Plugins/Media/ElectraPlayer/Source/ElectraPlayerRuntime/Private/Runtime/Player/HLS/InitSegmentCacheHLS.cpp
--- a/Engine/Plugins/Media/ElectraPlayer/Source/ElectraPlayerRuntime/Private/Runtime/Player/HLS/InitSegmentCacheHLS.cpp
+++ b/Engine/Plugins/Media/ElectraPlayer/Source/ElectraPlayerRuntime/Private/Runtime/Player/HLS/InitSegmentCacheHLS.cpp
@@ -12,7 +12,7 @@ namespace Electra
 class FInitSegmentCacheHLS : public IInitSegmentCacheHLS
 {
 public:
-	FInitSegmentCacheHLS(IPlayerSessionServices* SessionServices, const FParamDict& Options);
+	FInitSegmentCacheHLS(IPlayerSessionServices* SessionServices);
 	virtual ~FInitSegmentCacheHLS();
 	virtual TSharedPtrTS<const IParserISO14496_12> GetInitSegmentFor(const TSharedPtrTS<const FManifestHLSInternal::FMediaStream::FInitSegmentInfo>& InitSegmentInfo) override;
 	virtual void AddInitSegment(const TSharedPtrTS<const IParserISO14496_12>& InitSegment, const TSharedPtrTS<const FManifestHLSInternal::FMediaStream::FInitSegmentInfo>& InitSegmentInfo, const FTimeValue& ExpiresAtUTC) override;
@@ -20,7 +20,6 @@ private:
 
 	static const uint32 kMaxCacheEntries = 32;
 
-	FParamDict				Options;
 	IPlayerSessionServices* SessionServices;
 
 	struct FEntry
@@ -38,18 +37,17 @@ private:
 /***************************************************************************************************************************************************/
 /***************************************************************************************************************************************************/
 
-IInitSegmentCacheHLS* IInitSegmentCacheHLS::Create(IPlayerSessionServices* SessionServices, const FParamDict& Options)
+IInitSegmentCacheHLS* IInitSegmentCacheHLS::Create(IPlayerSessionServices* SessionServices)
 {
-	return new FInitSegmentCacheHLS(SessionServices, Options);
+	return new FInitSegmentCacheHLS(SessionServices);
 }
 
 /***************************************************************************************************************************************************/
 /***************************************************************************************************************************************************/
 /***************************************************************************************************************************************************/
 
-FInitSegmentCacheHLS::FInitSegmentCacheHLS(IPlayerSessionServices* InSessionServices, const FParamDict& InOptions)
-	: Options(InOptions)
-	, SessionServices(InSessionServices)
+FInitSegmentCacheHLS::FInitSegmentCacheHLS(IPlayerSessionServices* InSessionServices)
+	: SessionServices(InSessionServices)
 	, Cache(kMaxCacheEntries)
 {
 }
diff --git a/Engine/Plugins/Media/ElectraPlayer/Source/ElectraPlayerRuntime/Private/Runtime/Player/HLS/InitSegmentCacheHLS.h b/Engine/Plugins/Media/ElectraPlayer/Source/ElectraPlayerRuntime/Private/Runtime/Player/HLS/InitSegmentCacheHLS.h
--- a/Engine/Plugins/Media/ElectraPlayer/Source/ElectraPlayerRuntime/Private/Runtime/Player/HLS/InitSegmentCacheHLS.h
+++ b/Engine/Plugins/Media/ElectraPlayer/Source/ElectraPlayerRuntime/Private/Runtime/Player/HLS/InitSegmentCacheHLS.h
@@ -16,7 +16,7 @@ class IParserISO14496_12;
 class IInitSegmentCacheHLS
 {
 public:
-	static IInitSegmentCacheHLS* Create(IPlayerSessionServices* SessionServices, const FParamDict& Options);
+	static IInitSegmentCacheHLS* Create(IPlayerSessionServices* SessionServices);
 
 	virtual ~IInitSegmentCacheHLS() = default;
 
diff --git a/Engine/Plugins/Media/ElectraPlayer/Source/ElectraPlayerRuntime/Private/Runtime/Player/HLS/LicenseKeyCacheHLS.cpp b/Engine/Plugins/Media/ElectraPlayer/Source/ElectraPlayerRuntime/Private/Runtime/Player/HLS/LicenseKeyCacheHLS.cpp
--- a/Engine/Plugins/Media/ElectraPlayer/Source/ElectraPlayerRuntime/Private/Runtime/Player/HLS/LicenseKeyCacheHLS.cpp
+++ b/Engine/Plugins/Media/ElectraPlayer/Source/ElectraPlayerRuntime/Private/Runtime/Player/HLS/LicenseKeyCacheHLS.cpp
@@ -11,7 +11,7 @@ namespace Electra
 class FLicenseKeyCacheHLS : public ILicenseKeyCacheHLS
 {
 public:
-	FLicenseKeyCacheHLS(IPlayerSessionServices* SessionServices, const FParamDict& Options);
+	FLicenseKeyCacheHLS(IPlayerSessionServices* SessionServices);
 	virtual ~FLicenseKeyCacheHLS();
 
 	virtual TSharedPtr<TArray<uint8>, ESPMode::ThreadSafe> GetLicenseKeyFor(const TSharedPtr<const FManifestHLSInternal::FMediaStream::FDRMKeyInfo, ESPMode::ThreadSafe>& LicenseKeyInfo) override;
@@ -21,7 +21,6 @@ private:
 
 	static const uint32 kMaxCacheEntries = 32;
 
-	FParamDict				Options;
 	IPlayerSessionServices* SessionServices;
 
 	struct FEntry
@@ -39,20 +38,18 @@ private:
 /***************************************************************************************************************************************************/
 /***************************************************************************************************************************************************/
 
-ILicenseKeyCacheHLS* ILicenseKeyCacheHLS::Create(IPlayerSessionServices* SessionServices, const FParamDict& Options)
+ILicenseKeyCacheHLS* ILicenseKeyCacheHLS::Create(IPlayerSessionServices* SessionServices)
 {
-	return new FLicenseKeyCacheHLS(SessionServices, Options);
+	return new FLicenseKeyCacheHLS(SessionServices);
 }
 
 /***************************************************************************************************************************************************/
 /***************************************************************************************************************************************************/
 /***************************************************************************************************************************************************/
 
-FLicenseKeyCacheHLS::FLicenseKeyCacheHLS(IPlayerSessionServices* InSessionServices, const FParamDict& InOptions)
-	: Options(InOptions)
-	, SessionServices(InSessionServices)
+FLicenseKeyCacheHLS::FLicenseKeyCacheHLS(IPlayerSessionServices* InSessionServices)
+	: SessionServices(InSessionServices)
 	, Cache(kMaxCacheEntries)
-
 {
 }
 
diff --git a/Engine/Plugins/Media/ElectraPlayer/Source/ElectraPlayerRuntime/Private/Runtime/Player/HLS/LicenseKeyCacheHLS.h b/Engine/Plugins/Media/ElectraPlayer/Source/ElectraPlayerRuntime/Private/Runtime/Player/HLS/LicenseKeyCacheHLS.h
--- a/Engine/Plugins/Media/ElectraPlayer/Source/ElectraPlayerRuntime/Private/Runtime/Player/HLS/LicenseKeyCacheHLS.h
+++ b/Engine/Plugins/Media/ElectraPlayer/Source/ElectraPlayerRuntime/Private/Runtime/Player/HLS/LicenseKeyCacheHLS.h
@@ -14,7 +14,7 @@ class IPlayerSessionServices;
 class ILicenseKeyCacheHLS
 {
 public:
-	static ILicenseKeyCacheHLS* Create(IPlayerSessionServices* SessionServices, const FParamDict& Options);
+	static ILicenseKeyCacheHLS* Create(IPlayerSessionServices* SessionServices);
 
 	virtual ~ILicenseKeyCacheHLS() = default;
 
diff --git a/Engine/Plugins/Media/ElectraPlayer/Source/ElectraPlayerRuntime/Private/Runtime/Player/HLS/ManifestBuilderHLS.cpp b/Engine/Plugins/Media/ElectraPlayer/Source/ElectraPlayerRuntime/Private/Runtime/Player/HLS/ManifestBuilderHLS.cpp
--- a/Engine/Plugins/Media/ElectraPlayer/Source/ElectraPlayerRuntime/Private/Runtime/Player/HLS/ManifestBuilderHLS.cpp
+++ b/Engine/Plugins/Media/ElectraPlayer/Source/ElectraPlayerRuntime/Private/Runtime/Player/HLS/ManifestBuilderHLS.cpp
@@ -9,6 +9,7 @@
 #include "Player/AdaptiveStreamingPlayerABR.h"
 #include "Player/PlayerSessionServices.h"
 #include "Player/PlayerStreamFilter.h"
+#include "Player/AdaptivePlayerOptionKeynames.h"
 
 #include "InitSegmentCacheHLS.h"
 #include "LicenseKeyCacheHLS.h"
@@ -56,6 +57,9 @@ namespace Electra
 namespace
 {
 
+static const TCHAR* const DefaultAssetNameHLS = TEXT("asset.0");
+
+
 static void SplitOnOneOf(TArray<FString>& OutResults, const FString& Subject, const FString& SplitAt)
 {
 	OutResults.Empty();
@@ -125,227 +129,6 @@ static bool IsH264Codec(TArray<FString>& OutResults, const FString& Codec)
 
 
 
-namespace HLSTimeline
-{
-
-class FPlaybackAssetRepresentation : public IPlaybackAssetRepresentation
-{
-public:
-	FPlaybackAssetRepresentation()
-	{ }
-	virtual ~FPlaybackAssetRepresentation()
-	{ }
-	virtual FString GetUniqueIdentifier() const override
-	{
-		return UniqueIdentifier;
-	}
-	virtual const FStreamCodecInformation& GetCodecInformation() const override
-	{
-		return CodecInformation;
-	}
-	virtual int32 GetBitrate() const override
-	{
-		return Bitrate;
-	}
-	virtual int32 GetQualityIndex() const override
-	{ 
-		return QualityIndex;
-	}
-	virtual bool CanBePlayed() const override
-	{
-		return true;
-	}
-	FStreamCodecInformation	CodecInformation;
-	FString					UniqueIdentifier;
-	int32					Bitrate = 0;
-	int32					QualityIndex = 0;
-};
-
-class FPlaybackAssetAdaptationSet : public IPlaybackAssetAdaptationSet
-{
-public:
-	virtual ~FPlaybackAssetAdaptationSet()
-	{
-	}
-	virtual FString GetUniqueIdentifier() const override
-	{
-		return UniqueIdentifier;
-	}
-	virtual FString GetListOfCodecs() const override
-	{
-		return Codecs;
-	}
-	virtual FString GetLanguage() const override
-	{
-		return Language;
-	}
-	virtual int32 GetNumberOfRepresentations() const override
-	{
-		return Representations.Num();
-	}
-	virtual TSharedPtrTS<IPlaybackAssetRepresentation> GetRepresentationByIndex(int32 RepresentationIndex) const override
-	{
-		check(RepresentationIndex < Representations.Num());
-		if (RepresentationIndex < Representations.Num())
-		{
-			return Representations[RepresentationIndex];
-		}
-		return TSharedPtrTS<IPlaybackAssetRepresentation>();
-	}
-	virtual TSharedPtrTS<IPlaybackAssetRepresentation> GetRepresentationByUniqueIdentifier(const FString& InUniqueIdentifier) const override
-	{
-		for(int32 i=0, iMax=Representations.Num(); i<iMax; ++i)
-		{
-			if (Representations[i]->GetUniqueIdentifier() == InUniqueIdentifier)
-			{
-				return Representations[i];
-			}
-		}
-		return TSharedPtrTS<IPlaybackAssetRepresentation>();
-	}
-
-	TArray<TSharedPtrTS<IPlaybackAssetRepresentation>>		Representations;
-	FString													UniqueIdentifier;
-	FString													Language;
-	FString													Codecs;
-};
-
-
-class FTimelineMediaAsset : public ITimelineMediaAsset
-{
-public:
-	virtual ~FTimelineMediaAsset()
-	{
-	}
-	virtual FTimeRange GetTimeRange() const override
-	{
-		return TimeRange;
-	}
-	virtual FTimeValue GetDuration() const override
-	{
-		return Duration;
-	}
-	virtual FString GetAssetIdentifier() const override
-	{
-		return AssetIdentifier;
-	}
-	virtual FString GetUniqueIdentifier() const override
-	{
-		return UniqueIdentifier;
-	}
-	virtual int32 GetNumberOfAdaptationSets(EStreamType OfStreamType) const override
-	{
-		switch(OfStreamType)
-		{
-			case EStreamType::Video:
-				return VideoAdaptationSet.IsValid() ? 1 : 0;
-			case EStreamType::Audio:
-				return (int32) AudioAdaptationSets.Num();
-			default:
-				return 0;
-		}
-	}
-	virtual TSharedPtrTS<IPlaybackAssetAdaptationSet> GetAdaptationSetByTypeAndIndex(EStreamType OfStreamType, int32 AdaptationSetIndex) const override
-	{
-		switch(OfStreamType)
-		{
-			case EStreamType::Video:
-			{
-				check(AdaptationSetIndex == 0);
-				return AdaptationSetIndex == 0 ? VideoAdaptationSet : TSharedPtrTS<FPlaybackAssetAdaptationSet>();
-			}
-			case EStreamType::Audio:
-			{
-				check(AdaptationSetIndex < AudioAdaptationSets.Num());
-				return AdaptationSetIndex < AudioAdaptationSets.Num() ? AudioAdaptationSets[AdaptationSetIndex] : TSharedPtrTS<FPlaybackAssetAdaptationSet>();
-			}
-			default:
-			{
-				return TSharedPtrTS<FPlaybackAssetAdaptationSet>();
-			}
-		}
-	}
-	TSharedPtrTS<IPlaybackAssetAdaptationSet> GetAdaptationSetByTypeAndUniqueIdentifier(EStreamType OfStreamType, const FString& InUniqueIdentifier) const
-	{
-		switch(OfStreamType)
-		{
-			case EStreamType::Video:
-			{
-				return VideoAdaptationSet;
-			}
-			case EStreamType::Audio:
-			{
-				for(int32 i=0, iMax=(int32)AudioAdaptationSets.Num(); i<iMax; ++i)
-				{
-					if (AudioAdaptationSets[i]->GetUniqueIdentifier() == InUniqueIdentifier)
-					{
-						return AudioAdaptationSets[i];
-					}
-				}
-				return TSharedPtrTS<FPlaybackAssetAdaptationSet>();
-			}
-			default:
-			{
-				return TSharedPtrTS<FPlaybackAssetAdaptationSet>();
-			}
-		}
-	}
-
-	TSharedPtrTS<FPlaybackAssetAdaptationSet>				VideoAdaptationSet;
-	TArray<TSharedPtrTS<FPlaybackAssetAdaptationSet>>		AudioAdaptationSets;
-	FString													UniqueIdentifier;
-	FString													AssetIdentifier;
-	FTimeRange												TimeRange;
-	FTimeValue												Duration;
-};
-
-
-class FPlaybackTimeline : public IPlaybackAssetTimeline
-{
-public:
-	virtual ~FPlaybackTimeline()
-	{
-	}
-	virtual FTimeValue GetAnchorTime() const override
-	{
-		return FTimeValue::GetZero();
-	}
-	virtual FTimeRange GetTotalTimeRange() const override
-	{
-		return TotalTimeRange;
-	}
-	virtual FTimeRange GetSeekableTimeRange() const override
-	{
-		return SeekableTimeRange;
-	}
-	virtual void GetSeekablePositions(TArray<FTimespan>& OutPositions) const override
-	{
-		OutPositions = SeekablePositions;
-	}
-	virtual FTimeValue GetDuration() const override
-	{
-		return Duration;
-	}
-	virtual int32 GetNumberOfMediaAssets() const override
-	{
-		return 1;
-	}
-	virtual TSharedPtrTS<ITimelineMediaAsset> GetMediaAssetByIndex(int32 MediaAssetIndex) const override
-	{
-		return PlayAsset;
-	}
-
-	FTimeRange								TotalTimeRange;
-	FTimeRange								SeekableTimeRange;
-	FTimeValue								Duration;
-	TSharedPtrTS<FTimelineMediaAsset>		PlayAsset;
-	TArray<FTimespan>						SeekablePositions;
-};
-
-}
-
-
-
 
 class FManifestBuilderHLS : public IManifestBuilderHLS
 {
@@ -354,18 +137,18 @@ public:
 	void Initialize(IPlayerSessionServices* PlayerSessionServices);
 	virtual ~FManifestBuilderHLS();
 
-	virtual FErrorDetail BuildFromMasterPlaylist(TSharedPtrTS<FManifestHLSInternal>& OutHLSPlaylist, const HLSPlaylistParser::FPlaylist& Playlist, const FPlaylistLoadRequestHLS& SourceRequest, const HTTP::FConnectionInfo* ConnectionInfo, const FStreamPreferences& Preferences, const FParamDict& Options) override;
+	virtual FErrorDetail BuildFromMasterPlaylist(TSharedPtrTS<FManifestHLSInternal>& OutHLSPlaylist, const HLSPlaylistParser::FPlaylist& Playlist, const FPlaylistLoadRequestHLS& SourceRequest, const HTTP::FConnectionInfo* ConnectionInfo) override;
 
-	virtual FErrorDetail GetInitialPlaylistLoadRequests(TArray<FPlaylistLoadRequestHLS>& OutRequests, TSharedPtrTS<FManifestHLSInternal> Manifest, const FStreamPreferences& Preferences, const FParamDict& Options) override;
-	virtual UEMediaError UpdateFailedInitialPlaylistLoadRequest(FPlaylistLoadRequestHLS& InOutFailedRequest, const HTTP::FConnectionInfo* ConnectionInfo, TSharedPtrTS<HTTP::FRetryInfo> PreviousAttempts, const FTimeValue& BlacklistUntilUTC, TSharedPtrTS<FManifestHLSInternal> Manifest, const FStreamPreferences& Preferences, const FParamDict& Options) override;
+	virtual FErrorDetail GetInitialPlaylistLoadRequests(TArray<FPlaylistLoadRequestHLS>& OutRequests, TSharedPtrTS<FManifestHLSInternal> Manifest) override;
+	virtual UEMediaError UpdateFailedInitialPlaylistLoadRequest(FPlaylistLoadRequestHLS& InOutFailedRequest, const HTTP::FConnectionInfo* ConnectionInfo, TSharedPtrTS<HTTP::FRetryInfo> PreviousAttempts, const FTimeValue& BlacklistUntilUTC, TSharedPtrTS<FManifestHLSInternal> Manifest) override;
 
-	virtual FErrorDetail UpdateFromVariantPlaylist(TSharedPtrTS<FManifestHLSInternal> InOutHLSPlaylist, const HLSPlaylistParser::FPlaylist& VariantPlaylist, const FPlaylistLoadRequestHLS& SourceRequest, const HTTP::FConnectionInfo* ConnectionInfo, uint32 ResponseCRC, const FStreamPreferences& Preferences, const FParamDict& Options) override;
+	virtual FErrorDetail UpdateFromVariantPlaylist(TSharedPtrTS<FManifestHLSInternal> InOutHLSPlaylist, const HLSPlaylistParser::FPlaylist& VariantPlaylist, const FPlaylistLoadRequestHLS& SourceRequest, const HTTP::FConnectionInfo* ConnectionInfo, uint32 ResponseCRC) override;
 	virtual void SetVariantPlaylistFailure(TSharedPtrTS<FManifestHLSInternal> InHLSPlaylist, const FPlaylistLoadRequestHLS& SourceRequest, const HTTP::FConnectionInfo* ConnectionInfo, TSharedPtrTS<HTTP::FRetryInfo> PreviousAttempts, const FTimeValue& BlacklistUntilUTC) override;
 
 private:
 	FErrorDetail SetupRenditions(FManifestHLSInternal* Manifest, const HLSPlaylistParser::FPlaylist& Playlist);
 	FErrorDetail SetupVariants(FManifestHLSInternal* Manifest, const HLSPlaylistParser::FPlaylist& Playlist);
-	void UpdateManifestMetadataFromStream(FManifestHLSInternal* Manifest, const FManifestHLSInternal::FPlaylistBase* Playlist, FManifestHLSInternal::FMediaStream* Stream, const FStreamPreferences& Preferences, const FParamDict& Options);
+	void UpdateManifestMetadataFromStream(FManifestHLSInternal* Manifest, const FManifestHLSInternal::FPlaylistBase* Playlist, FManifestHLSInternal::FMediaStream* Stream);
 	void LogMessage(IInfoLog::ELevel Level, const FString& Message);
 	FErrorDetail CreateErrorAndLog(const FString& Message, uint16 Code, UEMediaError Error = UEMEDIA_ERROR_OK);
 
@@ -424,7 +207,7 @@ void FManifestBuilderHLS::LogMessage(IInfoLog::ELevel Level, const FString& Mess
 }
 
 
-FErrorDetail FManifestBuilderHLS::BuildFromMasterPlaylist(TSharedPtrTS<FManifestHLSInternal>& OutHLSPlaylist, const HLSPlaylistParser::FPlaylist& Playlist, const FPlaylistLoadRequestHLS& SourceRequest, const HTTP::FConnectionInfo* ConnectionInfo, const FStreamPreferences& Preferences, const FParamDict& Options)
+FErrorDetail FManifestBuilderHLS::BuildFromMasterPlaylist(TSharedPtrTS<FManifestHLSInternal>& OutHLSPlaylist, const HLSPlaylistParser::FPlaylist& Playlist, const FPlaylistLoadRequestHLS& SourceRequest, const HTTP::FConnectionInfo* ConnectionInfo)
 {
 	FErrorDetail Error;
 
@@ -435,8 +218,8 @@ FErrorDetail FManifestBuilderHLS::BuildFromMasterPlaylist(TSharedPtrTS<FManifest
 
 	TUniquePtr<FManifestHLSInternal> Manifest(new FManifestHLSInternal);
 
-	Manifest->InitSegmentCache = MakeShareable(IInitSegmentCacheHLS::Create(PlayerSessionServices, Options));
-	Manifest->LicenseKeyCache = MakeShareable(ILicenseKeyCacheHLS::Create(PlayerSessionServices, Options));
+	Manifest->InitSegmentCache = MakeShareable(IInitSegmentCacheHLS::Create(PlayerSessionServices));
+	Manifest->LicenseKeyCache = MakeShareable(ILicenseKeyCacheHLS::Create(PlayerSessionServices));
 
 	// Remember the URL we got the playlist from and the time we got it.
 	Manifest->MasterPlaylistVars.PlaylistLoadRequest = SourceRequest;
@@ -464,15 +247,9 @@ FErrorDetail FManifestBuilderHLS::BuildFromMasterPlaylist(TSharedPtrTS<FManifest
 	//		if (Playlist.ContainsTag(HLSPlaylistParser::ExtXSessionKey))
 
 
-	// Create timeline and its primary asset object
-	HLSTimeline::FPlaybackTimeline* Timeline = new HLSTimeline::FPlaybackTimeline;
-	HLSTimeline::FTimelineMediaAsset* Asset = new HLSTimeline::FTimelineMediaAsset;
-	Asset->UniqueIdentifier = TEXT("Asset0");
-	Timeline->PlayAsset = MakeShareable(Asset);
-	Manifest->PlaybackTimeline = MakeShareable(Timeline);
-	Timeline = nullptr;
-	Asset = nullptr;
-
+	// Create the media asset that setting up of variants and renditions will populate
+	Manifest->CurrentMediaAsset = MakeSharedTS<FTimelineMediaAssetHLS>();
+	Manifest->CurrentMediaAsset->UniqueIdentifier = Manifest->CurrentMediaAsset->AssetIdentifier = DefaultAssetNameHLS;
 
 	// Set up any renditions specified in the master playlist.
 	Error = SetupRenditions(Manifest.Get(), Playlist);
@@ -607,23 +384,23 @@ FErrorDetail FManifestBuilderHLS::SetupRenditions(FManifestHLSInternal* Manifest
 	}
 
 	// Get the timeline asset object.
-	HLSTimeline::FTimelineMediaAsset* Asset = static_cast<HLSTimeline::FTimelineMediaAsset*>(Manifest->PlaybackTimeline->GetMediaAssetByIndex(0).Get());
+	TSharedPtrTS<FTimelineMediaAssetHLS> Asset = Manifest->CurrentMediaAsset;
 
 	// Create adaptation sets for the audio rendition groups
 // FIXME: We will need to do this for the other rendition types (VIDEO, SUBTITLES and CLOSED-CAPTIONS as well)
 	for(TMultiMap<FString, TSharedPtrTS<FManifestHLSInternal::FRendition>>::TConstIterator It = Manifest->AudioRenditions.CreateConstIterator(); It; ++It)
 	{
-		HLSTimeline::FPlaybackAssetAdaptationSet* Adapt = static_cast<HLSTimeline::FPlaybackAssetAdaptationSet*>(Asset->GetAdaptationSetByTypeAndUniqueIdentifier(EStreamType::Audio, It.Key()).Get());
+		FPlaybackAssetAdaptationSetHLS* Adapt = static_cast<FPlaybackAssetAdaptationSetHLS*>(Asset->GetAdaptationSetByTypeAndUniqueIdentifier(EStreamType::Audio, It.Key()).Get());
 		if (!Adapt)
 		{
-			Adapt = new HLSTimeline::FPlaybackAssetAdaptationSet;
+			Adapt = new FPlaybackAssetAdaptationSetHLS;
 			Adapt->UniqueIdentifier = It.Key();
-			TSharedPtrTS<HLSTimeline::FPlaybackAssetAdaptationSet> IAdapt(Adapt);
+			TSharedPtrTS<FPlaybackAssetAdaptationSetHLS> IAdapt(Adapt);
 			Asset->AudioAdaptationSets.Push(IAdapt);
 		}
 
 		TSharedPtrTS<FManifestHLSInternal::FRendition> Rendition = It.Value();
-		HLSTimeline::FPlaybackAssetRepresentation* Repr = new HLSTimeline::FPlaybackAssetRepresentation;
+		FPlaybackAssetRepresentationHLS* Repr = new FPlaybackAssetRepresentationHLS;
 		Repr->UniqueIdentifier = LexToString(Rendition->Internal.UniqueID);
 		Rendition->Internal.AdaptationSetUniqueID = It.Key();
 		Rendition->Internal.RepresentationUniqueID = Repr->UniqueIdentifier;
@@ -645,7 +422,7 @@ FErrorDetail FManifestBuilderHLS::SetupVariants(FManifestHLSInternal* Manifest,
 	}
 
 	// Get the timeline asset object.
-	HLSTimeline::FTimelineMediaAsset* Asset = static_cast<HLSTimeline::FTimelineMediaAsset*>(Manifest->PlaybackTimeline->GetMediaAssetByIndex(0).Get());
+	TSharedPtrTS<FTimelineMediaAssetHLS> Asset = Manifest->CurrentMediaAsset;
 
 	TUniquePtr<IURLParser> UrlBuilder(IURLParser::Create());
 	UrlBuilder->ParseURL(Manifest->MasterPlaylistVars.PlaylistLoadRequest.URL);
@@ -928,14 +705,14 @@ FErrorDetail FManifestBuilderHLS::SetupVariants(FManifestHLSInternal* Manifest,
 				Manifest->StreamMetadataVideo.Push(StreamMetaData);
 
 				// Add to timeline asset
-				HLSTimeline::FPlaybackAssetAdaptationSet* Adapt = nullptr;
+				FPlaybackAssetAdaptationSetHLS* Adapt = nullptr;
 				if (!Asset->VideoAdaptationSet.IsValid())
 				{
-					Adapt = new HLSTimeline::FPlaybackAssetAdaptationSet;
+					Adapt = new FPlaybackAssetAdaptationSetHLS;
 					Asset->VideoAdaptationSet = MakeShareable(Adapt);
 					Adapt->UniqueIdentifier = TEXT("$video$");
 				}
-				Adapt = static_cast<HLSTimeline::FPlaybackAssetAdaptationSet*>(Asset->VideoAdaptationSet.Get());
+				Adapt = static_cast<FPlaybackAssetAdaptationSetHLS*>(Asset->VideoAdaptationSet.Get());
 				if (Adapt->Codecs.Find(StreamMetaData.CodecInformation.GetCodecSpecifierRFC6381()) == INDEX_NONE)
 				{
 					if (Adapt->Codecs.Len())
@@ -944,7 +721,7 @@ FErrorDetail FManifestBuilderHLS::SetupVariants(FManifestHLSInternal* Manifest,
 					}
 					Adapt->Codecs.Append(StreamMetaData.CodecInformation.GetCodecSpecifierRFC6381());
 				}
-				HLSTimeline::FPlaybackAssetRepresentation* Repr = new HLSTimeline::FPlaybackAssetRepresentation;
+				FPlaybackAssetRepresentationHLS* Repr = new FPlaybackAssetRepresentationHLS;
 				Repr->UniqueIdentifier = LexToString(StreamMetaData.StreamUniqueID);
 				Repr->CodecInformation = StreamMetaData.CodecInformation;
 				Repr->Bitrate   	   = StreamMetaData.Bandwidth;
@@ -957,7 +734,7 @@ FErrorDetail FManifestBuilderHLS::SetupVariants(FManifestHLSInternal* Manifest,
 				// Check if there is an audio adaptation set already due to an audio rendition group referenced by this variant.
 				if (bHaveAudioGroup)
 				{
-					Adapt = static_cast<HLSTimeline::FPlaybackAssetAdaptationSet*>(Asset->GetAdaptationSetByTypeAndUniqueIdentifier(EStreamType::Audio, vs->AudioGroupID).Get());
+					Adapt = static_cast<FPlaybackAssetAdaptationSetHLS*>(Asset->GetAdaptationSetByTypeAndUniqueIdentifier(EStreamType::Audio, vs->AudioGroupID).Get());
 					if (!Adapt)
 					{
 						return CreateErrorAndLog(FString::Printf(TEXT("EXT-X-STREAM-INF references AUDIO rendition group \"%s\" that has not been defined!"), *vs->AudioGroupID), ERRCODE_HLS_BUILDER_REFERENCED_GROUP_NOT_DEFINED, UEMEDIA_ERROR_FORMAT_ERROR);
@@ -980,7 +757,7 @@ FErrorDetail FManifestBuilderHLS::SetupVariants(FManifestHLSInternal* Manifest,
 					}
 					for(int32 i=0; i<Adapt->GetNumberOfRepresentations(); ++i)
 					{
-						Repr = static_cast<HLSTimeline::FPlaybackAssetRepresentation*>(Adapt->GetRepresentationByIndex(i).Get());
+						Repr = static_cast<FPlaybackAssetRepresentationHLS*>(Adapt->GetRepresentationByIndex(i).Get());
 						// FIXME: This is setting the last codec info from all codecs onto all renditions. Not sure if this is the correct thing to do.
 						Repr->CodecInformation = StreamMetaData.CodecInformation;
 						//Repr->CDN;
@@ -1007,10 +784,10 @@ FErrorDetail FManifestBuilderHLS::SetupVariants(FManifestHLSInternal* Manifest,
 				Manifest->StreamMetadataAudio.Push(StreamMetaData);
 
 				// Check if there is an audio adaptation set already due to an audio rendition group referenced by this variant.
-				HLSTimeline::FPlaybackAssetAdaptationSet* Adapt = nullptr;
+				FPlaybackAssetAdaptationSetHLS* Adapt = nullptr;
 				if (bHaveAudioGroup)
 				{
-					Adapt = static_cast<HLSTimeline::FPlaybackAssetAdaptationSet*>(Asset->GetAdaptationSetByTypeAndUniqueIdentifier(EStreamType::Audio, vs->AudioGroupID).Get());
+					Adapt = static_cast<FPlaybackAssetAdaptationSetHLS*>(Asset->GetAdaptationSetByTypeAndUniqueIdentifier(EStreamType::Audio, vs->AudioGroupID).Get());
 					if (!Adapt)
 					{
 						return CreateErrorAndLog(FString::Printf(TEXT("EXT-X-STREAM-INF references AUDIO rendition group \"%s\" that has not been defined!"), *vs->AudioGroupID), ERRCODE_HLS_BUILDER_REFERENCED_GROUP_NOT_DEFINED, UEMEDIA_ERROR_FORMAT_ERROR);
@@ -1019,12 +796,12 @@ FErrorDetail FManifestBuilderHLS::SetupVariants(FManifestHLSInternal* Manifest,
 				if (!Adapt)
 				{
 					// FIXME: What do we do if there are several audio-only variant streams with different language codes? Or codecs?
-					Adapt = static_cast<HLSTimeline::FPlaybackAssetAdaptationSet*>(Asset->GetAdaptationSetByTypeAndUniqueIdentifier(EStreamType::Audio, "$audio$").Get());
+					Adapt = static_cast<FPlaybackAssetAdaptationSetHLS*>(Asset->GetAdaptationSetByTypeAndUniqueIdentifier(EStreamType::Audio, "$audio$").Get());
 					if (!Adapt)
 					{
-						Adapt = new HLSTimeline::FPlaybackAssetAdaptationSet;
+						Adapt = new FPlaybackAssetAdaptationSetHLS;
 						Adapt->UniqueIdentifier = TEXT("$audio$");
-						TSharedPtrTS<HLSTimeline::FPlaybackAssetAdaptationSet> IAdapt(Adapt);
+						TSharedPtrTS<FPlaybackAssetAdaptationSetHLS> IAdapt(Adapt);
 						Asset->AudioAdaptationSets.Push(IAdapt);
 					}
 				}
@@ -1039,7 +816,7 @@ FErrorDetail FManifestBuilderHLS::SetupVariants(FManifestHLSInternal* Manifest,
 				}
 				if (!bHaveAudioGroup)
 				{
-					HLSTimeline::FPlaybackAssetRepresentation* Repr = new HLSTimeline::FPlaybackAssetRepresentation;
+					FPlaybackAssetRepresentationHLS* Repr = new FPlaybackAssetRepresentationHLS;
 					Repr->UniqueIdentifier = LexToString(StreamMetaData.StreamUniqueID);
 					Repr->CodecInformation = StreamMetaData.CodecInformation;
 					Repr->Bitrate   	   = StreamMetaData.Bandwidth;
@@ -1053,7 +830,7 @@ FErrorDetail FManifestBuilderHLS::SetupVariants(FManifestHLSInternal* Manifest,
 				{
 					for(int32 i=0; i<Adapt->GetNumberOfRepresentations(); ++i)
 					{
-						HLSTimeline::FPlaybackAssetRepresentation* Repr = static_cast<HLSTimeline::FPlaybackAssetRepresentation*>(Adapt->GetRepresentationByIndex(i).Get());
+						FPlaybackAssetRepresentationHLS* Repr = static_cast<FPlaybackAssetRepresentationHLS*>(Adapt->GetRepresentationByIndex(i).Get());
 						// FIXME: This is setting the last codec info from all codecs onto all renditions. Not sure if this is the correct thing to do.
 						Repr->CodecInformation = StreamMetaData.CodecInformation;
 						//Repr->CDN;
@@ -1062,7 +839,7 @@ FErrorDetail FManifestBuilderHLS::SetupVariants(FManifestHLSInternal* Manifest,
 						// We let those point back to this variant stream since that is what we effectively have to use.
 						TArray<TSharedPtrTS<FManifestHLSInternal::FRendition>*> RenditionRange;
 						Manifest->AudioRenditions.MultiFindPointer(vs->AudioGroupID, RenditionRange);
-						if (!RenditionRange.IsEmpty())
+						if (RenditionRange.Num())
 						{
 							for(int32 ii=0; ii < RenditionRange.Num(); ++ii)
 							{
@@ -1106,10 +883,10 @@ FErrorDetail FManifestBuilderHLS::SetupVariants(FManifestHLSInternal* Manifest,
 		// Set up variants from the renditions in the audio adaptation sets.
 		for(int32 i=0; i<Asset->AudioAdaptationSets.Num(); ++i)
 		{
-			const HLSTimeline::FPlaybackAssetAdaptationSet* Adapt = static_cast<const HLSTimeline::FPlaybackAssetAdaptationSet*>(Asset->GetAdaptationSetByTypeAndIndex(EStreamType::Audio, i).Get());
+			const FPlaybackAssetAdaptationSetHLS* Adapt = static_cast<const FPlaybackAssetAdaptationSetHLS*>(Asset->GetAdaptationSetByTypeAndIndex(EStreamType::Audio, i).Get());
 			for(int32 j=0; j<Adapt->GetNumberOfRepresentations(); ++j)
 			{
-				const HLSTimeline::FPlaybackAssetRepresentation* Repr = static_cast<const HLSTimeline::FPlaybackAssetRepresentation*>(Adapt->GetRepresentationByIndex(j).Get());
+				const FPlaybackAssetRepresentationHLS* Repr = static_cast<const FPlaybackAssetRepresentationHLS*>(Adapt->GetRepresentationByIndex(j).Get());
 				int32 ReprID = 0;
 				LexFromString(ReprID, *Repr->GetUniqueIdentifier());
 				for(TMultiMap<FString, TSharedPtrTS<FManifestHLSInternal::FRendition>>::TConstIterator It = Manifest->AudioRenditions.CreateConstIterator(); It; ++It)
@@ -1167,7 +944,7 @@ FErrorDetail FManifestBuilderHLS::SetupVariants(FManifestHLSInternal* Manifest,
 
 
 
-FErrorDetail FManifestBuilderHLS::GetInitialPlaylistLoadRequests(TArray<FPlaylistLoadRequestHLS>& OutRequests, TSharedPtrTS<FManifestHLSInternal> Manifest, const FStreamPreferences& Preferences, const FParamDict& Options)
+FErrorDetail FManifestBuilderHLS::GetInitialPlaylistLoadRequests(TArray<FPlaylistLoadRequestHLS>& OutRequests, TSharedPtrTS<FManifestHLSInternal> Manifest)
 {
 	OutRequests.Empty();
 
@@ -1182,7 +959,7 @@ FErrorDetail FManifestBuilderHLS::GetInitialPlaylistLoadRequests(TArray<FPlaylis
 	}
 
 	TSharedPtrTS<FManifestHLSInternal::FVariantStream> StartingVariantStream;
-	int32 InitialBitrate = (int32) Options.GetValue(IPlaylistReaderHLS::OptionKeyInitialBitrate).SafeGetInt64(0);
+	int32 InitialBitrate = (int32) PlayerSessionServices->GetOptions().GetValue(OptionKeyInitialBitrate).SafeGetInt64(0);
 	bool bIsAudioOnlyVariant = false;
 
 	if (Manifest->VariantStreams.Num())
@@ -1253,7 +1030,7 @@ FErrorDetail FManifestBuilderHLS::GetInitialPlaylistLoadRequests(TArray<FPlaylis
 		// Get the range of renditions
 		TArray< const TSharedPtrTS<FManifestHLSInternal::FRendition>* > RenditionRange;
 		Manifest->AudioRenditions.MultiFindPointer(StartingVariantStream->AudioGroupID, RenditionRange);
-		if (!RenditionRange.IsEmpty())
+		if (RenditionRange.Num())
 		{
 			for(int32 ii=0; ii < RenditionRange.Num(); ++ii)
 			{
@@ -1297,7 +1074,7 @@ FErrorDetail FManifestBuilderHLS::GetInitialPlaylistLoadRequests(TArray<FPlaylis
 	return FErrorDetail();
 }
 
-UEMediaError FManifestBuilderHLS::UpdateFailedInitialPlaylistLoadRequest(FPlaylistLoadRequestHLS& InOutFailedRequest, const HTTP::FConnectionInfo* ConnectionInfo, TSharedPtrTS<HTTP::FRetryInfo> PreviousAttempts, const FTimeValue& BlacklistUntilUTC, TSharedPtrTS<FManifestHLSInternal> Manifest, const FStreamPreferences& Preferences, const FParamDict& Options)
+UEMediaError FManifestBuilderHLS::UpdateFailedInitialPlaylistLoadRequest(FPlaylistLoadRequestHLS& InOutFailedRequest, const HTTP::FConnectionInfo* ConnectionInfo, TSharedPtrTS<HTTP::FRetryInfo> PreviousAttempts, const FTimeValue& BlacklistUntilUTC, TSharedPtrTS<FManifestHLSInternal> Manifest)
 {
 	// Is the failed request a video variant?
 	for(int32 i=0, iMax=Manifest->VariantStreams.Num(); i<iMax; ++i)
@@ -1363,7 +1140,7 @@ UEMediaError FManifestBuilderHLS::UpdateFailedInitialPlaylistLoadRequest(FPlayli
 			VideoVariant->Internal.Blacklisted = MakeSharedTS<FManifestHLSInternal::FBlacklist>();
 			VideoVariant->Internal.Blacklisted->PreviousAttempts				= PreviousAttempts;
 			VideoVariant->Internal.Blacklisted->BecomesAvailableAgainAtUTC  	= BlacklistUntilUTC;
-			VideoVariant->Internal.Blacklisted->AssetIDs.AssetUniqueID  		= Manifest->PlaybackTimeline->GetMediaAssetByIndex(0)->GetUniqueIdentifier();
+			VideoVariant->Internal.Blacklisted->AssetIDs.AssetUniqueID  		= DefaultAssetNameHLS;
 			VideoVariant->Internal.Blacklisted->AssetIDs.AdaptationSetUniqueID  = InOutFailedRequest.AdaptationSetUniqueID;
 			VideoVariant->Internal.Blacklisted->AssetIDs.RepresentationUniqueID = InOutFailedRequest.RepresentationUniqueID;
 			VideoVariant->Internal.Blacklisted->AssetIDs.CDN					= InOutFailedRequest.CDN;
@@ -1410,7 +1187,7 @@ UEMediaError FManifestBuilderHLS::UpdateFailedInitialPlaylistLoadRequest(FPlayli
 				AudioRendition->Internal.Blacklisted = MakeSharedTS<FManifestHLSInternal::FBlacklist>();
 				AudioRendition->Internal.Blacklisted->PreviousAttempts  			  = PreviousAttempts;
 				AudioRendition->Internal.Blacklisted->BecomesAvailableAgainAtUTC	  = BlacklistUntilUTC;
-				AudioRendition->Internal.Blacklisted->AssetIDs.AssetUniqueID		  = Manifest->PlaybackTimeline->GetMediaAssetByIndex(0)->GetUniqueIdentifier();
+				AudioRendition->Internal.Blacklisted->AssetIDs.AssetUniqueID		  = DefaultAssetNameHLS;
 				AudioRendition->Internal.Blacklisted->AssetIDs.AdaptationSetUniqueID  = InOutFailedRequest.AdaptationSetUniqueID;
 				AudioRendition->Internal.Blacklisted->AssetIDs.RepresentationUniqueID = InOutFailedRequest.RepresentationUniqueID;
 				AudioRendition->Internal.Blacklisted->AssetIDs.CDN  				  = InOutFailedRequest.CDN;
@@ -1451,7 +1228,7 @@ void FManifestBuilderHLS::SetVariantPlaylistFailure(TSharedPtrTS<FManifestHLSInt
 			Playlist->Internal.Blacklisted->PreviousAttempts		   = PreviousAttempts;
 			Playlist->Internal.Blacklisted->BecomesAvailableAgainAtUTC = BlacklistUntilUTC;
 
-			Playlist->Internal.Blacklisted->AssetIDs.AssetUniqueID  		= InHLSPlaylist->PlaybackTimeline->GetMediaAssetByIndex(0)->GetUniqueIdentifier();
+			Playlist->Internal.Blacklisted->AssetIDs.AssetUniqueID  		= DefaultAssetNameHLS;
 			Playlist->Internal.Blacklisted->AssetIDs.AdaptationSetUniqueID  = SourceRequest.AdaptationSetUniqueID;
 			Playlist->Internal.Blacklisted->AssetIDs.RepresentationUniqueID = SourceRequest.RepresentationUniqueID;
 			Playlist->Internal.Blacklisted->AssetIDs.CDN					= SourceRequest.CDN;
@@ -1468,7 +1245,7 @@ void FManifestBuilderHLS::SetVariantPlaylistFailure(TSharedPtrTS<FManifestHLSInt
 
 
 
-FErrorDetail FManifestBuilderHLS::UpdateFromVariantPlaylist(TSharedPtrTS<FManifestHLSInternal> InOutHLSPlaylist, const HLSPlaylistParser::FPlaylist& VariantPlaylist, const FPlaylistLoadRequestHLS& SourceRequest, const HTTP::FConnectionInfo* ConnectionInfo, uint32 ResponseCRC, const FStreamPreferences& Preferences, const FParamDict& Options)
+FErrorDetail FManifestBuilderHLS::UpdateFromVariantPlaylist(TSharedPtrTS<FManifestHLSInternal> InOutHLSPlaylist, const HLSPlaylistParser::FPlaylist& VariantPlaylist, const FPlaylistLoadRequestHLS& SourceRequest, const HTTP::FConnectionInfo* ConnectionInfo, uint32 ResponseCRC)
 {
 	// Find the variant or rendition this playlist updates.
 	TSharedPtrTS<FManifestHLSInternal::FPlaylistBase> Playlist;
@@ -1895,7 +1672,7 @@ FErrorDetail FManifestBuilderHLS::UpdateFromVariantPlaylist(TSharedPtrTS<FManife
 		Playlist->Internal.bReloadTriggered = false;
 		Playlist->Internal.bNewlySelected   = false;
 
-		UpdateManifestMetadataFromStream(InOutHLSPlaylist.Get(), Playlist.Get(), UpdatedMediaStream.Get(), Preferences, Options);
+		UpdateManifestMetadataFromStream(InOutHLSPlaylist.Get(), Playlist.Get(), UpdatedMediaStream.Get());
 		InOutHLSPlaylist->UnlockPlaylists();
 	}
 
@@ -1903,7 +1680,7 @@ FErrorDetail FManifestBuilderHLS::UpdateFromVariantPlaylist(TSharedPtrTS<FManife
 }
 
 
-void FManifestBuilderHLS::UpdateManifestMetadataFromStream(FManifestHLSInternal* Manifest, const FManifestHLSInternal::FPlaylistBase* Playlist, FManifestHLSInternal::FMediaStream* Stream, const FStreamPreferences& Preferences, const FParamDict& Options)
+void FManifestBuilderHLS::UpdateManifestMetadataFromStream(FManifestHLSInternal* Manifest, const FManifestHLSInternal::FPlaylistBase* Playlist, FManifestHLSInternal::FMediaStream* Stream)
 {
 	// NOTE: The master playlist mutex must be locked already!
 
@@ -1913,6 +1690,7 @@ void FManifestBuilderHLS::UpdateManifestMetadataFromStream(FManifestHLSInternal*
 	bool bIsAudioOnly = Manifest->AudioOnlyStreams.Num() && Manifest->VariantStreams.Num() == 0;
 	const TArray<FManifestHLSInternal::FMediaStream::FMediaSegment>& SegmentList = Stream->SegmentList;
 	check(SegmentList.Num());
+	const FParamDict& Options = PlayerSessionServices->GetOptions();
 	if (SegmentList.Num())
 	{
 		EarliestTimeline = SegmentList[0].AbsoluteDateTime;
@@ -1962,9 +1740,9 @@ void FManifestBuilderHLS::UpdateManifestMetadataFromStream(FManifestHLSInternal*
 				}
 			}
 			// Try the video config value regardless of presentation type.
-			if (!LastOffset.IsValid() && Options.HaveKey(IPlaylistReaderHLS::OptionKeyLiveSeekableEndOffset))
+			if (!LastOffset.IsValid() && Options.HaveKey(OptionKeyLiveSeekableEndOffset))
 			{
-				LastOffset = Options.GetValue(IPlaylistReaderHLS::OptionKeyLiveSeekableEndOffset).SafeGetTimeValue(ThreeTargetDurations);
+				LastOffset = Options.GetValue(OptionKeyLiveSeekableEndOffset).SafeGetTimeValue(ThreeTargetDurations);
 			}
 			// If still not valid use the default from the HLS specification.
 			if (!LastOffset.IsValid())
@@ -2117,26 +1895,14 @@ void FManifestBuilderHLS::UpdateManifestMetadataFromStream(FManifestHLSInternal*
 			Manifest->MasterPlaylistVars.SeekablePositions = Stream->SeekablePositions;
 		}
 
-		// Update the timeline.
-		// A new timeline object is created to avoid updating the existing one that may be held onto by some client code.
-		// This isn't overly expensive since we do share the actual streams.
-		check(Manifest->PlaybackTimeline.IsValid());
-		if (Manifest->PlaybackTimeline.IsValid())
-		{
-			HLSTimeline::FPlaybackTimeline* CurrentTimeline = static_cast<HLSTimeline::FPlaybackTimeline*>(Manifest->PlaybackTimeline.Get());
-			HLSTimeline::FPlaybackTimeline* Timeline		= new HLSTimeline::FPlaybackTimeline;
-			HLSTimeline::FTimelineMediaAsset* Asset 		= new HLSTimeline::FTimelineMediaAsset;
-			Timeline->PlayAsset = MakeShareable(Asset);
-			Timeline->SeekableTimeRange = Manifest->MasterPlaylistVars.SeekableRange;
-			Timeline->SeekablePositions = Manifest->MasterPlaylistVars.SeekablePositions;
-			Timeline->TotalTimeRange	= Asset->TimeRange = Manifest->MasterPlaylistVars.TimelineRange;
-			Timeline->Duration  		= Asset->Duration  = Manifest->MasterPlaylistVars.PresentationDuration;
-			Asset->UniqueIdentifier    = CurrentTimeline->PlayAsset->UniqueIdentifier;
-			Asset->AssetIdentifier     = CurrentTimeline->PlayAsset->AssetIdentifier;
-			Asset->VideoAdaptationSet  = CurrentTimeline->PlayAsset->VideoAdaptationSet;
-			Asset->AudioAdaptationSets = CurrentTimeline->PlayAsset->AudioAdaptationSets;
-			// Switch over
-			Manifest->PlaybackTimeline = MakeShareable(Timeline);
+		if (Manifest->CurrentMediaAsset.IsValid())
+		{
+			Manifest->CurrentMediaAsset->UpdateLock.Lock();
+			Manifest->CurrentMediaAsset->TimeRange         = Manifest->MasterPlaylistVars.TimelineRange;
+			Manifest->CurrentMediaAsset->Duration          = Manifest->MasterPlaylistVars.PresentationDuration;
+			Manifest->CurrentMediaAsset->SeekableTimeRange = Manifest->MasterPlaylistVars.SeekableRange;
+			Manifest->CurrentMediaAsset->SeekablePositions = Manifest->MasterPlaylistVars.SeekablePositions;
+			Manifest->CurrentMediaAsset->UpdateLock.Unlock();
 		}
 	}
 	else
