ElectraPlayer: Added MPEG DASH Live playback. Currently only for single period presentations.

diff --git a/Engine/Plugins/Media/ElectraPlayer/Source/ElectraPlayerRuntime/Private/Runtime/Player/HLS/ManifestBuilderHLS.h b/Engine/Plugins/Media/ElectraPlayer/Source/ElectraPlayerRuntime/Private/Runtime/Player/HLS/ManifestBuilderHLS.h
--- a/Engine/Plugins/Media/ElectraPlayer/Source/ElectraPlayerRuntime/Private/Runtime/Player/HLS/ManifestBuilderHLS.h
+++ b/Engine/Plugins/Media/ElectraPlayer/Source/ElectraPlayerRuntime/Private/Runtime/Player/HLS/ManifestBuilderHLS.h
@@ -17,6 +17,7 @@ namespace Electra
 {
 class IInitSegmentCacheHLS;
 class ILicenseKeyCacheHLS;
+class FTimelineMediaAssetHLS;
 
 struct FPlaylistLoadRequestHLS
 {
@@ -287,27 +288,28 @@ struct FManifestHLSInternal
 	// TODO: if provided we can put the EXT-X-START here
 	// TODO: For future DRM support we can put the EXT-X-SESSION-KEY here.
 
-	TArray<TSharedPtrTS<FVariantStream>>							VariantStreams;
-	TArray<TSharedPtrTS<FVariantStream>>							AudioOnlyStreams;				//!< Variants that are audio-only (have only audio codec specified in master playlist)
-	TArray<TSharedPtrTS<FVariantStream>>							IFrameOnlyStreams;
+	TArray<TSharedPtrTS<FVariantStream>>					VariantStreams;
+	TArray<TSharedPtrTS<FVariantStream>>					AudioOnlyStreams;				//!< Variants that are audio-only (have only audio codec specified in master playlist)
+	TArray<TSharedPtrTS<FVariantStream>>					IFrameOnlyStreams;
 
-	TMap<uint32, TWeakPtrTS<FPlaylistBase>>							PlaylistIDMap;
-	TMap<int32, int32>												BandwidthToQualityIndex;
+	TMap<uint32, TWeakPtrTS<FPlaylistBase>>					PlaylistIDMap;
+	TMap<int32, int32>										BandwidthToQualityIndex;
 
-	bool															bHasIndependentSegments;
+	bool													bHasIndependentSegments;
 
-	FMasterPlaylistVars												MasterPlaylistVars;
+	FMasterPlaylistVars										MasterPlaylistVars;
 
 
-	FMediaCriticalSection											VariantPlaylistAccessMutex;			//!< This mutex is used to access any(!) of the variant and rendition playlists. The master playlist is immutable.
-	TSet<uint32>													ActivelyReferencedStreamIDs;		//!< Unique IDs of streams from which segments are being fetched.
-	TSharedPtrTS<IInitSegmentCacheHLS>								InitSegmentCache;
-	TSharedPtrTS<ILicenseKeyCacheHLS>								LicenseKeyCache;
+	FMediaCriticalSection									VariantPlaylistAccessMutex;			//!< This mutex is used to access any(!) of the variant and rendition playlists. The master playlist is immutable.
+	TSet<uint32>											ActivelyReferencedStreamIDs;		//!< Unique IDs of streams from which segments are being fetched.
+	TSharedPtrTS<IInitSegmentCacheHLS>						InitSegmentCache;
+	TSharedPtrTS<ILicenseKeyCacheHLS>						LicenseKeyCache;
 
-	TArray<FStreamMetadata>											StreamMetadataVideo;
-	TArray<FStreamMetadata>											StreamMetadataAudio;
+	TArray<FStreamMetadata>									StreamMetadataVideo;
+	TArray<FStreamMetadata>									StreamMetadataAudio;
+
+	TSharedPtrTS<FTimelineMediaAssetHLS>					CurrentMediaAsset;
 
-	TSharedPtrTS<IPlaybackAssetTimeline>							PlaybackTimeline;
 
 
 	TSharedPtrTS<FPlaylistBase> GetPlaylistForUniqueID(uint32 UniqueID) const
@@ -406,12 +408,10 @@ public:
 	 * @param Playlist
 	 * @param SourceRequest
 	 * @param ConnectionInfo
-	 * @param Preferences
-	 * @param Options
 	 *
 	 * @return
 	 */
-	virtual FErrorDetail BuildFromMasterPlaylist(TSharedPtrTS<FManifestHLSInternal>& OutHLSPlaylist, const HLSPlaylistParser::FPlaylist& Playlist, const FPlaylistLoadRequestHLS& SourceRequest, const HTTP::FConnectionInfo* ConnectionInfo, const FStreamPreferences& Preferences, const FParamDict& Options) = 0;
+	virtual FErrorDetail BuildFromMasterPlaylist(TSharedPtrTS<FManifestHLSInternal>& OutHLSPlaylist, const HLSPlaylistParser::FPlaylist& Playlist, const FPlaylistLoadRequestHLS& SourceRequest, const HTTP::FConnectionInfo* ConnectionInfo) = 0;
 
 
 	/**
@@ -419,12 +419,10 @@ public:
 	 *
 	 * @param OutRequests
 	 * @param Manifest
-	 * @param Preferences
-	 * @param Options
 	 *
 	 * @return
 	 */
-	virtual FErrorDetail GetInitialPlaylistLoadRequests(TArray<FPlaylistLoadRequestHLS>& OutRequests, TSharedPtrTS<FManifestHLSInternal> Manifest, const FStreamPreferences& Preferences, const FParamDict& Options) = 0;
+	virtual FErrorDetail GetInitialPlaylistLoadRequests(TArray<FPlaylistLoadRequestHLS>& OutRequests, TSharedPtrTS<FManifestHLSInternal> Manifest) = 0;
 
 	/**
 	 * Updates an initial playlist load request that failed to load or parse with an
@@ -436,12 +434,10 @@ public:
 	 * @param PreviousAttempts
 	 * @param BlacklistUntilUTC
 	 * @param Manifest
-	 * @param Preferences
-	 * @param Options
 	 *
 	 * @return UEMEDIA_ERROR_OK if an alternative was found or UEMEDIA_ERROR_END_OF_STREAM if no further alternatives exist.
 	 */
-	virtual UEMediaError UpdateFailedInitialPlaylistLoadRequest(FPlaylistLoadRequestHLS& InOutFailedRequest, const HTTP::FConnectionInfo* ConnectionInfo, TSharedPtrTS<HTTP::FRetryInfo> PreviousAttempts, const FTimeValue& BlacklistUntilUTC, TSharedPtrTS<FManifestHLSInternal> Manifest, const FStreamPreferences& Preferences, const FParamDict& Options) = 0;
+	virtual UEMediaError UpdateFailedInitialPlaylistLoadRequest(FPlaylistLoadRequestHLS& InOutFailedRequest, const HTTP::FConnectionInfo* ConnectionInfo, TSharedPtrTS<HTTP::FRetryInfo> PreviousAttempts, const FTimeValue& BlacklistUntilUTC, TSharedPtrTS<FManifestHLSInternal> Manifest) = 0;
 
 
 	/**
@@ -452,12 +448,10 @@ public:
 	 * @param SourceRequest
 	 * @param ConnectionInfo
 	 * @param ResponseCRC
-	 * @param Preferences
-	 * @param Options
 	 *
 	 * @return
 	 */
-	virtual FErrorDetail UpdateFromVariantPlaylist(TSharedPtrTS<FManifestHLSInternal> InOutHLSPlaylist, const HLSPlaylistParser::FPlaylist& VariantPlaylist, const FPlaylistLoadRequestHLS& SourceRequest, const HTTP::FConnectionInfo* ConnectionInfo, uint32 ResponseCRC, const FStreamPreferences& Preferences, const FParamDict& Options) = 0;
+	virtual FErrorDetail UpdateFromVariantPlaylist(TSharedPtrTS<FManifestHLSInternal> InOutHLSPlaylist, const HLSPlaylistParser::FPlaylist& VariantPlaylist, const FPlaylistLoadRequestHLS& SourceRequest, const HTTP::FConnectionInfo* ConnectionInfo, uint32 ResponseCRC) = 0;
 
 
 	/**
@@ -473,6 +467,183 @@ public:
 };
 
 
+
+class FPlaybackAssetRepresentationHLS : public IPlaybackAssetRepresentation
+{
+public:
+	FPlaybackAssetRepresentationHLS()
+	{ }
+	virtual ~FPlaybackAssetRepresentationHLS()
+	{ }
+	virtual FString GetUniqueIdentifier() const override
+	{
+		return UniqueIdentifier;
+	}
+	virtual const FStreamCodecInformation& GetCodecInformation() const override
+	{
+		return CodecInformation;
+	}
+	virtual int32 GetBitrate() const override
+	{
+		return Bitrate;
+	}
+	virtual int32 GetQualityIndex() const override
+	{ 
+		return QualityIndex;
+	}
+	virtual bool CanBePlayed() const override
+	{
+		return true;
+	}
+	FStreamCodecInformation	CodecInformation;
+	FString					UniqueIdentifier;
+	int32					Bitrate = 0;
+	int32					QualityIndex = 0;
+};
+
+class FPlaybackAssetAdaptationSetHLS : public IPlaybackAssetAdaptationSet
+{
+public:
+	virtual ~FPlaybackAssetAdaptationSetHLS()
+	{
+	}
+	virtual FString GetUniqueIdentifier() const override
+	{
+		return UniqueIdentifier;
+	}
+	virtual FString GetListOfCodecs() const override
+	{
+		return Codecs;
+	}
+	virtual FString GetLanguage() const override
+	{
+		return Language;
+	}
+	virtual int32 GetNumberOfRepresentations() const override
+	{
+		return Representations.Num();
+	}
+	virtual TSharedPtrTS<IPlaybackAssetRepresentation> GetRepresentationByIndex(int32 RepresentationIndex) const override
+	{
+		check(RepresentationIndex < Representations.Num());
+		if (RepresentationIndex < Representations.Num())
+		{
+			return Representations[RepresentationIndex];
+		}
+		return TSharedPtrTS<IPlaybackAssetRepresentation>();
+	}
+	virtual TSharedPtrTS<IPlaybackAssetRepresentation> GetRepresentationByUniqueIdentifier(const FString& InUniqueIdentifier) const override
+	{
+		for(int32 i=0, iMax=Representations.Num(); i<iMax; ++i)
+		{
+			if (Representations[i]->GetUniqueIdentifier() == InUniqueIdentifier)
+			{
+				return Representations[i];
+			}
+		}
+		return TSharedPtrTS<IPlaybackAssetRepresentation>();
+	}
+
+	TArray<TSharedPtrTS<IPlaybackAssetRepresentation>>		Representations;
+	FString													UniqueIdentifier;
+	FString													Language;
+	FString													Codecs;
+};
+
+
+class FTimelineMediaAssetHLS : public ITimelineMediaAsset
+{
+public:
+	virtual ~FTimelineMediaAssetHLS()
+	{
+	}
+	virtual FTimeRange GetTimeRange() const override
+	{
+		return TimeRange;
+	}
+	virtual FTimeValue GetDuration() const override
+	{
+		return Duration;
+	}
+	virtual FString GetAssetIdentifier() const override
+	{
+		return AssetIdentifier;
+	}
+	virtual FString GetUniqueIdentifier() const override
+	{
+		return UniqueIdentifier;
+	}
+	virtual int32 GetNumberOfAdaptationSets(EStreamType OfStreamType) const override
+	{
+		switch(OfStreamType)
+		{
+			case EStreamType::Video:
+				return VideoAdaptationSet.IsValid() ? 1 : 0;
+			case EStreamType::Audio:
+				return (int32) AudioAdaptationSets.Num();
+			default:
+				return 0;
+		}
+	}
+	virtual TSharedPtrTS<IPlaybackAssetAdaptationSet> GetAdaptationSetByTypeAndIndex(EStreamType OfStreamType, int32 AdaptationSetIndex) const override
+	{
+		switch(OfStreamType)
+		{
+			case EStreamType::Video:
+			{
+				check(AdaptationSetIndex == 0);
+				return AdaptationSetIndex == 0 ? VideoAdaptationSet : TSharedPtrTS<FPlaybackAssetAdaptationSetHLS>();
+			}
+			case EStreamType::Audio:
+			{
+				check(AdaptationSetIndex < AudioAdaptationSets.Num());
+				return AdaptationSetIndex < AudioAdaptationSets.Num() ? AudioAdaptationSets[AdaptationSetIndex] : TSharedPtrTS<FPlaybackAssetAdaptationSetHLS>();
+			}
+			default:
+			{
+				return TSharedPtrTS<FPlaybackAssetAdaptationSetHLS>();
+			}
+		}
+	}
+	TSharedPtrTS<IPlaybackAssetAdaptationSet> GetAdaptationSetByTypeAndUniqueIdentifier(EStreamType OfStreamType, const FString& InUniqueIdentifier) const
+	{
+		switch(OfStreamType)
+		{
+			case EStreamType::Video:
+			{
+				return VideoAdaptationSet;
+			}
+			case EStreamType::Audio:
+			{
+				for(int32 i=0, iMax=(int32)AudioAdaptationSets.Num(); i<iMax; ++i)
+				{
+					if (AudioAdaptationSets[i]->GetUniqueIdentifier() == InUniqueIdentifier)
+					{
+						return AudioAdaptationSets[i];
+					}
+				}
+				return TSharedPtrTS<FPlaybackAssetAdaptationSetHLS>();
+			}
+			default:
+			{
+				return TSharedPtrTS<FPlaybackAssetAdaptationSetHLS>();
+			}
+		}
+	}
+
+	TSharedPtrTS<FPlaybackAssetAdaptationSetHLS>			VideoAdaptationSet;
+	TArray<TSharedPtrTS<FPlaybackAssetAdaptationSetHLS>>	AudioAdaptationSets;
+	FString													UniqueIdentifier;
+	FString													AssetIdentifier;
+	FCriticalSection										UpdateLock;
+	FTimeRange												TimeRange;
+	FTimeValue												Duration;
+	FTimeRange												SeekableTimeRange;
+	TArray<FTimespan>										SeekablePositions;
+};
+
+
+
 } // namespace Electra
 
 
diff --git a/Engine/Plugins/Media/ElectraPlayer/Source/ElectraPlayerRuntime/Private/Runtime/Player/HLS/ManifestHLS.cpp b/Engine/Plugins/Media/ElectraPlayer/Source/ElectraPlayerRuntime/Private/Runtime/Player/HLS/ManifestHLS.cpp
--- a/Engine/Plugins/Media/ElectraPlayer/Source/ElectraPlayerRuntime/Private/Runtime/Player/HLS/ManifestHLS.cpp
+++ b/Engine/Plugins/Media/ElectraPlayer/Source/ElectraPlayerRuntime/Private/Runtime/Player/HLS/ManifestHLS.cpp
@@ -35,13 +35,14 @@ public:
 	virtual IManifest::FResult GetNextSegment(TSharedPtrTS<IStreamSegment>& OutSegment, TSharedPtrTS<const IStreamSegment> CurrentSegment) override;
 	virtual IManifest::FResult GetRetrySegment(TSharedPtrTS<IStreamSegment>& OutSegment, TSharedPtrTS<const IStreamSegment> CurrentSegment, bool bReplaceWithFillerData) override;
 	virtual IManifest::FResult GetLoopingSegment(TSharedPtrTS<IStreamSegment>& OutSegment, FPlayerLoopState& InOutLoopState, const TMultiMap<EStreamType, TSharedPtrTS<IStreamSegment>>& InFinishedSegments, const FPlayStartPosition& StartPosition, IManifest::ESearchType SearchType) override;
+	virtual void IncreaseSegmentFetchDelay(const FTimeValue& IncreaseAmount) override;
 
 	// Obtains information on the stream segmentation of a particular stream starting at a given current reference segment (optional, if not given returns suitable default values).
-	virtual void GetSegmentInformation(TArray<FSegmentInformation>& OutSegmentInformation, FTimeValue& OutAverageSegmentDuration, TSharedPtrTS<const IStreamSegment> CurrentSegment, const FTimeValue& LookAheadTime, const TSharedPtrTS<IPlaybackAssetAdaptationSet>& AdaptationSet, const TSharedPtrTS<IPlaybackAssetRepresentation>& Representation) override;
+	virtual void GetSegmentInformation(TArray<FSegmentInformation>& OutSegmentInformation, FTimeValue& OutAverageSegmentDuration, TSharedPtrTS<const IStreamSegment> CurrentSegment, const FTimeValue& LookAheadTime, const FString& AdaptationSetID, const FString& RepresentationID) override;
 
 	virtual TSharedPtrTS<ITimelineMediaAsset> GetMediaAsset() const override;
 
-	virtual void SelectStream(const TSharedPtrTS<IPlaybackAssetAdaptationSet>& AdaptationSet, const TSharedPtrTS<IPlaybackAssetRepresentation>& Representation) override;
+	virtual void SelectStream(const FString& AdaptationSetID, const FString& RepresentationID) override;
 
 private:
 	struct FSegSearchParam
@@ -83,14 +84,13 @@ private:
 
 
 
-TSharedPtrTS<FManifestHLS> FManifestHLS::Create(IPlayerSessionServices* SessionServices, const FParamDict& Options, IPlaylistReaderHLS* PlaylistReader, TSharedPtrTS<FManifestHLSInternal> Manifest)
+TSharedPtrTS<FManifestHLS> FManifestHLS::Create(IPlayerSessionServices* SessionServices, IPlaylistReaderHLS* PlaylistReader, TSharedPtrTS<FManifestHLSInternal> Manifest)
 {
-	return TSharedPtrTS<FManifestHLS>(new FManifestHLS(SessionServices, Options, PlaylistReader, Manifest));
+	return TSharedPtrTS<FManifestHLS>(new FManifestHLS(SessionServices, PlaylistReader, Manifest));
 }
 
-FManifestHLS::FManifestHLS(IPlayerSessionServices* InSessionServices, const FParamDict& InOptions, IPlaylistReaderHLS* InPlaylistReader, TSharedPtrTS<FManifestHLSInternal> Manifest)
-	: Options(InOptions)
-	, InternalManifest(Manifest)
+FManifestHLS::FManifestHLS(IPlayerSessionServices* InSessionServices, IPlaylistReaderHLS* InPlaylistReader, TSharedPtrTS<FManifestHLSInternal> Manifest)
+	: InternalManifest(Manifest)
 	, SessionServices(InSessionServices)
 	, PlaylistReader(InPlaylistReader)
 {
@@ -107,13 +107,52 @@ IManifest::EType FManifestHLS::GetPresentationType() const
 }
 
 
-TSharedPtrTS<IPlaybackAssetTimeline> FManifestHLS::GetTimeline() const
+FTimeValue FManifestHLS::GetAnchorTime() const
 {
-	FManifestHLSInternal::ScopedLockPlaylists lock(InternalManifest);
-	return InternalManifest->PlaybackTimeline;
+	return FTimeValue::GetZero();
+}
+FTimeRange FManifestHLS::GetTotalTimeRange() const
+{
+	if (InternalManifest.IsValid() && InternalManifest->CurrentMediaAsset.IsValid())
+	{
+		FScopeLock lock(&InternalManifest->CurrentMediaAsset->UpdateLock);
+		return InternalManifest->CurrentMediaAsset->TimeRange;
+	}
+	return FTimeRange();
+}
+FTimeRange FManifestHLS::GetSeekableTimeRange() const
+{
+	if (InternalManifest.IsValid() && InternalManifest->CurrentMediaAsset.IsValid())
+	{
+		FScopeLock lock(&InternalManifest->CurrentMediaAsset->UpdateLock);
+		return InternalManifest->CurrentMediaAsset->SeekableTimeRange;
+	}
+	return FTimeRange();
+}
+void FManifestHLS::GetSeekablePositions(TArray<FTimespan>& OutPositions) const
+{
+	if (InternalManifest.IsValid() && InternalManifest->CurrentMediaAsset.IsValid())
+	{
+		FScopeLock lock(&InternalManifest->CurrentMediaAsset->UpdateLock);
+		OutPositions = InternalManifest->CurrentMediaAsset->SeekablePositions;
+	}
+}
+FTimeValue FManifestHLS::GetDuration() const
+{
+	if (InternalManifest.IsValid() && InternalManifest->CurrentMediaAsset.IsValid())
+	{
+		FScopeLock lock(&InternalManifest->CurrentMediaAsset->UpdateLock);
+		return InternalManifest->CurrentMediaAsset->Duration;
+	}
+	return FTimeValue();
+}
+FTimeValue FManifestHLS::GetDefaultStartTime() const
+{
+	return FTimeValue::GetInvalid();
 }
 
 
+
 /**
  * Returns the bitrate of the default stream (usually the first one specified).
  *
@@ -429,27 +468,24 @@ IManifest::FResult FPlayPeriodHLS::FindSegment(TSharedPtrTS<FStreamSegmentReques
 	TSharedPtrTS<FStreamSegmentRequestHLSfmp4>	Req(new FStreamSegmentRequestHLSfmp4);
 	Req->StreamType 	= StreamType;
 	Req->StreamUniqueID = StreamUniqueID;
-	if (InternalManifest->PlaybackTimeline.IsValid())
+	Req->MediaAsset = InternalManifest->CurrentMediaAsset;
+	if (!Req->MediaAsset.IsValid())
 	{
-		Req->MediaAsset = InternalManifest->PlaybackTimeline->GetMediaAssetByIndex(0);
-		if (!Req->MediaAsset.IsValid())
-		{
-			return IManifest::FResult(IManifest::FResult::EType::NotFound).SetErrorDetail(FErrorDetail().SetMessage("Internal error, media asset not found on asset timeline!"));
-		}
-		if (Req->MediaAsset->GetNumberOfAdaptationSets(StreamType) > 1)
-		{
-			return IManifest::FResult(IManifest::FResult::EType::NotFound).SetErrorDetail(FErrorDetail().SetMessage(FString::Printf(TEXT("Internal error, more than one %s rendition group found on asset timeline!"), GetStreamTypeName(StreamType))));
-		}
-		Req->AdaptationSet = Req->MediaAsset->GetAdaptationSetByTypeAndIndex(StreamType, 0);
-		if (!Req->AdaptationSet.IsValid())
-		{
-			return IManifest::FResult(IManifest::FResult::EType::NotFound).SetErrorDetail(FErrorDetail().SetMessage(FString::Printf(TEXT("Internal error, no %s rendition group found on asset timeline!"), GetStreamTypeName(StreamType))));
-		}
-		Req->Representation = Req->AdaptationSet->GetRepresentationByUniqueIdentifier(LexToString(StreamUniqueID));
-		if (!Req->Representation.IsValid())
-		{
-			return IManifest::FResult(IManifest::FResult::EType::NotFound).SetErrorDetail(FErrorDetail().SetMessage(FString::Printf(TEXT("Internal error, %s rendition not found in group on asset timeline!"), GetStreamTypeName(StreamType))));
-		}
+		return IManifest::FResult(IManifest::FResult::EType::NotFound).SetErrorDetail(FErrorDetail().SetMessage("Internal error, media asset not found on asset timeline!"));
+	}
+	if (Req->MediaAsset->GetNumberOfAdaptationSets(StreamType) > 1)
+	{
+		return IManifest::FResult(IManifest::FResult::EType::NotFound).SetErrorDetail(FErrorDetail().SetMessage(FString::Printf(TEXT("Internal error, more than one %s rendition group found on asset timeline!"), GetStreamTypeName(StreamType))));
+	}
+	Req->AdaptationSet = Req->MediaAsset->GetAdaptationSetByTypeAndIndex(StreamType, 0);
+	if (!Req->AdaptationSet.IsValid())
+	{
+		return IManifest::FResult(IManifest::FResult::EType::NotFound).SetErrorDetail(FErrorDetail().SetMessage(FString::Printf(TEXT("Internal error, no %s rendition group found on asset timeline!"), GetStreamTypeName(StreamType))));
+	}
+	Req->Representation = Req->AdaptationSet->GetRepresentationByUniqueIdentifier(LexToString(StreamUniqueID));
+	if (!Req->Representation.IsValid())
+	{
+		return IManifest::FResult(IManifest::FResult::EType::NotFound).SetErrorDetail(FErrorDetail().SetMessage(FString::Printf(TEXT("Internal error, %s rendition not found in group on asset timeline!"), GetStreamTypeName(StreamType))));
 	}
 	Req->InitSegmentCache   	   = InternalManifest->InitSegmentCache;
 	Req->LicenseKeyCache		   = InternalManifest->LicenseKeyCache;
@@ -839,6 +875,13 @@ IManifest::FResult FPlayPeriodHLS::GetStartingSegment(TSharedPtrTS<IStreamSegmen
 	}
 }
 
+
+void FPlayPeriodHLS::IncreaseSegmentFetchDelay(const FTimeValue& IncreaseAmount)
+{
+	// No-op.
+}
+
+
 IManifest::FResult FPlayPeriodHLS::GetLoopingSegment(TSharedPtrTS<IStreamSegment>& OutSegment, FPlayerLoopState& InOutLoopState, const TMultiMap<EStreamType, TSharedPtrTS<IStreamSegment>>& InFinishedSegments, const FPlayStartPosition& StartPosition, IManifest::ESearchType SearchType)
 {
 	if (InFinishedSegments.Num())
@@ -1066,10 +1109,10 @@ void FPlayPeriodHLS::RefreshBlacklistState()
  * @param OutAverageSegmentDuration
  * @param InCurrentSegment
  * @param LookAheadTime
- * @param AdaptationSet
- * @param Representation
+ * @param AdaptationSetID
+ * @param RepresentationID
  */
-void FPlayPeriodHLS::GetSegmentInformation(TArray<FSegmentInformation>& OutSegmentInformation, FTimeValue& OutAverageSegmentDuration, TSharedPtrTS<const IStreamSegment> InCurrentSegment, const FTimeValue& LookAheadTime, const TSharedPtrTS<IPlaybackAssetAdaptationSet>& AdaptationSet, const TSharedPtrTS<IPlaybackAssetRepresentation>& Representation)
+void FPlayPeriodHLS::GetSegmentInformation(TArray<FSegmentInformation>& OutSegmentInformation, FTimeValue& OutAverageSegmentDuration, TSharedPtrTS<const IStreamSegment> InCurrentSegment, const FTimeValue& LookAheadTime, const FString& AdaptationSetID, const FString& RepresentationID)
 {
 	OutSegmentInformation.Empty();
 	OutAverageSegmentDuration.SetToInvalid();
@@ -1096,7 +1139,7 @@ void FPlayPeriodHLS::GetSegmentInformation(TArray<FSegmentInformation>& OutSegme
 
 	// The representation ID is the unique ID of the stream as a string. Convert it back
 	uint32 UniqueID;
-	LexFromString(UniqueID, *Representation->GetUniqueIdentifier());
+	LexFromString(UniqueID, *RepresentationID);
 	TSharedPtrTS<FManifestHLSInternal::FPlaylistBase> Playlist = pInt->GetPlaylistForUniqueID(UniqueID);
 	if (!Playlist.IsValid() || !Playlist->IsVariantStream())
 	{
@@ -1200,64 +1243,51 @@ TSharedPtrTS<ITimelineMediaAsset> FPlayPeriodHLS::GetMediaAsset() const
 {
 	FManifestHLSInternal::ScopedLockPlaylists lock(InternalManifest);
 	const FManifestHLSInternal* pInt = InternalManifest.Get();
-	check(pInt && pInt->PlaybackTimeline.IsValid());
-	if (pInt && pInt->PlaybackTimeline.IsValid())
+	if (pInt)
 	{
 		// HLS only has a single playback "asset". (aka Period)
-		return pInt->PlaybackTimeline->GetMediaAssetByIndex(0);
+		return pInt->CurrentMediaAsset;
 	}
 	return TSharedPtrTS<ITimelineMediaAsset>();
 }
 
 
 
-void FPlayPeriodHLS::SelectStream(const TSharedPtrTS<IPlaybackAssetAdaptationSet>& AdaptationSet, const TSharedPtrTS<IPlaybackAssetRepresentation>& Representation)
+void FPlayPeriodHLS::SelectStream(const FString& AdaptationSetID, const FString& RepresentationID)
 {
-	if (AdaptationSet.IsValid() && Representation.IsValid())
-	{
-		RefreshBlacklistState();
+	RefreshBlacklistState();
 
-		// The representation ID is the unique ID of the stream as a string. Convert it back
-		uint32 UniqueID;
-		LexFromString(UniqueID, *Representation->GetUniqueIdentifier());
+	// The representation ID is the unique ID of the stream as a string. Convert it back
+	uint32 UniqueID;
+	LexFromString(UniqueID, *RepresentationID);
 
-		// Which stream type is this?
-		switch(Representation->GetCodecInformation().GetStreamType())
+	// Which stream type is this?
+	if (AdaptationSetID.Equals(TEXT("$video$")))
+	{
+		// Different from what we have actively selected?
+		if (UniqueID != ActiveVideoUniqueID)
 		{
-			case EStreamType::Video:
-			{
-				// Different from what we have actively selected?
-				if (UniqueID != ActiveVideoUniqueID)
-				{
-					// FIXME: We could emit a QoS event here. Although selecting the stream does not mean it will actually get used.
-					//        There is still a chance that (if the playlist is not loaded yet) another stream will be selected right away.
+			// FIXME: We could emit a QoS event here. Although selecting the stream does not mean it will actually get used.
+			//        There is still a chance that (if the playlist is not loaded yet) another stream will be selected right away.
 
-					// Tell the manifest that we are now using a different stream.
-					InternalManifest->SelectActiveStreamID(UniqueID, ActiveVideoUniqueID);
+			// Tell the manifest that we are now using a different stream.
+			InternalManifest->SelectActiveStreamID(UniqueID, ActiveVideoUniqueID);
 
-					ActiveVideoUniqueID = UniqueID;
-				}
-				break;
-			}
-			case EStreamType::Audio:
-			{
-				// Different from what we have actively selected?
-				if (UniqueID != ActiveAudioUniqueID)
-				{
-					// FIXME: We could emit a QoS event here. Although selecting the stream does not mean it will actually get used.
-					//        There is still a chance that (if the playlist is not loaded yet) another stream will be selected right away.
+			ActiveVideoUniqueID = UniqueID;
+		}
+	}
+	else if (AdaptationSetID.Equals(TEXT("$audio$")))
+	{
+		// Different from what we have actively selected?
+		if (UniqueID != ActiveAudioUniqueID)
+		{
+			// FIXME: We could emit a QoS event here. Although selecting the stream does not mean it will actually get used.
+			//        There is still a chance that (if the playlist is not loaded yet) another stream will be selected right away.
 
-					// Tell the manifest that we are now using a different stream.
-					InternalManifest->SelectActiveStreamID(UniqueID, ActiveAudioUniqueID);
+			// Tell the manifest that we are now using a different stream.
+			InternalManifest->SelectActiveStreamID(UniqueID, ActiveAudioUniqueID);
 
-					ActiveAudioUniqueID = UniqueID;
-				}
-				break;
-			}
-			default:
-			{
-				break;
-			}
+			ActiveAudioUniqueID = UniqueID;
 		}
 	}
 }
diff --git a/Engine/Plugins/Media/ElectraPlayer/Source/ElectraPlayerRuntime/Private/Runtime/Player/HLS/ManifestHLS.h b/Engine/Plugins/Media/ElectraPlayer/Source/ElectraPlayerRuntime/Private/Runtime/Player/HLS/ManifestHLS.h
--- a/Engine/Plugins/Media/ElectraPlayer/Source/ElectraPlayerRuntime/Private/Runtime/Player/HLS/ManifestHLS.h
+++ b/Engine/Plugins/Media/ElectraPlayer/Source/ElectraPlayerRuntime/Private/Runtime/Player/HLS/ManifestHLS.h
@@ -13,11 +13,16 @@ class IPlaylistReaderHLS;
 class FManifestHLS : public IManifest
 {
 public:
-	static TSharedPtrTS<FManifestHLS> Create(IPlayerSessionServices* SessionServices, const FParamDict& Options, IPlaylistReaderHLS* PlaylistReader, TSharedPtrTS<FManifestHLSInternal> Manifest);
+	static TSharedPtrTS<FManifestHLS> Create(IPlayerSessionServices* SessionServices, IPlaylistReaderHLS* PlaylistReader, TSharedPtrTS<FManifestHLSInternal> Manifest);
 
 	virtual ~FManifestHLS();
 	virtual EType GetPresentationType() const override;
-	virtual TSharedPtrTS<IPlaybackAssetTimeline> GetTimeline() const override;
+	virtual FTimeValue GetAnchorTime() const override;
+	virtual FTimeRange GetTotalTimeRange() const override;
+	virtual FTimeRange GetSeekableTimeRange() const override;
+	virtual void GetSeekablePositions(TArray<FTimespan>& OutPositions) const override;
+	virtual FTimeValue GetDuration() const override;
+	virtual FTimeValue GetDefaultStartTime() const override;
 	virtual int64 GetDefaultStartingBitrate() const override;
 	virtual FTimeValue GetMinBufferTime() const override;
 	virtual void GetStreamMetadata(TArray<FStreamMetadata>& OutMetadata, EStreamType StreamType) const override;
@@ -27,9 +32,8 @@ public:
 	virtual FResult FindPlayPeriod(TSharedPtrTS<IPlayPeriod>& OutPlayPeriod, const FPlayStartPosition& StartPosition, ESearchType SearchType) override;
 
 private:
-	FManifestHLS(IPlayerSessionServices* SessionServices, const FParamDict& InOptions, IPlaylistReaderHLS* PlaylistReader, TSharedPtrTS<FManifestHLSInternal> Manifest);
+	FManifestHLS(IPlayerSessionServices* SessionServices, IPlaylistReaderHLS* PlaylistReader, TSharedPtrTS<FManifestHLSInternal> Manifest);
 
-	FParamDict									Options;
 	TSharedPtrTS<FManifestHLSInternal>			InternalManifest;
 	IPlayerSessionServices* 					SessionServices;
 	IPlaylistReaderHLS*							PlaylistReader;
diff --git a/Engine/Plugins/Media/ElectraPlayer/Source/ElectraPlayerRuntime/Private/Runtime/Player/HLS/PlaylistReaderHLS.cpp b/Engine/Plugins/Media/ElectraPlayer/Source/ElectraPlayerRuntime/Private/Runtime/Player/HLS/PlaylistReaderHLS.cpp
--- a/Engine/Plugins/Media/ElectraPlayer/Source/ElectraPlayerRuntime/Private/Runtime/Player/HLS/PlaylistReaderHLS.cpp
+++ b/Engine/Plugins/Media/ElectraPlayer/Source/ElectraPlayerRuntime/Private/Runtime/Player/HLS/PlaylistReaderHLS.cpp
@@ -15,6 +15,7 @@
 #include "HAL/LowLevelMemTracker.h"
 #include "ElectraPlayerPrivate.h"
 #include "Player/AdaptiveStreamingPlayerResourceRequest.h"
+#include "Player/AdaptivePlayerOptionKeynames.h"
 
 #define ERRCODE_HLS_PARSER_ERROR							1
 #define ERRCODE_HLS_NO_MASTER_PLAYLIST						2
@@ -28,9 +29,7 @@
 namespace Electra
 {
 
-const FString IPlaylistReaderHLS::OptionKeyInitialBitrate(TEXT("initial_bitrate"));														//!< (int64) value indicating the bitrate to start with.
 const FString IPlaylistReaderHLS::OptionKeyLiveSeekableStartOffset(TEXT("seekable_range_live_start_offset"));							//!< (FTimeValue) value specifying how many seconds away from the Live media timeline the seekable range should start.
-const FString IPlaylistReaderHLS::OptionKeyLiveSeekableEndOffset(TEXT("seekable_range_live_end_offset"));								//!< (FTimeValue) value specifying how many seconds away from the Live media timeline the seekable range should end.
 const FString IPlaylistReaderHLS::OptionKeyLiveSeekableEndOffsetAudioOnly(TEXT("seekable_range_live_end_offset_audioonly"));			//!< (FTimeValue) value specifying how many seconds away from the Live media timeline the seekable range should end for audio-only playlists.
 const FString IPlaylistReaderHLS::OptionKeyLiveSeekableEndOffsetBeConservative(TEXT("seekable_range_live_end_offset_conservative"));	//!< (bool) true to use a larger Live edge distance, false to go with the smaller absolute difference
 
@@ -71,11 +70,8 @@ public:
 	 * Loads and parses the playlist.
 	 *
 	 * @param URL     URL of the playlist to load
-	 * @param Preferences
-	 *                User preferences for initial stream selection.
-	 * @param Options Options for the reader and parser
 	 */
-	virtual void LoadAndParse(const FString& URL, const FStreamPreferences& Preferences, const FParamDict& Options) override;
+	virtual void LoadAndParse(const FString& URL) override;
 
 	/**
 	 * Returns the URL from which the playlist was loaded (or supposed to be loaded).
@@ -104,9 +100,8 @@ private:
 	class FPlaylistRequest : public TSharedFromThis<FPlaylistRequest, ESPMode::ThreadSafe>
 	{
 	public:
-		FPlaylistRequest(const FPlaylistLoadRequestHLS& InPlaylistLoadRequest, const FParamDict& InOptions)
+		FPlaylistRequest(const FPlaylistLoadRequestHLS& InPlaylistLoadRequest)
 			: PlaylistLoadRequest(InPlaylistLoadRequest)
-			, Options(InOptions)
 			, bIsMasterPlaylist(false)
 		{
 		}
@@ -215,7 +210,6 @@ private:
 		};
 
 		FPlaylistLoadRequestHLS									PlaylistLoadRequest;
-		FParamDict												Options;
 		FTimeValue												ExecuteAtUTC;
 		TSharedPtrTS<IElectraHttpManager::FReceiveBuffer>		ReceiveBuffer;
 		TSharedPtrTS<IElectraHttpManager::FRequest>				HTTPRequest;
@@ -226,11 +220,10 @@ private:
 		HTTP::FConnectionInfo									StaticRequestConnectionInfo;
 	};
 	using FPlaylistRequestPtr = TSharedPtr<FPlaylistRequest, ESPMode::ThreadSafe>;
-	//using FPlaylistRequestRef = TSharedRef<FPlaylistRequest, ESPMode::ThreadSafe>;
 
 	void StartWorkerThread();
 	void StopWorkerThread();
-	void WorkerThread(void);
+	void WorkerThread();
 	void HandleEnqueuedPlaylistDownloads(const FTimeValue& TimeNow);
 	void HandleCompletedPlaylistDownloads(const FTimeValue& TimeNow);
 	void HandleStaticRequestCompletions(const FTimeValue& TimeNow);
@@ -251,8 +244,6 @@ private:
 
 
 	IPlayerSessionServices*									PlayerSessionServices;
-	FStreamPreferences										StreamPreferences;
-	FParamDict												Options;
 	FString													MasterPlaylistURL;
 	TSharedPtr<FMediaSemaphore, ESPMode::ThreadSafe>		WorkerThreadSignal;
 	bool													bIsWorkerThreadStarted;
@@ -392,10 +383,8 @@ void FPlaylistReaderHLS::LogMessage(IInfoLog::ELevel Level, const FString& Messa
 	}
 }
 
-void FPlaylistReaderHLS::LoadAndParse(const FString& URL, const FStreamPreferences& Preferences, const FParamDict& InOptions)
+void FPlaylistReaderHLS::LoadAndParse(const FString& URL)
 {
-	StreamPreferences = Preferences;
-	Options 		  = InOptions;
 	MasterPlaylistURL = URL;
 	StartWorkerThread();
 }
@@ -407,7 +396,7 @@ void FPlaylistReaderHLS::RequestPlaylistLoad(const FPlaylistLoadRequestHLS& Load
 
 void FPlaylistReaderHLS::EnqueueLoadPlaylist(const FPlaylistLoadRequestHLS& InPlaylistLoadRequest, bool bIsMasterPlaylist)
 {
-	FPlaylistRequestPtr Request = MakeShared<FPlaylistRequest, ESPMode::ThreadSafe>(InPlaylistLoadRequest, Options);
+	FPlaylistRequestPtr Request = MakeShared<FPlaylistRequest, ESPMode::ThreadSafe>(InPlaylistLoadRequest);
 	Request->SetIsMasterPlaylist(bIsMasterPlaylist);
 	EnqueuedPlaylistRequests.Push(Request);
 	WorkerThreadSignal->Release();
@@ -465,7 +454,7 @@ void FPlaylistReaderHLS::CheckForPlaylistUpdate(const FTimeValue& TimeNow)
 	}
 }
 
-void FPlaylistReaderHLS::WorkerThread(void)
+void FPlaylistReaderHLS::WorkerThread()
 {
 	LLM_SCOPE(ELLMTag::ElectraPlayer);
 	CSV_SCOPED_TIMING_STAT(ElectraPlayer, PlaylistReaderHLS_Worker);
@@ -623,6 +612,7 @@ void FPlaylistReaderHLS::FPlaylistRequest::Execute(TSharedPtrTS<IElectraHttpMana
 	HTTPRequest->ConnectionInfo.RetryInfo = RetryInfo;
 
 	// Set connection timeouts for master and variant playlist retrieval.
+	const FParamDict& Options =	InPlayerSessionServices->GetOptions();
 	if (bIsMasterPlaylist)
 	{
 		HTTPRequest->Parameters.ConnectTimeout = Options.GetValue(IPlaylistReaderHLS::OptionKeyMasterPlaylistLoadConnectTimeout).SafeGetTimeValue(FTimeValue().SetFromMilliseconds(1000 * 8));
@@ -851,7 +841,7 @@ void FPlaylistReaderHLS::HandleCompletedPlaylistDownloads(const FTimeValue& Time
 						RetryInfo->PreviousFailureStates.Push(ConnInfo->StatusInfo);
 						FTimeValue RetryDelay;
 						RetryDelay.SetFromMilliseconds(1000 * 10);
-						if (Builder->UpdateFailedInitialPlaylistLoadRequest(Request->GetPlaylistLoadRequest(), ConnInfo, RetryInfo, TimeNow + RetryDelay, Manifest, StreamPreferences, Options) == UEMEDIA_ERROR_OK)
+						if (Builder->UpdateFailedInitialPlaylistLoadRequest(Request->GetPlaylistLoadRequest(), ConnInfo, RetryInfo, TimeNow + RetryDelay, Manifest) == UEMEDIA_ERROR_OK)
 						{
 							// Retry with this initial playlist
 							LogMessage(IInfoLog::ELevel::Warning, FString::Printf(TEXT("Failed %s initial playlist \"%s\" (%s), trying alternative variant"), ParseError.IsSet() ? TEXT("parsing") : TEXT("downloading"), *ConnInfo->EffectiveURL, ParseError.IsSet() ? *ParseError.GetPrintable() : *ConnInfo->StatusInfo.ErrorDetail.GetMessage()));
@@ -879,10 +869,10 @@ void FPlaylistReaderHLS::HandleCompletedPlaylistDownloads(const FTimeValue& Time
 					// Was this the last one?
 					// NOTE: If downloading or parsing failed we will not attempt to create a manifest (would not work anyway)
 					//       and also not send a message to the player telling it that metadata is available.
-					else if (InitiallyRequiredPlaylistLoadRequests.IsEmpty())
+					else if (InitiallyRequiredPlaylistLoadRequests.Num() == 0)
 					{
 						// Create a wrapper manifest for use with the player.
-						PlayerManifest = FManifestHLS::Create(PlayerSessionServices, Options, this, Manifest);
+						PlayerManifest = FManifestHLS::Create(PlayerSessionServices, this, Manifest);
 
 						// Yes. We can now create the media timeline and report availability of metadata to the player.
 						PlayerSessionServices->SendMessageToPlayer(IPlaylistReader::PlaylistLoadedMessage::Create(LastErrorDetail, Request->GetConnectionInfo(), Playlist::EListType::Variant, Playlist::ELoadType::Initial));
@@ -1034,7 +1024,7 @@ FErrorDetail FPlaylistReaderHLS::ParsePlaylist(FPlaylistRequestPtr FromRequest)
 
 
 			TSharedPtrTS<FManifestHLSInternal> NewManifest;
-			Error = Builder->BuildFromMasterPlaylist(NewManifest, Playlist, FromRequest->GetPlaylistLoadRequest(), FromRequest->GetConnectionInfo(), StreamPreferences, Options);
+			Error = Builder->BuildFromMasterPlaylist(NewManifest, Playlist, FromRequest->GetPlaylistLoadRequest(), FromRequest->GetConnectionInfo());
 			if (Error.IsOK())
 			{
 				Manifest = NewManifest;
@@ -1045,7 +1035,7 @@ FErrorDetail FPlaylistReaderHLS::ParsePlaylist(FPlaylistRequestPtr FromRequest)
 
 				// Get the playlists we need to eventually start playback with.
 				TArray<FPlaylistLoadRequestHLS> PlaylistRequests;
-				Error = Builder->GetInitialPlaylistLoadRequests(PlaylistRequests, Manifest, StreamPreferences, Options);
+				Error = Builder->GetInitialPlaylistLoadRequests(PlaylistRequests, Manifest);
 				if (Error.IsOK())
 				{
 					// Enqueue those for loading.
@@ -1067,7 +1057,7 @@ FErrorDetail FPlaylistReaderHLS::ParsePlaylist(FPlaylistRequestPtr FromRequest)
 			// For our purposes there needs to be a master playlist.
 			if (Manifest.IsValid())
 			{
-				Error = Builder->UpdateFromVariantPlaylist(Manifest, Playlist, FromRequest->GetPlaylistLoadRequest(), FromRequest->GetConnectionInfo(), crc, StreamPreferences, Options);
+				Error = Builder->UpdateFromVariantPlaylist(Manifest, Playlist, FromRequest->GetPlaylistLoadRequest(), FromRequest->GetConnectionInfo(), crc);
 			}
 			else
 			{
diff --git a/Engine/Plugins/Media/ElectraPlayer/Source/ElectraPlayerRuntime/Private/Runtime/Player/HLS/PlaylistReaderHLS.h b/Engine/Plugins/Media/ElectraPlayer/Source/ElectraPlayerRuntime/Private/Runtime/Player/HLS/PlaylistReaderHLS.h
--- a/Engine/Plugins/Media/ElectraPlayer/Source/ElectraPlayerRuntime/Private/Runtime/Player/HLS/PlaylistReaderHLS.h
+++ b/Engine/Plugins/Media/ElectraPlayer/Source/ElectraPlayerRuntime/Private/Runtime/Player/HLS/PlaylistReaderHLS.h
@@ -21,10 +21,7 @@ public:
 
 	virtual ~IPlaylistReaderHLS() = default;
 
-
-	static const FString OptionKeyInitialBitrate;						//!< (int64) value indicating the bitrate to start with.
 	static const FString OptionKeyLiveSeekableStartOffset;				//!< (FTimeValue) value specifying how many seconds away from the Live media timeline the seekable range should start.
-	static const FString OptionKeyLiveSeekableEndOffset;				//!< (FTimeValue) value specifying how many seconds away from the Live media timeline the seekable range should end.
 	static const FString OptionKeyLiveSeekableEndOffsetAudioOnly;		//!< (FTimeValue) value specifying how many seconds away from the Live media timeline the seekable range should end for audio-only playlists.
 	static const FString OptionKeyLiveSeekableEndOffsetBeConservative;	//!< (bool) true to use a larger Live edge distance, false to go with the smaller absolute difference
 	static const FString OptionKeyMasterPlaylistLoadConnectTimeout;		//!< (FTimeValue) value specifying connection timeout fetching the master playlist
@@ -38,11 +35,8 @@ public:
 	 * Loads and parses the playlist.
 	 *
 	 * @param URL     URL of the playlist to load
-	 * @param Preferences
-	 *                User preferences for initial stream selection.
-	 * @param Options Options
 	 */
-	virtual void LoadAndParse(const FString& URL, const FStreamPreferences& Preferences, const FParamDict& Options) = 0;
+	virtual void LoadAndParse(const FString& URL) = 0;
 
 	/**
 	 * Returns the URL from which the playlist was loaded (or supposed to be loaded).
diff --git a/Engine/Plugins/Media/ElectraPlayer/Source/ElectraPlayerRuntime/Private/Runtime/Player/Manifest.h b/Engine/Plugins/Media/ElectraPlayer/Source/ElectraPlayerRuntime/Private/Runtime/Player/Manifest.h
--- a/Engine/Plugins/Media/ElectraPlayer/Source/ElectraPlayerRuntime/Private/Runtime/Player/Manifest.h
+++ b/Engine/Plugins/Media/ElectraPlayer/Source/ElectraPlayerRuntime/Private/Runtime/Player/Manifest.h
@@ -176,8 +176,56 @@ namespace Electra
 		//-------------------------------------------------------------------------
 		// Presentation and play time related functions.
 		//
-		//! Returns the current timeline object.
-		virtual TSharedPtrTS<IPlaybackAssetTimeline> GetTimeline() const = 0;
+
+		/**
+		 * Returns the time value the timeline is anchored at.
+		 * Usually this is a UTC timestamp at which the presentation began but it could
+		 * literally be anything.
+		 * All time values are absolute and have this anchor added in.
+		 *
+		 * @return The time this presentation timeline is anchored at.
+		 */
+		virtual FTimeValue GetAnchorTime() const = 0;
+
+		/**
+		 * Returns the total time range of the timeline, including the end time of the last sample
+		 * to which the player will not be able to seek to.
+		 *
+		 * @return Time range of assets on the timeline.
+		 */
+		virtual FTimeRange GetTotalTimeRange() const = 0;
+
+		/**
+		 * Returns the seekable range of the timeline, which is a subset of the total
+		 * time range. Playback can start at any point in the seekable range with
+		 * additional constraints imposed by the format.
+		 *
+		 * @return Time range in which playback can start.
+		 */
+		virtual FTimeRange GetSeekableTimeRange() const = 0;
+
+		/**
+		 * Returns the timestamps of the segments from the video or audio track (if no video is present).
+		 * Segments are required to start with a keyframe and can thus be used to start playback with.
+		 */
+		virtual void GetSeekablePositions(TArray<FTimespan>& OutPositions) const = 0;
+
+		/**
+		 * Returns the duration of the assets on the timeline.
+		 * Typically this is the difference of the end and start values of the total time range
+		 * unless this is a Live presentation for which the duration will be set to infinite.
+		 *
+		 * @return Duration of the timeline.
+		 */
+		virtual FTimeValue GetDuration() const = 0;
+
+		/**
+		 * Returns the playback start time as defined by the presentation itself.
+		 * If the presentation has no preferred start time an invalid value is returned.
+		 */
+		virtual FTimeValue GetDefaultStartTime() const = 0;
+
+
 
 		//! Returns the bitrate of the default stream (usually the first one specified).
 		virtual int64 GetDefaultStartingBitrate() const = 0;
@@ -185,7 +233,6 @@ namespace Electra
 		//! Returns stream metadata. For period based presentations the streams can be different per period in which case the metadata of the first period is returned.
 		virtual void GetStreamMetadata(TArray<FStreamMetadata>& OutMetadata, EStreamType StreamType) const = 0;
 
-
 		//
 		virtual FTimeValue GetMinBufferTime() const = 0;
 
@@ -234,10 +281,10 @@ namespace Electra
 			/**
 			 * Selects a specific stream to be used for all next segment downloads.
 			 *
-			 * @param AdaptationSet
-			 * @param Representation
+			 * @param AdaptationSetID
+			 * @param RepresentationID
 			 */
-			virtual void SelectStream(const TSharedPtrTS<IPlaybackAssetAdaptationSet>& AdaptationSet, const TSharedPtrTS<IPlaybackAssetRepresentation>& Representation) = 0;
+			virtual void SelectStream(const FString& AdaptationSetID, const FString& RepresentationID) = 0;
 
 			/**
 			 * Sets up a starting segment request to begin playback at the specified time.
@@ -288,6 +335,14 @@ namespace Electra
 			 */
 			virtual FResult GetRetrySegment(TSharedPtrTS<IStreamSegment>& OutSegment, TSharedPtrTS<const IStreamSegment> CurrentSegment, bool bReplaceWithFillerData) = 0;
 
+			/**
+			 * Called by the ABR to increase the delay in fetching the next segment in case the segment returned a 404 when fetched at
+			 * the announced availability time. This may reduce 404's on the next segment fetches.
+			 * 
+			 * @param IncreaseAmount
+			 */
+			virtual void IncreaseSegmentFetchDelay(const FTimeValue& IncreaseAmount) = 0;
+
 			struct FSegmentInformation
 			{
 				FTimeValue	Duration;
@@ -301,10 +356,10 @@ namespace Electra
 			 * @param OutAverageSegmentDuration
 			 * @param CurrentSegment
 			 * @param LookAheadTime
-			 * @param AdaptationSet
-			 * @param Representation
+			 * @param AdaptationSetID
+			 * @param RepresentationID
 			 */
-			virtual void GetSegmentInformation(TArray<FSegmentInformation>& OutSegmentInformation, FTimeValue& OutAverageSegmentDuration, TSharedPtrTS<const IStreamSegment> CurrentSegment, const FTimeValue& LookAheadTime, const TSharedPtrTS<IPlaybackAssetAdaptationSet>& AdaptationSet, const TSharedPtrTS<IPlaybackAssetRepresentation>& Representation) = 0;
+			virtual void GetSegmentInformation(TArray<FSegmentInformation>& OutSegmentInformation, FTimeValue& OutAverageSegmentDuration, TSharedPtrTS<const IStreamSegment> CurrentSegment, const FTimeValue& LookAheadTime, const FString& AdaptationSetID, const FString& RepresentationID) = 0;
 		};
 
 		virtual FResult FindPlayPeriod(TSharedPtrTS<IPlayPeriod>& OutPlayPeriod, const FPlayStartPosition& StartPosition, ESearchType SearchType) = 0;
diff --git a/Engine/Plugins/Media/ElectraPlayer/Source/ElectraPlayerRuntime/Private/Runtime/Player/PlaybackTimeline.h b/Engine/Plugins/Media/ElectraPlayer/Source/ElectraPlayerRuntime/Private/Runtime/Player/PlaybackTimeline.h
--- a/Engine/Plugins/Media/ElectraPlayer/Source/ElectraPlayerRuntime/Private/Runtime/Player/PlaybackTimeline.h
+++ b/Engine/Plugins/Media/ElectraPlayer/Source/ElectraPlayerRuntime/Private/Runtime/Player/PlaybackTimeline.h
@@ -117,80 +117,6 @@ namespace Electra
 	};
 
 
-	/**
-	 * Interface class to the media playback asset timeline.
-	 * This is a snapshot of the current timeline which may change periodically.
-	 */
-	class IPlaybackAssetTimeline
-	{
-	public:
-		virtual ~IPlaybackAssetTimeline() = default;
-
-		/**
-		 * Returns the time value the timeline is anchored at.
-		 * Usually this is a UTC timestamp at which the presentation began but it could
-		 * literally be anything.
-		 * All time values are absolute and have this anchor added in.
-		 *
-		 * @return The time this presentation timeline is anchored at.
-		 */
-		virtual FTimeValue GetAnchorTime() const = 0;
-
-		/**
-		 * Returns the total time range of the timeline, including the end time of the last sample
-		 * to which the player will not be able to seek to.
-		 *
-		 * @return Time range of assets on the timeline.
-		 */
-		virtual FTimeRange GetTotalTimeRange() const = 0;
-
-
-		/**
-		 * Returns the seekable range of the timeline, which is a subset of the total
-		 * time range. Playback can start at any point in the seekable range with
-		 * additional constraints imposed by the format.
-		 *
-		 * @return Time range in which playback can start.
-		 */
-		virtual FTimeRange GetSeekableTimeRange() const = 0;
-
-		/**
-		 * Returns the timestamps of the segments from the video or audio track (if no video is present).
-		 * Segments are required to start with a keyframe and can thus be used to start playback with.
-		 */
-		virtual void GetSeekablePositions(TArray<FTimespan>& OutPositions) const = 0;
-
-		/**
-		 * Returns the duration of the assets on the timeline.
-		 * Typically this is the difference of the end and start values of the total time range
-		 * unless this is a Live presentation for which the duration will be set to infinite.
-		 *
-		 * @return Duration of the timeline.
-		 */
-		virtual FTimeValue GetDuration() const = 0;
-
-		/**
-		 * Returns the number of media assets on the playback timeline.
-		 *
-		 * @return Number of media assets on the playback timeline.
-		 */
-		virtual int32 GetNumberOfMediaAssets() const = 0;
-
-		/**
-		 * Returns a media asset from the timeline by its index.
-		 * Please note that media assets are not sorted by their time range on the timeline.
-		 * This is done on purpose to keep the order of assets as they appear in the manifest.
-		 *
-		 * @param MediaAssetIndex
-		 *               Index (0 to GetNumberOfMediaAssets()-1) of the asset to get.
-		 *
-		 * @return Shared pointer to the requested media asset.
-		 */
-		virtual TSharedPtrTS<ITimelineMediaAsset> GetMediaAssetByIndex(int32 MediaAssetIndex) const = 0;
-	};
-
-
-
 } // namespace Electra
 
 
diff --git a/Engine/Plugins/Media/ElectraPlayer/Source/ElectraPlayerRuntime/Private/Runtime/Player/PlaylistReader.h b/Engine/Plugins/Media/ElectraPlayer/Source/ElectraPlayerRuntime/Private/Runtime/Player/PlaylistReader.h
--- a/Engine/Plugins/Media/ElectraPlayer/Source/ElectraPlayerRuntime/Private/Runtime/Player/PlaylistReader.h
+++ b/Engine/Plugins/Media/ElectraPlayer/Source/ElectraPlayerRuntime/Private/Runtime/Player/PlaylistReader.h
@@ -35,12 +35,8 @@ namespace Electra
 		 * Loads and parses the playlist.
 		 *
 		 * @param URL     URL of the playlist to load
-		 * @param Preferences
-		 *                User preferences (for initial stream selection)
-		 * @param Options Options for the playlist reader and parser specific to the format.
-		 *                See specifics in the format's implementation.
 		 */
-		virtual void LoadAndParse(const FString& URL, const FStreamPreferences& Preferences, const FParamDict& Options) = 0;
+		virtual void LoadAndParse(const FString& URL) = 0;
 
 		/**
 		 * Returns the URL from which the playlist was loaded (or supposed to be loaded).
diff --git a/Engine/Plugins/Media/ElectraPlayer/Source/ElectraPlayerRuntime/Private/Runtime/Player/mp4/ManifestMP4.cpp b/Engine/Plugins/Media/ElectraPlayer/Source/ElectraPlayerRuntime/Private/Runtime/Player/mp4/ManifestMP4.cpp
--- a/Engine/Plugins/Media/ElectraPlayer/Source/ElectraPlayerRuntime/Private/Runtime/Player/mp4/ManifestMP4.cpp
+++ b/Engine/Plugins/Media/ElectraPlayer/Source/ElectraPlayerRuntime/Private/Runtime/Player/mp4/ManifestMP4.cpp
@@ -87,21 +87,6 @@ IManifest::EType FManifestMP4Internal::GetPresentationType() const
 }
 
 
-//-----------------------------------------------------------------------------
-/**
- * Returns the media timeline object for this asset.
- *
- * @return
- */
-TSharedPtrTS<IPlaybackAssetTimeline> FManifestMP4Internal::GetTimeline() const
-{
-	// Since an mp4 file is fixed and will not change the timeline is fixed.
-	// That's why we are inheriting from IPlaybackAssetTimeline and return ourselves.
-	TSharedPtrTS<const FManifestMP4Internal> This = SharedThis(this);
-	return StaticCastSharedPtr<IPlaybackAssetTimeline>(ConstCastSharedPtr<FManifestMP4Internal>(This));
-}
-
-
 //-----------------------------------------------------------------------------
 /**
  * Returns the starting bitrate.
@@ -290,10 +275,10 @@ TSharedPtrTS<ITimelineMediaAsset> FManifestMP4Internal::FPlayPeriodMP4::GetMedia
 /**
  * Selects a particular stream (== internal track ID) for playback.
  *
- * @param AdaptationSet
- * @param Representation
+ * @param AdaptationSetID
+ * @param RepresentationID
  */
-void FManifestMP4Internal::FPlayPeriodMP4::SelectStream(const TSharedPtrTS<IPlaybackAssetAdaptationSet>& AdaptationSet, const TSharedPtrTS<IPlaybackAssetRepresentation>& Representation)
+void FManifestMP4Internal::FPlayPeriodMP4::SelectStream(const FString& AdaptationSetID, const FString& RepresentationID)
 {
 	// Presently this method is only called by the ABR to switch between quality levels.
 	// Since a single mp4 doesn't have different quality levels (technically it could, but we are concerning ourselves only with different bitrates and that doesn't apply since we are streaming
@@ -341,6 +326,19 @@ IManifest::FResult FManifestMP4Internal::FPlayPeriodMP4::GetLoopingSegment(TShar
 }
 
 
+//-----------------------------------------------------------------------------
+/**
+ * Called by the ABR to increase the delay in fetching the next segment in case the segment returned a 404 when fetched at
+ * the announced availability time. This may reduce 404's on the next segment fetches.
+ * 
+ * @param IncreaseAmount
+ */
+void FManifestMP4Internal::FPlayPeriodMP4::IncreaseSegmentFetchDelay(const FTimeValue& IncreaseAmount)
+{
+	// No-op.
+}
+
+
 //-----------------------------------------------------------------------------
 /**
  * Creates the next segment request.
@@ -385,12 +383,12 @@ IManifest::FResult FManifestMP4Internal::FPlayPeriodMP4::GetRetrySegment(TShared
  * @param AdaptationSet
  * @param Representation
  */
-void FManifestMP4Internal::FPlayPeriodMP4::GetSegmentInformation(TArray<FSegmentInformation>& OutSegmentInformation, FTimeValue& OutAverageSegmentDuration, TSharedPtrTS<const IStreamSegment> CurrentSegment, const FTimeValue& LookAheadTime, const TSharedPtrTS<IPlaybackAssetAdaptationSet>& AdaptationSet, const TSharedPtrTS<IPlaybackAssetRepresentation>& Representation)
+void FManifestMP4Internal::FPlayPeriodMP4::GetSegmentInformation(TArray<FSegmentInformation>& OutSegmentInformation, FTimeValue& OutAverageSegmentDuration, TSharedPtrTS<const IStreamSegment> CurrentSegment, const FTimeValue& LookAheadTime, const FString& AdaptationSetID, const FString& RepresentationID)
 {
 	TSharedPtrTS<FTimelineAssetMP4> ma = MediaAsset.Pin();
 	if (ma.IsValid())
 	{
-		ma->GetSegmentInformation(OutSegmentInformation, OutAverageSegmentDuration, CurrentSegment, LookAheadTime, AdaptationSet, Representation);
+		ma->GetSegmentInformation(OutSegmentInformation, OutAverageSegmentDuration, CurrentSegment, LookAheadTime, AdaptationSetID, RepresentationID);
 	}
 }
 
@@ -849,7 +847,7 @@ IManifest::FResult FManifestMP4Internal::FTimelineAssetMP4::GetLoopingSegment(TS
 }
 
 
-void FManifestMP4Internal::FTimelineAssetMP4::GetSegmentInformation(TArray<IManifest::IPlayPeriod::FSegmentInformation>& OutSegmentInformation, FTimeValue& OutAverageSegmentDuration, TSharedPtrTS<const IStreamSegment> CurrentSegment, const FTimeValue& LookAheadTime, const TSharedPtrTS<IPlaybackAssetAdaptationSet>& AdaptationSet, const TSharedPtrTS<IPlaybackAssetRepresentation>& Representation)
+void FManifestMP4Internal::FTimelineAssetMP4::GetSegmentInformation(TArray<IManifest::IPlayPeriod::FSegmentInformation>& OutSegmentInformation, FTimeValue& OutAverageSegmentDuration, TSharedPtrTS<const IStreamSegment> CurrentSegment, const FTimeValue& LookAheadTime, const FString& AdaptationSetID, const FString& RepresentationID)
 {
 	// This is not expected to be called. And if it does we return a dummy entry.
 	OutAverageSegmentDuration.SetFromSeconds(60.0);
@@ -859,7 +857,7 @@ void FManifestMP4Internal::FTimelineAssetMP4::GetSegmentInformation(TArray<IMani
 	si.Duration.SetFromSeconds(60.0);
 }
 
-TSharedPtrTS<IParserISO14496_12>	FManifestMP4Internal::FTimelineAssetMP4::GetMoovBoxParser()
+TSharedPtrTS<IParserISO14496_12> FManifestMP4Internal::FTimelineAssetMP4::GetMoovBoxParser()
 {
 	return MoovBoxParser;
 }
diff --git a/Engine/Plugins/Media/ElectraPlayer/Source/ElectraPlayerRuntime/Private/Runtime/Player/mp4/ManifestMP4.h b/Engine/Plugins/Media/ElectraPlayer/Source/ElectraPlayerRuntime/Private/Runtime/Player/mp4/ManifestMP4.h
--- a/Engine/Plugins/Media/ElectraPlayer/Source/ElectraPlayerRuntime/Private/Runtime/Player/mp4/ManifestMP4.h
+++ b/Engine/Plugins/Media/ElectraPlayer/Source/ElectraPlayerRuntime/Private/Runtime/Player/mp4/ManifestMP4.h
@@ -21,7 +21,7 @@ namespace Electra
  * This class represents the internal "manifest" or "playlist" of an mp4 file.
  * All metadata of the tracks it is composed of are maintained by this class.
  */
-class FManifestMP4Internal : public IManifest, public IPlaybackAssetTimeline, public TSharedFromThis<FManifestMP4Internal, ESPMode::ThreadSafe>
+class FManifestMP4Internal : public IManifest, public TSharedFromThis<FManifestMP4Internal, ESPMode::ThreadSafe>
 {
 public:
 	FManifestMP4Internal(IPlayerSessionServices* InPlayerSessionServices);
@@ -30,7 +30,32 @@ public:
 	FErrorDetail Build(TSharedPtrTS<IParserISO14496_12> MP4Parser, const FString& URL, const HTTP::FConnectionInfo& InConnectionInfo);
 
 	virtual EType GetPresentationType() const override;
-	virtual TSharedPtrTS<IPlaybackAssetTimeline> GetTimeline() const override;
+	virtual FTimeValue GetAnchorTime() const override
+	{ return FTimeValue::GetZero(); }
+	virtual FTimeRange GetTotalTimeRange() const override
+	{ return MediaAsset.IsValid() ? MediaAsset->GetTimeRange() : FTimeRange(); }
+	virtual FTimeRange GetSeekableTimeRange() const override
+	{
+		FTimeRange tr = GetTotalTimeRange();
+		// FIXME: This would need to be the time of the last sync frame of video (if it exists) or audio.
+		//        For now it does not need to be precise since the nearest sync frame is searched for when seeking.
+		//        It doesn't seem to make a lot of sense to seek very close to the end and play from there so we
+		//        allow only to go as close to the end as we deem feasible.
+		tr.End -= FTimeValue().SetFromSeconds(2.0);
+		if (tr.End < FTimeValue::GetZero())
+		{
+			tr.End = FTimeValue::GetZero();
+		}
+		return tr;
+	}
+	virtual void GetSeekablePositions(TArray<FTimespan>& OutPositions) const override
+	{
+		// For the time being we do not return anything here as that would require to iterate the tracks.
+	}
+	virtual FTimeValue GetDuration() const override
+	{ return MediaAsset.IsValid() ? MediaAsset->GetDuration() : FTimeValue(); }
+	virtual FTimeValue GetDefaultStartTime() const override
+	{ return FTimeValue::GetInvalid(); }
 	virtual int64 GetDefaultStartingBitrate() const override;
 	virtual void GetStreamMetadata(TArray<FStreamMetadata>& OutMetadata, EStreamType StreamType) const override;
 	virtual FTimeValue GetMinBufferTime() const override;
@@ -167,7 +192,7 @@ public:
 		FResult GetRetrySegment(TSharedPtrTS<IStreamSegment>& OutSegment, TSharedPtrTS<const IStreamSegment> CurrentSegment, bool bReplaceWithFillerData);
 		FResult GetLoopingSegment(TSharedPtrTS<IStreamSegment>& OutSegment, FPlayerLoopState& InOutLoopState, const TMultiMap<EStreamType, TSharedPtrTS<IStreamSegment>>& InFinishedSegments, const FPlayStartPosition& StartPosition, ESearchType SearchType);
 
-		void GetSegmentInformation(TArray<IManifest::IPlayPeriod::FSegmentInformation>& OutSegmentInformation, FTimeValue& OutAverageSegmentDuration, TSharedPtrTS<const IStreamSegment> CurrentSegment, const FTimeValue& LookAheadTime, const TSharedPtrTS<IPlaybackAssetAdaptationSet>& AdaptationSet, const TSharedPtrTS<IPlaybackAssetRepresentation>& Representation);
+		void GetSegmentInformation(TArray<IManifest::IPlayPeriod::FSegmentInformation>& OutSegmentInformation, FTimeValue& OutAverageSegmentDuration, TSharedPtrTS<const IStreamSegment> CurrentSegment, const FTimeValue& LookAheadTime, const FString& AdaptationSetID, const FString& RepresentationID);
 		TSharedPtrTS<IParserISO14496_12>	GetMoovBoxParser();
 		const FString& GetMediaURL() const
 		{
@@ -193,12 +218,13 @@ public:
 		virtual EReadyState GetReadyState() override;
 		virtual void PrepareForPlay(const FParamDict& Options) override;
 		virtual TSharedPtrTS<ITimelineMediaAsset> GetMediaAsset() const override;
-		virtual void SelectStream(const TSharedPtrTS<IPlaybackAssetAdaptationSet>& AdaptationSet, const TSharedPtrTS<IPlaybackAssetRepresentation>& Representation) override;
+		virtual void SelectStream(const FString& AdaptationSetID, const FString& RepresentationID) override;
 		virtual FResult GetStartingSegment(TSharedPtrTS<IStreamSegment>& OutSegment, const FPlayStartPosition& StartPosition, ESearchType SearchType) override;
 		virtual FResult GetNextSegment(TSharedPtrTS<IStreamSegment>& OutSegment, TSharedPtrTS<const IStreamSegment> CurrentSegment) override;
 		virtual FResult GetRetrySegment(TSharedPtrTS<IStreamSegment>& OutSegment, TSharedPtrTS<const IStreamSegment> CurrentSegment, bool bReplaceWithFillerData) override;
 		virtual FResult GetLoopingSegment(TSharedPtrTS<IStreamSegment>& OutSegment, FPlayerLoopState& InOutLoopState, const TMultiMap<EStreamType, TSharedPtrTS<IStreamSegment>>& InFinishedSegments, const FPlayStartPosition& StartPosition, ESearchType SearchType) override;
-		virtual void GetSegmentInformation(TArray<FSegmentInformation>& OutSegmentInformation, FTimeValue& OutAverageSegmentDuration, TSharedPtrTS<const IStreamSegment> CurrentSegment, const FTimeValue& LookAheadTime, const TSharedPtrTS<IPlaybackAssetAdaptationSet>& AdaptationSet, const TSharedPtrTS<IPlaybackAssetRepresentation>& Representation) override;
+		virtual void IncreaseSegmentFetchDelay(const FTimeValue& IncreaseAmount) override;
+		virtual void GetSegmentInformation(TArray<FSegmentInformation>& OutSegmentInformation, FTimeValue& OutAverageSegmentDuration, TSharedPtrTS<const IStreamSegment> CurrentSegment, const FTimeValue& LookAheadTime, const FString& AdaptationSetID, const FString& RepresentationID) override;
 	private:
 		TWeakPtrTS<FTimelineAssetMP4>			MediaAsset;
 		FStreamPreferences						Preferences;
@@ -206,49 +232,6 @@ public:
 		bool									bIsReady;
 	};
 
-	//-------------------------------------------------------------------------
-	// Methods from IPlaybackAssetTimeline
-	//
-	virtual FTimeValue GetAnchorTime() const override
-	{
-		return FTimeValue::GetZero();
-	}
-	virtual FTimeRange GetTotalTimeRange() const override
-	{
-		return MediaAsset.IsValid() ? MediaAsset->GetTimeRange() : FTimeRange();
-	}
-	virtual FTimeRange GetSeekableTimeRange() const override
-	{
-		FTimeRange tr = GetTotalTimeRange();
-		// FIXME: This would need to be the time of the last sync frame of video (if it exists) or audio.
-		//        For now it does not need to be precise since the nearest sync frame is searched for when seeking.
-		//        It doesn't seem to make a lot of sense to seek very close to the end and play from there so we
-		//        allow only to go as close to the end as we deem feasible.
-		tr.End -= FTimeValue().SetFromSeconds(2.0);
-		if (tr.End < FTimeValue::GetZero())
-		{
-			tr.End = FTimeValue::GetZero();
-		}
-		return tr;
-	}
-	virtual void GetSeekablePositions(TArray<FTimespan>& OutPositions) const override
-	{
-		// For the time being we do not return anything here as that would require to iterate the tracks.
-	}
-	virtual FTimeValue GetDuration() const override
-	{
-		return MediaAsset.IsValid() ? MediaAsset->GetDuration() : FTimeValue();
-	}
-	virtual int32 GetNumberOfMediaAssets() const override
-	{
-		return 1;
-	}
-	virtual TSharedPtrTS<ITimelineMediaAsset> GetMediaAssetByIndex(int32 MediaAssetIndex) const override
-	{
-		return MediaAsset;
-	}
-
-
 	void LogMessage(IInfoLog::ELevel Level, const FString& Message);
 
 	IPlayerSessionServices* 			PlayerSessionServices;
diff --git a/Engine/Plugins/Media/ElectraPlayer/Source/ElectraPlayerRuntime/Private/Runtime/Player/mp4/PlaylistReaderMP4.cpp b/Engine/Plugins/Media/ElectraPlayer/Source/ElectraPlayerRuntime/Private/Runtime/Player/mp4/PlaylistReaderMP4.cpp
--- a/Engine/Plugins/Media/ElectraPlayer/Source/ElectraPlayerRuntime/Private/Runtime/Player/mp4/PlaylistReaderMP4.cpp
+++ b/Engine/Plugins/Media/ElectraPlayer/Source/ElectraPlayerRuntime/Private/Runtime/Player/mp4/PlaylistReaderMP4.cpp
@@ -51,11 +51,8 @@ public:
 	 * Loads and parses the playlist.
 	 *
 	 * @param URL     URL of the playlist to load
-	 * @param Preferences
-	 *                User preferences for initial stream selection.
-	 * @param Options Options for the reader and parser
 	 */
-	virtual void LoadAndParse(const FString& URL, const FStreamPreferences& Preferences, const FParamDict& Options) override;
+	virtual void LoadAndParse(const FString& URL) override;
 
 	/**
 	 * Returns the URL from which the playlist was loaded (or supposed to be loaded).
@@ -108,7 +105,7 @@ private:
 
 	void StartWorkerThread();
 	void StopWorkerThread();
-	void WorkerThread(void);
+	void WorkerThread();
 
 	void PostError(const FString& Message, uint16 Code, UEMediaError Error = UEMEDIA_ERROR_OK);
 	void LogMessage(IInfoLog::ELevel Level, const FString& Message);
@@ -124,8 +121,6 @@ private:
 	}
 
 	IPlayerSessionServices*									PlayerSessionServices;
-	FStreamPreferences										StreamPreferences;
-	FParamDict												Options;
 	FString													MasterPlaylistURL;
 	FMediaEvent												WorkerThreadQuitSignal;
 	bool													bIsWorkerThreadStarted;
@@ -266,10 +261,8 @@ void FPlaylistReaderMP4::LogMessage(IInfoLog::ELevel Level, const FString& Messa
 	}
 }
 
-void FPlaylistReaderMP4::LoadAndParse(const FString& URL, const FStreamPreferences& Preferences, const FParamDict& InOptions)
+void FPlaylistReaderMP4::LoadAndParse(const FString& URL)
 {
-	StreamPreferences = Preferences;
-	Options 		  = InOptions;
 	MasterPlaylistURL = URL;
 	StartWorkerThread();
 }
@@ -286,7 +279,7 @@ void FPlaylistReaderMP4::HTTPCompletionCallback(const IElectraHttpManager::FRequ
 }
 
 
-void FPlaylistReaderMP4::WorkerThread(void)
+void FPlaylistReaderMP4::WorkerThread()
 {
 	LLM_SCOPE(ELLMTag::ElectraPlayer);
 	CSV_SCOPED_TIMING_STAT(ElectraPlayer, PlaylistReaderMP4_Worker);
diff --git a/Engine/Plugins/Media/ElectraPlayer/Source/ElectraPlayerRuntime/Private/Runtime/PlayerFacility.h b/Engine/Plugins/Media/ElectraPlayer/Source/ElectraPlayerRuntime/Private/Runtime/PlayerFacility.h
--- a/Engine/Plugins/Media/ElectraPlayer/Source/ElectraPlayerRuntime/Private/Runtime/PlayerFacility.h
+++ b/Engine/Plugins/Media/ElectraPlayer/Source/ElectraPlayerRuntime/Private/Runtime/PlayerFacility.h
@@ -38,6 +38,7 @@ namespace Electra
 			DASHMPDBuilder,
 			DASHManifest,
 			DASHFMP4Reader,
+			EntityCache,
 			LastEntry = 255
 		};
 
@@ -95,6 +96,8 @@ namespace Electra
 					return TEXT("DASH manifest");
 				case EFacility::DASHFMP4Reader:
 					return TEXT("DASH fmp4 reader");
+				case EFacility::EntityCache:
+					return TEXT("Entity cache");
 			}
 			return TEXT("???");
 		}
diff --git a/Engine/Plugins/Media/ElectraPlayer/Source/ElectraPlayerRuntime/Private/Runtime/SynchronizedClock.cpp b/Engine/Plugins/Media/ElectraPlayer/Source/ElectraPlayerRuntime/Private/Runtime/SynchronizedClock.cpp
--- a/Engine/Plugins/Media/ElectraPlayer/Source/ElectraPlayerRuntime/Private/Runtime/SynchronizedClock.cpp
+++ b/Engine/Plugins/Media/ElectraPlayer/Source/ElectraPlayerRuntime/Private/Runtime/SynchronizedClock.cpp
@@ -14,13 +14,14 @@ namespace Electra
 		virtual ~FSynchronizedUTCTime();
 
 		virtual void SetTime(const FTimeValue& TimeNow) override;
-
+		virtual void SetTime(const FTimeValue& LocalTime, const FTimeValue& UTCTime) override;
 		virtual FTimeValue GetTime() override;
+		virtual FTimeValue MapToSyncTime(const FTimeValue& InTimeToMap) override;
 
 	private:
+		FMediaCriticalSection	Lock;
 		FTimeValue				BaseUTCTime;
 		int64					BaseSystemTime;
-		FMediaCriticalSection	Lock;
 	};
 
 
@@ -51,6 +52,14 @@ namespace Electra
 		Lock.Unlock();
 	}
 
+	void FSynchronizedUTCTime::SetTime(const FTimeValue& LocalTime, const FTimeValue& UTCTime)
+	{
+		Lock.Lock();
+		BaseUTCTime = UTCTime;
+		BaseSystemTime = LocalTime.GetAsMilliseconds();
+		Lock.Unlock();
+	}
+
 	FTimeValue FSynchronizedUTCTime::GetTime()
 	{
 		int64 NowSystem = MEDIAutcTime::CurrentMSec();
@@ -58,10 +67,16 @@ namespace Electra
 		FTimeValue 	LastBaseUTC = BaseUTCTime;
 		int64   	LastBaseSystem = BaseSystemTime;
 		Lock.Unlock();
+		return LastBaseUTC + FTimeValue(FTimeValue::MillisecondsToHNS(NowSystem - LastBaseSystem));
+	}
 
-		FTimeValue Offset;
-		Offset.SetFromMilliseconds(NowSystem - LastBaseSystem);
-		return LastBaseUTC + Offset;
+	FTimeValue FSynchronizedUTCTime::MapToSyncTime(const FTimeValue& InTimeToMap)
+	{
+		Lock.Lock();
+		FTimeValue 	LastBaseUTC = BaseUTCTime;
+		int64   	LastBaseSystem = BaseSystemTime;
+		Lock.Unlock();
+		return LastBaseUTC + FTimeValue(FTimeValue::MillisecondsToHNS(InTimeToMap.GetAsMilliseconds() - LastBaseSystem));
 	}
 
 
diff --git a/Engine/Plugins/Media/ElectraPlayer/Source/ElectraPlayerRuntime/Private/Runtime/SynchronizedClock.h b/Engine/Plugins/Media/ElectraPlayer/Source/ElectraPlayerRuntime/Private/Runtime/SynchronizedClock.h
--- a/Engine/Plugins/Media/ElectraPlayer/Source/ElectraPlayerRuntime/Private/Runtime/SynchronizedClock.h
+++ b/Engine/Plugins/Media/ElectraPlayer/Source/ElectraPlayerRuntime/Private/Runtime/SynchronizedClock.h
@@ -17,10 +17,22 @@ namespace Electra
 		static ISynchronizedUTCTime* Create();
 		virtual ~ISynchronizedUTCTime() = default;
 
+		/**
+		 * Establishes a relationship between the current local system time and actual UTC time.
+		 */
 		virtual void SetTime(const FTimeValue& TimeNow) = 0;
 
+		virtual void SetTime(const FTimeValue& LocalTime, const FTimeValue& UTCTime) = 0;
+
+		/**
+		 * Gets the current system time as synchronized UTC time.
+		 */
 		virtual FTimeValue GetTime() = 0;
 
+		/**
+		 * Returns the synchronized UTC time for a given time.
+		 */
+		virtual FTimeValue MapToSyncTime(const FTimeValue& InTimeToMap) = 0;
 	};
 
 
diff --git a/Engine/Plugins/Media/ElectraPlayer/Source/ElectraPlayerRuntime/Private/Runtime/TimeUtilities.cpp b/Engine/Plugins/Media/ElectraPlayer/Source/ElectraPlayerRuntime/Private/Runtime/TimeUtilities.cpp
--- a/Engine/Plugins/Media/ElectraPlayer/Source/ElectraPlayerRuntime/Private/Runtime/TimeUtilities.cpp
+++ b/Engine/Plugins/Media/ElectraPlayer/Source/ElectraPlayerRuntime/Private/Runtime/TimeUtilities.cpp
@@ -6,28 +6,6 @@
 namespace Electra
 {
 
-	FTimeValue& FTimeValue::SetFromTimeFraction(const FTimeFraction& TimeFraction)
-	{
-		if (TimeFraction.IsValid())
-		{
-			if (TimeFraction.IsPositiveInfinity())
-			{
-				SetToPositiveInfinity();
-			}
-			else
-			{
-				SetFromND(TimeFraction.GetNumerator(), TimeFraction.GetDenominator());
-			}
-		}
-		else
-		{
-			SetToInvalid();
-		}
-		return *this;
-	}
-
-
-
 	namespace
 	{
 		struct FTimeComponents
@@ -292,7 +270,7 @@ namespace Electra
 					{
 						if (Anchor && !bHaveT && !bHaveY)
 						{
-							int32 Years = ParseSubStringToInt(Anchor, Start-Anchor);
+							int64 Years = ParseSubStringToInt(Anchor, Start-Anchor);
 							TotalSeconds += Years * (12 * 30 * 24 * 60 * 60);
 							Anchor = nullptr;
 							++Start;
@@ -308,7 +286,7 @@ namespace Electra
 					{
 						if (Anchor && !bHaveM)
 						{
-							int32 Value = ParseSubStringToInt(Anchor, Start-Anchor);
+							int64 Value = ParseSubStringToInt(Anchor, Start-Anchor);
 							if (!bHaveT)
 							{
 								TotalSeconds += Value * (30 * 24 * 60 * 60);
@@ -331,7 +309,7 @@ namespace Electra
 					{
 						if (Anchor && !bHaveT && !bHaveD)
 						{
-							int32 Days = ParseSubStringToInt(Anchor, Start-Anchor);
+							int64 Days = ParseSubStringToInt(Anchor, Start-Anchor);
 							TotalSeconds += Days * (24 * 60 * 60);
 							Anchor = nullptr;
 							++Start;
@@ -347,7 +325,7 @@ namespace Electra
 					{
 						if (Anchor && bHaveT && !bHaveH)
 						{
-							int32 Hours = ParseSubStringToInt(Anchor, Start-Anchor);
+							int64 Hours = ParseSubStringToInt(Anchor, Start-Anchor);
 							TotalSeconds += Hours * (60 * 60);
 							Anchor = nullptr;
 							++Start;
@@ -557,6 +535,39 @@ namespace Electra
 	} // namespace RFC7231
 
 
+	namespace RFC2326
+	{
+		bool ParseNPTTime(FTimeValue& OutTimeValue, const FString& NPTtime)
+		{
+			// See: https://www.w3.org/TR/media-frags/#naming-time
+			// and: https://www.ietf.org/rfc/rfc2326.txt  (section 3.6)
+			TArray<FString> NPTparts;
+			const TCHAR* const NPTDelimiter = TEXT(":");
+			NPTtime.ParseIntoArray(NPTparts, NPTDelimiter, false);
+			if (NPTparts.Num() < 4)
+			{
+				// Note that we do no validation here on whether or not the minutes and seconds are
+				// using double digits (in case of h:mm:ss.fff) or if they are within 0-59 range.
+				int64 h=0, m=0;
+				if (NPTparts.Num() == 3)
+				{
+					// h:mm:ss[.fff]
+					LexFromString(h, *NPTparts[0]);
+					LexFromString(m, *NPTparts[1]);
+				}
+				else if (NPTparts.Num() == 2)
+				{
+					// mm:ss[.fff]
+					LexFromString(m, *NPTparts[1]);
+				}
+				// The last part will always be seconds with optional fraction. We parse those independently from minutes and hours and add everything.
+				OutTimeValue = FTimeValue().SetFromMilliseconds((h * 3600 + m * 60) * 1000) + FTimeValue().SetFromTimeFraction(FTimeFraction().SetFromFloatString(NPTparts.Last()));
+				return true;
+			}
+			return false;
+		}
+	} // namespace RFC2326
+
 
 } // namespace Electra
 
diff --git a/Engine/Plugins/Media/ElectraPlayer/Source/ElectraPlayerRuntime/Private/Runtime/Utilities/TimeUtilities.h b/Engine/Plugins/Media/ElectraPlayer/Source/ElectraPlayerRuntime/Private/Runtime/Utilities/TimeUtilities.h
--- a/Engine/Plugins/Media/ElectraPlayer/Source/ElectraPlayerRuntime/Private/Runtime/Utilities/TimeUtilities.h
+++ b/Engine/Plugins/Media/ElectraPlayer/Source/ElectraPlayerRuntime/Private/Runtime/Utilities/TimeUtilities.h
@@ -20,5 +20,10 @@ namespace Electra
 		UEMediaError ParseDateTime(FTimeValue& OutTimeValue, const FString& DateTime);
 	}
 
+	namespace RFC2326
+	{
+		bool ParseNPTTime(FTimeValue& OutTimeValue, const FString& NPTtime);
+	}
+
 } // namespace Electra
 
diff --git a/Engine/Plugins/Media/ElectraPlayer/Source/ElectraPlayerRuntime/Private/Runtime/Utilities/URI.cpp b/Engine/Plugins/Media/ElectraPlayer/Source/ElectraPlayerRuntime/Private/Runtime/Utilities/URI.cpp
--- a/Engine/Plugins/Media/ElectraPlayer/Source/ElectraPlayerRuntime/Private/Runtime/Utilities/URI.cpp
+++ b/Engine/Plugins/Media/ElectraPlayer/Source/ElectraPlayerRuntime/Private/Runtime/Utilities/URI.cpp
@@ -344,7 +344,7 @@ namespace Electra
                                                      (Position + 2 == Size || InPath[Position + 2] == TCHAR('/')) &&
                                                      (Position == 0 || InPath[Position - 1] == TCHAR('/'))))
                 {
-                    if (!PathParts.IsEmpty())
+                    if (PathParts.Num())
                     {
                         PathParts.Pop(false);
                     }
diff --git a/Engine/Plugins/Media/ElectraPlayer/Source/ElectraPlayerRuntime/Private/Runtime/Utilities/URLParser.cpp b/Engine/Plugins/Media/ElectraPlayer/Source/ElectraPlayerRuntime/Private/Runtime/Utilities/URLParser.cpp
--- a/Engine/Plugins/Media/ElectraPlayer/Source/ElectraPlayerRuntime/Private/Runtime/Utilities/URLParser.cpp
+++ b/Engine/Plugins/Media/ElectraPlayer/Source/ElectraPlayerRuntime/Private/Runtime/Utilities/URLParser.cpp
@@ -235,7 +235,7 @@ namespace Electra
 				}
 				UrlEncode(URL, Path, RequiredEscapeCharsPath);
 			}
-			else if (Query.Len() || Fragment.Len())
+			else if ((Query.Len() && bIncludeQuery) || (Fragment.Len() && bIncludeFragment))
 			{
 				URL += TEXT("/");
 			}
@@ -244,12 +244,12 @@ namespace Electra
 		{
 			UrlEncode(URL, Path, RequiredEscapeCharsPath);
 		}
-		if (Query.Len())
+		if (Query.Len() && bIncludeQuery)
 		{
 			URL += TEXT("?");
 			URL += Query;
 		}
-		if (Fragment.Len())
+		if (Fragment.Len() && bIncludeFragment)
 		{
 			URL += TEXT("#");
 			UrlEncode(URL, Fragment, FString());
diff --git a/Engine/Plugins/Media/ElectraUtil/Source/ElectraBase/Public/PlayerTime.h b/Engine/Plugins/Media/ElectraUtil/Source/ElectraBase/Public/PlayerTime.h
--- a/Engine/Plugins/Media/ElectraUtil/Source/ElectraBase/Public/PlayerTime.h
+++ b/Engine/Plugins/Media/ElectraUtil/Source/ElectraBase/Public/PlayerTime.h
@@ -224,52 +224,10 @@ public:
 		return *this;
 	}
 
-	FTimeValue& SetFromND(int64 Numerator, uint32 Denominator)
-	{
-		if (Denominator != 0)
-		{
-			if (Denominator == 10000000)
-			{
-				HNS 		= Numerator;
-				bIsValid	= true;
-				bIsInfinity = false;
-			}
-			else if (Numerator >= -922337203685LL && Numerator <= 922337203685LL)
-			{
-				HNS 		= Numerator * 10000000 / Denominator;
-				bIsValid	= true;
-				bIsInfinity = false;
-			}
-			else
-			{
-				// Need a bigint here!
-				check(!"TODO");
-				HNS 		= (int64)(Numerator * (10000000.0 / (double)Denominator));
-				bIsValid	= true;
-				bIsInfinity = false;
-			}
-		}
-		else
-		{
-			HNS 		= Numerator>=0 ? 0x7fffffffffffffffLL : -0x7fffffffffffffffLL;
-			bIsValid	= true;
-			bIsInfinity = true;
-		}
-		return *this;
-	}
-
+	FTimeValue& SetFromND(int64 Numerator, uint32 Denominator);
 
 	FTimeValue& SetFromTimeFraction(const FTimeFraction& TimeFraction);
 
-	#if 0
-	FTimeValue& AdvanceBySeconds(double Seconds)
-	{
-		FTimeValue tv;
-		tv.SetFromSeconds(Seconds);
-		return operator += (tv);
-	}
-	#endif
-
 
 	bool operator == (const FTimeValue& rhs) const
 	{
@@ -324,7 +282,7 @@ public:
 				}
 				else
 				{
-					SetToInvalid();
+					SetToPositiveInfinity();
 				}
 			}
 			else
@@ -347,7 +305,7 @@ public:
 				}
 				else
 				{
-					SetToInvalid();
+					SetToPositiveInfinity();
 				}
 			}
 			else
@@ -394,6 +352,10 @@ public:
 				Result.HNS = HNS + rhs.HNS;
 				Result.bIsValid = true;
 			}
+			else
+			{
+				Result.SetToPositiveInfinity();
+			}
 		}
 		return Result;
 	}
@@ -408,6 +370,10 @@ public:
 				Result.HNS = HNS - rhs.HNS;
 				Result.bIsValid = true;
 			}
+			else
+			{
+				Result.SetToPositiveInfinity();
+			}
 		}
 		return Result;
 	}
@@ -415,9 +381,16 @@ public:
 	inline FTimeValue operator << (int32 Shift) const
 	{
 		FTimeValue Result(*this);
-		if (bIsValid && !bIsInfinity)
+		if (bIsValid)
 		{
-			Result.HNS <<= Shift;
+			if (!bIsInfinity)
+			{
+				Result.HNS <<= Shift;
+			}
+			else
+			{
+				Result.SetToPositiveInfinity();
+			}
 		}
 		return Result;
 	}
@@ -425,9 +398,16 @@ public:
 	inline FTimeValue operator >> (int32 Shift) const
 	{
 		FTimeValue Result(*this);
-		if (bIsValid && !bIsInfinity)
+		if (bIsValid)
 		{
-			Result.HNS >>= Shift;
+			if (!bIsInfinity)
+			{
+				Result.HNS >>= Shift;
+			}
+			else
+			{
+				Result.SetToPositiveInfinity();
+			}
 		}
 		return Result;
 	}
@@ -435,9 +415,16 @@ public:
 	inline FTimeValue operator * (int32 Scale) const
 	{
 		FTimeValue Result(*this);
-		if (bIsValid && !bIsInfinity)
+		if (bIsValid)
 		{
-			Result.HNS *= Scale;
+			if (!bIsInfinity)
+			{
+				Result.HNS *= Scale;
+			}
+			else
+			{
+				Result.SetToPositiveInfinity();
+			}
 		}
 		return Result;
 	}
@@ -445,11 +432,8 @@ public:
 	inline FTimeValue operator / (int32 Scale) const
 	{
 		FTimeValue Result(*this);
-		if (bIsValid && !bIsInfinity)
-		{
-			Result.HNS /= Scale;
-		}
-		return(Result);
+		Result.HNS /= Scale;
+		return Result;
 	}
 
 private:
@@ -596,7 +580,7 @@ class MEDIAutcTime
 public:
 	static Electra::FTimeValue Current()
 	{
-		return Electra::FTimeValue().SetFromMilliseconds(CurrentMSec());
+		return Electra::FTimeValue(Electra::FTimeValue::MillisecondsToHNS(CurrentMSec()));
 	}
 	static int64 CurrentMSec();
 };
