ElectraPlayer: fixing static analysis complaints

#rnx

diff --git a/Engine/Plugins/Media/ElectraPlayer/Source/ElectraPlayerRuntime/Private/Runtime/Player/DASH/PlaylistReaderDASH.cpp b/Engine/Plugins/Media/ElectraPlayer/Source/ElectraPlayerRuntime/Private/Runtime/Player/DASH/PlaylistReaderDASH.cpp
--- a/Engine/Plugins/Media/ElectraPlayer/Source/ElectraPlayerRuntime/Private/Runtime/Player/DASH/PlaylistReaderDASH.cpp
+++ b/Engine/Plugins/Media/ElectraPlayer/Source/ElectraPlayerRuntime/Private/Runtime/Player/DASH/PlaylistReaderDASH.cpp
@@ -729,29 +729,31 @@ void FPlaylistReaderDASH::ManifestDownloadCompleted(FResourceLoadRequestPtr Requ
 		// Get the Date header from the response and set our clock to this time.
 		// Do this once only with the date from the master playlist.
 		const HTTP::FConnectionInfo* ConnInfo = Request->GetConnectionInfo();
+		if (!ConnInfo)
+		{
+			return;
+		}
+
 		FString ETag;
 		FTimeValue FetchTime;
 		FString EffectiveURL = ConnInfo->EffectiveURL;
-		if (ConnInfo)
+		for(int32 i=0; i<ConnInfo->ResponseHeaders.Num(); ++i)
 		{
-			for(int32 i=0; i<ConnInfo->ResponseHeaders.Num(); ++i)
+			if (ConnInfo->ResponseHeaders[i].Header == "Date")
 			{
-				if (ConnInfo->ResponseHeaders[i].Header == "Date")
-				{
-					// Parse the header
-					FTimeValue DateFromHeader;
-					UEMediaError DateParseError = RFC7231::ParseDateTime(DateFromHeader, ConnInfo->ResponseHeaders[i].Value);
-					if (DateParseError == UEMEDIA_ERROR_OK && DateFromHeader.IsValid())
-					{
-						PlayerSessionServices->GetSynchronizedUTCTime()->SetTime(ConnInfo->RequestStartTime, DateFromHeader);
-					}
-					// The MPD 'FetchTime' is the time at which the request to get the MPD was made to the server.
-					FetchTime = PlayerSessionServices->GetSynchronizedUTCTime()->MapToSyncTime(ConnInfo->RequestStartTime);
-				}
-				else if (ConnInfo->ResponseHeaders[i].Header == "ETag")
+				// Parse the header
+				FTimeValue DateFromHeader;
+				UEMediaError DateParseError = RFC7231::ParseDateTime(DateFromHeader, ConnInfo->ResponseHeaders[i].Value);
+				if (DateParseError == UEMEDIA_ERROR_OK && DateFromHeader.IsValid())
 				{
-					ETag = ConnInfo->ResponseHeaders[i].Value;
+					PlayerSessionServices->GetSynchronizedUTCTime()->SetTime(ConnInfo->RequestStartTime, DateFromHeader);
 				}
+				// The MPD 'FetchTime' is the time at which the request to get the MPD was made to the server.
+				FetchTime = PlayerSessionServices->GetSynchronizedUTCTime()->MapToSyncTime(ConnInfo->RequestStartTime);
+			}
+			else if (ConnInfo->ResponseHeaders[i].Header == "ETag")
+			{
+				ETag = ConnInfo->ResponseHeaders[i].Value;
 			}
 		}
 		MostRecentMPDUpdateTime = FetchTime;
@@ -842,19 +844,20 @@ void FPlaylistReaderDASH::ManifestUpdateDownloadCompleted(FResourceLoadRequestPt
 		}
 		
 		const HTTP::FConnectionInfo* ConnInfo = Request->GetConnectionInfo();
+		if (!ConnInfo)
+		{
+			return;
+		}
 		FString ETag;
 		FString EffectiveURL;
 		FTimeValue FetchTime;
-		if (ConnInfo)
+		EffectiveURL = ConnInfo->EffectiveURL;
+		FetchTime = PlayerSessionServices->GetSynchronizedUTCTime()->MapToSyncTime(ConnInfo->RequestStartTime);
+		for(int32 i=0; i<ConnInfo->ResponseHeaders.Num(); ++i)
 		{
-			EffectiveURL = ConnInfo->EffectiveURL;
-			FetchTime = PlayerSessionServices->GetSynchronizedUTCTime()->MapToSyncTime(ConnInfo->RequestStartTime);
-			for(int32 i=0; i<ConnInfo->ResponseHeaders.Num(); ++i)
+			if (ConnInfo->ResponseHeaders[i].Header == "ETag")
 			{
-				if (ConnInfo->ResponseHeaders[i].Header == "ETag")
-				{
-					ETag = ConnInfo->ResponseHeaders[i].Value;
-				}
+				ETag = ConnInfo->ResponseHeaders[i].Value;
 			}
 		}
 		MostRecentMPDUpdateTime = FetchTime;
