ElectraPlayer: handling of multi-period DASH

#rnx

diff --git a/Engine/Plugins/Media/ElectraPlayer/Source/ElectraPlayerRuntime/Private/ElectraPlayer.cpp b/Engine/Plugins/Media/ElectraPlayer/Source/ElectraPlayerRuntime/Private/ElectraPlayer.cpp
--- a/Engine/Plugins/Media/ElectraPlayer/Source/ElectraPlayerRuntime/Private/ElectraPlayer.cpp
+++ b/Engine/Plugins/Media/ElectraPlayer/Source/ElectraPlayerRuntime/Private/ElectraPlayer.cpp
@@ -780,23 +780,23 @@ bool FElectraPlayer::Seek(const FTimespan& Time)
 	return false;
 }
 
-TSharedPtr<Electra::FStreamMetadata, ESPMode::ThreadSafe> FElectraPlayer::GetTrackStreamMetadata(EPlayerTrackType TrackType, int32 TrackIndex) const
+TSharedPtr<Electra::FTrackMetadata, ESPMode::ThreadSafe> FElectraPlayer::GetTrackStreamMetadata(EPlayerTrackType TrackType, int32 TrackIndex) const
 {
 	TSharedPtr<FInternalPlayerImpl, ESPMode::ThreadSafe> LockedPlayer = CurrentPlayer;
 	if (LockedPlayer.IsValid() && LockedPlayer->AdaptivePlayer.IsValid())
 	{
-		TArray<Electra::FStreamMetadata> StreamMetaData;
+		TArray<Electra::FTrackMetadata> TrackMetaData;
 		if (TrackType == EPlayerTrackType::Video)
 		{
-			CurrentPlayer->AdaptivePlayer->GetStreamMetadata(StreamMetaData, Electra::EStreamType::Video);
+			CurrentPlayer->AdaptivePlayer->GetTrackMetadata(TrackMetaData, Electra::EStreamType::Video);
 		}
 		else if (TrackType == EPlayerTrackType::Audio)
 		{
-			CurrentPlayer->AdaptivePlayer->GetStreamMetadata(StreamMetaData, Electra::EStreamType::Audio);
+			CurrentPlayer->AdaptivePlayer->GetTrackMetadata(TrackMetaData, Electra::EStreamType::Audio);
 		}
-		if (TrackIndex >= 0 && TrackIndex < StreamMetaData.Num())
+		if (TrackIndex >= 0 && TrackIndex < TrackMetaData.Num())
 		{
-			return MakeShared<Electra::FStreamMetadata, ESPMode::ThreadSafe>(StreamMetaData[TrackIndex]);
+			return MakeShared<Electra::FTrackMetadata, ESPMode::ThreadSafe>(TrackMetaData[TrackIndex]);
 		}
 	}
 	return nullptr;
@@ -807,10 +807,10 @@ bool FElectraPlayer::GetAudioTrackFormat(int32 TrackIndex, int32 FormatIndex, FA
 {
 	if (TrackIndex >= 0 && TrackIndex < NumTracksAudio && FormatIndex == 0)
 	{
-		TSharedPtr<Electra::FStreamMetadata, ESPMode::ThreadSafe> Meta = GetTrackStreamMetadata(EPlayerTrackType::Audio, TrackIndex);
+		TSharedPtr<Electra::FTrackMetadata, ESPMode::ThreadSafe> Meta = GetTrackStreamMetadata(EPlayerTrackType::Audio, TrackIndex);
 		if (Meta.IsValid())
 		{
-			const Electra::FStreamCodecInformation& ci = Meta->CodecInformation;
+			const Electra::FStreamCodecInformation& ci = Meta->HighestBandwidthCodec;
 			OutFormat.BitsPerSample = 16;
 			OutFormat.NumChannels = (uint32)ci.GetNumberOfChannels();
 			OutFormat.SampleRate = (uint32)ci.GetSamplingRate();
@@ -827,10 +827,10 @@ bool FElectraPlayer::GetVideoTrackFormat(int32 TrackIndex, int32 FormatIndex, FV
 {
 	if (TrackIndex == 0 && FormatIndex == 0)
 	{
-		TSharedPtr<Electra::FStreamMetadata, ESPMode::ThreadSafe> Meta = GetTrackStreamMetadata(EPlayerTrackType::Video, TrackIndex);
+		TSharedPtr<Electra::FTrackMetadata, ESPMode::ThreadSafe> Meta = GetTrackStreamMetadata(EPlayerTrackType::Video, TrackIndex);
 		if (Meta.IsValid())
 		{
-			const Electra::FStreamCodecInformation& ci = Meta->CodecInformation;
+			const Electra::FStreamCodecInformation& ci = Meta->HighestBandwidthCodec;
 			OutFormat.Dim.X = ci.GetResolution().Width;
 			OutFormat.Dim.Y = ci.GetResolution().Height;
 			OutFormat.FrameRate = (float)ci.GetFrameRate().GetAsDouble();
@@ -885,19 +885,16 @@ int32 FElectraPlayer::GetSelectedTrack(EPlayerTrackType TrackType) const
 
 FText FElectraPlayer::GetTrackDisplayName(EPlayerTrackType TrackType, int32 TrackIndex) const
 {
-	TSharedPtr<Electra::FStreamMetadata, ESPMode::ThreadSafe> Meta = GetTrackStreamMetadata(TrackType, TrackIndex);
+	TSharedPtr<Electra::FTrackMetadata, ESPMode::ThreadSafe> Meta = GetTrackStreamMetadata(TrackType, TrackIndex);
 	if (Meta.IsValid())
 	{
-		//const Electra::FStreamCodecInformation &ci = StreamMetaData[TrackIndex].CodecInformation;
 		if (TrackType == EPlayerTrackType::Video)
 		{
-			//return FText::Format(LOCTEXT("AudioTrackDisplayNameFormat", "Audio Track {0}"), TrackIndex);
-			//Track->Name = FString::Printf(TEXT("Video Track %i"), TrackIndex);
-			return FText::FromString(FString::Printf(TEXT("Video Track ID %u"), Meta->StreamUniqueID));
+			return FText::FromString(FString::Printf(TEXT("Video Track ID %s"), *Meta->TrackID));
 		}
 		else if (TrackType == EPlayerTrackType::Audio)
 		{
-			return FText::FromString(FString::Printf(TEXT("Audio Track ID %u"), Meta->StreamUniqueID));
+			return FText::FromString(FString::Printf(TEXT("Audio Track ID %s"), *Meta->TrackID));
 		}
 	}
 	return FText();
@@ -911,34 +908,16 @@ int32 FElectraPlayer::GetTrackFormat(EPlayerTrackType TrackType, int32 TrackInde
 
 FString FElectraPlayer::GetTrackLanguage(EPlayerTrackType TrackType, int32 TrackIndex) const
 {
-	TSharedPtr<Electra::FStreamMetadata, ESPMode::ThreadSafe> Meta = GetTrackStreamMetadata(TrackType, TrackIndex);
+	TSharedPtr<Electra::FTrackMetadata, ESPMode::ThreadSafe> Meta = GetTrackStreamMetadata(TrackType, TrackIndex);
 	if (Meta.IsValid())
 	{
-		return Meta->LanguageCode;
+		return Meta->Language;
 	}
 	return TEXT("");
 }
 
 FString FElectraPlayer::GetTrackName(EPlayerTrackType TrackType, int32 TrackIndex) const
 {
-#if 0
-	TSharedPtr<Electra::FStreamMetadata, ESPMode::ThreadSafe> Meta = GetTrackStreamMetadata(TrackType, TrackIndex);
-	if (Meta.IsValid())
-	{
-		//const Electra::FStreamCodecInformation &ci = StreamMetaData[TrackIndex].CodecInformation;
-		if (TrackType == EPlayerTrackType::Video)
-		{
-			//return FText::Format(LOCTEXT("AudioTrackDisplayNameFormat", "Audio Track {0}"), TrackIndex);
-			//Track->Name = FString::Printf(TEXT("Video Track %i"), TrackIndex);
-			return FText::FromString(FString::Printf(TEXT("Video Track ID %u"), Meta->StreamUniqueID));
-		}
-		else if (TrackType == EPlayerTrackType::Audio)
-		{
-			return FText::FromString(FString::Printf(TEXT("Audio Track ID %u"), Meta->StreamUniqueID));
-		}
-	}
-	return FText();
-#endif
 	return TEXT("");
 }
 
@@ -950,7 +929,7 @@ bool FElectraPlayer::SelectTrack(EPlayerTrackType TrackType, int32 TrackIndex)
 		if (TrackIndex >= 0)
 		{
 			// Do a quick check if the track index is valid by checking the presence of the track metadata.
-			TSharedPtr<Electra::FStreamMetadata, ESPMode::ThreadSafe> Meta = GetTrackStreamMetadata(EPlayerTrackType::Audio, TrackIndex);
+			TSharedPtr<Electra::FTrackMetadata, ESPMode::ThreadSafe> Meta = GetTrackStreamMetadata(EPlayerTrackType::Audio, TrackIndex);
 			if (Meta.IsValid())
 			{
 				PlaystartOptions.AudioTrackIndex = TrackIndex;
@@ -1148,12 +1127,18 @@ void FElectraPlayer::ReportReceivedPlaylists()
 	Statistics.MediaTimelineAtEnd = MediaTimeline;
 	Statistics.MediaDuration = MediaDuration.IsInfinity() ? -1.0 : MediaDuration.GetAsSeconds();
 	// Get the video bitrates and populate our number of segments per bitrate map.
-	TArray<Electra::FStreamMetadata> VideoStreamMetaData;
-	CurrentPlayer->AdaptivePlayer->GetStreamMetadata(VideoStreamMetaData, Electra::EStreamType::Video);
+	TArray<Electra::FTrackMetadata> VideoStreamMetaData;
+	CurrentPlayer->AdaptivePlayer->GetTrackMetadata(VideoStreamMetaData, Electra::EStreamType::Video);
 	NumTracksVideo = VideoStreamMetaData.Num();
-	for(int32 i=0; i<NumTracksVideo; ++i)
+	if (NumTracksVideo)
 	{
-		UE_LOG(LogElectraPlayer, Log, TEXT("[%p][%p] Found %d * %d video stream at bitrate %d"), this, CurrentPlayer.Get(), VideoStreamMetaData[i].CodecInformation.GetResolution().Width, VideoStreamMetaData[i].CodecInformation.GetResolution().Height, VideoStreamMetaData[i].Bandwidth);
+		for(int32 i=0; i<VideoStreamMetaData[0].StreamDetails.Num(); ++i)
+		{
+			UE_LOG(LogElectraPlayer, Log, TEXT("[%p][%p] Found %d * %d video stream at bitrate %d"), this, CurrentPlayer.Get(), 
+												VideoStreamMetaData[0].StreamDetails[i].CodecInformation.GetResolution().Width,
+												VideoStreamMetaData[0].StreamDetails[i].CodecInformation.GetResolution().Height,
+												VideoStreamMetaData[0].StreamDetails[i].Bandwidth);
+		}
 	}
 	StatisticsLock.Unlock();
 
@@ -1167,8 +1152,8 @@ void FElectraPlayer::ReportReceivedPlaylists()
 		EnqueueAnalyticsEvent(AnalyticEvent);
 	}
 
-	TArray<Electra::FStreamMetadata> AudioStreamMetaData;
-	CurrentPlayer->AdaptivePlayer->GetStreamMetadata(AudioStreamMetaData, Electra::EStreamType::Audio);
+	TArray<Electra::FTrackMetadata> AudioStreamMetaData;
+	CurrentPlayer->AdaptivePlayer->GetTrackMetadata(AudioStreamMetaData, Electra::EStreamType::Audio);
 	NumTracksAudio = AudioStreamMetaData.Num();
 
 	// Select the desired audio track by hard metadata index.
@@ -1204,6 +1189,36 @@ void FElectraPlayer::ReportReceivedPlaylists()
 	CurrentPlayer->AdaptivePlayer->SeekTo(playParam);
 }
 
+void FElectraPlayer::ReportTracksChanged()
+{
+	TSharedPtr<FInternalPlayerImpl, ESPMode::ThreadSafe> LockedPlayer = CurrentPlayer;
+	if (!LockedPlayer.IsValid() || !LockedPlayer->AdaptivePlayer.IsValid())
+	{
+		return;
+	}
+
+	TArray<Electra::FTrackMetadata> VideoStreamMetaData;
+	CurrentPlayer->AdaptivePlayer->GetTrackMetadata(VideoStreamMetaData, Electra::EStreamType::Video);
+	NumTracksVideo = VideoStreamMetaData.Num();
+	if (NumTracksVideo)
+	{
+		for(int32 i=0; i<VideoStreamMetaData[0].StreamDetails.Num(); ++i)
+		{
+			UE_LOG(LogElectraPlayer, Log, TEXT("[%p][%p] Found %d * %d video stream at bitrate %d"), this, CurrentPlayer.Get(), 
+												VideoStreamMetaData[0].StreamDetails[i].CodecInformation.GetResolution().Width,
+												VideoStreamMetaData[0].StreamDetails[i].CodecInformation.GetResolution().Height,
+												VideoStreamMetaData[0].StreamDetails[i].Bandwidth);
+		}
+	}
+
+	TArray<Electra::FTrackMetadata> AudioStreamMetaData;
+	CurrentPlayer->AdaptivePlayer->GetTrackMetadata(AudioStreamMetaData, Electra::EStreamType::Audio);
+	NumTracksAudio = AudioStreamMetaData.Num();
+
+	DeferredEvents.Enqueue(IElectraPlayerAdapterDelegate::EPlayerEvent::TracksChanged);
+}
+
+
 void FElectraPlayer::ReportPlaylistDownload(const Electra::Metrics::FPlaylistDownloadStats& PlaylistDownloadStats)
 {
 	TSharedPtr<FInternalPlayerImpl, ESPMode::ThreadSafe> LockedPlayer = CurrentPlayer;
@@ -1493,17 +1508,20 @@ void FElectraPlayer::ReportVideoQualityChange(int32 NewBitrate, int32 PreviousBi
 	int32 prvWidth  = Statistics.CurrentlyActiveResolutionWidth;
 	int32 prvHeight = Statistics.CurrentlyActiveResolutionHeight;
 	// Get the current playlist URL
-	TArray<Electra::FStreamMetadata> VideoStreamMetaData;
-	CurrentPlayer->AdaptivePlayer->GetStreamMetadata(VideoStreamMetaData, Electra::EStreamType::Video);
-	for(int32 i = 0; i < VideoStreamMetaData.Num(); ++i)
+	TArray<Electra::FTrackMetadata> VideoStreamMetaData;
+	CurrentPlayer->AdaptivePlayer->GetTrackMetadata(VideoStreamMetaData, Electra::EStreamType::Video);
+	if (VideoStreamMetaData.Num())
 	{
-		if (VideoStreamMetaData[i].Bandwidth == NewBitrate)
+		for(int32 i=0; i<VideoStreamMetaData[0].StreamDetails.Num(); ++i)
 		{
-			SelectedQuality = i;
-			Statistics.CurrentlyActivePlaylistURL = VideoStreamMetaData[i].PlaylistID;
-			Statistics.CurrentlyActiveResolutionWidth = VideoStreamMetaData[i].CodecInformation.GetResolution().Width;
-			Statistics.CurrentlyActiveResolutionHeight = VideoStreamMetaData[i].CodecInformation.GetResolution().Height;
-			break;
+			if (VideoStreamMetaData[0].StreamDetails[i].Bandwidth == NewBitrate)
+			{
+				SelectedQuality = i;
+				Statistics.CurrentlyActivePlaylistURL = VideoStreamMetaData[0].TrackID;
+				Statistics.CurrentlyActiveResolutionWidth = VideoStreamMetaData[0].StreamDetails[i].CodecInformation.GetResolution().Width;
+				Statistics.CurrentlyActiveResolutionHeight = VideoStreamMetaData[0].StreamDetails[i].CodecInformation.GetResolution().Height;
+				break;
+			}
 		}
 	}
 
diff --git a/Engine/Plugins/Media/ElectraPlayer/Source/ElectraPlayerRuntime/Private/ElectraPlayer.h b/Engine/Plugins/Media/ElectraPlayer/Source/ElectraPlayerRuntime/Private/ElectraPlayer.h
--- a/Engine/Plugins/Media/ElectraPlayer/Source/ElectraPlayerRuntime/Private/ElectraPlayer.h
+++ b/Engine/Plugins/Media/ElectraPlayer/Source/ElectraPlayerRuntime/Private/ElectraPlayer.h
@@ -130,6 +130,7 @@ private:
 	virtual void ReportOpenSource(const FString& URL) override;
 	virtual void ReportReceivedMasterPlaylist(const FString& EffectiveURL) override;
 	virtual void ReportReceivedPlaylists() override;
+	virtual void ReportTracksChanged() override;
 	virtual void ReportPlaylistDownload(const Metrics::FPlaylistDownloadStats& PlaylistDownloadStats) override;
 	virtual void ReportBufferingStart(Metrics::EBufferingReason BufferingReason) override;
 	virtual void ReportBufferingEnd(Metrics::EBufferingReason BufferingReason) override;
@@ -161,7 +162,7 @@ private:
 	void MediaStateOnSeekFinished();
 	void MediaStateOnError();
 
-	TSharedPtr<FStreamMetadata, ESPMode::ThreadSafe> GetTrackStreamMetadata(EPlayerTrackType TrackType, int32 TrackIndex) const;
+	TSharedPtr<FTrackMetadata, ESPMode::ThreadSafe> GetTrackStreamMetadata(EPlayerTrackType TrackType, int32 TrackIndex) const;
 
 	// Delegate to talk back to adapter host
 	TWeakPtr<IElectraPlayerAdapterDelegate, ESPMode::ThreadSafe>	AdapterDelegate;
diff --git a/Engine/Plugins/Media/ElectraPlayer/Source/ElectraPlayerRuntime/Private/Runtime/Demuxer/ParserISO14496-12.cpp b/Engine/Plugins/Media/ElectraPlayer/Source/ElectraPlayerRuntime/Private/Runtime/Demuxer/ParserISO14496-12.cpp
--- a/Engine/Plugins/Media/ElectraPlayer/Source/ElectraPlayerRuntime/Private/Runtime/Demuxer/ParserISO14496-12.cpp
+++ b/Engine/Plugins/Media/ElectraPlayer/Source/ElectraPlayerRuntime/Private/Runtime/Demuxer/ParserISO14496-12.cpp
@@ -528,8 +528,8 @@ namespace Electra
 		static const IParserISO14496_12::FBoxType kBox_gmhd = MAKE_BOX_ATOM('g', 'm', 'h', 'd');
 		static const IParserISO14496_12::FBoxType kBox_dinf = MAKE_BOX_ATOM('d', 'i', 'n', 'f');
 		static const IParserISO14496_12::FBoxType kBox_dref = MAKE_BOX_ATOM('d', 'r', 'e', 'f');
-		static const IParserISO14496_12::FBoxType kBox_url = MAKE_BOX_ATOM('u', 'r', 'l', ' ');
-		static const IParserISO14496_12::FBoxType kBox_urn = MAKE_BOX_ATOM('u', 'r', 'n', ' ');
+		static const IParserISO14496_12::FBoxType kBox_url  = MAKE_BOX_ATOM('u', 'r', 'l', ' ');
+		static const IParserISO14496_12::FBoxType kBox_urn  = MAKE_BOX_ATOM('u', 'r', 'n', ' ');
 		static const IParserISO14496_12::FBoxType kBox_tref = MAKE_BOX_ATOM('t', 'r', 'e', 'f');
 		static const IParserISO14496_12::FBoxType kBox_stbl = MAKE_BOX_ATOM('s', 't', 'b', 'l');
 		static const IParserISO14496_12::FBoxType kBox_stsd = MAKE_BOX_ATOM('s', 't', 's', 'd');
@@ -559,6 +559,26 @@ namespace Electra
 		static const IParserISO14496_12::FBoxType kBox_tfhd = MAKE_BOX_ATOM('t', 'f', 'h', 'd');
 		static const IParserISO14496_12::FBoxType kBox_tfdt = MAKE_BOX_ATOM('t', 'f', 'd', 't');
 		static const IParserISO14496_12::FBoxType kBox_trun = MAKE_BOX_ATOM('t', 'r', 'u', 'n');
+		static const IParserISO14496_12::FBoxType kBox_trep = MAKE_BOX_ATOM('t', 'r', 'e', 'p');
+		static const IParserISO14496_12::FBoxType kBox_colr = MAKE_BOX_ATOM('c', 'o', 'l', 'r');
+		static const IParserISO14496_12::FBoxType kBox_elng = MAKE_BOX_ATOM('e', 'l', 'n', 'g');
+		static const IParserISO14496_12::FBoxType kBox_prft = MAKE_BOX_ATOM('p', 'r', 'f', 't');
+		static const IParserISO14496_12::FBoxType kBox_emsg = MAKE_BOX_ATOM('e', 'm', 's', 'g');
+		static const IParserISO14496_12::FBoxType kBox_sbgp = MAKE_BOX_ATOM('s', 'b', 'g', 'p');
+		static const IParserISO14496_12::FBoxType kBox_sgpd = MAKE_BOX_ATOM('s', 'g', 'p', 'd');
+		static const IParserISO14496_12::FBoxType kBox_saiz = MAKE_BOX_ATOM('s', 'a', 'i', 'z');
+		static const IParserISO14496_12::FBoxType kBox_saio = MAKE_BOX_ATOM('s', 'a', 'i', 'o');
+
+		/**
+		 * Encryption related boxes
+		 */
+		static const IParserISO14496_12::FBoxType kBox_pssh = MAKE_BOX_ATOM('p', 's', 's', 'h');
+		static const IParserISO14496_12::FBoxType kBox_senc = MAKE_BOX_ATOM('s', 'e', 'n', 'c');
+		static const IParserISO14496_12::FBoxType kBox_tenc = MAKE_BOX_ATOM('t', 'e', 'n', 'c');
+		static const IParserISO14496_12::FBoxType kBox_sinf = MAKE_BOX_ATOM('s', 'i', 'n', 'f');
+		static const IParserISO14496_12::FBoxType kBox_frma = MAKE_BOX_ATOM('f', 'r', 'm', 'a');
+		static const IParserISO14496_12::FBoxType kBox_schm = MAKE_BOX_ATOM('s', 'c', 'h', 'm');
+		static const IParserISO14496_12::FBoxType kBox_schi = MAKE_BOX_ATOM('s', 'c', 'h', 'i');
 
 		/**
 		 * Track handler types
@@ -578,6 +598,7 @@ namespace Electra
 		static const IParserISO14496_12::FBoxType kSample_avc1 = MAKE_BOX_ATOM('a', 'v', 'c', '1');
 		static const IParserISO14496_12::FBoxType kSample_mp4a = MAKE_BOX_ATOM('m', 'p', '4', 'a');
 		static const IParserISO14496_12::FBoxType kSample_stpp = MAKE_BOX_ATOM('s', 't', 'p', 'p');
+		static const IParserISO14496_12::FBoxType kSample_sbtt = MAKE_BOX_ATOM('s', 'b', 't', 't');
 		static const IParserISO14496_12::FBoxType kSample_enca = MAKE_BOX_ATOM('e', 'n', 'c', 'a');
 		static const IParserISO14496_12::FBoxType kSample_encv = MAKE_BOX_ATOM('e', 'n', 'c', 'v');
 
@@ -2043,6 +2064,163 @@ namespace Electra
 	};
 
 
+	/**
+	 * 'elng' box.
+	 * ISO/IEC 14496-12:2015 - 8.4.6 - Extended language tag (within 'mdia' box)
+	 */
+	class FMP4BoxELNG : public FMP4BoxFull
+	{
+	public:
+		FMP4BoxELNG(IParserISO14496_12::FBoxType InBoxType, int64 InBoxSize, int64 InStartOffset, int64 InDataOffset, bool bInIsLeafBox)
+			: FMP4BoxFull(InBoxType, InBoxSize, InStartOffset, InDataOffset, bInIsLeafBox)
+		{
+		}
+
+		virtual ~FMP4BoxELNG()
+		{
+		}
+
+	private:
+		FMP4BoxELNG() = delete;
+		FMP4BoxELNG(const FMP4BoxELNG&) = delete;
+
+	protected:
+		virtual UEMediaError ReadAndParseAttributes(FMP4ParseInfo* ParseInfo) override
+		{
+			UEMediaError Error = UEMEDIA_ERROR_OK;
+			RETURN_IF_ERROR(FMP4BoxFull::ReadAndParseAttributes(ParseInfo));
+			uint32 BytesRemaining = BoxSize - (ParseInfo->Reader()->GetCurrentReadOffset() - StartOffset);
+			RETURN_IF_ERROR(ParseInfo->Reader()->ReadString(ExtendedLanguage, BytesRemaining));
+			return Error;
+		}
+
+	private:
+		FString ExtendedLanguage;
+	};
+
+
+	/**
+	 * 'prft' box.
+	 * ISO/IEC 14496-12:2015 - 8.16.5 - Producer Reference Time Box
+	 */
+	class FMP4BoxPRFT : public FMP4BoxFull
+	{
+	public:
+		FMP4BoxPRFT(IParserISO14496_12::FBoxType InBoxType, int64 InBoxSize, int64 InStartOffset, int64 InDataOffset, bool bInIsLeafBox)
+			: FMP4BoxFull(InBoxType, InBoxSize, InStartOffset, InDataOffset, bInIsLeafBox)
+			, ReferenceTrackID(0), MediaTime(0)
+		{
+		}
+
+		virtual ~FMP4BoxPRFT()
+		{
+		}
+
+	private:
+		FMP4BoxPRFT() = delete;
+		FMP4BoxPRFT(const FMP4BoxPRFT&) = delete;
+
+	protected:
+		virtual UEMediaError ReadAndParseAttributes(FMP4ParseInfo* ParseInfo) override
+		{
+			UEMediaError Error = UEMEDIA_ERROR_OK;
+			RETURN_IF_ERROR(FMP4BoxFull::ReadAndParseAttributes(ParseInfo));
+			RETURN_IF_ERROR(ParseInfo->Reader()->Read(ReferenceTrackID));		// reference_track_ID
+			RETURN_IF_ERROR(ParseInfo->Reader()->Read(NtpTimestamp));			// ntp_timestamp
+			// Which version of the box is this?
+			if (Version == 0)
+			{
+				uint32 Value32 = 0;
+				RETURN_IF_ERROR(ParseInfo->Reader()->Read(Value32));			// media_time
+				MediaTime = Value32;
+			}
+			else
+			{
+				RETURN_IF_ERROR(ParseInfo->Reader()->Read(MediaTime));			// media_time
+			}
+			return Error;
+		}
+
+	private:
+		uint32 ReferenceTrackID;
+		uint64 NtpTimestamp;
+		uint64 MediaTime;
+	};
+
+
+	/**
+	 * 'emsg' box.
+	 * ISO/IEC 23009-1:2019 - 5.10.3.3 Event Message Box
+	 */
+	class FMP4BoxEMSG : public FMP4BoxFull
+	{
+	public:
+		FMP4BoxEMSG(IParserISO14496_12::FBoxType InBoxType, int64 InBoxSize, int64 InStartOffset, int64 InDataOffset, bool bInIsLeafBox)
+			: FMP4BoxFull(InBoxType, InBoxSize, InStartOffset, InDataOffset, bInIsLeafBox)
+			, Timescale(0), PresentationTimeDelta(0), PresentationTime(0), EventDuration(0), ID(0)
+		{
+		}
+
+		virtual ~FMP4BoxEMSG()
+		{
+		}
+
+	private:
+		FMP4BoxEMSG() = delete;
+		FMP4BoxEMSG(const FMP4BoxEMSG&) = delete;
+
+	protected:
+		virtual UEMediaError ReadAndParseAttributes(FMP4ParseInfo* ParseInfo) override
+		{
+			UEMediaError Error = UEMEDIA_ERROR_OK;
+			RETURN_IF_ERROR(FMP4BoxFull::ReadAndParseAttributes(ParseInfo));
+			// Which version of the box is this?
+			uint32 BytesRemaining = BoxSize - (ParseInfo->Reader()->GetCurrentReadOffset() - StartOffset);
+			if (Version == 0)
+			{
+				RETURN_IF_ERROR(ParseInfo->Reader()->ReadString(SchemeIdUri, BytesRemaining));						// scheme_id_uri
+				BytesRemaining = BoxSize - (ParseInfo->Reader()->GetCurrentReadOffset() - StartOffset);
+				RETURN_IF_ERROR(ParseInfo->Reader()->ReadString(Value, BytesRemaining));							// value
+				RETURN_IF_ERROR(ParseInfo->Reader()->Read(Timescale));												// timescale
+				RETURN_IF_ERROR(ParseInfo->Reader()->Read(PresentationTimeDelta));									// presentation_time_delta
+				RETURN_IF_ERROR(ParseInfo->Reader()->Read(EventDuration));											// event_duration
+				RETURN_IF_ERROR(ParseInfo->Reader()->Read(ID));														// id
+			}
+			else if (Version == 1)
+			{
+				RETURN_IF_ERROR(ParseInfo->Reader()->Read(Timescale));												// timescale
+				RETURN_IF_ERROR(ParseInfo->Reader()->Read(PresentationTime));										// presentation_time
+				RETURN_IF_ERROR(ParseInfo->Reader()->Read(EventDuration));											// event_duration
+				RETURN_IF_ERROR(ParseInfo->Reader()->Read(ID));														// id
+				BytesRemaining = BoxSize - (ParseInfo->Reader()->GetCurrentReadOffset() - StartOffset);
+				RETURN_IF_ERROR(ParseInfo->Reader()->ReadString(SchemeIdUri, BytesRemaining));						// scheme_id_uri
+				BytesRemaining = BoxSize - (ParseInfo->Reader()->GetCurrentReadOffset() - StartOffset);
+				RETURN_IF_ERROR(ParseInfo->Reader()->ReadString(Value, BytesRemaining));							// value
+			}
+			else
+			{
+				return UEMEDIA_ERROR_FORMAT_ERROR;
+			}
+			BytesRemaining = BoxSize - (ParseInfo->Reader()->GetCurrentReadOffset() - StartOffset);
+			if (BytesRemaining)
+			{
+				MessageData.AddUninitialized(BytesRemaining);
+				RETURN_IF_ERROR(ParseInfo->Reader()->ReadBytes(MessageData.GetData(), BytesRemaining));				// message_data
+			}
+			return Error;
+		}
+	private:
+		TArray<uint8> MessageData;
+		FString SchemeIdUri;
+		FString Value;
+		uint32 Timescale;
+		uint32 PresentationTimeDelta;
+		uint64 PresentationTime;
+		uint32 EventDuration;
+		uint32 ID;
+	};
+
+
 	/**
 	 * 'stsd' box.
 	 * ISO/IEC 14496-12:2014 - 8.5.2 - Sample Description Box
@@ -3095,6 +3273,345 @@ namespace Electra
 
 
 
+	/**
+	 * 'pssh' box. ISO/IEC 23001-7:2016 - 8.1 - Protection system specific header box
+	 */
+	class FMP4BoxPSSH : public FMP4BoxFull
+	{
+	public:
+		FMP4BoxPSSH(IParserISO14496_12::FBoxType InBoxType, int64 InBoxSize, int64 InStartOffset, int64 InDataOffset, bool bInIsLeafBox)
+			: FMP4BoxFull(InBoxType, InBoxSize, InStartOffset, InDataOffset, bInIsLeafBox)
+		{
+		}
+
+		virtual ~FMP4BoxPSSH()
+		{
+		}
+
+	private:
+		FMP4BoxPSSH() = delete;
+		FMP4BoxPSSH(const FMP4BoxPSSH&) = delete;
+
+	protected:
+		virtual UEMediaError ReadAndParseAttributes(FMP4ParseInfo* ParseInfo) override
+		{
+			UEMediaError Error = UEMEDIA_ERROR_OK;
+			RETURN_IF_ERROR(FMP4BoxFull::ReadAndParseAttributes(ParseInfo));
+
+			SystemID.AddUninitialized(16);
+			RETURN_IF_ERROR(ParseInfo->Reader()->ReadBytes(SystemID.GetData(), 16));		// SystemID
+			if (Version > 0)
+			{
+				uint32 KID_count = 0;
+				RETURN_IF_ERROR(ParseInfo->Reader()->Read(KID_count));						// KID_count
+				for(uint32 i=0; i<KID_count; ++i)
+				{
+					TArray<uint8>& kid = KIDs.AddDefaulted_GetRef();
+					kid.AddUninitialized(16);
+					RETURN_IF_ERROR(ParseInfo->Reader()->ReadBytes(kid.GetData(), 16));		// KID
+				}
+			}
+			uint32 DataSize = 0;
+			RETURN_IF_ERROR(ParseInfo->Reader()->Read(DataSize));							// DataSize
+			if (DataSize)
+			{
+				Data.AddUninitialized(DataSize);
+				RETURN_IF_ERROR(ParseInfo->Reader()->ReadBytes(Data.GetData(), DataSize));	// Data
+			}
+			return Error;
+		}
+
+	private:
+		TArray<uint8>			SystemID;
+		TArray<TArray<uint8>>	KIDs;
+		TArray<uint8>			Data;
+	};
+
+
+	/**
+	 * 'sinf' box. ISO/IEC 14496-12:2015 - 8.12.1 Protection Scheme Information Box
+	 */
+	class FMP4BoxSINF : public FMP4BoxBasic
+	{
+	public:
+		FMP4BoxSINF(IParserISO14496_12::FBoxType InBoxType, int64 InBoxSize, int64 InStartOffset, int64 InDataOffset, bool bInIsLeafBox)
+			: FMP4BoxBasic(InBoxType, InBoxSize, InStartOffset, InDataOffset, bInIsLeafBox)
+		{
+		}
+
+		virtual ~FMP4BoxSINF()
+		{
+		}
+
+	private:
+		FMP4BoxSINF() = delete;
+		FMP4BoxSINF(const FMP4BoxSINF&) = delete;
+
+	protected:
+		virtual UEMediaError ReadAndParseAttributes(FMP4ParseInfo* ParseInfo) override
+		{
+			// Now that we parse this box we clear its leaf status. We call back into the generic box reader
+			// for which this flag must be clear to continue with the next boxes.
+			bIsLeafBox = false;
+			UEMediaError Error = UEMEDIA_ERROR_OK;
+			RETURN_IF_ERROR(ParseInfo->ReadAndParseNextBox(this));
+			return Error;
+		}
+	};
+
+
+	/**
+	 * 'frma' box. ISO/IEC 14496-12:2015 - 8.12.2 Original Format Box
+	 */
+	class FMP4BoxFRMA : public FMP4BoxBasic
+	{
+	public:
+		FMP4BoxFRMA(IParserISO14496_12::FBoxType InBoxType, int64 InBoxSize, int64 InStartOffset, int64 InDataOffset, bool bInIsLeafBox)
+			: FMP4BoxBasic(InBoxType, InBoxSize, InStartOffset, InDataOffset, bInIsLeafBox)
+			, DataFormat(0)
+		{
+		}
+
+		virtual ~FMP4BoxFRMA()
+		{
+		}
+
+	private:
+		FMP4BoxFRMA() = delete;
+		FMP4BoxFRMA(const FMP4BoxFRMA&) = delete;
+
+	protected:
+		virtual UEMediaError ReadAndParseAttributes(FMP4ParseInfo* ParseInfo) override
+		{
+			UEMediaError Error = UEMEDIA_ERROR_OK;
+			RETURN_IF_ERROR(ParseInfo->Reader()->Read(DataFormat));						// data_format
+			return Error;
+		}
+
+	private:
+		uint32 DataFormat;
+	};
+
+
+	/**
+	 * 'schm' box. ISO/IEC 14496-12:2015 - 8.12.5 Scheme Type Box
+	 */
+	class FMP4BoxSCHM : public FMP4BoxFull
+	{
+	public:
+		FMP4BoxSCHM(IParserISO14496_12::FBoxType InBoxType, int64 InBoxSize, int64 InStartOffset, int64 InDataOffset, bool bInIsLeafBox)
+			: FMP4BoxFull(InBoxType, InBoxSize, InStartOffset, InDataOffset, bInIsLeafBox)
+			, SchemeType(0), SchemeVersion(0)
+		{
+		}
+
+		virtual ~FMP4BoxSCHM()
+		{
+		}
+
+	private:
+		FMP4BoxSCHM() = delete;
+		FMP4BoxSCHM(const FMP4BoxSCHM&) = delete;
+
+	protected:
+		virtual UEMediaError ReadAndParseAttributes(FMP4ParseInfo* ParseInfo) override
+		{
+			UEMediaError Error = UEMEDIA_ERROR_OK;
+			RETURN_IF_ERROR(FMP4BoxFull::ReadAndParseAttributes(ParseInfo));
+
+			RETURN_IF_ERROR(ParseInfo->Reader()->Read(SchemeType));									// scheme_type
+			RETURN_IF_ERROR(ParseInfo->Reader()->Read(SchemeVersion));								// scheme_version
+			if ((Flags & 1) != 0)
+			{
+				uint32 BytesRemaining = BoxSize - (ParseInfo->Reader()->GetCurrentReadOffset() - StartOffset);
+				if (BytesRemaining)
+				{
+					RETURN_IF_ERROR(ParseInfo->Reader()->ReadString(SchemeURI, BytesRemaining));	// scheme_uri
+				}
+			}
+			return Error;
+		}
+
+	private:
+		uint32 SchemeType;
+		uint32 SchemeVersion;
+		FString SchemeURI;
+	};
+
+
+	/**
+	 * 'tenc' box. ISO/IEC 23001-7:2016 - 8.2 Track Encryption box
+	 */
+	class FMP4BoxTENC : public FMP4BoxFull
+	{
+	public:
+		FMP4BoxTENC(IParserISO14496_12::FBoxType InBoxType, int64 InBoxSize, int64 InStartOffset, int64 InDataOffset, bool bInIsLeafBox)
+			: FMP4BoxFull(InBoxType, InBoxSize, InStartOffset, InDataOffset, bInIsLeafBox)
+			, DefaultCryptByteBlock(0), DefaultSkipByteBlock(0), DefaultIsProtected(0), DefaultPerSampleIVSize(0)
+		{
+		}
+
+		virtual ~FMP4BoxTENC()
+		{
+		}
+
+	private:
+		FMP4BoxTENC() = delete;
+		FMP4BoxTENC(const FMP4BoxTENC&) = delete;
+
+	protected:
+		virtual UEMediaError ReadAndParseAttributes(FMP4ParseInfo* ParseInfo) override
+		{
+			UEMediaError Error = UEMEDIA_ERROR_OK;
+			RETURN_IF_ERROR(FMP4BoxFull::ReadAndParseAttributes(ParseInfo));
+
+			uint8 Value8 = 0;
+			RETURN_IF_ERROR(ParseInfo->Reader()->Read(Value8));									// reserved
+			if (Version == 0)
+			{
+				RETURN_IF_ERROR(ParseInfo->Reader()->Read(Value8));								// reserved
+			}
+			else
+			{
+				RETURN_IF_ERROR(ParseInfo->Reader()->Read(Value8));
+				DefaultCryptByteBlock = Value8 >> 4;											// default_crypt_byte_block
+				DefaultSkipByteBlock = Value8 & 0x0f;											// default_skip_byte_block
+			}
+			RETURN_IF_ERROR(ParseInfo->Reader()->Read(DefaultIsProtected));						// default_isProtected
+			RETURN_IF_ERROR(ParseInfo->Reader()->Read(DefaultPerSampleIVSize));					// default_Per_Sample_IV_Size
+			DefaultKID.AddUninitialized(16);
+			RETURN_IF_ERROR(ParseInfo->Reader()->ReadBytes(DefaultKID.GetData(), 16));			// default_KID
+			if (DefaultIsProtected == 1 && DefaultPerSampleIVSize == 0)
+			{
+				RETURN_IF_ERROR(ParseInfo->Reader()->Read(Value8));								// default_constant_IV_size
+				DefaultConstantIV.AddUninitialized(Value8);
+				RETURN_IF_ERROR(ParseInfo->Reader()->ReadBytes(DefaultConstantIV.GetData(), Value8));// default_constant_IV
+			}
+			return Error;
+		}
+
+	private:
+		uint8	DefaultCryptByteBlock;
+		uint8	DefaultSkipByteBlock;
+		uint8	DefaultIsProtected;
+		uint8	DefaultPerSampleIVSize;
+		TArray<uint8> DefaultKID;
+		TArray<uint8> DefaultConstantIV;
+	};
+
+
+	/**
+	 * 'saiz' box. ISO/IEC 14496-12:2015 - 8.7.8 Sample Auxiliary Information Sizes Box
+	 */
+	class FMP4BoxSAIZ : public FMP4BoxFull
+	{
+	public:
+		FMP4BoxSAIZ(IParserISO14496_12::FBoxType InBoxType, int64 InBoxSize, int64 InStartOffset, int64 InDataOffset, bool bInIsLeafBox)
+			: FMP4BoxFull(InBoxType, InBoxSize, InStartOffset, InDataOffset, bInIsLeafBox)
+			, AuxInfoType(0), AuxInfoTypeParameter(0), DefaultSampleInfoSize(0)
+		{
+		}
+
+		virtual ~FMP4BoxSAIZ()
+		{
+		}
+
+	private:
+		FMP4BoxSAIZ() = delete;
+		FMP4BoxSAIZ(const FMP4BoxSAIZ&) = delete;
+
+	protected:
+		virtual UEMediaError ReadAndParseAttributes(FMP4ParseInfo* ParseInfo) override
+		{
+			UEMediaError Error = UEMEDIA_ERROR_OK;
+			RETURN_IF_ERROR(FMP4BoxFull::ReadAndParseAttributes(ParseInfo));
+
+			if ((Flags & 1) != 0)
+			{
+				RETURN_IF_ERROR(ParseInfo->Reader()->Read(AuxInfoType));						// aux_info_type
+				RETURN_IF_ERROR(ParseInfo->Reader()->Read(AuxInfoTypeParameter));				// aux_info_type_parameter
+			}
+
+			RETURN_IF_ERROR(ParseInfo->Reader()->Read(DefaultSampleInfoSize));					// default_sample_info_size
+			uint32 SampleCount = 0;
+			RETURN_IF_ERROR(ParseInfo->Reader()->Read(SampleCount));							// sample_count
+			if (DefaultSampleInfoSize == 0 && SampleCount != 0)
+			{
+				SampleInfoSize.AddUninitialized(SampleCount);
+				RETURN_IF_ERROR(ParseInfo->Reader()->ReadBytes(SampleInfoSize.GetData(), SampleCount));// sample_info_size
+			}
+			return Error;
+		}
+
+	private:
+		uint32	AuxInfoType;
+		uint32	AuxInfoTypeParameter;
+		uint8	DefaultSampleInfoSize;
+		TArray<uint8> SampleInfoSize;
+	};
+
+
+	/**
+	 * 'saio' box. ISO/IEC 14496-12:2015 - 8.7.9 Sample Auxiliary Information Offsets Box
+	 */
+	class FMP4BoxSAIO : public FMP4BoxFull
+	{
+	public:
+		FMP4BoxSAIO(IParserISO14496_12::FBoxType InBoxType, int64 InBoxSize, int64 InStartOffset, int64 InDataOffset, bool bInIsLeafBox)
+			: FMP4BoxFull(InBoxType, InBoxSize, InStartOffset, InDataOffset, bInIsLeafBox)
+			, AuxInfoType(0), AuxInfoTypeParameter(0), DefaultSampleInfoSize(0)
+		{
+		}
+
+		virtual ~FMP4BoxSAIO()
+		{
+		}
+
+	private:
+		FMP4BoxSAIO() = delete;
+		FMP4BoxSAIO(const FMP4BoxSAIO&) = delete;
+
+	protected:
+		virtual UEMediaError ReadAndParseAttributes(FMP4ParseInfo* ParseInfo) override
+		{
+			UEMediaError Error = UEMEDIA_ERROR_OK;
+			RETURN_IF_ERROR(FMP4BoxFull::ReadAndParseAttributes(ParseInfo));
+
+			if ((Flags & 1) != 0)
+			{
+				RETURN_IF_ERROR(ParseInfo->Reader()->Read(AuxInfoType));						// aux_info_type
+				RETURN_IF_ERROR(ParseInfo->Reader()->Read(AuxInfoTypeParameter));				// aux_info_type_parameter
+			}
+			uint32 EntryCount = 0;
+			RETURN_IF_ERROR(ParseInfo->Reader()->Read(EntryCount));								// entry_count
+			if (EntryCount)
+			{
+				Offsets.AddUninitialized(EntryCount);
+				if (Version == 0)
+				{
+					uint32 Value32 = 0;
+					uint64* Offs = Offsets.GetData();
+					for(uint32 i=0; i<EntryCount; ++i)
+					{
+						RETURN_IF_ERROR(ParseInfo->Reader()->Read(Value32));					// offset
+						*Offs++ = Value32;
+					}
+				}
+				else
+				{
+					RETURN_IF_ERROR(ParseInfo->Reader()->ReadBytes(Offsets.GetData(), EntryCount));// offset
+				}
+			}
+			return Error;
+		}
+
+	private:
+		uint32	AuxInfoType;
+		uint32	AuxInfoTypeParameter;
+		uint8	DefaultSampleInfoSize;
+		TArray<uint64> Offsets;
+	};
+
 
 
 
@@ -3257,6 +3774,7 @@ namespace Electra
 						case FMP4Box::kBox_dinf:
 						case FMP4Box::kBox_stbl:
 						case FMP4Box::kBox_mvex:
+						case FMP4Box::kBox_schi:
 							NextBox = new FMP4BoxBasic(BoxType, BoxSize, BoxStartOffset, BoxDataOffset, false);
 							break;
 						case FMP4Box::kBox_trak:
@@ -3328,6 +3846,15 @@ namespace Electra
 						case FMP4Box::kBox_pasp:
 							NextBox = new FMP4BoxPASP(BoxType, BoxSize, BoxStartOffset, BoxDataOffset, true);
 							break;
+						case FMP4Box::kBox_elng:
+							NextBox = new FMP4BoxELNG(BoxType, BoxSize, BoxStartOffset, BoxDataOffset, true);
+							break;
+						case FMP4Box::kBox_prft:
+							NextBox = new FMP4BoxPRFT(BoxType, BoxSize, BoxStartOffset, BoxDataOffset, true);
+							break;
+						case FMP4Box::kBox_emsg:
+							NextBox = new FMP4BoxEMSG(BoxType, BoxSize, BoxStartOffset, BoxDataOffset, true);
+							break;
 						case FMP4Box::kBox_iods:
 						case FMP4Box::kBox_esds:
 							NextBox = new FMP4BoxESDS(BoxType, BoxSize, BoxStartOffset, BoxDataOffset, true);
@@ -3369,9 +3896,18 @@ namespace Electra
 						case FMP4Box::kBox_trun:
 							NextBox = new FMP4BoxTRUN(BoxType, BoxSize, BoxStartOffset, BoxDataOffset, true);
 							break;
-						case FMP4Box::kBox_tref:	// We ignore the 'tref' box for now.
-						case FMP4Box::kBox_sdtp:	// We ignore the 'sdtp' box for now.
-						case FMP4Box::kBox_gmhd:	// We ignore the 'gmhd' box for now.
+						case FMP4Box::kBox_saiz:
+							NextBox = new FMP4BoxSAIZ(BoxType, BoxSize, BoxStartOffset, BoxDataOffset, true);
+							break;
+						case FMP4Box::kBox_saio:
+							NextBox = new FMP4BoxSAIO(BoxType, BoxSize, BoxStartOffset, BoxDataOffset, true);
+							break;
+						// Boxes we ignore for now.
+						case FMP4Box::kBox_tref:
+						case FMP4Box::kBox_colr:
+						case FMP4Box::kBox_trep:
+						case FMP4Box::kBox_sdtp:
+						case FMP4Box::kBox_gmhd:
 							NextBox = new FMP4BoxIgnored(BoxType, BoxSize, BoxStartOffset, BoxDataOffset, true);
 							break;
 						case FMP4Box::kBox_skip:
@@ -3390,6 +3926,23 @@ namespace Electra
 							NextBox = new FMP4BoxIgnored(BoxType, BoxSize, BoxStartOffset, BoxDataOffset, true);
 							break;
 
+						// Encryption related boxes.
+						case FMP4Box::kBox_pssh:
+							NextBox = new FMP4BoxPSSH(BoxType, BoxSize, BoxStartOffset, BoxDataOffset, true);
+							break;
+						case FMP4Box::kBox_sinf:
+							NextBox = new FMP4BoxSINF(BoxType, BoxSize, BoxStartOffset, BoxDataOffset, true);		// 'sinf' is not a true leaf box, but its elements are a list that must be parsed.
+							break;
+						case FMP4Box::kBox_frma:
+							NextBox = new FMP4BoxFRMA(BoxType, BoxSize, BoxStartOffset, BoxDataOffset, true);
+							break;
+						case FMP4Box::kBox_schm:
+							NextBox = new FMP4BoxSCHM(BoxType, BoxSize, BoxStartOffset, BoxDataOffset, true);
+							break;
+						case FMP4Box::kBox_tenc:
+							NextBox = new FMP4BoxTENC(BoxType, BoxSize, BoxStartOffset, BoxDataOffset, true);
+							break;
+
 						default:
 #if MEDIA_DEBUG_HAS_BOX_NAMES
 							UE_LOG(LogElectraMP4Parser, Log, TEXT("Ignoring mp4 box: '%c%c%c%c' %lld B @ %lld"), boxName[0]?boxName[0]:'?', boxName[1]?boxName[1]:'?', boxName[2]?boxName[2]:'?', boxName[3]?boxName[3]:'?', (long long int)BoxSize, (long long int)BoxStartOffset);
diff --git a/Engine/Plugins/Media/ElectraPlayer/Source/ElectraPlayerRuntime/Private/Runtime/Player/AdaptiveStreamingPlayer.cpp b/Engine/Plugins/Media/ElectraPlayer/Source/ElectraPlayerRuntime/Private/Runtime/Player/AdaptiveStreamingPlayer.cpp
--- a/Engine/Plugins/Media/ElectraPlayer/Source/ElectraPlayerRuntime/Private/Runtime/Player/AdaptiveStreamingPlayer.cpp
+++ b/Engine/Plugins/Media/ElectraPlayer/Source/ElectraPlayerRuntime/Private/Runtime/Player/AdaptiveStreamingPlayer.cpp
@@ -84,6 +84,7 @@ FAdaptiveStreamingPlayer::FAdaptiveStreamingPlayer(const IAdaptiveStreamingPlaye
 
 	CurrentLoopState.Reset();
 	NextLoopStates.Clear();
+	bFirstSegmentRequestIsForLooping = false;
 
 	BitrateCeiling = 0;
 	VideoResolutionLimitWidth = 0;
@@ -496,25 +497,25 @@ void FAdaptiveStreamingPlayer::GetLoopState(FPlayerLoopState& OutLoopState) cons
 
 //-----------------------------------------------------------------------------
 /**
- * Returns stream metadata of the currently active play period.
+ * Returns track metadata of the currently active play period.
  *
- * @param OutStreamMetadata
+ * @param OutTrackMetadata
  * @param StreamType
  */
-void FAdaptiveStreamingPlayer::GetStreamMetadata(TArray<FStreamMetadata>& OutStreamMetadata, EStreamType StreamType) const
+void FAdaptiveStreamingPlayer::GetTrackMetadata(TArray<FTrackMetadata>& OutTrackMetadata, EStreamType StreamType) const
 {
-	TArray<FStreamMetadata> dummy;
+	TArray<FTrackMetadata> dummy;
 	if (StreamType == EStreamType::Video)
 	{
-		PlaybackState.GetStreamMetadata(OutStreamMetadata, dummy);
+		PlaybackState.GetTrackMetadata(OutTrackMetadata, dummy);
 	}
 	else if (StreamType == EStreamType::Audio)
 	{
-		PlaybackState.GetStreamMetadata(dummy, OutStreamMetadata);
+		PlaybackState.GetTrackMetadata(dummy, OutTrackMetadata);
 	}
 	else
 	{
-		OutStreamMetadata.Empty();
+		OutTrackMetadata.Empty();
 	}
 }
 
@@ -550,11 +551,11 @@ void FAdaptiveStreamingPlayer::SetMaxResolution(int32 MaxWidth, int32 MaxHeight)
  * Selects a track based from one of the array members returned by GetStreamMetadata().
  *
  * @param StreamType
- * @param StreamMetadata
+ * @param TrackMetadata
  */
-void FAdaptiveStreamingPlayer::SelectTrackByMetadata(EStreamType StreamType, const FStreamMetadata& StreamMetadata)
+void FAdaptiveStreamingPlayer::SelectTrackByMetadata(EStreamType StreamType, const FTrackMetadata& TrackMetadata)
 {
-	WorkerThread.SendTrackSelectByMetadataMessage(StreamType, StreamMetadata);
+	WorkerThread.SendTrackSelectByMetadataMessage(StreamType, TrackMetadata);
 }
 
 //-----------------------------------------------------------------------------
@@ -1173,22 +1174,22 @@ void FAdaptiveStreamingPlayer::WorkerThreadFN()
 					switch(msg.Data.TrackSelection.StreamType)
 					{
 						case EStreamType::Video:
-							InitialStreamSelectionVid = MakeShared<FStreamMetadata, ESPMode::ThreadSafe>(msg.Data.TrackSelection.StreamMetadata);
+							InitialTrackSelectionVid = MakeShared<FTrackMetadata, ESPMode::ThreadSafe>(msg.Data.TrackSelection.TrackMetadata);
 							break;
 						case EStreamType::Audio:
-							InitialStreamSelectionAud = MakeShared<FStreamMetadata, ESPMode::ThreadSafe>(msg.Data.TrackSelection.StreamMetadata);
+							InitialTrackSelectionAud = MakeShared<FTrackMetadata, ESPMode::ThreadSafe>(msg.Data.TrackSelection.TrackMetadata);
 							// For media with multiplexed tracks we can switch to a different track pretty much immediately.
 							if (ManifestReader.IsValid() && ManifestReader->GetPlaylistType() == TEXT("mp4"))
 							{
 								// Only select audio for now.
-								if (InitialStreamSelectionAud.IsValid())
+								if (InitialTrackSelectionAud.IsValid())
 								{
-									MultiStreamBufferAud.SelectTrackByID(InitialStreamSelectionAud->StreamUniqueID);
+									MultiStreamBufferAud.SelectTrackByID(InitialTrackSelectionAud->TrackID);
 								}
 							}
 							break;
 						case EStreamType::Subtitle:
-							InitialStreamSelectionTxt = MakeShared<FStreamMetadata, ESPMode::ThreadSafe>(msg.Data.TrackSelection.StreamMetadata);
+							InitialTrackSelectionTxt = MakeShared<FTrackMetadata, ESPMode::ThreadSafe>(msg.Data.TrackSelection.TrackMetadata);
 							break;
 					}
 					break;
@@ -1199,10 +1200,10 @@ void FAdaptiveStreamingPlayer::WorkerThreadFN()
 					switch(msg.Data.TrackSelection.StreamType)
 					{
 						case EStreamType::Video:
-							InitialStreamSelectionVid.Reset();
+							InitialTrackSelectionVid.Reset();
 							break;
 						case EStreamType::Audio:
-							InitialStreamSelectionAud.Reset();
+							InitialTrackSelectionAud.Reset();
 							// For media with multiplexed tracks we can switch to a different track pretty much immediately.
 							if (ManifestReader.IsValid() && ManifestReader->GetPlaylistType() == TEXT("mp4"))
 							{
@@ -1210,7 +1211,7 @@ void FAdaptiveStreamingPlayer::WorkerThreadFN()
 							}
 							break;
 						case EStreamType::Subtitle:
-							InitialStreamSelectionTxt.Reset();
+							InitialTrackSelectionTxt.Reset();
 							break;
 					}
 					break;
@@ -1395,14 +1396,49 @@ void FAdaptiveStreamingPlayer::HandleMetadataChanges()
 		PlaybackState.SetSeekableRange(Manifest->GetSeekableTimeRange());
 		PlaybackState.SetTimelineRange(Manifest->GetTotalTimeRange());
 		PlaybackState.SetDuration(Manifest->GetDuration());
-/*
-		TArray<FStreamMetadata> VideoStreamMetadata;
-		TArray<FStreamMetadata> AudioStreamMetadata;
-		pPresentation->GetStreamMetadata(VideoStreamMetadata, EStreamType::Video);
-		pPresentation->GetStreamMetadata(AudioStreamMetadata, EStreamType::Audio);
-		PlaybackState.SetStreamMetadata(VideoStreamMetadata, AudioStreamMetadata);
-		PlaybackState.SetHaveMetadata(true);
-*/
+	}
+
+	// See if we may be reaching the next period.
+	if (UpcomingPeriods.Num())
+	{
+		FTimeValue Current = GetCurrentPlayTime();
+		if (Current.IsValid())
+		{
+			for(int32 i=0; i<UpcomingPeriods.Num(); )
+			{
+				if (Current < UpcomingPeriods[i].TimeRange.Start)
+				{
+					break;
+				}
+				else
+				{
+					// Did we already pass the end of the period?
+					// This could be happening with really short periods and in this case we do not want to
+					// handle those. Since we passed them already there is no point.
+					if (!UpcomingPeriods[i].TimeRange.End.IsValid() || Current < UpcomingPeriods[i].TimeRange.End)
+					{
+						// Handle it.
+						TArray<FTrackMetadata> MetadataVideo;
+						TArray<FTrackMetadata> MetadataAudio;
+						//TArray<FTrackMetadata> MetadataSubtitle;
+						UpcomingPeriods[i].Period->GetMetaData(MetadataVideo, EStreamType::Video);
+						UpcomingPeriods[i].Period->GetMetaData(MetadataAudio, EStreamType::Audio);
+						//UpcomingPeriods[i].Period->GetMetaData(MetadataSubtitle, EStreamType::Subtitle);
+
+						bool bMetadataChanged = PlaybackState.SetTrackMetadata(MetadataVideo, MetadataAudio);
+						PlaybackState.SetHaveMetadata(true);
+						if (bMetadataChanged)
+						{
+							DispatchEvent(FMetricEvent::ReportTracksChanged());
+						}
+					}
+					
+					// Remove it and continue with the next.
+					UpcomingPeriods.RemoveAt(i);
+					++i;
+				}
+			}
+		}
 	}
 }
 
@@ -1622,6 +1658,33 @@ void FAdaptiveStreamingPlayer::HandleDecoderChanges()
 	// Nothing to do for now.
 }
 
+
+//-----------------------------------------------------------------------------
+/**
+ * Returns the current internal playback time.
+ * This is not the same as GetPlayPosition() in that this time here must not
+ * continue to run if the render clock is not stopped.
+ * This time should be fairly accurate but is not currently required to be.
+ * Events and period transitions are triggered off of this.
+ */
+FTimeValue FAdaptiveStreamingPlayer::GetCurrentPlayTime()
+{
+	FTimeValue Current = PlaybackState.GetPlayPosition();
+	FTimeValue v;
+	v = MultiStreamBufferVid.GetLastPoppedPTS();
+	if (v.IsValid() && v < Current)
+	{
+		Current = v;
+	}
+	v = MultiStreamBufferAud.GetLastPoppedPTS();
+	if (v.IsValid() && v < Current)
+	{
+		Current = v;
+	}
+	return Current;
+}
+
+
 //-----------------------------------------------------------------------------
 /**
  * Handle data in deselected buffers.
@@ -1798,7 +1861,7 @@ void FAdaptiveStreamingPlayer::HandlePendingMediaSegmentRequests()
 			}
 
 			// Do we have a play period yet?
-			if (CurrentPlayPeriod.Get() == nullptr)
+			if (!InitialPlayPeriod.IsValid())
 			{
 				// No. Check if we should ask for one now.
 				if (!PendingStartRequest->RetryAtTime.IsValid() || CurrentTime >= PendingStartRequest->RetryAtTime)
@@ -1824,8 +1887,7 @@ void FAdaptiveStreamingPlayer::HandlePendingMediaSegmentRequests()
 						}
 					}
 
-					TSharedPtrTS<IManifest::IPlayPeriod> StartingPeriod;
-					IManifest::FResult Result = Manifest->FindPlayPeriod(StartingPeriod, PendingStartRequest->StartAt, PendingStartRequest->SearchType);
+					IManifest::FResult Result = Manifest->FindPlayPeriod(InitialPlayPeriod, PendingStartRequest->StartAt, PendingStartRequest->SearchType);
 					switch(Result.GetType())
 					{
 						case IManifest::FResult::EType::TryAgainLater:
@@ -1835,22 +1897,25 @@ void FAdaptiveStreamingPlayer::HandlePendingMediaSegmentRequests()
 						}
 						case IManifest::FResult::EType::Found:
 						{
-							check(StartingPeriod.IsValid());
-							if (StartingPeriod.IsValid())
-							{
-								StartingPeriod->SetStreamPreferences(StreamPreferences);
-							}
-							CurrentPlayPeriod = StartingPeriod;
 							// NOTE: Do *not* reset the pending start request yet. It is still needed.
 							break;
 						}
+						case IManifest::FResult::EType::PastEOS:
+						{
+							PendingStartRequest.Reset();
+							InternalSetPlaybackEnded();
+							break;
+						}
 						case IManifest::FResult::EType::NotFound:
 						case IManifest::FResult::EType::BeforeStart:
-						case IManifest::FResult::EType::PastEOS:
 						case IManifest::FResult::EType::NotLoaded:
 						{
 							// Throw a playback error for now.
-							check(!"TODO");
+							FErrorDetail err;
+							err.SetFacility(Facility::EFacility::Player);
+							err.SetMessage("Could not locate the playback period to begin playback at.");
+							err.SetCode(INTERR_COULD_NOT_LOCATE_START_PERIOD);
+							PostError(err);
 							PendingStartRequest.Reset();
 							break;
 						}
@@ -1858,17 +1923,17 @@ void FAdaptiveStreamingPlayer::HandlePendingMediaSegmentRequests()
 				}
 			}
 			// Do we have a starting play period now?
-			if (CurrentPlayPeriod.IsValid())
+			if (InitialPlayPeriod.IsValid())
 			{
 				// Is it ready yet?
-				switch(CurrentPlayPeriod->GetReadyState())
+				switch(InitialPlayPeriod->GetReadyState())
 				{
 					case IManifest::IPlayPeriod::EReadyState::NotReady:
 					{
 						// Prepare the play period. This must immediately change the state away from NotReady
-						FParamDict Options;
-						CurrentPlayPeriod->PrepareForPlay(Options);
-						check(CurrentPlayPeriod->GetReadyState() != IManifest::IPlayPeriod::EReadyState::NotReady);
+						InitialPlayPeriod->SetStreamPreferences(StreamPreferences);
+						InitialPlayPeriod->PrepareForPlay();
+						check(InitialPlayPeriod->GetReadyState() != IManifest::IPlayPeriod::EReadyState::NotReady);
 						break;
 					}
 					case IManifest::IPlayPeriod::EReadyState::IsReady:
@@ -1876,10 +1941,19 @@ void FAdaptiveStreamingPlayer::HandlePendingMediaSegmentRequests()
 						// With the period being ready we can now get the initial media segment request.
 						check(PendingStartRequest.IsValid());	// must not have been released yet
 
+						// Now that we have a valid first-ever play start established we clear out the internal
+						// default start time so we can seek backwards/forward from it to a different time.
+						Manifest->ClearDefaultStartTime();
+
 					// TODO: this should update the stream metadata.
 
-						// Tell the ABR the current playback period.
-						StreamSelector->SetCurrentPlaybackPeriod(CurrentPlayPeriod);
+						// At playback start all streams begin in the same period.
+						CurrentPlayPeriodVideo = InitialPlayPeriod;
+						CurrentPlayPeriodAudio = InitialPlayPeriod;
+						CurrentPlayPeriodText = InitialPlayPeriod;
+						// Tell the ABR the current playback period. At playback start it is the same for all stream types.
+						StreamSelector->SetCurrentPlaybackPeriod(EStreamType::Video, CurrentPlayPeriodVideo);
+						StreamSelector->SetCurrentPlaybackPeriod(EStreamType::Audio, CurrentPlayPeriodAudio);
 
 						// Single stream mp4 files cannot switch streams. Failures need to retry the same file.
 						if (ManifestReader.IsValid() && ManifestReader->GetPlaylistType() == TEXT("mp4"))
@@ -1906,9 +1980,7 @@ void FAdaptiveStreamingPlayer::HandlePendingMediaSegmentRequests()
 						// Have the ABR select the initial bandwidth. Pass it an empty segment request to indicate this.
 						FTimeValue ActionDelay;
 						IAdaptiveStreamSelector::ESegmentAction Action = StreamSelector->SelectSuitableStreams(ActionDelay, TSharedPtrTS<IStreamSegment>());
-// FIXME: Can the starting choice with no previous segment even be anything besides failure or fetch?
 						check(Action == IAdaptiveStreamSelector::ESegmentAction::FetchNext || Action == IAdaptiveStreamSelector::ESegmentAction::Fail);
-						//if (Action == IAdaptiveStreamSelector::ESegmentAction::Fail)
 						if (Action != IAdaptiveStreamSelector::ESegmentAction::FetchNext)
 						{
 							FErrorDetail err;
@@ -1923,7 +1995,17 @@ void FAdaptiveStreamingPlayer::HandlePendingMediaSegmentRequests()
 						if (!PendingStartRequest->RetryAtTime.IsValid() || CurrentTime >= PendingStartRequest->RetryAtTime)
 						{
 							TSharedPtrTS<IStreamSegment> FirstSegmentRequest;
-							IManifest::FResult Result = CurrentPlayPeriod->GetStartingSegment(FirstSegmentRequest, PendingStartRequest->StartAt, PendingStartRequest->SearchType);
+							IManifest::FResult Result;
+							if (!PendingStartRequest->bForLooping)
+							{
+								Result = InitialPlayPeriod->GetStartingSegment(FirstSegmentRequest, PendingStartRequest->StartAt, PendingStartRequest->SearchType);
+								bFirstSegmentRequestIsForLooping = false;
+							}
+							else
+							{
+								Result = InitialPlayPeriod->GetLoopingSegment(FirstSegmentRequest, CurrentLoopState, PendingStartRequest->FinishedRequests, PendingStartRequest->StartAt, PendingStartRequest->SearchType);
+								bFirstSegmentRequestIsForLooping = true;
+							}
 							switch(Result.GetType())
 							{
 								case IManifest::FResult::EType::TryAgainLater:
@@ -1933,9 +2015,15 @@ void FAdaptiveStreamingPlayer::HandlePendingMediaSegmentRequests()
 								}
 								case IManifest::FResult::EType::Found:
 								{
+									// Add this period to the list of upcoming ones. This ensures the period metadata gets reported
+									// if the period playback starts in is not the first for which metadata was initially reported.
+									AddUpcomingPeriod(InitialPlayPeriod);
+
 									PendingFirstSegmentRequest = FirstSegmentRequest;
 									// No longer need the start request.
 									PendingStartRequest.Reset();
+									// Also no longer need the initial play period.
+									InitialPlayPeriod.Reset();
 									break;
 								}
 								case IManifest::FResult::EType::NotFound:
@@ -1943,7 +2031,7 @@ void FAdaptiveStreamingPlayer::HandlePendingMediaSegmentRequests()
 								{
 									// Reset the current play period and start over with the initial period selection.
 									PostLog(Facility::EFacility::Player, IInfoLog::ELevel::Info, FString::Printf(TEXT("Period start segments not found. Reselecting start period.")));
-									CurrentPlayPeriod.Reset();
+									InitialPlayPeriod.Reset();
 									break;
 								}
 								case IManifest::FResult::EType::BeforeStart:
@@ -1955,6 +2043,7 @@ void FAdaptiveStreamingPlayer::HandlePendingMediaSegmentRequests()
 									err.SetCode(INTERR_COULD_NOT_LOCATE_START_SEGMENT);
 									PostError(err);
 									PendingStartRequest.Reset();
+									InitialPlayPeriod.Reset();
 									break;
 								}
 							}
@@ -1964,13 +2053,6 @@ void FAdaptiveStreamingPlayer::HandlePendingMediaSegmentRequests()
 					case IManifest::IPlayPeriod::EReadyState::Preparing:
 					{
 						// Period is not yet ready. Check again on the next iteration.
-					// TODO: get a try-again-at value instead of checking each loop?
-						break;
-					}
-					case IManifest::IPlayPeriod::EReadyState::MustReselect:
-					{
-						// The period has asked to be discarded.
-						CurrentPlayPeriod.Reset();
 						break;
 					}
 				}
@@ -1989,41 +2071,74 @@ void FAdaptiveStreamingPlayer::HandlePendingMediaSegmentRequests()
 		check(StreamReaderHandler);
 		if (StreamReaderHandler)
 		{
-			FTimeValue CurrentPlayPos = PlaybackState.GetPlayPosition();
-			FTimeValue NewPlayPos = PendingFirstSegmentRequest->GetFirstPTS();
-			PlaybackState.SetPlayPosition(NewPlayPos);
-			PlaybackState.SetLoopState(CurrentLoopState);
-
-			check(CurrentState == EPlayerState::eState_Buffering || CurrentState == EPlayerState::eState_Rebuffering || CurrentState == EPlayerState::eState_Seeking);
-			if (CurrentState == EPlayerState::eState_Seeking)
+			if (!bFirstSegmentRequestIsForLooping)
 			{
-				DispatchEvent(FMetricEvent::ReportJumpInPlayPosition(NewPlayPos, CurrentPlayPos, Metrics::ETimeJumpReason::UserSeek));
-			}
+				FTimeValue CurrentPlayPos = PlaybackState.GetPlayPosition();
+				FTimeValue NewPlayPos = PendingFirstSegmentRequest->GetFirstPTS();
+				PlaybackState.SetPlayPosition(NewPlayPos);
+				PlaybackState.SetLoopState(CurrentLoopState);
 
-			// Get the expected stream types from the starting segment request.
-		// FIXME: get rid of the "bHave...XYZ...Reader" states to become more flexible with dynamic stream selection (enabling/disabling) during playback.
-			bHaveVideoReader.Set(false);
-			bHaveAudioReader.Set(false);
-			bHaveTextReader.Set(false);
-			switch(PendingFirstSegmentRequest->GetType())
-			{
-				case EStreamType::Video:	bHaveVideoReader.Set(true);	break;
-				case EStreamType::Audio:	bHaveAudioReader.Set(true);	break;
-				case EStreamType::Subtitle:	bHaveTextReader.Set(true);		break;
-				default:													break;
-			}
-			TArray<IStreamSegment::FDependentStreams> DependentStreams;
-			PendingFirstSegmentRequest->GetDependentStreams(DependentStreams);
-			for(int32 i=0; i<DependentStreams.Num(); ++i)
-			{
-				switch(DependentStreams[i].StreamType)
+				check(CurrentState == EPlayerState::eState_Buffering || CurrentState == EPlayerState::eState_Rebuffering || CurrentState == EPlayerState::eState_Seeking);
+				if (CurrentState == EPlayerState::eState_Seeking)
+				{
+					DispatchEvent(FMetricEvent::ReportJumpInPlayPosition(NewPlayPos, CurrentPlayPos, Metrics::ETimeJumpReason::UserSeek));
+				}
+
+				// Get the expected stream types from the starting segment request.
+			// FIXME: get rid of the "bHave...XYZ...Reader" states to become more flexible with dynamic stream selection (enabling/disabling) during playback.
+				bHaveVideoReader.Set(false);
+				bHaveAudioReader.Set(false);
+				bHaveTextReader.Set(false);
+				switch(PendingFirstSegmentRequest->GetType())
 				{
 					case EStreamType::Video:	bHaveVideoReader.Set(true);	break;
 					case EStreamType::Audio:	bHaveAudioReader.Set(true);	break;
-					case EStreamType::Subtitle:	bHaveTextReader.Set(true);		break;
-					default:													break;
+					case EStreamType::Subtitle:	bHaveTextReader.Set(true);	break;
+					default:												break;
+				}
+				TArray<IStreamSegment::FDependentStreams> DependentStreams;
+				PendingFirstSegmentRequest->GetDependentStreams(DependentStreams);
+				for(int32 i=0; i<DependentStreams.Num(); ++i)
+				{
+					switch(DependentStreams[i].StreamType)
+					{
+						case EStreamType::Video:	bHaveVideoReader.Set(true);	break;
+						case EStreamType::Audio:	bHaveAudioReader.Set(true);	break;
+						case EStreamType::Subtitle:	bHaveTextReader.Set(true);	break;
+						default:												break;
+					}
 				}
+
+				if (AudioRender.Renderer)
+				{
+					RenderClock->SetCurrentTime(IMediaRenderClock::ERendererType::Audio, NewPlayPos);
+					AudioRender.Renderer->SetNextApproximatePresentationTime(NewPlayPos);
+				}
+				if (VideoRender.Renderer)
+				{
+					RenderClock->SetCurrentTime(IMediaRenderClock::ERendererType::Video, NewPlayPos);
+					VideoRender.Renderer->SetNextApproximatePresentationTime(NewPlayPos);
+				}
+
+				PlaybackState.SetIsBuffering(true);
+				PlaybackState.SetIsSeeking(true);
+
+				// Send QoS buffering event
+				DispatchBufferingEvent(true, CurrentState);
+
+				// Remember why we are buffering so we can send the proper QoS event when buffering is done.
+				LastBufferingState = CurrentState;
+
+				// Whether we were seeking or rebuffering, we're now just buffering.
+				CurrentState = EPlayerState::eState_Buffering;
 			}
+			else
+			{
+				PostLog(Facility::EFacility::Player, IInfoLog::ELevel::Verbose, FString::Printf(TEXT("Enqueuing stream loop start request.")));
+				// Add the updated loop state to the upcoming loop FIFO. When the play position reaches the next enqueued loop state it gets popped.
+				NextLoopStates.Push(CurrentLoopState);
+			}
+
 			// Get the stream types that have already reached EOS in this start request because the streams are of different duration.
 			// The streams are there in principle so we need to set them up (in case we loop back into them) but we won't see any data arriving.
 			TArray<TSharedPtrTS<IStreamSegment>> AlreadyEndedStreams;
@@ -2033,32 +2148,16 @@ void FAdaptiveStreamingPlayer::HandlePendingMediaSegmentRequests()
 				CompletedSegmentRequests.Add(AlreadyEndedStreams[i]->GetType(), AlreadyEndedStreams[i]);
 			}
 
-			if (AudioRender.Renderer)
-			{
-				RenderClock->SetCurrentTime(IMediaRenderClock::ERendererType::Audio, NewPlayPos);
-				AudioRender.Renderer->SetNextApproximatePresentationTime(NewPlayPos);
-			}
-			if (VideoRender.Renderer)
+			// Kick off the request
+			IStreamReader::EAddResult ReqResult = StreamReaderHandler->AddRequest(CurrentPlaybackSequenceID, PendingFirstSegmentRequest);
+			if (ReqResult != IStreamReader::EAddResult::Added)
 			{
-				RenderClock->SetCurrentTime(IMediaRenderClock::ERendererType::Video, NewPlayPos);
-				VideoRender.Renderer->SetNextApproximatePresentationTime(NewPlayPos);
+				FErrorDetail err;
+				err.SetFacility(Facility::EFacility::Player);
+				err.SetMessage("Failed to add stream segment request to reader");
+				err.SetCode(INTERR_FRAGMENT_READER_REQUEST);
+				PostError(err);
 			}
-
-			PlaybackState.SetIsBuffering(true);
-			PlaybackState.SetIsSeeking(true);
-
-			// Send QoS buffering event
-			DispatchBufferingEvent(true, CurrentState);
-
-			// Remember why we are buffering so we can send the proper QoS event when buffering is done.
-			LastBufferingState = CurrentState;
-
-			// Whether we were seeking or rebuffering, we're now just buffering.
-			CurrentState = EPlayerState::eState_Buffering;
-
-			// Kick off the request
-//// TODO: check return code!!!
-			StreamReaderHandler->AddRequest(CurrentPlaybackSequenceID, PendingFirstSegmentRequest);
 		}
 		PendingFirstSegmentRequest.Reset();
 	}
@@ -2072,34 +2171,85 @@ void FAdaptiveStreamingPlayer::HandlePendingMediaSegmentRequests()
 		FPendingSegmentRequest FinishedReq;
 		pendingRequests.Dequeue(FinishedReq);
 
+		TSharedPtrTS<IManifest::IPlayPeriod> SegmentPlayPeriod;
+		if (FinishedReq.Request.IsValid())
+		{
+			switch(FinishedReq.Request->GetType())
+			{
+				default:
+				{
+					break;
+				}
+				case EStreamType::Video:
+				 {
+					SegmentPlayPeriod = CurrentPlayPeriodVideo;
+					break;
+				}
+				case EStreamType::Audio:
+				{
+					SegmentPlayPeriod = CurrentPlayPeriodAudio;
+					break;
+				}
+				case EStreamType::Subtitle:
+				{
+					SegmentPlayPeriod = CurrentPlayPeriodText;
+					break;
+				}
+			}
+		}
+
 		// If not yet ready to check again put the request back into the list.
 		if (FinishedReq.AtTime.IsValid() && FinishedReq.AtTime > CurrentTime)
 		{
 			NextPendingSegmentRequests.Enqueue(MoveTemp(FinishedReq));
 		}
-		else if (CurrentPlayPeriod.IsValid())
+		else if (SegmentPlayPeriod.IsValid())
 		{
-			// Evaluate ABR to select the next stream quality.
-			IAdaptiveStreamSelector::ESegmentAction Action;
-			if (FinishedReq.PreviousEvalResult.IsSet())
+			// Did we move into a new period that needs to be made ready?
+			if (FinishedReq.Period.IsValid())
 			{
-				Action = FinishedReq.PreviousEvalResult.Value();
-				FinishedReq.PreviousEvalResult.Reset();
-			}
-			else
-			{
-				FTimeValue ActionDelay(FTimeValue::GetZero());
-				Action = StreamSelector->SelectSuitableStreams(ActionDelay, FinishedReq.Request);
-				// If the selected choice is to be executed after a certain time take note of
-				// the choice and add the request back to the pending queue.
-				if (ActionDelay > FTimeValue::GetZero())
+				IManifest::IPlayPeriod::EReadyState PeriodState = FinishedReq.Period->GetReadyState();
+				if (PeriodState == IManifest::IPlayPeriod::EReadyState::NotReady)
 				{
-					FinishedReq.AtTime = CurrentTime + ActionDelay;
-					FinishedReq.PreviousEvalResult.Set(Action);
+					FinishedReq.Period->SetStreamPreferences(StreamPreferences);
+					FinishedReq.Period->PrepareForPlay();
 					NextPendingSegmentRequests.Enqueue(MoveTemp(FinishedReq));
 					continue;
 				}
+				else if (PeriodState == IManifest::IPlayPeriod::EReadyState::Preparing)
+				{
+					NextPendingSegmentRequests.Enqueue(MoveTemp(FinishedReq));
+					continue;
+				}
+				else
+				{
+					if (FinishedReq.Request->GetType() == EStreamType::Video)
+					{
+						SegmentPlayPeriod = FinishedReq.Period;
+						CurrentPlayPeriodVideo = FinishedReq.Period;
+						StreamSelector->SetCurrentPlaybackPeriod(EStreamType::Video, CurrentPlayPeriodVideo);
+						UpdateStreamResolutionLimit();
+					}
+					else if (FinishedReq.Request->GetType() == EStreamType::Audio)
+					{
+						SegmentPlayPeriod = FinishedReq.Period;
+						CurrentPlayPeriodAudio = FinishedReq.Period;
+						StreamSelector->SetCurrentPlaybackPeriod(EStreamType::Audio, CurrentPlayPeriodAudio);
+					}
+					else if (FinishedReq.Request->GetType() == EStreamType::Subtitle)
+					{
+						SegmentPlayPeriod = FinishedReq.Period;
+						CurrentPlayPeriodText = FinishedReq.Period;
+						//StreamSelector->SetCurrentPlaybackPeriod(EStreamType::Subtitle, CurrentPlayPeriodText);
+					}
+					// Add this to the upcoming periods the play position will move into and metadata will need to be updated then.
+					AddUpcomingPeriod(SegmentPlayPeriod);
+				}
 			}
+
+			// Evaluate ABR to select the next stream quality.
+			FTimeValue ActionDelay(FTimeValue::GetZero());
+			IAdaptiveStreamSelector::ESegmentAction Action = StreamSelector->SelectSuitableStreams(ActionDelay, FinishedReq.Request);
 			if (Action == IAdaptiveStreamSelector::ESegmentAction::Fail)
 			{
 				FErrorDetail err;
@@ -2114,11 +2264,11 @@ void FAdaptiveStreamingPlayer::HandlePendingMediaSegmentRequests()
 			IManifest::FResult Result;
 			if (Action == IAdaptiveStreamSelector::ESegmentAction::FetchNext)
 			{
-				Result = CurrentPlayPeriod->GetNextSegment(pNext, FinishedReq.Request);
+				Result = SegmentPlayPeriod->GetNextSegment(pNext, FinishedReq.Request);
 			}
 			else if (Action == IAdaptiveStreamSelector::ESegmentAction::Retry || Action == IAdaptiveStreamSelector::ESegmentAction::Fill)
 			{
-				Result = CurrentPlayPeriod->GetRetrySegment(pNext, FinishedReq.Request, Action == IAdaptiveStreamSelector::ESegmentAction::Fill);
+				Result = SegmentPlayPeriod->GetRetrySegment(pNext, FinishedReq.Request, Action == IAdaptiveStreamSelector::ESegmentAction::Fill);
 			}
 			switch(Result.GetType())
 			{
@@ -2130,23 +2280,56 @@ void FAdaptiveStreamingPlayer::HandlePendingMediaSegmentRequests()
 				}
 				case IManifest::FResult::EType::Found:
 				{
-//// TODO: check return code!!!
-					StreamReaderHandler->AddRequest(CurrentPlaybackSequenceID, pNext);
+					pNext->SetExecutionDelay(ActionDelay);
+					IStreamReader::EAddResult ReqResult = StreamReaderHandler->AddRequest(CurrentPlaybackSequenceID, pNext);
+					if (ReqResult != IStreamReader::EAddResult::Added)
+					{
+						FErrorDetail err;
+						err.SetFacility(Facility::EFacility::Player);
+						err.SetMessage("Failed to add stream segment request to reader");
+						err.SetCode(INTERR_FRAGMENT_READER_REQUEST);
+						PostError(err);
+					}
 					break;
 				}
 				case IManifest::FResult::EType::PastEOS:
 				{
-// TODO: move to next period, if any.
-					// Get the dependent stream types, if any.
-					// This is mainly for multiplexed streams like an .mp4 at this point.
-					// All stream types have reached EOS now.
-					TArray<IStreamSegment::FDependentStreams> DependentStreams;
-					FinishedReq.Request->GetDependentStreams(DependentStreams);
-					// Add the primary stream type to the list as well.
-					DependentStreams.AddDefaulted_GetRef().StreamType = FinishedReq.Request->GetType();
-					for(int32 i=0; i<DependentStreams.Num(); ++i)
+					TSharedPtrTS<IManifest::IPlayPeriod> NextPeriod;
+					IManifest::FResult NextPeriodResult = Manifest->FindNextPlayPeriod(NextPeriod, FinishedReq.Request);
+					switch(NextPeriodResult.GetType())
 					{
-						CompletedSegmentRequests.Add(DependentStreams[i].StreamType, FinishedReq.Request);
+						case IManifest::FResult::EType::TryAgainLater:
+						{
+							// Next period is either not there yet or is being resolved right now.
+							FinishedReq.AtTime = Result.GetRetryAgainAtTime();
+							NextPendingSegmentRequests.Enqueue(MoveTemp(FinishedReq));
+							break;
+						}
+						case IManifest::FResult::EType::Found:
+						{
+							// Get the next period ready for playback. Put it back into the queue.
+							FinishedReq.Period = NextPeriod;
+							NextPendingSegmentRequests.Enqueue(MoveTemp(FinishedReq));
+							break;
+						}
+						case IManifest::FResult::EType::PastEOS:
+						case IManifest::FResult::EType::NotFound:
+						case IManifest::FResult::EType::BeforeStart:
+						case IManifest::FResult::EType::NotLoaded:
+						{
+							// Get the dependent stream types, if any.
+							// This is mainly for multiplexed streams like an .mp4 at this point.
+							// All stream types have reached EOS now.
+							TArray<IStreamSegment::FDependentStreams> DependentStreams;
+							FinishedReq.Request->GetDependentStreams(DependentStreams);
+							// Add the primary stream type to the list as well.
+							DependentStreams.AddDefaulted_GetRef().StreamType = FinishedReq.Request->GetType();
+							for(int32 i=0; i<DependentStreams.Num(); ++i)
+							{
+								CompletedSegmentRequests.Add(DependentStreams[i].StreamType, FinishedReq.Request);
+							}
+							break;
+						}
 					}
 					break;
 				}
@@ -2211,41 +2394,46 @@ void FAdaptiveStreamingPlayer::HandlePendingMediaSegmentRequests()
 			// Looping enabled?
 			if (CurrentLoopParam.bEnableLooping)
 			{
-				if (Manifest.IsValid() && CurrentPlayPeriod.IsValid())
+				if (Manifest.IsValid())
 				{
-					FPlayStartPosition	StartAt;
-					FTimeRange			Seekable = Manifest->GetSeekableTimeRange();
-					StartAt.Time = Seekable.Start;
-
-					TSharedPtrTS<IStreamSegment> FirstSegmentRequest;
-					IManifest::FResult Result = CurrentPlayPeriod->GetLoopingSegment(FirstSegmentRequest, CurrentLoopState, LocalFinishedRequests, StartAt, IManifest::ESearchType::Closest);
-					if (Result.GetType() == IManifest::FResult::EType::Found)
-					{
-						// Get the stream types that have already reached EOS in this start request because the streams are of different duration.
-						// The streams are there in principle so we need to set them up (in case we loop back into them) but we won't see any data arriving.
-						TArray<TSharedPtrTS<IStreamSegment>> AlreadyEndedStreams;
-						FirstSegmentRequest->GetEndedStreams(AlreadyEndedStreams);
-						for(int32 i=0; i<AlreadyEndedStreams.Num(); ++i)
-						{
-							CompletedSegmentRequests.Add(AlreadyEndedStreams[i]->GetType(), AlreadyEndedStreams[i]);
-						}
+					FTimeRange Seekable = Manifest->GetSeekableTimeRange();
 
-						PostLog(Facility::EFacility::Player, IInfoLog::ELevel::Verbose, FString::Printf(TEXT("Enqueuing stream loop start request.")));
-						// Add the updated loop state to the upcoming loop FIFO. When the play position reaches the next enqueued loop state it gets popped.
-						NextLoopStates.Push(CurrentLoopState);
-						StreamReaderHandler->AddRequest(CurrentPlaybackSequenceID, FirstSegmentRequest);
-					}
-					else
-					{
-						// How could this have failed? Switch off looping to end the stream.
-						PostLog(Facility::EFacility::Player, IInfoLog::ELevel::Error, FString::Printf(TEXT("Could not locate segment to loop back to. Disabling looping and ending stream.")));
-						CurrentLoopParam.bEnableLooping = false;
-					}
+					PendingStartRequest = MakeSharedTS<FPendingStartRequest>();
+					PendingStartRequest->SearchType = IManifest::ESearchType::Closest;
+					PendingStartRequest->StartAt.Time = Seekable.Start;
+					PendingStartRequest->bForLooping = true;
+					PendingStartRequest->FinishedRequests = MoveTemp(LocalFinishedRequests);
 				}
 			}
 		}
 	}
+}
+
 
+//-----------------------------------------------------------------------------
+/**
+ * Adds the given period to the list of periods playback will eventually move into
+ * and a metadata change event may need to be generated.
+ */
+void FAdaptiveStreamingPlayer::AddUpcomingPeriod(TSharedPtrTS<IManifest::IPlayPeriod> InUpcomingPeriod)
+{
+	TSharedPtrTS<ITimelineMediaAsset> Period = InUpcomingPeriod->GetMediaAsset();
+	if (Period.IsValid())
+	{
+		FString PeriodID = Period->GetUniqueIdentifier();
+		// Only add if it is not already in the list.
+		if (!UpcomingPeriods.ContainsByPredicate([PeriodID](const FUpcomingPeriod& e){ return e.ID.Equals(PeriodID); }))
+		{
+			FUpcomingPeriod Next;
+			Next.ID = MoveTemp(PeriodID);
+			Next.TimeRange = Period->GetTimeRange();
+			Next.Period = MoveTemp(Period);
+			// Shift the time range by the current loop state.
+			Next.TimeRange.Start += CurrentLoopState.LoopBasetime;
+			Next.TimeRange.End += CurrentLoopState.LoopBasetime;
+			UpcomingPeriods.Emplace(MoveTemp(Next));
+		}
+	}
 }
 
 
@@ -2450,13 +2638,13 @@ void FAdaptiveStreamingPlayer::AudioDecoderOutputReady(const IDecoderOutputBuffe
 
 bool FAdaptiveStreamingPlayer::FindMatchingStreamInfo(FStreamCodecInformation& OutStreamInfo, int32 MaxWidth, int32 MaxHeight)
 {
-	if (!CurrentPlayPeriod.IsValid())
+	if (!CurrentPlayPeriodVideo.IsValid())
 	{
 		return false;
 	}
 	TArray<FStreamCodecInformation> VideoCodecInfos;
 
-	TSharedPtrTS<ITimelineMediaAsset> Asset = CurrentPlayPeriod->GetMediaAsset();
+	TSharedPtrTS<ITimelineMediaAsset> Asset = CurrentPlayPeriodVideo->GetMediaAsset();
 	if (Asset.IsValid())
 	{
 		if (Asset->GetNumberOfAdaptationSets(EStreamType::Video) > 0)
@@ -2646,14 +2834,7 @@ void FAdaptiveStreamingPlayer::CheckForStreamEnd()
 				// Everything at EOD?
 				if (bEndVid && bEndAud && bEndTxt)
 				{
-					PostrollVars.Clear();
-					PlaybackState.SetHasEnded(true);
-					StreamState = EStreamState::eStream_ReachedEnd;
-					// Go to pause mode now.
-					InternalPause();
-					DispatchEvent(FMetricEvent::ReportPlaybackEnded());
-					// In case a seek back into the stream will happen we reset the first start state.
-					PrerollVars.bIsVeryFirstStart = true;
+					InternalSetPlaybackEnded();
 				}
 			}
 			else if (StreamState == EStreamState::eStream_ReachedEnd)
@@ -2714,9 +2895,9 @@ bool FAdaptiveStreamingPlayer::InternalStartAt(const FSeekParam& NewPosition)
 	if (ManifestReader.IsValid() && ManifestReader->GetPlaylistType() == TEXT("mp4"))
 	{
 		// Only select audio for now.
-		if (InitialStreamSelectionAud.IsValid())
+		if (InitialTrackSelectionAud.IsValid())
 		{
-			MultiStreamBufferAud.SelectTrackByID(InitialStreamSelectionAud->StreamUniqueID);
+			MultiStreamBufferAud.SelectTrackByID(InitialTrackSelectionAud->TrackID);
 		}
 // TODO/FIXME: Cannot start with audio disabled at the moment
 #if 0
@@ -2757,15 +2938,18 @@ bool FAdaptiveStreamingPlayer::InternalStartAt(const FSeekParam& NewPosition)
 		return false;
 	}
 
-	check(PendingStartRequest.Get() == nullptr);
-	check(CurrentPlayPeriod.Get() == nullptr);
-	check(PendingFirstSegmentRequest.Get() == nullptr);
+	check(!PendingStartRequest.IsValid());
+	check(!InitialPlayPeriod.IsValid());
+	check(!PendingFirstSegmentRequest.IsValid());
 	check(NextPendingSegmentRequests.IsEmpty());
 	check(CompletedSegmentRequests.Num() == 0);
 	NextPendingSegmentRequests.Empty();
 	PendingFirstSegmentRequest.Reset();
 	CompletedSegmentRequests.Empty();
-	CurrentPlayPeriod.Reset();
+	InitialPlayPeriod.Reset();
+	CurrentPlayPeriodVideo.Reset();
+	CurrentPlayPeriodAudio.Reset();
+	CurrentPlayPeriodText.Reset();
 	PendingStartRequest = MakeSharedTS<FPendingStartRequest>();
 	PendingStartRequest->SearchType = IManifest::ESearchType::Closest;
 	PendingStartRequest->StartAt.Time = NewPosition.Time;
@@ -2782,11 +2966,24 @@ bool FAdaptiveStreamingPlayer::InternalStartAt(const FSeekParam& NewPosition)
 	{
 		PendingStartRequest->SearchType = IManifest::ESearchType::Closest;
 	}
+	bFirstSegmentRequestIsForLooping = false;
 
 	return true;
 }
 
 
+void FAdaptiveStreamingPlayer::InternalSetPlaybackEnded()
+{
+	PostrollVars.Clear();
+	PlaybackState.SetHasEnded(true);
+	StreamState = EStreamState::eStream_ReachedEnd;
+	// Go to pause mode now.
+	InternalPause();
+	DispatchEvent(FMetricEvent::ReportPlaybackEnded());
+	// In case a seek back into the stream will happen we reset the first start state.
+	PrerollVars.bIsVeryFirstStart = true;
+}
+
 
 //-----------------------------------------------------------------------------
 /**
@@ -2869,7 +3066,9 @@ void FAdaptiveStreamingPlayer::InternalStop(bool bHoldCurrentFrame)
 	// Destroy the stream reader handler.
 	IStreamReader* CurrentStreamReader = TMediaInterlockedExchangePointer(StreamReaderHandler, (IStreamReader*)nullptr);
 	if (CurrentStreamReader)
+	{
 		CurrentStreamReader->Close();
+	}
 	delete CurrentStreamReader;
 
 	bHaveVideoReader.Reset();
@@ -2881,7 +3080,11 @@ void FAdaptiveStreamingPlayer::InternalStop(bool bHoldCurrentFrame)
 	PendingFirstSegmentRequest.Reset();
 	NextPendingSegmentRequests.Empty();
 	CompletedSegmentRequests.Empty();
-	CurrentPlayPeriod.Reset();
+	InitialPlayPeriod.Reset();
+	CurrentPlayPeriodVideo.Reset();
+	CurrentPlayPeriodAudio.Reset();
+	CurrentPlayPeriodText.Reset();
+	UpcomingPeriods.Empty();
 
 	PlaybackRate = 0.0;
 	CurrentState = EPlayerState::eState_Paused;
@@ -3096,6 +3299,7 @@ void FAdaptiveStreamingPlayer::DebugPrint(void* pPlayer, void (*pPrintFN)(void*
 	pPrintFN(pPlayer, "    is seeking: %s", IsSeeking() ? "Yes" : "No");
 	pPrintFN(pPlayer, "     has ended: %s", HasEnded() ? "Yes" : "No");
 	pPrintFN(pPlayer, "clock (%s): %.3f", RenderClock->IsRunning()?"running":"paused", RenderClock->GetInterpolatedRenderTime(IMediaRenderClock::ERendererType::Video).GetAsSeconds(-1.0));
+	pPrintFN(pPlayer, " internal time: %.3f", GetCurrentPlayTime().GetAsSeconds());
 }
 void FAdaptiveStreamingPlayer::DebugHandle(void* pPlayer, void (*debugDrawPrintf)(void* pPlayer, const char *pFmt, ...))
 {
@@ -3189,6 +3393,9 @@ void FAdaptiveStreamingPlayerEventHandler::WorkerThreadFN()
 					case FMetricEvent::EType::ReceivedPlaylists:
 						MetricListeners[i]->ReportReceivedPlaylists();
 						break;
+					case FMetricEvent::EType::TracksChanged:
+						MetricListeners[i]->ReportTracksChanged();
+						break;
 					case FMetricEvent::EType::BufferingStart:
 						MetricListeners[i]->ReportBufferingStart(pEvt->Param.BufferingReason);
 						break;
diff --git a/Engine/Plugins/Media/ElectraPlayer/Source/ElectraPlayerRuntime/Private/Runtime/Player/AdaptiveStreamingPlayer.h b/Engine/Plugins/Media/ElectraPlayer/Source/ElectraPlayerRuntime/Private/Runtime/Player/AdaptiveStreamingPlayer.h
--- a/Engine/Plugins/Media/ElectraPlayer/Source/ElectraPlayerRuntime/Private/Runtime/Player/AdaptiveStreamingPlayer.h
+++ b/Engine/Plugins/Media/ElectraPlayer/Source/ElectraPlayerRuntime/Private/Runtime/Player/AdaptiveStreamingPlayer.h
@@ -144,8 +144,8 @@ public:
 	//! Returns the current loop state.
 	virtual void GetLoopState(FPlayerLoopState& OutLoopState) const = 0;
 
-	//! Returns stream metadata of the currently active play period.
-	virtual void GetStreamMetadata(TArray<FStreamMetadata>& OutStreamMetadata, EStreamType StreamType) const = 0;
+	//! Returns track metadata of the currently active play period.
+	virtual void GetTrackMetadata(TArray<FTrackMetadata>& OutTrackMetadata, EStreamType StreamType) const = 0;
 
 
 	//-------------------------------------------------------------------------
@@ -158,9 +158,9 @@ public:
 	//! Setting both will limit on either width or height, whichever limits first.
 	virtual void SetMaxResolution(int32 MaxWidth, int32 MaxHeight) = 0;
 
-	//! Selects a track based from one of the array members returned by GetStreamMetadata().
+	//! Selects a track based from one of the array members returned by GetTrackMetadata().
 	//! NOTE: Presently this must be done before starting playback!
-	virtual void SelectTrackByMetadata(EStreamType StreamType, const FStreamMetadata& StreamMetadata) = 0;
+	virtual void SelectTrackByMetadata(EStreamType StreamType, const FTrackMetadata& StreamMetadata) = 0;
 
 	//! Deselect track. The stream will continue to stream to allow for immediate selection/activation but no data will be fed to the decoder.
 	virtual void DeselectTrack(EStreamType StreamType) = 0;
diff --git a/Engine/Plugins/Media/ElectraPlayer/Source/ElectraPlayerRuntime/Private/Runtime/Player/AdaptiveStreamingPlayerABR.cpp b/Engine/Plugins/Media/ElectraPlayer/Source/ElectraPlayerRuntime/Private/Runtime/Player/AdaptiveStreamingPlayerABR.cpp
--- a/Engine/Plugins/Media/ElectraPlayer/Source/ElectraPlayerRuntime/Private/Runtime/Player/AdaptiveStreamingPlayerABR.cpp
+++ b/Engine/Plugins/Media/ElectraPlayer/Source/ElectraPlayerRuntime/Private/Runtime/Player/AdaptiveStreamingPlayerABR.cpp
@@ -124,7 +124,7 @@ namespace Electra
 		virtual ~FAdaptiveStreamSelector();
 		virtual void SetPresentationType(EMediaPresentationType PresentationType) override;
 		virtual void SetCanSwitchToAlternateStreams(bool bCanSwitch) override;
-		virtual void SetCurrentPlaybackPeriod(TSharedPtrTS<IManifest::IPlayPeriod> InCurrentPlayPeriod) override;
+		virtual void SetCurrentPlaybackPeriod(EStreamType InStreamType, TSharedPtrTS<IManifest::IPlayPeriod> InCurrentPlayPeriod) override;
 		virtual void SetBandwidth(int64 bitsPerSecond) override;
 		virtual void SetForcedNextBandwidth(int64 bitsPerSecond) override;
 		virtual void MarkStreamAsUnavailable(const FBlacklistedStream& BlacklistedStream) override;
@@ -146,6 +146,7 @@ namespace Electra
 		virtual void ReportOpenSource(const FString& URL) override {}
 		virtual void ReportReceivedMasterPlaylist(const FString& EffectiveURL) override {}
 		virtual void ReportReceivedPlaylists() override {}
+		virtual void ReportTracksChanged() override {}
 		virtual void ReportPlaylistDownload(const Metrics::FPlaylistDownloadStats& PlaylistDownloadStats) override {}
 		virtual void ReportBandwidth(int64 EffectiveBps, int64 ThroughputBps, double LatencyInSeconds) override {}
 		virtual void ReportBufferUtilization(const Metrics::FBufferStats& BufferStats) override {}
@@ -214,13 +215,13 @@ namespace Electra
 
 		FMediaCriticalSection								AccessMutex;
 		EMediaPresentationType								PresentationType;
-		TSharedPtrTS<IManifest::IPlayPeriod>				CurrentPlayPeriod;
+		TSharedPtrTS<IManifest::IPlayPeriod>				CurrentPlayPeriodVideo;
+		TSharedPtrTS<IManifest::IPlayPeriod>				CurrentPlayPeriodAudio;
 		bool												bPlayerIsBuffering;
 		bool												bCanSwitchToAlternateStreams;
 
 		TArray<TSharedPtrTS<FStreamInformation>>			StreamInformationVideo;
 		TArray<TSharedPtrTS<FStreamInformation>>			StreamInformationAudio;
-		FString												CurrentMediaAssetID;
 		FString												CurrentVideoAdaptationSetID;
 		FString												CurrentAudioAdaptationSetID;
 
@@ -352,84 +353,97 @@ namespace Electra
 	/**
 	 * Sets the playback period from which to select streams now.
 	 *
+	 * @param InStreamType
 	 * @param CurrentPlayPeriod
 	 */
-	void FAdaptiveStreamSelector::SetCurrentPlaybackPeriod(TSharedPtrTS<IManifest::IPlayPeriod> InCurrentPlayPeriod)
+	void FAdaptiveStreamSelector::SetCurrentPlaybackPeriod(EStreamType InStreamType, TSharedPtrTS<IManifest::IPlayPeriod> InCurrentPlayPeriod)
 	{
 		FMediaCriticalSection::ScopedLock lock(AccessMutex);
-		CurrentPlayPeriod = InCurrentPlayPeriod;
-		if (CurrentPlayPeriod.IsValid())
+		if (InStreamType == EStreamType::Video)
 		{
-			TSharedPtrTS<ITimelineMediaAsset> Asset = CurrentPlayPeriod->GetMediaAsset();
-			if (Asset.IsValid() && CurrentMediaAssetID != Asset->GetUniqueIdentifier())
+			CurrentPlayPeriodVideo = InCurrentPlayPeriod;
+			if (CurrentPlayPeriodVideo.IsValid())
 			{
-				CurrentMediaAssetID.Empty();
-				CurrentVideoAdaptationSetID.Empty();
-				CurrentAudioAdaptationSetID.Empty();
-				StreamInformationVideo.Empty();
-				StreamInformationAudio.Empty();
-				// Get video stream information
-				if (Asset->GetNumberOfAdaptationSets(EStreamType::Video) > 0)
+				TSharedPtrTS<ITimelineMediaAsset> Asset = CurrentPlayPeriodVideo->GetMediaAsset();
+				if (Asset.IsValid())
 				{
-					TSharedPtrTS<IPlaybackAssetAdaptationSet> Adapt = Asset->GetAdaptationSetByTypeAndIndex(EStreamType::Video, 0);
-					if (Adapt.IsValid())
+					CurrentVideoAdaptationSetID.Empty();
+					StreamInformationVideo.Empty();
+					// Get video stream information
+					if (Asset->GetNumberOfAdaptationSets(InStreamType) > 0)
 					{
-						CurrentMediaAssetID = Asset->GetUniqueIdentifier();
-						CurrentVideoAdaptationSetID = Adapt->GetUniqueIdentifier();
-						for(int32 i=0, iMax=Adapt->GetNumberOfRepresentations(); i<iMax; ++i)
+						TSharedPtrTS<IPlaybackAssetAdaptationSet> Adapt = Asset->GetAdaptationSetByTypeAndIndex(InStreamType, 0);
+						if (Adapt.IsValid())
 						{
-							TSharedPtrTS<IPlaybackAssetRepresentation> Repr = Adapt->GetRepresentationByIndex(i);
-
-							TSharedPtrTS<FStreamInformation> si = MakeSharedTS<FStreamInformation>();
-							si->AdaptationSetUniqueID = Adapt->GetUniqueIdentifier();
-							si->RepresentationUniqueID = Repr->GetUniqueIdentifier();
-							si->Bitrate = Repr->GetBitrate();
-							if (Repr->GetCodecInformation().IsVideoCodec())
+							CurrentVideoAdaptationSetID = Adapt->GetUniqueIdentifier();
+							for(int32 i=0, iMax=Adapt->GetNumberOfRepresentations(); i<iMax; ++i)
 							{
-								si->Resolution = Repr->GetCodecInformation().GetResolution();
+								TSharedPtrTS<IPlaybackAssetRepresentation> Repr = Adapt->GetRepresentationByIndex(i);
+								TSharedPtrTS<FStreamInformation> si = MakeSharedTS<FStreamInformation>();
+								si->AdaptationSetUniqueID = Adapt->GetUniqueIdentifier();
+								si->RepresentationUniqueID = Repr->GetUniqueIdentifier();
+								si->Bitrate = Repr->GetBitrate();
+								if (Repr->GetCodecInformation().IsVideoCodec())
+								{
+									si->Resolution = Repr->GetCodecInformation().GetResolution();
+								}
+								StreamInformationVideo.Push(si);
 							}
-							StreamInformationVideo.Push(si);
+							// Sort the representations by ascending bitrate
+							StreamInformationVideo.Sort([](const TSharedPtrTS<FStreamInformation>& a, const TSharedPtrTS<FStreamInformation>& b)
+							{
+								return a->Bitrate < b->Bitrate;
+							});
 						}
-						// Sort the representations by ascending bitrate
-						StreamInformationVideo.Sort([](const TSharedPtrTS<FStreamInformation>& a, const TSharedPtrTS<FStreamInformation>& b)
-						{
-							return a->Bitrate < b->Bitrate;
-						});
 					}
 				}
-				// Get audio stream information
-				if (Asset->GetNumberOfAdaptationSets(EStreamType::Audio) > 0)
+			}
+			else
+			{
+				CurrentVideoAdaptationSetID.Empty();
+				StreamInformationVideo.Empty();
+			}
+		}
+		else if (InStreamType == EStreamType::Audio)
+		{
+			CurrentPlayPeriodAudio = InCurrentPlayPeriod;
+			if (CurrentPlayPeriodAudio.IsValid())
+			{
+				TSharedPtrTS<ITimelineMediaAsset> Asset = CurrentPlayPeriodAudio->GetMediaAsset();
+				if (Asset.IsValid())
 				{
-					TSharedPtrTS<IPlaybackAssetAdaptationSet> Adapt = Asset->GetAdaptationSetByTypeAndIndex(EStreamType::Audio, 0);
-					if (Adapt.IsValid())
+					CurrentAudioAdaptationSetID.Empty();
+					StreamInformationAudio.Empty();
+					// Get audio stream information
+					if (Asset->GetNumberOfAdaptationSets(InStreamType) > 0)
 					{
-						CurrentMediaAssetID = Asset->GetUniqueIdentifier();
-						CurrentAudioAdaptationSetID = Adapt->GetUniqueIdentifier();
-						for(int32 i=0, iMax=Adapt->GetNumberOfRepresentations(); i<iMax; ++i)
+						TSharedPtrTS<IPlaybackAssetAdaptationSet> Adapt = Asset->GetAdaptationSetByTypeAndIndex(InStreamType, 0);
+						if (Adapt.IsValid())
 						{
-							TSharedPtrTS<IPlaybackAssetRepresentation> Repr = Adapt->GetRepresentationByIndex(i);
-							TSharedPtrTS<FStreamInformation> si = MakeSharedTS<FStreamInformation>();
-							si->AdaptationSetUniqueID = Adapt->GetUniqueIdentifier();
-							si->RepresentationUniqueID = Repr->GetUniqueIdentifier();
-							si->Bitrate = Repr->GetBitrate();
-							StreamInformationAudio.Push(si);
+							CurrentAudioAdaptationSetID = Adapt->GetUniqueIdentifier();
+							for(int32 i=0, iMax=Adapt->GetNumberOfRepresentations(); i<iMax; ++i)
+							{
+								TSharedPtrTS<IPlaybackAssetRepresentation> Repr = Adapt->GetRepresentationByIndex(i);
+								TSharedPtrTS<FStreamInformation> si = MakeSharedTS<FStreamInformation>();
+								si->AdaptationSetUniqueID = Adapt->GetUniqueIdentifier();
+								si->RepresentationUniqueID = Repr->GetUniqueIdentifier();
+								si->Bitrate = Repr->GetBitrate();
+								StreamInformationAudio.Push(si);
+							}
+							// Sort the representations by ascending bitrate
+							StreamInformationAudio.Sort([](const TSharedPtrTS<FStreamInformation>& a, const TSharedPtrTS<FStreamInformation>& b)
+							{
+								return a->Bitrate > b->Bitrate;
+							});
 						}
-						// Sort the representations by ascending bitrate
-						StreamInformationAudio.Sort([](const TSharedPtrTS<FStreamInformation>& a, const TSharedPtrTS<FStreamInformation>& b)
-						{
-							return a->Bitrate > b->Bitrate;
-						});
 					}
 				}
 			}
-		}
-		else
-		{
-			CurrentMediaAssetID.Empty();
-			CurrentVideoAdaptationSetID.Empty();
-			CurrentAudioAdaptationSetID.Empty();
-			StreamInformationVideo.Empty();
-			StreamInformationAudio.Empty();
+			else
+			{
+				CurrentAudioAdaptationSetID.Empty();
+				StreamInformationAudio.Empty();
+			}
 		}
 	}
 
@@ -794,10 +808,6 @@ namespace Electra
 	IAdaptiveStreamSelector::ESegmentAction FAdaptiveStreamSelector::SelectSuitableStreams(FTimeValue& OutDelay, TSharedPtrTS<const IStreamSegment> CurrentSegment)
 	{
 		OutDelay.SetToZero();
-		if (!CurrentPlayPeriod.IsValid())
-		{
-			return ESegmentAction::FetchNext;
-		}
 
 		FTimeValue TimeNow = PlayerSessionServices->GetSynchronizedUTCTime()->GetTime();
 		// Make all streams available again if their blacklist time has expired
@@ -820,6 +830,21 @@ namespace Electra
 
 
 		EStreamType StreamType = CurrentSegment.IsValid() ? CurrentSegment->GetType() : StreamInformationVideo.Num() ? EStreamType::Video : EStreamType::Audio;
+		
+		TSharedPtrTS<IManifest::IPlayPeriod> CurrentPlayPeriod;
+		if (StreamType == EStreamType::Video)
+		{
+			CurrentPlayPeriod = CurrentPlayPeriodVideo;
+		}
+		else if (StreamType == EStreamType::Audio)
+		{
+			CurrentPlayPeriod = CurrentPlayPeriodAudio;
+		}
+		if (!CurrentPlayPeriod.IsValid())
+		{
+			return ESegmentAction::FetchNext;
+		}
+
 		// Was the previous segment download in error?
 		bool bRetryIfPossible = false;
 		bool bSkipWithFiller = false;
@@ -827,8 +852,9 @@ namespace Electra
 		{
 			Metrics::FSegmentDownloadStats Stats;
 			CurrentSegment->GetDownloadStats(Stats);
+			// Try to get the stream information for the current downloaded segment. We may not find it on a period transition where the
+			// segment is the last of the previous period.
 			TSharedPtrTS<FAdaptiveStreamSelector::FStreamInformation> CurrentStreamInfo = GetStreamInformation(Stats);
-			check(CurrentStreamInfo.IsValid());
 			if (CurrentStreamInfo.IsValid())
 			{
 				if (!Stats.bWasSuccessful)
diff --git a/Engine/Plugins/Media/ElectraPlayer/Source/ElectraPlayerRuntime/Private/Runtime/Player/AdaptiveStreamingPlayerABR.h b/Engine/Plugins/Media/ElectraPlayer/Source/ElectraPlayerRuntime/Private/Runtime/Player/AdaptiveStreamingPlayerABR.h
--- a/Engine/Plugins/Media/ElectraPlayer/Source/ElectraPlayerRuntime/Private/Runtime/Player/AdaptiveStreamingPlayerABR.h
+++ b/Engine/Plugins/Media/ElectraPlayer/Source/ElectraPlayerRuntime/Private/Runtime/Player/AdaptiveStreamingPlayerABR.h
@@ -57,7 +57,7 @@ public:
 	virtual void SetCanSwitchToAlternateStreams(bool bCanSwitch) = 0;
 
 	//! Sets the manifest from which to select streams.
-	virtual void SetCurrentPlaybackPeriod(TSharedPtrTS<IManifest::IPlayPeriod> CurrentPlayPeriod) = 0;
+	virtual void SetCurrentPlaybackPeriod(EStreamType InStreamType, TSharedPtrTS<IManifest::IPlayPeriod> CurrentPlayPeriod) = 0;
 
 	//! Sets the initial bandwidth, either from guessing, past history or a current measurement.
 	virtual void SetBandwidth(int64 bitsPerSecond) = 0;
@@ -120,6 +120,7 @@ public:
 	virtual void ReportOpenSource(const FString& URL) = 0;
 	virtual void ReportReceivedMasterPlaylist(const FString& EffectiveURL) = 0;
 	virtual void ReportReceivedPlaylists() = 0;
+	virtual void ReportTracksChanged() = 0;
 	virtual void ReportPlaylistDownload(const Metrics::FPlaylistDownloadStats& PlaylistDownloadStats) = 0;
 	virtual void ReportBufferingStart(Metrics::EBufferingReason BufferingReason) = 0;
 	virtual void ReportBufferingEnd(Metrics::EBufferingReason BufferingReason) = 0;
diff --git a/Engine/Plugins/Media/ElectraPlayer/Source/ElectraPlayerRuntime/Private/Runtime/Player/AdaptiveStreamingPlayerInternal.h b/Engine/Plugins/Media/ElectraPlayer/Source/ElectraPlayerRuntime/Private/Runtime/Player/AdaptiveStreamingPlayerInternal.h
--- a/Engine/Plugins/Media/ElectraPlayer/Source/ElectraPlayerRuntime/Private/Runtime/Player/AdaptiveStreamingPlayerInternal.h
+++ b/Engine/Plugins/Media/ElectraPlayer/Source/ElectraPlayerRuntime/Private/Runtime/Player/AdaptiveStreamingPlayerInternal.h
@@ -22,7 +22,9 @@
 #define INTERR_ALL_STREAMS_HAVE_FAILED			1
 #define INTERR_UNSUPPORTED_FORMAT				2
 #define INTERR_COULD_NOT_LOCATE_START_SEGMENT	3
+#define INTERR_COULD_NOT_LOCATE_START_PERIOD	4
 #define INTERR_FRAGMENT_NOT_AVAILABLE			0x101
+#define INTERR_FRAGMENT_READER_REQUEST			0x102
 #define INTERR_CREATE_FRAGMENT_READER			0x103
 #define INTERR_REBUFFER_SHALL_THROW_ERROR		0x200
 
@@ -182,8 +184,8 @@
 	bool									bIsBuffering;
 	bool									bIsPlaying;
 	bool									bIsPaused;
-	TArray<FStreamMetadata>					VideoStreams;
-	TArray<FStreamMetadata>					AudioStreams;
+	TArray<FTrackMetadata>					VideoTracks;
+	TArray<FTrackMetadata>					AudioTracks;
 	TArray<FTimespan>						SeekablePositions;
 
 
@@ -315,17 +317,36 @@
 		bIsPlaying = bInIsPlaying;
 	}
 
-	void SetStreamMetadata(const TArray<FStreamMetadata> &InVideoStreams, const TArray<FStreamMetadata>& InAudioStreams)
+	bool SetTrackMetadata(const TArray<FTrackMetadata> &InVideoTracks, const TArray<FTrackMetadata>& InAudioTracks)
 	{
 		FMediaCriticalSection::ScopedLock lock(Lock);
-		VideoStreams = InVideoStreams;
-		AudioStreams = InAudioStreams;
+
+		auto ChangedTrackMetadata = [](const TArray<FTrackMetadata> &These, const TArray<FTrackMetadata>& Other) -> bool
+		{
+			if (These.Num() == Other.Num())
+			{
+				for(int32 i=0; i<These.Num(); ++i)
+				{
+					if (!These[i].Equals(Other[i]))
+					{
+						return true;
+					}
+				}
+				return false;
+			}
+			return true;
+		};
+
+		bool bChanged = ChangedTrackMetadata(VideoTracks, InVideoTracks) || ChangedTrackMetadata(AudioTracks, InAudioTracks);
+		VideoTracks = InVideoTracks;
+		AudioTracks = InAudioTracks;
+		return bChanged;
 	}
-	void GetStreamMetadata(TArray<FStreamMetadata> &OutVideoStreams, TArray<FStreamMetadata>& OutAudioStreams) const
+	void GetTrackMetadata(TArray<FTrackMetadata> &OutVideoTracks, TArray<FTrackMetadata>& OutAudioTracks) const
 	{
 		FMediaCriticalSection::ScopedLock lock(Lock);
-		OutVideoStreams = VideoStreams;
-		OutAudioStreams = AudioStreams;
+		OutVideoTracks = VideoTracks;
+		OutAudioTracks = AudioTracks;
 	}
 
 	void SetLoopState(const FPlayerLoopState& InLoopState)
@@ -358,6 +379,7 @@
 		OpenSource,
 		ReceivedMasterPlaylist,
 		ReceivedPlaylists,
+		TracksChanged,
 		BufferingStart,
 		BufferingEnd,
 		Bandwidth,
@@ -445,6 +467,12 @@
 		Evt->Type = EType::ReceivedPlaylists;
 		return Evt;
 	}
+	static TSharedPtrTS<FMetricEvent> ReportTracksChanged()
+	{
+		TSharedPtrTS<FMetricEvent> Evt = MakeSharedTS<FMetricEvent>();
+		Evt->Type = EType::TracksChanged;
+		return Evt;
+	}
 	static TSharedPtrTS<FMetricEvent> ReportBufferingStart(Metrics::EBufferingReason BufferingReason)
 	{
 		TSharedPtrTS<FMetricEvent> Evt = MakeSharedTS<FMetricEvent>();
@@ -663,12 +691,12 @@
 	virtual bool IsPaused() const override;
 
 	virtual void GetLoopState(FPlayerLoopState& OutLoopState) const override;
-	virtual void GetStreamMetadata(TArray<FStreamMetadata>& OutStreamMetadata, EStreamType StreamType) const override;
+	virtual void GetTrackMetadata(TArray<FTrackMetadata>& OutTrackMetadata, EStreamType StreamType) const override;
 
 	virtual void SetBitrateCeiling(int32 highestSelectableBitrate) override;
 	virtual void SetMaxResolution(int32 MaxWidth, int32 MaxHeight) override;
 
-	virtual void SelectTrackByMetadata(EStreamType StreamType, const FStreamMetadata& StreamMetadata) override;
+	virtual void SelectTrackByMetadata(EStreamType StreamType, const FTrackMetadata& StreamMetadata) override;
 	virtual void DeselectTrack(EStreamType StreamType) override;
 
 #if PLATFORM_ANDROID
@@ -988,7 +1016,7 @@
 				struct FMetadataTrackSelection
 				{
 					EStreamType										StreamType;
-					FStreamMetadata									StreamMetadata;
+					FTrackMetadata									TrackMetadata;
 				};
 
 				FLoadManifest				ManifestToLoad;
@@ -1088,12 +1116,12 @@
 			Msg.Data.Resolution.Height = Height;
 			WorkMessages.SendMessage(Msg);
 		}
-		void SendTrackSelectByMetadataMessage(EStreamType StreamType, const FStreamMetadata& StreamMetadata)
+		void SendTrackSelectByMetadataMessage(EStreamType StreamType, const FTrackMetadata& TrackMetadata)
 		{
 			FMessage Msg;
 			Msg.Type = FMessage::EType::SelectTrackByMetadata;
 			Msg.Data.TrackSelection.StreamType = StreamType;
-			Msg.Data.TrackSelection.StreamMetadata = StreamMetadata;
+			Msg.Data.TrackSelection.TrackMetadata = TrackMetadata;
 			WorkMessages.SendMessage(Msg);
 		}
 		void SendTrackDeselectMessage(EStreamType StreamType)
@@ -1109,6 +1137,111 @@
 		bool											bStarted;
 	};
 
+	struct FPendingStartRequest
+	{
+		FPlayStartPosition										StartAt;
+		IManifest::ESearchType									SearchType;
+		FTimeValue												RetryAtTime;
+		TMediaOptionalValue<int32>								InitialBandwidth;
+		bool													bForLooping = false;
+		TMultiMap<EStreamType, TSharedPtrTS<IStreamSegment>>	FinishedRequests;
+	};
 
+	struct FBufferStats
+	{
+		void Clear()
+		{
+			StreamBuffer.Clear();
+			DecoderInputBuffer.Clear();
+			DecoderOutputBuffer.Clear();
+		}
+		FAccessUnitBufferInfo								StreamBuffer;
+		IAccessUnitBufferListener::FBufferStats				DecoderInputBuffer;
+		IDecoderOutputBufferListener::FDecodeReadyStats		DecoderOutputBuffer;
+	};
+
+	struct FPrerollVars
+	{
+		FPrerollVars()
+		{
+			Clear();
+			bIsVeryFirstStart = true;
+		}
+		void Clear()
+		{
+			StartTime  	  = -1;
+			bHaveEnoughVideo = false;
+			bHaveEnoughAudio = false;
+			bHaveEnoughText  = false;
+		}
+		int64					StartTime;
+		bool					bHaveEnoughVideo;
+		bool					bHaveEnoughAudio;
+		bool					bHaveEnoughText;
+		bool					bIsVeryFirstStart;
+	};
+
+	struct FPostrollVars
+	{
+		struct FRenderState
+		{
+			FRenderState()
+			{
+				Clear();
+			}
+			void Clear()
+			{
+				LastCheckTime   = -1;
+				LastBufferCount = 0;
+			}
+			int64	LastCheckTime;
+			int32	LastBufferCount;
+		};
+
+		FPostrollVars()
+		{
+			Clear();
+		}
+
+		void Clear()
+		{
+			Video.Clear();
+			Audio.Clear();
+		}
+
+		FRenderState		Video;
+		FRenderState		Audio;
+	};
+
+	struct FStreamBitrateInfo
+	{
+		FStreamBitrateInfo()
+		{
+			Clear();
+		}
+		void Clear()
+		{
+			Bitrate 	 = 0;
+			QualityLevel = 0;
+		}
+		int32			Bitrate;
+		int32			QualityLevel;
+	};
+
+	struct FPendingSegmentRequest
+	{
+		TSharedPtrTS<IStreamSegment>			Request;
+		FTimeValue								AtTime;
+		TSharedPtrTS<IManifest::IPlayPeriod>	Period;			//!< Set if transitioning between periods. This is the new period that needs to be readied.
+	};
+
+	struct FUpcomingPeriod
+	{
+		TSharedPtrTS<ITimelineMediaAsset>		Period;
+		FString									ID;
+		FTimeRange								TimeRange;
+	};
+
+
 	void InternalLoadManifest(const FString& URL, const FString& MimeType);
 	bool SelectManifest();
@@ -1133,6 +1266,7 @@
 	void DestroyDecoders();
 	bool FindMatchingStreamInfo(FStreamCodecInformation& OutStreamInfo, int32 MaxWidth, int32 MaxHeight);
 	void UpdateStreamResolutionLimit();
+	void AddUpcomingPeriod(TSharedPtrTS<IManifest::IPlayPeriod> InUpcomingPeriod);
 
 	// AU memory / generic stream reader memory
 	void* AUAllocate(IAccessUnitMemoryProvider::EDataType type, SIZE_T size, SIZE_T alignment) override;
@@ -1144,36 +1278,7 @@
 	void OnFragmentReachedEOS(EStreamType InStreamType, TSharedPtr<const FStreamSourceInfo, ESPMode::ThreadSafe> InStreamSourceInfo) override;
 	void OnFragmentClose(TSharedPtrTS<IStreamSegment> pRequest) override;
 
-
-
-
-
-	struct FPendingStartRequest
-	{
-		FPlayStartPosition				StartAt;
-		IManifest::ESearchType			SearchType;
-		FTimeValue						RetryAtTime;
-		TMediaOptionalValue<int32>		InitialBandwidth;
-	};
-
-
-
-
-	struct FBufferStats
-	{
-		void Clear()
-		{
-			StreamBuffer.Clear();
-			DecoderInputBuffer.Clear();
-			DecoderOutputBuffer.Clear();
-		}
-		FAccessUnitBufferInfo								StreamBuffer;
-		IAccessUnitBufferListener::FBufferStats				DecoderInputBuffer;
-		IDecoderOutputBufferListener::FDecodeReadyStats		DecoderOutputBuffer;
-	};
 	void UpdateDiagnostics();
-
-
 	void HandleNewBufferedData();
 	void HandleNewOutputData();
 	void HandleSessionMessage(TSharedPtrTS<IPlayerMessage> SessionMessage);
@@ -1196,7 +1301,7 @@
 	void InternalStop(bool bHoldCurrentFrame);
 	void InternalClose();
 	void InternalSetLoop(const FLoopParam& LoopParam);
+	void InternalSetPlaybackEnded();
 
-
 	void DispatchEvent(TSharedPtrTS<FMetricEvent> Event);
 	void DispatchEventAndWait(TSharedPtrTS<FMetricEvent> Event);
@@ -1209,59 +1314,7 @@
 	void StartRendering();
 	void StopRendering();
 
-
-	struct FPrerollVars
-	{
-		FPrerollVars()
-		{
-			Clear();
-			bIsVeryFirstStart = true;
-		}
-		void Clear()
-		{
-			StartTime  	  = -1;
-			bHaveEnoughVideo = false;
-			bHaveEnoughAudio = false;
-			bHaveEnoughText  = false;
-		}
-		int64					StartTime;
-		bool					bHaveEnoughVideo;
-		bool					bHaveEnoughAudio;
-		bool					bHaveEnoughText;
-		bool					bIsVeryFirstStart;
-	};
-
-	struct FPostrollVars
-	{
-		struct FRenderState
-		{
-			FRenderState()
-			{
-				Clear();
-			}
-			void Clear()
-			{
-				LastCheckTime   = -1;
-				LastBufferCount = 0;
-			}
-			int64	LastCheckTime;
-			int32	LastBufferCount;
-		};
-
-		FPostrollVars()
-		{
-			Clear();
-		}
-
-		void Clear()
-		{
-			Video.Clear();
-			Audio.Clear();
-		}
-
-		FRenderState		Video;
-		FRenderState		Audio;
-	};
+	FTimeValue GetCurrentPlayTime();
 
 
 	//
@@ -1305,9 +1358,9 @@
 	Metrics::FDataAvailabilityChange									DataAvailabilityStateAud;
 	Metrics::FDataAvailabilityChange									DataAvailabilityStateTxt;
 
-	TSharedPtr<FStreamMetadata, ESPMode::ThreadSafe>					InitialStreamSelectionVid;
-	TSharedPtr<FStreamMetadata, ESPMode::ThreadSafe>					InitialStreamSelectionAud;
-	TSharedPtr<FStreamMetadata, ESPMode::ThreadSafe>					InitialStreamSelectionTxt;
+	TSharedPtr<FTrackMetadata, ESPMode::ThreadSafe>						InitialTrackSelectionVid;
+	TSharedPtr<FTrackMetadata, ESPMode::ThreadSafe>						InitialTrackSelectionAud;
+	TSharedPtr<FTrackMetadata, ESPMode::ThreadSafe>						InitialTrackSelectionTxt;
 
 	EPlayerState														CurrentState;
 	EPipelineState														PipelineState;
@@ -1329,22 +1382,8 @@
 	int32																VideoResolutionLimitWidth;
 	int32																VideoResolutionLimitHeight;
 	int32																MaxVideoTextureResolutionLimitHeight;
-	struct FStreamBitrateInfo
-	{
-		FStreamBitrateInfo()
-		{
-			Clear();
-		}
-		void Clear()
-		{
-			Bitrate 	 = 0;
-			QualityLevel = 0;
-		}
-		int32			Bitrate;
-		int32			QualityLevel;
-	};
+
 	FStreamBitrateInfo													CurrentVideoStreamBitrate;
 
-
 	bool																bShouldBePaused;
 	bool																bShouldBePlaying;
@@ -1354,18 +1393,16 @@
 	FLoopParam															CurrentLoopParam;
 	TMediaQueueDynamicNoLock<FPlayerLoopState>							NextLoopStates;
 
-	struct FPendingSegmentRequest
-	{
-		TSharedPtrTS<IStreamSegment>		Request;
-		FTimeValue							AtTime;
-		TMediaOptionalValue<IAdaptiveStreamSelector::ESegmentAction> PreviousEvalResult;
-	};
-
-	TSharedPtrTS<IManifest::IPlayPeriod>								CurrentPlayPeriod;
+	TArray<FUpcomingPeriod>												UpcomingPeriods;
+	TSharedPtrTS<IManifest::IPlayPeriod>								InitialPlayPeriod;
+	TSharedPtrTS<IManifest::IPlayPeriod>								CurrentPlayPeriodVideo;
+	TSharedPtrTS<IManifest::IPlayPeriod>								CurrentPlayPeriodAudio;
+	TSharedPtrTS<IManifest::IPlayPeriod>								CurrentPlayPeriodText;
 	TSharedPtrTS<FPendingStartRequest>									PendingStartRequest;
 	TSharedPtrTS<IStreamSegment>										PendingFirstSegmentRequest;
 	TQueue<FPendingSegmentRequest>										NextPendingSegmentRequests;
 	TMultiMap<EStreamType, TSharedPtrTS<IStreamSegment>>				CompletedSegmentRequests;
+	bool																bFirstSegmentRequestIsForLooping;
 
 	uint32																CurrentPlaybackSequenceID;
 
diff --git a/Engine/Plugins/Media/ElectraPlayer/Source/ElectraPlayerRuntime/Private/Runtime/Player/AdaptiveStreamingPlayerManifest.cpp b/Engine/Plugins/Media/ElectraPlayer/Source/ElectraPlayerRuntime/Private/Runtime/Player/AdaptiveStreamingPlayerManifest.cpp
--- a/Engine/Plugins/Media/ElectraPlayer/Source/ElectraPlayerRuntime/Private/Runtime/Player/AdaptiveStreamingPlayerManifest.cpp
+++ b/Engine/Plugins/Media/ElectraPlayer/Source/ElectraPlayerRuntime/Private/Runtime/Player/AdaptiveStreamingPlayerManifest.cpp
@@ -106,25 +106,20 @@ void FAdaptiveStreamingPlayer::InternalLoadManifest(const FString& URL, const FS
 		FString mimeType = Playlist::GetMIMETypeForURL(URL);
 		if (mimeType.Len())
 		{
-			CurrentState = EPlayerState::eState_ParsingManifest;
+			check(!ManifestReader.IsValid());
 
+			CurrentState = EPlayerState::eState_ParsingManifest;
 			if (mimeType == Playlist::MIMETypeHLS)
 			{
 				ManifestReader = IPlaylistReaderHLS::Create(this);
-				DispatchEvent(FMetricEvent::ReportOpenSource(URL));
-				ManifestReader->LoadAndParse(URL);
 			}
 			else if (mimeType == Playlist::MIMETypeMP4)
 			{
 				ManifestReader = IPlaylistReaderMP4::Create(this);
-				DispatchEvent(FMetricEvent::ReportOpenSource(URL));
-				ManifestReader->LoadAndParse(URL);
 			}
 			else if (mimeType == Playlist::MIMETypeDASH)
 			{
 				ManifestReader = IPlaylistReaderDASH::Create(this);
-				DispatchEvent(FMetricEvent::ReportOpenSource(URL));
-				ManifestReader->LoadAndParse(URL);
 			}
 			else
 			{
@@ -134,6 +129,12 @@ void FAdaptiveStreamingPlayer::InternalLoadManifest(const FString& URL, const FS
 				err.SetCode(INTERR_UNSUPPORTED_FORMAT);
 				PostError(err);
 			}
+
+			if (ManifestReader.IsValid())
+			{
+				DispatchEvent(FMetricEvent::ReportOpenSource(URL));
+				ManifestReader->LoadAndParse(URL);
+			}
 		}
 		else
 		{
@@ -176,11 +177,11 @@ bool FAdaptiveStreamingPlayer::SelectManifest()
 			PlaybackState.SetTimelineRange(NewPresentation->GetTotalTimeRange());
 			PlaybackState.SetDuration(NewPresentation->GetDuration());
 
-			TArray<FStreamMetadata> VideoStreamMetadata;
-			TArray<FStreamMetadata> AudioStreamMetadata;
-			NewPresentation->GetStreamMetadata(VideoStreamMetadata, EStreamType::Video);
-			NewPresentation->GetStreamMetadata(AudioStreamMetadata, EStreamType::Audio);
-			PlaybackState.SetStreamMetadata(VideoStreamMetadata, AudioStreamMetadata);
+			TArray<FTrackMetadata> VideoTrackMetadata;
+			TArray<FTrackMetadata> AudioTrackMetadata;
+			NewPresentation->GetTrackMetadata(VideoTrackMetadata, EStreamType::Video);
+			NewPresentation->GetTrackMetadata(AudioTrackMetadata, EStreamType::Audio);
+			PlaybackState.SetTrackMetadata(VideoTrackMetadata, AudioTrackMetadata);
 			PlaybackState.SetHaveMetadata(true);
 
 			Manifest = NewPresentation;
@@ -210,22 +211,15 @@ bool FAdaptiveStreamingPlayer::SelectManifest()
 
 void FAdaptiveStreamingPlayer::UpdateManifest()
 {
-	check(Manifest.IsValid());
-
-	TArray<FTimespan> SeekablePositions;
-
-	PlaybackState.SetSeekableRange(Manifest->GetSeekableTimeRange());
-	Manifest->GetSeekablePositions(SeekablePositions);
-	PlaybackState.SetSeekablePositions(SeekablePositions);
-	PlaybackState.SetTimelineRange(Manifest->GetTotalTimeRange());
-	PlaybackState.SetDuration(Manifest->GetDuration());
-
-	TArray<FStreamMetadata> VideoStreamMetadata;
-	TArray<FStreamMetadata> AudioStreamMetadata;
-	Manifest->GetStreamMetadata(VideoStreamMetadata, EStreamType::Video);
-	Manifest->GetStreamMetadata(AudioStreamMetadata, EStreamType::Audio);
-	PlaybackState.SetStreamMetadata(VideoStreamMetadata, AudioStreamMetadata);
-	PlaybackState.SetHaveMetadata(true);
+	if (Manifest.IsValid())
+	{
+		TArray<FTimespan> SeekablePositions;
+		PlaybackState.SetSeekableRange(Manifest->GetSeekableTimeRange());
+		Manifest->GetSeekablePositions(SeekablePositions);
+		PlaybackState.SetSeekablePositions(SeekablePositions);
+		PlaybackState.SetTimelineRange(Manifest->GetTotalTimeRange());
+		PlaybackState.SetDuration(Manifest->GetDuration());
+	}
 }
 
 } // namespace Electra
diff --git a/Engine/Plugins/Media/ElectraPlayer/Source/ElectraPlayerRuntime/Private/Runtime/Player/AdaptiveStreamingPlayerMetrics.h b/Engine/Plugins/Media/ElectraPlayer/Source/ElectraPlayerRuntime/Private/Runtime/Player/AdaptiveStreamingPlayerMetrics.h
--- a/Engine/Plugins/Media/ElectraPlayer/Source/ElectraPlayerRuntime/Private/Runtime/Player/AdaptiveStreamingPlayerMetrics.h
+++ b/Engine/Plugins/Media/ElectraPlayer/Source/ElectraPlayerRuntime/Private/Runtime/Player/AdaptiveStreamingPlayerMetrics.h
@@ -221,6 +221,11 @@ public:
 	 */
 	virtual void ReportReceivedPlaylists() = 0;
 
+	/**
+	 * Called when the available tracks or their properties have changed.
+	 */
+	virtual void ReportTracksChanged() = 0;
+
 	/**
 	 * Called at the end of every downloaded playlist.
 	 */
diff --git a/Engine/Plugins/Media/ElectraPlayer/Source/ElectraPlayerRuntime/Private/Runtime/Player/DASH/ManifestBuilderDASH.cpp b/Engine/Plugins/Media/ElectraPlayer/Source/ElectraPlayerRuntime/Private/Runtime/Player/DASH/ManifestBuilderDASH.cpp
--- a/Engine/Plugins/Media/ElectraPlayer/Source/ElectraPlayerRuntime/Private/Runtime/Player/DASH/ManifestBuilderDASH.cpp
+++ b/Engine/Plugins/Media/ElectraPlayer/Source/ElectraPlayerRuntime/Private/Runtime/Player/DASH/ManifestBuilderDASH.cpp
@@ -43,10 +43,6 @@
 
 namespace Electra
 {
-#ifdef LOG_DASHINT_CTORDTOR
-DECLARE_LOG_CATEGORY_EXTERN(LogElectraDASHMPD, Log, All);
-DEFINE_LOG_CATEGORY(LogElectraDASHMPD);
-#endif
 
 namespace
 {
@@ -1065,6 +1061,13 @@ FErrorDetail FManifestDASHInternal::BuildAfterInitialRemoteElementDownload()
 			{
 				return CreateErrorAndLog(PlayerSessionServices, FString::Printf(TEXT("Early available period \"%s\" must not be followed by a regular period (\"%s\")"), *Periods.Last()->ID, *p->ID), ERRCODE_DASH_MPD_BUILDER_EARLY_PERIODS_MUST_BE_LAST);
 			}
+
+			// If there are previous periods then these are now clearly followed by another period, namely this one.
+			if (Periods.Num())
+			{
+				Periods.Last()->bHasFollowingPeriod = true;
+			}
+
 			// If the previous period did not have a duration to calculate the end time with we set
 			// its end time to the start time of this period.
 			if (Periods.Num() && !Periods.Last()->End.IsValid())
@@ -1407,7 +1410,7 @@ void FManifestDASHInternal::PreparePeriodAdaptationSets(TSharedPtrTS<FPeriod> Pe
 						Representation->CodecInfo.SetAspectRatio(FStreamCodecInformation::FAspectRatio(MPDAdaptationSet->GetSAR().GetNumerator(), MPDAdaptationSet->GetSAR().GetDenominator()));
 					}
 
-					// Update the audio codec in the adaptation set with that of the highest bandwidth.
+					// Update the video codec in the adaptation set with that of the highest bandwidth.
 					if (MPDRepresentation->GetBandwidth() > AdaptationSet->MaxBandwidth)
 					{
 						AdaptationSet->MaxBandwidth = (int32) MPDRepresentation->GetBandwidth();
@@ -2101,50 +2104,11 @@ FTimeValue FManifestDASHInternal::GetDefaultStartTime() const
 	return FromTo.Start;
 }
 
-
-#ifdef LOG_DASHINT_CTORDTOR
-FManifestDASHInternal::~FManifestDASHInternal()
+void FManifestDASHInternal::ClearDefaultStartTime()
 {
-	UE_LOG(LogElectraDASHMPD, Log, TEXT("~FManifestDASHInternal(%p)"), this);
+	URLFragmentComponents.Empty();
 }
 
-FManifestDASHInternal::FManifestDASHInternal()
-{
-	UE_LOG(LogElectraDASHMPD, Log, TEXT("FManifestDASHInternal(%p)"), this);
-}
-
-FManifestDASHInternal::FPeriod::FPeriod()
-{
-	UE_LOG(LogElectraDASHMPD, Log, TEXT("FManifestDASHInternal::FPeriod(%p)"), this);
-}
-
-FManifestDASHInternal::FPeriod::~FPeriod()
-{
-	UE_LOG(LogElectraDASHMPD, Log, TEXT("FManifestDASHInternal::~FPeriod(%p)"), this);
-}
-
-FManifestDASHInternal::FAdaptationSet::FAdaptationSet()
-{
-	UE_LOG(LogElectraDASHMPD, Log, TEXT("FManifestDASHInternal::FAdaptationSet(%p)"), this);
-}
-
-FManifestDASHInternal::FAdaptationSet::~FAdaptationSet()
-{
-	UE_LOG(LogElectraDASHMPD, Log, TEXT("FManifestDASHInternal::~FAdaptationSet(%p)"), this);
-}
-
-FManifestDASHInternal::FRepresentation::FRepresentation()
-{
-	UE_LOG(LogElectraDASHMPD, Log, TEXT("FManifestDASHInternal::FRepresentation(%p)"), this);
-}
-
-FManifestDASHInternal::FRepresentation::~FRepresentation()
-{
-	UE_LOG(LogElectraDASHMPD, Log, TEXT("FManifestDASHInternal::~FRepresentation(%p)"), this);
-}
-
-
-#endif
 
 
 } // namespace Electra
diff --git a/Engine/Plugins/Media/ElectraPlayer/Source/ElectraPlayerRuntime/Private/Runtime/Player/DASH/ManifestDASH.cpp b/Engine/Plugins/Media/ElectraPlayer/Source/ElectraPlayerRuntime/Private/Runtime/Player/DASH/ManifestDASH.cpp
--- a/Engine/Plugins/Media/ElectraPlayer/Source/ElectraPlayerRuntime/Private/Runtime/Player/DASH/ManifestDASH.cpp
+++ b/Engine/Plugins/Media/ElectraPlayer/Source/ElectraPlayerRuntime/Private/Runtime/Player/DASH/ManifestDASH.cpp
@@ -98,7 +98,7 @@ public:
 	//
 	virtual void SetStreamPreferences(const FStreamPreferences& Preferences) override;
 	virtual EReadyState GetReadyState() override;
-	virtual void PrepareForPlay(const FParamDict& Options) override;
+	virtual void PrepareForPlay() override;
 	virtual TSharedPtrTS<ITimelineMediaAsset> GetMediaAsset() const override;
 	virtual void SelectStream(const FString& AdaptationSetID, const FString& RepresentationID) override;
 	virtual IManifest::FResult GetStartingSegment(TSharedPtrTS<IStreamSegment>& OutSegment, const FPlayStartPosition& StartPosition, IManifest::ESearchType SearchType) override;
@@ -114,7 +114,13 @@ private:
 	TSharedPtrTS<FManifestDASHInternal> GetCurrentManifest() const;
 
 	TSharedPtrTS<IPlaybackAssetRepresentation> GetRepresentationFromAdaptationByMaxBandwidth(TWeakPtrTS<IPlaybackAssetAdaptationSet> InAdaptationSet, int32 NotExceedingBandwidth);
-	IManifest::FResult GetNextOrRetrySegment(TSharedPtrTS<IStreamSegment>& OutSegment, TSharedPtrTS<const IStreamSegment> InCurrentSegment, bool bForRetry);
+	enum class ENextSegType
+	{
+		SamePeriodNext,
+		SamePeriodRetry,
+		NextPeriod
+	};
+	IManifest::FResult GetNextOrRetrySegment(TSharedPtrTS<IStreamSegment>& OutSegment, TSharedPtrTS<const IStreamSegment> InCurrentSegment, ENextSegType InNextType);
 
 	IPlayerSessionServices* PlayerSessionServices = nullptr;
 	EReadyState ReadyState = EReadyState::NotReady;
@@ -146,7 +152,6 @@ FManifestDASH::FManifestDASH(IPlayerSessionServices* InSessionServices, TSharedP
 	: PlayerSessionServices(InSessionServices)
 	, CurrentManifest(InManifest)
 {
-	bHaveCurrentMetadata = false;
 }
 
 FManifestDASH::~FManifestDASH()
@@ -207,6 +212,15 @@ FTimeValue FManifestDASH::GetDefaultStartTime() const
 	return Manifest.IsValid() ? Manifest->GetDefaultStartTime() : FTimeValue();
 }
 
+void FManifestDASH::ClearDefaultStartTime()
+{
+	TSharedPtrTS<FManifestDASHInternal> Manifest(CurrentManifest);
+	if (Manifest.IsValid())
+	{
+		Manifest->ClearDefaultStartTime();
+	}
+}
+
 int64 FManifestDASH::GetDefaultStartingBitrate() const
 {
 	return 2 * 1000 * 1000;
@@ -223,59 +237,15 @@ FTimeValue FManifestDASH::GetMinBufferTime() const
 	return FTimeValue();
 }
 
-void FManifestDASH::GetStreamMetadata(TArray<FStreamMetadata>& OutMetadata, EStreamType StreamType) const
+void FManifestDASH::GetTrackMetadata(TArray<FTrackMetadata>& OutMetadata, EStreamType StreamType) const
 {
-	if (!bHaveCurrentMetadata)
-	{
-		bHaveCurrentMetadata = true;
-		TSharedPtrTS<FManifestDASHInternal> Manifest(CurrentManifest);
-		if (Manifest.IsValid() && Manifest->GetPeriods().Num())
-		{
-			Manifest->PreparePeriodAdaptationSets(Manifest->GetPeriods()[0], false);
-			const TArray<TSharedPtrTS<FManifestDASHInternal::FAdaptationSet>>& AdaptationSets = Manifest->GetPeriods()[0]->GetAdaptationSets();
-			for(int32 nA=0; nA<AdaptationSets.Num(); ++nA)
-			{
-				FStreamMetadata sm;
-				sm.CodecInformation = AdaptationSets[nA]->GetCodec();
-				sm.PlaylistID = Manifest->GetPeriods()[0]->GetID();
-				sm.Bandwidth = AdaptationSets[nA]->GetMaxBandwidth();
-				sm.StreamUniqueID = 0;
-				sm.LanguageCode = AdaptationSets[nA]->GetLanguage();
-				switch(AdaptationSets[nA]->GetCodec().GetStreamType())
-				{
-					case EStreamType::Video:
-					{
-						CurrentMetadataVideo.Emplace(MoveTemp(sm));
-						break;
-					}
-					case EStreamType::Audio:
-					{
-						CurrentMetadataAudio.Emplace(MoveTemp(sm));
-						break;
-					}
-					case EStreamType::Subtitle:
-					{
-						CurrentMetadataSubtitle.Emplace(MoveTemp(sm));
-						break;
-					}
-				}
-			}
-		}
-	}
-
-	// At present we return metadata from the first period only as every period can have totally different
-	// number of streams and even codecs. There is no commonality between periods.
-	if (StreamType == EStreamType::Video)
-	{
-		OutMetadata = CurrentMetadataVideo;
-	}
-	else if (StreamType == EStreamType::Audio)
-	{
-		OutMetadata = CurrentMetadataAudio;
-	}
-	else if (StreamType == EStreamType::Subtitle)
+	TSharedPtrTS<FManifestDASHInternal> Manifest(CurrentManifest);
+	if (Manifest.IsValid() && Manifest->GetPeriods().Num())
 	{
-		OutMetadata = CurrentMetadataSubtitle;
+		// At present we return metadata from the first period only as every period can have totally different
+		// number of streams and even codecs. There is no commonality between periods.
+		Manifest->PreparePeriodAdaptationSets(Manifest->GetPeriods()[0], false);
+		Manifest->GetPeriods()[0]->GetMetaData(OutMetadata, StreamType);
 	}
 }
 
@@ -297,75 +267,95 @@ IManifest::FResult FManifestDASH::FindPlayPeriod(TSharedPtrTS<IPlayPeriod>& OutP
 		return IManifest::FResult(IManifest::FResult::EType::NotLoaded);
 	}
 
+	const TArray<TSharedPtrTS<FManifestDASHInternal::FPeriod>>& Periods = Manifest->GetPeriods();
+	TSharedPtrTS<FManifestDASHInternal::FPeriod> SelectedPeriod;
+	if (Periods.Num() == 0)
+	{
+		return IManifest::FResult(IManifest::FResult::EType::TryAgainLater).RetryAfterMilliseconds(1000);
+	}
+
 	FTimeValue StartTime = StartPosition.Time;
 
 	// All time values we communicate to the outside - and therefor get from the outside - are offset by the availabilityStartTime.
 	StartTime -= Manifest->GetAnchorTime();
 
+	// Quick out if the time falls outside the presentation.
+	FTimeValue TotalEndTime = Manifest->GetLastPeriodEndTime();
+	if (StartTime >= TotalEndTime)
+	{
+		return IManifest::FResult(IManifest::FResult::EType::PastEOS);
+	}
+	// If the time to search for is before the start of the first period we use the first period!
+	else if (StartTime < Periods[0]->GetStart())
+	{
+		LogMessage(PlayerSessionServices, IInfoLog::ELevel::Verbose, FString::Printf(TEXT("Start time is before the start of the first period. Using first period!")));
+		SelectedPeriod = Periods[0];
+	}
+
 	// Find the period into which the start time falls.
-	TSharedPtrTS<FManifestDASHInternal::FPeriod> SelectedPeriod;
-	const TArray<TSharedPtrTS<FManifestDASHInternal::FPeriod>>& Periods = Manifest->GetPeriods();
 	for(int32 nPeriod=0; !SelectedPeriod.IsValid() && nPeriod<Periods.Num(); ++nPeriod)
 	{
 		if (!Periods[nPeriod]->GetIsEarlyPeriod())
 		{
-			if (StartTime >= Periods[nPeriod]->GetStart())
+			FTimeValue PeriodStartTime = Periods[nPeriod]->GetStart();
+			FTimeValue PeriodEndTime = Periods[nPeriod]->GetEnd();
+			// When the period end time is not valid it must be the last period of a Live presentation
+			if (!PeriodEndTime.IsValid())
+			{
+				PeriodEndTime.SetToPositiveInfinity();
+			}
+			// Does the time fall into this period?
+			if (StartTime >= PeriodStartTime && StartTime < PeriodEndTime)
 			{
-				FTimeValue PeriodEndTime = Periods[nPeriod]->GetEnd();
-				// When the period end time is not valid it must be the last period of a Live presentation
-				if (!PeriodEndTime.IsValid())
+				FTimeValue DiffToNextPeriod = nPeriod + 1 < Periods.Num() && !Periods[nPeriod+1]->GetIsEarlyPeriod() ? Periods[nPeriod + 1]->GetStart() - StartTime : FTimeValue::GetPositiveInfinity();
+				if (!DiffToNextPeriod.IsValid())
 				{
-					PeriodEndTime.SetToPositiveInfinity();
+					DiffToNextPeriod = FTimeValue::GetPositiveInfinity();
 				}
-				// Does the time fall into this period?
-				if (StartTime < PeriodEndTime)
+				FTimeValue DiffToStart = StartTime - PeriodStartTime;
+				switch(SearchType)
 				{
-					FTimeValue DiffToNextPeriod = nPeriod + 1 < Periods.Num() && !Periods[nPeriod+1]->GetIsEarlyPeriod() ? Periods[nPeriod + 1]->GetStart() - StartTime : FTimeValue::GetPositiveInfinity();
-					FTimeValue DiffToEnd = PeriodEndTime - StartTime;
-					switch(SearchType)
+					case ESearchType::Closest:
 					{
-						case ESearchType::Closest:
+						// If there is a following period whose start is closer to the current position than the start of this period and the position
+						// is within a short duration to the next we jump ahead to the next.
+						if (DiffToNextPeriod < DiffToStart && DiffToNextPeriod.GetAsMilliseconds() < 1000)
 						{
-							// Closest only looks at the distance to the following period. It will never take us to the preceeding period since it is unlikely
-							// the intention (for closest) is to start playback at the end of a period only to transition into this one.
-							if (DiffToEnd.IsValid() && DiffToNextPeriod <= DiffToEnd && nPeriod + 1 < Periods.Num())
-							{
-								SelectedPeriod = Periods[nPeriod + 1];
-							}
-							else
-							{
-								SelectedPeriod = Periods[nPeriod];
-							}
-							break;
+							SelectedPeriod = Periods[nPeriod + 1];
 						}
-						case ESearchType::Before:
-						case ESearchType::Same:
-						case ESearchType::After:
+						else
 						{
-							// Before, Same and After have no meaning when looking for a period. The period the start time falls into is the one to use.
 							SelectedPeriod = Periods[nPeriod];
-							break;
 						}
-						case ESearchType::StrictlyAfter:
+						break;
+					}
+					case ESearchType::Before:
+					case ESearchType::Same:
+					case ESearchType::After:
+					{
+						// Before, Same and After have no meaning when looking for a period. The period the start time falls into is the one to use.
+						SelectedPeriod = Periods[nPeriod];
+						break;
+					}
+					case ESearchType::StrictlyAfter:
+					{
+						if (!DiffToNextPeriod.IsInfinity())
 						{
-							if (DiffToNextPeriod.IsValid() && !DiffToNextPeriod.IsInfinity())
-							{
-								SelectedPeriod = Periods[nPeriod + 1];
-							}
-							break;
+							SelectedPeriod = Periods[nPeriod + 1];
 						}
-						case ESearchType::StrictlyBefore:
+						break;
+					}
+					case ESearchType::StrictlyBefore:
+					{
+						if (nPeriod)
 						{
-							if (nPeriod)
-							{
-								SelectedPeriod = Periods[nPeriod - 1];
-							}
-							break;
+							SelectedPeriod = Periods[nPeriod - 1];
 						}
+						break;
 					}
-					// Time fell into this period. We have either found a candidate or not. We're done either way.
-					break;
 				}
+				// Time fell into this period. We have either found a candidate or not. We're done either way.
+				break;
 			}
 		}
 	}
@@ -401,7 +391,23 @@ IManifest::FResult FManifestDASH::FindPlayPeriod(TSharedPtrTS<IPlayPeriod>& OutP
 			return IManifest::FResult(IManifest::FResult::EType::TryAgainLater).RetryAfterMilliseconds(100);
 		}
 	}
-	return IManifest::FResult();
+	// Ok, we made sure to use the first period if the start time is less than that of the first one.
+	// Coming here can only mean that no period was found which can only means that the time is past the last one.
+	// Which means that the duration of the last period is actually less than what MPD@mediaPresentationDuration was saying.
+	// So, in a nutshell, we have reached the end.
+	return IManifest::FResult(IManifest::FResult::EType::PastEOS);
+}
+
+IManifest::FResult FManifestDASH::FindNextPlayPeriod(TSharedPtrTS<IPlayPeriod>& OutPlayPeriod, TSharedPtrTS<const IStreamSegment> InCurrentSegment)
+{
+	const FStreamSegmentRequestFMP4DASH* CurrentRequest = static_cast<const FStreamSegmentRequestFMP4DASH*>(InCurrentSegment.Get());
+	if (CurrentRequest)
+	{
+		FPlayStartPosition SearchTime;
+		SearchTime.Time = CurrentRequest->AST + CurrentRequest->PeriodStart + FTimeValue(CurrentRequest->Segment.Time - CurrentRequest->Segment.PTO, CurrentRequest->Segment.Timescale);
+		return FindPlayPeriod(OutPlayPeriod, SearchTime, IManifest::ESearchType::StrictlyAfter);
+	}
+	return IManifest::FResult(IManifest::FResult::EType::PastEOS);
 }
 
 
@@ -431,45 +437,41 @@ IManifest::IPlayPeriod::EReadyState FDASHPlayPeriod::GetReadyState()
 	return ReadyState;
 }
 
-void FDASHPlayPeriod::PrepareForPlay(const FParamDict& Options)
+void FDASHPlayPeriod::PrepareForPlay()
 {
 	TSharedPtrTS<FManifestDASHInternal> Manifest = GetCurrentManifest();
 	TSharedPtrTS<FManifestDASHInternal::FPeriod> Period = Manifest.IsValid() ? Manifest->GetPeriodByUniqueID(PeriodID) : nullptr;
-	if (!Period.IsValid())
+	if (Period.IsValid())
 	{
-		LogMessage(PlayerSessionServices, IInfoLog::ELevel::Verbose, FString::Printf(TEXT("Period in PrepareToPlay() is not or no longer valid. Must reselect.")));
-		ReadyState = EReadyState::MustReselect;
-		return;
-	}
-
-	// Prepare the adaptation sets and periods.
-	Manifest->PreparePeriodAdaptationSets(Period, false);
+		// Prepare the adaptation sets and periods.
+		Manifest->PreparePeriodAdaptationSets(Period, false);
 	
-	// We need to select one adaptation set per stream type we wish to play.
-	int32 NumVideoAdaptationSets = Period->GetNumberOfAdaptationSets(EStreamType::Video);
-	if (NumVideoAdaptationSets > 0)
-	{
-		// For now pick the first one!
-		TSharedPtrTS<IPlaybackAssetAdaptationSet> VideoAS = Period->GetAdaptationSetByTypeAndIndex(EStreamType::Video, 0);
-		ActiveVideoAdaptationSetID = VideoAS->GetUniqueIdentifier();
+		// We need to select one adaptation set per stream type we wish to play.
+		int32 NumVideoAdaptationSets = Period->GetNumberOfAdaptationSets(EStreamType::Video);
+		if (NumVideoAdaptationSets > 0)
+		{
+			// For now pick the first one!
+			TSharedPtrTS<IPlaybackAssetAdaptationSet> VideoAS = Period->GetAdaptationSetByTypeAndIndex(EStreamType::Video, 0);
+			ActiveVideoAdaptationSetID = VideoAS->GetUniqueIdentifier();
 
-		TSharedPtrTS<IPlaybackAssetRepresentation> VideoRepr = GetRepresentationFromAdaptationByMaxBandwidth(VideoAS, 2*1000*1000);
-		ActiveVideoRepresentationID = VideoRepr->GetUniqueIdentifier();
-	}
+			TSharedPtrTS<IPlaybackAssetRepresentation> VideoRepr = GetRepresentationFromAdaptationByMaxBandwidth(VideoAS, 2*1000*1000);
+			ActiveVideoRepresentationID = VideoRepr->GetUniqueIdentifier();
+		}
 
-	int32 NumAudioAdaptationSets = Period->GetNumberOfAdaptationSets(EStreamType::Audio);
-	if (NumAudioAdaptationSets > 0)
-	{
-		// For now pick the first one!
-		TSharedPtrTS<IPlaybackAssetAdaptationSet> AudioAS = Period->GetAdaptationSetByTypeAndIndex(EStreamType::Audio, 0);
-		ActiveAudioAdaptationSetID = AudioAS->GetUniqueIdentifier();
+		int32 NumAudioAdaptationSets = Period->GetNumberOfAdaptationSets(EStreamType::Audio);
+		if (NumAudioAdaptationSets > 0)
+		{
+			// For now pick the first one!
+			TSharedPtrTS<IPlaybackAssetAdaptationSet> AudioAS = Period->GetAdaptationSetByTypeAndIndex(EStreamType::Audio, 0);
+			ActiveAudioAdaptationSetID = AudioAS->GetUniqueIdentifier();
 
-		TSharedPtrTS<IPlaybackAssetRepresentation> AudioRepr = GetRepresentationFromAdaptationByMaxBandwidth(AudioAS, 128 * 1000);
-		ActiveAudioRepresentationID = AudioRepr->GetUniqueIdentifier();
-	}
+			TSharedPtrTS<IPlaybackAssetRepresentation> AudioRepr = GetRepresentationFromAdaptationByMaxBandwidth(AudioAS, 128 * 1000);
+			ActiveAudioRepresentationID = AudioRepr->GetUniqueIdentifier();
+		}
 
-//	ReadyState = EReadyState::Preparing;
-	ReadyState = EReadyState::IsReady;
+	//	ReadyState = EReadyState::Preparing;
+		ReadyState = EReadyState::IsReady;
+	}
 }
 
 TSharedPtrTS<ITimelineMediaAsset> FDASHPlayPeriod::GetMediaAsset() const
@@ -585,6 +587,7 @@ IManifest::FResult FDASHPlayPeriod::GetStartingSegment(TSharedPtrTS<IStreamSegme
 			{
 				SearchOpt.PeriodDuration = Manifest->GetLastPeriodEndTime() - AST;
 			}
+			SearchOpt.bHasFollowingPeriod = Period->GetHasFollowingPeriod();
 			SearchOpt.SearchType = SearchType;
 			FManifestDASHInternal::FRepresentation::ESearchResult SearchResult = Repr->FindSegment(PlayerSessionServices, SegmentInfo, RemoteElementLoadRequests, SearchOpt);
 			if (SearchResult == FManifestDASHInternal::FRepresentation::ESearchResult::NeedElement)
@@ -646,7 +649,7 @@ IManifest::FResult FDASHPlayPeriod::GetStartingSegment(TSharedPtrTS<IStreamSegme
 					// This is wasteful and prevents fast startup, so we set the start time to the beginning of the segment.
 					SegmentInfo.MediaLocalFirstAUTime = SegmentInfo.Time;
 
-					StartTime = Period->GetStart() + FTimeValue().SetFromND(SegmentInfo.Time - SegmentInfo.PTO, SegmentInfo.Timescale);
+					StartTime = FTimeValue().SetFromND(SegmentInfo.Time - SegmentInfo.PTO, SegmentInfo.Timescale);
 					SearchType = IManifest::ESearchType::Before;
 				}
 
@@ -714,7 +717,7 @@ IManifest::FResult FDASHPlayPeriod::GetStartingSegment(TSharedPtrTS<IStreamSegme
 }
 
 
-IManifest::FResult FDASHPlayPeriod::GetNextOrRetrySegment(TSharedPtrTS<IStreamSegment>& OutSegment, TSharedPtrTS<const IStreamSegment> InCurrentSegment, bool bForRetry)
+IManifest::FResult FDASHPlayPeriod::GetNextOrRetrySegment(TSharedPtrTS<IStreamSegment>& OutSegment, TSharedPtrTS<const IStreamSegment> InCurrentSegment, ENextSegType InNextType)
 {
 	TSharedPtrTS<const FStreamSegmentRequestFMP4DASH> Current = StaticCastSharedPtr<const FStreamSegmentRequestFMP4DASH>(InCurrentSegment);
 	if (Current->bIsInitialStartRequest)
@@ -755,7 +758,7 @@ IManifest::FResult FDASHPlayPeriod::GetNextOrRetrySegment(TSharedPtrTS<IStreamSe
 	FManifestDASHInternal::FSegmentInformation SegmentInfo;
 	FManifestDASHInternal::FSegmentSearchOption SearchOpt;
 	TArray<TWeakPtrTS<FMPDLoadRequestDASH>> RemoteElementLoadRequests;
-	if (!bForRetry)
+	if (InNextType == ENextSegType::SamePeriodNext)
 	{
 		// Set up the search time as the time three quarters into the current segment.
 		// This is to make sure the time is sufficiently large that it won't be affected by rounding errors in timescale conversions.
@@ -763,13 +766,19 @@ IManifest::FResult FDASHPlayPeriod::GetNextOrRetrySegment(TSharedPtrTS<IStreamSe
 		SearchOpt.PeriodDuration = Period->GetDuration();
 		SearchOpt.SearchType = IManifest::ESearchType::After;
 	}
-	else
+	else if (InNextType == ENextSegType::SamePeriodRetry)
 	{
 		// Use the same period local time for the retry representation as was used to locate the current segment.
 		SearchOpt.PeriodLocalTime.SetFromND(Current->Segment.Time - Current->Segment.PTO, Current->Segment.Timescale);
 		SearchOpt.PeriodDuration = Period->GetDuration();
 		SearchOpt.SearchType = IManifest::ESearchType::Closest;
 	}
+	else //if (InNextType == ENextSegType::NextPeriod)
+	{
+		SearchOpt.PeriodLocalTime.SetToZero();
+		SearchOpt.PeriodDuration = Period->GetDuration();
+		SearchOpt.SearchType = IManifest::ESearchType::Closest;
+	}
 	FTimeValue AST = Manifest->GetAnchorTime();
 	bool bUsesAST = Manifest->UsesAST();
 	bool bIsStaticType = Manifest->IsStaticType() || Manifest->IsEventType();
@@ -777,6 +786,7 @@ IManifest::FResult FDASHPlayPeriod::GetNextOrRetrySegment(TSharedPtrTS<IStreamSe
 	{
 		SearchOpt.PeriodDuration = Manifest->GetLastPeriodEndTime() - AST;
 	}
+	SearchOpt.bHasFollowingPeriod = Period->GetHasFollowingPeriod();
 
 	FManifestDASHInternal::FRepresentation::ESearchResult SearchResult = Repr->FindSegment(PlayerSessionServices, SegmentInfo, RemoteElementLoadRequests, SearchOpt);
 	if (SearchResult == FManifestDASHInternal::FRepresentation::ESearchResult::NeedElement)
@@ -843,16 +853,19 @@ IManifest::FResult FDASHPlayPeriod::GetNextOrRetrySegment(TSharedPtrTS<IStreamSe
 		{
 			SegmentRequest->bInsertFillerData = true;
 		}
-		// Because we are searching for the next segment we do not want any first access units to be truncated.
-		// We keep the current media local AU time for the case where with <SegmentTemplate> addressing we get greatly varying
-		// segment durations from the fixed value (up to +/- 50% variation are allowed!) and the current segment did not actually
-		// have any access units we wanted to have! In that case it is possible that this new segment would also have some initial
-		// access units outside the time we want. By retaining the initial value this is addressed.
-		// We do need to translate the value between potentially different timescales and potentially different local media times.
-		SegmentInfo.MediaLocalFirstAUTime = FTimeFraction(Current->Segment.MediaLocalFirstAUTime - Current->Segment.PTO, Current->Segment.Timescale).GetAsTimebase(SegmentInfo.Timescale) + SegmentInfo.PTO;
-		SegmentRequest->Segment = MoveTemp(SegmentInfo);
+		
+		if (InNextType != ENextSegType::NextPeriod)
+		{
+			// Because we are searching for the next segment we do not want any first access units to be truncated.
+			// We keep the current media local AU time for the case where with <SegmentTemplate> addressing we get greatly varying
+			// segment durations from the fixed value (up to +/- 50% variation are allowed!) and the current segment did not actually
+			// have any access units we wanted to have! In that case it is possible that this new segment would also have some initial
+			// access units outside the time we want. By retaining the initial value this is addressed.
+			// We do need to translate the value between potentially different timescales and potentially different local media times.
+			SegmentInfo.MediaLocalFirstAUTime = FTimeFraction(Current->Segment.MediaLocalFirstAUTime - Current->Segment.PTO, Current->Segment.Timescale).GetAsTimebase(SegmentInfo.Timescale) + SegmentInfo.PTO;
+		}
 		// For a retry request we have to increate the retry count to give up after n failed attempts.
-		if (bForRetry)
+		if (InNextType == ENextSegType::SamePeriodRetry)
 		{
 			SegmentRequest->NumOverallRetries = Current->NumOverallRetries + 1;
 		}
@@ -864,6 +877,7 @@ IManifest::FResult FDASHPlayPeriod::GetNextOrRetrySegment(TSharedPtrTS<IStreamSe
 			SegmentRequest->bWarnedAboutTimescale = Current->bWarnedAboutTimescale;
 		}
 
+		SegmentRequest->Segment = MoveTemp(SegmentInfo);
 		OutSegment = MoveTemp(SegmentRequest);
 		return IManifest::FResult(IManifest::FResult::EType::Found);
 	}
@@ -883,11 +897,21 @@ IManifest::FResult FDASHPlayPeriod::GetNextSegment(TSharedPtrTS<IStreamSegment>&
 	// Did the stream reader see a 'lmsg' brand on this segment?
 	// If so then this stream has ended and there will not be a next segment.
 	const FStreamSegmentRequestFMP4DASH* CurrentRequest = static_cast<const FStreamSegmentRequestFMP4DASH*>(InCurrentSegment.Get());
-	if (CurrentRequest->Segment.bSawLMSG)
+
+	// Check if we moved across a period.
+	if (CurrentRequest->Period->GetUniqueIdentifier().Equals(PeriodID))
 	{
-		return IManifest::FResult(IManifest::FResult::EType::PastEOS);
+		if (CurrentRequest->Segment.bSawLMSG)
+		{
+			return IManifest::FResult(IManifest::FResult::EType::PastEOS);
+		}
+		return GetNextOrRetrySegment(OutSegment, InCurrentSegment, ENextSegType::SamePeriodNext);
+	}
+	else
+	{
+		// Moved into a new period. This here is the new period.
+		return GetNextOrRetrySegment(OutSegment, InCurrentSegment, ENextSegType::NextPeriod);
 	}
-	return GetNextOrRetrySegment(OutSegment, InCurrentSegment, false);
 }
 
 IManifest::FResult FDASHPlayPeriod::GetRetrySegment(TSharedPtrTS<IStreamSegment>& OutSegment, TSharedPtrTS<const IStreamSegment> InCurrentSegment, bool bReplaceWithFillerData)
@@ -908,13 +932,45 @@ IManifest::FResult FDASHPlayPeriod::GetRetrySegment(TSharedPtrTS<IStreamSegment>
 		OutSegment = NewRequest;
 		return IManifest::FResult(IManifest::FResult::EType::Found);
 	}
-	return GetNextOrRetrySegment(OutSegment, InCurrentSegment, true);
+	return GetNextOrRetrySegment(OutSegment, InCurrentSegment, ENextSegType::SamePeriodRetry);
 }
 
 
 IManifest::FResult FDASHPlayPeriod::GetLoopingSegment(TSharedPtrTS<IStreamSegment>& OutSegment, FPlayerLoopState& InOutLoopState, const TMultiMap<EStreamType, TSharedPtrTS<IStreamSegment>>& InFinishedSegments, const FPlayStartPosition& StartPosition, IManifest::ESearchType SearchType)
 {
-	return IManifest::FResult(IManifest::FResult::EType::NotFound).SetErrorDetail(FErrorDetail().SetMessage("Looping is not currently supported!"));
+	if (InFinishedSegments.Num())
+	{
+		// Go over all finished segments and get the largest next expected timestamp from all of them.
+		FTimeValue LargestNextExpectedTimestamp(FTimeValue::GetZero());
+		for(auto It = InFinishedSegments.CreateConstIterator(); It; ++It)
+		{
+			const FStreamSegmentRequestFMP4DASH* FinishedRequest = static_cast<const FStreamSegmentRequestFMP4DASH*>(It->Value.Get());
+			if (FinishedRequest && FinishedRequest->NextLargestExpectedTimestamp > LargestNextExpectedTimestamp)
+			{
+				LargestNextExpectedTimestamp = FinishedRequest->NextLargestExpectedTimestamp;
+			}
+		}
+		IManifest::FResult Result = GetStartingSegment(OutSegment, StartPosition, SearchType);
+		if (Result.GetType() == IManifest::FResult::EType::Found)
+		{
+			FStreamSegmentRequestFMP4DASH* LoopRequest = static_cast<FStreamSegmentRequestFMP4DASH*>(OutSegment.Get());
+			InOutLoopState.bLoopEnabled = true;
+			InOutLoopState.LoopBasetime = LargestNextExpectedTimestamp;	// This is the absolute playback time at which the loop will occur
+			++InOutLoopState.LoopCount;
+			LoopRequest->PlayerLoopState = InOutLoopState;
+			// We have to subtract all the values that are added to the access units from the loop base time
+			// to avoid adding them twice.
+			LoopRequest->PlayerLoopState.LoopBasetime -= LoopRequest->PeriodStart + LoopRequest->AST + LoopRequest->AdditionalAdjustmentTime;			
+			// Set the loop state in the dependent streams as well.
+			for(int32 nDep=0, nDepMax=LoopRequest->DependentStreams.Num(); nDep<nDepMax; ++nDep)
+			{
+				LoopRequest->DependentStreams[nDep]->PlayerLoopState = LoopRequest->PlayerLoopState;
+			}
+			return Result;
+		}
+	}
+	// Return past EOS when we can't loop to indicate we're really done now.
+	return IManifest::FResult(IManifest::FResult::EType::PastEOS);
 }
 
 
@@ -1710,6 +1766,7 @@ FManifestDASHInternal::FRepresentation::ESearchResult FManifestDASHInternal::FRe
 		OutSegmentInfo.MediaLocalLastAUTime = MediaLocalPeriodEnd;
 		OutSegmentInfo.Timescale = SidxTimescale;
 		OutSegmentInfo.ATO = ATO;
+		OutSegmentInfo.bIsLastInPeriod = CurrentT + CurrentD >= MediaLocalPeriodEnd;
 		if (!PrepareDownloadURLs(InPlayerSessionServices, OutSegmentInfo, SegmentBase))
 		{
 			bIsUsable = false;
@@ -1861,6 +1918,7 @@ FManifestDASHInternal::FRepresentation::ESearchResult FManifestDASHInternal::FRe
 	OutSegmentInfo.MediaLocalLastAUTime = MediaLocalPeriodDuration + PTO;
 	OutSegmentInfo.Timescale = MPDTimescale;
 	OutSegmentInfo.bMayBeMissing = SegmentNum + 1 >= MaxSegmentsInPeriod;
+	OutSegmentInfo.bIsLastInPeriod = OutSegmentInfo.bMayBeMissing && InSearchOptions.bHasFollowingPeriod;
 	OutSegmentInfo.ATO = ATO;
 	if (!PrepareDownloadURLs(InPlayerSessionServices, OutSegmentInfo, SegmentTemplate))
 	{
@@ -2180,6 +2238,7 @@ FManifestDASHInternal::FRepresentation::ESearchResult FManifestDASHInternal::FRe
 		OutSegmentInfo.MediaLocalLastAUTime = MediaLocalPeriodEnd;
 		OutSegmentInfo.Timescale = MPDTimescale;
 		OutSegmentInfo.bMayBeMissing = CurrentT + CurrentD >= MediaLocalPeriodEnd;
+		OutSegmentInfo.bIsLastInPeriod = OutSegmentInfo.bMayBeMissing && InSearchOptions.bHasFollowingPeriod;
 		if (bIsCurrentlyAGap)
 		{
 			OutSegmentInfo.bMayBeMissing = true;
diff --git a/Engine/Plugins/Media/ElectraPlayer/Source/ElectraPlayerRuntime/Private/Runtime/Player/DASH/ManifestDASH.h b/Engine/Plugins/Media/ElectraPlayer/Source/ElectraPlayerRuntime/Private/Runtime/Player/DASH/ManifestDASH.h
--- a/Engine/Plugins/Media/ElectraPlayer/Source/ElectraPlayerRuntime/Private/Runtime/Player/DASH/ManifestDASH.h
+++ b/Engine/Plugins/Media/ElectraPlayer/Source/ElectraPlayerRuntime/Private/Runtime/Player/DASH/ManifestDASH.h
@@ -25,13 +25,15 @@ public:
 	virtual void GetSeekablePositions(TArray<FTimespan>& OutPositions) const override;
 	virtual FTimeValue GetDuration() const override;
 	virtual FTimeValue GetDefaultStartTime() const override;
+	virtual void ClearDefaultStartTime() override;
 	virtual int64 GetDefaultStartingBitrate() const override;
 	virtual FTimeValue GetMinBufferTime() const override;
-	virtual void GetStreamMetadata(TArray<FStreamMetadata>& OutMetadata, EStreamType StreamType) const override;
+	virtual void GetTrackMetadata(TArray<FTrackMetadata>& OutMetadata, EStreamType StreamType) const override;
 	virtual void UpdateDynamicRefetchCounter() override;
 	virtual IStreamReader* CreateStreamReaderHandler() override;
 
 	virtual FResult FindPlayPeriod(TSharedPtrTS<IPlayPeriod>& OutPlayPeriod, const FPlayStartPosition& StartPosition, ESearchType SearchType) override;
+	virtual FResult FindNextPlayPeriod(TSharedPtrTS<IPlayPeriod>& OutPlayPeriod, TSharedPtrTS<const IStreamSegment> CurrentSegment) override;
 
 private:
 	ELECTRA_IMPL_DEFAULT_ERROR_METHODS(DASHManifest);
@@ -40,10 +42,6 @@ private:
 
 	IPlayerSessionServices* 					PlayerSessionServices = nullptr;
 	TSharedPtrTS<FManifestDASHInternal>			CurrentManifest;
-	mutable TArray<FStreamMetadata>				CurrentMetadataVideo;
-	mutable TArray<FStreamMetadata>				CurrentMetadataAudio;
-	mutable TArray<FStreamMetadata>				CurrentMetadataSubtitle;
-	mutable bool								bHaveCurrentMetadata = false;
 	int64										CurrentPeriodAndAdaptationXLinkResolveID = 1;
 };
 
diff --git a/Engine/Plugins/Media/ElectraPlayer/Source/ElectraPlayerRuntime/Private/Runtime/Player/DASH/PlaylistReaderDASH_Internal.h b/Engine/Plugins/Media/ElectraPlayer/Source/ElectraPlayerRuntime/Private/Runtime/Player/DASH/PlaylistReaderDASH_Internal.h
--- a/Engine/Plugins/Media/ElectraPlayer/Source/ElectraPlayerRuntime/Private/Runtime/Player/DASH/PlaylistReaderDASH_Internal.h
+++ b/Engine/Plugins/Media/ElectraPlayer/Source/ElectraPlayerRuntime/Private/Runtime/Player/DASH/PlaylistReaderDASH_Internal.h
@@ -11,7 +11,6 @@
 #include "Utilities/URLParser.h"
 
 
-//#define LOG_DASHINT_CTORDTOR
 
 namespace Electra
 {
@@ -134,7 +133,7 @@ public:
 		FURL InitializationURL;
 		FURL MediaURL;
 		FTimeValue ATO;
-		int64 Time = 0;								//!< Time value T in timescale units without PTO, EPT or ELST applied
+		int64 Time = 0;								//!< Time value T in timescale units
 		int64 PTO = 0;								//!< PresentationTimeOffset
 		int64 EPTdelta = 0;
 		int64 Duration = 0;							//!< Duration of the segment. Not necessarily exact if <SegmentTemplate> is used).
@@ -145,6 +144,7 @@ public:
 		int64 MediaLocalFirstAUTime = 0;			//!< Time of the first AU to use in this segment in media local time
 		int64 MediaLocalLastAUTime = 0;				//!< Time at which the last AU to use in thie segment ends in media local time
 		uint32 Timescale = 0;						//!< Local media timescale
+		bool bIsLastInPeriod = false;				//!< true if known to be the last segment in the period.
 		bool bMayBeMissing = false;					//!< true if the last segment in <SegmentTemplate> that might not exist.
 		bool bIsMissing = false;					//!< Set to true if known to be missing.
 		bool bSawLMSG = false;						//!< Will be set to true by the stream reader if the 'lmsg' brand was found.
@@ -194,18 +194,14 @@ public:
 		FTimeValue PeriodDuration;					//!< Duration of the period. Needed to determine the number of segments in the period.
 		IManifest::ESearchType SearchType = IManifest::ESearchType::Closest;
 		int64 RequestID = 0;						//!< Sequential request ID across all segments during playback, needed to re-resolve potential UrlQueryInfo xlinks.
+		bool bHasFollowingPeriod = false;			//!< true if we know for sure there is another period following.
 	};
 
 	class FRepresentation : public IPlaybackAssetRepresentation, public TSharedFromThis<FRepresentation, ESPMode::ThreadSafe>
 	{
 	public:
-		#ifdef LOG_DASHINT_CTORDTOR
-			FRepresentation();
-			virtual ~FRepresentation();
-		#else
-			FRepresentation() = default;
-			virtual ~FRepresentation() = default;
-		#endif
+		FRepresentation() = default;
+		virtual ~FRepresentation() = default;
 
 		enum ESearchResult
 		{
@@ -215,7 +211,7 @@ public:
 			BadType,								//!< Representation is bad for some reason, most likely because it uses <SegmentList> addressing which is not supported.
 			Gone									//!< Underlying MPD Representation (held by a weak pointer) has gone and the representation is no longer accessible.
 		};
-		
+
 		ESearchResult FindSegment(IPlayerSessionServices* PlayerSessionServices, FSegmentInformation& OutSegmentInfo, TArray<TWeakPtrTS<FMPDLoadRequestDASH>>& OutRemoteElementLoadRequests, const FSegmentSearchOption& InSearchOptions);
 
 		void GetSegmentInformation(TArray<IManifest::IPlayPeriod::FSegmentInformation>& OutSegmentInformation, FTimeValue& OutAverageSegmentDuration, TSharedPtrTS<const IStreamSegment> CurrentSegment, const FTimeValue& LookAheadTime, const TSharedPtrTS<IPlaybackAssetAdaptationSet>& AdaptationSet);
@@ -282,13 +278,8 @@ public:
 	class FAdaptationSet : public IPlaybackAssetAdaptationSet
 	{
 	public:
-		#ifdef LOG_DASHINT_CTORDTOR
-			FAdaptationSet();
-			virtual ~FAdaptationSet();
-		#else
-			FAdaptationSet() = default;
-			virtual ~FAdaptationSet() = default;
-		#endif
+		FAdaptationSet() = default;
+		virtual ~FAdaptationSet() = default;
 
 		const FStreamCodecInformation& GetCodec() const								{ return Codec; }
 		const TArray<TSharedPtrTS<FRepresentation>>& GetRepresentations() const		{ return Representations; }
@@ -360,13 +351,8 @@ public:
 	class FPeriod : public ITimelineMediaAsset
 	{
 	public:
-		#ifdef LOG_DASHINT_CTORDTOR
-			FPeriod();
-			virtual ~FPeriod();
-		#else
-			FPeriod() = default;
-			virtual ~FPeriod() = default;
-		#endif
+		FPeriod() = default;
+		virtual ~FPeriod() = default;
 
 		bool GetHasBeenPrepared() const
 		{
@@ -377,6 +363,15 @@ public:
 			bHasBeenPrepared = bPrepared;
 		}
 
+		bool GetHasFollowingPeriod() const
+		{
+			return bHasFollowingPeriod;
+		}
+		void SetHasFollowingPeriod(bool bInHasFollowing)
+		{
+			bHasFollowingPeriod = bInHasFollowing;
+		}
+
 		const FString& GetID() const { return ID; }
 		const FTimeValue& GetStart() const { return Start; }
 		const FTimeValue& GetEnd() const { return End; }
@@ -439,7 +434,7 @@ public:
 			}
 			return Num;
 		}
-		virtual TSharedPtrTS<IPlaybackAssetAdaptationSet> GetAdaptationSetByTypeAndIndex(EStreamType OfStreamType, int32 AdaptationSetIndex) const
+		virtual TSharedPtrTS<IPlaybackAssetAdaptationSet> GetAdaptationSetByTypeAndIndex(EStreamType OfStreamType, int32 AdaptationSetIndex) const override
 		{
 			int32 Num=0;
 			for(int32 i=0; i<AdaptationSets.Num(); ++i)
@@ -455,6 +450,35 @@ public:
 			return nullptr;
 		}
 
+		virtual void GetMetaData(TArray<FTrackMetadata>& OutMetadata, EStreamType OfStreamType) const override
+		{
+			for(int32 i=0, iMax=GetNumberOfAdaptationSets(OfStreamType); i<iMax; ++i)
+			{
+				TSharedPtrTS<IPlaybackAssetAdaptationSet> AdaptationSet = GetAdaptationSetByTypeAndIndex(OfStreamType, i);
+				const FAdaptationSet* Adapt = static_cast<const FAdaptationSet*>(AdaptationSet.Get());
+				if (Adapt && Adapt->GetIsUsable())
+				{
+					FTrackMetadata tm;
+					tm.TrackID = Adapt->GetUniqueIdentifier();
+					tm.Language = Adapt->GetLanguage();
+					tm.HighestBandwidth = Adapt->GetMaxBandwidth();
+					tm.HighestBandwidthCodec = Adapt->GetCodec();
+					const TArray<TSharedPtrTS<FRepresentation>>& Reprs = Adapt->GetRepresentations();
+					for(int32 j=0; j<Reprs.Num(); ++j)
+					{
+						if (Reprs[j]->CanBePlayed())
+						{
+							FStreamMetadata sd;
+							sd.Bandwidth = Reprs[j]->GetBitrate();
+							sd.CodecInformation = Reprs[j]->GetCodecInformation();
+							tm.StreamDetails.Emplace(MoveTemp(sd));
+						}
+					}
+					OutMetadata.Emplace(MoveTemp(tm));
+				}
+			}
+		}
+
 	private:
 		friend class FManifestDASHInternal;
 		TArray<TSharedPtrTS<FAdaptationSet>> AdaptationSets;
@@ -464,6 +488,7 @@ public:
 		FTimeValue End;
 		FTimeValue Duration;
 		bool bIsEarlyPeriod = false;
+		bool bHasFollowingPeriod = false;
 		bool bHasBeenPrepared = false;
 	};
 
@@ -538,7 +563,7 @@ public:
 	{
 		return URLFragmentComponents;
 	}
-	
+
 	void SetURLFragmentComponents(TArray<FURL_RFC3986::FQueryParam> InURLFragmentComponents)
 	{
 		URLFragmentComponents = MoveTemp(InURLFragmentComponents);
@@ -560,6 +585,7 @@ public:
 	void GetSeekablePositions(TArray<FTimespan>& OutPositions) const;
 	FTimeValue GetDuration() const;
 	FTimeValue GetDefaultStartTime() const;
+	void ClearDefaultStartTime();
 
 	FTimeValue GetMPDValidityEndTime() const;
 	FTimeValue GetLastPeriodEndTime() const;
@@ -570,11 +596,6 @@ public:
 	FTimeValue GetAvailabilityEndTime() const;
 	FTimeValue GetTimeshiftBufferDepth() const;
 
-#ifdef LOG_DASHINT_CTORDTOR
-	FManifestDASHInternal();
-	~FManifestDASHInternal();
-#endif
-
 private:
 	FErrorDetail PrepareRemoteElementLoadRequest(TArray<TWeakPtrTS<FMPDLoadRequestDASH>>& OutRemoteElementLoadRequests, TWeakPtrTS<IDashMPDElement> ElementWithXLink, int64 RequestID);
 
diff --git a/Engine/Plugins/Media/ElectraPlayer/Source/ElectraPlayerRuntime/Private/Runtime/Player/DASH/StreamReaderFMP4DASH.cpp b/Engine/Plugins/Media/ElectraPlayer/Source/ElectraPlayerRuntime/Private/Runtime/Player/DASH/StreamReaderFMP4DASH.cpp
--- a/Engine/Plugins/Media/ElectraPlayer/Source/ElectraPlayerRuntime/Private/Runtime/Player/DASH/StreamReaderFMP4DASH.cpp
+++ b/Engine/Plugins/Media/ElectraPlayer/Source/ElectraPlayerRuntime/Private/Runtime/Player/DASH/StreamReaderFMP4DASH.cpp
@@ -39,6 +39,12 @@ uint32 FStreamSegmentRequestFMP4DASH::GetPlaybackSequenceID() const
 	return CurrentPlaybackSequenceID;
 }
 
+void FStreamSegmentRequestFMP4DASH::SetExecutionDelay(const FTimeValue& ExecutionDelay)
+{
+	DownloadDelayTime = ExecutionDelay;
+}
+
+
 EStreamType FStreamSegmentRequestFMP4DASH::GetType() const
 {
 	return StreamType;
@@ -472,6 +478,7 @@ void FStreamReaderFMP4DASH::FStreamHandler::HandleRequest()
 	FErrorDetail InitSegmentError;
 	UEMediaError Error;
 	bool bIsEmptyFillerSegment = Request->bInsertFillerData;
+	bool bIsLastSegment = Request->Segment.bIsLastInPeriod;
 
 	Metrics::FSegmentDownloadStats& ds = CurrentRequest->DownloadStats;
 	ds.StatsID = FMediaInterlockedIncrement(UniqueDownloadID);
@@ -508,14 +515,13 @@ void FStreamReaderFMP4DASH::FStreamHandler::HandleRequest()
 	if (Request->ASAST.IsValid())
 	{
 		FTimeValue Now = PlayerSessionService->GetSynchronizedUTCTime()->GetTime();
-#if 0
-		if (Now < Request->ASAST)
+		FTimeValue When = Request->ASAST;
+		if (Request->DownloadDelayTime.IsValid())
 		{
-			LogMessage(IInfoLog::ELevel::Info, FString::Printf(TEXT("waiting fr ASAST %lld"), (long long int) (Request->ASAST - Now).GetAsMilliseconds() ));
+			When += Request->DownloadDelayTime;
 		}
-#endif
 		// Availability start time not reached yet?
-		while(Now < Request->ASAST)
+		while(Now < When)
 		{
 			if (HasReadBeenAborted())
 			{
@@ -523,7 +529,7 @@ void FStreamReaderFMP4DASH::FStreamHandler::HandleRequest()
 				return;
 			}
 			const int64 SleepIntervalMS = 100;
-			FMediaRunnable::SleepMilliseconds((uint32) Utils::Max(Utils::Min((Request->ASAST - Now).GetAsMilliseconds(), SleepIntervalMS), (int64)0));
+			FMediaRunnable::SleepMilliseconds((uint32) Utils::Max(Utils::Min((When - Now).GetAsMilliseconds(), SleepIntervalMS), (int64)0));
 			Now = PlayerSessionService->GetSynchronizedUTCTime()->GetTime();
 		}
 
@@ -545,7 +551,8 @@ void FStreamReaderFMP4DASH::FStreamHandler::HandleRequest()
 
 	FTimeValue NextExpectedDTS;
 	FTimeValue LastKnownAUDuration;
-	FTimeValue TimeOffset = Request->PeriodStart + Request->AST + Request->AdditionalAdjustmentTime;
+	FTimeValue TimeOffset = Request->PeriodStart + Request->AST + Request->AdditionalAdjustmentTime + Request->PlayerLoopState.LoopBasetime;
+
 
 	// Get the init segment if there is one. Either gets it from the entity cache or requests it now and adds it to the cache.
 	InitSegmentError = GetInitSegment(MP4InitSegment, Request);
@@ -670,13 +677,9 @@ void FStreamReaderFMP4DASH::FStreamHandler::HandleRequest()
 										// Use it, but remember the amount of overlap.
 										AccessUnit->OverlapAdjust.SetFromND(Overlap, TrackTimescale);
 									}
-									/*
-										Taken out for now since varying segment durations and possible timescale conversions
-										can lead to inaccuracies. We need a flag saying this is the last segment of the period
-										and only when that is set and the period has a valid duration should we mark the
-										overlapping AUs. Unless the period has a duration this should never be set.
-						
-										// Entirely past the time allowed?
+									// Entirely past the time allowed?
+									if (bIsLastSegment)
+									{
 										if (AUPTS >= MediaLocalLastAUTime)
 										{
 											AccessUnit->DropState |= FAccessUnit::EDropState::PtsTooLate;
@@ -689,7 +692,7 @@ void FStreamReaderFMP4DASH::FStreamHandler::HandleRequest()
 											// boxes that we have to read (like 'emsg' or 'lmsg').
 											AccessUnit->OverlapAdjust.SetFromND(AUPTS + AUDuration - MediaLocalLastAUTime, TrackTimescale);
 										}
-									*/
+									}
 									FTimeValue Duration(TrackIterator->GetDuration(), TrackTimescale);
 									AccessUnit->Duration = Duration;
 
diff --git a/Engine/Plugins/Media/ElectraPlayer/Source/ElectraPlayerRuntime/Private/Runtime/Player/DASH/StreamReaderFMP4DASH.h b/Engine/Plugins/Media/ElectraPlayer/Source/ElectraPlayerRuntime/Private/Runtime/Player/DASH/StreamReaderFMP4DASH.h
--- a/Engine/Plugins/Media/ElectraPlayer/Source/ElectraPlayerRuntime/Private/Runtime/Player/DASH/StreamReaderFMP4DASH.h
+++ b/Engine/Plugins/Media/ElectraPlayer/Source/ElectraPlayerRuntime/Private/Runtime/Player/DASH/StreamReaderFMP4DASH.h
@@ -20,6 +20,7 @@ public:
 	virtual void SetPlaybackSequenceID(uint32 PlaybackSequenceID) override;
 	virtual uint32 GetPlaybackSequenceID() const override;
 
+	virtual void SetExecutionDelay(const FTimeValue& ExecutionDelay) override;
 
 	virtual EStreamType GetType() const override;
 
@@ -51,6 +52,7 @@ public:
 	// UTC wallclock times during which this segment can be fetched;
 	FTimeValue												ASAST;
 	FTimeValue												SAET;
+	FTimeValue												DownloadDelayTime;
 
 	// Internal work variables
 	FPlayerLoopState										PlayerLoopState;
diff --git a/Engine/Plugins/Media/ElectraPlayer/Source/ElectraPlayerRuntime/Private/Runtime/Player/HLS/ManifestBuilderHLS.cpp b/Engine/Plugins/Media/ElectraPlayer/Source/ElectraPlayerRuntime/Private/Runtime/Player/HLS/ManifestBuilderHLS.cpp
--- a/Engine/Plugins/Media/ElectraPlayer/Source/ElectraPlayerRuntime/Private/Runtime/Player/HLS/ManifestBuilderHLS.cpp
+++ b/Engine/Plugins/Media/ElectraPlayer/Source/ElectraPlayerRuntime/Private/Runtime/Player/HLS/ManifestBuilderHLS.cpp
@@ -395,6 +395,7 @@ FErrorDetail FManifestBuilderHLS::SetupRenditions(FManifestHLSInternal* Manifest
 		{
 			Adapt = new FPlaybackAssetAdaptationSetHLS;
 			Adapt->UniqueIdentifier = It.Key();
+			Adapt->Metadata.TrackID = Adapt->UniqueIdentifier;
 			TSharedPtrTS<FPlaybackAssetAdaptationSetHLS> IAdapt(Adapt);
 			Asset->AudioAdaptationSets.Push(IAdapt);
 		}
@@ -690,10 +691,20 @@ FErrorDetail FManifestBuilderHLS::SetupVariants(FManifestHLSInternal* Manifest,
 				Manifest->VariantStreams.Push(vs);
 				Manifest->BandwidthToQualityIndex.Add(vs->Bandwidth, vs->Bandwidth);
 
+				// Add to timeline asset
+				FPlaybackAssetAdaptationSetHLS* Adapt = nullptr;
+				if (!Asset->VideoAdaptationSet.IsValid())
+				{
+					Adapt = new FPlaybackAssetAdaptationSetHLS;
+					Asset->VideoAdaptationSet = MakeShareable(Adapt);
+					Adapt->UniqueIdentifier = TEXT("$video$");
+					Adapt->Metadata.TrackID = Adapt->UniqueIdentifier;
+				}
+				Adapt = static_cast<FPlaybackAssetAdaptationSetHLS*>(Asset->VideoAdaptationSet.Get());
+
 				FStreamMetadata StreamMetaData;
-				StreamMetaData.Bandwidth	  = vs->Bandwidth;
-				StreamMetaData.StreamUniqueID = vs->Internal.UniqueID;
-				StreamMetaData.PlaylistID     = UrlBuilder->ResolveWith(vs->URI);
+				StreamMetaData.Bandwidth = vs->Bandwidth;
+				StreamMetaData.ID = LexToString(vs->Internal.UniqueID);
 				for(int32 i=0; i<vs->StreamCodecInformationList.Num(); ++i)
 				{
 					if (vs->StreamCodecInformationList[i].IsVideoCodec())
@@ -702,17 +713,7 @@ FErrorDetail FManifestBuilderHLS::SetupVariants(FManifestHLSInternal* Manifest,
 						break;
 					}
 				}
-				Manifest->StreamMetadataVideo.Push(StreamMetaData);
 
-				// Add to timeline asset
-				FPlaybackAssetAdaptationSetHLS* Adapt = nullptr;
-				if (!Asset->VideoAdaptationSet.IsValid())
-				{
-					Adapt = new FPlaybackAssetAdaptationSetHLS;
-					Asset->VideoAdaptationSet = MakeShareable(Adapt);
-					Adapt->UniqueIdentifier = TEXT("$video$");
-				}
-				Adapt = static_cast<FPlaybackAssetAdaptationSetHLS*>(Asset->VideoAdaptationSet.Get());
 				if (Adapt->Codecs.Find(StreamMetaData.CodecInformation.GetCodecSpecifierRFC6381()) == INDEX_NONE)
 				{
 					if (Adapt->Codecs.Len())
@@ -722,14 +723,15 @@ FErrorDetail FManifestBuilderHLS::SetupVariants(FManifestHLSInternal* Manifest,
 					Adapt->Codecs.Append(StreamMetaData.CodecInformation.GetCodecSpecifierRFC6381());
 				}
 				FPlaybackAssetRepresentationHLS* Repr = new FPlaybackAssetRepresentationHLS;
-				Repr->UniqueIdentifier = LexToString(StreamMetaData.StreamUniqueID);
+				Repr->UniqueIdentifier = StreamMetaData.ID;
 				Repr->CodecInformation = StreamMetaData.CodecInformation;
 				Repr->Bitrate   	   = StreamMetaData.Bandwidth;
 				Adapt->Representations.Push(TSharedPtrTS<IPlaybackAssetRepresentation>(Repr));
+				Adapt->Metadata.StreamDetails.Emplace(StreamMetaData);
 
 				vs->Internal.AdaptationSetUniqueID  = Adapt->UniqueIdentifier;
 				vs->Internal.RepresentationUniqueID = Repr->UniqueIdentifier;
-				vs->Internal.CDN					= StreamMetaData.PlaylistID;
+				vs->Internal.CDN					= UrlBuilder->ResolveWith(vs->URI);;
 
 				// Check if there is an audio adaptation set already due to an audio rendition group referenced by this variant.
 				if (bHaveAudioGroup)
@@ -770,9 +772,8 @@ FErrorDetail FManifestBuilderHLS::SetupVariants(FManifestHLSInternal* Manifest,
 				Manifest->AudioOnlyStreams.Push(vs);
 
 				FStreamMetadata StreamMetaData;
-				StreamMetaData.Bandwidth	  = vs->Bandwidth;
-				StreamMetaData.StreamUniqueID = vs->Internal.UniqueID;
-				StreamMetaData.PlaylistID     = UrlBuilder->ResolveWith(vs->URI);
+				StreamMetaData.Bandwidth = vs->Bandwidth;
+				StreamMetaData.ID = LexToString(vs->Internal.UniqueID);
 				for(int32 i=0; i<vs->StreamCodecInformationList.Num(); ++i)
 				{
 					if (vs->StreamCodecInformationList[i].IsAudioCodec())
@@ -781,7 +782,6 @@ FErrorDetail FManifestBuilderHLS::SetupVariants(FManifestHLSInternal* Manifest,
 						break;
 					}
 				}
-				Manifest->StreamMetadataAudio.Push(StreamMetaData);
 
 				// Check if there is an audio adaptation set already due to an audio rendition group referenced by this variant.
 				FPlaybackAssetAdaptationSetHLS* Adapt = nullptr;
@@ -801,6 +801,7 @@ FErrorDetail FManifestBuilderHLS::SetupVariants(FManifestHLSInternal* Manifest,
 					{
 						Adapt = new FPlaybackAssetAdaptationSetHLS;
 						Adapt->UniqueIdentifier = TEXT("$audio$");
+						Adapt->Metadata.TrackID = Adapt->UniqueIdentifier;
 						TSharedPtrTS<FPlaybackAssetAdaptationSetHLS> IAdapt(Adapt);
 						Asset->AudioAdaptationSets.Push(IAdapt);
 					}
@@ -817,17 +818,20 @@ FErrorDetail FManifestBuilderHLS::SetupVariants(FManifestHLSInternal* Manifest,
 				if (!bHaveAudioGroup)
 				{
 					FPlaybackAssetRepresentationHLS* Repr = new FPlaybackAssetRepresentationHLS;
-					Repr->UniqueIdentifier = LexToString(StreamMetaData.StreamUniqueID);
+					Repr->UniqueIdentifier = StreamMetaData.ID;
 					Repr->CodecInformation = StreamMetaData.CodecInformation;
 					Repr->Bitrate   	   = StreamMetaData.Bandwidth;
 					Adapt->Representations.Push(TSharedPtrTS<IPlaybackAssetRepresentation>(Repr));
+					Adapt->Metadata.StreamDetails.Emplace(StreamMetaData);
 
 					vs->Internal.AdaptationSetUniqueID  = Adapt->UniqueIdentifier;
 					vs->Internal.RepresentationUniqueID = Repr->UniqueIdentifier;
-					vs->Internal.CDN					= StreamMetaData.PlaylistID;
+					vs->Internal.CDN					= UrlBuilder->ResolveWith(vs->URI);
 				}
 				else
 				{
+					Adapt->Metadata.StreamDetails.Emplace(StreamMetaData);
+
 					for(int32 i=0; i<Adapt->GetNumberOfRepresentations(); ++i)
 					{
 						FPlaybackAssetRepresentationHLS* Repr = static_cast<FPlaybackAssetRepresentationHLS*>(Adapt->GetRepresentationByIndex(i).Get());
@@ -841,7 +845,7 @@ FErrorDetail FManifestBuilderHLS::SetupVariants(FManifestHLSInternal* Manifest,
 						Manifest->AudioRenditions.MultiFindPointer(vs->AudioGroupID, RenditionRange);
 						if (RenditionRange.Num())
 						{
-							for(int32 ii=0; ii < RenditionRange.Num(); ++ii)
+							for(int32 ii=0; ii<RenditionRange.Num(); ++ii)
 							{
 								TSharedPtrTS<FManifestHLSInternal::FRendition>& Rendition = *RenditionRange[ii];
 								if (Rendition->URI.Len() == 0)
@@ -851,7 +855,7 @@ FErrorDetail FManifestBuilderHLS::SetupVariants(FManifestHLSInternal* Manifest,
 									Repr->Bitrate = StreamMetaData.Bandwidth;
 									vs->Internal.AdaptationSetUniqueID  = Adapt->UniqueIdentifier;
 									vs->Internal.RepresentationUniqueID = Repr->UniqueIdentifier;
-									vs->Internal.CDN					= StreamMetaData.PlaylistID;
+									vs->Internal.CDN					= UrlBuilder->ResolveWith(vs->URI);
 								}
 							}
 						}
@@ -877,13 +881,12 @@ FErrorDetail FManifestBuilderHLS::SetupVariants(FManifestHLSInternal* Manifest,
 			{
 				Manifest->PlaylistIDMap.Remove(Manifest->AudioOnlyStreams[i]->Internal.UniqueID);
 				Manifest->AudioOnlyStreams.RemoveAt(i);
-				Manifest->StreamMetadataAudio.RemoveAt(i);
 			}
 		}
 		// Set up variants from the renditions in the audio adaptation sets.
 		for(int32 i=0; i<Asset->AudioAdaptationSets.Num(); ++i)
 		{
-			const FPlaybackAssetAdaptationSetHLS* Adapt = static_cast<const FPlaybackAssetAdaptationSetHLS*>(Asset->GetAdaptationSetByTypeAndIndex(EStreamType::Audio, i).Get());
+			FPlaybackAssetAdaptationSetHLS* Adapt = static_cast<FPlaybackAssetAdaptationSetHLS*>(Asset->GetAdaptationSetByTypeAndIndex(EStreamType::Audio, i).Get());
 			for(int32 j=0; j<Adapt->GetNumberOfRepresentations(); ++j)
 			{
 				const FPlaybackAssetRepresentationHLS* Repr = static_cast<const FPlaybackAssetRepresentationHLS*>(Adapt->GetRepresentationByIndex(j).Get());
@@ -902,10 +905,8 @@ FErrorDetail FManifestBuilderHLS::SetupVariants(FManifestHLSInternal* Manifest,
 						vs->Internal.UniqueID = FMediaInterlockedIncrement(NextUniqueID) + 1;
 
 						FStreamMetadata StreamMetaData;
-						StreamMetaData.Bandwidth	  = 128000;
-						StreamMetaData.StreamUniqueID = vs->Internal.UniqueID;
-						StreamMetaData.PlaylistID     = UrlBuilder->ResolveWith(vs->URI);
-						StreamMetaData.LanguageCode   = Rendition->Language;
+						StreamMetaData.Bandwidth = 128000;
+						StreamMetaData.ID = LexToString(vs->Internal.UniqueID);
 						for(int32 k=0; k<vs->StreamCodecInformationList.Num(); ++k)
 						{
 							if (vs->StreamCodecInformationList[k].IsAudioCodec())
@@ -914,11 +915,11 @@ FErrorDetail FManifestBuilderHLS::SetupVariants(FManifestHLSInternal* Manifest,
 								break;
 							}
 						}
-						Manifest->StreamMetadataAudio.Push(StreamMetaData);
+						Adapt->Metadata.StreamDetails.Emplace(StreamMetaData);
 
 						vs->Internal.AdaptationSetUniqueID  = Adapt->UniqueIdentifier;
 						vs->Internal.RepresentationUniqueID = Repr->UniqueIdentifier;
-						vs->Internal.CDN					= StreamMetaData.PlaylistID;
+						vs->Internal.CDN					= UrlBuilder->ResolveWith(vs->URI);
 						Manifest->AudioOnlyStreams.Push(vs);
 
 						Manifest->PlaylistIDMap.Add(vs->Internal.UniqueID, vs);
@@ -930,6 +931,33 @@ FErrorDetail FManifestBuilderHLS::SetupVariants(FManifestHLSInternal* Manifest,
 		}
 	}
 
+	// Set up the highest bandwidth and corresponding codec information in the adaptation sets.
+	if (Asset->VideoAdaptationSet.IsValid())
+	{
+		FTrackMetadata& Meta = Asset->VideoAdaptationSet->Metadata;
+		for(int32 i=0; i<Meta.StreamDetails.Num(); ++i)
+		{
+			if (Meta.StreamDetails[i].Bandwidth > Meta.HighestBandwidth)
+			{
+				Meta.HighestBandwidth = Meta.StreamDetails[i].Bandwidth;
+				Meta.HighestBandwidthCodec = Meta.StreamDetails[i].CodecInformation;
+			}
+		}
+		Manifest->TrackMetadataVideo.Push(Meta);
+	}
+	for(int32 j=0; j<Asset->AudioAdaptationSets.Num(); ++j)
+	{
+		FTrackMetadata& Meta = Asset->AudioAdaptationSets[j]->Metadata;
+		for(int32 i=0; i<Meta.StreamDetails.Num(); ++i)
+		{
+			if (Meta.StreamDetails[i].Bandwidth > Meta.HighestBandwidth)
+			{
+				Meta.HighestBandwidth = Meta.StreamDetails[i].Bandwidth;
+				Meta.HighestBandwidthCodec = Meta.StreamDetails[i].CodecInformation;
+			}
+		}
+		Manifest->TrackMetadataAudio.Push(Meta);
+	}
 
 	// Index the stream quality level map. Lower quality = lower index
 	int32 QualityIndex = 0;
diff --git a/Engine/Plugins/Media/ElectraPlayer/Source/ElectraPlayerRuntime/Private/Runtime/Player/HLS/ManifestBuilderHLS.h b/Engine/Plugins/Media/ElectraPlayer/Source/ElectraPlayerRuntime/Private/Runtime/Player/HLS/ManifestBuilderHLS.h
--- a/Engine/Plugins/Media/ElectraPlayer/Source/ElectraPlayerRuntime/Private/Runtime/Player/HLS/ManifestBuilderHLS.h
+++ b/Engine/Plugins/Media/ElectraPlayer/Source/ElectraPlayerRuntime/Private/Runtime/Player/HLS/ManifestBuilderHLS.h
@@ -305,8 +305,8 @@ struct FManifestHLSInternal
 	TSharedPtrTS<IInitSegmentCacheHLS>						InitSegmentCache;
 	TSharedPtrTS<ILicenseKeyCacheHLS>						LicenseKeyCache;
 
-	TArray<FStreamMetadata>									StreamMetadataVideo;
-	TArray<FStreamMetadata>									StreamMetadataAudio;
+	TArray<FTrackMetadata>									TrackMetadataVideo;
+	TArray<FTrackMetadata>									TrackMetadataAudio;
 
 	TSharedPtrTS<FTimelineMediaAssetHLS>					CurrentMediaAsset;
 
@@ -548,6 +548,7 @@ public:
 	FString													UniqueIdentifier;
 	FString													Language;
 	FString													Codecs;
+	FTrackMetadata											Metadata;
 };
 
 
@@ -605,6 +606,34 @@ public:
 			}
 		}
 	}
+
+	virtual void GetMetaData(TArray<FTrackMetadata>& OutMetadata, EStreamType OfStreamType) const override
+	{
+		switch(OfStreamType)
+		{
+			case EStreamType::Video:
+			{
+				if (VideoAdaptationSet.IsValid())
+				{
+					OutMetadata.Emplace(VideoAdaptationSet->Metadata);
+				}
+				break;
+			}
+			case EStreamType::Audio:
+			{
+				for(int32 i=0, iMax=AudioAdaptationSets.Num(); i<iMax; ++i)
+				{
+					OutMetadata.Emplace(AudioAdaptationSets[i]->Metadata);
+				}
+				break;
+			}
+			default:
+			{
+				break;
+			}
+		}
+	}
+
 	TSharedPtrTS<IPlaybackAssetAdaptationSet> GetAdaptationSetByTypeAndUniqueIdentifier(EStreamType OfStreamType, const FString& InUniqueIdentifier) const
 	{
 		switch(OfStreamType)
@@ -615,7 +644,7 @@ public:
 			}
 			case EStreamType::Audio:
 			{
-				for(int32 i=0, iMax=(int32)AudioAdaptationSets.Num(); i<iMax; ++i)
+				for(int32 i=0, iMax=AudioAdaptationSets.Num(); i<iMax; ++i)
 				{
 					if (AudioAdaptationSets[i]->GetUniqueIdentifier() == InUniqueIdentifier)
 					{
diff --git a/Engine/Plugins/Media/ElectraPlayer/Source/ElectraPlayerRuntime/Private/Runtime/Player/HLS/ManifestHLS.cpp b/Engine/Plugins/Media/ElectraPlayer/Source/ElectraPlayerRuntime/Private/Runtime/Player/HLS/ManifestHLS.cpp
--- a/Engine/Plugins/Media/ElectraPlayer/Source/ElectraPlayerRuntime/Private/Runtime/Player/HLS/ManifestHLS.cpp
+++ b/Engine/Plugins/Media/ElectraPlayer/Source/ElectraPlayerRuntime/Private/Runtime/Player/HLS/ManifestHLS.cpp
@@ -27,7 +27,7 @@ public:
 	virtual void SetStreamPreferences(const FStreamPreferences& Preferences) override;
 
 	virtual EReadyState GetReadyState() override;
-	virtual void PrepareForPlay(const FParamDict& Options) override;
+	virtual void PrepareForPlay() override;
 
 	// TODO: need to provide metadata (duration, streams, languages, etc.)
 
@@ -150,6 +150,9 @@ FTimeValue FManifestHLS::GetDefaultStartTime() const
 {
 	return FTimeValue::GetInvalid();
 }
+void FManifestHLS::ClearDefaultStartTime()
+{
+}
 
 
 
@@ -169,21 +172,21 @@ int64 FManifestHLS::GetDefaultStartingBitrate() const
 }
 
 /**
- * Returns stream metadata. For period based presentations the streams can be different per period in which case the metadata of the first period is returned.
+ * Returns track metadata. For period based presentations the streams can be different per period in which case the metadata of the first period is returned.
  *
  * @param OutMetadata
  * @param StreamType
  */
-void FManifestHLS::GetStreamMetadata(TArray<FStreamMetadata>& OutMetadata, EStreamType StreamType) const
+void FManifestHLS::GetTrackMetadata(TArray<FTrackMetadata>& OutMetadata, EStreamType StreamType) const
 {
 	FManifestHLSInternal::ScopedLockPlaylists lock(InternalManifest);
 	switch(StreamType)
 	{
 		case EStreamType::Video:
-			OutMetadata = InternalManifest->StreamMetadataVideo;
+			OutMetadata = InternalManifest->TrackMetadataVideo;
 			break;
 		case EStreamType::Audio:
-			OutMetadata = InternalManifest->StreamMetadataAudio;
+			OutMetadata = InternalManifest->TrackMetadataAudio;
 			break;
 		case EStreamType::Subtitle:
 			OutMetadata.Empty();
@@ -232,6 +235,12 @@ IManifest::FResult FManifestHLS::FindPlayPeriod(TSharedPtrTS<IPlayPeriod>& OutPl
 	return IManifest::FResult(IManifest::FResult::EType::Found);
 }
 
+IManifest::FResult FManifestHLS::FindNextPlayPeriod(TSharedPtrTS<IPlayPeriod>& OutPlayPeriod, TSharedPtrTS<const IStreamSegment> CurrentSegment)
+{
+	// There is no following period.
+	return IManifest::FResult(IManifest::FResult::EType::PastEOS);
+}
+
 
 
 /**
@@ -305,10 +314,8 @@ IManifest::IPlayPeriod::EReadyState FPlayPeriodHLS::GetReadyState()
 
 /**
  * Prepares the period for playback.
- *
- * @param Options
  */
-void FPlayPeriodHLS::PrepareForPlay(const FParamDict& Options)
+void FPlayPeriodHLS::PrepareForPlay()
 {
 	// For now we just go with the streams for which we loaded the playlists initially.
 	// FIXME: in the future, based on preferences and options, select the streams we want.
diff --git a/Engine/Plugins/Media/ElectraPlayer/Source/ElectraPlayerRuntime/Private/Runtime/Player/HLS/ManifestHLS.h b/Engine/Plugins/Media/ElectraPlayer/Source/ElectraPlayerRuntime/Private/Runtime/Player/HLS/ManifestHLS.h
--- a/Engine/Plugins/Media/ElectraPlayer/Source/ElectraPlayerRuntime/Private/Runtime/Player/HLS/ManifestHLS.h
+++ b/Engine/Plugins/Media/ElectraPlayer/Source/ElectraPlayerRuntime/Private/Runtime/Player/HLS/ManifestHLS.h
@@ -23,13 +23,15 @@ public:
 	virtual void GetSeekablePositions(TArray<FTimespan>& OutPositions) const override;
 	virtual FTimeValue GetDuration() const override;
 	virtual FTimeValue GetDefaultStartTime() const override;
+	virtual void ClearDefaultStartTime() override;
 	virtual int64 GetDefaultStartingBitrate() const override;
 	virtual FTimeValue GetMinBufferTime() const override;
-	virtual void GetStreamMetadata(TArray<FStreamMetadata>& OutMetadata, EStreamType StreamType) const override;
+	virtual void GetTrackMetadata(TArray<FTrackMetadata>& OutMetadata, EStreamType StreamType) const override;
 	virtual void UpdateDynamicRefetchCounter() override;
 	virtual IStreamReader *CreateStreamReaderHandler() override;
 
 	virtual FResult FindPlayPeriod(TSharedPtrTS<IPlayPeriod>& OutPlayPeriod, const FPlayStartPosition& StartPosition, ESearchType SearchType) override;
+	virtual FResult FindNextPlayPeriod(TSharedPtrTS<IPlayPeriod>& OutPlayPeriod, TSharedPtrTS<const IStreamSegment> CurrentSegment) override;
 
 private:
 	FManifestHLS(IPlayerSessionServices* SessionServices, IPlaylistReaderHLS* PlaylistReader, TSharedPtrTS<FManifestHLSInternal> Manifest);
diff --git a/Engine/Plugins/Media/ElectraPlayer/Source/ElectraPlayerRuntime/Private/Runtime/Player/HLS/StreamReaderHLSfmp4.cpp b/Engine/Plugins/Media/ElectraPlayer/Source/ElectraPlayerRuntime/Private/Runtime/Player/HLS/StreamReaderHLSfmp4.cpp
--- a/Engine/Plugins/Media/ElectraPlayer/Source/ElectraPlayerRuntime/Private/Runtime/Player/HLS/StreamReaderHLSfmp4.cpp
+++ b/Engine/Plugins/Media/ElectraPlayer/Source/ElectraPlayerRuntime/Private/Runtime/Player/HLS/StreamReaderHLSfmp4.cpp
@@ -57,6 +57,11 @@ uint32 FStreamSegmentRequestHLSfmp4::GetPlaybackSequenceID() const
 	return CurrentPlaybackSequenceID;
 }
 
+void FStreamSegmentRequestHLSfmp4::SetExecutionDelay(const FTimeValue& ExecutionDelay)
+{
+	// No-op for HLS.
+}
+
 EStreamType FStreamSegmentRequestHLSfmp4::GetType() const
 {
 	return StreamType;
diff --git a/Engine/Plugins/Media/ElectraPlayer/Source/ElectraPlayerRuntime/Private/Runtime/Player/HLS/StreamReaderHLSfmp4.h b/Engine/Plugins/Media/ElectraPlayer/Source/ElectraPlayerRuntime/Private/Runtime/Player/HLS/StreamReaderHLSfmp4.h
--- a/Engine/Plugins/Media/ElectraPlayer/Source/ElectraPlayerRuntime/Private/Runtime/Player/HLS/StreamReaderHLSfmp4.h
+++ b/Engine/Plugins/Media/ElectraPlayer/Source/ElectraPlayerRuntime/Private/Runtime/Player/HLS/StreamReaderHLSfmp4.h
@@ -22,6 +22,7 @@ public:
 	virtual void SetPlaybackSequenceID(uint32 PlaybackSequenceID) override;
 	virtual uint32 GetPlaybackSequenceID() const override;
 
+	virtual void SetExecutionDelay(const FTimeValue& ExecutionDelay) override;
 
 	virtual EStreamType GetType() const override;
 
diff --git a/Engine/Plugins/Media/ElectraPlayer/Source/ElectraPlayerRuntime/Private/Runtime/Player/Manifest.h b/Engine/Plugins/Media/ElectraPlayer/Source/ElectraPlayerRuntime/Private/Runtime/Player/Manifest.h
--- a/Engine/Plugins/Media/ElectraPlayer/Source/ElectraPlayerRuntime/Private/Runtime/Player/Manifest.h
+++ b/Engine/Plugins/Media/ElectraPlayer/Source/ElectraPlayerRuntime/Private/Runtime/Player/Manifest.h
@@ -44,6 +44,7 @@ namespace Electra
 		 */
 		virtual uint32 GetPlaybackSequenceID() const = 0;
 
+		virtual void SetExecutionDelay(const FTimeValue& ExecutionDelay) = 0;
 
 		virtual EStreamType GetType() const = 0;
 
@@ -224,14 +225,19 @@ namespace Electra
 		 * If the presentation has no preferred start time an invalid value is returned.
 		 */
 		virtual FTimeValue GetDefaultStartTime() const = 0;
+		
+		/**
+		 * Clears the internal default start time so it will not be used again.
+		 */
+		virtual void ClearDefaultStartTime() = 0;
 
 
 
 		//! Returns the bitrate of the default stream (usually the first one specified).
 		virtual int64 GetDefaultStartingBitrate() const = 0;
 
-		//! Returns stream metadata. For period based presentations the streams can be different per period in which case the metadata of the first period is returned.
-		virtual void GetStreamMetadata(TArray<FStreamMetadata>& OutMetadata, EStreamType StreamType) const = 0;
+		//! Returns track metadata. For period based presentations the streams can be different per period in which case the metadata of the first period is returned.
+		virtual void GetTrackMetadata(TArray<FTrackMetadata>& OutMetadata, EStreamType StreamType) const = 0;
 
 		//
 		virtual FTimeValue GetMinBufferTime() const = 0;
@@ -266,10 +272,9 @@ namespace Electra
 				NotReady,
 				Preparing,
 				IsReady,
-				MustReselect						//!< Player must forget this Play Period and reselect it through FindPlayPeriod(). Needed if the timeline got altered.
 			};
 			virtual EReadyState GetReadyState() = 0;
-			virtual void PrepareForPlay(const FParamDict& Options) = 0;
+			virtual void PrepareForPlay() = 0;
 
 			/**
 			 * Returns the media asset for this play period.
@@ -362,7 +367,11 @@ namespace Electra
 			virtual void GetSegmentInformation(TArray<FSegmentInformation>& OutSegmentInformation, FTimeValue& OutAverageSegmentDuration, TSharedPtrTS<const IStreamSegment> CurrentSegment, const FTimeValue& LookAheadTime, const FString& AdaptationSetID, const FString& RepresentationID) = 0;
 		};
 
+		//! Finds the playback period the specified start time falls into.
 		virtual FResult FindPlayPeriod(TSharedPtrTS<IPlayPeriod>& OutPlayPeriod, const FPlayStartPosition& StartPosition, ESearchType SearchType) = 0;
+
+		//! Locates the period following the given segment.
+		virtual FResult FindNextPlayPeriod(TSharedPtrTS<IPlayPeriod>& OutPlayPeriod, TSharedPtrTS<const IStreamSegment> CurrentSegment) = 0;
 	};
 
 
diff --git a/Engine/Plugins/Media/ElectraPlayer/Source/ElectraPlayerRuntime/Private/Runtime/Player/PlaybackTimeline.h b/Engine/Plugins/Media/ElectraPlayer/Source/ElectraPlayerRuntime/Private/Runtime/Player/PlaybackTimeline.h
--- a/Engine/Plugins/Media/ElectraPlayer/Source/ElectraPlayerRuntime/Private/Runtime/Player/PlaybackTimeline.h
+++ b/Engine/Plugins/Media/ElectraPlayer/Source/ElectraPlayerRuntime/Private/Runtime/Player/PlaybackTimeline.h
@@ -114,6 +114,13 @@ namespace Electra
 		 * @return Shared pointer to the requested adaptation set.
 		 */
 		virtual TSharedPtrTS<IPlaybackAssetAdaptationSet> GetAdaptationSetByTypeAndIndex(EStreamType OfStreamType, int32 AdaptationSetIndex) const = 0;
+
+
+		/**
+		 * Returns "track" metadata for the specified type of stream. This is essentially the adaptation sets per type with their
+		 * individual representations.
+		 */
+		virtual void GetMetaData(TArray<FTrackMetadata>& OutMetadata, EStreamType OfStreamType) const = 0;
 	};
 
 
diff --git a/Engine/Plugins/Media/ElectraPlayer/Source/ElectraPlayerRuntime/Private/Runtime/Player/mp4/ManifestMP4.cpp b/Engine/Plugins/Media/ElectraPlayer/Source/ElectraPlayerRuntime/Private/Runtime/Player/mp4/ManifestMP4.cpp
--- a/Engine/Plugins/Media/ElectraPlayer/Source/ElectraPlayerRuntime/Private/Runtime/Player/mp4/ManifestMP4.cpp
+++ b/Engine/Plugins/Media/ElectraPlayer/Source/ElectraPlayerRuntime/Private/Runtime/Player/mp4/ManifestMP4.cpp
@@ -110,35 +110,16 @@ int64 FManifestMP4Internal::GetDefaultStartingBitrate() const
 
 //-----------------------------------------------------------------------------
 /**
- * Returns stream metadata.
+ * Returns track metadata.
  *
  * @param OutMetadata
  * @param StreamType
  */
-void FManifestMP4Internal::GetStreamMetadata(TArray<FStreamMetadata>& OutMetadata, EStreamType StreamType) const
+void FManifestMP4Internal::GetTrackMetadata(TArray<FTrackMetadata>& OutMetadata, EStreamType StreamType) const
 {
 	if (MediaAsset.IsValid())
 	{
-		for(int32 i=0, iMax = MediaAsset->GetNumberOfAdaptationSets(StreamType); i<iMax; ++i)
-		{
-			TSharedPtrTS<IPlaybackAssetAdaptationSet> AdaptSet = MediaAsset->GetAdaptationSetByTypeAndIndex(StreamType, i);
-			if (AdaptSet.IsValid())
-			{
-				for(int32 j=0, jMax=AdaptSet->GetNumberOfRepresentations(); j<jMax; ++j)
-				{
-					TSharedPtrTS<FRepresentationMP4> Repr = StaticCastSharedPtr<FRepresentationMP4>(AdaptSet->GetRepresentationByIndex(j));
-					if (Repr.IsValid())
-					{
-						FStreamMetadata& meta = OutMetadata.AddDefaulted_GetRef();
-						meta.CodecInformation = Repr->GetCodecInformation();
-						LexFromString(meta.StreamUniqueID, *Repr->GetUniqueIdentifier());
-						meta.PlaylistID 	  = FString(TEXT("mp4"));
-						meta.Bandwidth  	  = Repr->GetBitrate();
-						meta.LanguageCode     = AdaptSet->GetLanguage();
-					}
-				}
-			}
-		}
+		MediaAsset->GetMetaData(OutMetadata, StreamType);
 	}
 }
 
@@ -192,6 +173,12 @@ IManifest::FResult FManifestMP4Internal::FindPlayPeriod(TSharedPtrTS<IPlayPeriod
 	return IManifest::FResult(IManifest::FResult::EType::Found);
 }
 
+IManifest::FResult FManifestMP4Internal::FindNextPlayPeriod(TSharedPtrTS<IPlayPeriod>& OutPlayPeriod, TSharedPtrTS<const IStreamSegment> CurrentSegment)
+{
+	// There is no following period.
+	return IManifest::FResult(IManifest::FResult::EType::PastEOS);
+}
+
 
 
 
@@ -247,12 +234,9 @@ IManifest::IPlayPeriod::EReadyState FManifestMP4Internal::FPlayPeriodMP4::GetRea
  * Prepares the playback period for playback.
  * With an mp4 file we are actually always ready for playback, but we say we're not
  * one time to get here with any possible options.
- *
- * @param InOptions
  */
-void FManifestMP4Internal::FPlayPeriodMP4::PrepareForPlay(const FParamDict& InOptions)
+void FManifestMP4Internal::FPlayPeriodMP4::PrepareForPlay()
 {
-	Options = InOptions;
 	bIsReady = true;
 }
 
diff --git a/Engine/Plugins/Media/ElectraPlayer/Source/ElectraPlayerRuntime/Private/Runtime/Player/mp4/ManifestMP4.h b/Engine/Plugins/Media/ElectraPlayer/Source/ElectraPlayerRuntime/Private/Runtime/Player/mp4/ManifestMP4.h
--- a/Engine/Plugins/Media/ElectraPlayer/Source/ElectraPlayerRuntime/Private/Runtime/Player/mp4/ManifestMP4.h
+++ b/Engine/Plugins/Media/ElectraPlayer/Source/ElectraPlayerRuntime/Private/Runtime/Player/mp4/ManifestMP4.h
@@ -56,13 +56,16 @@ public:
 	{ return MediaAsset.IsValid() ? MediaAsset->GetDuration() : FTimeValue(); }
 	virtual FTimeValue GetDefaultStartTime() const override
 	{ return FTimeValue::GetInvalid(); }
+	virtual void ClearDefaultStartTime() override
+	{ }
 	virtual int64 GetDefaultStartingBitrate() const override;
-	virtual void GetStreamMetadata(TArray<FStreamMetadata>& OutMetadata, EStreamType StreamType) const override;
+	virtual void GetTrackMetadata(TArray<FTrackMetadata>& OutMetadata, EStreamType StreamType) const override;
 	virtual FTimeValue GetMinBufferTime() const override;
 	virtual void UpdateDynamicRefetchCounter() override;
 
 	virtual IStreamReader* CreateStreamReaderHandler() override;
 	virtual FResult FindPlayPeriod(TSharedPtrTS<IPlayPeriod>& OutPlayPeriod, const FPlayStartPosition& StartPosition, ESearchType SearchType) override;
+	virtual FResult FindNextPlayPeriod(TSharedPtrTS<IPlayPeriod>& OutPlayPeriod, TSharedPtrTS<const IStreamSegment> CurrentSegment) override;
 
 
 	class FRepresentationMP4 : public IPlaybackAssetRepresentation
@@ -184,6 +187,38 @@ public:
 			return TSharedPtrTS<IPlaybackAssetAdaptationSet>();
 		}
 
+		virtual void GetMetaData(TArray<FTrackMetadata>& OutMetadata, EStreamType StreamType) const override
+		{
+			for(int32 i=0, iMax=GetNumberOfAdaptationSets(StreamType); i<iMax; ++i)
+			{
+				TSharedPtrTS<FAdaptationSetMP4> AdaptSet = StaticCastSharedPtr<FAdaptationSetMP4>(GetAdaptationSetByTypeAndIndex(StreamType, i));
+				if (AdaptSet.IsValid())
+				{
+					FTrackMetadata tm;
+					tm.TrackID = AdaptSet->GetUniqueIdentifier();
+					tm.Language = AdaptSet->GetLanguage();
+
+					for(int32 j=0, jMax=AdaptSet->GetNumberOfRepresentations(); j<jMax; ++j)
+					{
+						TSharedPtrTS<FRepresentationMP4> Repr = StaticCastSharedPtr<FRepresentationMP4>(AdaptSet->GetRepresentationByIndex(j));
+						if (Repr.IsValid())
+						{
+							FStreamMetadata sd;
+							sd.Bandwidth = Repr->GetBitrate();
+							sd.CodecInformation = Repr->GetCodecInformation();
+							sd.ID = Repr->GetUniqueIdentifier();
+							// There is only 1 "stream" per "track" so we can set the highest bitrate and codec info the same as the track.
+							tm.HighestBandwidth = sd.Bandwidth;
+							tm.HighestBandwidthCodec = sd.CodecInformation;
+
+							tm.StreamDetails.Emplace(MoveTemp(sd));
+						}
+					}
+					OutMetadata.Emplace(MoveTemp(tm));
+				}
+			}
+		}
+
 
 		void LimitSegmentDownloadSize(TSharedPtrTS<IStreamSegment>& InOutSegment, TSharedPtr<IParserISO14496_12::IAllTrackIterator, ESPMode::ThreadSafe> AllTrackIterator);
 
@@ -216,7 +251,7 @@ public:
 		virtual ~FPlayPeriodMP4();
 		virtual void SetStreamPreferences(const FStreamPreferences& Preferences) override;
 		virtual EReadyState GetReadyState() override;
-		virtual void PrepareForPlay(const FParamDict& Options) override;
+		virtual void PrepareForPlay() override;
 		virtual TSharedPtrTS<ITimelineMediaAsset> GetMediaAsset() const override;
 		virtual void SelectStream(const FString& AdaptationSetID, const FString& RepresentationID) override;
 		virtual FResult GetStartingSegment(TSharedPtrTS<IStreamSegment>& OutSegment, const FPlayStartPosition& StartPosition, ESearchType SearchType) override;
diff --git a/Engine/Plugins/Media/ElectraPlayer/Source/ElectraPlayerRuntime/Private/Runtime/Player/mp4/StreamReaderMP4.cpp b/Engine/Plugins/Media/ElectraPlayer/Source/ElectraPlayerRuntime/Private/Runtime/Player/mp4/StreamReaderMP4.cpp
--- a/Engine/Plugins/Media/ElectraPlayer/Source/ElectraPlayerRuntime/Private/Runtime/Player/mp4/StreamReaderMP4.cpp
+++ b/Engine/Plugins/Media/ElectraPlayer/Source/ElectraPlayerRuntime/Private/Runtime/Player/mp4/StreamReaderMP4.cpp
@@ -57,6 +57,11 @@ EStreamType FStreamSegmentRequestMP4::GetType() const
 	return  PrimaryStreamType;
 }
 
+void FStreamSegmentRequestMP4::SetExecutionDelay(const FTimeValue& ExecutionDelay)
+{
+	// No op for mp4.
+}
+
 void FStreamSegmentRequestMP4::GetDependentStreams(TArray<FDependentStreams>& OutDependentStreams) const
 {
 	// Those are not "real" dependent streams in that they are multiplexed and do not need to be fetched from
@@ -415,7 +420,6 @@ void FStreamReaderMP4::WorkerThread()
 				if (!st.StreamSourceInfo.IsValid())
 				{
 					auto meta = MakeShared<FStreamSourceInfo, ESPMode::ThreadSafe>();
-					meta->NumericTrackID = tkid;
 
 					// Check if this track is in the list of selected tracks.
 					const FPlaylistTrackMetadata* SelectedTrackMetadata = SelectedTrackMap.Find(tkid);
diff --git a/Engine/Plugins/Media/ElectraPlayer/Source/ElectraPlayerRuntime/Private/Runtime/Player/mp4/StreamReaderMP4.h b/Engine/Plugins/Media/ElectraPlayer/Source/ElectraPlayerRuntime/Private/Runtime/Player/mp4/StreamReaderMP4.h
--- a/Engine/Plugins/Media/ElectraPlayer/Source/ElectraPlayerRuntime/Private/Runtime/Player/mp4/StreamReaderMP4.h
+++ b/Engine/Plugins/Media/ElectraPlayer/Source/ElectraPlayerRuntime/Private/Runtime/Player/mp4/StreamReaderMP4.h
@@ -31,6 +31,7 @@ public:
 	virtual void SetPlaybackSequenceID(uint32 InPlaybackSequenceID) override;
 	virtual uint32 GetPlaybackSequenceID() const override;
 
+	virtual void SetExecutionDelay(const FTimeValue& ExecutionDelay) override;
 	virtual EStreamType GetType() const override;
 
 	virtual void GetDependentStreams(TArray<FDependentStreams>& OutDependentStreams) const override;
diff --git a/Engine/Plugins/Media/ElectraPlayer/Source/ElectraPlayerRuntime/Private/Runtime/StreamAccessUnitBuffer.cpp b/Engine/Plugins/Media/ElectraPlayer/Source/ElectraPlayerRuntime/Private/Runtime/StreamAccessUnitBuffer.cpp
--- a/Engine/Plugins/Media/ElectraPlayer/Source/ElectraPlayerRuntime/Private/Runtime/StreamAccessUnitBuffer.cpp
+++ b/Engine/Plugins/Media/ElectraPlayer/Source/ElectraPlayerRuntime/Private/Runtime/StreamAccessUnitBuffer.cpp
@@ -11,8 +11,6 @@ namespace Electra
 	FMultiTrackAccessUnitBuffer::FMultiTrackAccessUnitBuffer()
 	{
 		EmptyBuffer = CreateNewBuffer();
-		ActiveOutputID = -1;
-		LastPoppedBufferID = -1;
 		LastPoppedDTS.SetToInvalid();
 		LastPoppedPTS.SetToInvalid();
 		bAutoselectFirstTrack = false;
@@ -42,10 +40,10 @@ namespace Electra
 
 	bool FMultiTrackAccessUnitBuffer::Push(FAccessUnit*& AU)
 	{
-		uint32 TrackID = 0;
+		FString TrackID;
 		if (AU->StreamSourceInfo.IsValid())
 		{
-			TrackID = AU->StreamSourceInfo->NumericTrackID;
+			TrackID = AU->StreamSourceInfo->AdaptationSetID;
 		}
 		// Do we have a buffer for this track ID yet?
 		AccessLock.Lock();
@@ -58,7 +56,7 @@ namespace Electra
 		if (bAutoselectFirstTrack)
 		{
 			bAutoselectFirstTrack = false;
-			ActiveOutputID = (int32)TrackID;
+			ActiveOutputID = TrackID;
 		}
 
 		// Pushing data to either buffer means that there is data and we are not at EOD here any more.
@@ -66,7 +64,7 @@ namespace Electra
 		bEndOfData = false;
 		AccessLock.Unlock();
 		// If this is the active buffer or no active buffer has been set yet we use the configured buffer limit in pushing.
-		const FAccessUnitBuffer::FConfiguration* PushConfig = ActiveOutputID == -1 || ActiveOutputID == (int32)TrackID ? &PrimaryBufferConfiguration : nullptr;
+		const FAccessUnitBuffer::FConfiguration* PushConfig = ActiveOutputID.IsEmpty() || ActiveOutputID.Equals(TrackID) ? &PrimaryBufferConfiguration : nullptr;
 		bool bWasPushed = TrackBuffer->Push(AU, PushConfig);
 		bLastPushWasBlocked = TrackBuffer->WasLastPushBlocked();
 		return bWasPushed;
@@ -74,10 +72,10 @@ namespace Electra
 
 	void FMultiTrackAccessUnitBuffer::PushEndOfDataFor(TSharedPtr<const FStreamSourceInfo, ESPMode::ThreadSafe> InStreamSourceInfo)
 	{
-		uint32 TrackID = 0;
+		FString TrackID;
 		if (InStreamSourceInfo.IsValid())
 		{
-			TrackID = InStreamSourceInfo->NumericTrackID;
+			TrackID = InStreamSourceInfo->AdaptationSetID;
 		}
 		// Do we have a buffer for this track ID yet?
 		AccessLock.Lock();
@@ -90,7 +88,7 @@ namespace Electra
 		if (bAutoselectFirstTrack)
 		{
 			bAutoselectFirstTrack = false;
-			ActiveOutputID = (int32)TrackID;
+			ActiveOutputID = TrackID;
 		}
 
 		AccessLock.Unlock();
@@ -120,7 +118,7 @@ namespace Electra
 		LastPoppedPTS.SetToInvalid();
 		bEndOfData = false;
 		bLastPushWasBlocked = false;
-		LastPoppedBufferID = -1;
+		LastPoppedBufferID.Empty();
 	}
 
 	void FMultiTrackAccessUnitBuffer::PurgeAll()
@@ -132,7 +130,7 @@ namespace Electra
 		LastPoppedPTS.SetToInvalid();
 		bEndOfData = false;
 		bLastPushWasBlocked = false;
-		LastPoppedBufferID = -1;
+		LastPoppedBufferID.Empty();
 	}
 
 	void FMultiTrackAccessUnitBuffer::AutoselectFirstTrack()
@@ -143,7 +141,7 @@ namespace Electra
 	}
 
 
-	void FMultiTrackAccessUnitBuffer::SelectTrackByID(int32 TrackID)
+	void FMultiTrackAccessUnitBuffer::SelectTrackByID(const FString& TrackID)
 	{
 		FMediaCriticalSection::ScopedLock lock(AccessLock);
 		if (TrackID != ActiveOutputID)
@@ -179,7 +177,7 @@ namespace Electra
 	{
 		FAccessUnitBufferPtr TrackBuffer;
 		AccessLock.Lock();
-		if (ActiveOutputID != -1 && TrackBuffers.Contains(ActiveOutputID))
+		if (TrackBuffers.Contains(ActiveOutputID))
 		{
 			TrackBuffer = TrackBuffers[ActiveOutputID];
 		}
@@ -195,7 +193,7 @@ namespace Electra
 	{
 		FAccessUnitBufferPtr TrackBuffer;
 		AccessLock.Lock();
-		if (ActiveOutputID != -1 && TrackBuffers.Contains(ActiveOutputID))
+		if (TrackBuffers.Contains(ActiveOutputID))
 		{
 			TrackBuffer = TrackBuffers[ActiveOutputID];
 		}
@@ -228,7 +226,7 @@ namespace Electra
 			if (bDidPop && OutAU)
 			{
 				// Did we just pop from a different buffer than last time?
-				if (LastPoppedBufferID != -1 && LastPoppedBufferID != ActiveOutputID)
+				if (LastPoppedBufferID != ActiveOutputID)
 				{
 					// FIXME: We may need to discard everything that is not tagged as a sync sample to ensure proper stream switching.
 					OutAU->bTrackChangeDiscontinuity = true;
diff --git a/Engine/Plugins/Media/ElectraPlayer/Source/ElectraPlayerRuntime/Private/Runtime/StreamAccessUnitBuffer.h b/Engine/Plugins/Media/ElectraPlayer/Source/ElectraPlayerRuntime/Private/Runtime/StreamAccessUnitBuffer.h
--- a/Engine/Plugins/Media/ElectraPlayer/Source/ElectraPlayerRuntime/Private/Runtime/StreamAccessUnitBuffer.h
+++ b/Engine/Plugins/Media/ElectraPlayer/Source/ElectraPlayerRuntime/Private/Runtime/StreamAccessUnitBuffer.h
@@ -32,13 +32,11 @@ namespace Electra
 
 	struct FStreamSourceInfo
 	{
-		FStreamSourceInfo() : NumericTrackID(0) {}
 		FString		Language;
 		FString		Role;
 		FString		PeriodID;
 		FString		AdaptationSetID;
 		FString		RepresentationID;
-		uint32		NumericTrackID;
 	};
 
 	struct FAccessUnit
@@ -544,9 +542,9 @@ namespace Electra
 		~FMultiTrackAccessUnitBuffer();
 		void CapacitySet(const FAccessUnitBuffer::FConfiguration& Config);
 		void AutoselectFirstTrack();
-		void SelectTrackByID(int32 TrackID);
+		void SelectTrackByID(const FString& TrackID);
 		void Deselect();
-		//int32 GetSelectedTrackID() const;
+		//FString GetSelectedTrackID() const;
 		bool Push(FAccessUnit*& AU);
 		void PushEndOfDataFor(TSharedPtr<const FStreamSourceInfo, ESPMode::ThreadSafe> InStreamSourceInfo);
 		void PushEndOfDataAll();
@@ -588,12 +586,12 @@ namespace Electra
 		TSharedPtr<FAccessUnitBuffer, ESPMode::ThreadSafe> GetSelectedTrackBuffer();
 		TSharedPtr<const FAccessUnitBuffer, ESPMode::ThreadSafe> GetSelectedTrackBuffer() const;
 
-		mutable FMediaCriticalSection						AccessLock;
+		mutable FMediaCriticalSection			AccessLock;
 		FAccessUnitBuffer::FConfiguration		PrimaryBufferConfiguration;			//!< Configuration of the primary (active) buffer. All other buffers are unbounded to ensure no stalling of the demuxer.
 		FAccessUnitBufferPtr					EmptyBuffer;						//!< An empty buffer
-		int32									ActiveOutputID;						//!< Track ID of the buffer that is feeding the decoder.
-		int32									LastPoppedBufferID;					//!< Buffer ID from which an AU was popped last.
-		TMap<uint32, FAccessUnitBufferPtr>		TrackBuffers;						//!< Map of track buffers. One per track ID.
+		FString									ActiveOutputID;						//!< Track ID of the buffer that is feeding the decoder.
+		FString									LastPoppedBufferID;					//!< Buffer ID from which an AU was popped last.
+		TMap<FString, FAccessUnitBufferPtr>		TrackBuffers;						//!< Map of track buffers. One per track ID.
 		FTimeValue								LastPoppedDTS;
 		FTimeValue								LastPoppedPTS;
 		bool									bAutoselectFirstTrack;
diff --git a/Engine/Plugins/Media/ElectraPlayer/Source/ElectraPlayerRuntime/Private/Runtime/StreamTypes.h b/Engine/Plugins/Media/ElectraPlayer/Source/ElectraPlayerRuntime/Private/Runtime/StreamTypes.h
--- a/Engine/Plugins/Media/ElectraPlayer/Source/ElectraPlayerRuntime/Private/Runtime/StreamTypes.h
+++ b/Engine/Plugins/Media/ElectraPlayer/Source/ElectraPlayerRuntime/Private/Runtime/StreamTypes.h
@@ -416,6 +416,36 @@ namespace Electra
 			// We do not compare frame rate here as we are interested in values that may require a decoder reconfiguration.
 			return Codec != Other.Codec || Resolution != Other.Resolution || ProfileLevel != Other.ProfileLevel || AspectRatio != Other.AspectRatio;
 		}
+
+		bool Equals(const FStreamCodecInformation& Other) const
+		{
+			if (StreamType == Other.StreamType && Codec == Other.Codec && CodecSpecifier.Equals(Other.CodecSpecifier) && StreamLanguageCode.Equals(Other.StreamLanguageCode))
+			{
+				if (StreamType == EStreamType::Video)
+				{
+					return Resolution == Other.Resolution &&
+						   Crop == Other.Crop &&
+						   AspectRatio == Other.AspectRatio &&
+						   FrameRate == Other.FrameRate &&
+						   ProfileLevel == Other.ProfileLevel;
+				}
+				else if (StreamType == EStreamType::Audio)
+				{
+					return SampleRate == Other.SampleRate &&
+						   NumChannels == Other.NumChannels &&
+						   ChannelConfiguration == Other.ChannelConfiguration &&
+						   AudioDecodingComplexity == Other.AudioDecodingComplexity &&
+						   AudioAccessibility == Other.AudioAccessibility &&
+						   NumberOfAudioObjects == Other.NumberOfAudioObjects;
+				}
+				else if (StreamType == EStreamType::Subtitle)
+				{
+					return true;
+				}
+			}
+			return false;
+		}
+
 	private:
 		struct FProfileLevel
 		{
@@ -475,14 +505,45 @@ namespace Electra
 	struct FStreamMetadata
 	{
 		FStreamCodecInformation		CodecInformation;					//!< Stream codec information
-		FString						PlaylistID;							//!< URL or some other ID of the playlist (FIXME: this is mostly for HLS)
+		FString						ID;									//!< ID of this stream
 		int32						Bandwidth;							//!< Bandwidth required for this stream in bits per second
-		uint32						StreamUniqueID;						//!< Uniquely identifies this stream
-		FString						LanguageCode;
-		// TODO: additional information
+		bool Equals(const FStreamMetadata& Other) const
+		{
+			return ID == Other.ID && Bandwidth == Other.Bandwidth && CodecInformation.Equals(Other.CodecInformation);
+		}
 	};
 
-
+	/**
+	 * Metadata per track type.
+	 */
+	struct FTrackMetadata
+	{
+		FString TrackID;
+		FString Role;
+		FString Language;
+		TArray<FStreamMetadata> StreamDetails;
+		FStreamCodecInformation HighestBandwidthCodec;
+		int32 HighestBandwidth = 0;
+		bool Equals(const FTrackMetadata& Other) const
+		{
+			bool bEquals = TrackID.Equals(Other.TrackID) && Role.Equals(Other.Role) && Language.Equals(Other.Language);
+			if (bEquals)
+			{
+				if (StreamDetails.Num() == Other.StreamDetails.Num())
+				{
+					for(int32 i=0; i<StreamDetails.Num(); ++i)
+					{
+						if (!StreamDetails[i].Equals(Other.StreamDetails[i]))
+						{
+							return false;
+						}
+					}
+					return true;
+				}
+			}
+			return false;
+		}
+	};
 
 
 	/**
diff --git a/Engine/Plugins/Media/ElectraUtil/Source/ElectraBase/Public/PlayerTime.h b/Engine/Plugins/Media/ElectraUtil/Source/ElectraBase/Public/PlayerTime.h
--- a/Engine/Plugins/Media/ElectraUtil/Source/ElectraBase/Public/PlayerTime.h
+++ b/Engine/Plugins/Media/ElectraUtil/Source/ElectraBase/Public/PlayerTime.h
@@ -561,6 +561,17 @@ public:
 
 	FTimeFraction& SetFromFloatString(const FString& In);
 
+
+	bool operator == (const FTimeFraction& rhs) const
+	{
+		return bIsValid == rhs.bIsValid && Numerator == rhs.Numerator && Denominator == rhs.Denominator;
+	}
+
+	bool operator != (const FTimeFraction& rhs) const
+	{
+		return !(*this == rhs);
+	}
+
 private:
 	int64	Numerator;
 	uint32	Denominator;
