ElectraPlayer: possible fix for a crash using a no longer valid IMediaOptions pointer; using a shared pointer when removing HTTP requests for safety

#rnx

diff --git a/Engine/Plugins/Media/ElectraPlayer/Source/ElectraPlayerPlugin/Private/ElectraPlayerPlugin.cpp b/Engine/Plugins/Media/ElectraPlayer/Source/ElectraPlayerPlugin/Private/ElectraPlayerPlugin.cpp
--- a/Engine/Plugins/Media/ElectraPlayer/Source/ElectraPlayerPlugin/Private/ElectraPlayerPlugin.cpp
+++ b/Engine/Plugins/Media/ElectraPlayer/Source/ElectraPlayerPlugin/Private/ElectraPlayerPlugin.cpp
@@ -15,8 +15,6 @@ using namespace Electra;
 //-----------------------------------------------------------------------------
 
 FElectraPlayerPlugin::FElectraPlayerPlugin()
-	: OptionInterface(nullptr)
-	, EventSink(nullptr)
 {
 	// Make sure a few assumptions are correct...
 	static_assert((int32)EMediaEvent::MediaBuffering == (int32)IElectraPlayerAdapterDelegate::EPlayerEvent::MediaBuffering, "check alignment of both enums");
@@ -64,7 +62,9 @@ bool FElectraPlayerPlugin::Initialize(IMediaEventSink& InEventSink,
 	FElectraPlayerReportVideoStreamingErrorDelegate& InReportVideoStreamingErrorDelegate,
 	FElectraPlayerReportSubtitlesMetricsDelegate& InReportSubtitlesFileMetricsDelegate)
 {
+	CallbackPointerLock.Lock();
 	EventSink = &InEventSink;
+	CallbackPointerLock.Unlock();
 
 	MediaSamples.Reset(new FMediaSamples);
 
@@ -78,6 +78,18 @@ bool FElectraPlayerPlugin::Initialize(IMediaEventSink& InEventSink,
 
 FElectraPlayerPlugin::~FElectraPlayerPlugin()
 {
+	CallbackPointerLock.Lock();
+	EventSink = nullptr;
+	OptionInterface = nullptr;
+	CallbackPointerLock.Unlock();
+	if (Player.IsValid())
+	{
+		Player->CloseInternal(true);
+		Player.Reset();
+	}
+	PlayerDelegate.Reset();
+	PlayerResourceDelegate.Reset();
+	MediaSamples.Reset();
 }
 
 //-----------------------------------------------------------------------------
@@ -86,70 +98,69 @@ Electra::FVariantValue FElectraPlayerPlugin::FPlayerAdapterDelegate::QueryOption
 {
 	if (TSharedPtr<FElectraPlayerPlugin, ESPMode::ThreadSafe> PinnedHost = Host.Pin())
 	{
-		switch (Type)
+		FScopeLock lock(&PinnedHost->CallbackPointerLock);
+		if (PinnedHost->OptionInterface)
 		{
-		case	EOptionType::MaxVerticalStreamResolution:
-			if (PinnedHost->OptionInterface)
-			{
-				return FVariantValue((int64)PinnedHost->OptionInterface->GetMediaOption(TEXT("MaxResolutionForMediaStreaming"), (int64)0));
-			}
-			break;
-
-		case	EOptionType::MaxBandwidthForStreaming:
-			if (PinnedHost->OptionInterface)
+			switch (Type)
 			{
-				return FVariantValue((int64)PinnedHost->OptionInterface->GetMediaOption(TEXT("ElectraMaxStreamingBandwidth"), (int64)0));
-			}
-			break;
+				case EOptionType::MaxVerticalStreamResolution:
+				{
+					return FVariantValue((int64)PinnedHost->OptionInterface->GetMediaOption(TEXT("MaxResolutionForMediaStreaming"), (int64)0));
+				}
 
-		case EOptionType::PlayListData:
-			if (PinnedHost->OptionInterface)
-			{
-				const FName PlaylistOptionKey = TEXT("ElectraGetPlaylistData");
-				if (PinnedHost->OptionInterface->HasMediaOption(PlaylistOptionKey))
+				case EOptionType::MaxBandwidthForStreaming:
 				{
-					check(Param.IsType(FVariantValue::EDataType::TypeFString));
-					return FVariantValue(PinnedHost->OptionInterface->GetMediaOption(PlaylistOptionKey, Param.GetFString()));
+					return FVariantValue((int64)PinnedHost->OptionInterface->GetMediaOption(TEXT("ElectraMaxStreamingBandwidth"), (int64)0));
 				}
-			}
-			break;
 
-		case EOptionType::LicenseKeyData:
-			if (PinnedHost->OptionInterface)
-			{
-				const FName LicenseKeyDataOptionKey = TEXT("ElectraGetLicenseKeyData");
-				if (PinnedHost->OptionInterface->HasMediaOption(LicenseKeyDataOptionKey))
+				case EOptionType::PlayListData:
 				{
-					check(Param.IsType(FVariantValue::EDataType::TypeFString));
-					return FVariantValue(PinnedHost->OptionInterface->GetMediaOption(LicenseKeyDataOptionKey, Param.GetFString()));
+					const FName PlaylistOptionKey = TEXT("ElectraGetPlaylistData");
+					if (PinnedHost->OptionInterface->HasMediaOption(PlaylistOptionKey))
+					{
+						check(Param.IsType(FVariantValue::EDataType::TypeFString));
+						return FVariantValue(PinnedHost->OptionInterface->GetMediaOption(PlaylistOptionKey, Param.GetFString()));
+					}
+					break;
 				}
-			}
-			break;
 
-		case EOptionType::PlaystartPosFromSeekPositions:
-			if (PinnedHost->OptionInterface)
-			{
-				const FName PlaystartOptionKey = TEXT("ElectraGetPlaystartPosFromSeekPositions");
-				if (PinnedHost->OptionInterface->HasMediaOption(PlaystartOptionKey))
+				case EOptionType::LicenseKeyData:
 				{
-					check(Param.IsType(FVariantValue::EDataType::TypeSharedPointer));
+					const FName LicenseKeyDataOptionKey = TEXT("ElectraGetLicenseKeyData");
+					if (PinnedHost->OptionInterface->HasMediaOption(LicenseKeyDataOptionKey))
+					{
+						check(Param.IsType(FVariantValue::EDataType::TypeFString));
+						return FVariantValue(PinnedHost->OptionInterface->GetMediaOption(LicenseKeyDataOptionKey, Param.GetFString()));
+					}
+					break;
+				}
 
-					TSharedPtr<TArray<FTimespan>, ESPMode::ThreadSafe> PosArray = Param.GetSharedPointer<TArray<FTimespan>>();
-					if (PosArray.IsValid())
+				case EOptionType::PlaystartPosFromSeekPositions:
+				{
+					const FName PlaystartOptionKey = TEXT("ElectraGetPlaystartPosFromSeekPositions");
+					if (PinnedHost->OptionInterface->HasMediaOption(PlaystartOptionKey))
 					{
-						TSharedPtr<FElectraSeekablePositions, ESPMode::ThreadSafe> Res = StaticCastSharedPtr<FElectraSeekablePositions, IMediaOptions::FDataContainer, ESPMode::ThreadSafe>(PinnedHost->OptionInterface->GetMediaOption(PlaystartOptionKey, MakeShared<FElectraSeekablePositions, ESPMode::ThreadSafe>(*PosArray)));
-						if (Res.IsValid() && Res->Data.Num())
+						check(Param.IsType(FVariantValue::EDataType::TypeSharedPointer));
+
+						TSharedPtr<TArray<FTimespan>, ESPMode::ThreadSafe> PosArray = Param.GetSharedPointer<TArray<FTimespan>>();
+						if (PosArray.IsValid())
 						{
-							return FVariantValue(int64(Res->Data[0].GetTicks())); // return HNS
+							TSharedPtr<FElectraSeekablePositions, ESPMode::ThreadSafe> Res = StaticCastSharedPtr<FElectraSeekablePositions, IMediaOptions::FDataContainer, ESPMode::ThreadSafe>(PinnedHost->OptionInterface->GetMediaOption(PlaystartOptionKey, MakeShared<FElectraSeekablePositions, ESPMode::ThreadSafe>(*PosArray)));
+							if (Res.IsValid() && Res->Data.Num())
+							{
+								return FVariantValue(int64(Res->Data[0].GetTicks())); // return HNS
+							}
 						}
+						return FVariantValue();
 					}
-					return FVariantValue();
+					break;
 				}
-			}
-			break;
 
-		default:
-			break;
+				default:
+				{
+					break;
+				}
+			}
 		}
 	}
 	return FVariantValue();
@@ -160,6 +171,7 @@ void FElectraPlayerPlugin::FPlayerAdapterDelegate::SendMediaEvent(EPlayerEvent E
 {
 	if (TSharedPtr<FElectraPlayerPlugin, ESPMode::ThreadSafe> PinnedHost = Host.Pin())
 	{
+		FScopeLock lock(&PinnedHost->CallbackPointerLock);
 		if (PinnedHost->EventSink)
 		{
 			PinnedHost->EventSink->ReceiveMediaEvent((EMediaEvent)Event);
@@ -295,7 +307,9 @@ bool FElectraPlayerPlugin::Open(const FString& Url, const IMediaOptions* Options
 bool FElectraPlayerPlugin::Open(const FString& Url, const IMediaOptions* Options, const FMediaPlayerOptions* InPlayerOptions)
 {
 	// Remember the option interface to poll for changes during playback.
+	CallbackPointerLock.Lock();
 	OptionInterface = Options;
+	CallbackPointerLock.Unlock();
 
 	IElectraPlayerInterface::FPlaystartOptions LocalPlaystartOptions;
 
@@ -422,6 +436,9 @@ bool FElectraPlayerPlugin::Open(const TSharedRef<FArchive, ESPMode::ThreadSafe>&
 */
 void FElectraPlayerPlugin::Close()
 {
+	CallbackPointerLock.Lock();
+	OptionInterface = nullptr;
+	CallbackPointerLock.Unlock();
 	Player->CloseInternal(true);
 }
 
diff --git a/Engine/Plugins/Media/ElectraPlayer/Source/ElectraPlayerPlugin/Private/ElectraPlayerPlugin.h b/Engine/Plugins/Media/ElectraPlayer/Source/ElectraPlayerPlugin/Private/ElectraPlayerPlugin.h
--- a/Engine/Plugins/Media/ElectraPlayer/Source/ElectraPlayerPlugin/Private/ElectraPlayerPlugin.h
+++ b/Engine/Plugins/Media/ElectraPlayer/Source/ElectraPlayerPlugin/Private/ElectraPlayerPlugin.h
@@ -180,11 +180,14 @@ private:
 	/** Output queues as needed by MediaFramework */
 	TUniquePtr<FMediaSamples> MediaSamples;
 
+	/** Lock to guard the POD callback pointers from being changed while being used. */
+	FCriticalSection CallbackPointerLock;
+
 	/** Option interface */
-	const IMediaOptions* OptionInterface;
+	const IMediaOptions* OptionInterface = nullptr;
 
 	/** The media event handler */
-	IMediaEventSink* EventSink;
+	IMediaEventSink* EventSink = nullptr;
 
 	/** The actual player */
 	TSharedPtr<IElectraPlayerInterface, ESPMode::ThreadSafe> Player;
diff --git a/Engine/Plugins/Media/ElectraPlayer/Source/ElectraPlayerRuntime/Private/Runtime/HTTP/HTTPManager.cpp b/Engine/Plugins/Media/ElectraPlayer/Source/ElectraPlayerRuntime/Private/Runtime/HTTP/HTTPManager.cpp
--- a/Engine/Plugins/Media/ElectraPlayer/Source/ElectraPlayerRuntime/Private/Runtime/HTTP/HTTPManager.cpp
+++ b/Engine/Plugins/Media/ElectraPlayer/Source/ElectraPlayerRuntime/Private/Runtime/HTTP/HTTPManager.cpp
@@ -312,14 +312,14 @@ namespace Electra
 		{
 			void SignalDone()
 			{
-				if (WaitingEvent)
+				if (WaitingEvent.IsValid())
 				{
 					WaitingEvent->Signal();
-					WaitingEvent = nullptr;
+					WaitingEvent.Reset();
 				}
 			}
 			TSharedPtrTS<FRequest> Request;
-			FMediaEvent* WaitingEvent = nullptr;
+			TSharedPtrTS<FMediaEvent> WaitingEvent;
 		};
 
 		FCriticalSection												Lock;
@@ -404,15 +404,15 @@ namespace Electra
 
 	void FElectraHttpManager::RemoveRequest(TSharedPtrTS<FRequest> Request)
 	{
-		FMediaEvent Sig;
+		TSharedPtrTS<FMediaEvent> WaitingEvent = MakeSharedTS<FMediaEvent>();
 		FRemoveRequest Remove;
 		Remove.Request = Request;
-		Remove.WaitingEvent = &Sig;
+		Remove.WaitingEvent = WaitingEvent;
 		Lock.Lock();
 		RequestsToRemove.Enqueue(MoveTemp(Remove));
 		RequestChangesEvent.Signal();
 		Lock.Unlock();
-		Sig.Wait();
+		WaitingEvent->Wait();
 	}
 
 	FElectraHttpManager::FHandle* FElectraHttpManager::CreateLocalFileHandle(const FTimeValue& Now, FTransportError& OutError, const TSharedPtrTS<IElectraHttpManager::FRequest>& Request)
