ElectraPlayer: fixed an infinite loop splitting a URL into its components when it specifies a path ending with a slash.

#rnx

diff --git a/Engine/Plugins/Media/ElectraPlayer/Source/ElectraPlayerRuntime/Private/Runtime/Utilities/StringHelpers.cpp b/Engine/Plugins/Media/ElectraPlayer/Source/ElectraPlayerRuntime/Private/Runtime/Utilities/StringHelpers.cpp
--- a/Engine/Plugins/Media/ElectraPlayer/Source/ElectraPlayerRuntime/Private/Runtime/Utilities/StringHelpers.cpp
+++ b/Engine/Plugins/Media/ElectraPlayer/Source/ElectraPlayerRuntime/Private/Runtime/Utilities/StringHelpers.cpp
@@ -101,11 +101,11 @@ namespace Electra
 					{
 						OutSplits.Push(subs);
 					}
-					if (SplitPos == INDEX_NONE)
+					FirstPos = SplitPos + SplitAt.Len();
+					if (SplitPos == INDEX_NONE || FirstPos >= InString.Len())
 					{
 						break;
 					}
-					FirstPos = SplitPos + SplitAt.Len();
 				}
 			}
 		}
