ElectraPlayer: Events coming from the player proper are now enqueued for being handled in Tick() only to ensure they do not happen asynchronously. Closing the player is now protected from being called inside Tick() and the destruction of the player instance outside of Tick().

#rnx

diff --git a/Engine/Plugins/Media/ElectraPlayer/Source/ElectraPlayerRuntime/Private/ElectraPlayer.cpp b/Engine/Plugins/Media/ElectraPlayer/Source/ElectraPlayerRuntime/Private/ElectraPlayer.cpp
--- a/Engine/Plugins/Media/ElectraPlayer/Source/ElectraPlayerRuntime/Private/ElectraPlayer.cpp
+++ b/Engine/Plugins/Media/ElectraPlayer/Source/ElectraPlayerRuntime/Private/ElectraPlayer.cpp
@@ -86,13 +86,12 @@ FElectraPlayer::FElectraPlayer(const TSharedPtr<IElectraPlayerAdapterDelegate, E
 	, SendAnalyticMetricsPerMinuteDelegate(InSendAnalyticMetricsPerMinuteDelegate)
 	, ReportVideoStreamingErrorDelegate(InReportVideoStreamingErrorDelegate)
 	, ReportSubtitlesMetricsDelegate(InReportSubtitlesFileMetricsDelegate)
-	, State(EPlayerState::Closed)
-	, Status(EPlayerStatus::None)
-	, bWasClosedOnError(false)
-	, bAllowKillAfterCloseEvent(false)
 {
 	CSV_EVENT(ElectraPlayer, TEXT("Player Creation"));
 
+	WaitForPlayerDestroyedEvent = FPlatformProcess::GetSynchEventFromPool(true);
+	WaitForPlayerDestroyedEvent->Trigger();
+
 	FString	OSMinor;
 	AnalyticsGPUType = InAdapterDelegate->GetVideoAdapterName().TrimStartAndEnd();
 	FPlatformMisc::GetOSVersions(AnalyticsOSVersion, OSMinor);
@@ -103,15 +102,19 @@ FElectraPlayer::FElectraPlayer(const TSharedPtr<IElectraPlayerAdapterDelegate, E
 	ReportVideoStreamingErrorDelegate.AddRaw(this, &FElectraPlayer::ReportVideoStreamingError);
 	ReportSubtitlesMetricsDelegate.AddRaw(this, &FElectraPlayer::ReportSubtitlesMetrics);
 
-	LastPresentedFrameDimension = FIntPoint::ZeroValue;
-
-	SelectedQuality = 0;
-	SelectedVideoTrackIndex = -1;
-	SelectedAudioTrackIndex = -1;
+	State = EPlayerState::Closed;
+	Status = EPlayerStatus::None;
+	bAllowKillAfterCloseEvent = false;
+	bCloseInProgress = false;
+	bHasPendingError = false;
+	AnalyticsInstanceEventCount = 0;
+	NumQueuedAnalyticEvents = 0;
 
 	StaticResourceProvider = MakeShared<FAdaptiveStreamingPlayerResourceProvider, ESPMode::ThreadSafe>(AdapterDelegate);
 	VideoDecoderResourceDelegate = PlatformCreateVideoDecoderResourceDelegate(AdapterDelegate);
 
+	ClearToDefaultState();
+
 	UE_LOG(LogElectraPlayer, Log, TEXT("[%p] FElectraPlayer() created."), this);
 }
 
@@ -122,6 +125,7 @@ FElectraPlayer::FElectraPlayer(const TSharedPtr<IElectraPlayerAdapterDelegate, E
 FElectraPlayer::~FElectraPlayer()
 {
 	CloseInternal(false);
+	WaitForPlayerDestroyedEvent->Wait();
 
 	CSV_EVENT(ElectraPlayer, TEXT("Player Destruction"));
 
@@ -135,6 +139,23 @@ FElectraPlayer::~FElectraPlayer()
 	{
 		AsyncResourceReleaseNotification->Signal(ResourceFlags_OutputBuffers);
 	}
+	
+	FPlatformProcess::ReturnSynchEventToPool(WaitForPlayerDestroyedEvent);
+}
+
+
+void FElectraPlayer::ClearToDefaultState()
+{
+	FScopeLock lock(&PlayerLock);
+
+	NumTracksAudio = 0;
+	NumTracksVideo = 0;
+	SelectedQuality = 0;
+	SelectedVideoTrackIndex = -1;
+	SelectedAudioTrackIndex = -1;
+	LastPresentedFrameDimension = FIntPoint::ZeroValue;
+	DeferredPlayerEvents.Empty();
+	MediaUrl.Empty();
 }
 
 
@@ -148,12 +169,13 @@ bool FElectraPlayer::OpenInternal(const FString& Url, const FParamDict & PlayerO
 	CSV_EVENT(ElectraPlayer, TEXT("Open"));
 
 	CloseInternal(false);
-
 	// Clear out our work variables
-	LastPresentedFrameDimension = FIntPoint::ZeroValue;
+	ClearToDefaultState();
+	bAllowKillAfterCloseEvent = false;
 	State = EPlayerState::Closed;
 	Status = EPlayerStatus::None;
-	bWasClosedOnError = false;
+	bCloseInProgress = false;
+	bHasPendingError = false;
 
 	// Start statistics with a clean slate.
 	Statistics.Reset();
@@ -163,7 +185,6 @@ bool FElectraPlayer::OpenInternal(const FString& Url, const FParamDict & PlayerO
 	// Create a guid string for the analytics. We do this here and not in the constructor in case the same instance is used over again.
 	AnalyticsInstanceGuid = FGuid::NewGuid().ToString(EGuidFormats::Digits);
 
-
 	PlaystartOptions = InPlaystartOptions;
 
 	// Create a new empty player structure. This contains the actual player instance, its associated renderers and sample queues.
@@ -222,10 +243,16 @@ void FElectraPlayer::CloseInternal(bool bKillAfterClose)
 {
 	LLM_SCOPE(ELLMTag::ElectraPlayer);
 
-	if (!CurrentPlayer.IsValid())
+	PlayerLock.Lock();
+	if (bCloseInProgress || !CurrentPlayer.IsValid())
 	{
+		PlayerLock.Unlock();
 		return;
 	}
+	bCloseInProgress = true;
+	WaitForPlayerDestroyedEvent->Reset();
+	PlayerLock.Unlock();
+
 
 	UE_LOG(LogElectraPlayer, Log, TEXT("[%p][%p] IMediaPlayer::Close()"), this, CurrentPlayer.Get());
 	CSV_EVENT(ElectraPlayer, TEXT("Close"));
@@ -288,18 +315,21 @@ void FElectraPlayer::CloseInternal(bool bKillAfterClose)
 	StaticResourceProvider->ClearPendingRequests();
 
 	// Pretend we got the playback stopped event via metrics, which we did not because we unregistered ourselves already.
-	ReportPlaybackStopped();
+	HandlePlayerEventPlaybackStopped();
 	LogStatistics();
 	DeferredEvents.Enqueue(IElectraPlayerAdapterDelegate::EPlayerEvent::TracksChanged);
 	DeferredEvents.Enqueue(IElectraPlayerAdapterDelegate::EPlayerEvent::MediaClosed);
 
 	// Clear out the player instance now.
 	CurrentPlayer.Reset();
+	bCloseInProgress = false;
 
 	// Kick off asynchronous closing now.
 	FInternalPlayerImpl::DoCloseAsync(MoveTemp(Player), AsyncResourceReleaseNotification);
 
 	bAllowKillAfterCloseEvent = bKillAfterClose;
+
+	WaitForPlayerDestroyedEvent->Trigger();
 }
 
 void FElectraPlayer::FInternalPlayerImpl::DoCloseAsync(TSharedPtr<FInternalPlayerImpl, ESPMode::ThreadSafe> && Player, TSharedPtr<FElectraPlayer::IAsyncResourceReleaseNotifyContainer, ESPMode::ThreadSafe> AsyncResourceReleaseNotification)
@@ -365,16 +395,10 @@ void FElectraPlayer::Tick(FTimespan DeltaTime, FTimespan Timecode)
 	SCOPE_CYCLE_COUNTER(STAT_ElectraPlayer_ElectraPlayer_TickInput);
 	CSV_SCOPED_TIMING_STAT(ElectraPlayer, TickInput);
 
-	if (State == EPlayerState::Error && !bWasClosedOnError)
-	{
-		CloseInternal(true);
-		// Close has set the state to close but we want it to stick at error.
-		State = EPlayerState::Error;
-		bWasClosedOnError = true;
-	}
-
 	// Handle the internal player, if we have one.
-	if (CurrentPlayer.IsValid())
+	PlayerLock.Lock();
+	TSharedPtr<FInternalPlayerImpl, ESPMode::ThreadSafe> Player = CurrentPlayer;
+	if (!bCloseInProgress && CurrentPlayer.IsValid() && State != EPlayerState::Error)
 	{
 		if (CurrentPlayer->RendererVideo.IsValid())
 		{
@@ -410,7 +434,30 @@ void FElectraPlayer::Tick(FTimespan DeltaTime, FTimespan Timecode)
 				}
 			}
 		}
+		
+		// Process accumulated player events.
+		HandleDeferredPlayerEvents();
+		if (bHasPendingError)
+		{
+			bHasPendingError = false;
+			if (State == EPlayerState::Preparing)
+			{
+				DeferredEvents.Enqueue(IElectraPlayerAdapterDelegate::EPlayerEvent::MediaOpenFailed);
+			}
+			else if (State == EPlayerState::Playing)
+			{
+				DeferredEvents.Enqueue(IElectraPlayerAdapterDelegate::EPlayerEvent::MediaClosed);
+			}
+
+			CloseInternal(true);
+			State = EPlayerState::Error;
+		}
 	}
+	else
+	{
+		DeferredPlayerEvents.Empty();
+	}
+	PlayerLock.Unlock();
 
 	// Forward enqueued session events. We do this even with no current internal player to ensure all pending events are sent and none are lost.
 	if (TSharedPtr<IElectraPlayerAdapterDelegate, ESPMode::ThreadSafe> PinnedAdapterDelegate = AdapterDelegate.Pin())
@@ -999,15 +1046,138 @@ void FElectraPlayer::EnqueueAnalyticsEvent(TSharedPtr<FAnalyticsEvent> InAnalyti
 /**
  *	State management information from media player...
  */
-
-void FElectraPlayer::ReportOpenSource(const FString& URL)
+void FElectraPlayer::HandleDeferredPlayerEvents()
 {
-	TSharedPtr<FInternalPlayerImpl, ESPMode::ThreadSafe> LockedPlayer = CurrentPlayer;
-	if (!LockedPlayer.IsValid() || !LockedPlayer->AdaptivePlayer.IsValid())
+	TSharedPtrTS<FPlayerMetricEventBase> Event;
+	while(DeferredPlayerEvents.Dequeue(Event))
 	{
-		return;
+		switch(Event->Type)
+		{
+			case FPlayerMetricEventBase::EType::OpenSource:
+			{
+				FPlayerMetricEvent_OpenSource* Ev = static_cast<FPlayerMetricEvent_OpenSource*>(Event.Get());
+				HandlePlayerEventOpenSource(Ev->URL);
+				break;
+			}
+			case FPlayerMetricEventBase::EType::ReceivedMasterPlaylist:
+			{
+				FPlayerMetricEvent_ReceivedMasterPlaylist* Ev = static_cast<FPlayerMetricEvent_ReceivedMasterPlaylist*>(Event.Get());
+				HandlePlayerEventReceivedMasterPlaylist(Ev->EffectiveURL);
+				break;
+			}
+			case FPlayerMetricEventBase::EType::ReceivedPlaylists:
+				HandlePlayerEventReceivedPlaylists();
+				break;
+			case FPlayerMetricEventBase::EType::TracksChanged:
+				HandlePlayerEventTracksChanged();
+				break;
+			case FPlayerMetricEventBase::EType::PlaylistDownload:
+			{
+				FPlayerMetricEvent_PlaylistDownload* Ev = static_cast<FPlayerMetricEvent_PlaylistDownload*>(Event.Get());
+				HandlePlayerEventPlaylistDownload(Ev->PlaylistDownloadStats);
+				break;
+			}
+			case FPlayerMetricEventBase::EType::BufferingStart:
+			{
+				FPlayerMetricEvent_BufferingStart* Ev = static_cast<FPlayerMetricEvent_BufferingStart*>(Event.Get());
+				HandlePlayerEventBufferingStart(Ev->BufferingReason);
+				break;
+			}
+			case FPlayerMetricEventBase::EType::BufferingEnd:
+			{
+				FPlayerMetricEvent_BufferingEnd* Ev = static_cast<FPlayerMetricEvent_BufferingEnd*>(Event.Get());
+				HandlePlayerEventBufferingEnd(Ev->BufferingReason);
+				break;
+			}
+			case FPlayerMetricEventBase::EType::Bandwidth:
+			{
+				FPlayerMetricEvent_Bandwidth* Ev = static_cast<FPlayerMetricEvent_Bandwidth*>(Event.Get());
+				HandlePlayerEventBandwidth(Ev->EffectiveBps, Ev->ThroughputBps, Ev->LatencyInSeconds);
+				break;
+			}
+			case FPlayerMetricEventBase::EType::BufferUtilization:
+			{
+				FPlayerMetricEvent_BufferUtilization* Ev = static_cast<FPlayerMetricEvent_BufferUtilization*>(Event.Get());
+				HandlePlayerEventBufferUtilization(Ev->BufferStats);
+				break;
+			}
+			case FPlayerMetricEventBase::EType::SegmentDownload:
+			{
+				FPlayerMetricEvent_SegmentDownload* Ev = static_cast<FPlayerMetricEvent_SegmentDownload*>(Event.Get());
+				HandlePlayerEventSegmentDownload(Ev->SegmentDownloadStats);
+				break;
+			}
+			case FPlayerMetricEventBase::EType::LicenseKey:
+			{
+				FPlayerMetricEvent_LicenseKey* Ev = static_cast<FPlayerMetricEvent_LicenseKey*>(Event.Get());
+				HandlePlayerEventLicenseKey(Ev->LicenseKeyStats);
+				break;
+			}
+			case FPlayerMetricEventBase::EType::DataAvailabilityChange:
+			{
+				FPlayerMetricEvent_DataAvailabilityChange* Ev = static_cast<FPlayerMetricEvent_DataAvailabilityChange*>(Event.Get());
+				HandlePlayerEventDataAvailabilityChange(Ev->DataAvailability);
+				break;
+			}
+			case FPlayerMetricEventBase::EType::VideoQualityChange:
+			{
+				FPlayerMetricEvent_VideoQualityChange* Ev = static_cast<FPlayerMetricEvent_VideoQualityChange*>(Event.Get());
+				HandlePlayerEventVideoQualityChange(Ev->NewBitrate, Ev->PreviousBitrate, Ev->bIsDrasticDownswitch);
+				break;
+			}
+			case FPlayerMetricEventBase::EType::PrerollStart:
+				HandlePlayerEventPrerollStart();
+				break;
+			case FPlayerMetricEventBase::EType::PrerollEnd:
+				HandlePlayerEventPrerollEnd();
+				break;
+			case FPlayerMetricEventBase::EType::PlaybackStart:
+				HandlePlayerEventPlaybackStart();
+				break;
+			case FPlayerMetricEventBase::EType::PlaybackPaused:
+				HandlePlayerEventPlaybackPaused();
+				break;
+			case FPlayerMetricEventBase::EType::PlaybackResumed:
+				HandlePlayerEventPlaybackResumed();
+				break;
+			case FPlayerMetricEventBase::EType::PlaybackEnded:
+				HandlePlayerEventPlaybackEnded();
+				break;
+			case FPlayerMetricEventBase::EType::JumpInPlayPosition:
+			{
+				FPlayerMetricEvent_JumpInPlayPosition* Ev = static_cast<FPlayerMetricEvent_JumpInPlayPosition*>(Event.Get());
+				HandlePlayerEventJumpInPlayPosition(Ev->ToNewTime, Ev->FromTime, Ev->TimejumpReason);
+				break;
+			}
+			case FPlayerMetricEventBase::EType::PlaybackStopped:
+				HandlePlayerEventPlaybackStopped();
+				break;
+			case FPlayerMetricEventBase::EType::Error:
+			{
+				FPlayerMetricEvent_Error* Ev = static_cast<FPlayerMetricEvent_Error*>(Event.Get());
+				HandlePlayerEventError(Ev->ErrorReason);
+				break;
+			}
+			case FPlayerMetricEventBase::EType::LogMessage:
+			{
+				FPlayerMetricEvent_LogMessage* Ev = static_cast<FPlayerMetricEvent_LogMessage*>(Event.Get());
+				HandlePlayerEventLogMessage(Ev->LogLevel, Ev->LogMessage, Ev->PlayerWallclockMilliseconds);
+				break;
+			}
+			case FPlayerMetricEventBase::EType::DroppedVideoFrame:
+				HandlePlayerEventDroppedVideoFrame();
+				break;
+			case FPlayerMetricEventBase::EType::DroppedAudioFrame:
+				HandlePlayerEventDroppedAudioFrame();
+				break; 
+			default:
+				break;
+		}
 	}
+}
 
+void FElectraPlayer::HandlePlayerEventOpenSource(const FString& URL)
+{
 	Status = Status | EPlayerStatus::Connecting;
 
 	State = EPlayerState::Preparing;
@@ -1031,14 +1201,8 @@ void FElectraPlayer::ReportOpenSource(const FString& URL)
 	}
 }
 
-void FElectraPlayer::ReportReceivedMasterPlaylist(const FString& EffectiveURL)
+void FElectraPlayer::HandlePlayerEventReceivedMasterPlaylist(const FString& EffectiveURL)
 {
-	TSharedPtr<FInternalPlayerImpl, ESPMode::ThreadSafe> LockedPlayer = CurrentPlayer;
-	if (!LockedPlayer.IsValid() || !LockedPlayer->AdaptivePlayer.IsValid())
-	{
-		return;
-	}
-
 	// Update statistics
 	FScopeLock Lock(&StatisticsLock);
 	// Note the time it took to get the master playlist
@@ -1057,14 +1221,8 @@ void FElectraPlayer::ReportReceivedMasterPlaylist(const FString& EffectiveURL)
 	}
 }
 
-void FElectraPlayer::ReportReceivedPlaylists()
+void FElectraPlayer::HandlePlayerEventReceivedPlaylists()
 {
-	TSharedPtr<FInternalPlayerImpl, ESPMode::ThreadSafe> LockedPlayer = CurrentPlayer;
-	if (!LockedPlayer.IsValid() || !LockedPlayer->AdaptivePlayer.IsValid())
-	{
-		return;
-	}
-
 	Status = Status & ~EPlayerStatus::Connecting;
 
 	//Status = Status | EPlayerStatus::Buffering;
@@ -1189,14 +1347,8 @@ void FElectraPlayer::ReportReceivedPlaylists()
 	CurrentPlayer->AdaptivePlayer->SeekTo(playParam);
 }
 
-void FElectraPlayer::ReportTracksChanged()
+void FElectraPlayer::HandlePlayerEventTracksChanged()
 {
-	TSharedPtr<FInternalPlayerImpl, ESPMode::ThreadSafe> LockedPlayer = CurrentPlayer;
-	if (!LockedPlayer.IsValid() || !LockedPlayer->AdaptivePlayer.IsValid())
-	{
-		return;
-	}
-
 	TArray<Electra::FTrackMetadata> VideoStreamMetaData;
 	CurrentPlayer->AdaptivePlayer->GetTrackMetadata(VideoStreamMetaData, Electra::EStreamType::Video);
 	NumTracksVideo = VideoStreamMetaData.Num();
@@ -1219,14 +1371,8 @@ void FElectraPlayer::ReportTracksChanged()
 }
 
 
-void FElectraPlayer::ReportPlaylistDownload(const Electra::Metrics::FPlaylistDownloadStats& PlaylistDownloadStats)
+void FElectraPlayer::HandlePlayerEventPlaylistDownload(const Electra::Metrics::FPlaylistDownloadStats& PlaylistDownloadStats)
 {
-	TSharedPtr<FInternalPlayerImpl, ESPMode::ThreadSafe> LockedPlayer = CurrentPlayer;
-	if (!LockedPlayer.IsValid() || !LockedPlayer->AdaptivePlayer.IsValid())
-	{
-		return;
-	}
-
 	// To reduce the number of playlist events during a Live presentation we will only report the initial playlist load
 	// and later on only failed loads but not successful ones.
 	bool bReport =  PlaylistDownloadStats.LoadType == Electra::Playlist::ELoadType::Initial || !PlaylistDownloadStats.bWasSuccessful;
@@ -1249,7 +1395,7 @@ void FElectraPlayer::ReportPlaylistDownload(const Electra::Metrics::FPlaylistDow
 	}
 }
 
-void FElectraPlayer::ReportLicenseKey(const Electra::Metrics::FLicenseKeyStats& LicenseKeyStats)
+void FElectraPlayer::HandlePlayerEventLicenseKey(const Electra::Metrics::FLicenseKeyStats& LicenseKeyStats)
 {
 	// TBD
 	if (LicenseKeyStats.bWasSuccessful)
@@ -1263,7 +1409,7 @@ void FElectraPlayer::ReportLicenseKey(const Electra::Metrics::FLicenseKeyStats&
 }
 
 
-void FElectraPlayer::ReportDataAvailabilityChange(const Electra::Metrics::FDataAvailabilityChange& DataAvailability)
+void FElectraPlayer::HandlePlayerEventDataAvailabilityChange(const Electra::Metrics::FDataAvailabilityChange& DataAvailability)
 {
 	// Pass this event up to the media player facade. We do not act on this here right now.
 	if (DataAvailability.StreamType == Electra::EStreamType::Video)
@@ -1291,14 +1437,8 @@ void FElectraPlayer::ReportDataAvailabilityChange(const Electra::Metrics::FDataA
 }
 
 
-void FElectraPlayer::ReportBufferingStart(Electra::Metrics::EBufferingReason BufferingReason)
+void FElectraPlayer::HandlePlayerEventBufferingStart(Electra::Metrics::EBufferingReason BufferingReason)
 {
-	TSharedPtr<FInternalPlayerImpl, ESPMode::ThreadSafe> LockedPlayer = CurrentPlayer;
-	if (!LockedPlayer.IsValid() || !LockedPlayer->AdaptivePlayer.IsValid())
-	{
-		return;
-	}
-
 	Status = Status | EPlayerStatus::Buffering;
 
 	// Update statistics
@@ -1330,14 +1470,8 @@ void FElectraPlayer::ReportBufferingStart(Electra::Metrics::EBufferingReason Buf
 	CSV_EVENT(ElectraPlayer, TEXT("Buffering starts"));
 }
 
-void FElectraPlayer::ReportBufferingEnd(Electra::Metrics::EBufferingReason BufferingReason)
+void FElectraPlayer::HandlePlayerEventBufferingEnd(Electra::Metrics::EBufferingReason BufferingReason)
 {
-	TSharedPtr<FInternalPlayerImpl, ESPMode::ThreadSafe> LockedPlayer = CurrentPlayer;
-	if (!LockedPlayer.IsValid() || !LockedPlayer->AdaptivePlayer.IsValid())
-	{
-		return;
-	}
-
 	// Note: While this event signals the end of buffering the player will now immediately transition into the pre-rolling
 	//       state from which a playback start is not quite possible yet and would incur a slight delay until it is.
 	//       To avoid this we keep the state as buffering until the pre-rolling phase has also completed.
@@ -1378,36 +1512,19 @@ void FElectraPlayer::ReportBufferingEnd(Electra::Metrics::EBufferingReason Buffe
 	CSV_EVENT(ElectraPlayer, TEXT("Buffering ends"));
 }
 
-void FElectraPlayer::ReportBandwidth(int64 EffectiveBps, int64 ThroughputBps, double LatencyInSeconds)
+void FElectraPlayer::HandlePlayerEventBandwidth(int64 EffectiveBps, int64 ThroughputBps, double LatencyInSeconds)
 {
-	TSharedPtr<FInternalPlayerImpl, ESPMode::ThreadSafe> LockedPlayer = CurrentPlayer;
-	if (!LockedPlayer.IsValid() || !LockedPlayer->AdaptivePlayer.IsValid())
-	{
-		return;
-	}
-
-	//	FScopeLock Lock(&StatisticsLock);
+//	FScopeLock Lock(&StatisticsLock);
 	UE_LOG(LogElectraPlayer, Verbose, TEXT("[%p][%p] Observed bandwidth of %d bps; throughput = %d; latency = %.3fs"), this, CurrentPlayer.Get(), (int32)EffectiveBps, (int32)ThroughputBps, LatencyInSeconds);
 }
 
-void FElectraPlayer::ReportBufferUtilization(const Electra::Metrics::FBufferStats& BufferStats)
+void FElectraPlayer::HandlePlayerEventBufferUtilization(const Electra::Metrics::FBufferStats& BufferStats)
 {
-	TSharedPtr<FInternalPlayerImpl, ESPMode::ThreadSafe> LockedPlayer = CurrentPlayer;
-	if (!LockedPlayer.IsValid() || !LockedPlayer->AdaptivePlayer.IsValid())
-	{
-		return;
-	}
 //	FScopeLock Lock(&StatisticsLock);
 }
 
-void FElectraPlayer::ReportSegmentDownload(const Electra::Metrics::FSegmentDownloadStats& SegmentDownloadStats)
+void FElectraPlayer::HandlePlayerEventSegmentDownload(const Electra::Metrics::FSegmentDownloadStats& SegmentDownloadStats)
 {
-	TSharedPtr<FInternalPlayerImpl, ESPMode::ThreadSafe> LockedPlayer = CurrentPlayer;
-	if (!LockedPlayer.IsValid() || !LockedPlayer->AdaptivePlayer.IsValid())
-	{
-		return;
-	}
-
 	// Update statistics
 	FScopeLock Lock(&StatisticsLock);
 	if (SegmentDownloadStats.StreamType == Electra::EStreamType::Video)
@@ -1467,14 +1584,8 @@ void FElectraPlayer::ReportSegmentDownload(const Electra::Metrics::FSegmentDownl
 	}
 }
 
-void FElectraPlayer::ReportVideoQualityChange(int32 NewBitrate, int32 PreviousBitrate, bool bIsDrasticDownswitch)
+void FElectraPlayer::HandlePlayerEventVideoQualityChange(int32 NewBitrate, int32 PreviousBitrate, bool bIsDrasticDownswitch)
 {
-	TSharedPtr<FInternalPlayerImpl, ESPMode::ThreadSafe> LockedPlayer = CurrentPlayer;
-	if (!LockedPlayer.IsValid() || !LockedPlayer->AdaptivePlayer.IsValid())
-	{
-		return;
-	}
-
 	// Update statistics
 	FScopeLock Lock(&StatisticsLock);
 	if (PreviousBitrate == 0)
@@ -1541,14 +1652,8 @@ void FElectraPlayer::ReportVideoQualityChange(int32 NewBitrate, int32 PreviousBi
 	CSV_EVENT(ElectraPlayer, TEXT("QualityChange %d -> %d"), PreviousBitrate, NewBitrate);
 }
 
-void FElectraPlayer::ReportPrerollStart()
+void FElectraPlayer::HandlePlayerEventPrerollStart()
 {
-	TSharedPtr<FInternalPlayerImpl, ESPMode::ThreadSafe> LockedPlayer = CurrentPlayer;
-	if (!LockedPlayer.IsValid() || !LockedPlayer->AdaptivePlayer.IsValid())
-	{
-		return;
-	}
-
 	// Update statistics
 	FScopeLock Lock(&StatisticsLock);
 	Statistics.TimeAtPrerollBegin = FPlatformTime::Seconds();
@@ -1565,14 +1670,8 @@ void FElectraPlayer::ReportPrerollStart()
 	}
 }
 
-void FElectraPlayer::ReportPrerollEnd()
+void FElectraPlayer::HandlePlayerEventPrerollEnd()
 {
-	TSharedPtr<FInternalPlayerImpl, ESPMode::ThreadSafe> LockedPlayer = CurrentPlayer;
-	if (!LockedPlayer.IsValid() || !LockedPlayer->AdaptivePlayer.IsValid())
-	{
-		return;
-	}
-
 	// Note: See comments in ReportBufferingEnd()
 	//       Preroll follows at the end of buffering and we keep the buffering state until preroll has finished as well.
 	Status = Status & ~EPlayerStatus::Buffering;
@@ -1595,14 +1694,8 @@ void FElectraPlayer::ReportPrerollEnd()
 	}
 }
 
-void FElectraPlayer::ReportPlaybackStart()
+void FElectraPlayer::HandlePlayerEventPlaybackStart()
 {
-	TSharedPtr<FInternalPlayerImpl, ESPMode::ThreadSafe> LockedPlayer = CurrentPlayer;
-	if (!LockedPlayer.IsValid() || !LockedPlayer->AdaptivePlayer.IsValid())
-	{
-		return;
-	}
-
 	Status = Status & ~EPlayerStatus::Buffering;
 	MediaStateOnPlay();
 
@@ -1626,14 +1719,8 @@ void FElectraPlayer::ReportPlaybackStart()
 	}
 }
 
-void FElectraPlayer::ReportPlaybackPaused()
+void FElectraPlayer::HandlePlayerEventPlaybackPaused()
 {
-	TSharedPtr<FInternalPlayerImpl, ESPMode::ThreadSafe> LockedPlayer = CurrentPlayer;
-	if (!LockedPlayer.IsValid() || !LockedPlayer->AdaptivePlayer.IsValid())
-	{
-		return;
-	}
-
 	MediaStateOnPause();
 	FScopeLock Lock(&StatisticsLock);
 	Statistics.LastState = "Paused";
@@ -1650,14 +1737,8 @@ void FElectraPlayer::ReportPlaybackPaused()
 	}
 }
 
-void FElectraPlayer::ReportPlaybackResumed()
+void FElectraPlayer::HandlePlayerEventPlaybackResumed()
 {
-	TSharedPtr<FInternalPlayerImpl, ESPMode::ThreadSafe> LockedPlayer = CurrentPlayer;
-	if (!LockedPlayer.IsValid() || !LockedPlayer->AdaptivePlayer.IsValid())
-	{
-		return;
-	}
-
 	MediaStateOnPlay();
 	FScopeLock Lock(&StatisticsLock);
 	Statistics.LastState = "Playing";
@@ -1674,14 +1755,8 @@ void FElectraPlayer::ReportPlaybackResumed()
 	}
 }
 
-void FElectraPlayer::ReportPlaybackEnded()
+void FElectraPlayer::HandlePlayerEventPlaybackEnded()
 {
-	TSharedPtr<FInternalPlayerImpl, ESPMode::ThreadSafe> LockedPlayer = CurrentPlayer;
-	if (!LockedPlayer.IsValid() || !LockedPlayer->AdaptivePlayer.IsValid())
-	{
-		return;
-	}
-
 	UpdatePlayEndStatistics();
 	double PlayPos = CurrentPlayer->AdaptivePlayer->GetPlayPosition().GetAsSeconds();
 
@@ -1704,14 +1779,8 @@ void FElectraPlayer::ReportPlaybackEnded()
 	MediaStateOnEndReached();
 }
 
-void FElectraPlayer::ReportJumpInPlayPosition(const Electra::FTimeValue& ToNewTime, const Electra::FTimeValue& FromTime, Electra::Metrics::ETimeJumpReason TimejumpReason)
+void FElectraPlayer::HandlePlayerEventJumpInPlayPosition(const Electra::FTimeValue& ToNewTime, const Electra::FTimeValue& FromTime, Electra::Metrics::ETimeJumpReason TimejumpReason)
 {
-	TSharedPtr<FInternalPlayerImpl, ESPMode::ThreadSafe> LockedPlayer = CurrentPlayer;
-	if (!LockedPlayer.IsValid() || !LockedPlayer->AdaptivePlayer.IsValid())
-	{
-		return;
-	}
-
 	Electra::FTimeRange MediaTimeline;
 	CurrentPlayer->AdaptivePlayer->GetTimelineRange(MediaTimeline);
 
@@ -1751,14 +1820,8 @@ void FElectraPlayer::ReportJumpInPlayPosition(const Electra::FTimeValue& ToNewTi
 	}
 }
 
-void FElectraPlayer::ReportPlaybackStopped()
+void FElectraPlayer::HandlePlayerEventPlaybackStopped()
 {
-	TSharedPtr<FInternalPlayerImpl, ESPMode::ThreadSafe> LockedPlayer = CurrentPlayer;
-	if (!LockedPlayer.IsValid() || !LockedPlayer->AdaptivePlayer.IsValid())
-	{
-		return;
-	}
-
 	UpdatePlayEndStatistics();
 	double PlayPos = CurrentPlayer->AdaptivePlayer->GetPlayPosition().GetAsSeconds();
 
@@ -1779,16 +1842,10 @@ void FElectraPlayer::ReportPlaybackStopped()
 }
 
 
-void FElectraPlayer::ReportError(const FString& ErrorReason)
+void FElectraPlayer::HandlePlayerEventError(const FString& ErrorReason)
 {
-	TSharedPtr<FInternalPlayerImpl, ESPMode::ThreadSafe> LockedPlayer = CurrentPlayer;
-	if (!LockedPlayer.IsValid() || !LockedPlayer->AdaptivePlayer.IsValid())
-	{
-		return;
-	}
+	bHasPendingError = true;
 
-	MediaStateOnError();
-	
 	// Update statistics
 	FScopeLock Lock(&StatisticsLock);
 	// If there is already an error do not overwrite it. First come, first serve!
@@ -1810,14 +1867,8 @@ void FElectraPlayer::ReportError(const FString& ErrorReason)
 	}
 }
 
-void FElectraPlayer::ReportLogMessage(Electra::IInfoLog::ELevel InLogLevel, const FString& InLogMessage, int64 InPlayerWallclockMilliseconds)
+void FElectraPlayer::HandlePlayerEventLogMessage(Electra::IInfoLog::ELevel InLogLevel, const FString& InLogMessage, int64 InPlayerWallclockMilliseconds)
 {
-	TSharedPtr<FInternalPlayerImpl, ESPMode::ThreadSafe> LockedPlayer = CurrentPlayer;
-	if (!LockedPlayer.IsValid() || !LockedPlayer->AdaptivePlayer.IsValid())
-	{
-		return;
-	}
-
 	FString m(SanitizeMessage(InLogMessage));
 	switch(InLogLevel)
 	{
@@ -1836,22 +1887,12 @@ void FElectraPlayer::ReportLogMessage(Electra::IInfoLog::ELevel InLogLevel, cons
 	}
 }
 
-void FElectraPlayer::ReportDroppedVideoFrame()
+void FElectraPlayer::HandlePlayerEventDroppedVideoFrame()
 {
-	TSharedPtr<FInternalPlayerImpl, ESPMode::ThreadSafe> LockedPlayer = CurrentPlayer;
-	if (!LockedPlayer.IsValid() || !LockedPlayer->AdaptivePlayer.IsValid())
-	{
-		return;
-	}
 }
 
-void FElectraPlayer::ReportDroppedAudioFrame()
+void FElectraPlayer::HandlePlayerEventDroppedAudioFrame()
 {
-	TSharedPtr<FInternalPlayerImpl, ESPMode::ThreadSafe> LockedPlayer = CurrentPlayer;
-	if (!LockedPlayer.IsValid() || !LockedPlayer->AdaptivePlayer.IsValid())
-	{
-		return;
-	}
 }
 
 
@@ -2210,21 +2251,6 @@ void FElectraPlayer::MediaStateOnSeekFinished()
 	DeferredEvents.Enqueue(IElectraPlayerAdapterDelegate::EPlayerEvent::SeekCompleted);
 }
 
-void FElectraPlayer::MediaStateOnError()
-{
-	CSV_EVENT(ElectraPlayer, TEXT("MediaStateOnError"));
-
-	if (State == EPlayerState::Preparing)
-	{
-		DeferredEvents.Enqueue(IElectraPlayerAdapterDelegate::EPlayerEvent::MediaOpenFailed);
-	}
-	else if (State == EPlayerState::Playing)
-	{
-		DeferredEvents.Enqueue(IElectraPlayerAdapterDelegate::EPlayerEvent::MediaClosed);
-	}
-	State = EPlayerState::Error;
-}
-
 // ------------------------------------------------------------------------------------------------------------------
 
 FElectraPlayer::FAdaptiveStreamingPlayerResourceProvider::FAdaptiveStreamingPlayerResourceProvider(const TWeakPtr<IElectraPlayerAdapterDelegate, ESPMode::ThreadSafe>& InAdapterDelegate)
diff --git a/Engine/Plugins/Media/ElectraPlayer/Source/ElectraPlayerRuntime/Private/ElectraPlayer.h b/Engine/Plugins/Media/ElectraPlayer/Source/ElectraPlayerRuntime/Private/ElectraPlayer.h
--- a/Engine/Plugins/Media/ElectraPlayer/Source/ElectraPlayerRuntime/Private/ElectraPlayer.h
+++ b/Engine/Plugins/Media/ElectraPlayer/Source/ElectraPlayerRuntime/Private/ElectraPlayer.h
@@ -119,48 +119,188 @@ public:
 	void NotifyOfOptionChange() override;
 
 private:
+	struct FPlayerMetricEventBase
+	{
+		enum class EType
+		{
+			OpenSource,
+			ReceivedMasterPlaylist,
+			ReceivedPlaylists,
+			TracksChanged,
+			PlaylistDownload,
+			BufferingStart,
+			BufferingEnd,
+			Bandwidth,
+			BufferUtilization,
+			SegmentDownload,
+			LicenseKey,
+			DataAvailabilityChange,
+			VideoQualityChange,
+			PrerollStart,
+			PrerollEnd,
+			PlaybackStart,
+			PlaybackPaused,
+			PlaybackResumed,
+			PlaybackEnded,
+			JumpInPlayPosition,
+			PlaybackStopped,
+			Error,
+			LogMessage,
+			DroppedVideoFrame,
+			DroppedAudioFrame
+		};
+		FPlayerMetricEventBase(EType InType) : Type(InType) {}
+		virtual ~FPlayerMetricEventBase() = default;
+		EType Type;
+	};
+	struct FPlayerMetricEvent_OpenSource : public FPlayerMetricEventBase
+	{
+		FPlayerMetricEvent_OpenSource(const FString& InURL) : FPlayerMetricEventBase(EType::OpenSource), URL(InURL) {}
+		FString URL;
+	};
+	struct FPlayerMetricEvent_ReceivedMasterPlaylist : public FPlayerMetricEventBase
+	{
+		FPlayerMetricEvent_ReceivedMasterPlaylist(const FString& InEffectiveURL) : FPlayerMetricEventBase(EType::ReceivedMasterPlaylist), EffectiveURL(InEffectiveURL) {}
+		FString EffectiveURL;
+	};
+	struct FPlayerMetricEvent_PlaylistDownload : public FPlayerMetricEventBase
+	{
+		FPlayerMetricEvent_PlaylistDownload(const Metrics::FPlaylistDownloadStats& InPlaylistDownloadStats) : FPlayerMetricEventBase(EType::PlaylistDownload), PlaylistDownloadStats(InPlaylistDownloadStats) {}
+		Metrics::FPlaylistDownloadStats PlaylistDownloadStats;
+	};
+	struct FPlayerMetricEvent_BufferingStart : public FPlayerMetricEventBase
+	{
+		FPlayerMetricEvent_BufferingStart(Metrics::EBufferingReason InBufferingReason) : FPlayerMetricEventBase(EType::BufferingStart), BufferingReason(InBufferingReason) {}
+		Metrics::EBufferingReason BufferingReason;
+	};
+	struct FPlayerMetricEvent_BufferingEnd : public FPlayerMetricEventBase
+	{
+		FPlayerMetricEvent_BufferingEnd(Metrics::EBufferingReason InBufferingReason) : FPlayerMetricEventBase(EType::BufferingEnd), BufferingReason(InBufferingReason) {}
+		Metrics::EBufferingReason BufferingReason;
+	};
+	struct FPlayerMetricEvent_Bandwidth : public FPlayerMetricEventBase
+	{
+		FPlayerMetricEvent_Bandwidth(int64 InEffectiveBps, int64 InThroughputBps, double InLatencyInSeconds) : FPlayerMetricEventBase(EType::Bandwidth), EffectiveBps(InEffectiveBps), ThroughputBps(InThroughputBps), LatencyInSeconds(InLatencyInSeconds) {}
+		int64 EffectiveBps;
+		int64 ThroughputBps;
+		double LatencyInSeconds;
+	};
+	struct FPlayerMetricEvent_BufferUtilization : public FPlayerMetricEventBase
+	{
+		FPlayerMetricEvent_BufferUtilization(const Metrics::FBufferStats& InBufferStats) : FPlayerMetricEventBase(EType::BufferUtilization), BufferStats(InBufferStats) {}
+		Metrics::FBufferStats BufferStats;
+	};
+	struct FPlayerMetricEvent_SegmentDownload : public FPlayerMetricEventBase
+	{
+		FPlayerMetricEvent_SegmentDownload(const Metrics::FSegmentDownloadStats& InSegmentDownloadStats) : FPlayerMetricEventBase(EType::SegmentDownload), SegmentDownloadStats(InSegmentDownloadStats) {}
+		Metrics::FSegmentDownloadStats SegmentDownloadStats;
+	};
+	struct FPlayerMetricEvent_LicenseKey : public FPlayerMetricEventBase
+	{
+		FPlayerMetricEvent_LicenseKey(const Metrics::FLicenseKeyStats& InLicenseKeyStats) : FPlayerMetricEventBase(EType::LicenseKey), LicenseKeyStats(InLicenseKeyStats) {}
+		Metrics::FLicenseKeyStats LicenseKeyStats;
+	};
+	struct FPlayerMetricEvent_DataAvailabilityChange : public FPlayerMetricEventBase
+	{
+		FPlayerMetricEvent_DataAvailabilityChange(const Metrics::FDataAvailabilityChange& InDataAvailability) : FPlayerMetricEventBase(EType::DataAvailabilityChange), DataAvailability(InDataAvailability) {}
+		Metrics::FDataAvailabilityChange DataAvailability;
+	};
+	struct FPlayerMetricEvent_VideoQualityChange : public FPlayerMetricEventBase
+	{
+		FPlayerMetricEvent_VideoQualityChange(int32 InNewBitrate, int32 InPreviousBitrate, bool bInIsDrasticDownswitch) : FPlayerMetricEventBase(EType::VideoQualityChange), NewBitrate(InNewBitrate), PreviousBitrate(InPreviousBitrate), bIsDrasticDownswitch(bInIsDrasticDownswitch) {}
+		int32 NewBitrate;
+		int32 PreviousBitrate;
+		bool bIsDrasticDownswitch;
+	};
+	struct FPlayerMetricEvent_JumpInPlayPosition : public FPlayerMetricEventBase
+	{
+		FPlayerMetricEvent_JumpInPlayPosition(const FTimeValue& InToNewTime, const FTimeValue& InFromTime, Metrics::ETimeJumpReason InTimejumpReason) : FPlayerMetricEventBase(EType::JumpInPlayPosition), ToNewTime(InToNewTime), FromTime(InFromTime), TimejumpReason(InTimejumpReason) {}
+		FTimeValue ToNewTime;
+		FTimeValue FromTime;
+		Metrics::ETimeJumpReason TimejumpReason;
+	};
+	struct FPlayerMetricEvent_Error : public FPlayerMetricEventBase
+	{
+		FPlayerMetricEvent_Error(const FString& InErrorReason) : FPlayerMetricEventBase(EType::Error), ErrorReason(InErrorReason) {}
+		FString ErrorReason;
+	};
+	struct FPlayerMetricEvent_LogMessage : public FPlayerMetricEventBase
+	{
+		FPlayerMetricEvent_LogMessage(IInfoLog::ELevel InLogLevel, const FString& InLogMessage, int64 InPlayerWallclockMilliseconds) : FPlayerMetricEventBase(EType::LogMessage), LogLevel(InLogLevel), LogMessage(InLogMessage), PlayerWallclockMilliseconds(InPlayerWallclockMilliseconds) {}
+		IInfoLog::ELevel LogLevel;
+		FString LogMessage;
+		int64 PlayerWallclockMilliseconds;
+	};
+
+
 	void CalculateTargetSeekTime(FTimespan& OutTargetTime, const FTimespan& InTime);
 
 	bool PresentVideoFrame(const FVideoDecoderOutputPtr& InVideoFrame);
 	bool PresentAudioFrame(const IAudioDecoderOutputPtr& DecoderOutput);
-	
+
 	void PlatformNotifyOfOptionChange();
 
 	// Methods from IAdaptiveStreamingPlayerMetrics
-	virtual void ReportOpenSource(const FString& URL) override;
-	virtual void ReportReceivedMasterPlaylist(const FString& EffectiveURL) override;
-	virtual void ReportReceivedPlaylists() override;
-	virtual void ReportTracksChanged() override;
-	virtual void ReportPlaylistDownload(const Metrics::FPlaylistDownloadStats& PlaylistDownloadStats) override;
-	virtual void ReportBufferingStart(Metrics::EBufferingReason BufferingReason) override;
-	virtual void ReportBufferingEnd(Metrics::EBufferingReason BufferingReason) override;
-	virtual void ReportBandwidth(int64 EffectiveBps, int64 ThroughputBps, double LatencyInSeconds) override;
-	virtual void ReportBufferUtilization(const Metrics::FBufferStats& BufferStats) override;
-	virtual void ReportSegmentDownload(const Metrics::FSegmentDownloadStats& SegmentDownloadStats) override;
-	virtual void ReportLicenseKey(const Metrics::FLicenseKeyStats& LicenseKeyStats) override;
-	virtual void ReportDataAvailabilityChange(const Metrics::FDataAvailabilityChange& DataAvailability) override;
-	virtual void ReportVideoQualityChange(int32 NewBitrate, int32 PreviousBitrate, bool bIsDrasticDownswitch) override;
-	virtual void ReportPrerollStart() override;
-	virtual void ReportPrerollEnd() override;
-	virtual void ReportPlaybackStart() override;
-	virtual void ReportPlaybackPaused() override;
-	virtual void ReportPlaybackResumed() override;
-	virtual void ReportPlaybackEnded() override;
-	virtual void ReportJumpInPlayPosition(const FTimeValue& ToNewTime, const FTimeValue& FromTime, Metrics::ETimeJumpReason TimejumpReason) override;
-	virtual void ReportPlaybackStopped() override;
-	virtual void ReportError(const FString& ErrorReason) override;
-	virtual void ReportLogMessage(IInfoLog::ELevel InLogLevel, const FString& InLogMessage, int64 InPlayerWallclockMilliseconds) override;
-	virtual void ReportDroppedVideoFrame() override;
-	virtual void ReportDroppedAudioFrame() override;
-	
+	virtual void ReportOpenSource(const FString& URL) override
+	{ DeferredPlayerEvents.Enqueue(MakeSharedTS<FPlayerMetricEvent_OpenSource>(URL)); }
+	virtual void ReportReceivedMasterPlaylist(const FString& EffectiveURL) override
+	{ DeferredPlayerEvents.Enqueue(MakeSharedTS<FPlayerMetricEvent_ReceivedMasterPlaylist>(EffectiveURL)); }
+	virtual void ReportReceivedPlaylists() override
+	{ DeferredPlayerEvents.Enqueue(MakeSharedTS<FPlayerMetricEventBase>(FPlayerMetricEventBase::EType::ReceivedPlaylists)); }
+	virtual void ReportTracksChanged() override
+	{ DeferredPlayerEvents.Enqueue(MakeSharedTS<FPlayerMetricEventBase>(FPlayerMetricEventBase::EType::TracksChanged)); }
+	virtual void ReportPlaylistDownload(const Metrics::FPlaylistDownloadStats& PlaylistDownloadStats) override
+	{ DeferredPlayerEvents.Enqueue(MakeSharedTS<FPlayerMetricEvent_PlaylistDownload>(PlaylistDownloadStats)); }
+	virtual void ReportBufferingStart(Metrics::EBufferingReason BufferingReason) override
+	{ DeferredPlayerEvents.Enqueue(MakeSharedTS<FPlayerMetricEvent_BufferingStart>(BufferingReason)); }
+	virtual void ReportBufferingEnd(Metrics::EBufferingReason BufferingReason) override
+	{ DeferredPlayerEvents.Enqueue(MakeSharedTS<FPlayerMetricEvent_BufferingEnd>(BufferingReason)); }
+	virtual void ReportBandwidth(int64 EffectiveBps, int64 ThroughputBps, double LatencyInSeconds) override
+	{ DeferredPlayerEvents.Enqueue(MakeSharedTS<FPlayerMetricEvent_Bandwidth>(EffectiveBps, ThroughputBps, LatencyInSeconds)); }
+	virtual void ReportBufferUtilization(const Metrics::FBufferStats& BufferStats) override
+	{ DeferredPlayerEvents.Enqueue(MakeSharedTS<FPlayerMetricEvent_BufferUtilization>(BufferStats)); }
+	virtual void ReportSegmentDownload(const Metrics::FSegmentDownloadStats& SegmentDownloadStats) override
+	{ DeferredPlayerEvents.Enqueue(MakeSharedTS<FPlayerMetricEvent_SegmentDownload>(SegmentDownloadStats)); }
+	virtual void ReportLicenseKey(const Metrics::FLicenseKeyStats& LicenseKeyStats) override
+	{ DeferredPlayerEvents.Enqueue(MakeSharedTS<FPlayerMetricEvent_LicenseKey>(LicenseKeyStats)); }
+	virtual void ReportDataAvailabilityChange(const Metrics::FDataAvailabilityChange& DataAvailability) override
+	{ DeferredPlayerEvents.Enqueue(MakeSharedTS<FPlayerMetricEvent_DataAvailabilityChange>(DataAvailability)); }
+	virtual void ReportVideoQualityChange(int32 NewBitrate, int32 PreviousBitrate, bool bIsDrasticDownswitch) override
+	{ DeferredPlayerEvents.Enqueue(MakeSharedTS<FPlayerMetricEvent_VideoQualityChange>(NewBitrate, PreviousBitrate, bIsDrasticDownswitch)); }
+	virtual void ReportPrerollStart() override
+	{ DeferredPlayerEvents.Enqueue(MakeSharedTS<FPlayerMetricEventBase>(FPlayerMetricEventBase::EType::PrerollStart)); }
+	virtual void ReportPrerollEnd() override
+	{ DeferredPlayerEvents.Enqueue(MakeSharedTS<FPlayerMetricEventBase>(FPlayerMetricEventBase::EType::PrerollEnd)); }
+	virtual void ReportPlaybackStart() override
+	{ DeferredPlayerEvents.Enqueue(MakeSharedTS<FPlayerMetricEventBase>(FPlayerMetricEventBase::EType::PlaybackStart)); }
+	virtual void ReportPlaybackPaused() override
+	{ DeferredPlayerEvents.Enqueue(MakeSharedTS<FPlayerMetricEventBase>(FPlayerMetricEventBase::EType::PlaybackPaused)); }
+	virtual void ReportPlaybackResumed() override
+	{ DeferredPlayerEvents.Enqueue(MakeSharedTS<FPlayerMetricEventBase>(FPlayerMetricEventBase::EType::PlaybackResumed)); }
+	virtual void ReportPlaybackEnded() override
+	{ DeferredPlayerEvents.Enqueue(MakeSharedTS<FPlayerMetricEventBase>(FPlayerMetricEventBase::EType::PlaybackEnded)); }
+	virtual void ReportJumpInPlayPosition(const FTimeValue& ToNewTime, const FTimeValue& FromTime, Metrics::ETimeJumpReason TimejumpReason) override
+	{ DeferredPlayerEvents.Enqueue(MakeSharedTS<FPlayerMetricEvent_JumpInPlayPosition>(ToNewTime, FromTime, TimejumpReason)); }
+	virtual void ReportPlaybackStopped() override
+	{ DeferredPlayerEvents.Enqueue(MakeSharedTS<FPlayerMetricEventBase>(FPlayerMetricEventBase::EType::PlaybackStopped)); }
+	virtual void ReportError(const FString& ErrorReason) override
+	{ DeferredPlayerEvents.Enqueue(MakeSharedTS<FPlayerMetricEvent_Error>(ErrorReason)); }
+	virtual void ReportLogMessage(IInfoLog::ELevel InLogLevel, const FString& InLogMessage, int64 InPlayerWallclockMilliseconds) override
+	{ DeferredPlayerEvents.Enqueue(MakeSharedTS<FPlayerMetricEvent_LogMessage>(InLogLevel, InLogMessage, InPlayerWallclockMilliseconds)); }
+	virtual void ReportDroppedVideoFrame() override
+	{ DeferredPlayerEvents.Enqueue(MakeSharedTS<FPlayerMetricEventBase>(FPlayerMetricEventBase::EType::DroppedVideoFrame)); }
+	virtual void ReportDroppedAudioFrame() override
+	{ DeferredPlayerEvents.Enqueue(MakeSharedTS<FPlayerMetricEventBase>(FPlayerMetricEventBase::EType::DroppedAudioFrame)); }
+
 	void LogPresentationFramesQueues(FTimespan DeltaTime);
 
+	void ClearToDefaultState();
+
 	void MediaStateOnPreparingFinished();
 	bool MediaStateOnPlay();
 	bool MediaStateOnPause();
 	void MediaStateOnEndReached();
 	void MediaStateOnSeekFinished();
-	void MediaStateOnError();
 
 	TSharedPtr<FTrackMetadata, ESPMode::ThreadSafe> GetTrackStreamMetadata(EPlayerTrackType TrackType, int32 TrackIndex) const;
 
@@ -191,12 +331,15 @@ private:
 	/**  */
 	TAtomic<EPlayerState>							State;
 	TAtomic<EPlayerStatus>							Status;
-	bool											bWasClosedOnError;
+
+	TAtomic<bool>									bCloseInProgress;
+	TAtomic<bool>									bHasPendingError;
 
 	bool											bAllowKillAfterCloseEvent;
 
 	/** Queued events */
 	TQueue<IElectraPlayerAdapterDelegate::EPlayerEvent>	DeferredEvents;
+	TQueue<TSharedPtrTS<FPlayerMetricEventBase>>	DeferredPlayerEvents;
 
 	/** The URL of the currently opened media. */
 	FString											MediaUrl;
@@ -215,7 +358,9 @@ private:
 		static void DoCloseAsync(TSharedPtr<FInternalPlayerImpl, ESPMode::ThreadSafe> && Player, TSharedPtr<IAsyncResourceReleaseNotifyContainer, ESPMode::ThreadSafe> AsyncDestructNotification);
 	};
 
+	FCriticalSection												PlayerLock;
 	TSharedPtr<FInternalPlayerImpl, ESPMode::ThreadSafe>			CurrentPlayer;
+	FEvent*															WaitForPlayerDestroyedEvent;
 
 	TSharedPtr<IAsyncResourceReleaseNotifyContainer, ESPMode::ThreadSafe> AsyncResourceReleaseNotification;
 
@@ -224,7 +369,7 @@ private:
 	public:
 		FAdaptiveStreamingPlayerResourceProvider(const TWeakPtr<IElectraPlayerAdapterDelegate, ESPMode::ThreadSafe> & AdapterDelegate);
 		virtual ~FAdaptiveStreamingPlayerResourceProvider() = default;
-		
+
 		virtual void ProvideStaticPlaybackDataForURL(TSharedPtr<IAdaptiveStreamingPlayerResourceRequest, ESPMode::ThreadSafe> InOutRequest) override;
 
 		void ProcessPendingStaticResourceRequests();
@@ -479,6 +624,35 @@ private:
 	FString									AnalyticsInstanceGuid;
 	/** Sequential analytics event number. Helps sorting events. **/
 	uint32									AnalyticsInstanceEventCount;
+
+
+	void HandleDeferredPlayerEvents();
+	void HandlePlayerEventOpenSource(const FString& URL);
+	void HandlePlayerEventReceivedMasterPlaylist(const FString& EffectiveURL);
+	void HandlePlayerEventReceivedPlaylists();
+	void HandlePlayerEventTracksChanged();
+	void HandlePlayerEventPlaylistDownload(const Metrics::FPlaylistDownloadStats& PlaylistDownloadStats);
+	void HandlePlayerEventBufferingStart(Metrics::EBufferingReason BufferingReason);
+	void HandlePlayerEventBufferingEnd(Metrics::EBufferingReason BufferingReason);
+	void HandlePlayerEventBandwidth(int64 EffectiveBps, int64 ThroughputBps, double LatencyInSeconds);
+	void HandlePlayerEventBufferUtilization(const Metrics::FBufferStats& BufferStats);
+	void HandlePlayerEventSegmentDownload(const Metrics::FSegmentDownloadStats& SegmentDownloadStats);
+	void HandlePlayerEventLicenseKey(const Metrics::FLicenseKeyStats& LicenseKeyStats);
+	void HandlePlayerEventDataAvailabilityChange(const Metrics::FDataAvailabilityChange& DataAvailability);
+	void HandlePlayerEventVideoQualityChange(int32 NewBitrate, int32 PreviousBitrate, bool bIsDrasticDownswitch);
+	void HandlePlayerEventPrerollStart();
+	void HandlePlayerEventPrerollEnd();
+	void HandlePlayerEventPlaybackStart();
+	void HandlePlayerEventPlaybackPaused();
+	void HandlePlayerEventPlaybackResumed();
+	void HandlePlayerEventPlaybackEnded();
+	void HandlePlayerEventJumpInPlayPosition(const FTimeValue& ToNewTime, const FTimeValue& FromTime, Metrics::ETimeJumpReason TimejumpReason);
+	void HandlePlayerEventPlaybackStopped();
+	void HandlePlayerEventError(const FString& ErrorReason);
+	void HandlePlayerEventLogMessage(IInfoLog::ELevel InLogLevel, const FString& InLogMessage, int64 InPlayerWallclockMilliseconds);
+	void HandlePlayerEventDroppedVideoFrame();
+	void HandlePlayerEventDroppedAudioFrame();
+
 };
 
 ENUM_CLASS_FLAGS(FElectraPlayer::EPlayerStatus);
