ElectraPlayer: added player option to pass stream mime type on the side in case it cannot be inferred from the URL itself. If none is passed a http head request is made to get the content type from the server. mp4 streams no longer need to be fast-startable but they still should be.

#rnx

diff --git a/Engine/Plugins/Media/ElectraPlayer/Source/ElectraPlayerPlugin/Private/ElectraPlayerPlugin.cpp b/Engine/Plugins/Media/ElectraPlayer/Source/ElectraPlayerPlugin/Private/ElectraPlayerPlugin.cpp
--- a/Engine/Plugins/Media/ElectraPlayer/Source/ElectraPlayerPlugin/Private/ElectraPlayerPlugin.cpp
+++ b/Engine/Plugins/Media/ElectraPlayer/Source/ElectraPlayerPlugin/Private/ElectraPlayerPlugin.cpp
@@ -322,24 +322,20 @@ bool FElectraPlayerPlugin::Open(const FString& Url, const IMediaOptions* Options
 
 	// Set up options to initialize the internal player with.
 	Electra::FParamDict PlayerOptions;
-	//	PlayerOptions.Set("initial_bitrate", Electra::FVariantValue((int64)2000000));
 
-		// Check for one-time initialization options that can't be changed during playback.
-		/*
-		int64 MinStreamFPS = Options->GetMediaOption(TEXT("MinStreamFPSForMediaStreaming"), (int64)-1);
-		if (MinStreamFPS > 0)
-		{
-			PlayerOptions.Set("min_stream_fps", Electra::FVariantValue(MinStreamFPS));
-			UE_LOG(LogElectraPlayerPlugin, Log, TEXT("[%p][%p] IMediaPlayer::Open: Limiting min stream fps to %d"), this, AdaptivePlayer.Get(), (int32)MinStreamFPS);
-		}
-
-		int64 MaxStreamFPS = Options->GetMediaOption(TEXT("MaxStreamFPSForMediaStreaming"), (int64)-1);
-		if (MaxStreamFPS > 0)
-		{
-			PlayerOptions.Set("max_stream_fps", Electra::FVariantValue(MaxStreamFPS));
-			UE_LOG(LogElectraPlayerPlugin, Log, TEXT("[%p][%p] IMediaPlayer::Open: Limiting max stream fps to %d"), this, AdaptivePlayer.Get(), (int32)MaxStreamFPS);
-		}
-		*/
+	// Check for one-time initialization options that can't be changed during playback.
+	int64 InitialStreamBitrate = Options->GetMediaOption(TEXT("ElectraInitialBitrate"), (int64)-1);
+	if (InitialStreamBitrate > 0)
+	{
+		PlayerOptions.Set("initial_bitrate", Electra::FVariantValue(InitialStreamBitrate));
+		UE_LOG(LogElectraPlayerPlugin, Log, TEXT("[%p] IMediaPlayer::Open: Using initial bitrate of %d bits/second"), this, (int32)InitialStreamBitrate);
+	}
+	FString MediaMimeType = Options->GetMediaOption(TEXT("mimetype"), FString());
+	if (MediaMimeType.Len())
+	{
+		PlayerOptions.Set("mime_type", Electra::FVariantValue(MediaMimeType));
+		UE_LOG(LogElectraPlayerPlugin, Log, TEXT("[%p] IMediaPlayer::Open: Setting media mime type to \"%s\""), this, *MediaMimeType);
+	}
 	int64 MaxVerticalHeight = Options->GetMediaOption(TEXT("MaxElectraVerticalResolution"), (int64)-1);
 	if (MaxVerticalHeight > 0)
 	{
diff --git a/Engine/Plugins/Media/ElectraPlayer/Source/ElectraPlayerRuntime/Private/Runtime/Core/MediaTypes.h b/Engine/Plugins/Media/ElectraPlayer/Source/ElectraPlayerRuntime/Private/Runtime/Core/MediaTypes.h
--- a/Engine/Plugins/Media/ElectraPlayer/Source/ElectraPlayerRuntime/Private/Runtime/Core/MediaTypes.h
+++ b/Engine/Plugins/Media/ElectraPlayer/Source/ElectraPlayerRuntime/Private/Runtime/Core/MediaTypes.h
@@ -24,6 +24,7 @@ typedef uint32  UEMediaError;
 #define UEMEDIA_ERROR_INSUFFICIENT_DATA	(6)
 #define UEMEDIA_ERROR_READ_ERROR		(7)
 #define UEMEDIA_ERROR_FORMAT_ERROR		(8)
+#define UEMEDIA_ERROR_NOT_SUPPORTED		(9)
 #define UEMEDIA_ERROR_TRY_AGAIN			(100)
 #define UEMEDIA_ERROR_OOM				(0x80000000U)
 #define UEMEDIA_ERROR_BAD_ARGUMENTS		(0x80000001U)
diff --git a/Engine/Plugins/Media/ElectraPlayer/Source/ElectraPlayerRuntime/Private/Runtime/Demuxer/ParserISO14496-12.cpp b/Engine/Plugins/Media/ElectraPlayer/Source/ElectraPlayerRuntime/Private/Runtime/Demuxer/ParserISO14496-12.cpp
--- a/Engine/Plugins/Media/ElectraPlayer/Source/ElectraPlayerRuntime/Private/Runtime/Demuxer/ParserISO14496-12.cpp
+++ b/Engine/Plugins/Media/ElectraPlayer/Source/ElectraPlayerRuntime/Private/Runtime/Demuxer/ParserISO14496-12.cpp
@@ -42,7 +42,6 @@ namespace Electra
 	static const uint32 kAllOnes32 = uint32(0xffffffff);
 	static const uint64 kAllOnes64 = uint64(0xffffffffffffffffL);
 
-
 	class FMP4BoxReader : private TMediaNoncopyable<FMP4BoxReader>
 	{
 	public:
@@ -51,7 +50,6 @@ namespace Electra
 			, BoxParseCallback(InBoxParseCallback)
 		{
 			check(DataReader);
-			check(BoxParseCallback);
 		}
 
 		~FMP4BoxReader()
@@ -157,6 +155,15 @@ namespace Electra
 			return IParserISO14496_12::IBoxCallback::EParseContinuation::Stop;
 		}
 
+		IParserISO14496_12::IBoxCallback::EParseContinuation NotifyEndOfBox(IParserISO14496_12::FBoxType InBoxType, int64 InBoxSize, int64 InStartOffset, int64 InDataOffset)
+		{
+			if (BoxParseCallback)
+			{
+				return BoxParseCallback->OnEndOfBox(InBoxType, InBoxSize, InStartOffset, InDataOffset);
+			}
+			return IParserISO14496_12::IBoxCallback::EParseContinuation::Stop;
+		}
+
 
 	private:
 		FMP4BoxReader() = delete;
@@ -167,6 +174,38 @@ namespace Electra
 	};
 
 
+class FDataBufferReader : public IParserISO14496_12::IReader
+{
+public:
+	FDataBufferReader(const TArray<uint8>& InDataBufferToReadFrom) : DataBufferRef(InDataBufferToReadFrom)
+	{}
+	virtual ~FDataBufferReader() = default;
+	virtual int64 ReadData(void* IntoBuffer, int64 NumBytesToRead) override
+	{
+		int64 NumAvail = DataBufferRef.Num() - CurrentOffset;
+		if (NumAvail >= NumBytesToRead)
+		{
+			if (IntoBuffer)
+			{
+				FMemory::Memcpy(IntoBuffer, DataBufferRef.GetData() + CurrentOffset, NumBytesToRead);
+			}
+			CurrentOffset += NumBytesToRead;
+			return NumBytesToRead;
+		}
+		return -1;
+	}
+	virtual bool HasReachedEOF() const override
+	{ return CurrentOffset >= DataBufferRef.Num(); }
+	virtual bool HasReadBeenAborted() const override
+	{ return false;	}
+	virtual int64 GetCurrentOffset() const override
+	{ return CurrentOffset;	}
+private:
+	const TArray<uint8>& DataBufferRef;
+	int64 CurrentOffset = 0;
+};
+
+
 #define RETURN_IF_ERROR(expr)											\
 	if ((Error = expr) != UEMEDIA_ERROR_OK)								\
 	{																	\
@@ -184,7 +223,7 @@ namespace Electra
 
 		virtual ~FMP4ParseInfo();
 
-		UEMediaError Parse(FMP4BoxReader* Reader, IPlayerSessionServices* PlayerSession);
+		UEMediaError Parse(FMP4BoxReader* Reader, IPlayerSessionServices* PlayerSession, const FMP4ParseInfo* InOptionalMoovParseInfo);
 
 		/**
 		 * Reads the type and size of the next box. If the box is an uuid box the 16 byte uuid is stored in the provided buffer.
@@ -311,6 +350,8 @@ namespace Electra
 			return BoxNestingLevel;
 		}
 
+		const FMP4ParseInfo* OptionalMoovParseInfo = nullptr;
+
 		FMP4Box* RootBox = nullptr;					//!< A root box that is not an actual file box but a container representing the file itself.
 		FMP4BoxReader* BoxReader = nullptr;			//!< Instance of the box reader we were given. This is not ours so we must not delete it!
 		int32 BoxNestingLevel = 0;					//!< Box tree depth
@@ -324,7 +365,6 @@ namespace Electra
 		FMP4Box* CurrentHandlerBox = nullptr;		//!< hdlr being parsed at the moment.
 		FMP4Box* CurrentMediaHandlerBox = nullptr;	//!< media handler being parsed at the moment ('vmhd', 'smhd', 'sthd', nmhd')
 
-
 		IPlayerSessionServices* PlayerSession = nullptr;
 		int32 NumTotalBoxesParsed = 0;
 	};
@@ -602,6 +642,19 @@ namespace Electra
 		static const IParserISO14496_12::FBoxType kSample_enca = MAKE_BOX_ATOM('e', 'n', 'c', 'a');
 		static const IParserISO14496_12::FBoxType kSample_encv = MAKE_BOX_ATOM('e', 'n', 'c', 'v');
 
+		/**
+		 * Encryption scheme types
+		 */
+		static const IParserISO14496_12::FBoxType kEncryptionScheme_cenc = MAKE_BOX_ATOM('c', 'e', 'n', 'c');
+
+		/**
+		 * Known UUIDs for PIFF boxes before they were added to ISO/IEC 14496-12
+		 */
+		static constexpr uint8 UUID_atom[16] = { 0x00,0x00,0x00,0x00,0x00,0x11,0x00,0x10,0x80,0x00,0x00,0xAA,0x00,0x38,0x9B,0x71 };
+		static constexpr uint8 UUID_pssh[16] = { 0xD0,0x8A,0x4F,0x18,0x10,0xF3,0x4A,0x82,0xB6,0xC8,0x32,0xD8,0xAB,0xA1,0x83,0xD3 };
+		static constexpr uint8 UUID_tenc[16] = { 0x89,0x74,0xDB,0xCE,0x7B,0xE7,0x4C,0x51,0x84,0xF9,0x71,0x48,0xF9,0x88,0x25,0x54 };
+		static constexpr uint8 UUID_senc[16] = { 0xA2,0x39,0x4F,0x52,0x5A,0x9B,0x4F,0x14,0xA2,0x44,0x6C,0x42,0x7C,0x64,0x8D,0xF4 };
+
 		//	#undef MAKE_BOX_ATOM
 
 	private:
@@ -731,6 +784,36 @@ namespace Electra
 	};
 
 
+	/**
+	 * 'mdat' box.
+	 * This box has no parseable data. If we are asked to parse it then this is probably in locating a 'moov' box of
+	 * a file that is not fast-startable and we will actually skip over it.
+	 */
+	class FMP4BoxMDAT : public FMP4BoxBasic
+	{
+	public:
+		FMP4BoxMDAT(IParserISO14496_12::FBoxType InBoxType, int64 InBoxSize, int64 InStartOffset, int64 InDataOffset, bool bInIsLeafBox)
+			: FMP4BoxBasic(InBoxType, InBoxSize, InStartOffset, InDataOffset, bInIsLeafBox)
+		{
+		}
+
+		virtual ~FMP4BoxMDAT()
+		{
+		}
+
+	private:
+		FMP4BoxMDAT() = delete;
+		FMP4BoxMDAT(const FMP4BoxMDAT&) = delete;
+
+	protected:
+		virtual UEMediaError ReadAndParseAttributes(FMP4ParseInfo* ParseInfo) override
+		{
+			// Read the data from this box to NULL.
+			return ParseInfo->Reader()->ReadBytes(nullptr, BoxSize - (ParseInfo->Reader()->GetCurrentReadOffset() - StartOffset));
+		}
+	};
+
+
 	/**
 	 * 'ftyp' box.
 	 * ISO/IEC 14496-12:2014 - 4.3 - File Type Box
@@ -3376,6 +3459,11 @@ namespace Electra
 		{
 		}
 
+		uint32 GetDataFormat() const
+		{
+			return DataFormat;
+		}
+
 	private:
 		FMP4BoxFRMA() = delete;
 		FMP4BoxFRMA(const FMP4BoxFRMA&) = delete;
@@ -3409,6 +3497,11 @@ namespace Electra
 		{
 		}
 
+		uint32 GetSchemeType() const
+		{
+			return SchemeType;
+		}
+
 	private:
 		FMP4BoxSCHM() = delete;
 		FMP4BoxSCHM(const FMP4BoxSCHM&) = delete;
@@ -3455,6 +3548,35 @@ namespace Electra
 		{
 		}
 
+		bool HasDefaultCryptBlockValues() const
+		{
+			return Version != 0;
+		}
+		uint8 GetDefaultCryptByteBlock() const
+		{
+			return DefaultCryptByteBlock;
+		}
+		uint8 GetDefaultSkipByteBlock() const
+		{
+			return DefaultSkipByteBlock;
+		}
+		uint8 GetDefaultIsProtected() const
+		{
+			return DefaultIsProtected;
+		}
+		uint8 GetDefaultPerSampleIVSize() const
+		{
+			return DefaultPerSampleIVSize;
+		}
+		const TArray<uint8>& GetDefaultKID() const
+		{
+			return DefaultKID;
+		}
+		const TArray<uint8>& GetDefaultConstantIV() const
+		{
+			return DefaultConstantIV;
+		}
+
 	private:
 		FMP4BoxTENC() = delete;
 		FMP4BoxTENC(const FMP4BoxTENC&) = delete;
@@ -3613,6 +3735,114 @@ namespace Electra
 	};
 
 
+	/**
+	 * 'senc' box. ISO/IEC 23001-7:2016 - 7.2.1 Sample Encryption Box
+	 */
+	class FMP4BoxSENC : public FMP4BoxFull
+	{
+	public:
+		FMP4BoxSENC(IParserISO14496_12::FBoxType InBoxType, int64 InBoxSize, int64 InStartOffset, int64 InDataOffset, bool bInIsLeafBox)
+			: FMP4BoxFull(InBoxType, InBoxSize, InStartOffset, InDataOffset, bInIsLeafBox)
+		{
+		}
+
+		virtual ~FMP4BoxSENC()
+		{
+		}
+
+		struct FEntry
+		{
+			TArray<uint8> IV;
+			struct FSubSample
+			{
+				uint16	BytesOfClearData;
+				uint32	BytesOfProtectedData;
+			};
+			TArray<FSubSample> SubSamples;
+		};
+
+		const TArray<FEntry>& GetEntries() const
+		{
+			return Entries;
+		}
+
+		UEMediaError Parse(const FMP4BoxTENC* TENCBox, const FMP4BoxSAIZ* SAIZBox, const FMP4BoxSAIO* SAIOBox)
+		{
+			UEMediaError Error = UEMEDIA_ERROR_OK;
+			// We need to have the tenc box here. Right now the saiz and saio are not relevant.
+			if (!TENCBox)
+			{
+				return UEMEDIA_ERROR_FORMAT_ERROR;
+			}
+			
+			int32 PerSampleIVSize = TENCBox->GetDefaultPerSampleIVSize();
+
+			FDataBufferReader BufferReader(BoxData);
+			FMP4BoxReader Reader(&BufferReader, nullptr);
+
+			if (SampleCount)
+			{
+				Entries.Reserve(SampleCount);
+				for(uint32 i=0; i<SampleCount; ++i)
+				{
+					FEntry& Entry = Entries.AddDefaulted_GetRef();
+					if (PerSampleIVSize)
+					{
+						Entry.IV.AddUninitialized(PerSampleIVSize);
+						RETURN_IF_ERROR(Reader.ReadBytes(Entry.IV.GetData(), PerSampleIVSize));
+					}
+					if ((Flags & 2) != 0)
+					{
+						uint16 SubSampleCount = 0;
+						RETURN_IF_ERROR(Reader.Read(SubSampleCount));
+						for(int32 j=0; j<SubSampleCount; ++j)
+						{
+							FEntry::FSubSample& SubSample = Entry.SubSamples.AddDefaulted_GetRef();
+							RETURN_IF_ERROR(Reader.Read(SubSample.BytesOfClearData));
+							RETURN_IF_ERROR(Reader.Read(SubSample.BytesOfProtectedData));
+						}
+					}
+				}
+			}
+			return UEMEDIA_ERROR_OK;
+		}
+
+	private:
+		FMP4BoxSENC() = delete;
+		FMP4BoxSENC(const FMP4BoxSENC&) = delete;
+
+	protected:
+		virtual UEMediaError ReadAndParseAttributes(FMP4ParseInfo* ParseInfo) override
+		{
+			UEMediaError Error = UEMEDIA_ERROR_OK;
+			RETURN_IF_ERROR(FMP4BoxFull::ReadAndParseAttributes(ParseInfo));
+
+			// Check for a pre-PIFF 1.3 box that has elements of the 'tenc' box here when flags&1.
+			if ((Flags & 1) != 0)
+			{
+				return UEMEDIA_ERROR_FORMAT_ERROR;
+			}
+			
+			RETURN_IF_ERROR(ParseInfo->Reader()->Read(SampleCount));
+			// The contents of this box cannot be parsed without information from another box
+			// containing the size of the IV.
+			// We read the data of this box and parse it later when we need its data.
+			uint32 BytesRemaining = BoxSize - (ParseInfo->Reader()->GetCurrentReadOffset() - StartOffset);
+			if (BytesRemaining)
+			{
+				BoxData.AddUninitialized(BytesRemaining);
+				RETURN_IF_ERROR(ParseInfo->Reader()->ReadBytes(BoxData.GetData(), BytesRemaining));
+			}
+			return Error;
+		}
+
+	private:
+		uint32 SampleCount = 0;
+		TArray<uint8> BoxData;
+		TArray<FEntry> Entries;
+	};
+
+
 
 
 	FMP4ParseInfo::~FMP4ParseInfo()
@@ -3658,9 +3888,10 @@ namespace Electra
 	}
 
 
-	UEMediaError FMP4ParseInfo::Parse(FMP4BoxReader* Reader, IPlayerSessionServices* InPlayerSession)
+	UEMediaError FMP4ParseInfo::Parse(FMP4BoxReader* Reader, IPlayerSessionServices* InPlayerSession, const FMP4ParseInfo* InOptionalMoovParseInfo)
 	{
 		PlayerSession = InPlayerSession;
+		OptionalMoovParseInfo = InOptionalMoovParseInfo;
 
 		// New parse or continuing a previous?
 		if (RootBox == nullptr)
@@ -3759,6 +3990,10 @@ namespace Electra
 					FMP4Box* NextBox = nullptr;
 					switch(BoxType)
 					{
+						case FMP4Box::kBox_mdat:
+							NextBox = new FMP4BoxMDAT(BoxType, BoxSize, BoxStartOffset, BoxDataOffset, true);
+							break;
+
 						case FMP4Box::kBox_moov:
 							check(GetCurrentMoovBox() == nullptr);
 							NextBox = new FMP4BoxBasic(BoxType, BoxSize, BoxStartOffset, BoxDataOffset, false);
@@ -3921,9 +4156,29 @@ namespace Electra
 							break;
 
 						case FMP4Box::kBox_uuid:
-							// Here we could check if we know the UUID of the box and then handle it specially.
-								// if (memcmp(BoxUUID, ...., 16) == 0)
-							NextBox = new FMP4BoxIgnored(BoxType, BoxSize, BoxStartOffset, BoxDataOffset, true);
+							// Check if this is a long form of an atom. This is rare and not supported.
+							if (FMemory::Memcmp(BoxUUID+4, FMP4Box::UUID_atom+4, 12) == 0)
+							{
+								// This is not handled.
+								return UEMEDIA_ERROR_FORMAT_ERROR;
+							}
+							// Compare the UUID with known boxes.
+							else if (FMemory::Memcmp(BoxUUID, FMP4Box::UUID_pssh, 16) == 0)
+							{
+								NextBox = new FMP4BoxPSSH(FMP4Box::kBox_pssh, BoxSize, BoxStartOffset, BoxDataOffset, true);
+							}
+							else if (FMemory::Memcmp(BoxUUID, FMP4Box::UUID_tenc, 16) == 0)
+							{
+								NextBox = new FMP4BoxTENC(FMP4Box::kBox_tenc, BoxSize, BoxStartOffset, BoxDataOffset, true);
+							}
+							else if (FMemory::Memcmp(BoxUUID, FMP4Box::UUID_senc, 16) == 0)
+							{
+								NextBox = new FMP4BoxSENC(FMP4Box::kBox_senc, BoxSize, BoxStartOffset, BoxDataOffset, true);
+							}
+							else
+							{
+								NextBox = new FMP4BoxIgnored(BoxType, BoxSize, BoxStartOffset, BoxDataOffset, true);
+							}
 							break;
 
 						// Encryption related boxes.
@@ -3942,6 +4197,9 @@ namespace Electra
 						case FMP4Box::kBox_tenc:
 							NextBox = new FMP4BoxTENC(BoxType, BoxSize, BoxStartOffset, BoxDataOffset, true);
 							break;
+						case FMP4Box::kBox_senc:
+							NextBox = new FMP4BoxSENC(BoxType, BoxSize, BoxStartOffset, BoxDataOffset, true);
+							break;
 
 						default:
 #if MEDIA_DEBUG_HAS_BOX_NAMES
@@ -4030,6 +4288,12 @@ namespace Electra
 							return Error;
 						}
 					}
+					ParseContinuation = GetNestingLevel() == 0 ? BoxReader->NotifyEndOfBox(BoxType, BoxSize, BoxStartOffset, BoxDataOffset) : IParserISO14496_12::IBoxCallback::EParseContinuation::Continue;
+					if (ParseContinuation == IParserISO14496_12::IBoxCallback::EParseContinuation::Stop)
+					{
+						// Stop parsing.
+						return UEMEDIA_ERROR_OK;
+					}
 				}
 				else
 				{
@@ -4059,7 +4323,7 @@ namespace Electra
 	public:
 		FParserISO14496_12();
 		virtual ~FParserISO14496_12();
-		virtual UEMediaError ParseHeader(IReader* DataReader, IBoxCallback* BoxParseCallback, IPlayerSessionServices* PlayerSession) override;
+		virtual UEMediaError ParseHeader(IReader* DataReader, IBoxCallback* BoxParseCallback, IPlayerSessionServices* PlayerSession, const IParserISO14496_12* OptionalInitSegment) override;
 
 		virtual UEMediaError PrepareTracks(TSharedPtrTS<const IParserISO14496_12>	OptionalMP4InitSegment) override;
 
@@ -4224,13 +4488,20 @@ namespace Electra
 			const FMP4BoxSTSZ* STSZBox = nullptr;
 			const FMP4BoxSTCO* STCOBox = nullptr;
 			const FMP4BoxSTSS* STSSBox = nullptr;
+			// Encryption
+			TArray<const FMP4BoxPSSH*>	PSSHBoxes;
+			const FMP4BoxSCHM* SCHMBox = nullptr;
+			const FMP4BoxTENC* TENCBox = nullptr;
+			const FMP4BoxSENC* SENCBox = nullptr;
+			const FMP4BoxSAIZ* SAIZBox = nullptr;
+			const FMP4BoxSAIO* SAIOBox = nullptr;
 			// Fragmented track
 			const FMP4Box*	   MOOFBox = nullptr;
 			const FMP4BoxTREX* TREXBox = nullptr;
 			const FMP4BoxTFHD* TFHDBox = nullptr;
 			const FMP4BoxTFDT* TFDTBox = nullptr;
 			TArray<const FMP4BoxTRUN*>	TRUNBoxes;
-			// Future boxes: sbgp, sgpd, subs, saiz, saio
+			// Future boxes: sbgp, sgpd, subs
 
 			FStreamCodecInformation					CodecInformation;
 
@@ -4308,6 +4579,9 @@ namespace Electra
 		};
 
 
+		UEMediaError ParseAVC1SampleType(FTrack* Track, const FMP4Box* SampleBox);
+		UEMediaError ParseMP4ASampleType(FTrack* Track, const FMP4Box* SampleBox);
+
 		FMP4ParseInfo* ParsedData;
 
 		FTrackInfo* ParsedTrackInfo;
@@ -5093,7 +5367,7 @@ namespace Electra
 		delete ParsedData;
 	}
 
-	UEMediaError FParserISO14496_12::ParseHeader(IReader* InDataReader, IBoxCallback* InBoxParseCallback, IPlayerSessionServices* PlayerSession)
+	UEMediaError FParserISO14496_12::ParseHeader(IReader* InDataReader, IBoxCallback* InBoxParseCallback, IPlayerSessionServices* PlayerSession, const IParserISO14496_12* InOptionalInitSegment)
 	{
 		if (!InDataReader || !InBoxParseCallback)
 		{
@@ -5110,13 +5384,177 @@ namespace Electra
 			}
 		}
 
+		const FParserISO14496_12* OptionalInitSegment = static_cast<const FParserISO14496_12*>(InOptionalInitSegment);
 		FMP4BoxReader Reader(InDataReader, InBoxParseCallback);
-		UEMediaError ParseError = ParsedData->Parse(&Reader, PlayerSession);
-
+		UEMediaError ParseError = ParsedData->Parse(&Reader, PlayerSession, OptionalInitSegment ? OptionalInitSegment->ParsedData : nullptr);
 		return ParseError;
 	}
 
 
+	UEMediaError FParserISO14496_12::ParseAVC1SampleType(FTrack* Track, const FMP4Box* SampleBox)
+	{
+		const FMP4BoxVisualSampleEntry* VisualSampleEntry = static_cast<const FMP4BoxVisualSampleEntry*>(SampleBox);
+		check(VisualSampleEntry->GetNumberOfChildren() > 0);
+		if (VisualSampleEntry->GetNumberOfChildren() > 0)
+		{
+			// There may be several entries. Usually there is an 'avcC' (required) and optional boxes like 'pasp', 'btrt', 'clap', etc.
+			bool bGotVideoFormat = false;
+			for(int32 j = 0, jMax = VisualSampleEntry->GetNumberOfChildren(); j < jMax; ++j)
+			{
+				const FMP4Box* SampleEntry = VisualSampleEntry->GetChildBox(j);
+				switch(SampleEntry->GetType())
+				{
+					case FMP4Box::kBox_avcC:
+					{
+						const FMP4BoxAVCC* AVCCBox = static_cast<const FMP4BoxAVCC*>(SampleEntry);
+						Track->CodecSpecificDataAVC = AVCCBox->GetDecoderConfigurationRecord();
+						bool bOk = Track->CodecSpecificDataAVC.Parse();
+						check(bOk);
+						if (bOk)
+						{
+							Track->CodecInformation.SetStreamType(EStreamType::Video);
+							Track->CodecInformation.SetCodec(FStreamCodecInformation::ECodec::H264);
+							Track->CodecInformation.SetStreamLanguageCode(Track->GetLanguage());
+							if (Track->CodecSpecificDataAVC.GetNumberOfSPS() == 0)
+							{
+								return UEMEDIA_ERROR_FORMAT_ERROR;
+							}
+							const MPEG::FISO14496_10_seq_parameter_set_data& sps = Track->CodecSpecificDataAVC.GetParsedSPS(0);
+							int32 CropL, CropR, CropT, CropB;
+							sps.GetCrop(CropL, CropR, CropT, CropB);
+							Track->CodecInformation.SetResolution(FStreamCodecInformation::FResolution(sps.GetWidth() - CropL - CropR, sps.GetHeight() - CropT - CropB));
+							Track->CodecInformation.SetCrop(FStreamCodecInformation::FCrop(CropL, CropT, CropR, CropB));
+							FStreamCodecInformation::FAspectRatio ar;
+							sps.GetAspect(ar.Width, ar.Height);
+							Track->CodecInformation.SetAspectRatio(ar);
+							Track->CodecInformation.SetFrameRate(sps.GetTiming());
+							Track->CodecInformation.SetProfile(sps.profile_idc);
+							Track->CodecInformation.SetProfileLevel(sps.level_idc);
+							uint8 Constraints = (sps.constraint_set0_flag << 7) | (sps.constraint_set1_flag << 6) | (sps.constraint_set2_flag << 5) | (sps.constraint_set3_flag << 4) | (sps.constraint_set4_flag << 3) | (sps.constraint_set5_flag << 2);
+							Track->CodecInformation.SetProfileConstraints(Constraints);
+							Track->CodecInformation.SetCodecSpecifierRFC6381(FString::Printf(TEXT("avc1.%02x%02x%02x"), sps.profile_idc, Constraints, sps.level_idc));
+							bGotVideoFormat = true;
+						}
+						else
+						{
+							return UEMEDIA_ERROR_FORMAT_ERROR;
+						}
+						break;
+					}
+					case FMP4Box::kBox_btrt:
+					{
+						const FMP4BoxBTRT* BTRTBox = static_cast<const FMP4BoxBTRT*>(SampleEntry);
+						Track->BitrateInfo.BufferSizeDB = BTRTBox->GetBufferSizeDB();
+						Track->BitrateInfo.MaxBitrate = BTRTBox->GetMaxBitrate();
+						Track->BitrateInfo.AvgBitrate = BTRTBox->GetAverageBitrate();
+						break;
+					}
+					case FMP4Box::kBox_pasp:
+					{
+						break;
+					}
+					default:
+					{
+						break;
+					}
+				}
+			}
+			// Check that we got the required information
+			if (bGotVideoFormat)
+			{
+				return UEMEDIA_ERROR_OK;
+			}
+		}
+		return UEMEDIA_ERROR_FORMAT_ERROR;
+	}
+
+	UEMediaError FParserISO14496_12::ParseMP4ASampleType(FTrack* Track, const FMP4Box* SampleBox)
+	{
+		const FMP4BoxAudioSampleEntry* AudioSampleEntry = static_cast<const FMP4BoxAudioSampleEntry*>(Track->STSDBox->GetChildBox(0));
+		check(AudioSampleEntry->GetNumberOfChildren() > 0);
+		if (AudioSampleEntry->GetNumberOfChildren() > 0)
+		{
+			// There may be several entries. Usually there is the required sample format box (eg 'mp4a') and optional boxes like 'btrt'
+			bool bGotAudioFormat = false;
+			bool bIsSupported = true;
+			for(int32 j = 0, jMax = AudioSampleEntry->GetNumberOfChildren(); j < jMax; ++j)
+			{
+				const FMP4Box* SampleEntry = AudioSampleEntry->GetChildBox(j);
+				switch(SampleEntry->GetType())
+				{
+					case FMP4Box::kBox_esds:
+					{
+						const FMP4BoxESDS* ESDSBox = static_cast<const FMP4BoxESDS*>(AudioSampleEntry->GetChildBox(0));
+						Track->CodecSpecificDataMP4A = ESDSBox->GetESDescriptor();
+						bool bOk = Track->CodecSpecificDataMP4A.Parse();
+						if (bOk)
+						{
+							// Is this AAC audio?
+							if (Track->CodecSpecificDataMP4A.GetObjectTypeID() == MPEG::FESDescriptor::FObjectTypeID::MPEG4_Audio &&
+								Track->CodecSpecificDataMP4A.GetStreamType() == MPEG::FESDescriptor::FStreamType::Audio)
+							{
+								Track->CodecInformation.SetStreamType(EStreamType::Audio);
+								Track->CodecInformation.SetCodec(FStreamCodecInformation::ECodec::AAC);
+								Track->CodecInformation.SetStreamLanguageCode(Track->GetLanguage());
+
+								MPEG::FAACDecoderConfigurationRecord ConfigRecord;
+								if (ConfigRecord.ParseFrom(Track->CodecSpecificDataMP4A.GetCodecSpecificData().GetData(), Track->CodecSpecificDataMP4A.GetCodecSpecificData().Num()))
+								{
+									Track->CodecInformation.SetCodecSpecifierRFC6381(FString::Printf(TEXT("mp4a.40.%d"), ConfigRecord.ExtAOT ? ConfigRecord.ExtAOT : ConfigRecord.AOT));
+									Track->CodecInformation.SetSamplingRate(ConfigRecord.ExtSamplingFrequency ? ConfigRecord.ExtSamplingFrequency : ConfigRecord.SamplingRate);
+									Track->CodecInformation.SetNumberOfChannels((int32)ConfigRecord.ChannelConfiguration);
+									// We assume that all platforms can decode PS (parametric stereo). As such we change the channel count from mono to stereo
+									// to convey the _decoded_ format, not the source format.
+									if (ConfigRecord.ChannelConfiguration == 1 && ConfigRecord.PSSignal > 0)
+									{
+										Track->CodecInformation.SetNumberOfChannels(2);
+									}
+									Track->CodecInformation.GetExtras().Set(TEXT("samples_per_block"), FVariantValue(ConfigRecord.SBRSignal > 0 ? (int64)2048 : (int64)1024));
+								}
+
+								// Typically an mp4a track will not have a 'btrt' box because the bitrate is stored in the DecoderConfigDescriptor.
+								Track->BitrateInfo.BufferSizeDB = Track->CodecSpecificDataMP4A.GetBufferSize();
+								Track->BitrateInfo.MaxBitrate = Track->CodecSpecificDataMP4A.GetMaxBitrate();
+								Track->BitrateInfo.AvgBitrate = Track->CodecSpecificDataMP4A.GetAvgBitrate();
+
+								bGotAudioFormat = true;
+							}
+							else
+							{
+								// Not expected format. Ignore this track.
+								bIsSupported = false;
+								// Got the info, just not the format we want.
+								bGotAudioFormat = true;
+							}
+						}
+						else
+						{
+							return UEMEDIA_ERROR_FORMAT_ERROR;
+						}
+						break;
+					}
+					case FMP4Box::kBox_btrt:
+					{
+						const FMP4BoxBTRT* BTRTBox = static_cast<const FMP4BoxBTRT*>(SampleEntry);
+						Track->BitrateInfo.BufferSizeDB = BTRTBox->GetBufferSizeDB();
+						Track->BitrateInfo.MaxBitrate = BTRTBox->GetMaxBitrate();
+						Track->BitrateInfo.AvgBitrate = BTRTBox->GetAverageBitrate();
+						break;
+					}
+					default:
+					{
+						break;
+					}
+				}
+			}
+			// Check that we got the required information
+			if (bGotAudioFormat)
+			{
+				return bIsSupported ? UEMEDIA_ERROR_OK : UEMEDIA_ERROR_NOT_SUPPORTED;
+			}
+		}
+		return UEMEDIA_ERROR_FORMAT_ERROR;
+	}
 
 	UEMediaError FParserISO14496_12::PrepareTracks(TSharedPtrTS<const IParserISO14496_12> OptionalMP4InitSegment)
 	{
@@ -5314,190 +5752,99 @@ namespace Electra
 						}
 
 						bool bIsSupported = true;
-						switch(Track->STSDBox->GetChildBox(0)->GetType())
+						const FMP4Box* STSDFirstChildBox = Track->STSDBox->GetChildBox(0);
+						switch(STSDFirstChildBox->GetType())
 						{
-							case FMP4Box::kSample_avc1:
+							case FMP4Box::kSample_encv:
 							{
-								const FMP4BoxVisualSampleEntry* VisualSampleEntry = static_cast<const FMP4BoxVisualSampleEntry*>(Track->STSDBox->GetChildBox(0));
-								check(VisualSampleEntry->GetNumberOfChildren() > 0);
-								if (VisualSampleEntry->GetNumberOfChildren() > 0)
+								const FMP4BoxFRMA* FRMABox = static_cast<const FMP4BoxFRMA*>(STSDFirstChildBox->GetBoxPath(FMP4Box::kBox_sinf, FMP4Box::kBox_frma));
+								const FMP4BoxSCHM* SCHMBox = static_cast<const FMP4BoxSCHM*>(STSDFirstChildBox->GetBoxPath(FMP4Box::kBox_sinf, FMP4Box::kBox_schm));
+								if (!FRMABox || !SCHMBox)
 								{
-									// There may be several entries. Usually there is an 'avcC' (required) and optional boxes like 'pasp', 'btrt', 'clap', etc.
-									bool bGotVideoFormat = false;
-									for(int32 j = 0, jMax = VisualSampleEntry->GetNumberOfChildren(); j < jMax; ++j)
-									{
-										const FMP4Box* SampleEntry = VisualSampleEntry->GetChildBox(j);
-										switch(SampleEntry->GetType())
-										{
-											case FMP4Box::kBox_avcC:
-											{
-												const FMP4BoxAVCC* AVCCBox = static_cast<const FMP4BoxAVCC*>(SampleEntry);
-												Track->CodecSpecificDataAVC = AVCCBox->GetDecoderConfigurationRecord();
-												bool bOk = Track->CodecSpecificDataAVC.Parse();
-												check(bOk);
-												if (bOk)
-												{
-													Track->CodecInformation.SetStreamType(EStreamType::Video);
-													Track->CodecInformation.SetCodec(FStreamCodecInformation::ECodec::H264);
-													Track->CodecInformation.SetStreamLanguageCode(Track->GetLanguage());
-													if (Track->CodecSpecificDataAVC.GetNumberOfSPS() == 0)
-													{
-														return UEMEDIA_ERROR_FORMAT_ERROR;
-													}
-													const MPEG::FISO14496_10_seq_parameter_set_data& sps = Track->CodecSpecificDataAVC.GetParsedSPS(0);
-													int32 CropL, CropR, CropT, CropB;
-													sps.GetCrop(CropL, CropR, CropT, CropB);
-													Track->CodecInformation.SetResolution(FStreamCodecInformation::FResolution(sps.GetWidth() - CropL - CropR, sps.GetHeight() - CropT - CropB));
-													Track->CodecInformation.SetCrop(FStreamCodecInformation::FCrop(CropL, CropT, CropR, CropB));
-													FStreamCodecInformation::FAspectRatio ar;
-													sps.GetAspect(ar.Width, ar.Height);
-													Track->CodecInformation.SetAspectRatio(ar);
-													Track->CodecInformation.SetFrameRate(sps.GetTiming());
-													Track->CodecInformation.SetProfile(sps.profile_idc);
-													Track->CodecInformation.SetProfileLevel(sps.level_idc);
-													uint8 Constraints = (sps.constraint_set0_flag << 7) | (sps.constraint_set1_flag << 6) | (sps.constraint_set2_flag << 5) | (sps.constraint_set3_flag << 4) | (sps.constraint_set4_flag << 3) | (sps.constraint_set5_flag << 2);
-													Track->CodecInformation.SetProfileConstraints(Constraints);
-													Track->CodecInformation.SetCodecSpecifierRFC6381(FString::Printf(TEXT("avc1.%02x%02x%02x"), sps.profile_idc, Constraints, sps.level_idc));
-													bGotVideoFormat = true;
-												}
-												else
-												{
-													return UEMEDIA_ERROR_FORMAT_ERROR;
-												}
-												break;
-											}
-											case FMP4Box::kBox_btrt:
-											{
-												const FMP4BoxBTRT* BTRTBox = static_cast<const FMP4BoxBTRT*>(SampleEntry);
-												Track->BitrateInfo.BufferSizeDB = BTRTBox->GetBufferSizeDB();
-												Track->BitrateInfo.MaxBitrate = BTRTBox->GetMaxBitrate();
-												Track->BitrateInfo.AvgBitrate = BTRTBox->GetAverageBitrate();
-												break;
-											}
-											case FMP4Box::kBox_pasp:
-											{
-												break;
-											}
-											default:
-											{
-												break;
-											}
-										}
-									}
-									// Check that we got the required information
-									if (!bGotVideoFormat)
-									{
-										return UEMEDIA_ERROR_FORMAT_ERROR;
-									}
-
+									return UEMEDIA_ERROR_FORMAT_ERROR;
 								}
-								else
+								// Expecting avc1 for video right now.
+								if (FRMABox->GetDataFormat() != FMP4Box::kSample_avc1)
 								{
 									return UEMEDIA_ERROR_FORMAT_ERROR;
 								}
-
+								// Likewise fo encryption it needs to be 'cenc'
+								if (SCHMBox->GetSchemeType() != FMP4Box::kEncryptionScheme_cenc)
+								{
+									return UEMEDIA_ERROR_FORMAT_ERROR;
+								}
+								// Parse the sample
+								UEMediaError Error = ParseAVC1SampleType(Track.Get(), STSDFirstChildBox);
+								if (Error == UEMEDIA_ERROR_NOT_SUPPORTED)
+								{
+									bIsSupported = false;
+								}
+								else if (Error != UEMEDIA_ERROR_OK)
+								{
+									return Error;
+								}
+								// Remember the encryption boxes.
+								Track->SCHMBox = SCHMBox;
+								Track->TENCBox = static_cast<const FMP4BoxTENC*>(STSDFirstChildBox->GetBoxPath(FMP4Box::kBox_sinf, FMP4Box::kBox_schi, FMP4Box::kBox_tenc));
 								break;
 							}
-							case FMP4Box::kSample_encv:
+							case FMP4Box::kSample_enca:
 							{
-								check(!TEXT("TODO")); // get CSD
+								const FMP4BoxFRMA* FRMABox = static_cast<const FMP4BoxFRMA*>(STSDFirstChildBox->GetBoxPath(FMP4Box::kBox_sinf, FMP4Box::kBox_frma));
+								const FMP4BoxSCHM* SCHMBox = static_cast<const FMP4BoxSCHM*>(STSDFirstChildBox->GetBoxPath(FMP4Box::kBox_sinf, FMP4Box::kBox_schm));
+								if (!FRMABox || !SCHMBox)
+								{
+									return UEMEDIA_ERROR_FORMAT_ERROR;
+								}
+								// Expecting mp4a for audio right now.
+								if (FRMABox->GetDataFormat() != FMP4Box::kSample_mp4a)
+								{
+									return UEMEDIA_ERROR_FORMAT_ERROR;
+								}
+								// Likewise fo encryption it needs to be 'cenc'
+								if (SCHMBox->GetSchemeType() != FMP4Box::kEncryptionScheme_cenc)
+								{
+									return UEMEDIA_ERROR_FORMAT_ERROR;
+								}
+								// Parse the sample
+								UEMediaError Error = ParseMP4ASampleType(Track.Get(), STSDFirstChildBox);
+								if (Error == UEMEDIA_ERROR_NOT_SUPPORTED)
+								{
+									bIsSupported = false;
+								}
+								else if (Error != UEMEDIA_ERROR_OK)
+								{
+									return Error;
+								}
+								// Remember the encryption boxes.
+								Track->SCHMBox = SCHMBox;
+								Track->TENCBox = static_cast<const FMP4BoxTENC*>(STSDFirstChildBox->GetBoxPath(FMP4Box::kBox_sinf, FMP4Box::kBox_schi, FMP4Box::kBox_tenc));
 								break;
 							}
-							case FMP4Box::kSample_mp4a:
+							case FMP4Box::kSample_avc1:
 							{
-								const FMP4BoxAudioSampleEntry* AudioSampleEntry = static_cast<const FMP4BoxAudioSampleEntry*>(Track->STSDBox->GetChildBox(0));
-								check(AudioSampleEntry->GetNumberOfChildren() > 0);
-								if (AudioSampleEntry->GetNumberOfChildren() > 0)
+								UEMediaError Error = ParseAVC1SampleType(Track.Get(), STSDFirstChildBox);
+								if (Error == UEMEDIA_ERROR_NOT_SUPPORTED)
 								{
-									// There may be several entries. Usually there is the required sample format box (eg 'mp4a') and optional boxes like 'btrt'
-									bool bGotAudioFormat = false;
-									for(int32 j = 0, jMax = AudioSampleEntry->GetNumberOfChildren(); j < jMax; ++j)
-									{
-										const FMP4Box* SampleEntry = AudioSampleEntry->GetChildBox(j);
-										switch(SampleEntry->GetType())
-										{
-											case FMP4Box::kBox_esds:
-											{
-												const FMP4BoxESDS* ESDSBox = static_cast<const FMP4BoxESDS*>(AudioSampleEntry->GetChildBox(0));
-												Track->CodecSpecificDataMP4A = ESDSBox->GetESDescriptor();
-												bool bOk = Track->CodecSpecificDataMP4A.Parse();
-												if (bOk)
-												{
-													// Is this AAC audio?
-													if (Track->CodecSpecificDataMP4A.GetObjectTypeID() == MPEG::FESDescriptor::FObjectTypeID::MPEG4_Audio &&
-														Track->CodecSpecificDataMP4A.GetStreamType() == MPEG::FESDescriptor::FStreamType::Audio)
-													{
-														Track->CodecInformation.SetStreamType(EStreamType::Audio);
-														Track->CodecInformation.SetCodec(FStreamCodecInformation::ECodec::AAC);
-														Track->CodecInformation.SetStreamLanguageCode(Track->GetLanguage());
-
-														MPEG::FAACDecoderConfigurationRecord ConfigRecord;
-														if (ConfigRecord.ParseFrom(Track->CodecSpecificDataMP4A.GetCodecSpecificData().GetData(), Track->CodecSpecificDataMP4A.GetCodecSpecificData().Num()))
-														{
-															Track->CodecInformation.SetCodecSpecifierRFC6381(FString::Printf(TEXT("mp4a.40.%d"), ConfigRecord.ExtAOT ? ConfigRecord.ExtAOT : ConfigRecord.AOT));
-															Track->CodecInformation.SetSamplingRate(ConfigRecord.ExtSamplingFrequency ? ConfigRecord.ExtSamplingFrequency : ConfigRecord.SamplingRate);
-															Track->CodecInformation.SetNumberOfChannels((int32)ConfigRecord.ChannelConfiguration);
-															// We assume that all platforms can decode PS (parametric stereo). As such we change the channel count from mono to stereo
-															// to convey the _decoded_ format, not the source format.
-															if (ConfigRecord.ChannelConfiguration == 1 && ConfigRecord.PSSignal > 0)
-															{
-																Track->CodecInformation.SetNumberOfChannels(2);
-															}
-															Track->CodecInformation.GetExtras().Set(TEXT("samples_per_block"), FVariantValue(ConfigRecord.SBRSignal > 0 ? (int64)2048 : (int64)1024));
-														}
-
-														// Typically an mp4a track will not have a 'btrt' box because the bitrate is stored in the DecoderConfigDescriptor.
-														Track->BitrateInfo.BufferSizeDB = Track->CodecSpecificDataMP4A.GetBufferSize();
-														Track->BitrateInfo.MaxBitrate = Track->CodecSpecificDataMP4A.GetMaxBitrate();
-														Track->BitrateInfo.AvgBitrate = Track->CodecSpecificDataMP4A.GetAvgBitrate();
-
-														bGotAudioFormat = true;
-													}
-													else
-													{
-														// Not expected format. Ignore this track.
-														bIsSupported = false;
-														// Got the info, just not the format we want.
-														bGotAudioFormat = true;
-													}
-												}
-												else
-												{
-													return UEMEDIA_ERROR_FORMAT_ERROR;
-												}
-												break;
-											}
-											case FMP4Box::kBox_btrt:
-											{
-												const FMP4BoxBTRT* BTRTBox = static_cast<const FMP4BoxBTRT*>(SampleEntry);
-												Track->BitrateInfo.BufferSizeDB = BTRTBox->GetBufferSizeDB();
-												Track->BitrateInfo.MaxBitrate = BTRTBox->GetMaxBitrate();
-												Track->BitrateInfo.AvgBitrate = BTRTBox->GetAverageBitrate();
-												break;
-											}
-											default:
-											{
-												break;
-											}
-										}
-									}
-									// Check that we got the required information
-									if (!bGotAudioFormat)
-									{
-										return UEMEDIA_ERROR_FORMAT_ERROR;
-									}
-
+									bIsSupported = false;
 								}
-								else
+								else if (Error != UEMEDIA_ERROR_OK)
 								{
-									return UEMEDIA_ERROR_FORMAT_ERROR;
+									return Error;
 								}
-
 								break;
 							}
-							case FMP4Box::kSample_enca:
+							case FMP4Box::kSample_mp4a:
 							{
-								check(!TEXT("TODO")); // get CSD
+								UEMediaError Error = ParseMP4ASampleType(Track.Get(), STSDFirstChildBox);
+								if (Error == UEMEDIA_ERROR_NOT_SUPPORTED)
+								{
+									bIsSupported = false;
+								}
+								else if (Error != UEMEDIA_ERROR_OK)
+								{
+									return Error;
+								}
 								break;
 							}
 							case FMP4Box::kSample_stpp:
@@ -5560,6 +5907,27 @@ namespace Electra
 									{
 										Track->TRUNBoxes.Add(static_cast<const FMP4BoxTRUN*>(AllTRUNBoxes[nTruns]));
 									}
+
+									// Locate encryption related boxes. At present we do this only for fragmented streams.
+									Track->SAIZBox = static_cast<const FMP4BoxSAIZ*>(Box->FindBox(FMP4Box::kBox_saiz, 0));
+									Track->SAIOBox = static_cast<const FMP4BoxSAIO*>(Box->FindBox(FMP4Box::kBox_saio, 0));
+									Track->SENCBox = static_cast<const FMP4BoxSENC*>(Box->FindBox(FMP4Box::kBox_senc, 0));
+									// Get the PSSH boxes from the track fragment.
+									TArray<const FMP4Box*> AllPSSHBoxes;
+									Box->GetAllBoxInstances(AllPSSHBoxes, FMP4Box::kBox_pssh);
+									for(int32 nTruns=0; nTruns<AllPSSHBoxes.Num(); ++nTruns)
+									{
+										Track->PSSHBoxes.Add(static_cast<const FMP4BoxPSSH*>(AllPSSHBoxes[nTruns]));
+									}
+									// Now that we have the information needed to parse the senc box, do it.
+									if (Track->SENCBox)
+									{
+										UEMediaError Error = const_cast<FMP4BoxSENC*>(Track->SENCBox)->Parse(Track->TENCBox, Track->SAIZBox, Track->SAIOBox);
+										if (Error != UEMEDIA_ERROR_OK)
+										{
+											return Error;
+										}
+									}
 								}
 								else
 								{
diff --git a/Engine/Plugins/Media/ElectraPlayer/Source/ElectraPlayerRuntime/Private/Runtime/Demuxer/ParserISO14496-12.h b/Engine/Plugins/Media/ElectraPlayer/Source/ElectraPlayerRuntime/Private/Runtime/Demuxer/ParserISO14496-12.h
--- a/Engine/Plugins/Media/ElectraPlayer/Source/ElectraPlayerRuntime/Private/Runtime/Demuxer/ParserISO14496-12.h
+++ b/Engine/Plugins/Media/ElectraPlayer/Source/ElectraPlayerRuntime/Private/Runtime/Demuxer/ParserISO14496-12.h
@@ -112,7 +112,7 @@ namespace Electra
 			};
 
 			/**
-			 * Notifies the caller of the start of a new box.
+			 * Notifies the caller of the start/end of a new box.
 			 *
 			 * The caller may continue to parse this box or ask to stop parsing.
 			 *
@@ -123,6 +123,7 @@ namespace Electra
 			 * @return Whether or not to continue parsing thix box.
 			 */
 			virtual EParseContinuation OnFoundBox(FBoxType Box, int64 BoxSizeInBytes, int64 FileDataOffset, int64 BoxDataOffset) = 0;
+			virtual EParseContinuation OnEndOfBox(FBoxType Box, int64 BoxSizeInBytes, int64 FileDataOffset, int64 BoxDataOffset) = 0;
 		};
 
 
@@ -131,7 +132,7 @@ namespace Electra
 		/**
 		 * Parses the header boxes (all non-MDAT boxes).
 		 */
-		virtual UEMediaError ParseHeader(IReader* DataReader, IBoxCallback* BoxParseCallback, IPlayerSessionServices* PlayerSession) = 0;
+		virtual UEMediaError ParseHeader(IReader* DataReader, IBoxCallback* BoxParseCallback, IPlayerSessionServices* PlayerSession, const IParserISO14496_12* OptionalInitSegment) = 0;
 
 
 		/** A brand is a 32 bit value in an mp4 file. */
diff --git a/Engine/Plugins/Media/ElectraPlayer/Source/ElectraPlayerRuntime/Private/Runtime/HTTP/HTTPManager.cpp b/Engine/Plugins/Media/ElectraPlayer/Source/ElectraPlayerRuntime/Private/Runtime/HTTP/HTTPManager.cpp
--- a/Engine/Plugins/Media/ElectraPlayer/Source/ElectraPlayerRuntime/Private/Runtime/HTTP/HTTPManager.cpp
+++ b/Engine/Plugins/Media/ElectraPlayer/Source/ElectraPlayerRuntime/Private/Runtime/HTTP/HTTPManager.cpp
@@ -469,7 +469,14 @@ namespace Electra
 		Handle->HttpRequest->OnRequestProgress().BindThreadSafeSP(Handle->HttpsRequestCallbackWrapper.ToSharedRef(), &FHTTPCallbackWrapper::ReportRequestProgress);
 		Handle->HttpRequest->OnHeaderReceived().BindThreadSafeSP(Handle->HttpsRequestCallbackWrapper.ToSharedRef(), &FHTTPCallbackWrapper::ReportRequestHeaderReceived);
 		Handle->HttpRequest->SetURL(Request->Parameters.URL);
-		Handle->HttpRequest->SetVerb(TEXT("GET"));
+		if (!Request->Parameters.Verb.IsEmpty())
+		{
+			Handle->HttpRequest->SetVerb(Request->Parameters.Verb);
+		}
+		else
+		{
+			Handle->HttpRequest->SetVerb(TEXT("GET"));
+		}
 		Handle->HttpRequest->SetHeader(TEXT("User-Agent"), "X-UnrealEngine-Agent");
 		// Set accepted encoding first. If this is also present in custom headers let those overwrite it.
 		if (Request->Parameters.AcceptEncoding.IsSet())
diff --git a/Engine/Plugins/Media/ElectraPlayer/Source/ElectraPlayerRuntime/Private/Runtime/HTTP/HTTPManager.h b/Engine/Plugins/Media/ElectraPlayer/Source/ElectraPlayerRuntime/Private/Runtime/HTTP/HTTPManager.h
--- a/Engine/Plugins/Media/ElectraPlayer/Source/ElectraPlayerRuntime/Private/Runtime/HTTP/HTTPManager.h
+++ b/Engine/Plugins/Media/ElectraPlayer/Source/ElectraPlayerRuntime/Private/Runtime/HTTP/HTTPManager.h
@@ -386,6 +386,7 @@ namespace Electra
 				int64			DocumentSize = -1;
 			};
 			FString								URL;							//!< URL
+			FString								Verb;							//!< GET (default if not set), HEAD, OPTIONS,....
 			FRange								Range;							//!< Optional request range
 			int32								SubRangeRequestSize = 0;		//!< If not 0 the size to break the request into smaller range requests into.
 			TArray<HTTP::FHTTPHeader>			RequestHeaders;					//!< Request headers
diff --git a/Engine/Plugins/Media/ElectraPlayer/Source/ElectraPlayerRuntime/Private/Runtime/Player/AdaptivePlayerOptionKeynames.h b/Engine/Plugins/Media/ElectraPlayer/Source/ElectraPlayerRuntime/Private/Runtime/Player/AdaptivePlayerOptionKeynames.h
--- a/Engine/Plugins/Media/ElectraPlayer/Source/ElectraPlayerRuntime/Private/Runtime/Player/AdaptivePlayerOptionKeynames.h
+++ b/Engine/Plugins/Media/ElectraPlayer/Source/ElectraPlayerRuntime/Private/Runtime/Player/AdaptivePlayerOptionKeynames.h
@@ -6,6 +6,7 @@
 
 namespace Electra
 {
+	const TCHAR* const OptionKeyMimeType = TEXT("mime_type");										//!< (FString) mime type of URL to load
 	const TCHAR* const OptionKeyInitialBitrate = TEXT("initial_bitrate");							//!< (int64) value indicating the bitrate to start with.
 	const TCHAR* const OptionKeyLiveSeekableEndOffset = TEXT("seekable_range_live_end_offset");		//!< (FTimeValue) value specifying how many seconds away from the Live media timeline the seekable range should end.
 
diff --git a/Engine/Plugins/Media/ElectraPlayer/Source/ElectraPlayerRuntime/Private/Runtime/Player/AdaptiveStreamingPlayer.cpp b/Engine/Plugins/Media/ElectraPlayer/Source/ElectraPlayerRuntime/Private/Runtime/Player/AdaptiveStreamingPlayer.cpp
--- a/Engine/Plugins/Media/ElectraPlayer/Source/ElectraPlayerRuntime/Private/Runtime/Player/AdaptiveStreamingPlayer.cpp
+++ b/Engine/Plugins/Media/ElectraPlayer/Source/ElectraPlayerRuntime/Private/Runtime/Player/AdaptiveStreamingPlayer.cpp
@@ -287,9 +287,9 @@ void FAdaptiveStreamingPlayer::StopWorkerThread()
 
 
 
-void FAdaptiveStreamingPlayer::LoadManifest(const FString& ManifestURL)
+void FAdaptiveStreamingPlayer::LoadManifest(const FString& InManifestURL)
 {
-	WorkerThread.SendLoadManifestMessage(ManifestURL, FString());
+	WorkerThread.SendLoadManifestMessage(InManifestURL, FString());
 }
 
 
@@ -1337,7 +1337,10 @@ void FAdaptiveStreamingPlayer::WorkerThreadFN()
 			// Handle decoder changes
 			HandleDecoderChanges();
 			// Handle entity cache expirations.
-			EntityCache->HandleEntityExpiration();
+			if (EntityCache.IsValid())
+			{
+				EntityCache->HandleEntityExpiration();
+			}
 		}
 	}
 
@@ -3114,6 +3117,7 @@ void FAdaptiveStreamingPlayer::InternalStop(bool bHoldCurrentFrame)
 void FAdaptiveStreamingPlayer::InternalClose()
 {
 	// No longer need the manifest reader/updater
+	InternalCancelLoadManifest();
 	if (ManifestReader.IsValid())
 	{
 		ManifestReader->Close();
diff --git a/Engine/Plugins/Media/ElectraPlayer/Source/ElectraPlayerRuntime/Private/Runtime/Player/AdaptiveStreamingPlayerInternal.h b/Engine/Plugins/Media/ElectraPlayer/Source/ElectraPlayerRuntime/Private/Runtime/Player/AdaptiveStreamingPlayerInternal.h
--- a/Engine/Plugins/Media/ElectraPlayer/Source/ElectraPlayerRuntime/Private/Runtime/Player/AdaptiveStreamingPlayerInternal.h
+++ b/Engine/Plugins/Media/ElectraPlayer/Source/ElectraPlayerRuntime/Private/Runtime/Player/AdaptiveStreamingPlayerInternal.h
@@ -12,6 +12,7 @@
 
 #include "Demuxer/ParserISO14496-12.h"
 
+#include "Player/AdaptiveStreamingPlayerResourceRequest.h"
 #include "Player/PlayerStreamReader.h"
 #include "Player/PlaylistReader.h"
 #include "Player/PlayerStreamFilter.h"
@@ -1244,8 +1245,10 @@ private:
 
 
 	void InternalLoadManifest(const FString& URL, const FString& MimeType);
+	void InternalCancelLoadManifest();
 	bool SelectManifest();
 	void UpdateManifest();
+	void OnManifestGetMimeTypeComplete(TSharedPtrTS<FHTTPResourceRequest> InRequest);
 
 	void VideoDecoderInputNeeded(const IAccessUnitBufferListener::FBufferStats& currentInputBufferStats);
 	void VideoDecoderOutputReady(const IDecoderOutputBufferListener::FDecodeReadyStats& currentReadyStats);
@@ -1341,8 +1344,10 @@ private:
 
 	TMediaQueueDynamic<TSharedPtrTS<FErrorDetail>>						ErrorQueue;
 
+	FString																ManifestURL;
 	TSharedPtrTS<IManifest>												Manifest;
 	TSharedPtrTS<IPlaylistReader>										ManifestReader;
+	TSharedPtrTS<FHTTPResourceRequest>									ManifestMimeTypeRequest;
 
 	IStreamReader*														StreamReaderHandler;
 	TMediaOptionalValue<bool> 											bHaveVideoReader;
diff --git a/Engine/Plugins/Media/ElectraPlayer/Source/ElectraPlayerRuntime/Private/Runtime/Player/AdaptiveStreamingPlayerManifest.cpp b/Engine/Plugins/Media/ElectraPlayer/Source/ElectraPlayerRuntime/Private/Runtime/Player/AdaptiveStreamingPlayerManifest.cpp
--- a/Engine/Plugins/Media/ElectraPlayer/Source/ElectraPlayerRuntime/Private/Runtime/Player/AdaptiveStreamingPlayerManifest.cpp
+++ b/Engine/Plugins/Media/ElectraPlayer/Source/ElectraPlayerRuntime/Private/Runtime/Player/AdaptiveStreamingPlayerManifest.cpp
@@ -3,6 +3,7 @@
 #include "PlayerCore.h"
 #include "Misc/Paths.h"
 #include "Player/AdaptiveStreamingPlayerInternal.h"
+#include "Player/AdaptivePlayerOptionKeynames.h"
 #include "Player/HLS/PlaylistReaderHLS.h"
 #include "Player/mp4/PlaylistReaderMP4.h"
 #include "Player/DASH/PlaylistReaderDASH.h"
@@ -12,7 +13,6 @@
 #include "ParameterDictionary.h"
 
 
-
 namespace Electra
 {
 
@@ -90,6 +90,32 @@ FString GetMIMETypeForURL(const FString &URL)
 
 
 
+void FAdaptiveStreamingPlayer::OnManifestGetMimeTypeComplete(TSharedPtrTS<FHTTPResourceRequest> InRequest)
+{
+	if (!InRequest->GetWasCanceled())
+	{
+		FString mimeType(TEXT("application/octet-stream"));
+		if (InRequest->GetError() == 0)
+		{
+			const HTTP::FConnectionInfo* ci = InRequest->GetConnectionInfo();
+			if (ci)
+			{
+				mimeType = ci->ContentType;
+			}
+		}
+		WorkerThread.SendLoadManifestMessage(ManifestURL, mimeType);
+	}
+}
+
+
+void FAdaptiveStreamingPlayer::InternalCancelLoadManifest()
+{
+	if (ManifestMimeTypeRequest.IsValid())
+	{
+		ManifestMimeTypeRequest->Cancel();
+		ManifestMimeTypeRequest.Reset();
+	}
+}
 
 
 //-----------------------------------------------------------------------------
@@ -101,9 +127,30 @@ FString GetMIMETypeForURL(const FString &URL)
  */
 void FAdaptiveStreamingPlayer::InternalLoadManifest(const FString& URL, const FString& MimeType)
 {
+	// Remember the original request URL since we may lose the fragment part in requests.
+	ManifestURL = URL;
+	ManifestMimeTypeRequest.Reset();
 	if (CurrentState == EPlayerState::eState_Idle)
 	{
-		FString mimeType = Playlist::GetMIMETypeForURL(URL);
+		FString mimeType = MimeType;
+		if (GetOptions().HaveKey(OptionKeyMimeType))
+		{
+			mimeType = GetOptions().GetValue(OptionKeyMimeType).GetFString();
+		}
+		else if (mimeType.IsEmpty())
+		{
+			mimeType = Playlist::GetMIMETypeForURL(URL);
+		}
+		// If there is no mime type we need to issue a HTTP HEAD request in order to get the "Content-Type" header.
+		if (mimeType.IsEmpty() && (URL.StartsWith("https://", ESearchCase::CaseSensitive) || URL.StartsWith("http://", ESearchCase::CaseSensitive)))
+		{
+			ManifestMimeTypeRequest = MakeSharedTS<FHTTPResourceRequest>();
+			ManifestMimeTypeRequest->URL(URL);
+			ManifestMimeTypeRequest->Verb(TEXT("HEAD"));
+			ManifestMimeTypeRequest->Callback().BindThreadSafeSP(AsShared(), &FAdaptiveStreamingPlayer::OnManifestGetMimeTypeComplete);
+			ManifestMimeTypeRequest->StartGet(this);
+			return;
+		}
 		if (mimeType.Len())
 		{
 			check(!ManifestReader.IsValid());
diff --git a/Engine/Plugins/Media/ElectraPlayer/Source/ElectraPlayerRuntime/Private/Runtime/Player/AdaptiveStreamingPlayerResourceRequest.h b/Engine/Plugins/Media/ElectraPlayer/Source/ElectraPlayerRuntime/Private/Runtime/Player/AdaptiveStreamingPlayerResourceRequest.h
--- a/Engine/Plugins/Media/ElectraPlayer/Source/ElectraPlayerRuntime/Private/Runtime/Player/AdaptiveStreamingPlayerResourceRequest.h
+++ b/Engine/Plugins/Media/ElectraPlayer/Source/ElectraPlayerRuntime/Private/Runtime/Player/AdaptiveStreamingPlayerResourceRequest.h
@@ -65,6 +65,9 @@ public:
 
 	virtual FHTTPResourceRequest& URL(const FString& InURL)
 	{ Request->Parameters.URL = InURL; return *this; }
+
+	virtual FHTTPResourceRequest& Verb(const FString& InVerb)
+	{ Request->Parameters.Verb = InVerb; return *this; }
 	
 	virtual FHTTPResourceRequest& Range(const FString& InRange)
 	{ Request->Parameters.Range.Set(InRange); return *this; }
diff --git a/Engine/Plugins/Media/ElectraPlayer/Source/ElectraPlayerRuntime/Private/Runtime/Player/DASH/ManifestDASH.cpp b/Engine/Plugins/Media/ElectraPlayer/Source/ElectraPlayerRuntime/Private/Runtime/Player/DASH/ManifestDASH.cpp
--- a/Engine/Plugins/Media/ElectraPlayer/Source/ElectraPlayerRuntime/Private/Runtime/Player/DASH/ManifestDASH.cpp
+++ b/Engine/Plugins/Media/ElectraPlayer/Source/ElectraPlayerRuntime/Private/Runtime/Player/DASH/ManifestDASH.cpp
@@ -70,6 +70,10 @@ namespace DashUtils
 			}
 			return IParserISO14496_12::IBoxCallback::EParseContinuation::Continue;
 		}
+		virtual EParseContinuation OnEndOfBox(IParserISO14496_12::FBoxType Box, int64 BoxSizeInBytes, int64 FileDataOffset, int64 BoxDataOffset) override
+		{
+			return IParserISO14496_12::IBoxCallback::EParseContinuation::Continue;
+		}
 		bool bHaveSIDX = false;
 	};
 
@@ -2407,7 +2411,7 @@ void FManifestDASHInternal::FRepresentation::SegmentIndexDownloadComplete(TShare
 			DashUtils::FMP4SidxBoxReader BoxReader;
 			BoxReader.SetParseData(LoadRequest->Request->GetResponseBuffer());
 			TSharedPtrTS<IParserISO14496_12> Index = IParserISO14496_12::CreateParser();
-			UEMediaError parseError = Index->ParseHeader(&BoxReader, &BoxReader, LoadRequest->PlayerSessionServices);
+			UEMediaError parseError = Index->ParseHeader(&BoxReader, &BoxReader, LoadRequest->PlayerSessionServices, nullptr);
 			if (parseError == UEMEDIA_ERROR_OK || parseError == UEMEDIA_ERROR_END_OF_STREAM)
 			{
 				if (Index->PrepareTracks(nullptr) == UEMEDIA_ERROR_OK && Index->GetNumberOfSegmentIndices() > 0)
diff --git a/Engine/Plugins/Media/ElectraPlayer/Source/ElectraPlayerRuntime/Private/Runtime/Player/DASH/StreamReaderFMP4DASH.cpp b/Engine/Plugins/Media/ElectraPlayer/Source/ElectraPlayerRuntime/Private/Runtime/Player/DASH/StreamReaderFMP4DASH.cpp
--- a/Engine/Plugins/Media/ElectraPlayer/Source/ElectraPlayerRuntime/Private/Runtime/Player/DASH/StreamReaderFMP4DASH.cpp
+++ b/Engine/Plugins/Media/ElectraPlayer/Source/ElectraPlayerRuntime/Private/Runtime/Player/DASH/StreamReaderFMP4DASH.cpp
@@ -429,7 +429,7 @@ FErrorDetail FStreamReaderFMP4DASH::FStreamHandler::GetInitSegment(TSharedPtrTS<
 	FMP4StaticDataReader StaticDataReader;
 	StaticDataReader.SetParseData(LoadReq->Request->GetResponseBuffer());
 	TSharedPtrTS<IParserISO14496_12> Init = IParserISO14496_12::CreateParser();
-	UEMediaError parseError = Init->ParseHeader(&StaticDataReader, this, PlayerSessionService);
+	UEMediaError parseError = Init->ParseHeader(&StaticDataReader, this, PlayerSessionService, nullptr);
 	if (parseError == UEMEDIA_ERROR_OK || parseError == UEMEDIA_ERROR_END_OF_STREAM)
 	{
 		// Parse the tracks of the init segment. We do this mainly to get to the CSD we might need should we have to insert filler data later.
@@ -608,7 +608,7 @@ void FStreamReaderFMP4DASH::FStreamHandler::HandleRequest()
 			bool bIsFirstAU = true;
 			while(!bDone && !HasErrored() && !HasReadBeenAborted())
 			{
-				UEMediaError parseError = MP4Parser->ParseHeader(this, this, PlayerSessionService);
+				UEMediaError parseError = MP4Parser->ParseHeader(this, this, PlayerSessionService, MP4InitSegment.Get());
 				if (parseError == UEMEDIA_ERROR_OK)
 				{
 					Request->Segment.bSawLMSG = MP4Parser->HasBrand(IParserISO14496_12::BrandType_lmsg);
@@ -1226,6 +1226,11 @@ IParserISO14496_12::IBoxCallback::EParseContinuation FStreamReaderFMP4DASH::FStr
 	}
 }
 
+IParserISO14496_12::IBoxCallback::EParseContinuation FStreamReaderFMP4DASH::FStreamHandler::OnEndOfBox(IParserISO14496_12::FBoxType Box, int64 BoxSizeInBytes, int64 FileDataOffset, int64 BoxDataOffset)
+{
+	return IParserISO14496_12::IBoxCallback::EParseContinuation::Continue;
+}
+
 
 
 } // namespace Electra
diff --git a/Engine/Plugins/Media/ElectraPlayer/Source/ElectraPlayerRuntime/Private/Runtime/Player/DASH/StreamReaderFMP4DASH.h b/Engine/Plugins/Media/ElectraPlayer/Source/ElectraPlayerRuntime/Private/Runtime/Player/DASH/StreamReaderFMP4DASH.h
--- a/Engine/Plugins/Media/ElectraPlayer/Source/ElectraPlayerRuntime/Private/Runtime/Player/DASH/StreamReaderFMP4DASH.h
+++ b/Engine/Plugins/Media/ElectraPlayer/Source/ElectraPlayerRuntime/Private/Runtime/Player/DASH/StreamReaderFMP4DASH.h
@@ -159,6 +159,7 @@ private:
 		virtual int64 GetCurrentOffset() const override;
 		// Methods from IParserISO14496_12::IBoxCallback
 		virtual IParserISO14496_12::IBoxCallback::EParseContinuation OnFoundBox(IParserISO14496_12::FBoxType Box, int64 BoxSizeInBytes, int64 FileDataOffset, int64 BoxDataOffset) override;
+		virtual IParserISO14496_12::IBoxCallback::EParseContinuation OnEndOfBox(IParserISO14496_12::FBoxType Box, int64 BoxSizeInBytes, int64 FileDataOffset, int64 BoxDataOffset) override;
 	};
 
 	// Currently set to use 2 handlers, one for video and one for audio. This could become a pool of n if we need to stream
diff --git a/Engine/Plugins/Media/ElectraPlayer/Source/ElectraPlayerRuntime/Private/Runtime/Player/HLS/StreamReaderHLSfmp4.cpp b/Engine/Plugins/Media/ElectraPlayer/Source/ElectraPlayerRuntime/Private/Runtime/Player/HLS/StreamReaderHLSfmp4.cpp
--- a/Engine/Plugins/Media/ElectraPlayer/Source/ElectraPlayerRuntime/Private/Runtime/Player/HLS/StreamReaderHLSfmp4.cpp
+++ b/Engine/Plugins/Media/ElectraPlayer/Source/ElectraPlayerRuntime/Private/Runtime/Player/HLS/StreamReaderHLSfmp4.cpp
@@ -614,7 +614,7 @@ FStreamReaderHLSfmp4::FStreamHandler::EInitSegmentResult FStreamReaderHLSfmp4::F
 //		 Presently the assumption is that the init segment is not encrypted.
 
 				TSharedPtrTS<IParserISO14496_12> InitSegmentParser = IParserISO14496_12::CreateParser();
-				UEMediaError parseError = InitSegmentParser->ParseHeader(this, this, PlayerSessionService);
+				UEMediaError parseError = InitSegmentParser->ParseHeader(this, this, PlayerSessionService, nullptr);
 				if (parseError == UEMEDIA_ERROR_OK || parseError == UEMEDIA_ERROR_END_OF_STREAM)
 				{
 					// Parse the tracks of the init segment. We do this mainly to get to the CSD we might need should we have to insert filler data later.
@@ -880,7 +880,7 @@ void FStreamReaderHLSfmp4::FStreamHandler::HandleRequest()
 				bool bIsFirstAU = true;
 				while(!bDone && !HasErrored() && !HasReadBeenAborted())
 				{
-					UEMediaError parseError = MP4Parser->ParseHeader(this, this, PlayerSessionService);
+					UEMediaError parseError = MP4Parser->ParseHeader(this, this, PlayerSessionService, MP4InitSegment.Get());
 					if (parseError == UEMEDIA_ERROR_OK)
 					{
 						parseError = MP4Parser->PrepareTracks(MP4InitSegment);
@@ -1533,6 +1533,10 @@ IParserISO14496_12::IBoxCallback::EParseContinuation FStreamReaderHLSfmp4::FStre
 	}
 }
 
+IParserISO14496_12::IBoxCallback::EParseContinuation FStreamReaderHLSfmp4::FStreamHandler::OnEndOfBox(IParserISO14496_12::FBoxType Box, int64 BoxSizeInBytes, int64 FileDataOffset, int64 BoxDataOffset)
+{
+	return IParserISO14496_12::IBoxCallback::EParseContinuation::Continue;
+}
 
 
 } // namespace Electra
diff --git a/Engine/Plugins/Media/ElectraPlayer/Source/ElectraPlayerRuntime/Private/Runtime/Player/HLS/StreamReaderHLSfmp4.h b/Engine/Plugins/Media/ElectraPlayer/Source/ElectraPlayerRuntime/Private/Runtime/Player/HLS/StreamReaderHLSfmp4.h
--- a/Engine/Plugins/Media/ElectraPlayer/Source/ElectraPlayerRuntime/Private/Runtime/Player/HLS/StreamReaderHLSfmp4.h
+++ b/Engine/Plugins/Media/ElectraPlayer/Source/ElectraPlayerRuntime/Private/Runtime/Player/HLS/StreamReaderHLSfmp4.h
@@ -241,6 +241,7 @@ private:
 		virtual int64 GetCurrentOffset() const override;
 		// Methods from IParserISO14496_12::IBoxCallback
 		virtual IParserISO14496_12::IBoxCallback::EParseContinuation OnFoundBox(IParserISO14496_12::FBoxType Box, int64 BoxSizeInBytes, int64 FileDataOffset, int64 BoxDataOffset) override;
+		virtual IParserISO14496_12::IBoxCallback::EParseContinuation OnEndOfBox(IParserISO14496_12::FBoxType Box, int64 BoxSizeInBytes, int64 FileDataOffset, int64 BoxDataOffset) override;
 	};
 
 	// Currently set to use 2 handlers, one for video and one for audio. This could become a pool of n if we need to stream
diff --git a/Engine/Plugins/Media/ElectraPlayer/Source/ElectraPlayerRuntime/Private/Runtime/Player/mp4/PlaylistReaderMP4.cpp b/Engine/Plugins/Media/ElectraPlayer/Source/ElectraPlayerRuntime/Private/Runtime/Player/mp4/PlaylistReaderMP4.cpp
--- a/Engine/Plugins/Media/ElectraPlayer/Source/ElectraPlayerRuntime/Private/Runtime/Player/mp4/PlaylistReaderMP4.cpp
+++ b/Engine/Plugins/Media/ElectraPlayer/Source/ElectraPlayerRuntime/Private/Runtime/Player/mp4/PlaylistReaderMP4.cpp
@@ -69,32 +69,6 @@ public:
 	virtual TSharedPtrTS<IManifest> GetManifest() override;
 
 private:
-	struct FReadBuffer
-	{
-		FReadBuffer()
-		{
-			Reset();
-		}
-		void Reset()
-		{
-			ReceiveBuffer.Reset();
-			ParsePos = 0;
-			bHasBeenAborted = false;
-		}
-		void Abort()
-		{
-			TSharedPtrTS<IElectraHttpManager::FReceiveBuffer>	CurrentBuffer = ReceiveBuffer;
-			if (CurrentBuffer.IsValid())
-			{
-				CurrentBuffer->Buffer.Abort();
-			}
-			bHasBeenAborted = true;
-		}
-		TSharedPtrTS<IElectraHttpManager::FReceiveBuffer>	ReceiveBuffer;
-		int64												ParsePos;
-		bool												bHasBeenAborted;
-	};
-
 	// Methods from IParserISO14496_12::IReader
 	virtual int64 ReadData(void* IntoBuffer, int64 NumBytesToRead) override;
 	virtual bool HasReachedEOF() const override;
@@ -102,6 +76,7 @@ private:
 	virtual int64 GetCurrentOffset() const override;
 	// Methods from IParserISO14496_12::IBoxCallback
 	virtual IParserISO14496_12::IBoxCallback::EParseContinuation OnFoundBox(IParserISO14496_12::FBoxType Box, int64 BoxSizeInBytes, int64 FileDataOffset, int64 BoxDataOffset) override;
+	virtual IParserISO14496_12::IBoxCallback::EParseContinuation OnEndOfBox(IParserISO14496_12::FBoxType Box, int64 BoxSizeInBytes, int64 FileDataOffset, int64 BoxDataOffset) override;
 
 	void StartWorkerThread();
 	void StopWorkerThread();
@@ -115,27 +90,43 @@ private:
 	int32 HTTPProgressCallback(const IElectraHttpManager::FRequest* Request);
 	void HTTPCompletionCallback(const IElectraHttpManager::FRequest* Request);
 
+	void ClearRequest();
+
+	void ReadNextChunk(int64 InFromOffset, int64 ChunkSize);
+
 	bool HasErrored() const
 	{
 		return bHasErrored;
 	}
 
-	IPlayerSessionServices*									PlayerSessionServices;
+	const int32 kChunkReadSize = 65536;
+
+	IPlayerSessionServices*									PlayerSessionServices = nullptr;
 	FString													MasterPlaylistURL;
 	FMediaEvent												WorkerThreadQuitSignal;
-	bool													bIsWorkerThreadStarted;
+	bool													bIsWorkerThreadStarted = false;
 
+	FCriticalSection										Lock;
+	TSharedPtrTS<IElectraHttpManager::FRequest> 			Request;
+	TSharedPtrTS<IElectraHttpManager::FReceiveBuffer> 		ReceiveBuffer;
 	TSharedPtrTS<IElectraHttpManager::FProgressListener>	ProgressListener;
-	FReadBuffer												ReadBuffer;
-	bool													bAbort;
-	bool													bHasErrored;
+	HTTP::FConnectionInfo									ConnectionInfo;
+	FPODRingbuffer 											Buffer;
+	int64													ParsePos = 0;
+	int64													ChunkReadOffset = 0;
+	int64													FileSize = -1;
+	bool													bChunkReadInProgress = false;
+
+	bool													bAbort = false;
+	bool													bHasErrored = false;
 
 	TSharedPtrTS<IParserISO14496_12>						MP4Parser;
-	bool													bFoundBoxFTYP;
-	bool													bFoundBoxMOOV;
-	bool													bFoundBoxSIDX;
-	bool													bFoundBoxMOOF;
-	bool													bFoundBoxMDAT;
+	bool													bFoundBoxFTYP = false;
+	bool													bFoundBoxMOOV = false;
+	bool													bFoundBoxSIDX = false;
+	bool													bFoundBoxMOOF = false;
+	bool													bFoundBoxMDAT = false;
+	bool													bIsFastStartable = false;
 
 	TSharedPtrTS<FManifestMP4Internal>						Manifest;
 	FErrorDetail											LastErrorDetail;
@@ -163,15 +154,6 @@ TSharedPtrTS<IPlaylistReader> IPlaylistReaderMP4::Create(IPlayerSessionServices*
 
 FPlaylistReaderMP4::FPlaylistReaderMP4()
 	: FMediaThread("ElectraPlayer::MP4 Playlist")
-	, PlayerSessionServices(nullptr)
-	, bIsWorkerThreadStarted(false)
-	, bAbort(false)
-	, bHasErrored(false)
-	, bFoundBoxFTYP(false)
-	, bFoundBoxMOOV(false)
-	, bFoundBoxSIDX(false)
-	, bFoundBoxMOOF(false)
-	, bFoundBoxMDAT(false)
 {
 }
 
@@ -193,15 +175,13 @@ TSharedPtrTS<IManifest> FPlaylistReaderMP4::GetManifest()
 void FPlaylistReaderMP4::Initialize(IPlayerSessionServices* InPlayerSessionServices)
 {
 	PlayerSessionServices = InPlayerSessionServices;
-	ProgressListener = MakeSharedTS<IElectraHttpManager::FProgressListener>();
-	ProgressListener->CompletionDelegate = Electra::MakeDelegate(this, &FPlaylistReaderMP4::HTTPCompletionCallback);
-	ProgressListener->ProgressDelegate   = Electra::MakeDelegate(this, &FPlaylistReaderMP4::HTTPProgressCallback);
+	Buffer.Reserve(kChunkReadSize);
 }
 
 void FPlaylistReaderMP4::Close()
 {
-	ReadBuffer.Abort();
 	bAbort = true;
+	ClearRequest();
 
 	StopWorkerThread();
 }
@@ -267,7 +247,7 @@ void FPlaylistReaderMP4::LoadAndParse(const FString& URL)
 	StartWorkerThread();
 }
 
-int32 FPlaylistReaderMP4::HTTPProgressCallback(const IElectraHttpManager::FRequest* Request)
+int32 FPlaylistReaderMP4::HTTPProgressCallback(const IElectraHttpManager::FRequest* InRequest)
 {
 	// Aborted?
 	return bAbort ? 1 : 0;
@@ -275,80 +255,133 @@ int32 FPlaylistReaderMP4::HTTPProgressCallback(const IElectraHttpManager::FReque
 
 void FPlaylistReaderMP4::HTTPCompletionCallback(const IElectraHttpManager::FRequest* InRequest)
 {
-	bHasErrored = InRequest->ConnectionInfo.StatusInfo.ErrorDetail.IsError();
+	bool bFailed = InRequest->ConnectionInfo.StatusInfo.ErrorDetail.IsError();
+	ConnectionInfo = InRequest->ConnectionInfo;
+	if (!bFailed)
+	{
+		// Set the size of the resource if we don't have it yet.
+		if (FileSize < 0)
+		{
+			IElectraHttpManager::FParams::FRange crh;
+			if (crh.ParseFromContentRangeResponse(InRequest->ConnectionInfo.ContentRangeHeader))
+			{
+				FileSize = crh.GetDocumentSize();
+			}
+		}
+
+		// Copy the read data across.
+		int32 NumRead = ReceiveBuffer->Buffer.Num();
+		Buffer.PushData(ReceiveBuffer->Buffer.GetLinearReadData(), NumRead);
+		if (FileSize >= 0 && ChunkReadOffset + NumRead >= FileSize)
+		{
+			Buffer.SetEOD();
+		}
+	}
+	bHasErrored = bFailed;
+	bChunkReadInProgress = false;
 }
 
+void FPlaylistReaderMP4::ClearRequest()
+{
+	FScopeLock lock(&Lock);
+	ProgressListener.Reset();
+	ReceiveBuffer.Reset();
+	if (Request.IsValid())
+	{
+		PlayerSessionServices->GetHTTPManager()->RemoveRequest(Request);
+		Request.Reset();
+	}
+}
 
-void FPlaylistReaderMP4::WorkerThread()
+void FPlaylistReaderMP4::ReadNextChunk(int64 InFromOffset, int64 ChunkSize)
 {
-	LLM_SCOPE(ELLMTag::ElectraPlayer);
-	CSV_SCOPED_TIMING_STAT(ElectraPlayer, PlaylistReaderMP4_Worker);
+	ClearRequest();
 
-	ReadBuffer.Reset();
-	ReadBuffer.ReceiveBuffer = MakeSharedTS<IElectraHttpManager::FReceiveBuffer>();
-	ReadBuffer.ReceiveBuffer->Buffer.Reserve(1 << 20);
-	ReadBuffer.ReceiveBuffer->bEnableRingbuffer = true;
+	FScopeLock lock(&Lock);
+	// Asked to go beyond the size of the file?
+	ChunkReadOffset = InFromOffset;
+	if (FileSize >= 0 && InFromOffset >= FileSize)
+	{
+		Buffer.SetEOD();
+		return;
+	}
+	ProgressListener = MakeSharedTS<IElectraHttpManager::FProgressListener>();
+	ProgressListener->CompletionDelegate = Electra::MakeDelegate(this, &FPlaylistReaderMP4::HTTPCompletionCallback);
+	ProgressListener->ProgressDelegate   = Electra::MakeDelegate(this, &FPlaylistReaderMP4::HTTPProgressCallback);
 
-	TSharedPtrTS<IElectraHttpManager::FRequest> HTTP(new IElectraHttpManager::FRequest);
-	HTTP->Parameters.URL = MasterPlaylistURL;
-	// Perform small reads of 256K ranges only. We just want to fetch the boxes up to the first 'mdat' and those should be fairly small.
-	HTTP->Parameters.SubRangeRequestSize = 256 << 10;
-	HTTP->ReceiveBuffer = ReadBuffer.ReceiveBuffer;
-	HTTP->ProgressListener = ProgressListener;
-	PlayerSessionServices->GetHTTPManager()->AddRequest(HTTP);
+	ReceiveBuffer = MakeSharedTS<IElectraHttpManager::FReceiveBuffer>();
+	ReceiveBuffer->Buffer.Reserve(ChunkSize);
 
-	// Create the parser we need for parsing the "moov" box containing all the track information.
-	MP4Parser = IParserISO14496_12::CreateParser();
+	Request = MakeSharedTS<IElectraHttpManager::FRequest>();
+	Request->Parameters.URL = MasterPlaylistURL;
+	Request->Parameters.Range.SetStart(InFromOffset);
+	int64 LastByte = InFromOffset + ChunkSize - 1;
+	if (FileSize >= 0 && LastByte > FileSize-1)
+	{
+		LastByte = FileSize - 1;
+	}
+	Request->Parameters.Range.SetEndIncluding(LastByte);
+	Request->ReceiveBuffer = ReceiveBuffer;
+	Request->ProgressListener = ProgressListener;
+	PlayerSessionServices->GetHTTPManager()->AddRequest(Request);
+}
 
-	UEMediaError parseError = MP4Parser->ParseHeader(this, this, PlayerSessionServices);
 
-	ProgressListener.Reset();
-	PlayerSessionServices->GetHTTPManager()->RemoveRequest(HTTP);
+
+
+void FPlaylistReaderMP4::WorkerThread()
+{
+	LLM_SCOPE(ELLMTag::ElectraPlayer);
+	CSV_SCOPED_TIMING_STAT(ElectraPlayer, PlaylistReaderMP4_Worker);
+
+	MP4Parser = IParserISO14496_12::CreateParser();
+	UEMediaError parseError = MP4Parser->ParseHeader(this, this, PlayerSessionServices, nullptr);
+	ClearRequest();
 
 	if (parseError != UEMEDIA_ERROR_ABORTED)
 	{
-		const HTTP::FConnectionInfo* ConnInfo = &HTTP->ConnectionInfo;
-
 		// Notify the download of the "master playlist". This indicates the download only, not the parsing thereof.
-		PlayerSessionServices->SendMessageToPlayer(IPlaylistReader::PlaylistDownloadMessage::Create(ConnInfo, Playlist::EListType::Master, Playlist::ELoadType::Initial));
+		PlayerSessionServices->SendMessageToPlayer(IPlaylistReader::PlaylistDownloadMessage::Create(&ConnectionInfo, Playlist::EListType::Master, Playlist::ELoadType::Initial));
 		// Notify that the "master playlist" has been parsed, successfully or not.
-		PlayerSessionServices->SendMessageToPlayer(IPlaylistReader::PlaylistLoadedMessage::Create(LastErrorDetail, ConnInfo, Playlist::EListType::Master, Playlist::ELoadType::Initial));
+		PlayerSessionServices->SendMessageToPlayer(IPlaylistReader::PlaylistLoadedMessage::Create(LastErrorDetail, &ConnectionInfo, Playlist::EListType::Master, Playlist::ELoadType::Initial));
 
-		if (parseError == UEMEDIA_ERROR_OK)
+		if (parseError == UEMEDIA_ERROR_OK || parseError == UEMEDIA_ERROR_END_OF_STREAM)
 		{
 			// See that we have parsed all the boxes we need.
 			if (bFoundBoxFTYP && bFoundBoxMOOV)
 			{
+				if (!bIsFastStartable)
+				{
+					LogMessage(IInfoLog::ELevel::Info, FString::Printf(TEXT("The mp4 at \"%s\" is not fast-startable. Consider moving the 'moov' box in front of the 'mdat' for faster startup times."), *ConnectionInfo.EffectiveURL));
+				}
+
 				// Prepare the tracks in the stream that are of a supported codec.
 				parseError = MP4Parser->PrepareTracks(TSharedPtrTS<const IParserISO14496_12>());
 				if (parseError == UEMEDIA_ERROR_OK)
 				{
 					Manifest = MakeSharedTS<FManifestMP4Internal>(PlayerSessionServices);
-					FErrorDetail err = Manifest->Build(MP4Parser, MasterPlaylistURL, *ConnInfo);
+					FErrorDetail err = Manifest->Build(MP4Parser, MasterPlaylistURL, ConnectionInfo);
 
 					// Notify that the "variant playlists" are ready. There are no variants in an mp4, but this is the trigger that the playlists are all set up and are good to go now.
-					PlayerSessionServices->SendMessageToPlayer(IPlaylistReader::PlaylistLoadedMessage::Create(LastErrorDetail, ConnInfo, Playlist::EListType::Variant, Playlist::ELoadType::Initial));
+					PlayerSessionServices->SendMessageToPlayer(IPlaylistReader::PlaylistLoadedMessage::Create(LastErrorDetail, &ConnectionInfo, Playlist::EListType::Variant, Playlist::ELoadType::Initial));
 				}
 				else
 				{
-					PostError(FString::Printf(TEXT("Failed to parse tracks in mp4 \"%s\" with error %u"), *ConnInfo->EffectiveURL, parseError), ERRCODE_MP4_INVALID_FILE, UEMEDIA_ERROR_FORMAT_ERROR);
+					PostError(FString::Printf(TEXT("Failed to parse tracks in mp4 \"%s\" with error %u"), *ConnectionInfo.EffectiveURL, parseError), ERRCODE_MP4_INVALID_FILE, UEMEDIA_ERROR_FORMAT_ERROR);
 				}
 			}
 			else
 			{
 				// No moov box usually means this is not a fast-start file.
-				PostError(FString::Printf(TEXT("No moov box found before mdat in \"%s\". Make sure the video is fast-startable."), *ConnInfo->EffectiveURL), ERRCODE_MP4_INVALID_FILE, UEMEDIA_ERROR_FORMAT_ERROR);
+				PostError(FString::Printf(TEXT("No moov box found in \"%s\". This is not a valid file."), *ConnectionInfo.EffectiveURL), ERRCODE_MP4_INVALID_FILE, UEMEDIA_ERROR_FORMAT_ERROR);
 			}
 		}
 		else
 		{
-			PostError(FString::Printf(TEXT("Failed to parse mp4 \"%s\" with error %u"), *ConnInfo->EffectiveURL, parseError), ERRCODE_MP4_INVALID_FILE, UEMEDIA_ERROR_FORMAT_ERROR);
+			PostError(FString::Printf(TEXT("Failed to parse mp4 \"%s\" with error %u"), *ConnectionInfo.EffectiveURL, parseError), ERRCODE_MP4_INVALID_FILE, UEMEDIA_ERROR_FORMAT_ERROR);
 		}
 	}
 
-	HTTP.Reset();
-	ReadBuffer.Reset();
-
 	// This thread's work is done. We only wait for termination now.
 	WorkerThreadQuitSignal.Wait();
 }
@@ -364,55 +397,70 @@ void FPlaylistReaderMP4::WorkerThread()
  * Reading must return the number of bytes asked to get, if necessary by blocking.
  * If a read error prevents reading the number of bytes -1 must be returned.
  *
- * @param ToBuffer Buffer into which to store the data bytes. If nullptr is passed the data must be skipped over.
- * @param NumBytes The number of bytes to read. Must not read more bytes and no less than requested.
+ * @param IntoBuffer Buffer into which to store the data bytes. If nullptr is passed the data must be skipped over.
+ * @param NumBytesToRead The number of bytes to read. Must not read more bytes and no less than requested.
  * @return The number of bytes read or -1 on a read error.
  */
-int64 FPlaylistReaderMP4::ReadData(void* ToBuffer, int64 NumBytes)
+int64 FPlaylistReaderMP4::ReadData(void* IntoBuffer, int64 NumBytesToRead)
 {
-	FPODRingbuffer& SourceBuffer = ReadBuffer.ReceiveBuffer->Buffer;
-
-	uint8* OutputBuffer = (uint8*)ToBuffer;
+	uint8* OutputBuffer = (uint8*)IntoBuffer;
 	// Do we have enough data in the ringbuffer to satisfy the read?
-	if (SourceBuffer.Num() >= NumBytes)
+	if (Buffer.Num() >= NumBytesToRead)
 	{
 		// Yes. Get the data and return.
-		int32 NumGot = SourceBuffer.PopData(OutputBuffer, NumBytes);
-		check(NumGot == NumBytes);
-		ReadBuffer.ParsePos += NumBytes;
-		return NumBytes;
+		int32 NumGot = Buffer.PopData(OutputBuffer, NumBytesToRead);
+		check(NumGot == NumBytesToRead);
+		ParsePos += NumBytesToRead;
+		return NumBytesToRead;
 	}
 	else
 	{
 		// Do not have enough data yet or we want to read more than the ringbuffer can hold.
-		int32 NumBytesToGo = NumBytes;
+		int32 NumBytesToGo = NumBytesToRead;
+		int64 NextChunkReadOffset = ParsePos;
 		while(NumBytesToGo > 0)
 		{
-			if (HasErrored() || SourceBuffer.WasAborted() || bAbort)
+			if (bHasErrored || bAbort)
 			{
 				return -1;
 			}
 			// EOD?
-			if (SourceBuffer.IsEndOfData())
+			if (Buffer.IsEndOfData())
 			{
 				return 0;
 			}
 
 			// Get whatever amount of data is currently available to free up the buffer for receiving more data.
-			int32 NumGot = SourceBuffer.PopData(OutputBuffer, NumBytesToGo);
+			int32 NumGot = Buffer.PopData(OutputBuffer, NumBytesToGo);
 			if ((NumBytesToGo -= NumGot) > 0)
 			{
 				if (OutputBuffer)
 				{
 					OutputBuffer += NumGot;
 				}
+				// Trigger read of next chunk of data.
+				if (!bChunkReadInProgress)
+				{
+					// Is the data to read actually used or is it skipped over?
+					if (OutputBuffer)
+					{
+						NextChunkReadOffset += NumGot;
+						ReadNextChunk(NextChunkReadOffset, kChunkReadSize);
+						bChunkReadInProgress = true;
+					}
+					else
+					{
+						// Data is not used, so do not request it.
+						break;
+					}
+				}
 				// Wait for data to arrive in the ringbuffer.
-				int32 WaitForBytes = NumBytesToGo > SourceBuffer.Capacity() ? SourceBuffer.Capacity() : NumBytesToGo;
-				SourceBuffer.WaitUntilSizeAvailable(WaitForBytes, 1000 * 100);
+				int32 WaitForBytes = NumBytesToGo > Buffer.Capacity() ? Buffer.Capacity() : NumBytesToGo;
+				Buffer.WaitUntilSizeAvailable(WaitForBytes, 1000 * 100);
 			}
 		}
-		ReadBuffer.ParsePos += NumBytes;
-		return NumBytes;
+		ParsePos += NumBytesToRead;
+		return NumBytesToRead;
 	}
 }
 
@@ -423,8 +471,7 @@ int64 FPlaylistReaderMP4::ReadData(void* ToBuffer, int64 NumBytes)
  */
 bool FPlaylistReaderMP4::HasReachedEOF() const
 {
-	FPODRingbuffer& SourceBuffer = ReadBuffer.ReceiveBuffer->Buffer;
-	return SourceBuffer.IsEndOfData();
+	return Buffer.IsEndOfData();
 }
 
 /**
@@ -434,7 +481,7 @@ bool FPlaylistReaderMP4::HasReachedEOF() const
  */
 bool FPlaylistReaderMP4::HasReadBeenAborted() const
 {
-	return bAbort || ReadBuffer.bHasBeenAborted;
+	return bAbort;
 }
 
 /**
@@ -446,7 +493,7 @@ bool FPlaylistReaderMP4::HasReadBeenAborted() const
  */
 int64 FPlaylistReaderMP4::GetCurrentOffset() const
 {
-	return ReadBuffer.ParsePos;
+	return ParsePos;
 }
 
 
@@ -463,35 +510,34 @@ IParserISO14496_12::IBoxCallback::EParseContinuation FPlaylistReaderMP4::OnFound
 	switch(Box)
 	{
 		case IParserISO14496_12::BoxType_ftyp:
-		{
 			bFoundBoxFTYP = true;
-			return IParserISO14496_12::IBoxCallback::EParseContinuation::Continue;
-		}
+			break;
 		case IParserISO14496_12::BoxType_moov:
-		{
 			bFoundBoxMOOV = true;
-			return IParserISO14496_12::IBoxCallback::EParseContinuation::Continue;
-		}
+			bIsFastStartable = !bFoundBoxMDAT;
+			break;
 		case IParserISO14496_12::BoxType_sidx:
-		{
 			bFoundBoxSIDX = true;
-			return IParserISO14496_12::IBoxCallback::EParseContinuation::Continue;
-		}
+			break;
 		case IParserISO14496_12::BoxType_moof:
-		{
 			bFoundBoxMOOF = true;
 			return IParserISO14496_12::IBoxCallback::EParseContinuation::Stop;
-		}
 		case IParserISO14496_12::BoxType_mdat:
-		{
 			bFoundBoxMDAT = true;
-			return IParserISO14496_12::IBoxCallback::EParseContinuation::Stop;
-		}
+			break;
 		default:
-		{
-			return IParserISO14496_12::IBoxCallback::EParseContinuation::Continue;
-		}
+			break;
+	}
+	return IParserISO14496_12::IBoxCallback::EParseContinuation::Continue;
+}
+
+IParserISO14496_12::IBoxCallback::EParseContinuation FPlaylistReaderMP4::OnEndOfBox(IParserISO14496_12::FBoxType Box, int64 BoxSizeInBytes, int64 FileDataOffset, int64 BoxDataOffset)
+{
+	if (bFoundBoxMOOV)
+	{
+		return IParserISO14496_12::IBoxCallback::EParseContinuation::Stop;
 	}
+	return IParserISO14496_12::IBoxCallback::EParseContinuation::Continue;
 }
 
 
