ElectraPlayer: silencing a PVS code analysis warning

#rnx

diff --git a/Engine/Plugins/Media/ElectraPlayer/Source/ElectraPlayerRuntime/Private/Runtime/Demuxer/ParserISO14496-12.cpp b/Engine/Plugins/Media/ElectraPlayer/Source/ElectraPlayerRuntime/Private/Runtime/Demuxer/ParserISO14496-12.cpp
--- a/Engine/Plugins/Media/ElectraPlayer/Source/ElectraPlayerRuntime/Private/Runtime/Demuxer/ParserISO14496-12.cpp
+++ b/Engine/Plugins/Media/ElectraPlayer/Source/ElectraPlayerRuntime/Private/Runtime/Demuxer/ParserISO14496-12.cpp
@@ -5775,11 +5775,11 @@ private:
 								}
 								// Parse the sample
 								UEMediaError Error = ParseAVC1SampleType(Track.Get(), STSDFirstChildBox);
-								if (Error == UEMEDIA_ERROR_NOT_SUPPORTED)
+								/*if (Error == UEMEDIA_ERROR_NOT_SUPPORTED)
 								{
 									bIsSupported = false;
 								}
-								else if (Error != UEMEDIA_ERROR_OK)
+								else*/ if (Error != UEMEDIA_ERROR_OK)
 								{
 									return Error;
 								}
@@ -5824,11 +5824,12 @@ private:
 							case FMP4Box::kSample_avc1:
 							{
 								UEMediaError Error = ParseAVC1SampleType(Track.Get(), STSDFirstChildBox);
+								/*
 								if (Error == UEMEDIA_ERROR_NOT_SUPPORTED)
 								{
 									bIsSupported = false;
 								}
-								else if (Error != UEMEDIA_ERROR_OK)
+								else*/ if (Error != UEMEDIA_ERROR_OK)
 								{
 									return Error;
 								}
