ElectraPlayer: Synchronizing player time to server based on (supported) DASH <UTCTiming> elements. Fixed an issue in the URL parser resolving a path with authority but no scheme.

#rnx

diff --git a/Engine/Plugins/Media/ElectraPlayer/Source/ElectraPlayerRuntime/Private/Runtime/Player/AdaptiveStreamingPlayerManifest.cpp b/Engine/Plugins/Media/ElectraPlayer/Source/ElectraPlayerRuntime/Private/Runtime/Player/AdaptiveStreamingPlayerManifest.cpp
--- a/Engine/Plugins/Media/ElectraPlayer/Source/ElectraPlayerRuntime/Private/Runtime/Player/AdaptiveStreamingPlayerManifest.cpp
+++ b/Engine/Plugins/Media/ElectraPlayer/Source/ElectraPlayerRuntime/Private/Runtime/Player/AdaptiveStreamingPlayerManifest.cpp
@@ -36,20 +36,15 @@ static const FString MIMETypeDASH(TEXT("application/dash+xml"));
 FString GetMIMETypeForURL(const FString &URL)
 {
 	FString MimeType;
-	TUniquePtr<IURLParser> UrlParser(IURLParser::Create());
-	if (!UrlParser)
-	{
-		return MimeType;
-	}
-	UEMediaError Error = UrlParser->ParseURL(URL);
-	if (Error != UEMEDIA_ERROR_OK)
+	FURL_RFC3986 UrlParser;
+	if (!UrlParser.Parse(URL))
 	{
 		return MimeType;
 	}
 
 	//	FString PathOnly = UrlParser->GetPath();
 	TArray<FString> PathComponents;
-	UrlParser->GetPathComponents(PathComponents);
+	UrlParser.GetPathComponents(PathComponents);
 	if (PathComponents.Num())
 	{
 		FString LowerCaseExtension = FPaths::GetExtension(PathComponents.Last().ToLower());
diff --git a/Engine/Plugins/Media/ElectraPlayer/Source/ElectraPlayerRuntime/Private/Runtime/Player/DASH/MPDElementsDASH.h b/Engine/Plugins/Media/ElectraPlayer/Source/ElectraPlayerRuntime/Private/Runtime/Player/DASH/MPDElementsDASH.h
--- a/Engine/Plugins/Media/ElectraPlayer/Source/ElectraPlayerRuntime/Private/Runtime/Player/DASH/MPDElementsDASH.h
+++ b/Engine/Plugins/Media/ElectraPlayer/Source/ElectraPlayerRuntime/Private/Runtime/Player/DASH/MPDElementsDASH.h
@@ -1659,6 +1659,11 @@ public:
 	void SetFetchTime(const FTimeValue& InFetchTime)		{ FScopeLock Lock(&UpdateLock); FetchTime = InFetchTime; }
 	FString GetETag() const									{ return ETag; }
 	void SetETag(const FString& InETag)						{ ETag = InETag; }
+	void RemoveUTCTimingElement(TSharedPtrTS<FDashMPD_DescriptorType> TimingElement)
+	{
+		FScopeLock lock(&UpdateLock);
+		UTCTimings.Remove(TimingElement);
+	}
 
 	// Methods to manipulate the presentation type.
 	void LockAccess()											{ UpdateLock.Lock(); }
diff --git a/Engine/Plugins/Media/ElectraPlayer/Source/ElectraPlayerRuntime/Private/Runtime/Player/DASH/PlaylistReaderDASH.cpp b/Engine/Plugins/Media/ElectraPlayer/Source/ElectraPlayerRuntime/Private/Runtime/Player/DASH/PlaylistReaderDASH.cpp
--- a/Engine/Plugins/Media/ElectraPlayer/Source/ElectraPlayerRuntime/Private/Runtime/Player/DASH/PlaylistReaderDASH.cpp
+++ b/Engine/Plugins/Media/ElectraPlayer/Source/ElectraPlayerRuntime/Private/Runtime/Player/DASH/PlaylistReaderDASH.cpp
@@ -63,6 +63,7 @@ private:
 	void OnHTTPResourceRequestComplete(TSharedPtrTS<FHTTPResourceRequest> InRequest);
 
 	FErrorDetail GetXMLResponseString(FString& OutXMLString, FResourceLoadRequestPtr FromRequest);
+	FErrorDetail GetResponseString(FString& OutString, FResourceLoadRequestPtr FromRequest);
 
 	void StartWorkerThread();
 	void StopWorkerThread();
@@ -72,6 +73,8 @@ private:
 	void HandleStaticRequestCompletions(const FTimeValue& TimeNow);
 	void CheckForMPDUpdate();
 
+	void TriggerTimeSynchronization();
+
 	void PostError(const FErrorDetail& Error);
 	void PostError(const FString& Message, uint16 Code, UEMediaError Error = UEMEDIA_ERROR_OK);
 	void LogMessage(IInfoLog::ELevel Level, const FString& Message);
@@ -88,6 +91,9 @@ private:
 	void ManifestUpdateDownloadCompleted(FResourceLoadRequestPtr Request, bool bSuccess);
 	void InitialMPDXLinkElementDownloadCompleted(FResourceLoadRequestPtr Request, bool bSuccess);
 
+	void Timesync_httphead_Completed(FResourceLoadRequestPtr Request, bool bSuccess);
+	void Timesync_httpxsdate_Completed(FResourceLoadRequestPtr Request, bool bSuccess);
+	void Timesync_httpiso_Completed(FResourceLoadRequestPtr Request, bool bSuccess);
 
 	IPlayerSessionServices*									PlayerSessionServices = nullptr;
 	FString													MPDURL;
@@ -118,6 +124,29 @@ private:
 	FErrorDetail											LastErrorDetail;
 
 	TSharedPtrTS<FManifestDASH>								PlayerManifest;
+
+	// Time synchronization.
+	TArray<TSharedPtrTS<FDashMPD_DescriptorType>>			AttemptedTimesyncDescriptors;
+	bool													bTimeSyncInProgress = false;
+	bool													bInitialTimeSyncedToMPDDateTimeHeader = false;
+
+	// Warnings
+	bool													bWarnedAboutNoUTCTimingElements = false;
+	bool													bWarnedAboutUnsupportedUTCTimingElement = false;
+};
+
+
+namespace DASHTimingSources
+{
+    const TCHAR* const Scheme_urn_mpeg_dash_utc_httphead2014 = TEXT("urn:mpeg:dash:utc:http-head:2014");
+    const TCHAR* const Scheme_urn_mpeg_dash_utc_httpxsdate2014 = TEXT("urn:mpeg:dash:utc:http-xsdate:2014");
+    const TCHAR* const Scheme_urn_mpeg_dash_utc_httpiso2014 = TEXT("urn:mpeg:dash:utc:http-iso:2014");
+    const TCHAR* const Scheme_urn_mpeg_dash_utc_direct2014 = TEXT("urn:mpeg:dash:utc:direct:2014");
+/* not currently supported
+    const TCHAR* const Scheme_urn_mpeg_dash_utc_httpntp2014 = TEXT("urn:mpeg:dash:utc:http-ntp:2014");
+    const TCHAR* const Scheme_urn_mpeg_dash_utc_ntp2014 = TEXT("urn:mpeg:dash:utc:ntp:2014");
+    const TCHAR* const Scheme_urn_mpeg_dash_utc_sntp2014 = TEXT("urn:mpeg:dash:utc:sntp:2014");
+*/
 };
 
 
@@ -412,6 +441,134 @@ void FPlaylistReaderDASH::RequestMPDUpdate(bool bForcedUpdate)
 }
 
 
+void FPlaylistReaderDASH::TriggerTimeSynchronization()
+{
+	if (bTimeSyncInProgress)
+	{
+		return;
+	}
+
+	// Time sync is only necessary when there is either an MPD@availabilityStartTime or MPD@type is dynamic.
+	if (Manifest.IsValid() && (Manifest->UsesAST() || !Manifest->IsStaticType()))
+	{
+		TSharedPtrTS<FDashMPD_MPDType> MPDRoot = Manifest->GetMPDRoot();
+		if (MPDRoot.IsValid())
+		{
+			bool bFoundSupportedElement = false;
+			FString MPDDocumentURL = MPDRoot->GetDocumentURL();
+			const TArray<TSharedPtrTS<FDashMPD_DescriptorType>>& UTCTimings = MPDRoot->GetUTCTimings();
+
+			// Emit a one-time only warning if there is no <UTCTiming> element in the MPD when there should be.
+			if (UTCTimings.Num() == 0 && !bWarnedAboutNoUTCTimingElements)
+			{
+				bWarnedAboutNoUTCTimingElements = true;
+				LogMessage(IInfoLog::ELevel::Warning, FString::Printf(TEXT("No <UTCTiming> element found in MPD. This could result in playback failures. For detailed validation of your MPD please use https://conformance.dashif.org/")));
+			}
+
+			/*
+				We go over the UTC timing elements in two passes.
+				In the first we look only for the (discouraged) urn:mpeg:dash:utc:direct:2014 scheme because
+				if it is available we can set it immediately. Its value is baked into the manifest and is therefore
+				only really valid right at the time the MPD was fetched.
+				If we try other schemes first that involve a network request and those fail then there will have
+				passed that much time that the hardcoded value in the MPD will no longer make sense to associate
+				with the current system time.
+			*/
+			for(int32 i=0; i<UTCTimings.Num(); ++i)
+			{
+				if (UTCTimings[i]->GetSchemeIdUri().Equals(DASHTimingSources::Scheme_urn_mpeg_dash_utc_direct2014))
+				{
+					FTimeValue DirectTime;
+					if (ISO8601::ParseDateTime(DirectTime, UTCTimings[i]->GetValue()) == UEMEDIA_ERROR_OK)
+					{
+						PlayerSessionServices->GetSynchronizedUTCTime()->SetTime(DirectTime);
+						// Remove this timing element. It can not be used indefinitely and we best get rid of it now.
+						MPDRoot->RemoveUTCTimingElement(UTCTimings[i]);
+						--i;
+					}
+					bFoundSupportedElement = true;
+				}
+			}
+			// 2nd pass. Use the more appropriate timing elements.
+			for(int32 i=0; i<UTCTimings.Num(); ++i)
+			{
+				auto GetRequestURL = [MPDDocumentURL](FString ListOfURLs, int32 Index) -> FString
+				{
+					TArray<FString> URLs;
+					ListOfURLs.ParseIntoArrayWS(URLs);
+					if (Index < URLs.Num())
+					{
+						TArray<TSharedPtrTS<const FDashMPD_BaseURLType>> OutBaseURLs;
+						FString URL, RequestHeader;
+						FTimeValue UrlATO;
+						DASHUrlHelpers::BuildAbsoluteElementURL(URL, UrlATO, MPDDocumentURL, OutBaseURLs, URLs[Index]);
+						return URL;
+					}
+					return FString();
+				};
+
+
+				const FString& SchemeIdUri = UTCTimings[i]->GetSchemeIdUri();
+				const FString& Value = UTCTimings[i]->GetValue();
+
+				TSharedPtrTS<FMPDLoadRequestDASH> TimeSyncRequest;
+				// Check for supported schemes. They appear in the manifest in the order the author has prioritized them.
+				if (SchemeIdUri.Equals(DASHTimingSources::Scheme_urn_mpeg_dash_utc_httpxsdate2014))
+				{
+					TimeSyncRequest = MakeSharedTS<FMPDLoadRequestDASH>();
+					TimeSyncRequest->LoadType = FMPDLoadRequestDASH::ELoadType::TimeSync;
+					TimeSyncRequest->URL = GetRequestURL(Value, 0);
+					TimeSyncRequest->CompleteCallback.BindThreadSafeSP(AsShared(), &FPlaylistReaderDASH::Timesync_httpxsdate_Completed);
+				}
+				else if (SchemeIdUri.Equals(DASHTimingSources::Scheme_urn_mpeg_dash_utc_httpiso2014))
+				{
+					TimeSyncRequest = MakeSharedTS<FMPDLoadRequestDASH>();
+					TimeSyncRequest->LoadType = FMPDLoadRequestDASH::ELoadType::TimeSync;
+					TimeSyncRequest->URL = GetRequestURL(Value, 0);
+					TimeSyncRequest->CompleteCallback.BindThreadSafeSP(AsShared(), &FPlaylistReaderDASH::Timesync_httpiso_Completed);
+				}
+				else if (SchemeIdUri.Equals(DASHTimingSources::Scheme_urn_mpeg_dash_utc_httphead2014))
+				{
+					TimeSyncRequest = MakeSharedTS<FMPDLoadRequestDASH>();
+					TimeSyncRequest->LoadType = FMPDLoadRequestDASH::ELoadType::TimeSync;
+					TimeSyncRequest->URL = GetRequestURL(Value, 0);
+					TimeSyncRequest->Verb = TEXT("HEAD");
+					TimeSyncRequest->CompleteCallback.BindThreadSafeSP(AsShared(), &FPlaylistReaderDASH::Timesync_httphead_Completed);
+				}
+				else
+				{
+					// Not supported.
+				}
+
+				// When a request has been set up, enqueue it and leave the loop.
+				if (TimeSyncRequest.IsValid())
+				{
+					bTimeSyncInProgress = true;
+					AttemptedTimesyncDescriptors.Emplace(UTCTimings[i]);
+					EnqueueResourceRequest(MoveTemp(TimeSyncRequest));
+					bFoundSupportedElement = true;
+					break;
+				}
+			}
+			// Warn once when there is no UTCTiming element we support.
+			if (!bFoundSupportedElement && !bWarnedAboutUnsupportedUTCTimingElement)
+			{
+				bWarnedAboutUnsupportedUTCTimingElement = true;
+				if (!bInitialTimeSyncedToMPDDateTimeHeader)
+				{
+					LogMessage(IInfoLog::ELevel::Warning, FString::Printf(TEXT("No supported <UTCTiming> element found in MPD. This could result in playback failures.")));
+				}
+				else
+				{
+					LogMessage(IInfoLog::ELevel::Warning, FString::Printf(TEXT("No supported <UTCTiming> element found in MPD, but time was synchronized to the MPD's HTTP response Date header. This may not be accurate enough.")));
+				}
+			}
+		}
+	}
+}
+
+
+
 void FPlaylistReaderDASH::WorkerThread()
 {
 	LLM_SCOPE(ELLMTag::ElectraPlayer);
@@ -488,7 +645,7 @@ void FPlaylistReaderDASH::ExecutePendingRequests(const FTimeValue& TimeNow)
 				ActiveRequests.Push(Request);
 
 				Request->Request = MakeSharedTS<FHTTPResourceRequest>();
-				Request->Request->URL(Request->URL).Range(Request->Range).Headers(Request->Headers).Object(Request);
+				Request->Request->URL(Request->URL).Verb(Request->Verb).Range(Request->Range).Headers(Request->Headers).Object(Request);
 				// Set up timeouts and also accepted encodings depending on the type of request.
 				SetupRequestTimeouts(Request);
 				Request->Request->AllowStaticQuery(IAdaptiveStreamingPlayerResourceRequest::EPlaybackResourceType::Playlist);
@@ -717,6 +874,20 @@ FErrorDetail FPlaylistReaderDASH::GetXMLResponseString(FString& OutXMLString, FR
 }
 
 
+FErrorDetail FPlaylistReaderDASH::GetResponseString(FString& OutString, FResourceLoadRequestPtr FromRequest)
+{
+	if (FromRequest.IsValid() && FromRequest->Request.IsValid())
+	{
+		TSharedPtrTS<IElectraHttpManager::FReceiveBuffer> ResponseBuffer = FromRequest->Request->GetResponseBuffer();
+		int32 NumResponseBytes = ResponseBuffer->Buffer.Num();
+		const uint8* ResponseBytes = (const uint8*)ResponseBuffer->Buffer.GetLinearReadData();
+		FString UTF8Text(NumResponseBytes, (TCHAR*)FUTF8ToTCHAR((const ANSICHAR*)ResponseBytes, NumResponseBytes).Get());
+		OutString = MoveTemp(UTF8Text);
+	}
+	return FErrorDetail();
+}
+
+
 void FPlaylistReaderDASH::ManifestDownloadCompleted(FResourceLoadRequestPtr Request, bool bSuccess)
 {
 	if (bSuccess)
@@ -739,7 +910,7 @@ void FPlaylistReaderDASH::ManifestDownloadCompleted(FResourceLoadRequestPtr Requ
 		FString EffectiveURL = ConnInfo->EffectiveURL;
 		for(int32 i=0; i<ConnInfo->ResponseHeaders.Num(); ++i)
 		{
-			if (ConnInfo->ResponseHeaders[i].Header == "Date")
+			if (ConnInfo->ResponseHeaders[i].Header.Equals(TEXT("Date")))
 			{
 				// Parse the header
 				FTimeValue DateFromHeader;
@@ -750,8 +921,9 @@ void FPlaylistReaderDASH::ManifestDownloadCompleted(FResourceLoadRequestPtr Requ
 				}
 				// The MPD 'FetchTime' is the time at which the request to get the MPD was made to the server.
 				FetchTime = PlayerSessionServices->GetSynchronizedUTCTime()->MapToSyncTime(ConnInfo->RequestStartTime);
+				bInitialTimeSyncedToMPDDateTimeHeader = true;
 			}
-			else if (ConnInfo->ResponseHeaders[i].Header == "ETag")
+			else if (ConnInfo->ResponseHeaders[i].Header.Equals(TEXT("ETag")))
 			{
 				ETag = ConnInfo->ResponseHeaders[i].Value;
 			}
@@ -780,6 +952,7 @@ void FPlaylistReaderDASH::ManifestDownloadCompleted(FResourceLoadRequestPtr Requ
 					NewManifest->TransformIntoEpicEvent();
 				}
 				Manifest = NewManifest;
+				TriggerTimeSynchronization();
 				PlayerManifest = FManifestDASH::Create(PlayerSessionServices, Manifest);
 
 				// Check if the MPD defines an @minimumUpdatePeriod
@@ -855,9 +1028,10 @@ void FPlaylistReaderDASH::ManifestUpdateDownloadCompleted(FResourceLoadRequestPt
 		FetchTime = PlayerSessionServices->GetSynchronizedUTCTime()->MapToSyncTime(ConnInfo->RequestStartTime);
 		for(int32 i=0; i<ConnInfo->ResponseHeaders.Num(); ++i)
 		{
-			if (ConnInfo->ResponseHeaders[i].Header == "ETag")
+			if (ConnInfo->ResponseHeaders[i].Header.Equals(TEXT("ETag")))
 			{
 				ETag = ConnInfo->ResponseHeaders[i].Value;
+				break;
 			}
 		}
 		MostRecentMPDUpdateTime = FetchTime;
@@ -900,6 +1074,8 @@ void FPlaylistReaderDASH::ManifestUpdateDownloadCompleted(FResourceLoadRequestPt
 					// Also update in the external manifest we handed out to the player.
 					PlayerManifest->UpdateInternalManifest(Manifest);
 
+					TriggerTimeSynchronization();
+
 					// Check if the new MPD also defines an @minimumUpdatePeriod
 					FTimeValue mup = Manifest->GetMinimumUpdatePeriod();
 					if (mup.IsValid() && mup > FTimeValue::GetZero())
@@ -975,6 +1151,74 @@ void FPlaylistReaderDASH::InitialMPDXLinkElementDownloadCompleted(FResourceLoadR
 }
 
 
+void FPlaylistReaderDASH::Timesync_httphead_Completed(FResourceLoadRequestPtr Request, bool bSuccess)
+{
+	if (bSuccess)
+	{
+		const HTTP::FConnectionInfo* ConnInfo = Request->GetConnectionInfo();
+		if (ConnInfo)
+		{
+			for(int32 i=0; i<ConnInfo->ResponseHeaders.Num(); ++i)
+			{
+				if (ConnInfo->ResponseHeaders[i].Header.Equals(TEXT("Date")))
+				{
+					// Parse the header
+					FTimeValue DateFromHeader;
+					UEMediaError DateParseError = RFC7231::ParseDateTime(DateFromHeader, ConnInfo->ResponseHeaders[i].Value);
+					if (DateParseError == UEMEDIA_ERROR_OK && DateFromHeader.IsValid())
+					{
+						PlayerSessionServices->GetSynchronizedUTCTime()->SetTime(DateFromHeader);
+					}
+					break;
+				}
+			}
+		}
+	}
+	// Presently we do not try another time sync method if this one has failed. Clear out what we have attempted so far and be done with it.
+	AttemptedTimesyncDescriptors.Empty();
+	bTimeSyncInProgress = false;
+}
+
+void FPlaylistReaderDASH::Timesync_httpxsdate_Completed(FResourceLoadRequestPtr Request, bool bSuccess)
+{
+	if (bSuccess)
+	{
+		FString Response;
+		if (GetResponseString(Response, Request).IsOK())
+		{
+			FTimeValue DirectTime;
+			// Note: Yes, this is a bit weird, but the xs:dateTime format is actually exactly the same as ISO-8601
+			//       so it is not clear why there is a distinction made in the UTCTiming scheme.
+			//       We keep the different callback handlers here though for completeness sake.
+			if (ISO8601::ParseDateTime(DirectTime, Response) == UEMEDIA_ERROR_OK)
+			{
+				PlayerSessionServices->GetSynchronizedUTCTime()->SetTime(DirectTime);
+			}
+		}
+	}
+	// Presently we do not try another time sync method if this one has failed. Clear out what we have attempted so far and be done with it.
+	AttemptedTimesyncDescriptors.Empty();
+	bTimeSyncInProgress = false;
+}
+
+void FPlaylistReaderDASH::Timesync_httpiso_Completed(FResourceLoadRequestPtr Request, bool bSuccess)
+{
+	if (bSuccess)
+	{
+		FString Response;
+		if (GetResponseString(Response, Request).IsOK())
+		{
+			FTimeValue DirectTime;
+			if (ISO8601::ParseDateTime(DirectTime, Response) == UEMEDIA_ERROR_OK)
+			{
+				PlayerSessionServices->GetSynchronizedUTCTime()->SetTime(DirectTime);
+			}
+		}
+	}
+	// Presently we do not try another time sync method if this one has failed. Clear out what we have attempted so far and be done with it.
+	AttemptedTimesyncDescriptors.Empty();
+	bTimeSyncInProgress = false;
+}
 
 } // namespace Electra
 
diff --git a/Engine/Plugins/Media/ElectraPlayer/Source/ElectraPlayerRuntime/Private/Runtime/Player/DASH/PlaylistReaderDASH_Internal.h b/Engine/Plugins/Media/ElectraPlayer/Source/ElectraPlayerRuntime/Private/Runtime/Player/DASH/PlaylistReaderDASH_Internal.h
--- a/Engine/Plugins/Media/ElectraPlayer/Source/ElectraPlayerRuntime/Private/Runtime/Player/DASH/PlaylistReaderDASH_Internal.h
+++ b/Engine/Plugins/Media/ElectraPlayer/Source/ElectraPlayerRuntime/Private/Runtime/Player/DASH/PlaylistReaderDASH_Internal.h
@@ -33,7 +33,8 @@ struct FMPDLoadRequestDASH : public IHTTPResourceRequestObject
 		XLink_URLQuery,
 		XLink_InitializationSet,
 		Callback,
-		Segment
+		Segment,
+		TimeSync
 	};
 	const TCHAR* const GetRequestTypeName() const
 	{
@@ -49,6 +50,7 @@ struct FMPDLoadRequestDASH : public IHTTPResourceRequestObject
 			case ELoadType::XLink_InitializationSet:return TEXT("remote InitializationSet");
 			case ELoadType::Callback:				return TEXT("Callback");
 			case ELoadType::Segment:				return TEXT("Segment");
+			case ELoadType::TimeSync:				return TEXT("Time sync");
 			default:								return TEXT("<unknown>");
 		}
 	}
@@ -61,6 +63,7 @@ struct FMPDLoadRequestDASH : public IHTTPResourceRequestObject
 	FMPDLoadRequestDASH() : LoadType(ELoadType::MPD) {}
 	FString			URL;	// For xlink requests this could be "urn:mpeg:dash:resolve-to-zero:2013" indicating removal of the element.
 	FString			Range;
+	FString			Verb;
 	TArray<HTTP::FHTTPHeader> Headers;
 	FTimeValue		ExecuteAtUTC;
 
diff --git a/Engine/Plugins/Media/ElectraPlayer/Source/ElectraPlayerRuntime/Private/Runtime/Player/HLS/ManifestBuilderHLS.cpp b/Engine/Plugins/Media/ElectraPlayer/Source/ElectraPlayerRuntime/Private/Runtime/Player/HLS/ManifestBuilderHLS.cpp
--- a/Engine/Plugins/Media/ElectraPlayer/Source/ElectraPlayerRuntime/Private/Runtime/Player/HLS/ManifestBuilderHLS.cpp
+++ b/Engine/Plugins/Media/ElectraPlayer/Source/ElectraPlayerRuntime/Private/Runtime/Player/HLS/ManifestBuilderHLS.cpp
@@ -425,8 +425,8 @@ FErrorDetail FManifestBuilderHLS::SetupVariants(FManifestHLSInternal* Manifest,
 	// Get the timeline asset object.
 	TSharedPtrTS<FTimelineMediaAssetHLS> Asset = Manifest->CurrentMediaAsset;
 
-	TUniquePtr<IURLParser> UrlBuilder(IURLParser::Create());
-	UrlBuilder->ParseURL(Manifest->MasterPlaylistVars.PlaylistLoadRequest.URL);
+	FURL_RFC3986 UrlBuilder;
+	UrlBuilder.Parse(Manifest->MasterPlaylistVars.PlaylistLoadRequest.URL);
 
 	for(const HLSPlaylistParser::FMediaPlaylist& VariantStream : Playlist.GetPlaylists())
 	{
@@ -731,7 +731,7 @@ FErrorDetail FManifestBuilderHLS::SetupVariants(FManifestHLSInternal* Manifest,
 
 				vs->Internal.AdaptationSetUniqueID  = Adapt->UniqueIdentifier;
 				vs->Internal.RepresentationUniqueID = Repr->UniqueIdentifier;
-				vs->Internal.CDN					= UrlBuilder->ResolveWith(vs->URI);;
+				vs->Internal.CDN					= FURL_RFC3986(UrlBuilder).ResolveWith(vs->URI).Get();
 
 				// Check if there is an audio adaptation set already due to an audio rendition group referenced by this variant.
 				if (bHaveAudioGroup)
@@ -826,7 +826,7 @@ FErrorDetail FManifestBuilderHLS::SetupVariants(FManifestHLSInternal* Manifest,
 
 					vs->Internal.AdaptationSetUniqueID  = Adapt->UniqueIdentifier;
 					vs->Internal.RepresentationUniqueID = Repr->UniqueIdentifier;
-					vs->Internal.CDN					= UrlBuilder->ResolveWith(vs->URI);
+					vs->Internal.CDN					= FURL_RFC3986(UrlBuilder).ResolveWith(vs->URI).Get();
 				}
 				else
 				{
@@ -855,7 +855,7 @@ FErrorDetail FManifestBuilderHLS::SetupVariants(FManifestHLSInternal* Manifest,
 									Repr->Bitrate = StreamMetaData.Bandwidth;
 									vs->Internal.AdaptationSetUniqueID  = Adapt->UniqueIdentifier;
 									vs->Internal.RepresentationUniqueID = Repr->UniqueIdentifier;
-									vs->Internal.CDN					= UrlBuilder->ResolveWith(vs->URI);
+									vs->Internal.CDN					= FURL_RFC3986(UrlBuilder).ResolveWith(vs->URI).Get();
 								}
 							}
 						}
@@ -919,7 +919,7 @@ FErrorDetail FManifestBuilderHLS::SetupVariants(FManifestHLSInternal* Manifest,
 
 						vs->Internal.AdaptationSetUniqueID  = Adapt->UniqueIdentifier;
 						vs->Internal.RepresentationUniqueID = Repr->UniqueIdentifier;
-						vs->Internal.CDN					= UrlBuilder->ResolveWith(vs->URI);
+						vs->Internal.CDN					= FURL_RFC3986(UrlBuilder).ResolveWith(vs->URI).Get();
 						Manifest->AudioOnlyStreams.Push(vs);
 
 						Manifest->PlaylistIDMap.Add(vs->Internal.UniqueID, vs);
@@ -1038,14 +1038,14 @@ FErrorDetail FManifestBuilderHLS::GetInitialPlaylistLoadRequests(TArray<FPlaylis
 	}
 
 	// Add the variant playlist to the list of playlists to be fetched.
-	TUniquePtr<IURLParser> UrlBuilder(IURLParser::Create());
-	UrlBuilder->ParseURL(Manifest->MasterPlaylistVars.PlaylistLoadRequest.URL);
+	FURL_RFC3986 UrlBuilder;
+	UrlBuilder.Parse(Manifest->MasterPlaylistVars.PlaylistLoadRequest.URL);
 	FPlaylistLoadRequestHLS& VariantRequest = OutRequests.AddDefaulted_GetRef();
 	check(StartingVariantStream->Internal.LoadState == FManifestHLSInternal::FPlaylistBase::FInternal::ELoadState::NotLoaded);
 	StartingVariantStream->Internal.LoadState = FManifestHLSInternal::FPlaylistBase::FInternal::ELoadState::Pending;
 	VariantRequest.InternalUniqueID 		  = StartingVariantStream->Internal.UniqueID;
 	VariantRequest.RequestedAtTime  		  = PlayerSessionServices->GetSynchronizedUTCTime()->GetTime();
-	VariantRequest.URL  					  = UrlBuilder->ResolveWith(StartingVariantStream->URI);
+	VariantRequest.URL  					  = FURL_RFC3986(UrlBuilder).ResolveWith(StartingVariantStream->URI).Get();
 	VariantRequest.LoadType 				  = FPlaylistLoadRequestHLS::ELoadType::Initial;
 	VariantRequest.AdaptationSetUniqueID	  = StartingVariantStream->Internal.AdaptationSetUniqueID;
 	VariantRequest.RepresentationUniqueID     = StartingVariantStream->Internal.RepresentationUniqueID;
@@ -1077,7 +1077,7 @@ FErrorDetail FManifestBuilderHLS::GetInitialPlaylistLoadRequests(TArray<FPlaylis
 					Rendition->Internal.LoadState   		= FManifestHLSInternal::FPlaylistBase::FInternal::ELoadState::Pending;
 					RenditionRequest.InternalUniqueID   	= Rendition->Internal.UniqueID;
 					RenditionRequest.RequestedAtTime		= PlayerSessionServices->GetSynchronizedUTCTime()->GetTime();
-					RenditionRequest.URL					= UrlBuilder->ResolveWith(Rendition->URI);
+					RenditionRequest.URL					= FURL_RFC3986(UrlBuilder).ResolveWith(Rendition->URI).Get();
 					RenditionRequest.LoadType   			= FPlaylistLoadRequestHLS::ELoadType::Initial;
 					RenditionRequest.AdaptationSetUniqueID  = Rendition->Internal.AdaptationSetUniqueID;
 					RenditionRequest.RepresentationUniqueID = Rendition->Internal.RepresentationUniqueID;
@@ -1176,12 +1176,12 @@ UEMediaError FManifestBuilderHLS::UpdateFailedInitialPlaylistLoadRequest(FPlayli
 			TSharedPtrTS<IAdaptiveStreamSelector> StreamSelector(PlayerSessionServices->GetStreamSelector());
 			StreamSelector->MarkStreamAsUnavailable(VideoVariant->Internal.Blacklisted->AssetIDs);
 
-			TUniquePtr<IURLParser> UrlBuilder(IURLParser::Create());
-			UrlBuilder->ParseURL(Manifest->MasterPlaylistVars.PlaylistLoadRequest.URL);
+			FURL_RFC3986 UrlBuilder;
+			UrlBuilder.Parse(Manifest->MasterPlaylistVars.PlaylistLoadRequest.URL);
 			AlternateVariantStream->Internal.LoadState = FManifestHLSInternal::FPlaylistBase::FInternal::ELoadState::Pending;
 			// Update the changed request parameters.
 			InOutFailedRequest.InternalUniqueID 		  = AlternateVariantStream->Internal.UniqueID;
-			InOutFailedRequest.URL  					  = UrlBuilder->ResolveWith(AlternateVariantStream->URI);
+			InOutFailedRequest.URL  					  = FURL_RFC3986(UrlBuilder).ResolveWith(AlternateVariantStream->URI).Get();
 			InOutFailedRequest.AdaptationSetUniqueID	  = AlternateVariantStream->Internal.AdaptationSetUniqueID;
 			InOutFailedRequest.RepresentationUniqueID     = AlternateVariantStream->Internal.RepresentationUniqueID;
 			InOutFailedRequest.CDN  					  = AlternateVariantStream->Internal.CDN;
@@ -1427,8 +1427,8 @@ FErrorDetail FManifestBuilderHLS::UpdateFromVariantPlaylist(TSharedPtrTS<FManife
 		if (MediaSegmentList.Num())
 		{
 			TArray<FString> SplitResults;
-			TUniquePtr<IURLParser> UrlBuilder(IURLParser::Create());
-			UrlBuilder->ParseURL(SourceRequest.URL);
+			FURL_RFC3986 UrlBuilder;
+			UrlBuilder.Parse(SourceRequest.URL);
 
 			int64 NextSequenceNum      = MediaStream->MediaSequence;
 			int64 NextDiscontinuityNum = MediaStream->DiscontinuitySequence;
@@ -1568,7 +1568,7 @@ FErrorDetail FManifestBuilderHLS::UpdateFromVariantPlaylist(TSharedPtrTS<FManife
 						GetSegmentAttributeString(URI, HLSPlaylistParser::ExtXKey, TEXT("URI"));
 						GetSegmentAttributeString(IV, HLSPlaylistParser::ExtXKey, TEXT("IV"));
 						// Resolve the URL if it is relative.
-						FString LicenseKeyURL = UrlBuilder->ResolveWith(URI);
+						FString LicenseKeyURL = FURL_RFC3986(UrlBuilder).ResolveWith(URI).Get();
 						// Is there a change in attributes?
 						if (LicenseKeyInfo.IsValid() && (LicenseKeyInfo->Method != Method || LicenseKeyInfo->URI != LicenseKeyURL || LicenseKeyInfo->IV != IV))
 						{
@@ -1593,7 +1593,7 @@ FErrorDetail FManifestBuilderHLS::UpdateFromVariantPlaylist(TSharedPtrTS<FManife
 				// Is there a new init segment specified with EXT-X-MAP ?
 				if (GetSegmentAttributeString(AttrValue, HLSPlaylistParser::ExtXMap, TEXT("URI")))
 				{
-					FString InitSegmentURL = UrlBuilder->ResolveWith(AttrValue);
+					FString InitSegmentURL = FURL_RFC3986(UrlBuilder).ResolveWith(AttrValue).Get();
 					InitSegmentInfo = MakeSharedTS<FManifestHLSInternal::FMediaStream::FInitSegmentInfo>();
 					InitSegmentInfo->URI = InitSegmentURL;
 // FIXME: The init segment could/can have a separate EXT-X-KEY applied to it!!
diff --git a/Engine/Plugins/Media/ElectraPlayer/Source/ElectraPlayerRuntime/Private/Runtime/Player/HLS/ManifestHLS.cpp b/Engine/Plugins/Media/ElectraPlayer/Source/ElectraPlayerRuntime/Private/Runtime/Player/HLS/ManifestHLS.cpp
--- a/Engine/Plugins/Media/ElectraPlayer/Source/ElectraPlayerRuntime/Private/Runtime/Player/HLS/ManifestHLS.cpp
+++ b/Engine/Plugins/Media/ElectraPlayer/Source/ElectraPlayerRuntime/Private/Runtime/Player/HLS/ManifestHLS.cpp
@@ -470,8 +470,8 @@ IManifest::FResult FPlayPeriodHLS::FindSegment(TSharedPtrTS<FStreamSegmentReques
 ///////	if (InStream->PlaylistType == FManifestHLSInternal::FMediaStream::EPlaylistType::VOD || InStream->PlaylistType == FManifestHLSInternal::FMediaStream::EPlaylistType::Event)
 
 
-	TUniquePtr<IURLParser> UrlBuilder(IURLParser::Create());
-	UrlBuilder->ParseURL(InPlaylist->Internal.PlaylistLoadRequest.URL);
+	FURL_RFC3986 UrlBuilder;
+	UrlBuilder.Parse(InPlaylist->Internal.PlaylistLoadRequest.URL);
 	TSharedPtrTS<FStreamSegmentRequestHLSfmp4>	Req(new FStreamSegmentRequestHLSfmp4);
 	Req->StreamType 	= StreamType;
 	Req->StreamUniqueID = StreamUniqueID;
@@ -664,7 +664,7 @@ IManifest::FResult FPlayPeriodHLS::FindSegment(TSharedPtrTS<FStreamSegmentReques
 			Req->LocalIndex 			   = SelectedSegmentIndex;
 			Req->bIsPrefetch			   = SegmentList[SelectedSegmentIndex].bIsPrefetch;
 			Req->bIsEOSSegment  		   = false;
-			Req->URL					   = UrlBuilder->ResolveWith(SegmentList[SelectedSegmentIndex].URI);
+			Req->URL					   = FURL_RFC3986(UrlBuilder).ResolveWith(SegmentList[SelectedSegmentIndex].URI).Get();
 			Req->FirstAUTimeOffset  	   = searchTime - SegmentList[SelectedSegmentIndex].AbsoluteDateTime;
 			Req->InitSegmentInfo		   = SegmentList[SelectedSegmentIndex].InitSegmentInfo;
 			Req->LicenseKeyInfo 		   = SegmentList[SelectedSegmentIndex].DRMKeyInfo;
@@ -831,13 +831,13 @@ IManifest::FResult FPlayPeriodHLS::GetStartingSegment(TSharedPtrTS<IStreamSegmen
 
 			VideoPlaylist->Internal.LoadState = FManifestHLSInternal::FPlaylistBase::FInternal::ELoadState::Pending;
 
-			TUniquePtr<IURLParser> UrlBuilder(IURLParser::Create());
-			UrlBuilder->ParseURL(InternalManifest->MasterPlaylistVars.PlaylistLoadRequest.URL);
+			FURL_RFC3986 UrlBuilder;
+			UrlBuilder.Parse(InternalManifest->MasterPlaylistVars.PlaylistLoadRequest.URL);
 			FPlaylistLoadRequestHLS Req;
 			Req.LoadType			   = FPlaylistLoadRequestHLS::ELoadType::First;
 			Req.InternalUniqueID	   = ActiveVideoUniqueID;
 			Req.RequestedAtTime 	   = SessionServices->GetSynchronizedUTCTime()->GetTime();
-			Req.URL 				   = UrlBuilder->ResolveWith(VideoPlaylist->GetURL());
+			Req.URL 				   = FURL_RFC3986(UrlBuilder).ResolveWith(VideoPlaylist->GetURL()).Get();
 			Req.AdaptationSetUniqueID  = VideoPlaylist->Internal.AdaptationSetUniqueID;
 			Req.RepresentationUniqueID = VideoPlaylist->Internal.RepresentationUniqueID;
 			Req.CDN 				   = VideoPlaylist->Internal.CDN;
@@ -850,13 +850,13 @@ IManifest::FResult FPlayPeriodHLS::GetStartingSegment(TSharedPtrTS<IStreamSegmen
 
 			AudioPlaylist->Internal.LoadState = FManifestHLSInternal::FPlaylistBase::FInternal::ELoadState::Pending;
 
-			TUniquePtr<IURLParser> UrlBuilder(IURLParser::Create());
-			UrlBuilder->ParseURL(InternalManifest->MasterPlaylistVars.PlaylistLoadRequest.URL);
+			FURL_RFC3986 UrlBuilder;
+			UrlBuilder.Parse(InternalManifest->MasterPlaylistVars.PlaylistLoadRequest.URL);
 			FPlaylistLoadRequestHLS Req;
 			Req.LoadType			   = FPlaylistLoadRequestHLS::ELoadType::First;
 			Req.InternalUniqueID	   = ActiveAudioUniqueID;
 			Req.RequestedAtTime 	   = SessionServices->GetSynchronizedUTCTime()->GetTime();
-			Req.URL 				   = UrlBuilder->ResolveWith(AudioPlaylist->GetURL());
+			Req.URL 				   = FURL_RFC3986(UrlBuilder).ResolveWith(AudioPlaylist->GetURL()).Get();
 			Req.AdaptationSetUniqueID  = AudioPlaylist->Internal.AdaptationSetUniqueID;
 			Req.RepresentationUniqueID = AudioPlaylist->Internal.RepresentationUniqueID;
 			Req.CDN 				   = AudioPlaylist->Internal.CDN;
@@ -1020,13 +1020,13 @@ IManifest::FResult FPlayPeriodHLS::GetNextOrRetrySegment(TSharedPtrTS<IStreamSeg
 
 			Playlist->Internal.LoadState = FManifestHLSInternal::FPlaylistBase::FInternal::ELoadState::Pending;
 
-			TUniquePtr<IURLParser> UrlBuilder(IURLParser::Create());
-			UrlBuilder->ParseURL(InternalManifest->MasterPlaylistVars.PlaylistLoadRequest.URL);
+			FURL_RFC3986 UrlBuilder;
+			UrlBuilder.Parse(InternalManifest->MasterPlaylistVars.PlaylistLoadRequest.URL);
 			FPlaylistLoadRequestHLS Req;
 			Req.LoadType			   = FPlaylistLoadRequestHLS::ELoadType::First;
 			Req.InternalUniqueID	   = ForStreamID;
 			Req.RequestedAtTime 	   = SessionServices->GetSynchronizedUTCTime()->GetTime();
-			Req.URL 				   = UrlBuilder->ResolveWith(Playlist->GetURL());
+			Req.URL 				   = UrlBuilder.ResolveWith(Playlist->GetURL()).Get();
 			Req.AdaptationSetUniqueID  = Playlist->Internal.AdaptationSetUniqueID;
 			Req.RepresentationUniqueID = Playlist->Internal.RepresentationUniqueID;
 			Req.CDN 				   = Playlist->Internal.CDN;
diff --git a/Engine/Plugins/Media/ElectraPlayer/Source/ElectraPlayerRuntime/Private/Runtime/Utilities/StringHelpers.h b/Engine/Plugins/Media/ElectraPlayer/Source/ElectraPlayerRuntime/Private/Runtime/Utilities/StringHelpers.h
--- a/Engine/Plugins/Media/ElectraPlayer/Source/ElectraPlayerRuntime/Private/Runtime/Utilities/StringHelpers.h
+++ b/Engine/Plugins/Media/ElectraPlayer/Source/ElectraPlayerRuntime/Private/Runtime/Utilities/StringHelpers.h
@@ -105,6 +105,11 @@ namespace Electra
 				return Index >= 0 && Index < StringToIterate.Len();
 			}
 
+			const TCHAR* GetRemainder() const
+			{
+				return &StringToIterate[Index];
+			}
+
 			/** Returns an index to the current element. */
 			int32 GetIndex() const
 			{
diff --git a/Engine/Plugins/Media/ElectraPlayer/Source/ElectraPlayerRuntime/Private/Runtime/Utilities/URLParser.cpp b/Engine/Plugins/Media/ElectraPlayer/Source/ElectraPlayerRuntime/Private/Runtime/Utilities/URLParser.cpp
--- a/Engine/Plugins/Media/ElectraPlayer/Source/ElectraPlayerRuntime/Private/Runtime/Utilities/URLParser.cpp
+++ b/Engine/Plugins/Media/ElectraPlayer/Source/ElectraPlayerRuntime/Private/Runtime/Utilities/URLParser.cpp
@@ -1,104 +1,11 @@
 // Copyright Epic Games, Inc. All Rights Reserved.
 
 #include "Utilities/URLParser.h"
-#include "Utilities/URI.h"
-#include "Utilities/StringHelpers.h"
 
 
 namespace Electra
 {
 
-	class FURLParser : public IURLParser
-	{
-	public:
-		FURLParser();
-		virtual ~FURLParser();
-
-		virtual UEMediaError ParseURL(const FString& URL) override;
-		virtual bool IsAbsoluteURL() const override;
-		virtual FString GetPath() const override;
-		virtual void GetPathComponents(TArray<FString>& OutPathComponents) const override;
-
-		virtual FString ResolveWith(const FString& RelativeURL) override;
-
-		virtual FString GetURL() const
-		{
-			return CurrentURL;
-		}
-
-	private:
-		Utilities::FURI* URLHelper;
-		FString CurrentURL;
-	};
-
-
-	IURLParser* IURLParser::Create()
-	{
-		return new FURLParser;
-	}
-
-
-
-	FURLParser::FURLParser()
-		: URLHelper(nullptr)
-	{
-	}
-
-	FURLParser::~FURLParser()
-	{
-		delete URLHelper;
-	}
-
-	UEMediaError FURLParser::ParseURL(const FString& URL)
-	{
-		delete URLHelper;
-		URLHelper = new Utilities::FURI(Utilities::FURI::Parse(URL));
-		if (!URLHelper->IsValid())
-		{
-			return UEMEDIA_ERROR_FORMAT_ERROR;
-		}
-		CurrentURL = URL;
-		return UEMEDIA_ERROR_OK;
-	}
-
-	bool FURLParser::IsAbsoluteURL() const
-	{
-		if (URLHelper)
-		{
-			return URLHelper->Scheme.Len() > 0;
-		}
-		return false;
-	}
-
-
-	FString FURLParser::GetPath() const
-	{
-		return URLHelper ? URLHelper->Path : FString();
-	}
-
-	void FURLParser::GetPathComponents(TArray<FString>& OutPathComponents) const
-	{
-		OutPathComponents.Empty();
-		if (URLHelper)
-		{
-			static const FString kTextSlash(TEXT("/"));
-			StringHelpers::SplitByDelimiter(OutPathComponents, URLHelper->Path, kTextSlash);
-		}
-	}
-
-	FString FURLParser::ResolveWith(const FString& RelativeURL)
-	{
-		if (URLHelper)
-		{
-			CurrentURL = URLHelper->Resolve(Utilities::FURI::Parse(RelativeURL)).Format();
-		}
-		return CurrentURL;
-	}
-
-	/***************************************************************************************************************************************************/
-	/***************************************************************************************************************************************************/
-	/***************************************************************************************************************************************************/
-
 	static inline void _SwapStrings(FString& a, FString& b)
 	{
 		FString s = MoveTemp(a);
@@ -135,30 +42,64 @@ namespace Electra
 		{
 			return true;
 		}
+		// Not handling URNs here.
+		if (InURL.StartsWith(TEXT("urn:")))
+		{
+			return false;
+		}
+		/*
+		RFC 3986, section 3: Syntax components
+
+				 foo://example.com:8042/over/there?name=ferret#nose
+				 \_/   \______________/\_________/ \_________/ \__/
+				  |           |            |            |        |
+			   scheme     authority       path        query   fragment
+		*/
 		StringHelpers::FStringIterator it(InURL);
-		if (*it != TCHAR('/') && *it != TCHAR('.') && *it != TCHAR('?') && *it != TCHAR('#'))
+		if (InURL.StartsWith(TEXT("//")))
 		{
+			it += 2;
+			if (ParseAuthority(it))
+			{
+				return ParsePathAndQueryFragment(it);
+			}
+			return false;
+		}
+		// Check if the URL begins with a path, query or fragment.
+		else if (IsPathSeparator(*it) || IsQueryOrFragmentSeparator(*it) || *it == TCHAR('.'))
+		{
+			return ParsePathAndQueryFragment(it);
+		}
+		else
+		{
+			// If we get a URL without a scheme we will trip over the colon separating host and port, mistaking it for
+			// the scheme delimiter. For our purposes we only handle URLs that have a scheme when they have an authority!
 			FString Component;
-			Component.Empty();
-			while(it && *it != TCHAR(':') && *it != TCHAR('?') && *it != TCHAR('#') && *it != TCHAR('/'))
+			// Parse out the first component up to the delimiting colon, path, query or fragment.
+			while(it && !IsColonSeparator(*it) && !IsQueryOrFragmentSeparator(*it) && !IsPathSeparator(*it))
 			{
 				Component += *it++;
 			}
-			if (it && *it == TCHAR(':'))
+			// If we found the colon we assume it separates the scheme from the authority.
+			if (it && IsColonSeparator(*it))
 			{
 				++it;
 				if (!it)
 				{
-					// Scheme must be followed by authority or path
+					// Cannot just be a scheme with no authority or path.
 					return false;
 				}
+				// Assume we parsed the scheme.
 				Scheme = Component;
-				if (*it == TCHAR('/'))
+				// Does an absolute path ('/') or an authority ('//') follow?
+				if (IsPathSeparator(*it))
 				{
 					++it;
-					if (it && *it == TCHAR('/'))
+					// Authority?
+					if (it && IsPathSeparator(*it))
 					{
 						++it;
+						// Parse out the authority and if successful continue parsing out the path.
 						if (!ParseAuthority(it))
 						{
 							return false;
@@ -166,6 +107,7 @@ namespace Electra
 					}
 					else
 					{
+						// Unwind the absolute path '/' and parse the path.
 						--it;
 					}
 				}
@@ -176,14 +118,11 @@ namespace Electra
 			}
 			else
 			{
+				// Got either '/', '?' or '#' while scanning for the scheme. Rewind the iterator and handle the entire thing as a path with or without query or fragment.
 				it.Reset();
 				return ParsePathAndQueryFragment(it);
 			}
 		}
-		else
-		{
-			return ParsePathAndQueryFragment(it);
-		}
 		return true;
 	}
 
@@ -229,7 +168,7 @@ namespace Electra
 			}
 			if (Path.Len())
 			{
-				if (Authority.Len() && Path[0] != TCHAR('/'))
+				if (Authority.Len() && !IsPathSeparator(Path[0]))
 				{
 					URL += TEXT("/");
 				}
@@ -363,7 +302,7 @@ namespace Electra
 			Authority += UserInfo;
 			Authority += TEXT("@");
 		}
-		// IPv6 address?
+		// Is the host an IPv6 address that needs to be enclosed in square brackets?
 		int32 DummyPos = 0;
 		if (Host.FindChar(TCHAR(':'), DummyPos))
 		{
@@ -375,6 +314,7 @@ namespace Electra
 		{
 			Authority += Host;
 		}
+		// Need to append a port?
 		if (Port.Len())
 		{
 			Authority += TEXT(":");
@@ -403,21 +343,23 @@ namespace Electra
 	bool FURL_RFC3986::ParseAuthority(StringHelpers::FStringIterator& it)
 	{
 		UserInfo.Empty();
-		FString SubPart;
-		while(it && *it != TCHAR('/') && *it != TCHAR('?') && *it != TCHAR('#'))
+		FString Component;
+		while(it && !IsPathSeparator(*it) && !IsQueryOrFragmentSeparator(*it))
 		{
+			// If there is a user-info delimiter?
 			if (*it == TCHAR('@'))
 			{
-				UserInfo = SubPart;
-				SubPart.Empty();
+				// Yes, what we have gathered so far is the user info. Set it and gather from scratch.
+				UserInfo = Component;
+				Component.Empty();
 			}
 			else
 			{
-				SubPart += *it;
+				Component += *it;
 			}
 			++it;
 		}
-		StringHelpers::FStringIterator SubIt(SubPart);
+		StringHelpers::FStringIterator SubIt(Component);
 		return ParseHostAndPort(SubIt);
 	}
 
@@ -428,7 +370,7 @@ namespace Electra
 			return true;
 		}
 		Host.Empty();
-		// IPv6 adress?
+		// IPv6 adress in [xxx:xxx:xxx] notation?
 		if (*it == TCHAR('['))
 		{
 			++it;
@@ -438,19 +380,19 @@ namespace Electra
 			}
 			if (!it)
 			{
-				// Missing closing ']'
+				// Need to have at least the closing ']'
 				return false;
 			}
 			++it;
 		}
 		else
 		{
-			while(it && *it != TCHAR(':'))
+			while(it && !IsColonSeparator(*it))
 			{
 				Host += *it++;
 			}
 		}
-		if (it && *it == TCHAR(':'))
+		if (it && IsColonSeparator(*it))
 		{
 			++it;
 			Port.Empty();
@@ -466,25 +408,23 @@ namespace Electra
 	{
 		if (it)
 		{
-			if (*it != TCHAR('?') && *it != TCHAR('#'))
+			if (!IsQueryOrFragmentSeparator(*it))
 			{
 				if (!ParsePath(it))
 				{
 					return false;
 				}
 			}
-			if (it && *it == TCHAR('?'))
+			if (it && IsQuerySeparator(*it))
 			{
-				++it;
-				if (!ParseQuery(it))
+				if (!ParseQuery(++it))
 				{
 					return false;
 				}
 			}
-			if (it && *it == TCHAR('#'))
+			if (it && IsFragmentSeparator(*it))
 			{
-				++it;
-				return ParseFragment(it);
+				return ParseFragment(++it);
 			}
 		}
 		return true;
@@ -493,10 +433,11 @@ namespace Electra
 	bool FURL_RFC3986::ParsePath(StringHelpers::FStringIterator& it)
 	{
 		Path.Empty();
-		while(it && *it != TCHAR('?') && *it != TCHAR('#'))
+		while(it && !IsQueryOrFragmentSeparator(*it))
 		{
 			Path += *it++;
 		}
+		// URL decode the path internally.
 		FString Decoded;
 		if (UrlDecode(Decoded, Path))
 		{
@@ -509,7 +450,8 @@ namespace Electra
 	bool FURL_RFC3986::ParseQuery(StringHelpers::FStringIterator& it)
 	{
 		Query.Empty();
-		while(it && *it != TCHAR('#'))
+		// Query is all up to the end or the start of the fragment.
+		while(it && !IsFragmentSeparator(*it))
 		{
 			Query += *it++;
 		}
@@ -518,11 +460,10 @@ namespace Electra
 
 	bool FURL_RFC3986::ParseFragment(StringHelpers::FStringIterator& it)
 	{
-		Fragment.Empty();
-		while(it)
-		{
-			Fragment += *it++;
-		}
+		// Fragment being the last part of the URL extends all the way to the end.
+		Fragment = it.GetRemainder();
+		it.SetToEnd();
+		// URL decode the fragment internally.
 		FString Decoded;
 		if (UrlDecode(Decoded, Fragment))
 		{
@@ -620,20 +561,22 @@ namespace Electra
 	void FURL_RFC3986::GetPathComponents(TArray<FString>& OutPathComponents, const FString& InPath)
 	{
 		TArray<FString> Components;
+		// Split on '/', ignoring resulting empty parts (eg. "a//b" gives ["a", "b"] only instead of ["a", "", "b"] !!
 		InPath.ParseIntoArray(Components, TEXT("/"), true);
 		OutPathComponents.Append(Components);
 	}
 
-	void FURL_RFC3986::ResolveWith(const FString& InChildURL)
+	FURL_RFC3986& FURL_RFC3986::ResolveWith(const FString& InChildURL)
 	{
 		FURL_RFC3986 Other;
 		if (Other.Parse(InChildURL))
 		{
 			ResolveWith(Other);
 		}
+		return *this;
 	}
 
-	void FURL_RFC3986::ResolveAgainst(const FString& InParentURL)
+	FURL_RFC3986& FURL_RFC3986::ResolveAgainst(const FString& InParentURL)
 	{
 		FURL_RFC3986 Parent;
 		if (Parent.Parse(InParentURL))
@@ -641,6 +584,7 @@ namespace Electra
 			Parent.ResolveWith(*this);
 			Swap(Parent);
 		}
+		return *this;
 	}
 	
 	/**
@@ -656,7 +600,7 @@ namespace Electra
 			Port     = Other.Port;
 			Path     = Other.Path;
 			Query    = Other.Query;
-			RemoveDotComponents();
+			RemoveDotSegments();
 		}
 		else
 		{
@@ -667,7 +611,7 @@ namespace Electra
 				Port     = Other.Port;
 				Path     = Other.Path;
 				Query    = Other.Query;
-				RemoveDotComponents();
+				RemoveDotSegments();
 			}
 			else
 			{
@@ -680,15 +624,15 @@ namespace Electra
 				}
 				else
 				{
-					if (Other.Path[0] == TCHAR('/'))
+					if (IsPathSeparator(Other.Path[0]))
 					{
 						Path = Other.Path;
-						RemoveDotComponents();
 					}
 					else
 					{
-						AppendPath(Other.Path);
+						MergePath(Other.Path);
 					}
+					RemoveDotSegments();
 					Query = Other.Query;
 				}
 			}
@@ -696,7 +640,7 @@ namespace Electra
 		Fragment = Other.Fragment;
 	}
 
-	void FURL_RFC3986::RebuildPathFromComponents(const TArray<FString>& Components, bool bAddLeadingSlash, bool bAddTrailingSlash)
+	void FURL_RFC3986::BuildPathFromSegments(const TArray<FString>& Components, bool bAddLeadingSlash, bool bAddTrailingSlash)
 	{
 		Path.Empty();
 		int32 DummyPos = 0;
@@ -704,10 +648,18 @@ namespace Electra
 		{
 			if (i == 0)
 			{
-				if (bAddTrailingSlash)
+				if (bAddLeadingSlash)
 				{
 					Path = TEXT("/");
 				}
+				/*
+					As per RFC 3986 Section 4.2 Relative Reference:
+						A path segment that contains a colon character (e.g., "this:that")
+						cannot be used as the first segment of a relative-path reference, as
+						it would be mistaken for a scheme name.  Such a segment must be
+						preceded by a dot-segment (e.g., "./this:that") to make a relative-
+						path reference.
+				*/
 				else if (Scheme.IsEmpty() && Components[0].FindChar(TCHAR(':'), DummyPos))
 				{
 					Path = TEXT("./");
@@ -725,79 +677,61 @@ namespace Electra
 		}
 	}
 
-	void FURL_RFC3986::AppendPath(const FString& InPathToAppend)
+	void FURL_RFC3986::MergePath(const FString& InPathToMerge)
 	{
-		TArray<FString> Components;
-		TArray<FString> NormalizedComponents;
-		bool bAddLeadingSlash = false;
-		if (!Path.IsEmpty())
+		// RFC 3986 Section 5.2.3. Merge Paths
+		if (Host.Len() && Path.IsEmpty())
 		{
-			GetPathComponents(Components);
-			bool bEndsWithSlash = Path.EndsWith(TEXT("/"));
-			if (!bEndsWithSlash && Components.Num())
-			{
-				Components.Pop();
-			}
-			bAddLeadingSlash = Path[0] == TCHAR('/');
+			Path = TEXT("/");
+			Path += InPathToMerge;
 		}
-		GetPathComponents(Components, InPathToAppend);
-		bAddLeadingSlash = bAddLeadingSlash || (InPathToAppend.Len() && InPathToAppend[0] == TCHAR('/'));
-		bool bHasTrailingSlash = InPathToAppend.Len() && InPathToAppend.EndsWith(TEXT("/"));
-		bool bAddTrailingSlash = false;
-		for(int32 i=0,iMax=Components.Num(); i<iMax; ++i)
+		else
 		{
-			if (Components[i].Equals(TEXT("..")))
+			int32 LastSlashPos = INDEX_NONE;
+			// Is there a '/' somewhere in the current path that's not at the end?
+			if (Path.FindLastChar(TCHAR('/'), LastSlashPos))
 			{
-				bAddTrailingSlash = true;
-				if (NormalizedComponents.Num())
+				if (LastSlashPos != Path.Len()-1)
 				{
-					NormalizedComponents.Pop();
+					Path.LeftInline(LastSlashPos+1);
 				}
+				Path += InPathToMerge;
 			}
-			else if (!Components[i].Equals(TEXT(".")))
-			{
-				bAddTrailingSlash = false;
-				NormalizedComponents.Add(Components[i]);
-			}
-			else 
+			else
 			{
-				bAddTrailingSlash = true;
+				Path = InPathToMerge;
 			}
 		}
-		RebuildPathFromComponents(NormalizedComponents, bAddLeadingSlash, bHasTrailingSlash || bAddTrailingSlash);
 	}
 
-	void FURL_RFC3986::RemoveDotComponents()
+	void FURL_RFC3986::RemoveDotSegments()
 	{
 		if (Path.Len())
 		{
-			bool bHasLeadingSlash = Path[0] == TCHAR('/');
+			TArray<FString> NewSegments;
+			// Remember the current path starting or ending in a slash.
+			bool bHasLeadingSlash = IsPathSeparator(Path[0]);
 			bool bHasTrailingSlash = Path.EndsWith(TEXT("/"));
-			TArray<FString> Components;
-			TArray<FString> NormalizedComponents;
-			GetPathComponents(Components);
-			for(int32 i=0,iMax=Components.Num(); i<iMax; ++i)
+			// Split the path into its segments
+			TArray<FString> Segments;
+			GetPathComponents(Segments);
+			for(int32 i=0; i<Segments.Num(); ++i)
 			{
-				if (Components[i].Equals(TEXT("..")))
+				// For every ".." go up one level.
+				if (Segments[i].Equals(TEXT("..")))
 				{
-					if (NormalizedComponents.Num())
+					if (NewSegments.Num())
 					{
-						if (NormalizedComponents.Last().Equals(TEXT("..")))
-						{
-							NormalizedComponents.Add(Components[i]);
-						}
-						else
-						{
-							NormalizedComponents.Pop();
-						}
+						NewSegments.Pop();
 					}
 				}
-				else if (!Components[i].Equals(TEXT(".")))
+				// Add non- "." segments.
+				else if (!Segments[i].Equals(TEXT(".")))
 				{
-					NormalizedComponents.Add(Components[i]);
+					NewSegments.Add(Segments[i]);
 				}
 			}
-			RebuildPathFromComponents(NormalizedComponents, bHasLeadingSlash, bHasTrailingSlash);
+			BuildPathFromSegments(NewSegments, bHasLeadingSlash, bHasTrailingSlash);
 		}
 	}
 
diff --git a/Engine/Plugins/Media/ElectraPlayer/Source/ElectraPlayerRuntime/Private/Runtime/Utilities/URLParser.h b/Engine/Plugins/Media/ElectraPlayer/Source/ElectraPlayerRuntime/Private/Runtime/Utilities/URLParser.h
--- a/Engine/Plugins/Media/ElectraPlayer/Source/ElectraPlayerRuntime/Private/Runtime/Utilities/URLParser.h
+++ b/Engine/Plugins/Media/ElectraPlayer/Source/ElectraPlayerRuntime/Private/Runtime/Utilities/URLParser.h
@@ -9,34 +9,10 @@
 namespace Electra
 {
 
-	class IURLParser
-	{
-	public:
-		static IURLParser* Create();
-		virtual ~IURLParser() = default;
-
-		virtual UEMediaError ParseURL(const FString& URL) = 0;
-
-		virtual bool IsAbsoluteURL() const = 0;
-
-		virtual FString GetPath() const = 0;
-		virtual void GetPathComponents(TArray<FString>& OutPathComponents) const = 0;
-
-		virtual FString ResolveWith(const FString& RelativeURL) = 0;
-
-		virtual FString GetURL() const = 0;
-
-	protected:
-		IURLParser() = default;
-		IURLParser(const IURLParser&) = delete;
-	};
-
-
-
-
 	class FURL_RFC3986
 	{
 	public:
+		// Parses the given URL.
 		bool Parse(const FString& InURL);
 		// Returns the scheme, if present. Does not include the :// sequence.
 		FString GetScheme() const;
@@ -55,9 +31,9 @@ namespace Electra
 		// Returns the path as individual components. Like GetPath() the components will have escape sequences already decoded.
 		void GetPathComponents(TArray<FString>& OutPathComponents) const;
 		// Resolves a relative URL against this one.
-		void ResolveWith(const FString& InChildURL);
+		FURL_RFC3986& ResolveWith(const FString& InChildURL);
 		// Resolves this URL (which should be relative) against the specified URL.
-		void ResolveAgainst(const FString& InParentURL);
+		FURL_RFC3986& ResolveAgainst(const FString& InParentURL);
 
 		// Appends or prepends additional query parameters.
 		void AddQueryParameters(const FString& InQueryParameters, bool bAppend);
@@ -94,6 +70,16 @@ namespace Electra
 
 		void Empty();
 		void Swap(FURL_RFC3986& Other);
+		inline bool IsColonSeparator(TCHAR c)
+		{ return c == TCHAR(':'); }
+		inline bool IsPathSeparator(TCHAR c)
+		{ return c == TCHAR('/'); }
+		inline bool IsQuerySeparator(TCHAR c)
+		{ return c == TCHAR('?'); }
+		inline bool IsFragmentSeparator(TCHAR c)
+		{ return c == TCHAR('#'); }
+		inline bool IsQueryOrFragmentSeparator(TCHAR c)
+		{ return IsQuerySeparator(c) || IsFragmentSeparator(c); }
 
 		bool ParseAuthority(StringHelpers::FStringIterator& it);
 		bool ParseHostAndPort(StringHelpers::FStringIterator& it);
@@ -104,9 +90,9 @@ namespace Electra
 
 		FString GetAuthority() const;
 		void ResolveWith(const FURL_RFC3986& Other);
-		void RebuildPathFromComponents(const TArray<FString>& Components, bool bAddLeadingSlash, bool bAddTrailingSlash);
-		void AppendPath(const FString& InPathToAppend);
-		void RemoveDotComponents();
+		void BuildPathFromSegments(const TArray<FString>& Components, bool bAddLeadingSlash, bool bAddTrailingSlash);
+		void MergePath(const FString& InPathToMerge);
+		void RemoveDotSegments();
 	};
 
 
