ElectraPlayer: Support for DASH EventStream and InbandEventStream

#rnx

diff --git a/Engine/Plugins/Media/ElectraPlayer/Source/ElectraPlayerRuntime/Private/ElectraPlayer.cpp b/Engine/Plugins/Media/ElectraPlayer/Source/ElectraPlayerRuntime/Private/ElectraPlayer.cpp
--- a/Engine/Plugins/Media/ElectraPlayer/Source/ElectraPlayerRuntime/Private/ElectraPlayer.cpp
+++ b/Engine/Plugins/Media/ElectraPlayer/Source/ElectraPlayerRuntime/Private/ElectraPlayer.cpp
@@ -11,6 +11,7 @@
 #include "Renderer/RendererAudio.h"
 #include "VideoDecoderResourceDelegate.h"
 #include "Crypto/StreamCryptoAES128.h"
+#include "Utilities/Utilities.h"
 
 #include "Async/Async.h"
 
@@ -209,6 +210,11 @@ bool FElectraPlayer::OpenInternal(const FString& Url, const FParamDict & PlayerO
 	NewPlayer->AdaptivePlayer->SetStaticResourceProviderCallback(StaticResourceProvider);
 	NewPlayer->AdaptivePlayer->SetVideoDecoderResourceDelegate(VideoDecoderResourceDelegate);
 
+	// Create a new media player event receiver and register it.
+	MediaPlayerEventReceiver = MakeSharedTS<FAEMSEventReceiver>();
+	MediaPlayerEventReceiver->GetEventReceivedDelegate().BindRaw(this, &FElectraPlayer::OnMediaPlayerEventReceived);
+	NewPlayer->AdaptivePlayer->AddAEMSReceiver(MediaPlayerEventReceiver, TEXT("*"), TEXT(""), IAdaptiveStreamingPlayerAEMSReceiver::EDispatchMode::OnStart);
+
 	if (!NewPlayer->AdaptivePlayer->Initialize(PlayerOptions))
 	{
 		// FIXME: Can Initialize() fail? If so, we'd need to handle the error here.
@@ -301,6 +307,13 @@ void FElectraPlayer::CloseInternal(bool bKillAfterClose)
 	check(Player->AdaptivePlayer.IsValid());
 	if (Player->AdaptivePlayer.IsValid())
 	{
+		if (MediaPlayerEventReceiver.IsValid())
+		{
+			MediaPlayerEventReceiver->GetEventReceivedDelegate().Unbind();
+			Player->AdaptivePlayer->RemoveAEMSReceiver(MediaPlayerEventReceiver, TEXT("*"), TEXT(""), IAdaptiveStreamingPlayerAEMSReceiver::EDispatchMode::OnStart);
+			MediaPlayerEventReceiver.Reset();
+		}
+
 		// Unregister ourselves as the provider for static resources.
 		Player->AdaptivePlayer->SetStaticResourceProviderCallback(nullptr);
 		// Also unregister us from receiving further metric callbacks.
@@ -1010,6 +1023,19 @@ bool FElectraPlayer::SelectTrack(EPlayerTrackType TrackType, int32 TrackIndex)
 	return false;
 }
 
+
+
+void FElectraPlayer::OnMediaPlayerEventReceived(TSharedPtrTS<IAdaptiveStreamingPlayerAEMSEvent> InEvent, IAdaptiveStreamingPlayerAEMSReceiver::EDispatchMode InDispatchMode)
+{
+	const TCHAR* const Origins[] = { TEXT("Playlist"), TEXT("Inband"), TEXT("TimedMetadata"), TEXT("???") };
+	UE_LOG(LogElectraPlayer, Log, TEXT("[%p][%p] %s event %s with \"%s\", \"%s\", \"%s\" PTS @ %.3f for %.3fs"), this, CurrentPlayer.Get(),
+		Origins[Electra::Utils::Min((int32)InEvent->GetOrigin(), (int32)UE_ARRAY_COUNT(Origins)-1)],
+		InDispatchMode==IAdaptiveStreamingPlayerAEMSReceiver::EDispatchMode::OnReceive?TEXT("received"):TEXT("started"),
+		*InEvent->GetSchemeIdUri(), *InEvent->GetValue(), *InEvent->GetID(),
+		InEvent->GetPresentationTime().GetAsSeconds(), InEvent->GetDuration().GetAsSeconds());
+}
+
+
 TSharedPtr<FElectraPlayer::FAnalyticsEvent> FElectraPlayer::CreateAnalyticsEvent(FString InEventName)
 {
 	// Since analytics are popped from the outside only we check if we have accumulated a lot without them having been retrieved.
@@ -1440,6 +1466,7 @@ void FElectraPlayer::HandlePlayerEventDataAvailabilityChange(const Electra::Metr
 void FElectraPlayer::HandlePlayerEventBufferingStart(Electra::Metrics::EBufferingReason BufferingReason)
 {
 	Status = Status | EPlayerStatus::Buffering;
+	DeferredEvents.Enqueue(IElectraPlayerAdapterDelegate::EPlayerEvent::MediaBuffering);
 
 	// Update statistics
 	FScopeLock Lock(&StatisticsLock);
diff --git a/Engine/Plugins/Media/ElectraPlayer/Source/ElectraPlayerRuntime/Private/ElectraPlayer.h b/Engine/Plugins/Media/ElectraPlayer/Source/ElectraPlayerRuntime/Private/ElectraPlayer.h
--- a/Engine/Plugins/Media/ElectraPlayer/Source/ElectraPlayerRuntime/Private/ElectraPlayer.h
+++ b/Engine/Plugins/Media/ElectraPlayer/Source/ElectraPlayerRuntime/Private/ElectraPlayer.h
@@ -119,6 +119,20 @@ public:
 	void NotifyOfOptionChange() override;
 
 private:
+	DECLARE_DELEGATE_TwoParams(FOnMediaPlayerEventReceivedDelegate, TSharedPtrTS<IAdaptiveStreamingPlayerAEMSEvent> /*InEvent*/, IAdaptiveStreamingPlayerAEMSReceiver::EDispatchMode /*InDispatchMode*/);
+	class FAEMSEventReceiver : public IAdaptiveStreamingPlayerAEMSReceiver
+	{
+	public:
+		virtual ~FAEMSEventReceiver() = default;
+		FOnMediaPlayerEventReceivedDelegate& GetEventReceivedDelegate()
+		{ return EventReceivedDelegate; }
+	private:
+		virtual void OnMediaPlayerEventReceived(TSharedPtrTS<IAdaptiveStreamingPlayerAEMSEvent> InEvent, IAdaptiveStreamingPlayerAEMSReceiver::EDispatchMode InDispatchMode) override
+		{ EventReceivedDelegate.ExecuteIfBound(InEvent, InDispatchMode); }
+		FOnMediaPlayerEventReceivedDelegate EventReceivedDelegate;
+	};
+
+
 	struct FPlayerMetricEventBase
 	{
 		enum class EType
@@ -240,6 +254,8 @@ private:
 
 	void PlatformNotifyOfOptionChange();
 
+	void OnMediaPlayerEventReceived(TSharedPtrTS<IAdaptiveStreamingPlayerAEMSEvent> InEvent, IAdaptiveStreamingPlayerAEMSReceiver::EDispatchMode InDispatchMode);
+
 	// Methods from IAdaptiveStreamingPlayerMetrics
 	virtual void ReportOpenSource(const FString& URL) override
 	{ DeferredPlayerEvents.Enqueue(MakeSharedTS<FPlayerMetricEvent_OpenSource>(URL)); }
@@ -340,6 +356,7 @@ private:
 	/** Queued events */
 	TQueue<IElectraPlayerAdapterDelegate::EPlayerEvent>	DeferredEvents;
 	TQueue<TSharedPtrTS<FPlayerMetricEventBase>>	DeferredPlayerEvents;
+	TSharedPtrTS<FAEMSEventReceiver>				MediaPlayerEventReceiver;
 
 	/** The URL of the currently opened media. */
 	FString											MediaUrl;
diff --git a/Engine/Plugins/Media/ElectraPlayer/Source/ElectraPlayerRuntime/Private/Runtime/Demuxer/ParserISO14496-12.cpp b/Engine/Plugins/Media/ElectraPlayer/Source/ElectraPlayerRuntime/Private/Runtime/Demuxer/ParserISO14496-12.cpp
--- a/Engine/Plugins/Media/ElectraPlayer/Source/ElectraPlayerRuntime/Private/Runtime/Demuxer/ParserISO14496-12.cpp
+++ b/Engine/Plugins/Media/ElectraPlayer/Source/ElectraPlayerRuntime/Private/Runtime/Demuxer/ParserISO14496-12.cpp
@@ -265,6 +265,15 @@ private:
 			return CurrentSidxBox;
 		}
 
+		int32 GetNumberOfEMSGBoxes() const
+		{
+			return EMSGBoxes.Num();
+		}
+		FMP4Box* GetEMSGBox(int32 Index) const
+		{
+			return Index < EMSGBoxes.Num() ? EMSGBoxes[Index] : nullptr;
+		}
+
 		FMP4Box* GetCurrentTrackBox() const
 		{
 			return CurrentTrackBox;
@@ -322,6 +331,11 @@ private:
 			CurrentSidxBox = SidxBox;
 		}
 
+		void AddEMSGBox(FMP4Box* EmsgBox)
+		{
+			EMSGBoxes.Emplace(EmsgBox);
+		}
+
 		void SetCurrentTrackBox(FMP4Box* TrackBox)
 		{
 			CurrentTrackBox = TrackBox;
@@ -365,6 +379,8 @@ private:
 		FMP4Box* CurrentHandlerBox = nullptr;		//!< hdlr being parsed at the moment.
 		FMP4Box* CurrentMediaHandlerBox = nullptr;	//!< media handler being parsed at the moment ('vmhd', 'smhd', 'sthd', nmhd')
 
+		TArray<FMP4Box*> EMSGBoxes;
+
 		IPlayerSessionServices* PlayerSession = nullptr;
 		int32 NumTotalBoxesParsed = 0;
 	};
@@ -2240,7 +2256,7 @@ private:
 	 * 'emsg' box.
 	 * ISO/IEC 23009-1:2019 - 5.10.3.3 Event Message Box
 	 */
-	class FMP4BoxEMSG : public FMP4BoxFull
+	class FMP4BoxEMSG : public FMP4BoxFull, public IParserISO14496_12::IEventMessage
 	{
 	public:
 		FMP4BoxEMSG(IParserISO14496_12::FBoxType InBoxType, int64 InBoxSize, int64 InStartOffset, int64 InDataOffset, bool bInIsLeafBox)
@@ -2253,6 +2269,43 @@ private:
 		{
 		}
 
+		virtual int32 GetVersion() const override
+		{ 
+			return Version; 
+		}
+		virtual const FString& GetSchemeIdUri() const override
+		{ 
+			return SchemeIdUri; 
+		}
+		virtual const FString& GetValue() const override
+		{ 
+			return Value; 
+		}
+		virtual uint32 GetTimescale() const override
+		{
+			return Timescale;
+		}
+		virtual uint32 GetPresentationTimeDelta() const override
+		{
+			return PresentationTimeDelta;
+		}
+		virtual uint64 GetPresentationTime() const override
+		{
+			return PresentationTime;
+		}
+		virtual uint32 GetEventDuration() const override
+		{
+			return EventDuration;
+		}
+		virtual uint32 GetID() const override
+		{
+			return ID;
+		}
+		virtual const TArray<uint8>& GetMessageData() const override
+		{
+			return MessageData;
+		}
+
 	private:
 		FMP4BoxEMSG() = delete;
 		FMP4BoxEMSG(const FMP4BoxEMSG&) = delete;
@@ -4094,6 +4147,7 @@ private:
 							break;
 						case FMP4Box::kBox_emsg:
 							NextBox = new FMP4BoxEMSG(BoxType, BoxSize, BoxStartOffset, BoxDataOffset, true);
+							AddEMSGBox(NextBox);
 							break;
 						case FMP4Box::kBox_iods:
 						case FMP4Box::kBox_esds:
@@ -4335,12 +4389,14 @@ private:
 		virtual TMediaOptionalValue<FTimeFraction> GetMovieDuration() const override;
 		virtual int32 GetNumberOfTracks() const override;
 		virtual int32 GetNumberOfSegmentIndices() const override;
+		virtual int32 GetNumberOfEventMessages() const override;
 		virtual int32 GetNumberOfBrands() const override;
 		virtual FBrandType GetBrandByIndex(int32 Index) const override;
 		virtual bool HasBrand(const FBrandType InBrand) const override;
 		virtual const ITrack* GetTrackByIndex(int32 Index) const override;
 		virtual const ITrack* GetTrackByTrackID(int32 TrackID) const override;
 		virtual const ISegmentIndex* GetSegmentIndexByIndex(int32 Index) const override;
+		virtual const IEventMessage* GetEventMessageByIndex(int32 Index) const override;
 
 		virtual TSharedPtr<IAllTrackIterator, ESPMode::ThreadSafe> CreateAllTrackIteratorByFilePos(int64 InFromFilePos) const override;
 
@@ -5986,6 +6042,12 @@ private:
 		return ParsedData ? (ParsedData->GetCurrentSidxBox() ? 1 : 0) : 0;
 	}
 
+	int32 FParserISO14496_12::GetNumberOfEventMessages() const
+	{
+		return ParsedData ? ParsedData->GetNumberOfEMSGBoxes() : 0;
+	}
+
+
 	const FParserISO14496_12::ITrack* FParserISO14496_12::GetTrackByIndex(int32 Index) const
 	{
 		return ParsedTrackInfo ? ParsedTrackInfo->GetTrackByIndex(Index) : nullptr;
@@ -6003,6 +6065,11 @@ private:
 		return Index < GetNumberOfSegmentIndices() ? static_cast<FMP4BoxSIDX*>(ParsedData->GetCurrentSidxBox()) : nullptr;
 	}
 
+	const FParserISO14496_12::IEventMessage* FParserISO14496_12::GetEventMessageByIndex(int32 Index) const
+	{
+		return ParsedData ? static_cast<FMP4BoxEMSG*>(ParsedData->GetEMSGBox(Index)) : nullptr;
+	}
+
 	int32 FParserISO14496_12::GetNumberOfBrands() const
 	{
 		const FMP4BoxFTYP* Box = static_cast<const FMP4BoxFTYP*>(ParsedData->GetCurrentStypBox() ? ParsedData->GetCurrentStypBox() : ParsedData->GetCurrentFtypBox());
diff --git a/Engine/Plugins/Media/ElectraPlayer/Source/ElectraPlayerRuntime/Private/Runtime/Demuxer/ParserISO14496-12.h b/Engine/Plugins/Media/ElectraPlayer/Source/ElectraPlayerRuntime/Private/Runtime/Demuxer/ParserISO14496-12.h
--- a/Engine/Plugins/Media/ElectraPlayer/Source/ElectraPlayerRuntime/Private/Runtime/Demuxer/ParserISO14496-12.h
+++ b/Engine/Plugins/Media/ElectraPlayer/Source/ElectraPlayerRuntime/Private/Runtime/Demuxer/ParserISO14496-12.h
@@ -158,6 +158,8 @@ namespace Electra
 
 		virtual int32 GetNumberOfSegmentIndices() const = 0;
 
+		virtual int32 GetNumberOfEventMessages() const = 0;
+
 		class ITrackIterator
 		{
 		public:
@@ -259,6 +261,22 @@ namespace Electra
 			virtual const FEntry& GetEntry(int32 Index) const = 0;
 		};
 
+		class IEventMessage
+		{
+		public:
+			virtual ~IEventMessage() = default;
+			virtual int32 GetVersion() const = 0;
+			virtual const FString& GetSchemeIdUri() const = 0;
+			virtual const FString& GetValue() const = 0;
+			virtual uint32 GetTimescale() const = 0;
+			virtual uint32 GetPresentationTimeDelta() const = 0;
+			virtual uint64 GetPresentationTime() const = 0;
+			virtual uint32 GetEventDuration() const = 0;
+			virtual uint32 GetID() const = 0;
+			virtual const TArray<uint8>& GetMessageData() const = 0;
+		};
+
+
 		virtual TSharedPtr<IAllTrackIterator, ESPMode::ThreadSafe> CreateAllTrackIteratorByFilePos(int64 InFromFilePos) const = 0;
 
 		virtual const ITrack* GetTrackByIndex(int32 Index) const = 0;
@@ -266,6 +284,7 @@ namespace Electra
 
 		virtual const ISegmentIndex* GetSegmentIndexByIndex(int32 Index) const = 0;
 
+		virtual const IEventMessage* GetEventMessageByIndex(int32 Index) const = 0;
 	};
 
 } // namespace Electra
diff --git a/Engine/Plugins/Media/ElectraPlayer/Source/ElectraPlayerRuntime/Private/Runtime/Player/AdaptiveStreamingPlayer.cpp b/Engine/Plugins/Media/ElectraPlayer/Source/ElectraPlayerRuntime/Private/Runtime/Player/AdaptiveStreamingPlayer.cpp
--- a/Engine/Plugins/Media/ElectraPlayer/Source/ElectraPlayerRuntime/Private/Runtime/Player/AdaptiveStreamingPlayer.cpp
+++ b/Engine/Plugins/Media/ElectraPlayer/Source/ElectraPlayerRuntime/Private/Runtime/Player/AdaptiveStreamingPlayer.cpp
@@ -78,6 +78,9 @@ FAdaptiveStreamingPlayer::FAdaptiveStreamingPlayer(const IAdaptiveStreamingPlaye
 	// Create the render clock
 	RenderClock = MakeSharedTS<FMediaRenderClock>();
 
+	// Create the AEMS handler. This is needed early on for the client to register receivers on.
+	AEMSEventHandler = IAdaptiveStreamingPlayerAEMSHandler::Create();
+
 	bShouldBePaused  = false;
 	bShouldBePlaying = false;
 	bSeekPending	 = false;
@@ -119,6 +122,8 @@ FAdaptiveStreamingPlayer::~FAdaptiveStreamingPlayer()
 	RenderClock.Reset();
 	delete SynchronizedUTCTime;
 
+	delete AEMSEventHandler;
+
 	Electra::RemoveActivePlayerInstance();
 }
 
@@ -221,6 +226,15 @@ void FAdaptiveStreamingPlayer::RemoveMetricsReceiver(IAdaptiveStreamingPlayerMet
 	}
 }
 
+void FAdaptiveStreamingPlayer::AddAEMSReceiver(TWeakPtrTS<IAdaptiveStreamingPlayerAEMSReceiver> InReceiver, FString InForSchemeIdUri, FString InForValue, IAdaptiveStreamingPlayerAEMSReceiver::EDispatchMode InDispatchMode)
+{
+	AEMSEventHandler->AddAEMSReceiver(InReceiver, InForSchemeIdUri, InForValue, InDispatchMode, true);
+}
+
+void FAdaptiveStreamingPlayer::RemoveAEMSReceiver(TWeakPtrTS<IAdaptiveStreamingPlayerAEMSReceiver> InReceiver, FString InForSchemeIdUri, FString InForValue, IAdaptiveStreamingPlayerAEMSReceiver::EDispatchMode InDispatchMode)
+{
+	AEMSEventHandler->RemoveAEMSReceiver(InReceiver, InForSchemeIdUri, InForValue, InDispatchMode);
+}
 
 
 //-----------------------------------------------------------------------------
@@ -809,6 +823,11 @@ TSharedPtrTS<IPlayerEntityCache> FAdaptiveStreamingPlayer::GetEntityCache()
 	return EntityCache;
 }
 
+IAdaptiveStreamingPlayerAEMSHandler* FAdaptiveStreamingPlayer::GetAEMSEventHandler()
+{
+	return AEMSEventHandler;
+}
+
 FParamDict& FAdaptiveStreamingPlayer::GetOptions()
 {
 	return PlayerOptions;
@@ -1328,6 +1347,8 @@ void FAdaptiveStreamingPlayer::WorkerThreadFN()
 			HandlePendingMediaSegmentRequests();
 			// Handle changes in metadata, like timeline changes or track availability.
 			HandleMetadataChanges();
+			// Handle AEMS events
+			HandleAEMSEvents();
 			// Handle buffer level changes
 			HandleNewBufferedData();
 			// Handle new output data (finish prerolling)
@@ -1436,6 +1457,9 @@ void FAdaptiveStreamingPlayer::HandleMetadataChanges()
 						}
 					}
 					
+					// Inform AEMS handler of the period transition.
+					AEMSEventHandler->PlaybackPeriodTransition(UpcomingPeriods[i].ID, UpcomingPeriods[i].TimeRange);
+
 					// Remove it and continue with the next.
 					UpcomingPeriods.RemoveAt(i);
 					++i;
@@ -1446,6 +1470,16 @@ void FAdaptiveStreamingPlayer::HandleMetadataChanges()
 }
 
 
+void FAdaptiveStreamingPlayer::HandleAEMSEvents()
+{
+	if (CurrentState == EPlayerState::eState_Playing)
+	{
+		FTimeValue Current = GetCurrentPlayTime();
+		AEMSEventHandler->ExecutePendingEvents(Current);
+	}
+}
+
+
 //-----------------------------------------------------------------------------
 /**
  * Handles a player session message sent by one of the player sub-systems.
@@ -2434,6 +2468,9 @@ void FAdaptiveStreamingPlayer::AddUpcomingPeriod(TSharedPtrTS<IManifest::IPlayPe
 			// Shift the time range by the current loop state.
 			Next.TimeRange.Start += CurrentLoopState.LoopBasetime;
 			Next.TimeRange.End += CurrentLoopState.LoopBasetime;
+			// Tell the AEMS handler that a new period will be coming up. It needs this information to cut overlapping events.
+			AEMSEventHandler->NewUpcomingPeriod(Next.ID, Next.TimeRange);
+
 			UpcomingPeriods.Emplace(MoveTemp(Next));
 		}
 	}
@@ -2971,6 +3008,9 @@ bool FAdaptiveStreamingPlayer::InternalStartAt(const FSeekParam& NewPosition)
 	}
 	bFirstSegmentRequestIsForLooping = false;
 
+	// Let the AEMS handler know that we are starting up.
+	AEMSEventHandler->PlaybackStartingUp();
+
 	return true;
 }
 
@@ -2985,6 +3025,8 @@ void FAdaptiveStreamingPlayer::InternalSetPlaybackEnded()
 	DispatchEvent(FMetricEvent::ReportPlaybackEnded());
 	// In case a seek back into the stream will happen we reset the first start state.
 	PrerollVars.bIsVeryFirstStart = true;
+	// Let the AEMS handler know as well.
+	AEMSEventHandler->PlaybackReachedEnd(GetCurrentPlayTime());
 }
 
 
@@ -3111,6 +3153,9 @@ void FAdaptiveStreamingPlayer::InternalStop(bool bHoldCurrentFrame)
 	// Flush the renderers again, this time to discard everything the decoders may have emitted while being flushed.
 	AudioRender.Flush();
 	VideoRender.Flush(bHoldCurrentFrame);
+
+	// Flush dynamic events from the AEMS handler.
+	AEMSEventHandler->FlushDynamic();
 }
 
 
@@ -3155,6 +3200,8 @@ void FAdaptiveStreamingPlayer::InternalClose()
 	CurrentLoopState.Reset();
 	NextLoopStates.Clear();
 
+	// Flush all events from the AEMS handler
+	AEMSEventHandler->FlushEverything();
 
 	// Clear error state.
 	LastErrorDetail.Clear();
diff --git a/Engine/Plugins/Media/ElectraPlayer/Source/ElectraPlayerRuntime/Private/Runtime/Player/AdaptiveStreamingPlayer.h b/Engine/Plugins/Media/ElectraPlayer/Source/ElectraPlayerRuntime/Private/Runtime/Player/AdaptiveStreamingPlayer.h
--- a/Engine/Plugins/Media/ElectraPlayer/Source/ElectraPlayerRuntime/Private/Runtime/Player/AdaptiveStreamingPlayer.h
+++ b/Engine/Plugins/Media/ElectraPlayer/Source/ElectraPlayerRuntime/Private/Runtime/Player/AdaptiveStreamingPlayer.h
@@ -10,6 +10,7 @@
 #include "ParameterDictionary.h"
 #include "StreamTypes.h"
 #include "AdaptiveStreamingPlayerMetrics.h"
+#include "AdaptiveStreamingPlayerEvents.h"
 #include "AdaptiveStreamingPlayerResourceRequest.h"
 
 class IOptionPointerValueContainer;
@@ -47,6 +48,13 @@ public:
 	virtual void AddMetricsReceiver(IAdaptiveStreamingPlayerMetrics* InMetricsReceiver) = 0;
 	virtual void RemoveMetricsReceiver(IAdaptiveStreamingPlayerMetrics* InMetricsReceiver) = 0;
 
+	
+	//-------------------------------------------------------------------------
+	// Application Event or Metadata Stream (AEMS) receiver
+	//   Please refer to ISO/IEC 23009-1:2019/DAM 1:2020(E)
+	virtual void AddAEMSReceiver(TWeakPtrTS<IAdaptiveStreamingPlayerAEMSReceiver> InReceiver, FString InForSchemeIdUri, FString InForValue, IAdaptiveStreamingPlayerAEMSReceiver::EDispatchMode InDispatchMode) = 0;
+	virtual void RemoveAEMSReceiver(TWeakPtrTS<IAdaptiveStreamingPlayerAEMSReceiver> InReceiver, FString InForSchemeIdUri, FString InForValue, IAdaptiveStreamingPlayerAEMSReceiver::EDispatchMode InDispatchMode) = 0;
+
 
 	//-------------------------------------------------------------------------
 	// Initializes the player. Options may be passed in to affect behaviour.
diff --git a/Engine/Plugins/Media/ElectraPlayer/Source/ElectraPlayerRuntime/Private/Runtime/Player/AdaptiveStreamingPlayerEvents.cpp b/Engine/Plugins/Media/ElectraPlayer/Source/ElectraPlayerRuntime/Private/Runtime/Player/AdaptiveStreamingPlayerEvents.cpp
--- a/Engine/Plugins/Media/ElectraPlayer/Source/ElectraPlayerRuntime/Private/Runtime/Player/AdaptiveStreamingPlayerEvents.cpp
+++ b/Engine/Plugins/Media/ElectraPlayer/Source/ElectraPlayerRuntime/Private/Runtime/Player/AdaptiveStreamingPlayerEvents.cpp
@@ -1,6 +1,7 @@
 // Copyright Epic Games, Inc. All Rights Reserved.
 
 #include "Player/AdaptiveStreamingPlayerInternal.h"
+#include "Player/DASH/PlayerEventDASH_Internal.h"
 
 namespace Electra
 {
@@ -71,5 +72,342 @@ void FAdaptiveStreamingPlayer::DispatchBufferUtilizationEvent(EStreamType Buffer
 
 
 
+/***************************************************************************************************************************************************/
+/***************************************************************************************************************************************************/
+/***************************************************************************************************************************************************/
+
+namespace ReservedEventSchemeNames
+{
+	static const TArray<FString> IgnoredByWildcard { DASH::Schemes::ManifestEvents::Scheme_urn_mpeg_dash_event_2012, 
+													 DASH::Schemes::ManifestEvents::Scheme_urn_mpeg_dash_event_callback_2015, 
+													 DASH::Schemes::ManifestEvents::Scheme_urn_mpeg_dash_event_ttfn_2016 };
+};
+
+    
+
+
+class FAdaptiveStreamingPlayerAEMSHandler : public IAdaptiveStreamingPlayerAEMSHandler
+{
+public:
+	FAdaptiveStreamingPlayerAEMSHandler();
+	virtual ~FAdaptiveStreamingPlayerAEMSHandler();
+	
+	virtual void ExecutePendingEvents(FTimeValue CurrentPlaybackTime) override;
+	virtual void FlushEverything() override;
+	virtual void FlushDynamic() override;
+	virtual void NewUpcomingPeriod(FString NewPeriodID, FTimeRange NewPeriodTimeRange) override;
+	virtual void PlaybackStartingUp() override;
+	virtual void PlaybackPeriodTransition(FString NewPeriodID, FTimeRange NewPeriodTimeRange) override;
+	virtual void PlaybackReachedEnd(FTimeValue CurrentPlaybackTime) override;
+
+	virtual void AddEvent(TSharedPtrTS<IAdaptiveStreamingPlayerAEMSEvent> InNewEvent, FString InPeriodID, EEventAddMode InAddMode) override;
+	virtual void AddAEMSReceiver(TWeakPtrTS<IAdaptiveStreamingPlayerAEMSReceiver> InReceiver, FString InForSchemeIdUri, FString InForValue, IAdaptiveStreamingPlayerAEMSReceiver::EDispatchMode InDispatchMode, bool bIsAppReceiver) override;
+	virtual void RemoveAEMSReceiver(TWeakPtrTS<IAdaptiveStreamingPlayerAEMSReceiver> InReceiver, FString InForSchemeIdUri, FString InForValue, IAdaptiveStreamingPlayerAEMSReceiver::EDispatchMode InDispatchMode) override;
+
+private:
+	struct FSubscriber
+	{
+		TWeakPtrTS<IAdaptiveStreamingPlayerAEMSReceiver> Receiver;
+		FString ForValue;
+		IAdaptiveStreamingPlayerAEMSReceiver::EDispatchMode DispatchMode;
+		bool bIsAppSubscriber = false;
+		bool bReceivesAll = false;
+	};
+
+	struct FEnqueuedEvent
+	{
+		TArray<TWeakPtrTS<FSubscriber>> Subscribers;
+		FString PeriodID;
+		TSharedPtrTS<IAdaptiveStreamingPlayerAEMSEvent> Event;
+	};
+
+	void DiscardAllPeriodEvents(const FString& ForPeriod);
+	void FireEvent(TSharedPtrTS<FEnqueuedEvent> InEvent, IAdaptiveStreamingPlayerAEMSReceiver::EDispatchMode InDispatchMode);
+	void FireEventsInRange(FTimeRange InRange);
+
+	FCriticalSection Lock;
+	TMultiMap<FString, TSharedPtrTS<FSubscriber>> Subscribers;
+	TArray<TSharedPtrTS<FEnqueuedEvent>> EnqueuedEvents;
+	FString CurrentPeriodID;
+	bool bHaveCurrentPlayPeriod = false;
+	FTimeValue LastPlayPos;
+
+	static const FString AllSchemeWildcard;
+	static const FTimeValue HandlingInterval;
+	static const FTimeValue PlaybackEndThreshold;
+};
+const FString FAdaptiveStreamingPlayerAEMSHandler::AllSchemeWildcard(TEXT("*"));
+const FTimeValue FAdaptiveStreamingPlayerAEMSHandler::HandlingInterval(0.1);
+const FTimeValue FAdaptiveStreamingPlayerAEMSHandler::PlaybackEndThreshold(0.5);
+
+/***************************************************************************************************************************************************/
+
+IAdaptiveStreamingPlayerAEMSHandler* IAdaptiveStreamingPlayerAEMSHandler::Create()
+{
+	return new FAdaptiveStreamingPlayerAEMSHandler;
+}
+
+/***************************************************************************************************************************************************/
+
+FAdaptiveStreamingPlayerAEMSHandler::FAdaptiveStreamingPlayerAEMSHandler()
+{
+}
+
+FAdaptiveStreamingPlayerAEMSHandler::~FAdaptiveStreamingPlayerAEMSHandler()
+{
+	FlushEverything();
+}
+
+void FAdaptiveStreamingPlayerAEMSHandler::ExecutePendingEvents(FTimeValue CurrentPlaybackTime)
+{
+	FTimeValue DeltaTime = CurrentPlaybackTime - LastPlayPos;
+	if (!DeltaTime.IsValid() || DeltaTime < FTimeValue::GetZero())
+	{
+		DeltaTime = HandlingInterval;
+	}
+	// Only update so often. We do not want to iterate through all the events every time
+	// even though we _think_ there aren't that many.
+	if (DeltaTime >= HandlingInterval)
+	{
+		FTimeRange TriggerRange;
+		TriggerRange.Start = CurrentPlaybackTime - HandlingInterval;
+		TriggerRange.End = CurrentPlaybackTime;
+		LastPlayPos = CurrentPlaybackTime;
+		FireEventsInRange(TriggerRange);
+	}
+}
+
+void FAdaptiveStreamingPlayerAEMSHandler::FlushEverything()
+{
+	FScopeLock lock(&Lock);
+	EnqueuedEvents.Empty();
+}
+
+void FAdaptiveStreamingPlayerAEMSHandler::FlushDynamic()
+{
+	FScopeLock lock(&Lock);
+	for(int32 i=0; i<EnqueuedEvents.Num(); ++i)
+	{
+		switch(EnqueuedEvents[i]->Event->GetOrigin())
+		{
+			// Keep only those that came in with the playlist itself.
+			case IAdaptiveStreamingPlayerAEMSEvent::EOrigin::EventStream:
+				break;
+			default:
+				EnqueuedEvents.RemoveAt(i);
+				--i;
+				break;
+		}
+	}
+}
+
+void FAdaptiveStreamingPlayerAEMSHandler::DiscardAllPeriodEvents(const FString& ForPeriod)
+{
+	FScopeLock lock(&Lock);
+	for(int32 i=0; i<EnqueuedEvents.Num(); ++i)
+	{
+		if (EnqueuedEvents[i]->PeriodID.Equals(ForPeriod))
+		{
+			EnqueuedEvents.RemoveAt(i);
+			--i;
+		}
+	}
+}
+
+
+void FAdaptiveStreamingPlayerAEMSHandler::NewUpcomingPeriod(FString NewPeriodID, FTimeRange NewPeriodTimeRange)
+{
+	// No-op for now.
+}
+
+void FAdaptiveStreamingPlayerAEMSHandler::PlaybackStartingUp()
+{
+	CurrentPeriodID.Empty();
+	bHaveCurrentPlayPeriod = false;
+	FlushDynamic();
+	LastPlayPos.SetToInvalid();
+}
+
+void FAdaptiveStreamingPlayerAEMSHandler::PlaybackPeriodTransition(FString NewPeriodID, FTimeRange NewPeriodTimeRange)
+{
+	if (bHaveCurrentPlayPeriod)
+	{
+		// We are transitioning out of the period we know
+		DiscardAllPeriodEvents(CurrentPeriodID);
+	}
+	CurrentPeriodID = NewPeriodID;
+	bHaveCurrentPlayPeriod = true;
+}
+
+void FAdaptiveStreamingPlayerAEMSHandler::PlaybackReachedEnd(FTimeValue CurrentPlaybackTime)
+{
+	FTimeRange TriggerRange;
+	// Fire everything that is within a late and future threshold around the playback position.
+	TriggerRange.Start = CurrentPlaybackTime - HandlingInterval;
+	TriggerRange.End = CurrentPlaybackTime + PlaybackEndThreshold;
+	FireEventsInRange(TriggerRange);
+}
+
+void FAdaptiveStreamingPlayerAEMSHandler::AddEvent(TSharedPtrTS<IAdaptiveStreamingPlayerAEMSEvent> InNewEvent, FString InPeriodID, EEventAddMode InAddMode)
+{
+	// Only events for which there are registered receivers need to be handled. If there is no receiver at the moment but one is added later it will
+	// not receive events that were added before.
+	FScopeLock lock(&Lock);
+
+	TArray<TSharedPtrTS<FSubscriber>> Receivers;
+	Subscribers.MultiFind(InNewEvent->GetSchemeIdUri(), Receivers);
+	Subscribers.MultiFind(AllSchemeWildcard, Receivers);
+	if (Receivers.Num())
+	{
+		// Find an already existing event in the queue.
+		// Note: The PeriodID does not enter into this. Only event specific values.
+		int32 EventIdx = EnqueuedEvents.IndexOfByPredicate([InNewEvent](const TSharedPtrTS<FEnqueuedEvent>& This) {
+			return InNewEvent->GetSchemeIdUri().Equals(This->Event->GetSchemeIdUri()) &&
+				   InNewEvent->GetID().Equals(This->Event->GetID()) &&
+				   (InNewEvent->GetValue().IsEmpty() || InNewEvent->GetValue().Equals(This->Event->GetValue()));
+		});
+
+		TSharedPtrTS<FEnqueuedEvent> Ev;
+		bool bIsNewEvent = false;
+		if (EventIdx == INDEX_NONE)
+		{
+			Ev = MakeSharedTS<FEnqueuedEvent>();
+			bIsNewEvent = true;
+		}
+		else if (InAddMode == EEventAddMode::UpdateIfExists)
+		{
+			Ev = EnqueuedEvents[EventIdx];
+		}
+		if (Ev.IsValid())
+		{
+			Ev->PeriodID = MoveTemp(InPeriodID);
+			Ev->Event = InNewEvent;
+			Ev->Subscribers.Empty();
+			for(int32 i=0; i<Receivers.Num(); ++i)
+			{
+				Ev->Subscribers.Emplace(Receivers[i]);
+			}
+			if (bIsNewEvent)
+			{
+				EnqueuedEvents.Emplace(Ev);
+				lock.Unlock();
+				FireEvent(Ev, IAdaptiveStreamingPlayerAEMSReceiver::EDispatchMode::OnReceive);
+			}
+		}
+	}
+}
+
+void FAdaptiveStreamingPlayerAEMSHandler::AddAEMSReceiver(TWeakPtrTS<IAdaptiveStreamingPlayerAEMSReceiver> InReceiver, FString InForSchemeIdUri, FString InForValue, IAdaptiveStreamingPlayerAEMSReceiver::EDispatchMode InDispatchMode, bool bIsAppReceiver)
+{
+	bool bWildcard = InForSchemeIdUri.Equals(TEXT("*"));
+
+	if (InForSchemeIdUri.Equals(TEXT("urn:mpeg:dash:event:catchall:2020")))
+	{
+		InForSchemeIdUri = AllSchemeWildcard;
+	}
+	if (InForSchemeIdUri.Equals(AllSchemeWildcard))
+	{
+		InForValue.Empty();
+	}
+
+	TSharedPtrTS<FSubscriber> sub = MakeSharedTS<FSubscriber>();
+	sub->Receiver = InReceiver;
+	sub->ForValue = InForValue;
+	sub->DispatchMode = InDispatchMode;
+	sub->bIsAppSubscriber = bIsAppReceiver;
+	sub->bReceivesAll = bWildcard;
+	FScopeLock lock(&Lock);
+	Subscribers.Add(InForSchemeIdUri, sub);
+}
+
+void FAdaptiveStreamingPlayerAEMSHandler::RemoveAEMSReceiver(TWeakPtrTS<IAdaptiveStreamingPlayerAEMSReceiver> InReceiver, FString InForSchemeIdUri, FString InForValue, IAdaptiveStreamingPlayerAEMSReceiver::EDispatchMode InDispatchMode)
+{
+	if (InForSchemeIdUri.Equals(TEXT("urn:mpeg:dash:event:catchall:2020")))
+	{
+		InForSchemeIdUri = AllSchemeWildcard;
+	}
+	if (InForSchemeIdUri.Equals(AllSchemeWildcard))
+	{
+		InForValue.Empty();
+	}
+	FScopeLock lock(&Lock);
+	for(TMultiMap<FString, TSharedPtrTS<FSubscriber>>::TKeyIterator It = Subscribers.CreateKeyIterator(InForSchemeIdUri); It; ++It)
+	{
+		if (It.Value()->Receiver == InReceiver && It.Value()->ForValue == InForValue && It.Value()->DispatchMode == InDispatchMode)
+		{
+			It.RemoveCurrent();
+			break;
+		}
+	}
+}
+
+
+void FAdaptiveStreamingPlayerAEMSHandler::FireEvent(TSharedPtrTS<FEnqueuedEvent> InEvent, IAdaptiveStreamingPlayerAEMSReceiver::EDispatchMode InDispatchMode)
+{
+	for(int32 nSub=0; nSub<InEvent->Subscribers.Num(); ++nSub)
+	{
+		TSharedPtrTS<FSubscriber> Subscriber = InEvent->Subscribers[nSub].Pin();
+		if (Subscriber.IsValid())
+		{
+			TSharedPtrTS<IAdaptiveStreamingPlayerAEMSReceiver> Receiver = Subscriber->Receiver.Pin();
+			if (Receiver.IsValid())
+			{
+				if (Subscriber->DispatchMode == InDispatchMode &&
+					(Subscriber->ForValue.IsEmpty() || Subscriber->ForValue.Equals(InEvent->Event->GetValue())))
+				{
+					// Check if the subscriber is the application and it wants to listen to everything.
+					if (Subscriber->bIsAppSubscriber && Subscriber->bReceivesAll)
+					{
+						// Well, internal events meant for the player only we will not send to it.
+						// The app can still subscribe explicitly to the event type if it needs them but
+						// we do not send them on our own.
+						if (ReservedEventSchemeNames::IgnoredByWildcard.Contains(InEvent->Event->GetSchemeIdUri()))
+						{
+							continue;
+						}
+					}
+					Receiver->OnMediaPlayerEventReceived(InEvent->Event, InDispatchMode);
+				}
+			}
+		}
+	}
+}
+
+void FAdaptiveStreamingPlayerAEMSHandler::FireEventsInRange(FTimeRange InRange)
+{
+	TArray<TSharedPtrTS<FEnqueuedEvent>> EventsToFire;
+	FScopeLock lock(&Lock);
+	// Events that are to be fired get moved into a temporary list so we can
+	// trigger them outside the lock.
+	for(int32 i=0; i<EnqueuedEvents.Num(); ++i)
+	{
+		FTimeValue evS = EnqueuedEvents[i]->Event->GetPresentationTime();
+		FTimeValue evD = EnqueuedEvents[i]->Event->GetDuration();
+		FTimeValue evE = evS + (evD.IsValid() ? evD : FTimeValue::GetZero());
+		
+		// Event already expired?
+		if (evE < InRange.Start)
+		{
+			EnqueuedEvents.RemoveAt(i);
+			--i;
+		}
+		else if (evS >= InRange.Start && evS <= InRange.End)
+		{
+			EventsToFire.Emplace(EnqueuedEvents[i]);
+			EnqueuedEvents.RemoveAt(i);
+			--i;
+		}
+	}
+	lock.Unlock();
+	// Fire the events.
+	for(int32 i=0; i<EventsToFire.Num(); ++i)
+	{
+		FireEvent(EventsToFire[i], IAdaptiveStreamingPlayerAEMSReceiver::EDispatchMode::OnStart);
+	}
+}
+
+
+
+
 } // namespace Electra
 
diff --git a/Engine/Plugins/Media/ElectraPlayer/Source/ElectraPlayerRuntime/Private/Runtime/Player/AdaptiveStreamingPlayerEvents.h b/Engine/Plugins/Media/ElectraPlayer/Source/ElectraPlayerRuntime/Private/Runtime/Player/AdaptiveStreamingPlayerEvents.h
new file mode 100644
--- /dev/null
+++ b/Engine/Plugins/Media/ElectraPlayer/Source/ElectraPlayerRuntime/Private/Runtime/Player/AdaptiveStreamingPlayerEvents.h
@@ -0,0 +1,112 @@
+// Copyright Epic Games, Inc. All Rights Reserved.
+
+#pragma once
+
+#include "PlayerCore.h"
+
+namespace Electra
+{
+
+class IAdaptiveStreamingPlayerAEMSEvent
+{
+public:
+	enum class EOrigin
+	{
+		EventStream,
+		InbandEventStream,
+		TimedMetadata
+	};
+
+	virtual ~IAdaptiveStreamingPlayerAEMSEvent() = default;
+
+	// Returns where the event originated from.
+	virtual EOrigin GetOrigin() const = 0;
+
+	// Returns the event's scheme id URI
+	virtual FString GetSchemeIdUri() const = 0;
+
+	// Returns the event's value.
+	virtual FString GetValue() const = 0;
+
+	// Returns the event's ID.
+	virtual FString GetID() const = 0;
+
+	// Returns the event's presentation time.
+	virtual FTimeValue GetPresentationTime() const = 0;
+
+	// Returns the event's duration. If the event has no defined duration the return value will be invalid.
+	virtual FTimeValue GetDuration() const = 0;
+
+	// Returns the event's message data.
+	virtual const TArray<uint8>& GetMessageData() const = 0;
+};
+
+
+
+/**
+ * Client side event receiver.
+ * 
+ * As per ISO/IEC 23009-1:2019/DAM 1 event receivers subscribing to "on_receive" will receive the event the
+ * instant it is received by the player. Handling of the event is the sole responsibility of the receiver.
+ * Events will be received at any time and the same event can be received multiple times. Events will also
+ * be received in no particular order and events may also already be outdated when received.
+ * This happens with events contained inside the playlist every time the playlist is updated on the server
+ * and reloaded by the client.
+ */
+class IAdaptiveStreamingPlayerAEMSReceiver
+{
+public:
+	virtual ~IAdaptiveStreamingPlayerAEMSReceiver() = default;
+
+	enum class EDispatchMode
+	{
+		OnReceive,
+		OnStart
+	};
+
+	virtual void OnMediaPlayerEventReceived(TSharedPtrTS<IAdaptiveStreamingPlayerAEMSEvent> InEvent, EDispatchMode InDispatchMode) = 0;
+};
+
+
+
+
+class IAdaptiveStreamingPlayerAEMSHandler
+{
+public:
+	static IAdaptiveStreamingPlayerAEMSHandler* Create();
+	virtual ~IAdaptiveStreamingPlayerAEMSHandler() = default;
+	
+	// Call periodically to dispatch pending events.
+	virtual void ExecutePendingEvents(FTimeValue CurrentPlaybackTime) = 0;
+	// Unconditionally flushes all events from the queue.
+	virtual void FlushEverything() = 0;
+	// Informs that playback has moved to a new position in either direction. This may discard dynamic events from streams but keep period events from the playlist.
+	virtual void FlushDynamic() = 0;
+	// Informs that a new playback period is coming up. This is to truncate period events that would extend over their originating period.
+	virtual void NewUpcomingPeriod(FString NewPeriodID, FTimeRange NewPeriodTimeRange) = 0;
+	// Informs that playback is starting up.
+	virtual void PlaybackStartingUp() = 0;
+	// Informs that playback has moved forward into another period. Events of an earlier period that are outside its associated periods time range will not fire.
+	virtual void PlaybackPeriodTransition(FString NewPeriodID, FTimeRange NewPeriodTimeRange) = 0;
+	// Informs that playback has ended.
+	virtual void PlaybackReachedEnd(FTimeValue CurrentPlaybackTime) = 0;
+
+	enum class EEventAddMode
+	{
+		AddIfNotExists,
+		UpdateIfExists
+	};
+	// Enqueues an event. Checks if the new event ID (GetID()) for the given scheme/value has already been added earlier
+	// and either adds, updates or ignores the event based on the given InAddMode.
+	virtual void AddEvent(TSharedPtrTS<IAdaptiveStreamingPlayerAEMSEvent> InNewEvent, FString InPeriodID, EEventAddMode InAddMode) = 0;
+
+
+	// Adds/removes event receivers.
+	virtual void AddAEMSReceiver(TWeakPtrTS<IAdaptiveStreamingPlayerAEMSReceiver> InReceiver, FString InForSchemeIdUri, FString InForValue, IAdaptiveStreamingPlayerAEMSReceiver::EDispatchMode InDispatchMode, bool bIsAppReceiver) = 0;
+	virtual void RemoveAEMSReceiver(TWeakPtrTS<IAdaptiveStreamingPlayerAEMSReceiver> InReceiver, FString InForSchemeIdUri, FString InForValue, IAdaptiveStreamingPlayerAEMSReceiver::EDispatchMode InDispatchMode) = 0;
+};
+
+
+} // namespace Electra
+
+
diff --git a/Engine/Plugins/Media/ElectraPlayer/Source/ElectraPlayerRuntime/Private/Runtime/Player/AdaptiveStreamingPlayerInternal.h b/Engine/Plugins/Media/ElectraPlayer/Source/ElectraPlayerRuntime/Private/Runtime/Player/AdaptiveStreamingPlayerInternal.h
--- a/Engine/Plugins/Media/ElectraPlayer/Source/ElectraPlayerRuntime/Private/Runtime/Player/AdaptiveStreamingPlayerInternal.h
+++ b/Engine/Plugins/Media/ElectraPlayer/Source/ElectraPlayerRuntime/Private/Runtime/Player/AdaptiveStreamingPlayerInternal.h
@@ -652,7 +652,7 @@ public:
 	virtual void AddMetricsReceiver(IAdaptiveStreamingPlayerMetrics* InMetricsReceiver) override;
 	virtual void RemoveMetricsReceiver(IAdaptiveStreamingPlayerMetrics* InMetricsReceiver) override;
 
-	// Matrics receiver accessor for event dispatcher thread.
+	// Metrics receiver accessor for event dispatcher thread.
 	void LockMetricsReceivers()
 	{
 		MetricListenerCriticalSection.Lock();
@@ -667,6 +667,9 @@ public:
 	}
 
 
+	virtual void AddAEMSReceiver(TWeakPtrTS<IAdaptiveStreamingPlayerAEMSReceiver> InReceiver, FString InForSchemeIdUri, FString InForValue, IAdaptiveStreamingPlayerAEMSReceiver::EDispatchMode InDispatchMode) override;
+	virtual void RemoveAEMSReceiver(TWeakPtrTS<IAdaptiveStreamingPlayerAEMSReceiver> InReceiver, FString InForSchemeIdUri, FString InForValue, IAdaptiveStreamingPlayerAEMSReceiver::EDispatchMode InDispatchMode) override;
+
 	virtual bool Initialize(const FParamDict& Options) override;
 
 	virtual void LoadManifest(const FString& manifestURL) override;
@@ -721,6 +724,7 @@ private:
 	virtual IPlayerStreamFilter* GetStreamFilter() override;
 	virtual	TSharedPtrTS<IPlaylistReader> GetManifestReader() override;
 	virtual TSharedPtrTS<IPlayerEntityCache> GetEntityCache() override;
+	virtual IAdaptiveStreamingPlayerAEMSHandler* GetAEMSEventHandler() override;
 	virtual FParamDict& GetOptions() override;
 
 	// Methods from IPlayerStreamFilter
@@ -1290,6 +1294,7 @@ private:
 	void HandleDeselectedBuffers();
 	void HandleDecoderChanges();
 	void HandleMetadataChanges();
+	void HandleAEMSEvents();
 
 	void CheckForStreamEnd();
 
@@ -1335,6 +1340,7 @@ private:
 	FPlaybackState														PlaybackState;
 	FStreamPreferences													StreamPreferences;
 	ISynchronizedUTCTime*												SynchronizedUTCTime;
+	IAdaptiveStreamingPlayerAEMSHandler*								AEMSEventHandler;
 
 	TSharedPtrTS<FMediaRenderClock>										RenderClock;
 
diff --git a/Engine/Plugins/Media/ElectraPlayer/Source/ElectraPlayerRuntime/Private/Runtime/Player/DASH/MPDElementsDASH.h b/Engine/Plugins/Media/ElectraPlayer/Source/ElectraPlayerRuntime/Private/Runtime/Player/DASH/MPDElementsDASH.h
--- a/Engine/Plugins/Media/ElectraPlayer/Source/ElectraPlayerRuntime/Private/Runtime/Player/DASH/MPDElementsDASH.h
+++ b/Engine/Plugins/Media/ElectraPlayer/Source/ElectraPlayerRuntime/Private/Runtime/Player/DASH/MPDElementsDASH.h
@@ -1581,6 +1581,11 @@ public:
 	const TMediaOptionalValue<bool>& GetBitstreamSwitching() const			{ return bBitstreamSwitching; }
 	virtual FXLink& GetXLink() override										{ return XLink; }
 
+	// Methods to manipulate the presentation type.
+	void SetDuration(const FTimeValue& NewDuration)
+	{
+		Duration = NewDuration;
+	}
 
 private:
 	virtual bool ProcessElement(FManifestParserDASH* Builder, const TCHAR* ElementName, const TCHAR* ElementData, int32 XmlFileLineNumber) override;
diff --git a/Engine/Plugins/Media/ElectraPlayer/Source/ElectraPlayerRuntime/Private/Runtime/Player/DASH/ManifestBuilderDASH.cpp b/Engine/Plugins/Media/ElectraPlayer/Source/ElectraPlayerRuntime/Private/Runtime/Player/DASH/ManifestBuilderDASH.cpp
--- a/Engine/Plugins/Media/ElectraPlayer/Source/ElectraPlayerRuntime/Private/Runtime/Player/DASH/ManifestBuilderDASH.cpp
+++ b/Engine/Plugins/Media/ElectraPlayer/Source/ElectraPlayerRuntime/Private/Runtime/Player/DASH/ManifestBuilderDASH.cpp
@@ -9,6 +9,7 @@
 #include "Player/PlayerEntityCache.h"
 #include "Player/DASH/OptionKeynamesDASH.h"
 #include "Player/AdaptivePlayerOptionKeynames.h"
+#include "Player/DASH/PlayerEventDASH.h"
 
 #include "ManifestBuilderDASH.h"
 #include "ManifestParserDASH.h"
@@ -689,7 +690,7 @@ public:
 	FManifestBuilderDASH(IPlayerSessionServices* InPlayerSessionServices);
 	virtual ~FManifestBuilderDASH() = default;
 	
-	virtual FErrorDetail BuildFromMPD(TSharedPtrTS<FManifestDASHInternal>& OutMPD, TCHAR* InOutMPDXML, const FString& EffectiveURL, const FTimeValue& FetchTime, const FString& ETag) override;
+	virtual FErrorDetail BuildFromMPD(TSharedPtrTS<FManifestDASHInternal>& OutMPD, TCHAR* InOutMPDXML, const FString& EffectiveURL, const FString& ETag) override;
 
 private:
 	ELECTRA_IMPL_DEFAULT_ERROR_METHODS(DASHMPDBuilder);
@@ -711,7 +712,7 @@ FManifestBuilderDASH::FManifestBuilderDASH(IPlayerSessionServices* InPlayerSessi
 {
 }
 
-FErrorDetail FManifestBuilderDASH::BuildFromMPD(TSharedPtrTS<FManifestDASHInternal>& OutMPD, TCHAR* InOutMPDXML, const FString& EffectiveURL, const FTimeValue& FetchTime, const FString& ETag)
+FErrorDetail FManifestBuilderDASH::BuildFromMPD(TSharedPtrTS<FManifestDASHInternal>& OutMPD, TCHAR* InOutMPDXML, const FString& EffectiveURL, const FString& ETag)
 {
 	TSharedPtrTS<FDashMPD_MPDType> NewMPD;
 	TArray<TWeakPtrTS<IDashMPDElement>> XLinkElements;
@@ -728,7 +729,6 @@ FErrorDetail FManifestBuilderDASH::BuildFromMPD(TSharedPtrTS<FManifestDASHIntern
 	}
 	NewMPD = RootEntities.MPDs[0];
 	NewMPD->SetDocumentURL(EffectiveURL);
-	NewMPD->SetFetchTime(FetchTime);
 	NewMPD->SetETag(ETag);
 
 	// Check if this MPD uses a profile we can handle. (If our array is empty we claim to handle anything)
@@ -1177,6 +1177,13 @@ FErrorDetail FManifestDASHInternal::BuildAfterInitialRemoteElementDownload()
 		Periods.Emplace(MoveTemp(p));
 	}
 
+	// As per ISO/IEC 23009-1:2019/DAM 1 "Change 2: Event Stream and Timed Metadata Processing" Section A.11.11 Detailed processing
+	// all events from MPD EventStreams are to be collected right now.
+	for(int32 i=0; i<Periods.Num(); ++i)
+	{
+		SendEventsFromAllPeriodEventStreams(Periods[i]);
+	}
+
 	return Error;
 }
 
@@ -1542,6 +1549,63 @@ void FManifestDASHInternal::PreparePeriodAdaptationSets(TSharedPtrTS<FPeriod> Pe
 
 
 
+void FManifestDASHInternal::SendEventsFromAllPeriodEventStreams(TSharedPtrTS<FPeriod> InPeriod)
+{
+	check(InPeriod.IsValid());
+	TSharedPtrTS<FDashMPD_PeriodType> MPDPeriod = InPeriod->Period.Pin();
+	if (MPDPeriod.IsValid())
+	{
+		FTimeValue PeriodTimeOffset = GetAnchorTime() + InPeriod->GetStart();
+		// Iterate of all the period event streams. We expect them to have been xlink resolved already.
+		const TArray<TSharedPtrTS<FDashMPD_EventStreamType>>& MPDEventStreams = MPDPeriod->GetEventStreams();
+		for(int32 nEvS=0; nEvS<MPDEventStreams.Num(); ++nEvS)
+		{
+			const TSharedPtrTS<FDashMPD_EventStreamType>& EvS = MPDEventStreams[nEvS];
+			uint32 Timescale = EvS->GetTimescale().GetWithDefault(1);
+			int64 PTO = (int64) EvS->GetPresentationTimeOffset().GetWithDefault(0);
+			
+			const TArray<TSharedPtrTS<FDashMPD_EventType>>& Events = EvS->GetEvents();
+			for(int32 i=0; i<Events.Num(); ++i)
+			{
+				const TSharedPtrTS<FDashMPD_EventType>& Event = Events[i];
+				TSharedPtrTS<DASH::FPlayerEvent> NewEvent = MakeSharedTS<DASH::FPlayerEvent>();
+				NewEvent->SetOrigin(IAdaptiveStreamingPlayerAEMSEvent::EOrigin::EventStream);
+				NewEvent->SetSchemeIdUri(EvS->GetSchemeIdUri());
+				NewEvent->SetValue(EvS->GetValue());
+				if (Event->GetID().IsSet())
+				{
+					NewEvent->SetID(LexToString(Event->GetID().Value()));
+				}
+				FTimeValue PTS((int64)Event->GetPresentationTime() - PTO, Timescale);
+				PTS += PeriodTimeOffset;
+				NewEvent->SetPresentationTime(PTS);
+				if (Event->GetDuration().IsSet())
+				{
+					NewEvent->SetDuration(FTimeValue((int64)Event->GetDuration().Value(), Timescale));
+				}
+				FString Data = Event->GetData();
+				if (Data.IsEmpty())
+				{
+					Data = Event->GetMessageData();
+				}
+				// If the data is still empty then there could have been an entire XML element tree which we have parsed.
+				// We could reconstruct the XML from it here and pass that along, but for now we do not do this.
+				/*
+					if (Data.IsEmpty())
+					{
+					}
+				*/
+				NewEvent->SetMessageData(Data, Event->GetContentEncoding().Equals(TEXT("base64")));
+				NewEvent->SetPeriodID(InPeriod->GetUniqueIdentifier());
+				// Add the event to the handler.
+				PlayerSessionServices->GetAEMSEventHandler()->AddEvent(NewEvent, InPeriod->GetID(), IAdaptiveStreamingPlayerAEMSHandler::EEventAddMode::UpdateIfExists);
+			}
+		}
+	}
+}
+
+
+
 /**
  * Resolves an xlink request made by the initial load of the MPD.
  */
@@ -1834,7 +1898,10 @@ FTimeValue FManifestDASHInternal::GetLastPeriodEndTime() const
 	}
 
 	FTimeValue End = ast + MediaPresentationDuration;
-	if (MPDRoot->GetMinimumUpdatePeriod().IsValid())
+	// If MUP is zero then it is expected that InbandEventStream is used to signal when to update the MPD.
+	// Since the MPD will not update through MUP in this case we need to return that the period goes up
+	// to mediaPresentationDuration, in this case infinity.
+	if (MPDRoot->GetMinimumUpdatePeriod().IsValid() && MPDRoot->GetMinimumUpdatePeriod() > FTimeValue::GetZero())
 	{
 		FTimeValue CheckTime = FetchTime + MPDRoot->GetMinimumUpdatePeriod();
 		return CheckTime < End ? CheckTime : End;
@@ -1984,6 +2051,32 @@ FTimeValue FManifestDASHInternal::GetDuration() const
 }
 
 
+void FManifestDASHInternal::EndPresentationAt(const FTimeValue& EndsAt, const FString& InPeriod)
+{
+	TSharedPtrTS<FPeriod> Period;
+	if (InPeriod.IsEmpty() && Periods.Num())
+	{
+		Period = Periods.Last();
+	}
+	else
+	{
+		Period = GetPeriodByUniqueID(InPeriod);
+	}
+	if (Period.IsValid())
+	{
+		Period->EndPresentationAt(EndsAt - GetAnchorTime());
+	}
+	// Updates no longer expected.
+	MPDRoot->SetMinimumUpdatePeriod(FTimeValue::GetInvalid());
+	if (Periods.Num())
+	{
+		FTimeValue NewDuration = EndsAt - (GetAnchorTime() + Periods[0]->GetStart());
+		MPDRoot->SetMediaPresentationDuration(NewDuration);
+	}
+}
+
+
+
 FTimeRange FManifestDASHInternal::GetPlayTimesFromURI() const
 {
 	FTimeRange FromTo;
diff --git a/Engine/Plugins/Media/ElectraPlayer/Source/ElectraPlayerRuntime/Private/Runtime/Player/DASH/ManifestBuilderDASH.h b/Engine/Plugins/Media/ElectraPlayer/Source/ElectraPlayerRuntime/Private/Runtime/Player/DASH/ManifestBuilderDASH.h
--- a/Engine/Plugins/Media/ElectraPlayer/Source/ElectraPlayerRuntime/Private/Runtime/Player/DASH/ManifestBuilderDASH.h
+++ b/Engine/Plugins/Media/ElectraPlayer/Source/ElectraPlayerRuntime/Private/Runtime/Player/DASH/ManifestBuilderDASH.h
@@ -27,12 +27,11 @@ public:
 	 * @param OutMPD
 	 * @param InOutMPDXML
 	 * @param EffectiveURL
-	 * @param FetchTime
 	 * @param ETag
 	 *
 	 * @return
 	 */
-	virtual FErrorDetail BuildFromMPD(TSharedPtrTS<FManifestDASHInternal>& OutMPD, TCHAR* InOutMPDXML, const FString& EffectiveURL, const FTimeValue& FetchTime, const FString& ETag) = 0;
+	virtual FErrorDetail BuildFromMPD(TSharedPtrTS<FManifestDASHInternal>& OutMPD, TCHAR* InOutMPDXML, const FString& EffectiveURL, const FString& ETag) = 0;
 };
 
 
diff --git a/Engine/Plugins/Media/ElectraPlayer/Source/ElectraPlayerRuntime/Private/Runtime/Player/DASH/ManifestDASH.cpp b/Engine/Plugins/Media/ElectraPlayer/Source/ElectraPlayerRuntime/Private/Runtime/Player/DASH/ManifestDASH.cpp
--- a/Engine/Plugins/Media/ElectraPlayer/Source/ElectraPlayerRuntime/Private/Runtime/Player/DASH/ManifestDASH.cpp
+++ b/Engine/Plugins/Media/ElectraPlayer/Source/ElectraPlayerRuntime/Private/Runtime/Player/DASH/ManifestDASH.cpp
@@ -256,6 +256,17 @@ void FManifestDASH::GetTrackMetadata(TArray<FTrackMetadata>& OutMetadata, EStrea
 void FManifestDASH::UpdateDynamicRefetchCounter()
 {
 	++CurrentPeriodAndAdaptationXLinkResolveID;
+
+	// Since we don't know which streams will be used now we have to let the manifest reader know
+	// that currently no stream is active that is providing inband events.
+	TSharedPtrTS<IPlaylistReader> ManifestReader = PlayerSessionServices->GetManifestReader();
+	IPlaylistReaderDASH* Reader = static_cast<IPlaylistReaderDASH*>(ManifestReader.Get());
+	if (Reader)
+	{
+		Reader->SetStreamInbandEventUsage(EStreamType::Video, false);
+		Reader->SetStreamInbandEventUsage(EStreamType::Audio, false);
+		Reader->SetStreamInbandEventUsage(EStreamType::Subtitle, false);
+	}
 }
 
 IStreamReader* FManifestDASH::CreateStreamReaderHandler()
@@ -473,6 +484,9 @@ void FDASHPlayPeriod::PrepareForPlay()
 			ActiveAudioRepresentationID = AudioRepr->GetUniqueIdentifier();
 		}
 
+		// Emit all <EventStream> events of the period to the AEMS event handler.
+		Manifest->SendEventsFromAllPeriodEventStreams(Period);
+
 	//	ReadyState = EReadyState::Preparing;
 		ReadyState = EReadyState::IsReady;
 	}
@@ -823,7 +837,7 @@ IManifest::FResult FDASHPlayPeriod::GetNextOrRetrySegment(TSharedPtrTS<IStreamSe
 
 		IPlaylistReaderDASH* Reader = static_cast<IPlaylistReaderDASH*>(PlayerSessionServices->GetManifestReader().Get());
 		check(Reader->GetPlaylistType().Equals(TEXT("dash")));
-		Reader->RequestMPDUpdate(false);
+		Reader->RequestMPDUpdate(IPlaylistReaderDASH::EMPDRequestType::GetLatestSegment);
 		return IManifest::FResult().RetryAfterMilliseconds(250);
 	}
 	else if (SearchResult == FManifestDASHInternal::FRepresentation::ESearchResult::Gone)
@@ -1480,8 +1494,30 @@ FString FManifestDASHInternal::FRepresentation::ApplyTemplateStrings(FString Tem
 		}
 	}
 	return NewURL;
+}
 
+void FManifestDASHInternal::FRepresentation::CollectInbandEventStreams(IPlayerSessionServices* InPlayerSessionServices, FSegmentInformation& InOutSegmentInfo)
+{
+	TSharedPtrTS<FDashMPD_RepresentationType> MPDRepresentation = Representation.Pin();
 
+	TArray<TSharedPtrTS<FDashMPD_EventStreamType>> EvS = MPDRepresentation->GetInbandEventStreams();
+	for(TSharedPtrTS<const IDashMPDElement> Parent = MPDRepresentation->GetParentElement(); Parent.IsValid(); Parent=Parent->GetParentElement())
+	{
+		if (Parent->GetElementType() == IDashMPDElement::EType::AdaptationSet)
+		{
+			EvS.Append(static_cast<const FDashMPD_AdaptationSetType*>(Parent.Get())->GetInbandEventStreams());
+			break;
+		}
+	}
+	for(int32 i=0; i<EvS.Num(); ++i)
+	{
+		FSegmentInformation::FInbandEventStream ibs;
+		ibs.SchemeIdUri = EvS[i]->GetSchemeIdUri();
+		ibs.Value = EvS[i]->GetValue();
+		ibs.PTO = (int64) EvS[i]->GetPresentationTimeOffset().GetWithDefault(0);
+		ibs.Timescale = EvS[i]->GetTimescale().GetWithDefault(1);
+		InOutSegmentInfo.InbandEventStreams.Emplace(MoveTemp(ibs));
+	}
 }
 
 
@@ -1771,6 +1807,7 @@ FManifestDASHInternal::FRepresentation::ESearchResult FManifestDASHInternal::FRe
 		OutSegmentInfo.Timescale = SidxTimescale;
 		OutSegmentInfo.ATO = ATO;
 		OutSegmentInfo.bIsLastInPeriod = CurrentT + CurrentD >= MediaLocalPeriodEnd;
+		CollectInbandEventStreams(InPlayerSessionServices, OutSegmentInfo);
 		if (!PrepareDownloadURLs(InPlayerSessionServices, OutSegmentInfo, SegmentBase))
 		{
 			bIsUsable = false;
@@ -1924,6 +1961,7 @@ FManifestDASHInternal::FRepresentation::ESearchResult FManifestDASHInternal::FRe
 	OutSegmentInfo.bMayBeMissing = SegmentNum + 1 >= MaxSegmentsInPeriod;
 	OutSegmentInfo.bIsLastInPeriod = OutSegmentInfo.bMayBeMissing && InSearchOptions.bHasFollowingPeriod;
 	OutSegmentInfo.ATO = ATO;
+	CollectInbandEventStreams(InPlayerSessionServices, OutSegmentInfo);
 	if (!PrepareDownloadURLs(InPlayerSessionServices, OutSegmentInfo, SegmentTemplate))
 	{
 		bIsUsable = false;
@@ -2250,6 +2288,7 @@ FManifestDASHInternal::FRepresentation::ESearchResult FManifestDASHInternal::FRe
 			LogMessage(InPlayerSessionServices, IInfoLog::ELevel::Warning, FString::Printf(TEXT("Representation \"%s\" <SegmentTimeline> gap encountered for needed 't' value of %lld. Replacing with an empty filler segment."), *MPDRepresentation->GetID(), (long long int)CurrentT));
 		}
 		OutSegmentInfo.ATO = ATO;
+		CollectInbandEventStreams(InPlayerSessionServices, OutSegmentInfo);
 		if (!PrepareDownloadURLs(InPlayerSessionServices, OutSegmentInfo, SegmentTemplate))
 		{
 			bIsUsable = false;
diff --git a/Engine/Plugins/Media/ElectraPlayer/Source/ElectraPlayerRuntime/Private/Runtime/Player/DASH/ManifestParserDASH.cpp b/Engine/Plugins/Media/ElectraPlayer/Source/ElectraPlayerRuntime/Private/Runtime/Player/DASH/ManifestParserDASH.cpp
--- a/Engine/Plugins/Media/ElectraPlayer/Source/ElectraPlayerRuntime/Private/Runtime/Player/DASH/ManifestParserDASH.cpp
+++ b/Engine/Plugins/Media/ElectraPlayer/Source/ElectraPlayerRuntime/Private/Runtime/Player/DASH/ManifestParserDASH.cpp
@@ -724,7 +724,8 @@ bool IDashMPDElement::ProcessElement(FManifestParserDASH* Builder, const TCHAR*
 {
 	bool bEmitWarning = true;
 	// Since descriptors are largely application specific we cannot know if their elements are meaningful or not, so we don't show a warning.
-	if (GetElementType() == EType::Descriptor)
+	// Similarly, events may carry XML elements as their payload and we will sadly parse them, too.
+	if (GetElementType() == EType::Descriptor || GetElementType() == EType::Event || GetElementType() == EType::Other)
 	{
 		bEmitWarning = false;
 	}
diff --git a/Engine/Plugins/Media/ElectraPlayer/Source/ElectraPlayerRuntime/Private/Runtime/Player/DASH/PlayerEventDASH.h b/Engine/Plugins/Media/ElectraPlayer/Source/ElectraPlayerRuntime/Private/Runtime/Player/DASH/PlayerEventDASH.h
new file mode 100644
--- /dev/null
+++ b/Engine/Plugins/Media/ElectraPlayer/Source/ElectraPlayerRuntime/Private/Runtime/Player/DASH/PlayerEventDASH.h
@@ -0,0 +1,92 @@
+// Copyright Epic Games, Inc. All Rights Reserved.
+
+#pragma once
+
+#include "PlayerCore.h"
+#include "Player/AdaptiveStreamingPlayerEvents.h"
+#include "Misc/Base64.h"
+
+namespace Electra
+{
+
+namespace DASH
+{
+
+class FPlayerEvent : public IAdaptiveStreamingPlayerAEMSEvent
+{
+public:
+	virtual ~FPlayerEvent() = default;
+
+
+	void SetOrigin(EOrigin InOrigin)
+	{ Origin = InOrigin; }
+	void SetSchemeIdUri(FString InSchemeIdUri)
+	{ SchemeIdUri = MoveTemp(InSchemeIdUri); }
+	void SetValue(FString InValue)
+	{ Value = MoveTemp(InValue); }
+	void SetID(FString InID)
+	{ ID = MoveTemp(InID); }
+	void SetPresentationTime(FTimeValue InPresentationTime)
+	{ PresentationTime = MoveTemp(InPresentationTime); }
+	void SetDuration(FTimeValue InDuration)
+	{ Duration = MoveTemp(InDuration); }
+	void SetMessageData(const FString& InData, bool bDecodeBase64)
+	{
+		if (bDecodeBase64)
+		{
+			if (!FBase64::Decode(InData, MessageData))
+			{
+				MessageData.Empty();
+			}
+		}
+		else
+		{
+			StringHelpers::StringToArray(MessageData, InData);
+		}
+	}
+	void SetMessageData(const TArray<uint8>& InData)
+	{
+		MessageData = InData;
+	}
+
+	void SetPeriodID(const FString& InPeriodID)
+	{
+		PeriodID = InPeriodID;
+	}
+	FString GetPeriodID() const
+	{
+		return PeriodID;
+	}
+
+	virtual EOrigin GetOrigin() const override
+	{ return Origin; }
+	virtual FString GetSchemeIdUri() const override
+	{ return SchemeIdUri; }
+	virtual FString GetValue() const override
+	{ return Value; }
+	virtual FString GetID() const override
+	{ return ID; }
+	virtual FTimeValue GetPresentationTime() const override
+	{ return PresentationTime; }
+	virtual FTimeValue GetDuration() const override
+	{ return Duration; }
+	virtual const TArray<uint8>& GetMessageData() const override
+	{ return MessageData; }
+
+private:
+	EOrigin Origin;
+	FString SchemeIdUri;
+	FString Value;
+	FString ID;
+	FTimeValue PresentationTime;
+	FTimeValue Duration;
+	TArray<uint8> MessageData;
+	FString PeriodID;
+};
+
+
+}
+
+} // namespace Electra
+
+
diff --git a/Engine/Plugins/Media/ElectraPlayer/Source/ElectraPlayerRuntime/Private/Runtime/Player/DASH/PlayerEventDASH_Internal.h b/Engine/Plugins/Media/ElectraPlayer/Source/ElectraPlayerRuntime/Private/Runtime/Player/DASH/PlayerEventDASH_Internal.h
new file mode 100644
--- /dev/null
+++ b/Engine/Plugins/Media/ElectraPlayer/Source/ElectraPlayerRuntime/Private/Runtime/Player/DASH/PlayerEventDASH_Internal.h
@@ -0,0 +1,42 @@
+// Copyright Epic Games, Inc. All Rights Reserved.
+
+#pragma once
+
+#include "PlayerCore.h"
+
+namespace Electra
+{
+
+namespace DASH
+{
+
+namespace Schemes
+{
+
+namespace TimingSources
+{
+    const TCHAR* const Scheme_urn_mpeg_dash_utc_httphead2014 = TEXT("urn:mpeg:dash:utc:http-head:2014");
+    const TCHAR* const Scheme_urn_mpeg_dash_utc_httpxsdate2014 = TEXT("urn:mpeg:dash:utc:http-xsdate:2014");
+    const TCHAR* const Scheme_urn_mpeg_dash_utc_httpiso2014 = TEXT("urn:mpeg:dash:utc:http-iso:2014");
+    const TCHAR* const Scheme_urn_mpeg_dash_utc_direct2014 = TEXT("urn:mpeg:dash:utc:direct:2014");
+/* not currently supported
+    const TCHAR* const Scheme_urn_mpeg_dash_utc_httpntp2014 = TEXT("urn:mpeg:dash:utc:http-ntp:2014");
+    const TCHAR* const Scheme_urn_mpeg_dash_utc_ntp2014 = TEXT("urn:mpeg:dash:utc:ntp:2014");
+    const TCHAR* const Scheme_urn_mpeg_dash_utc_sntp2014 = TEXT("urn:mpeg:dash:utc:sntp:2014");
+*/
+}
+
+namespace ManifestEvents
+{
+    const TCHAR* const Scheme_urn_mpeg_dash_event_2012 = TEXT("urn:mpeg:dash:event:2012");
+    const TCHAR* const Scheme_urn_mpeg_dash_event_callback_2015 = TEXT("urn:mpeg:dash:event:callback:2015");
+    const TCHAR* const Scheme_urn_mpeg_dash_event_ttfn_2016 = TEXT("urn:mpeg:dash:event:ttfn:2016");
+};
+
+}
+
+}
+
+} // namespace Electra
+
+
diff --git a/Engine/Plugins/Media/ElectraPlayer/Source/ElectraPlayerRuntime/Private/Runtime/Player/DASH/PlaylistReaderDASH.cpp b/Engine/Plugins/Media/ElectraPlayer/Source/ElectraPlayerRuntime/Private/Runtime/Player/DASH/PlaylistReaderDASH.cpp
--- a/Engine/Plugins/Media/ElectraPlayer/Source/ElectraPlayerRuntime/Private/Runtime/Player/DASH/PlaylistReaderDASH.cpp
+++ b/Engine/Plugins/Media/ElectraPlayer/Source/ElectraPlayerRuntime/Private/Runtime/Player/DASH/PlaylistReaderDASH.cpp
@@ -14,8 +14,11 @@
 #include "HAL/LowLevelMemTracker.h"
 #include "ElectraPlayerPrivate.h"
 #include "Player/AdaptiveStreamingPlayerResourceRequest.h"
+#include "Player/AdaptiveStreamingPlayerEvents.h"
 #include "Player/PlayerEntityCache.h"
 #include "Player/DASH/OptionKeynamesDASH.h"
+#include "Player/DASH/PlayerEventDASH.h"
+#include "Player/DASH/PlayerEventDASH_Internal.h"
 
 #include "Misc/DateTime.h"
 
@@ -34,7 +37,7 @@ namespace Electra
 /**
  * This class is responsible for downloading a DASH MPD and parsing it.
  */
-class FPlaylistReaderDASH : public TSharedFromThis<FPlaylistReaderDASH, ESPMode::ThreadSafe>, public IPlaylistReaderDASH, public FMediaThread
+class FPlaylistReaderDASH : public TSharedFromThis<FPlaylistReaderDASH, ESPMode::ThreadSafe>, public IPlaylistReaderDASH, public IAdaptiveStreamingPlayerAEMSReceiver, public FMediaThread
 {
 public:
 	FPlaylistReaderDASH();
@@ -51,11 +54,15 @@ public:
 	virtual FString GetURL() const override;
 	virtual TSharedPtrTS<IManifest> GetManifest() override;
 	virtual void AddElementLoadRequests(const TArray<TWeakPtrTS<FMPDLoadRequestDASH>>& RemoteElementLoadRequests) override;
-	virtual void RequestMPDUpdate(bool bForcedUpdate) override;
+	virtual void RequestMPDUpdate(EMPDRequestType InRequestType) override;
 	virtual TSharedPtrTS<FManifestDASHInternal> GetCurrentMPD() override
 	{
 		return Manifest;
 	}
+	virtual void SetStreamInbandEventUsage(EStreamType InStreamType, bool bUsesInbandDASHEvents) override;
+
+
+	virtual void OnMediaPlayerEventReceived(TSharedPtrTS<IAdaptiveStreamingPlayerAEMSEvent> InEvent, IAdaptiveStreamingPlayerAEMSReceiver::EDispatchMode InDispatchMode) override;
 
 private:
 	using FResourceLoadRequestPtr = TSharedPtrTS<FMPDLoadRequestDASH>;
@@ -94,6 +101,7 @@ private:
 	void Timesync_httphead_Completed(FResourceLoadRequestPtr Request, bool bSuccess);
 	void Timesync_httpxsdate_Completed(FResourceLoadRequestPtr Request, bool bSuccess);
 	void Timesync_httpiso_Completed(FResourceLoadRequestPtr Request, bool bSuccess);
+	void Callback_Completed(FResourceLoadRequestPtr Request, bool bSuccess);
 
 	IPlayerSessionServices*									PlayerSessionServices = nullptr;
 	FString													MPDURL;
@@ -110,13 +118,18 @@ private:
 	FTimeValue												NextMPDUpdateTime;
 	FTimeValue												MostRecentMPDUpdateTime;
 	bool													bIsMPDUpdateInProgress = false;
-	enum class EUpdateRequestType
+	struct FMPDUpdateRequest
 	{
-		None,
-		Regular,
-		Forced
+		IPlaylistReaderDASH::EMPDRequestType Type = IPlaylistReaderDASH::EMPDRequestType::MinimumUpdatePeriod;
+		FString	PublishTime;
+		FString	PeriodId;
+		FTimeValue ValidUntil;
+		FTimeValue NewDuration;
+		FString NewMPD;
 	};
-	EUpdateRequestType										UpdateRequested = EUpdateRequestType::None;
+	TArray<TSharedPtrTS<FMPDUpdateRequest>>					UpdateRequested;
+
+	bool													bInbandEventByStreamType[3] { false, false, false };
 
 	TUniquePtr<IManifestBuilderDASH>						Builder;
 
@@ -136,18 +149,6 @@ private:
 };
 
 
-namespace DASHTimingSources
-{
-    const TCHAR* const Scheme_urn_mpeg_dash_utc_httphead2014 = TEXT("urn:mpeg:dash:utc:http-head:2014");
-    const TCHAR* const Scheme_urn_mpeg_dash_utc_httpxsdate2014 = TEXT("urn:mpeg:dash:utc:http-xsdate:2014");
-    const TCHAR* const Scheme_urn_mpeg_dash_utc_httpiso2014 = TEXT("urn:mpeg:dash:utc:http-iso:2014");
-    const TCHAR* const Scheme_urn_mpeg_dash_utc_direct2014 = TEXT("urn:mpeg:dash:utc:direct:2014");
-/* not currently supported
-    const TCHAR* const Scheme_urn_mpeg_dash_utc_httpntp2014 = TEXT("urn:mpeg:dash:utc:http-ntp:2014");
-    const TCHAR* const Scheme_urn_mpeg_dash_utc_ntp2014 = TEXT("urn:mpeg:dash:utc:ntp:2014");
-    const TCHAR* const Scheme_urn_mpeg_dash_utc_sntp2014 = TEXT("urn:mpeg:dash:utc:sntp:2014");
-*/
-};
 
 
 /***************************************************************************************************************************************************/
@@ -271,36 +272,148 @@ void FPlaylistReaderDASH::LoadAndParse(const FString& URL)
 
 void FPlaylistReaderDASH::CheckForMPDUpdate()
 {
+	// Get the time now and do not pass it in from the worker thread as the clock could have been
+	// resynchronized to the server time in the meantime.
+	FTimeValue Now = PlayerSessionServices->GetSynchronizedUTCTime()->GetTime();
+
 	RequestsLock.Lock();
-	FTimeValue Next = NextMPDUpdateTime;
 	bool bRequestNow = false;
-	bool bIsForced = false;
-	if (UpdateRequested == EUpdateRequestType::Forced)
-	{
-		bIsForced = true;
-		Next.SetToZero();
-	}
-	else if (UpdateRequested == EUpdateRequestType::Regular)
+	bool bTimedRequestAdded = false;
+
+	// Check if the time to update the MPD through the MPD@minimumUpdatePeriod has come.
+	if (NextMPDUpdateTime.IsValid())
 	{
-		if (!Next.IsValid())
+		// Limit the frequency of the updates.
+		bool bIsPossible = !MostRecentMPDUpdateTime.IsValid() || Now - MostRecentMPDUpdateTime > MinTimeBetweenUpdates;
+		if (NextMPDUpdateTime < Now && bIsPossible)
 		{
-			Next.SetToZero();
+			// When it's time to do the update add a request to the request queue. This gets handled
+			// right away next thing. Doing it that way means we can put the handling into one spot.
+			NextMPDUpdateTime.SetToInvalid();
+			TSharedPtrTS<FMPDUpdateRequest> Request = MakeSharedTS<FMPDUpdateRequest>();
+			Request->Type = IPlaylistReaderDASH::EMPDRequestType::MinimumUpdatePeriod;
+			UpdateRequested.Emplace(MoveTemp(Request));
+			bTimedRequestAdded = true;
 		}
 	}
-	if (!bIsMPDUpdateInProgress && Next.IsValid())
+
+	// Go over all currently pending update requests.
+	while(UpdateRequested.Num())
 	{
-		// Get the time now and do not pass it in from the worker thread as the clock could have been
-		// resynchronized to the server time in the meantime.
-		FTimeValue Now = PlayerSessionServices->GetSynchronizedUTCTime()->GetTime();
-		// Limit the frequency of the updates.
-		bool bIsPossible = !MostRecentMPDUpdateTime.IsValid() || Now - MostRecentMPDUpdateTime > MinTimeBetweenUpdates;
-		if (Next < Now && bIsPossible)
+		TSharedPtrTS<FMPDUpdateRequest> Request = UpdateRequested[0];
+		UpdateRequested.RemoveAt(0);
+		if (Request->Type == IPlaylistReaderDASH::EMPDRequestType::MinimumUpdatePeriod ||
+			Request->Type == IPlaylistReaderDASH::EMPDRequestType::GetLatestSegment)
 		{
-			bRequestNow = true;
-			NextMPDUpdateTime.SetToInvalid();
+			// If the request to update is because we need the latest segment and there is currently
+			// no MPD update by time we ask for one with the next check.
+			if (Request->Type == IPlaylistReaderDASH::EMPDRequestType::GetLatestSegment)
+			{
+				if (!bTimedRequestAdded && !NextMPDUpdateTime.IsValid())
+				{
+					NextMPDUpdateTime.SetToZero();
+				}
+				continue;
+			}
+
+			if (!bIsMPDUpdateInProgress)
+			{
+				bRequestNow = true;
+				/*
+					If there is any stream receiving inband events we do not load the MPD through
+					expiration of time. The events have sole control and responsibility here.
+
+					Note: The expectation at the moment is that inband MPD event streams are only
+						  used with Live streams and are carried in the audio stream instead of
+						  each and every video stream. This here being the worker thread to
+						  refresh the MPD and downloading remote entities (index segments, xlink
+						  elements and such) it has no immediate knowledge which streams are
+						  actively used in the playing session and whether these are carrying
+						  inband event streams or not. Instead it relies on the segment downloader
+						  to tell it if a particular type of stream (video, audio, etc.) is
+						  using inband events as each segment is downloaded.
+						  Meaning that if a stream type stops downloading (ie it reaching its end)
+						  it cannot clear this flag and we assume we are still receiving events
+						  from it. Hence the expectation this is used with Live streams that
+						  just do not end and keep running.
+						  In addition this is only done when MPD@minimumUpdatePeriod is zero
+						  as per DASH-IF-IOP v4.3 Setion 4.5.
+				*/
+				FTimeValue mup = Manifest.IsValid() && Manifest->GetMPDRoot().IsValid() ? Manifest->GetMPDRoot()->GetMinimumUpdatePeriod() : FTimeValue::GetInvalid();
+				// check video and audio only
+				if ((bInbandEventByStreamType[0] || bInbandEventByStreamType[1]) && mup == FTimeValue::GetZero())
+				{
+					bRequestNow = false;
+				}
+			}
+		}
+		else if (Request->Type == IPlaylistReaderDASH::EMPDRequestType::EventMessage)
+		{
+			// Unconditional update?
+			if (Request->PublishTime.IsEmpty() && !Request->NewDuration.IsValid() && !Request->ValidUntil.IsValid())
+			{
+				// Note: We could check the publish time against the current MPD's publish time to see if the event
+				//       still applies, but that seem superfluous considering we're only reloading the MPD.
+				bRequestNow = true;
+			}
+			else
+			{
+				// Does this terminate the presentation at a specified time?
+				if (Request->ValidUntil.IsValid() && Request->NewDuration.IsValid())
+				{
+					bool bEventApplies = true;
+
+					if (Manifest.IsValid())
+					{
+						// Have to check the event publish time against the current MPD publish time to see if the event still applies
+						// to the current MPD.
+						TSharedPtrTS<const FDashMPD_MPDType> MPDRoot = Manifest->GetMPDRoot();
+						FTimeValue CurrentPUBT = MPDRoot.IsValid() ? MPDRoot->GetPublishTime() : FTimeValue::GetInvalid();
+						FTimeValue EventPUBT;
+						ISO8601::ParseDateTime(EventPUBT, Request->PublishTime);
+						// We take it that if neither PUBT is valid the event is still to be used.
+						// Only if the times are given and the event PUBT has already passed we ignore the event.
+						if (CurrentPUBT.IsValid() && EventPUBT.IsValid() && CurrentPUBT > EventPUBT)
+						{
+							bEventApplies = false;
+						}
+						if (bEventApplies)
+						{
+							Manifest->EndPresentationAt(Request->ValidUntil + Request->NewDuration, Request->PeriodId);
+						}
+					}
+					if (bEventApplies)
+					{
+						// Set internal times such that no MPD reloads will trigger based on time.
+						NextMPDUpdateTime.SetToPositiveInfinity();
+						MostRecentMPDUpdateTime.SetToPositiveInfinity();
+						// Clear out all current requests.
+						bRequestNow = false;
+						UpdateRequested.Empty();
+						// Send a fake update message to the player to have it evaluate the media timeline again.
+						PlayerSessionServices->SendMessageToPlayer(IPlaylistReader::PlaylistLoadedMessage::Create(FErrorDetail(), nullptr, Playlist::EListType::Master, Playlist::ELoadType::Update));
+					}
+				}
+				else
+				{
+					// The event could provide a new MPD to be used. This case is not supported yet though.
+					if (!Request->NewMPD.IsEmpty())
+					{
+						// Something to implement in the future.
+						LogMessage(IInfoLog::ELevel::Info, FString::Printf(TEXT("DASH events providing a new MPD are not supported yet. Triggering an MPD reload instead.")));
+					}
+					// Just trigger an MPD reload instead.
+					// This is also where we a failed segment download of a segment with an inband event will arrive.
+					bRequestNow = true;
+				}
+			}
+		}
+		else
+		{
+			check(!"Unhandled update type");
 		}
 	}
-	UpdateRequested = EUpdateRequestType::None;
+
 	RequestsLock.Unlock();
 	if (bRequestNow)
 	{
@@ -427,16 +540,30 @@ void FPlaylistReaderDASH::AddElementLoadRequests(const TArray<TWeakPtrTS<FMPDLoa
 	}
 }
 
-void FPlaylistReaderDASH::RequestMPDUpdate(bool bForcedUpdate)
+void FPlaylistReaderDASH::RequestMPDUpdate(IPlaylistReaderDASH::EMPDRequestType InRequestType)
 {
+	// This method must not be used with event messages since it does not have all the required values.
+	check(InRequestType != 	IPlaylistReaderDASH::EMPDRequestType::EventMessage);
+
+	TSharedPtrTS<FMPDUpdateRequest> Request = MakeSharedTS<FMPDUpdateRequest>();
+	Request->Type = InRequestType;
+
 	FScopeLock lock(&RequestsLock);
-	if (bForcedUpdate)
-	{
-		UpdateRequested = EUpdateRequestType::Forced;
-	}
-	else if (UpdateRequested == EUpdateRequestType::None)
+	UpdateRequested.Emplace(MoveTemp(Request));
+}
+
+void FPlaylistReaderDASH::SetStreamInbandEventUsage(EStreamType InStreamType, bool bUsesInbandDASHEvents)
+{
+	switch(InStreamType)
 	{
-		UpdateRequested = EUpdateRequestType::Regular;
+		case EStreamType::Video: 
+			bInbandEventByStreamType[0] = bUsesInbandDASHEvents; 
+			break;
+		case EStreamType::Audio: 
+			bInbandEventByStreamType[1] = bUsesInbandDASHEvents; 
+			break;
+		default:
+			break;
 	}
 }
 
@@ -476,7 +603,7 @@ void FPlaylistReaderDASH::TriggerTimeSynchronization()
 			*/
 			for(int32 i=0; i<UTCTimings.Num(); ++i)
 			{
-				if (UTCTimings[i]->GetSchemeIdUri().Equals(DASHTimingSources::Scheme_urn_mpeg_dash_utc_direct2014))
+				if (UTCTimings[i]->GetSchemeIdUri().Equals(DASH::Schemes::TimingSources::Scheme_urn_mpeg_dash_utc_direct2014))
 				{
 					FTimeValue DirectTime;
 					if (ISO8601::ParseDateTime(DirectTime, UTCTimings[i]->GetValue()) == UEMEDIA_ERROR_OK)
@@ -513,21 +640,21 @@ void FPlaylistReaderDASH::TriggerTimeSynchronization()
 
 				TSharedPtrTS<FMPDLoadRequestDASH> TimeSyncRequest;
 				// Check for supported schemes. They appear in the manifest in the order the author has prioritized them.
-				if (SchemeIdUri.Equals(DASHTimingSources::Scheme_urn_mpeg_dash_utc_httpxsdate2014))
+				if (SchemeIdUri.Equals(DASH::Schemes::TimingSources::Scheme_urn_mpeg_dash_utc_httpxsdate2014))
 				{
 					TimeSyncRequest = MakeSharedTS<FMPDLoadRequestDASH>();
 					TimeSyncRequest->LoadType = FMPDLoadRequestDASH::ELoadType::TimeSync;
 					TimeSyncRequest->URL = GetRequestURL(Value, 0);
 					TimeSyncRequest->CompleteCallback.BindThreadSafeSP(AsShared(), &FPlaylistReaderDASH::Timesync_httpxsdate_Completed);
 				}
-				else if (SchemeIdUri.Equals(DASHTimingSources::Scheme_urn_mpeg_dash_utc_httpiso2014))
+				else if (SchemeIdUri.Equals(DASH::Schemes::TimingSources::Scheme_urn_mpeg_dash_utc_httpiso2014))
 				{
 					TimeSyncRequest = MakeSharedTS<FMPDLoadRequestDASH>();
 					TimeSyncRequest->LoadType = FMPDLoadRequestDASH::ELoadType::TimeSync;
 					TimeSyncRequest->URL = GetRequestURL(Value, 0);
 					TimeSyncRequest->CompleteCallback.BindThreadSafeSP(AsShared(), &FPlaylistReaderDASH::Timesync_httpiso_Completed);
 				}
-				else if (SchemeIdUri.Equals(DASHTimingSources::Scheme_urn_mpeg_dash_utc_httphead2014))
+				else if (SchemeIdUri.Equals(DASH::Schemes::TimingSources::Scheme_urn_mpeg_dash_utc_httphead2014))
 				{
 					TimeSyncRequest = MakeSharedTS<FMPDLoadRequestDASH>();
 					TimeSyncRequest->LoadType = FMPDLoadRequestDASH::ELoadType::TimeSync;
@@ -574,6 +701,15 @@ void FPlaylistReaderDASH::WorkerThread()
 	LLM_SCOPE(ELLMTag::ElectraPlayer);
 	CSV_SCOPED_TIMING_STAT(ElectraPlayer, MPDReaderDASH_Worker);
 
+	bInbandEventByStreamType[0] = false;
+	bInbandEventByStreamType[1] = false;
+	bInbandEventByStreamType[2] = false;
+
+	// Register event callbacks
+	PlayerSessionServices->GetAEMSEventHandler()->AddAEMSReceiver(AsShared(), DASH::Schemes::ManifestEvents::Scheme_urn_mpeg_dash_event_2012, TEXT(""), IAdaptiveStreamingPlayerAEMSReceiver::EDispatchMode::OnStart, false);
+	PlayerSessionServices->GetAEMSEventHandler()->AddAEMSReceiver(AsShared(), DASH::Schemes::ManifestEvents::Scheme_urn_mpeg_dash_event_callback_2015, TEXT("1"), IAdaptiveStreamingPlayerAEMSReceiver::EDispatchMode::OnStart, false);
+	PlayerSessionServices->GetAEMSEventHandler()->AddAEMSReceiver(AsShared(), DASH::Schemes::ManifestEvents::Scheme_urn_mpeg_dash_event_ttfn_2016, TEXT(""), IAdaptiveStreamingPlayerAEMSReceiver::EDispatchMode::OnStart, false);
+
 	// Get the minimum MPD update time limit.
 	MinTimeBetweenUpdates = PlayerSessionServices->GetOptions().GetValue(DASH::OptionKey_MinTimeBetweenMPDUpdates).SafeGetTimeValue(FTimeValue().SetFromMilliseconds(1000));
 
@@ -604,6 +740,11 @@ void FPlaylistReaderDASH::WorkerThread()
 		CheckForMPDUpdate();
 	}
 
+	// Unregister event callbacks
+	PlayerSessionServices->GetAEMSEventHandler()->RemoveAEMSReceiver(AsShared(), DASH::Schemes::ManifestEvents::Scheme_urn_mpeg_dash_event_ttfn_2016, TEXT(""), IAdaptiveStreamingPlayerAEMSReceiver::EDispatchMode::OnStart);
+	PlayerSessionServices->GetAEMSEventHandler()->RemoveAEMSReceiver(AsShared(), DASH::Schemes::ManifestEvents::Scheme_urn_mpeg_dash_event_callback_2015, TEXT("1"), IAdaptiveStreamingPlayerAEMSReceiver::EDispatchMode::OnStart);
+	PlayerSessionServices->GetAEMSEventHandler()->RemoveAEMSReceiver(AsShared(), DASH::Schemes::ManifestEvents::Scheme_urn_mpeg_dash_event_2012, TEXT(""), IAdaptiveStreamingPlayerAEMSReceiver::EDispatchMode::OnStart);
+
 	// Cleanup!
 	RequestsLock.Lock();
 	PendingRequests.Empty();
@@ -928,7 +1069,6 @@ void FPlaylistReaderDASH::ManifestDownloadCompleted(FResourceLoadRequestPtr Requ
 				ETag = ConnInfo->ResponseHeaders[i].Value;
 			}
 		}
-		MostRecentMPDUpdateTime = FetchTime;
 
 		PlayerSessionServices->SendMessageToPlayer(IPlaylistReader::PlaylistDownloadMessage::Create(ConnInfo, Playlist::EListType::Master, Playlist::ELoadType::Initial));
 
@@ -938,7 +1078,7 @@ void FPlaylistReaderDASH::ManifestDownloadCompleted(FResourceLoadRequestPtr Requ
 		{
 			Builder.Reset(IManifestBuilderDASH::Create(PlayerSessionServices));
 			TSharedPtrTS<FManifestDASHInternal> NewManifest;
-			LastErrorDetail = Builder->BuildFromMPD(NewManifest, XML.GetCharArray().GetData(), EffectiveURL, FetchTime, ETag);
+			LastErrorDetail = Builder->BuildFromMPD(NewManifest, XML.GetCharArray().GetData(), EffectiveURL, ETag);
 			if (LastErrorDetail.IsOK() || LastErrorDetail.IsTryAgain())
 			{
 				if (NewManifest.IsValid())
@@ -947,16 +1087,27 @@ void FPlaylistReaderDASH::ManifestDownloadCompleted(FResourceLoadRequestPtr Requ
 					TArray<FURL_RFC3986::FQueryParam> URLFragmentComponents;
 					FURL_RFC3986::GetQueryParams(URLFragmentComponents, Fragment, false);	// The fragment is already URL escaped, so no need to do it again.
 					NewManifest->SetURLFragmentComponents(MoveTemp(URLFragmentComponents));
-					// If the URL has a special fragment part to turn this presentation into an event, do so.
-					// Note: This MAY require the client clock to be synced to the server IF the special keyword 'now' is used.
-					NewManifest->TransformIntoEpicEvent();
 				}
 				Manifest = NewManifest;
 				TriggerTimeSynchronization();
+				if (Manifest.IsValid())
+				{
+					// Note: This is not like the standard defines FetchTime. It should be the time the server is processing the request
+					//       but we can't know this. Also if there is a direct <UTCTiming> element in the MPD it may be different from
+					//       what our local clock is set to at the moment. So instead we set FetchTime to NOW.
+					FetchTime = PlayerSessionServices->GetSynchronizedUTCTime()->GetTime();
+					Manifest->GetMPDRoot()->SetFetchTime(FetchTime);
+					MostRecentMPDUpdateTime = FetchTime;
+
+					// If the URL has a special fragment part to turn this presentation into an event, do so.
+					// Note: This MAY require the client clock to be synced to the server IF the special keyword 'now' is used.
+					Manifest->TransformIntoEpicEvent();
+				}
+
 				PlayerManifest = FManifestDASH::Create(PlayerSessionServices, Manifest);
 
 				// Check if the MPD defines an @minimumUpdatePeriod
-				FTimeValue mup = Manifest->GetMinimumUpdatePeriod();
+				FTimeValue mup = Manifest.IsValid() ? Manifest->GetMinimumUpdatePeriod() : FTimeValue::GetInvalid();
 				// If the update time is zero then updates happen just in time when segments are required or through
 				// an inband event stream. Either way, we do not need to update periodically.
 				if (mup.IsValid() && mup > FTimeValue::GetZero())
@@ -1025,7 +1176,9 @@ void FPlaylistReaderDASH::ManifestUpdateDownloadCompleted(FResourceLoadRequestPt
 		FString EffectiveURL;
 		FTimeValue FetchTime;
 		EffectiveURL = ConnInfo->EffectiveURL;
-		FetchTime = PlayerSessionServices->GetSynchronizedUTCTime()->MapToSyncTime(ConnInfo->RequestStartTime);
+		// See above why we use NOW for FetchTime.
+		//FetchTime = PlayerSessionServices->GetSynchronizedUTCTime()->MapToSyncTime(ConnInfo->RequestStartTime);
+		FetchTime = PlayerSessionServices->GetSynchronizedUTCTime()->GetTime();
 		for(int32 i=0; i<ConnInfo->ResponseHeaders.Num(); ++i)
 		{
 			if (ConnInfo->ResponseHeaders[i].Header.Equals(TEXT("ETag")))
@@ -1060,7 +1213,7 @@ void FPlaylistReaderDASH::ManifestUpdateDownloadCompleted(FResourceLoadRequestPt
 			if (LastErrorDetail.IsOK())
 			{
 				TSharedPtrTS<FManifestDASHInternal> NewManifest;
-				LastErrorDetail = Builder->BuildFromMPD(NewManifest, XML.GetCharArray().GetData(), EffectiveURL, FetchTime, ETag);
+				LastErrorDetail = Builder->BuildFromMPD(NewManifest, XML.GetCharArray().GetData(), EffectiveURL, ETag);
 				if (LastErrorDetail.IsOK() || LastErrorDetail.IsTryAgain())
 				{
 					// Copy over the initial document URL fragments.
@@ -1075,9 +1228,14 @@ void FPlaylistReaderDASH::ManifestUpdateDownloadCompleted(FResourceLoadRequestPt
 					PlayerManifest->UpdateInternalManifest(Manifest);
 
 					TriggerTimeSynchronization();
+					if (Manifest.IsValid())
+					{
+						Manifest->GetMPDRoot()->SetFetchTime(FetchTime);
+						MostRecentMPDUpdateTime = FetchTime;
+					}
 
 					// Check if the new MPD also defines an @minimumUpdatePeriod
-					FTimeValue mup = Manifest->GetMinimumUpdatePeriod();
+					FTimeValue mup = Manifest.IsValid() ? Manifest->GetMinimumUpdatePeriod() : FTimeValue::GetInvalid();
 					if (mup.IsValid() && mup > FTimeValue::GetZero())
 					{
 						if (mup.GetAsMilliseconds() < 1000)
@@ -1186,13 +1344,19 @@ void FPlaylistReaderDASH::Timesync_httpxsdate_Completed(FResourceLoadRequestPtr
 		FString Response;
 		if (GetResponseString(Response, Request).IsOK())
 		{
-			FTimeValue DirectTime;
+			FTimeValue NewTime;
 			// Note: Yes, this is a bit weird, but the xs:dateTime format is actually exactly the same as ISO-8601
 			//       so it is not clear why there is a distinction made in the UTCTiming scheme.
 			//       We keep the different callback handlers here though for completeness sake.
-			if (ISO8601::ParseDateTime(DirectTime, Response) == UEMEDIA_ERROR_OK)
+			if (ISO8601::ParseDateTime(NewTime, Response) == UEMEDIA_ERROR_OK)
 			{
-				PlayerSessionServices->GetSynchronizedUTCTime()->SetTime(DirectTime);
+				PlayerSessionServices->GetSynchronizedUTCTime()->SetTime(NewTime);
+			}
+			// If parsing failed then maybe the response is just a number (possibly with frational digits) giving the
+			// current Unix epoch time.
+			else if (UnixEpoch::ParseFloatString(NewTime, Response))
+			{
+				PlayerSessionServices->GetSynchronizedUTCTime()->SetTime(NewTime);
 			}
 		}
 	}
@@ -1208,10 +1372,10 @@ void FPlaylistReaderDASH::Timesync_httpiso_Completed(FResourceLoadRequestPtr Req
 		FString Response;
 		if (GetResponseString(Response, Request).IsOK())
 		{
-			FTimeValue DirectTime;
-			if (ISO8601::ParseDateTime(DirectTime, Response) == UEMEDIA_ERROR_OK)
+			FTimeValue NewTime;
+			if (ISO8601::ParseDateTime(NewTime, Response) == UEMEDIA_ERROR_OK)
 			{
-				PlayerSessionServices->GetSynchronizedUTCTime()->SetTime(DirectTime);
+				PlayerSessionServices->GetSynchronizedUTCTime()->SetTime(NewTime);
 			}
 		}
 	}
@@ -1220,6 +1384,80 @@ void FPlaylistReaderDASH::Timesync_httpiso_Completed(FResourceLoadRequestPtr Req
 	bTimeSyncInProgress = false;
 }
 
+void FPlaylistReaderDASH::OnMediaPlayerEventReceived(TSharedPtrTS<IAdaptiveStreamingPlayerAEMSEvent> InEvent, IAdaptiveStreamingPlayerAEMSReceiver::EDispatchMode InDispatchMode)
+{
+	if (InEvent->GetSchemeIdUri().Equals(DASH::Schemes::ManifestEvents::Scheme_urn_mpeg_dash_event_2012))
+	{
+		// This cast is safe because we only exist when the format is DASH and the events must therefore be of that type.
+		const DASH::FPlayerEvent* DASHEvent = static_cast<const DASH::FPlayerEvent*>(InEvent.Get());
+		FString EventPeriodId = DASHEvent->GetPeriodID();
+
+		// ISO/IEC 23009-1:2019 Section 5.10.4 DASH-specific events
+		if (InEvent->GetOrigin() == IAdaptiveStreamingPlayerAEMSEvent::EOrigin::EventStream)
+		{
+			// For an MPD triggered event there is no point in including a patch or a full MPD.
+			// If either was already available when the current MPD was generated the MPD could
+			// have already included the changes. As such we only handle value 1 here which
+			// sets a validity expiration, which for an event coming from the MPD means to just
+			// fetch a new MPD now.
+			if (InEvent->GetValue().Equals(TEXT("1")))
+			{
+				// Add a new MPD update event. Set EventMessage for type but leave the other fields empty.
+				TSharedPtrTS<FMPDUpdateRequest> Request = MakeSharedTS<FMPDUpdateRequest>();
+				Request->Type = IPlaylistReaderDASH::EMPDRequestType::EventMessage;
+				Request->PeriodId = EventPeriodId;
+				FScopeLock lock(&RequestsLock);
+				UpdateRequested.Emplace(MoveTemp(Request));
+			}
+		}
+		else if (InEvent->GetOrigin() == IAdaptiveStreamingPlayerAEMSEvent::EOrigin::InbandEventStream)
+		{
+			// An inband event needs more consideration.
+			TSharedPtrTS<FMPDUpdateRequest> Request = MakeSharedTS<FMPDUpdateRequest>();
+			Request->Type = IPlaylistReaderDASH::EMPDRequestType::EventMessage;
+			Request->PeriodId = EventPeriodId;
+			Request->ValidUntil = InEvent->GetPresentationTime();
+			Request->NewDuration = InEvent->GetDuration();
+			int32 NULPos = INDEX_NONE;
+			if (InEvent->GetMessageData().Num() && InEvent->GetMessageData().Find(0, NULPos))
+			{
+				Request->PublishTime = StringHelpers::ArrayToString(TArray<uint8>(InEvent->GetMessageData().GetData(), NULPos));
+				FString NewMPD;
+				if (InEvent->GetValue().Equals(TEXT("3")))
+				{
+					Request->NewMPD = StringHelpers::ArrayToString(TArray<uint8>(InEvent->GetMessageData().GetData() + NULPos + 1, InEvent->GetMessageData().Num() - NULPos - 1));
+				}
+			}
+			FScopeLock lock(&RequestsLock);
+			UpdateRequested.Emplace(MoveTemp(Request));
+		}
+	}
+	else if (InEvent->GetSchemeIdUri().Equals(DASH::Schemes::ManifestEvents::Scheme_urn_mpeg_dash_event_callback_2015))
+	{
+		FString URL = StringHelpers::ArrayToString(InEvent->GetMessageData());
+		if (!URL.IsEmpty())
+		{
+			TSharedPtrTS<FMPDLoadRequestDASH> CallbackRequest;
+			CallbackRequest = MakeSharedTS<FMPDLoadRequestDASH>();
+			CallbackRequest->LoadType = FMPDLoadRequestDASH::ELoadType::Callback;
+			CallbackRequest->URL = URL;
+			CallbackRequest->CompleteCallback.BindThreadSafeSP(AsShared(), &FPlaylistReaderDASH::Callback_Completed);
+			EnqueueResourceRequest(MoveTemp(CallbackRequest));
+		}
+	}
+	else if (InEvent->GetSchemeIdUri().Equals(DASH::Schemes::ManifestEvents::Scheme_urn_mpeg_dash_event_ttfn_2016))
+	{
+		// TBD
+	}
+}
+
+void FPlaylistReaderDASH::Callback_Completed(FResourceLoadRequestPtr Request, bool bSuccess)
+{
+	// Callback responses are ignored.
+}
+
+
+
 } // namespace Electra
 
 
diff --git a/Engine/Plugins/Media/ElectraPlayer/Source/ElectraPlayerRuntime/Private/Runtime/Player/DASH/PlaylistReaderDASH.h b/Engine/Plugins/Media/ElectraPlayer/Source/ElectraPlayerRuntime/Private/Runtime/Player/DASH/PlaylistReaderDASH.h
--- a/Engine/Plugins/Media/ElectraPlayer/Source/ElectraPlayerRuntime/Private/Runtime/Player/DASH/PlaylistReaderDASH.h
+++ b/Engine/Plugins/Media/ElectraPlayer/Source/ElectraPlayerRuntime/Private/Runtime/Player/DASH/PlaylistReaderDASH.h
@@ -46,12 +46,23 @@ public:
 	/**
 	 * Requests an MPD update.
 	 */
-	virtual void RequestMPDUpdate(bool bForcedUpdate) = 0;
+	enum class EMPDRequestType
+	{
+		MinimumUpdatePeriod,		// Triggered through MPD@minimumUpdatePeriod
+		GetLatestSegment,			// Locating the segment in the Live period
+		EventMessage				// Triggered through an event message
+	};
+	virtual void RequestMPDUpdate(EMPDRequestType InRequestType) = 0;
 
 	/**
 	 * Access the current internal MPD
 	 */
 	virtual TSharedPtrTS<FManifestDASHInternal> GetCurrentMPD() = 0;
+
+	/**
+	 * Informs whether or not a currently playing stream is using inband events.
+	 */
+	virtual void SetStreamInbandEventUsage(EStreamType InStreamType, bool bUsesInbandDASHEvents) = 0;
 };
 
 
diff --git a/Engine/Plugins/Media/ElectraPlayer/Source/ElectraPlayerRuntime/Private/Runtime/Player/DASH/PlaylistReaderDASH_Internal.h b/Engine/Plugins/Media/ElectraPlayer/Source/ElectraPlayerRuntime/Private/Runtime/Player/DASH/PlaylistReaderDASH_Internal.h
--- a/Engine/Plugins/Media/ElectraPlayer/Source/ElectraPlayerRuntime/Private/Runtime/Player/DASH/PlaylistReaderDASH_Internal.h
+++ b/Engine/Plugins/Media/ElectraPlayer/Source/ElectraPlayerRuntime/Private/Runtime/Player/DASH/PlaylistReaderDASH_Internal.h
@@ -133,6 +133,13 @@ public:
 			FString CDN;
 			FString CustomHeader;
 		};
+		struct FInbandEventStream
+		{
+			FString SchemeIdUri;
+			FString Value;
+			int64 PTO = 0;
+			uint32 Timescale = 0;
+		};
 		FURL InitializationURL;
 		FURL MediaURL;
 		FTimeValue ATO;
@@ -151,6 +158,7 @@ public:
 		bool bMayBeMissing = false;					//!< true if the last segment in <SegmentTemplate> that might not exist.
 		bool bIsMissing = false;					//!< Set to true if known to be missing.
 		bool bSawLMSG = false;						//!< Will be set to true by the stream reader if the 'lmsg' brand was found.
+		TArray<FInbandEventStream> InbandEventStreams;
 
 		FTimeValue CalculateASAST(const FTimeValue& AST, const FTimeValue& PeriodStart, bool bIsStatic)
 		{
@@ -252,8 +260,9 @@ public:
 		ESearchResult FindSegment_Timeline(IPlayerSessionServices* PlayerSessionServices, FSegmentInformation& OutSegmentInfo, TArray<TWeakPtrTS<FMPDLoadRequestDASH>>& OutRemoteElementLoadRequests, const FSegmentSearchOption& InSearchOptions, const TSharedPtrTS<FDashMPD_RepresentationType>& MPDRepresentation, const TArray<TSharedPtrTS<FDashMPD_SegmentTemplateType>>& SegmentTemplate, const TSharedPtrTS<FDashMPD_SegmentTimelineType>& SegmentTimeline);
 
 		bool PrepareDownloadURLs(IPlayerSessionServices* PlayerSessionServices, FSegmentInformation& InOutSegmentInfo, const TArray<TSharedPtrTS<FDashMPD_SegmentBaseType>>& SegmentBase);
-		bool PrepareDownloadURLs(IPlayerSessionServices* PlayerSessionServices, FSegmentInformation& InOutSegmentInfo, const TArray<TSharedPtrTS<FDashMPD_SegmentTemplateType>>& SegmentBase);
+		bool PrepareDownloadURLs(IPlayerSessionServices* PlayerSessionServices, FSegmentInformation& InOutSegmentInfo, const TArray<TSharedPtrTS<FDashMPD_SegmentTemplateType>>& SegmentTemplate);
 		FString ApplyTemplateStrings(FString TemplateURL, const FSegmentInformation& InSegmentInfo);
+		void CollectInbandEventStreams(IPlayerSessionServices* PlayerSessionServices, FSegmentInformation& InOutSegmentInfo);
 
 		void SegmentIndexDownloadComplete(TSharedPtrTS<FMPDLoadRequestDASH> Request, bool bSuccess);
 		friend class FManifestDASHInternal;
@@ -392,6 +401,22 @@ public:
 			}
 			return nullptr;
 		}
+
+		void EndPresentationAt(const FTimeValue& EndsAt)
+		{
+			FTimeValue NewDur = EndsAt - Start;
+			if (NewDur >= FTimeValue::GetZero())
+			{
+				Duration = NewDur;
+				End = Start + NewDur;
+				TSharedPtrTS<FDashMPD_PeriodType> MPDPeriod = Period.Pin();
+				if (MPDPeriod.IsValid())
+				{
+					MPDPeriod->SetDuration(NewDur);
+				}
+			}
+		}
+
 /*
 		TSharedPtrTS<FRepresentation> GetRepresentationFromAdaptationSetIDs(const FString& AdaptationSetID, const FString& RepresentationID) const
 		{
@@ -562,6 +587,8 @@ public:
 
 	void PreparePeriodAdaptationSets(TSharedPtrTS<FPeriod> InPeriod, bool bRequestXlink);
 
+	void SendEventsFromAllPeriodEventStreams(TSharedPtrTS<FPeriod> InPeriod);
+
 	const TArray<FURL_RFC3986::FQueryParam>& GetURLFragmentComponents() const
 	{
 		return URLFragmentComponents;
@@ -599,6 +626,8 @@ public:
 	FTimeValue GetAvailabilityEndTime() const;
 	FTimeValue GetTimeshiftBufferDepth() const;
 
+	void EndPresentationAt(const FTimeValue& EndsAt, const FString& InPeriod);
+
 private:
 	FErrorDetail PrepareRemoteElementLoadRequest(TArray<TWeakPtrTS<FMPDLoadRequestDASH>>& OutRemoteElementLoadRequests, TWeakPtrTS<IDashMPDElement> ElementWithXLink, int64 RequestID);
 
diff --git a/Engine/Plugins/Media/ElectraPlayer/Source/ElectraPlayerRuntime/Private/Runtime/Player/DASH/StreamReaderFMP4DASH.cpp b/Engine/Plugins/Media/ElectraPlayer/Source/ElectraPlayerRuntime/Private/Runtime/Player/DASH/StreamReaderFMP4DASH.cpp
--- a/Engine/Plugins/Media/ElectraPlayer/Source/ElectraPlayerRuntime/Private/Runtime/Player/DASH/StreamReaderFMP4DASH.cpp
+++ b/Engine/Plugins/Media/ElectraPlayer/Source/ElectraPlayerRuntime/Private/Runtime/Player/DASH/StreamReaderFMP4DASH.cpp
@@ -12,6 +12,8 @@
 #include "Player/PlayerEntityCache.h"
 #include "Player/DASH/PlaylistReaderDASH.h"
 #include "Player/DASH/OptionKeynamesDASH.h"
+#include "Player/DASH/PlayerEventDASH.h"
+#include "Player/DASH/PlayerEventDASH_Internal.h"
 #include "Utilities/Utilities.h"
 
 #define INTERNAL_ERROR_INIT_SEGMENT_DOWNLOAD_ERROR					1
@@ -210,15 +212,8 @@ IStreamReader::EAddResult FStreamReaderFMP4DASH::AddRequest(uint32 CurrentPlayba
 		}
 
 		Request->SetPlaybackSequenceID(CurrentPlaybackSequenceID);
-		// Only add the request if this is not an EOS segment.
-		// A segment that is already at EOS is not meant to be loaded as it does not exist and there
-		// would not be another segment following it either. They are meant to indicate to the player
-		// that a stream has ended and will not be delivering any more data.
-		if (!Request->bIsEOSSegment)
-		{
-			Handler->CurrentRequest = Request;
-			Handler->SignalWork();
-		}
+		Handler->CurrentRequest = Request;
+		Handler->SignalWork();
 	}
 	return IStreamReader::EAddResult::Added;
 }
@@ -463,6 +458,112 @@ FErrorDetail FStreamReaderFMP4DASH::FStreamHandler::GetInitSegment(TSharedPtrTS<
 	}
 }
 
+void FStreamReaderFMP4DASH::FStreamHandler::CheckForInbandDASHEvents()
+{
+	bool bHasInbandEvent = false;
+	if (!CurrentRequest->bIsEOSSegment)
+	{
+		for(int32 i=0; i<CurrentRequest->Segment.InbandEventStreams.Num(); ++i)
+		{
+			const FManifestDASHInternal::FSegmentInformation::FInbandEventStream& ibs = CurrentRequest->Segment.InbandEventStreams[i];
+			if (ibs.SchemeIdUri.Equals(DASH::Schemes::ManifestEvents::Scheme_urn_mpeg_dash_event_2012))
+			{
+				bHasInbandEvent = true;
+				break;
+			}
+		}
+	}
+	TSharedPtrTS<IPlaylistReader> ManifestReader = PlayerSessionService->GetManifestReader();
+	IPlaylistReaderDASH* Reader = static_cast<IPlaylistReaderDASH*>(ManifestReader.Get());
+	if (Reader)
+	{
+		Reader->SetStreamInbandEventUsage(CurrentRequest->GetType(), bHasInbandEvent);
+	}
+}
+
+void FStreamReaderFMP4DASH::FStreamHandler::HandleEventMessages()
+{
+	// Are there 'emsg' boxes we need to handle?
+	if (MP4Parser->GetNumberOfEventMessages())
+	{
+		FTimeValue AbsPeriodStart = CurrentRequest->PeriodStart + CurrentRequest->AST + CurrentRequest->AdditionalAdjustmentTime + CurrentRequest->PlayerLoopState.LoopBasetime;
+		// We may need the EPT from the 'sidx' if there is one.
+		const IParserISO14496_12::ISegmentIndex* Sidx = MP4Parser->GetSegmentIndexByIndex(0);
+		for(int32 nEmsg=0, nEmsgs=MP4Parser->GetNumberOfEventMessages(); nEmsg<nEmsgs; ++nEmsg)
+		{
+			const IParserISO14496_12::IEventMessage* Emsg = MP4Parser->GetEventMessageByIndex(nEmsg);
+			// This event must be described by an <InbandEventStream> in order to be processed.
+			for(int32 i=0; i<CurrentRequest->Segment.InbandEventStreams.Num(); ++i)
+			{
+				const FManifestDASHInternal::FSegmentInformation::FInbandEventStream& ibs = CurrentRequest->Segment.InbandEventStreams[i];
+				if (ibs.SchemeIdUri.Equals(Emsg->GetSchemeIdUri()) &&
+					(ibs.Value.IsEmpty() || ibs.Value.Equals(Emsg->GetValue())))
+				{
+					TSharedPtrTS<DASH::FPlayerEvent> NewEvent = MakeSharedTS<DASH::FPlayerEvent>();
+					NewEvent->SetOrigin(IAdaptiveStreamingPlayerAEMSEvent::EOrigin::InbandEventStream);
+					NewEvent->SetSchemeIdUri(Emsg->GetSchemeIdUri());
+					NewEvent->SetValue(Emsg->GetValue());
+					NewEvent->SetID(LexToString(Emsg->GetID()));
+					uint32 Timescale = Emsg->GetTimescale();
+					uint32 Duration = Emsg->GetEventDuration();
+					FTimeValue PTS;
+					if (Emsg->GetVersion() == 0)
+					{
+						// Version 0 uses a presentation time delta relative to the EPT of the SIDX, if it exists, or if not
+						// to the PTS of the first AU, which should be the same as the segment media start time.
+						FTimeValue PTD((int64)Emsg->GetPresentationTimeDelta(), Timescale);
+						FTimeValue EPT;
+						if (Sidx)
+						{
+							EPT.SetFromND((int64)Sidx->GetEarliestPresentationTime(), Sidx->GetTimescale());
+						}
+						else
+						{
+							EPT.SetFromND((int64)CurrentRequest->Segment.Time, CurrentRequest->Segment.Timescale);
+						}
+						FTimeValue PTO(CurrentRequest->Segment.PTO, CurrentRequest->Segment.Timescale);
+						PTS = AbsPeriodStart - PTO + EPT + PTD;
+					}
+					else if (Emsg->GetVersion() == 1)
+					{
+						FTimeValue EventTime(Emsg->GetPresentationTime(), Timescale);
+						FTimeValue PTO(FTimeValue::GetZero());
+						PTS = AbsPeriodStart - PTO + EventTime;
+					}
+					NewEvent->SetPresentationTime(PTS);
+					if (Duration != 0xffffffff)
+					{
+						NewEvent->SetDuration(FTimeValue((int64)Duration, Timescale));
+					}
+					NewEvent->SetMessageData(Emsg->GetMessageData());
+					NewEvent->SetPeriodID(CurrentRequest->Period->GetUniqueIdentifier());
+					// Add the event to the handler.
+					if (PTS.IsValid())
+					{
+						/*
+							Check that we have no seen this event in this segment already. This is for the case where the 'emsg' appears
+							inbetween multiple 'moof' boxes. As per ISO/IEC 23009-1:2019 Section 5.10.3.3.1 General:
+								A Media Segment if based on the ISO BMFF container may contain one or more event message ('emsg') boxes. If present, any 'emsg' box shall be placed as follows:
+								- It may be placed before the first 'moof' box of the segment.
+								- It may be placed in between any 'mdat' and 'moof' box. In this case, an equivalent 'emsg' with the same id value shall be present before the first 'moof' box of any Segment.
+						*/
+						int32 EventIdx = SegmentEventsFound.IndexOfByPredicate([&NewEvent](const TSharedPtrTS<DASH::FPlayerEvent>& This) {
+							return NewEvent->GetSchemeIdUri().Equals(This->GetSchemeIdUri()) &&
+								   NewEvent->GetID().Equals(This->GetID()) &&
+								   (NewEvent->GetValue().IsEmpty() || NewEvent->GetValue().Equals(This->GetValue()));
+						});
+						if (EventIdx == INDEX_NONE)
+						{
+							PlayerSessionService->GetAEMSEventHandler()->AddEvent(NewEvent, CurrentRequest->Period->GetUniqueIdentifier(), IAdaptiveStreamingPlayerAEMSHandler::EEventAddMode::AddIfNotExists);
+							SegmentEventsFound.Emplace(MoveTemp(NewEvent));
+						}
+					}
+					break;
+				}
+			}
+		}
+	}
+}
 
 void FStreamReaderFMP4DASH::FStreamHandler::HandleRequest()
 {
@@ -510,6 +611,24 @@ void FStreamReaderFMP4DASH::FStreamHandler::HandleRequest()
 	ds.RetryNumber  	= Request->NumOverallRetries;
 	ds.ABRState.Reset();
 
+	// Clear out the list of events found the last time.
+	SegmentEventsFound.Empty();
+	CheckForInbandDASHEvents();
+	/*
+		Actually handle the request only if this is not an EOS segment.
+		A segment that is already at EOS is not meant to be loaded as it does not exist and there
+		would not be another segment following it either. They are meant to indicate to the player
+		that a stream has ended and will not be delivering any more data.
+		
+		NOTE:
+		  We had to handle this request up to detecting the use of inband event streams in order to
+		  signal that this stream has now ended and will not be receiving any further inband events!
+	*/
+	if (Request->bIsEOSSegment)
+	{
+		CurrentRequest.Reset();
+		return;
+	}
 
 	// Does this request have an availability window?
 	if (Request->ASAST.IsValid())
@@ -622,6 +741,8 @@ void FStreamReaderFMP4DASH::FStreamHandler::HandleRequest()
 							check(Track);
 							if (Track)
 							{
+								HandleEventMessages();
+
 								CSD->CodecSpecificData = Track->GetCodecSpecificData();
 								CSD->RawCSD = Track->GetCodecSpecificDataRAW();
 								CSD->ParsedInfo = Track->GetCodecInformation();
@@ -1044,6 +1165,26 @@ void FStreamReaderFMP4DASH::FStreamHandler::HandleRequest()
 		}
 	}
 
+	// If we failed to get the segment and there is an inband DASH event stream which triggers MPD events and
+	// we did not get such an event in the 'emsg' boxes, then we err on the safe side and assume this segment
+	// would have carried an MPD update event and fire an artificial event.
+	if (!ds.bWasSuccessful && CurrentRequest->Segment.InbandEventStreams.FindByPredicate([](const FManifestDASHInternal::FSegmentInformation::FInbandEventStream& This){ return This.SchemeIdUri.Equals(DASH::Schemes::ManifestEvents::Scheme_urn_mpeg_dash_event_2012); }))
+	{
+		if (SegmentEventsFound.IndexOfByPredicate([](const TSharedPtrTS<DASH::FPlayerEvent>& This){return This->GetSchemeIdUri().Equals(DASH::Schemes::ManifestEvents::Scheme_urn_mpeg_dash_event_2012);}) == INDEX_NONE)
+		{
+			TSharedPtrTS<DASH::FPlayerEvent> NewEvent = MakeSharedTS<DASH::FPlayerEvent>();
+			NewEvent->SetOrigin(IAdaptiveStreamingPlayerAEMSEvent::EOrigin::InbandEventStream);
+			NewEvent->SetSchemeIdUri(DASH::Schemes::ManifestEvents::Scheme_urn_mpeg_dash_event_2012);
+			NewEvent->SetValue(TEXT("1"));
+			NewEvent->SetID("$missed$");
+			FTimeValue EPT((int64)CurrentRequest->Segment.Time, CurrentRequest->Segment.Timescale);
+			FTimeValue PTO(CurrentRequest->Segment.PTO, CurrentRequest->Segment.Timescale);
+			NewEvent->SetPresentationTime(TimeOffset - PTO + EPT);
+			NewEvent->SetPeriodID(CurrentRequest->Period->GetUniqueIdentifier());
+			PlayerSessionService->GetAEMSEventHandler()->AddEvent(NewEvent, CurrentRequest->Period->GetUniqueIdentifier(), IAdaptiveStreamingPlayerAEMSHandler::EEventAddMode::AddIfNotExists);
+		}
+	}
+
 	StreamSelector->ReportDownloadEnd(ds);
 
 	// Remember the next expected timestamp.
@@ -1054,6 +1195,7 @@ void FStreamReaderFMP4DASH::FStreamHandler::HandleRequest()
 	CurrentRequest.Reset();
 	ReadBuffer.Reset();
 	MP4Parser.Reset();
+	SegmentEventsFound.Empty();
 
 	Parameters.EventListener->OnFragmentClose(FinishedRequest);
 }
diff --git a/Engine/Plugins/Media/ElectraPlayer/Source/ElectraPlayerRuntime/Private/Runtime/Player/DASH/StreamReaderFMP4DASH.h b/Engine/Plugins/Media/ElectraPlayer/Source/ElectraPlayerRuntime/Private/Runtime/Player/DASH/StreamReaderFMP4DASH.h
--- a/Engine/Plugins/Media/ElectraPlayer/Source/ElectraPlayerRuntime/Private/Runtime/Player/DASH/StreamReaderFMP4DASH.h
+++ b/Engine/Plugins/Media/ElectraPlayer/Source/ElectraPlayerRuntime/Private/Runtime/Player/DASH/StreamReaderFMP4DASH.h
@@ -5,6 +5,7 @@
 #include "Player/PlayerStreamReader.h"
 #include "Player/AdaptiveStreamingPlayerResourceRequest.h"
 #include "Player/DASH/PlaylistReaderDASH_Internal.h"
+#include "Player/DASH/PlayerEventDASH.h"
 #include "Demuxer/ParserISO14496-12.h"
 #include "HTTP/HTTPManager.h"
 
@@ -130,6 +131,8 @@ private:
 		FTimeValue 												DurationSuccessfullyRead;
 		FTimeValue 												DurationSuccessfullyDelivered;
 
+		TArray<TSharedPtrTS<DASH::FPlayerEvent>>				SegmentEventsFound;
+
 		FMediaCriticalSection									MetricUpdateLock;
 		int32													ProgressReportCount = 0;
 		TSharedPtrTS<IAdaptiveStreamSelector>					StreamSelector;
@@ -141,8 +144,10 @@ private:
 		void SignalWork();
 		void WorkerThread();
 		void HandleRequest();
-		
+
 		FErrorDetail GetInitSegment(TSharedPtrTS<const IParserISO14496_12>& OutMP4InitSegment, const TSharedPtrTS<FStreamSegmentRequestFMP4DASH>& InRequest);
+		void CheckForInbandDASHEvents();
+		void HandleEventMessages();
 
 		void LogMessage(IInfoLog::ELevel Level, const FString& Message);
 
diff --git a/Engine/Plugins/Media/ElectraPlayer/Source/ElectraPlayerRuntime/Private/Runtime/Player/PlayerSessionServices.h b/Engine/Plugins/Media/ElectraPlayer/Source/ElectraPlayerRuntime/Private/Runtime/Player/PlayerSessionServices.h
--- a/Engine/Plugins/Media/ElectraPlayer/Source/ElectraPlayerRuntime/Private/Runtime/Player/PlayerSessionServices.h
+++ b/Engine/Plugins/Media/ElectraPlayer/Source/ElectraPlayerRuntime/Private/Runtime/Player/PlayerSessionServices.h
@@ -18,6 +18,7 @@ namespace Electra
 	class IAdaptiveStreamingPlayerResourceProvider;
 	class IPlayerEntityCache;
 	class IPlaylistReader;
+	class IAdaptiveStreamingPlayerAEMSHandler;
 
 
 	class IPlayerMessage
@@ -115,6 +116,11 @@ namespace Electra
 		virtual	TSharedPtrTS<IPlaylistReader> GetManifestReader() = 0;
 
 
+		/**
+		 * Returns the "Application Event or Metadata Streams" (AEMS) handler.
+		 */
+		virtual IAdaptiveStreamingPlayerAEMSHandler* GetAEMSEventHandler() = 0;
+
 		/**
 		 * Returns the mutable player option dictionary. Other than initial options this serves as an interface to
 		 * pass values between internal sub systems.
diff --git a/Engine/Plugins/Media/ElectraPlayer/Source/ElectraPlayerRuntime/Private/Runtime/TimeUtilities.cpp b/Engine/Plugins/Media/ElectraPlayer/Source/ElectraPlayerRuntime/Private/Runtime/TimeUtilities.cpp
--- a/Engine/Plugins/Media/ElectraPlayer/Source/ElectraPlayerRuntime/Private/Runtime/TimeUtilities.cpp
+++ b/Engine/Plugins/Media/ElectraPlayer/Source/ElectraPlayerRuntime/Private/Runtime/TimeUtilities.cpp
@@ -569,6 +569,22 @@ namespace Electra
 	} // namespace RFC2326
 
 
+
+	namespace UnixEpoch
+	{
+		bool ParseFloatString(FTimeValue& OutTimeValue, const FString& Seconds)
+		{
+			FTimeFraction t;
+			if (t.SetFromFloatString(Seconds).IsValid())
+			{
+				OutTimeValue.SetFromTimeFraction(t);
+				return true;
+			}
+			return false;
+		}
+	}
+
+
 } // namespace Electra
 
 
diff --git a/Engine/Plugins/Media/ElectraPlayer/Source/ElectraPlayerRuntime/Private/Runtime/Utilities/StringHelpers.cpp b/Engine/Plugins/Media/ElectraPlayer/Source/ElectraPlayerRuntime/Private/Runtime/Utilities/StringHelpers.cpp
--- a/Engine/Plugins/Media/ElectraPlayer/Source/ElectraPlayerRuntime/Private/Runtime/Utilities/StringHelpers.cpp
+++ b/Engine/Plugins/Media/ElectraPlayer/Source/ElectraPlayerRuntime/Private/Runtime/Utilities/StringHelpers.cpp
@@ -121,6 +121,24 @@ namespace Electra
 			return FPlatformString::Strncmp(s1, s2, n) == 0; 
 		}
 
+		void StringToArray(TArray<uint8>& OutArray, const FString& InString)
+		{
+			FTCHARToUTF8 cnv(*InString);
+			int32 Len = cnv.Length();
+			OutArray.AddUninitialized(Len);
+			FMemory::Memcpy(OutArray.GetData(), cnv.Get(), Len);
+		}
+
+		FString ArrayToString(const TArray<uint8>& InArray)
+		{
+			FUTF8ToTCHAR cnv((const ANSICHAR*)InArray.GetData(), InArray.Num());
+			FString UTF8Text(cnv.Length(), cnv.Get());
+			//return (TCHAR*)FUTF8ToTCHAR((const ANSICHAR*)InArray.GetData(), InArray.Num()).Get();
+			//FString UTF8Text(InArray.Num(), (TCHAR*)FUTF8ToTCHAR((const ANSICHAR*)InArray.GetData(), InArray.Num()).Get());
+			return MoveTemp(UTF8Text);
+		}
+
+
 	} // namespace StringHelpers
 } // namespace Electra
 
diff --git a/Engine/Plugins/Media/ElectraPlayer/Source/ElectraPlayerRuntime/Private/Runtime/Utilities/StringHelpers.h b/Engine/Plugins/Media/ElectraPlayer/Source/ElectraPlayerRuntime/Private/Runtime/Utilities/StringHelpers.h
--- a/Engine/Plugins/Media/ElectraPlayer/Source/ElectraPlayerRuntime/Private/Runtime/Utilities/StringHelpers.h
+++ b/Engine/Plugins/Media/ElectraPlayer/Source/ElectraPlayerRuntime/Private/Runtime/Utilities/StringHelpers.h
@@ -24,7 +24,9 @@ namespace Electra
 
 		bool StringStartsWith(const TCHAR * const s1, const TCHAR * const s2, SIZE_T n);
 
+		void StringToArray(TArray<uint8>& OutArray, const FString& InString);
 
+		FString ArrayToString(const TArray<uint8>& InArray);
 
 		/**
 		 * There is a known anomaly in the FString::TConstIterator. It iterates all TCHARs in the string *including* the terminating zero character.
diff --git a/Engine/Plugins/Media/ElectraPlayer/Source/ElectraPlayerRuntime/Private/Runtime/Utilities/TimeUtilities.h b/Engine/Plugins/Media/ElectraPlayer/Source/ElectraPlayerRuntime/Private/Runtime/Utilities/TimeUtilities.h
--- a/Engine/Plugins/Media/ElectraPlayer/Source/ElectraPlayerRuntime/Private/Runtime/Utilities/TimeUtilities.h
+++ b/Engine/Plugins/Media/ElectraPlayer/Source/ElectraPlayerRuntime/Private/Runtime/Utilities/TimeUtilities.h
@@ -25,5 +25,10 @@ namespace Electra
 		bool ParseNPTTime(FTimeValue& OutTimeValue, const FString& NPTtime);
 	}
 
+	namespace UnixEpoch
+	{
+		bool ParseFloatString(FTimeValue& OutTimeValue, const FString& Seconds);
+	}
+
 } // namespace Electra
 
