ElectraPlayer: changed video resolution limit to allow for changes in content between DASH periods and adaptation sets

#rnx

diff --git a/Engine/Plugins/Media/ElectraPlayer/Source/ElectraPlayerPlugin/Private/ElectraPlayerPlugin.cpp b/Engine/Plugins/Media/ElectraPlayer/Source/ElectraPlayerPlugin/Private/ElectraPlayerPlugin.cpp
--- a/Engine/Plugins/Media/ElectraPlayer/Source/ElectraPlayerPlugin/Private/ElectraPlayerPlugin.cpp
+++ b/Engine/Plugins/Media/ElectraPlayer/Source/ElectraPlayerPlugin/Private/ElectraPlayerPlugin.cpp
@@ -374,18 +374,6 @@ bool FElectraPlayerPlugin::Open(const FString& Url, const IMediaOptions* Options
 		PlayerOptions.Set("seekable_range_live_end_offset_conservative", Electra::FVariantValue(bUseConservativeLiveEdgeDistance));
 		UE_LOG(LogElectraPlayerPlugin, Log, TEXT("[%p] IMediaPlayer::Open: Using conservative live edge for distance calculation"), this);
 	}
-	bool bDontAllowInsecureForMediaSegments = Options->GetMediaOption(TEXT("ElectraDontAllowInsecureForEncryptedMediaSegments"), (bool)false);
-	if (bDontAllowInsecureForMediaSegments)
-	{
-		PlayerOptions.Set("dont_use_insecure_for_media_segments", Electra::FVariantValue(bDontAllowInsecureForMediaSegments));
-		UE_LOG(LogElectraPlayerPlugin, Log, TEXT("[%p] IMediaPlayer::Open: Do not allow downgrade to http for encrypted media segments"), this);
-	}
-	bool bDontAllowInsecureForInitSegments = Options->GetMediaOption(TEXT("ElectraDontAllowInsecureForInitSegments"), (bool)false);
-	if (bDontAllowInsecureForInitSegments)
-	{
-		PlayerOptions.Set("dont_use_insecure_for_init_segments", Electra::FVariantValue(bDontAllowInsecureForInitSegments));
-		UE_LOG(LogElectraPlayerPlugin, Log, TEXT("[%p] IMediaPlayer::Open: Do not allow downgrade to http for init segments"), this);
-	}
 	bool bThrowErrorWhenRebuffering = Options->GetMediaOption(TEXT("ElectraThrowErrorWhenRebuffering"), (bool)false);
 	if (bThrowErrorWhenRebuffering)
 	{
diff --git a/Engine/Plugins/Media/ElectraPlayer/Source/ElectraPlayerRuntime/Private/Runtime/Player/AdaptiveStreamingPlayer.cpp b/Engine/Plugins/Media/ElectraPlayer/Source/ElectraPlayerRuntime/Private/Runtime/Player/AdaptiveStreamingPlayer.cpp
--- a/Engine/Plugins/Media/ElectraPlayer/Source/ElectraPlayerRuntime/Private/Runtime/Player/AdaptiveStreamingPlayer.cpp
+++ b/Engine/Plugins/Media/ElectraPlayer/Source/ElectraPlayerRuntime/Private/Runtime/Player/AdaptiveStreamingPlayer.cpp
@@ -8,6 +8,7 @@
 #include "Player/AdaptiveStreamingPlayerABR.h"
 #include "Player/PlayerLicenseKey.h"
 #include "Player/AdaptivePlayerOptionKeynames.h"
+#include "Utilities/Utilities.h"
 
 #include "HAL/LowLevelMemTracker.h"
 
@@ -92,7 +93,6 @@ FAdaptiveStreamingPlayer::FAdaptiveStreamingPlayer(const IAdaptiveStreamingPlaye
 	BitrateCeiling = 0;
 	VideoResolutionLimitWidth = 0;
 	VideoResolutionLimitHeight = 0;
-	MaxVideoTextureResolutionLimitHeight = 0;
 
 	TMediaInterlockedExchangePointer(PointerToLatestPlayer, this);
 
@@ -190,6 +190,60 @@ bool FAdaptiveStreamingPlayer::Initialize(const FParamDict& Options)
 	StreamSelector = IAdaptiveStreamSelector::Create(this, PlayerConfig.StreamSelectorConfig);
 	AddMetricsReceiver(StreamSelector.Get());
 
+	// Set up video decoder resolution limits. As the media playlists are parsed the video streams will be
+	// compared against these limits and those that exceed the limit will not be considered for playback.
+	
+	// Maximum allowed vertical resolution specified?
+	if (PlayerOptions.HaveKey("max_resoY"))
+	{
+		PlayerConfig.H264LimitUpto30fps.MaxResolution.Height = (int32) PlayerOptions.GetValue("max_resoY").GetInt64();
+		PlayerConfig.H264LimitAbove30fps.MaxResolution.Height = (int32) PlayerOptions.GetValue("max_resoY").GetInt64();
+	}
+	// A limit in vertical resolution for streams with more than 30fps?
+	if (PlayerOptions.HaveKey("max_resoY_above_30fps"))
+	{
+		PlayerConfig.H264LimitAbove30fps.MaxResolution.Height = (int32) PlayerOptions.GetValue("max_resoY_above_30fps").GetInt64();
+	}
+	// Note: We could add additional limits here if need be.
+	//       Eventually these need to be differentiated based on codec as well.
+
+	// Get global video decoder capabilities and if supported, set the stream resolution limit accordingly.
+	IVideoDecoderH264::FStreamDecodeCapability Capability, StreamParam;
+	// Do a one-time global capability check with a default-empty stream param structure.
+	// This then gets used in the individual stream capability checks.
+	if (IVideoDecoderH264::GetStreamDecodeCapability(Capability, StreamParam))
+	{
+		if (Capability.Profile && Capability.Level)
+		{
+			PlayerConfig.H264LimitUpto30fps.MaxTierProfileLevel.Profile = Capability.Profile;
+			PlayerConfig.H264LimitUpto30fps.MaxTierProfileLevel.Level = Capability.Level;
+			PlayerConfig.H264LimitAbove30fps.MaxTierProfileLevel.Profile = Capability.Profile;
+			PlayerConfig.H264LimitAbove30fps.MaxTierProfileLevel.Level = Capability.Level;
+		}
+		if (Capability.Height)
+		{
+			PlayerConfig.H264LimitUpto30fps.MaxResolution.Height = Utils::Min(PlayerConfig.H264LimitUpto30fps.MaxResolution.Height, Capability.Height);
+			PlayerConfig.H264LimitAbove30fps.MaxResolution.Height = Utils::Min(PlayerConfig.H264LimitAbove30fps.MaxResolution.Height, Capability.Height);
+		}
+		
+		// If this is software decoding only and there is limit for this in place on Windows, apply it.
+		if (PlayerOptions.HaveKey("max_resoY_windows_software") && Capability.DecoderSupportType == IVideoDecoderH264::FStreamDecodeCapability::ESupported::SoftwareOnly)
+		{
+			int32 MaxWinSWHeight = (int32) PlayerOptions.GetValue("max_resoY_windows_software").GetInt64();
+			PlayerConfig.H264LimitUpto30fps.MaxResolution.Height = Utils::Min(PlayerConfig.H264LimitUpto30fps.MaxResolution.Height, MaxWinSWHeight);
+			PlayerConfig.H264LimitAbove30fps.MaxResolution.Height = Utils::Min(PlayerConfig.H264LimitAbove30fps.MaxResolution.Height, MaxWinSWHeight);
+		}
+
+		// If the maximum fps is only up to 30 fps set the resolution for streams above 30fps so small
+		// that they will get rejected.
+		if (Capability.FPS > 0.0 && Capability.FPS <= 30.0)
+		{
+			PlayerConfig.H264LimitAbove30fps.MaxResolution.Height = 16;
+			PlayerConfig.H264LimitAbove30fps.MaxTierProfileLevel.Profile = 66;
+			PlayerConfig.H264LimitAbove30fps.MaxTierProfileLevel.Level = 10;
+		}
+	}
+
 	return true;
 }
 
@@ -668,21 +722,11 @@ int32 FAdaptiveStreamingPlayer::CreateInitialDecoder(EStreamType type)
 	{
 		if (VideoDecoder.Decoder == nullptr)
 		{
-			// Do we have a user defined fixed limited for vertical stream resolution?
-			MaxVideoTextureResolutionLimitHeight = 0;
-			if (PlayerOptions.HaveKey("max_resoY"))
-			{
-				MaxVideoTextureResolutionLimitHeight = (int32) PlayerOptions.GetValue("max_resoY").GetInt64();
-			}
-
-			// Get the largest stream resolution. We pass this to the decoder's Open() call to let it know what
-			// the largest resolution is it may need to decode at some point.
-		// FIXME: This is only correct for HLS and simple mp4 playback. For MPEG DASH multi period playback we need to look at all
-		//        periods and get the largest resolution across all of them. And even then, for a Live presentation where a new
-		//        period can be added at any point there is the possibility that this may be of a higher resolution than any period
-		//        encountered so far/
+			// Get the largest stream resolution of the currently selected video adaptation set.
+			// This is only an initial selection as there could be other adaptation sets in upcoming periods
+			// that have a larger resolution that is still within the allowed limits.
 			FStreamCodecInformation HighestStream;
-			FindMatchingStreamInfo(HighestStream, 0, MaxVideoTextureResolutionLimitHeight);
+			FindMatchingStreamInfo(HighestStream, 0, 0);
 
 			// Create H.264 video decoder
 			VideoDecoder.Decoder = IVideoDecoderH264::Create();
@@ -859,79 +903,30 @@ bool FAdaptiveStreamingPlayer::CanDecodeStream(const FStreamCodecInformation& In
 {
 	if (InStreamCodecInfo.IsVideoCodec())
 	{
-		if (InStreamCodecInfo.GetFrameRate().IsValid())
+		const AdaptiveStreamingPlayerConfig::FConfiguration::FVideoDecoderLimit* DecoderLimit = &PlayerConfig.H264LimitUpto30fps;
+		double Rate = InStreamCodecInfo.GetFrameRate().IsValid() ? InStreamCodecInfo.GetFrameRate().GetAsDouble() : 30.0;
+		if (Rate > 31.0)
 		{
-			// Check framerate related configuration options. For this the stream needs to have the framerate specified.
-			FTimeValue StreamFPS;
-			StreamFPS.SetFromTimeFraction(InStreamCodecInfo.GetFrameRate());
-			if (StreamFPS.IsValid())
-			{
-				// Required minimum fps?
-				if (PlayerOptions.HaveKey("min_stream_fps"))
-				{
-					FTimeValue MinFPS(FTimeValue().SetFromND(PlayerOptions.GetValue("min_stream_fps").GetInt64(), 1));
-					if (StreamFPS < MinFPS)
-					{
-						return false;
-					}
-				}
-				// Required maximum fps?
-				if (PlayerOptions.HaveKey("max_stream_fps"))
-				{
-					FTimeValue MaxFPS(FTimeValue().SetFromND(PlayerOptions.GetValue("max_stream_fps").GetInt64(), 1));
-					if (StreamFPS > MaxFPS)
-					{
-						return false;
-					}
-				}
-				// Stream fps faster than 30fps and a vertical resolution limit? (check against 31 fps in case of slight rounding errors)
-				if (PlayerOptions.HaveKey("max_resoY_above_30fps") && StreamFPS > FTimeValue().SetFromND(31,1))
-				{
-					int64 MaxVertReso = PlayerOptions.GetValue("max_resoY_above_30fps").GetInt64();
-
-					// Check if a maximum texture dimension limit is in place. If it is then it must never be exceeded.
-					if (MaxVideoTextureResolutionLimitHeight > 0 && MaxVertReso > MaxVideoTextureResolutionLimitHeight)
-					{
-						MaxVertReso = MaxVideoTextureResolutionLimitHeight;
-					}
-
-					if (InStreamCodecInfo.GetResolution().Height > MaxVertReso)
-					{
-						return false;
-					}
-				}
-			}
+			DecoderLimit = &PlayerConfig.H264LimitAbove30fps;
+		}
+		// Check against user configured resolution limit
+		if (DecoderLimit->MaxResolution.Height && InStreamCodecInfo.GetResolution().Height > DecoderLimit->MaxResolution.Height)
+		{
+			return false;
 		}
 
+		// Check against video decoder capabilities.
 		IVideoDecoderH264::FStreamDecodeCapability StreamParam, Capability;
-		// Do a one-time global capability check with a default-empty stream param structure.
-		// This then gets used in the individual stream capability checks.
+		StreamParam.Width = InStreamCodecInfo.GetResolution().Width;
+		StreamParam.Height = InStreamCodecInfo.GetResolution().Height;
+		StreamParam.Profile = InStreamCodecInfo.GetProfile();
+		StreamParam.Level = InStreamCodecInfo.GetProfileLevel();
+		StreamParam.FPS = Rate;
 		if (IVideoDecoderH264::GetStreamDecodeCapability(Capability, StreamParam))
 		{
-			// Manual vertical resolution selection?
-			if (PlayerOptions.HaveKey("max_resoY_windows_software"))
-			{
-				int64 MaxWinSWHeight = PlayerOptions.GetValue("max_resoY_windows_software").GetInt64();
-				// If the global capability indicates software decoding, check if this stream exceeds the vertical limit. If so, reject it.
-				if (MaxWinSWHeight > 0 && Capability.DecoderSupportType == IVideoDecoderH264::FStreamDecodeCapability::ESupported::SoftwareOnly && InStreamCodecInfo.GetResolution().Height > MaxWinSWHeight)
-				{
-					return false;
-				}
-			}
-			else
+			if (Capability.DecoderSupportType == IVideoDecoderH264::FStreamDecodeCapability::ESupported::NotSupported)
 			{
-				StreamParam.Width = InStreamCodecInfo.GetResolution().Width;
-				StreamParam.Height = InStreamCodecInfo.GetResolution().Height;
-				StreamParam.Profile = InStreamCodecInfo.GetProfile();
-				StreamParam.Level = InStreamCodecInfo.GetProfileLevel();
-				StreamParam.FPS = InStreamCodecInfo.GetFrameRate().IsValid() ? InStreamCodecInfo.GetFrameRate().GetAsDouble() : 30.0;
-				if (IVideoDecoderH264::GetStreamDecodeCapability(Capability, StreamParam))
-				{
-					if (Capability.DecoderSupportType == IVideoDecoderH264::FStreamDecodeCapability::ESupported::NotSupported)
-					{
-						return false;
-					}
-				}
+				return false;
 			}
 		}
 	}
@@ -1470,6 +1465,12 @@ void FAdaptiveStreamingPlayer::HandleMetadataChanges()
 }
 
 
+//-----------------------------------------------------------------------------
+/**
+ * Handles Application Event or Metadata Stream (AEMS) events triggering on current playback position.
+ *
+ * @param SessionMessage
+ */
 void FAdaptiveStreamingPlayer::HandleAEMSEvents()
 {
 	if (CurrentState == EPlayerState::eState_Playing)
@@ -2810,20 +2811,11 @@ bool FAdaptiveStreamingPlayer::FindMatchingStreamInfo(FStreamCodecInformation& O
 void FAdaptiveStreamingPlayer::UpdateStreamResolutionLimit()
 {
 	FStreamCodecInformation StreamInfo;
-	// Check if there is a maximum texture resolution limit in place that the video decoder was configured in the initial Open() call with.
-	// If it is then this is an absolute override over everything else.
-	if (MaxVideoTextureResolutionLimitHeight > 0 && VideoResolutionLimitHeight > MaxVideoTextureResolutionLimitHeight)
-	{
-		VideoResolutionLimitHeight = MaxVideoTextureResolutionLimitHeight;
-	}
-
 	if (FindMatchingStreamInfo(StreamInfo, VideoResolutionLimitWidth, VideoResolutionLimitHeight))
 	{
-		VideoResolutionLimitWidth  = StreamInfo.GetResolution().Width;
-		VideoResolutionLimitHeight = StreamInfo.GetResolution().Height;
 		if (VideoDecoder.Decoder)
 		{
-			VideoDecoder.Decoder->SetMaximumDecodeCapability(VideoResolutionLimitWidth, VideoResolutionLimitHeight, StreamInfo.GetProfile(), StreamInfo.GetProfileLevel(), StreamInfo.GetExtras());
+			VideoDecoder.Decoder->SetMaximumDecodeCapability(StreamInfo.GetResolution().Width, StreamInfo.GetResolution().Height, StreamInfo.GetProfile(), StreamInfo.GetProfileLevel(), StreamInfo.GetExtras());
 		}
 	}
 	StreamSelector->SetMaxVideoResolution(VideoResolutionLimitWidth, VideoResolutionLimitHeight);
diff --git a/Engine/Plugins/Media/ElectraPlayer/Source/ElectraPlayerRuntime/Private/Runtime/Player/AdaptiveStreamingPlayerInternal.h b/Engine/Plugins/Media/ElectraPlayer/Source/ElectraPlayerRuntime/Private/Runtime/Player/AdaptiveStreamingPlayerInternal.h
--- a/Engine/Plugins/Media/ElectraPlayer/Source/ElectraPlayerRuntime/Private/Runtime/Player/AdaptiveStreamingPlayerInternal.h
+++ b/Engine/Plugins/Media/ElectraPlayer/Source/ElectraPlayerRuntime/Private/Runtime/Player/AdaptiveStreamingPlayerInternal.h
@@ -1392,7 +1392,6 @@ private:
 	int32																BitrateCeiling;
 	int32																VideoResolutionLimitWidth;
 	int32																VideoResolutionLimitHeight;
-	int32																MaxVideoTextureResolutionLimitHeight;
 
 	FStreamBitrateInfo													CurrentVideoStreamBitrate;
 
diff --git a/Engine/Plugins/Media/ElectraPlayer/Source/ElectraPlayerRuntime/Private/Runtime/Player/AdaptiveStreamingPlayerInternalConfig.h b/Engine/Plugins/Media/ElectraPlayer/Source/ElectraPlayerRuntime/Private/Runtime/Player/AdaptiveStreamingPlayerInternalConfig.h
--- a/Engine/Plugins/Media/ElectraPlayer/Source/ElectraPlayerRuntime/Private/Runtime/Player/AdaptiveStreamingPlayerInternalConfig.h
+++ b/Engine/Plugins/Media/ElectraPlayer/Source/ElectraPlayerRuntime/Private/Runtime/Player/AdaptiveStreamingPlayerInternalConfig.h
@@ -14,8 +14,6 @@ namespace Electra
 
 	namespace AdaptiveStreamingPlayerConfig
 	{
-
-
 		struct FConfiguration
 		{
 			FConfiguration()
@@ -35,18 +33,24 @@ namespace Electra
 
 				bHoldLastFrameDuringSeek = true;
 
-				MaxPrerollDurationUntilForcedStart = 3.0;
-
 				InitialBufferMinTimeAvailBeforePlayback = 5.0;
 				SeekBufferMinTimeAvailBeforePlayback = 5.0;
 				RebufferMinTimeAvailBeforePlayback = 5.0;
-
-				// And apply the following scaling factors to determine the individual state buffering amounts
-				ScaleMPDInitialBufferMinTimeBeforePlayback = 1.0;
-				ScaleMPDSeekBufferMinTimeAvailBeforePlayback = 1.0;
-				ScaleMPDRebufferMinTimeAvailBeforePlayback = 1.0;
 			}
 
+			struct FVideoDecoderLimit
+			{
+				struct FTierProfileLevel
+				{
+					int32 Tier = 0;
+					int32 Profile = 0;
+					int32 Level = 0;
+				};
+				FStreamCodecInformation::FResolution	MaxResolution;
+				FTierProfileLevel						MaxTierProfileLevel;
+			};
+
+
 			FMediaRunnable::Param								WorkerThread;
 
 			FAccessUnitBuffer::FConfiguration					StreamBufferConfigVideo;
@@ -58,17 +62,14 @@ namespace Electra
 			IVideoDecoderH264::FInstanceConfiguration			DecoderCfg264;
 			IAudioDecoderAAC::FInstanceConfiguration			DecoderCfgAAC;
 
-			bool												bHoldLastFrameDuringSeek;
+			FVideoDecoderLimit									H264LimitUpto30fps;
+			FVideoDecoderLimit									H264LimitAbove30fps;
 
-			double												MaxPrerollDurationUntilForcedStart;
+			bool												bHoldLastFrameDuringSeek;
 
 			double												InitialBufferMinTimeAvailBeforePlayback;
 			double												SeekBufferMinTimeAvailBeforePlayback;
 			double												RebufferMinTimeAvailBeforePlayback;
-
-			double												ScaleMPDInitialBufferMinTimeBeforePlayback;
-			double												ScaleMPDSeekBufferMinTimeAvailBeforePlayback;
-			double												ScaleMPDRebufferMinTimeAvailBeforePlayback;
 		};
 
 	} // namespace AdaptiveStreamingPlayerConfig
diff --git a/Engine/Plugins/Media/ElectraPlayer/Source/ElectraPlayerRuntime/Private/Runtime/Player/AdaptiveStreamingPlayerManifest.cpp b/Engine/Plugins/Media/ElectraPlayer/Source/ElectraPlayerRuntime/Private/Runtime/Player/AdaptiveStreamingPlayerManifest.cpp
--- a/Engine/Plugins/Media/ElectraPlayer/Source/ElectraPlayerRuntime/Private/Runtime/Player/AdaptiveStreamingPlayerManifest.cpp
+++ b/Engine/Plugins/Media/ElectraPlayer/Source/ElectraPlayerRuntime/Private/Runtime/Player/AdaptiveStreamingPlayerManifest.cpp
@@ -231,9 +231,9 @@ bool FAdaptiveStreamingPlayer::SelectManifest()
 			CurrentState = EPlayerState::eState_Ready;
 
 			double minBufTimeMPD = Manifest->GetMinBufferTime().GetAsSeconds();
-			PlayerConfig.InitialBufferMinTimeAvailBeforePlayback = Utils::Min(minBufTimeMPD * PlayerConfig.ScaleMPDInitialBufferMinTimeBeforePlayback,   PlayerConfig.InitialBufferMinTimeAvailBeforePlayback);
-			PlayerConfig.SeekBufferMinTimeAvailBeforePlayback    = Utils::Min(minBufTimeMPD * PlayerConfig.ScaleMPDSeekBufferMinTimeAvailBeforePlayback, PlayerConfig.SeekBufferMinTimeAvailBeforePlayback);
-			PlayerConfig.RebufferMinTimeAvailBeforePlayback 	 = Utils::Min(minBufTimeMPD * PlayerConfig.ScaleMPDRebufferMinTimeAvailBeforePlayback,   PlayerConfig.RebufferMinTimeAvailBeforePlayback);
+			PlayerConfig.InitialBufferMinTimeAvailBeforePlayback = Utils::Min(minBufTimeMPD, PlayerConfig.InitialBufferMinTimeAvailBeforePlayback);
+			PlayerConfig.SeekBufferMinTimeAvailBeforePlayback    = Utils::Min(minBufTimeMPD, PlayerConfig.SeekBufferMinTimeAvailBeforePlayback);
+			PlayerConfig.RebufferMinTimeAvailBeforePlayback 	 = Utils::Min(minBufTimeMPD, PlayerConfig.RebufferMinTimeAvailBeforePlayback);
 
 			return true;
 		}
