Fixuing Electra crash with Vulkan rendering device used on Windows
#rnx

diff --git a/Engine/Plugins/Media/ElectraPlayer/Source/ElectraPlayerRuntime/Private/Runtime/Decoder/DX/VideoDecoderH264_DX.cpp b/Engine/Plugins/Media/ElectraPlayer/Source/ElectraPlayerRuntime/Private/Runtime/Decoder/DX/VideoDecoderH264_DX.cpp
--- a/Engine/Plugins/Media/ElectraPlayer/Source/ElectraPlayerRuntime/Private/Runtime/Decoder/DX/VideoDecoderH264_DX.cpp
+++ b/Engine/Plugins/Media/ElectraPlayer/Source/ElectraPlayerRuntime/Private/Runtime/Decoder/DX/VideoDecoderH264_DX.cpp
@@ -459,6 +459,9 @@ bool FVideoDecoderH264::FallbackToSwDecoding(FString Reason)
 		// Once multithread protection is enabled we don't disable it, so UE4's rendering device stays protected for the rest of its lifetime.
 		// Some other system could enable multithread protection after we did it, we have no means to know about this, and so disabling it
 		// at the end of playback can cause GPU driver crash
+		//
+		// The DX11 rendering device will also come in as null if DX12 or a non-DX rendering API is used on the PC.
+		//
 		if (Electra::FDXDeviceInfo::s_DXDeviceInfo->RenderingDx11Device)
 		{
 			HRESULT res;
diff --git a/Engine/Plugins/Media/ElectraPlayer/Source/ElectraPlayerRuntime/Private/Runtime/Decoder/Windows/VideoDecoderH264_PC.cpp b/Engine/Plugins/Media/ElectraPlayer/Source/ElectraPlayerRuntime/Private/Runtime/Decoder/Windows/VideoDecoderH264_PC.cpp
--- a/Engine/Plugins/Media/ElectraPlayer/Source/ElectraPlayerRuntime/Private/Runtime/Decoder/Windows/VideoDecoderH264_PC.cpp
+++ b/Engine/Plugins/Media/ElectraPlayer/Source/ElectraPlayerRuntime/Private/Runtime/Decoder/Windows/VideoDecoderH264_PC.cpp
@@ -165,7 +165,9 @@ void FElectraPlayerVideoDecoderOutputPC::PreInitForDecode(FIntPoint OutputDim, c
 
 		SampleDim = Dim;
 
-		if (Electra::IsWindows8Plus() && (Electra::FDXDeviceInfo::s_DXDeviceInfo->DxVersion != Electra::FDXDeviceInfo::ED3DVersion::Version12Win10))
+		// Windows 8+ and a valid DX11 rendering device?
+		// (Rendering device will be null if the rendering device is not DX11 (DX12, Vulkan...))
+		if (Electra::IsWindows8Plus() && Electra::FDXDeviceInfo::s_DXDeviceInfo->RenderingDx11Device)
 		{
 			// Software decode into DX11 texture
 
