ElectraPlayer: Truncate segment data when last segment in a period is used only partially due to the next period cutting into it.

#rnx

diff --git a/Engine/Plugins/Media/ElectraPlayer/Source/ElectraPlayerRuntime/Private/Runtime/Decoder/Android/VideoDecoderH264_Android.cpp b/Engine/Plugins/Media/ElectraPlayer/Source/ElectraPlayerRuntime/Private/Runtime/Decoder/Android/VideoDecoderH264_Android.cpp
--- a/Engine/Plugins/Media/ElectraPlayer/Source/ElectraPlayerRuntime/Private/Runtime/Decoder/Android/VideoDecoderH264_Android.cpp
+++ b/Engine/Plugins/Media/ElectraPlayer/Source/ElectraPlayerRuntime/Private/Runtime/Decoder/Android/VideoDecoderH264_Android.cpp
@@ -1449,6 +1449,7 @@ void FVideoDecoderH264::WorkerThread()
 
 	int32 result = 0;
 	bool bDone  = false;
+	bool bGotLastSequenceAU = false;
 
 	bError = false;
 
@@ -1614,7 +1615,8 @@ void FVideoDecoderH264::WorkerThread()
 						}
 					}
 
-					bool bStreamFormatChanged = CurrentStreamFormatInfo.IsDifferentFrom(CurrentAccessUnit);
+					bool bStreamFormatChanged = CurrentStreamFormatInfo.IsDifferentFrom(CurrentAccessUnit) || bGotLastSequenceAU;
+					bGotLastSequenceAU = CurrentAccessUnit->bIsLastInPeriod;
 
 					// A new AU following dummy data?
 					bool bPrevWasDummy = false;
diff --git a/Engine/Plugins/Media/ElectraPlayer/Source/ElectraPlayerRuntime/Private/Runtime/Decoder/Apple/VideoDecoderH264_Apple.cpp b/Engine/Plugins/Media/ElectraPlayer/Source/ElectraPlayerRuntime/Private/Runtime/Decoder/Apple/VideoDecoderH264_Apple.cpp
--- a/Engine/Plugins/Media/ElectraPlayer/Source/ElectraPlayerRuntime/Private/Runtime/Decoder/Apple/VideoDecoderH264_Apple.cpp
+++ b/Engine/Plugins/Media/ElectraPlayer/Source/ElectraPlayerRuntime/Private/Runtime/Decoder/Apple/VideoDecoderH264_Apple.cpp
@@ -1294,6 +1294,7 @@ FVideoDecoderH264::EDecodeResult FVideoDecoderH264::Decode(FAccessUnit* AccessUn
 void FVideoDecoderH264::WorkerThread()
 {
 	bool bDone  = false;
+    bool bGotLastSequenceAU = false;
 
 	bError = false;
 
@@ -1385,8 +1386,9 @@ void FVideoDecoderH264::WorkerThread()
 					ReplayAccessUnitBuffer.Flush();
 				}
 
-				bool bStreamFormatChanged = CurrentStreamFormatInfo.IsDifferentFrom(AccessUnit);
+				bool bStreamFormatChanged = CurrentStreamFormatInfo.IsDifferentFrom(AccessUnit) || bGotLastSequenceAU;
 				bool bNeedNewDecoder = false;
+                bGotLastSequenceAU = AccessUnit->bIsLastInPeriod;
 				if (bStreamFormatChanged || !DecoderHandle)
 				{
 					CMFormatDescriptionRef NewFormatDescr = nullptr;
diff --git a/Engine/Plugins/Media/ElectraPlayer/Source/ElectraPlayerRuntime/Private/Runtime/Decoder/DX/VideoDecoderH264_DX.cpp b/Engine/Plugins/Media/ElectraPlayer/Source/ElectraPlayerRuntime/Private/Runtime/Decoder/DX/VideoDecoderH264_DX.cpp
--- a/Engine/Plugins/Media/ElectraPlayer/Source/ElectraPlayerRuntime/Private/Runtime/Decoder/DX/VideoDecoderH264_DX.cpp
+++ b/Engine/Plugins/Media/ElectraPlayer/Source/ElectraPlayerRuntime/Private/Runtime/Decoder/DX/VideoDecoderH264_DX.cpp
@@ -1296,6 +1296,7 @@ void FVideoDecoderH264::WorkerThread()
 
 	bool bDone = false;
 	bool bInDummyDecodeMode = false;
+	bool bGotLastSequenceAU = false;
 
 	// Require a new media input type based on the actual first access unit.
 	bool bNeedInitialReconfig = true;
@@ -1354,7 +1355,7 @@ void FVideoDecoderH264::WorkerThread()
 					if (!bError)
 					{
 						// Update the buffer acquisition after a change in streams?
-						bool bFormatChangedJustNow = CurrentSampleInfo.IsDifferentFromOtherVideo(NewSampleInfo);
+						bool bFormatChangedJustNow = CurrentSampleInfo.IsDifferentFromOtherVideo(NewSampleInfo) || bGotLastSequenceAU;
 
 						if (bInDummyDecodeMode)
 						{
@@ -1383,7 +1384,7 @@ void FVideoDecoderH264::WorkerThread()
 							SetupBufferAcquisitionProperties();
 						}
 						bNeedInitialReconfig = false;
-
+						bGotLastSequenceAU = CurrentAccessUnit->bIsLastInPeriod;
 						if (!Decode(CurrentAccessUnit, false))
 						{
 							bError = true;
diff --git a/Engine/Plugins/Media/ElectraPlayer/Source/ElectraPlayerRuntime/Private/Runtime/Player/DASH/StreamReaderFMP4DASH.cpp b/Engine/Plugins/Media/ElectraPlayer/Source/ElectraPlayerRuntime/Private/Runtime/Player/DASH/StreamReaderFMP4DASH.cpp
--- a/Engine/Plugins/Media/ElectraPlayer/Source/ElectraPlayerRuntime/Private/Runtime/Player/DASH/StreamReaderFMP4DASH.cpp
+++ b/Engine/Plugins/Media/ElectraPlayer/Source/ElectraPlayerRuntime/Private/Runtime/Player/DASH/StreamReaderFMP4DASH.cpp
@@ -726,6 +726,7 @@ void FStreamReaderFMP4DASH::FStreamHandler::HandleRequest()
 			uint32 TrackTimescale = 0;
 			bool bDone = false;
 			bool bIsFirstAU = true;
+			bool bReadPastLastPTS = false;
 
 			// Encrypted?
 			TSharedPtr<ElectraCDM::IMediaCDMDecrypter, ESPMode::ThreadSafe> Decrypter;
@@ -918,6 +919,27 @@ void FStreamReaderFMP4DASH::FStreamHandler::HandleRequest()
 										break;
 									}
 
+									// Check if the AU is outside the time range we are allowed to read.
+									// The last one (the one that is already outside the range, actually) is tagged as such and sent into the buffer.
+									// The respective decoder has to handle this flag if necessary and/or drop the AU.
+									if (AccessUnit)
+									{
+										// Already sent the last one?
+										if (bReadPastLastPTS)
+										{
+											// Yes. Release this AU and do not forward it. Continue reading however.
+											FAccessUnit::Release(AccessUnit);
+											AccessUnit = nullptr;
+										}
+										else if ((AccessUnit->DropState & (FAccessUnit::EDropState::PtsTooLate | FAccessUnit::EDropState::DtsTooLate)) == (FAccessUnit::EDropState::PtsTooLate | FAccessUnit::EDropState::DtsTooLate))
+										{
+											// Tag the last one and send it off, but stop doing so for the remainder of the segment.
+											// Note: we continue reading this segment all the way to the end on purpose in case there are further 'emsg' boxes.
+											AccessUnit->bIsLastInPeriod = true;
+											bReadPastLastPTS = true;
+										}
+									}
+
 									if (AccessUnit)
 									{
 										AccessUnitFIFO.Push(AccessUnit);
diff --git a/Engine/Plugins/Media/ElectraPlayer/Source/ElectraPlayerRuntime/Private/Runtime/StreamAccessUnitBuffer.h b/Engine/Plugins/Media/ElectraPlayer/Source/ElectraPlayerRuntime/Private/Runtime/StreamAccessUnitBuffer.h
--- a/Engine/Plugins/Media/ElectraPlayer/Source/ElectraPlayerRuntime/Private/Runtime/StreamAccessUnitBuffer.h
+++ b/Engine/Plugins/Media/ElectraPlayer/Source/ElectraPlayerRuntime/Private/Runtime/StreamAccessUnitBuffer.h
@@ -68,6 +68,7 @@ namespace Electra
 		TSharedPtrTS<CodecData>		AUCodecData;					//!< If set, points to sideband data for this access unit.
 		uint8						DropState;						//!< If set this access unit is not to be rendered. If possible also not to be decoded.
 		bool						bIsFirstInSequence;				//!< true for the first AU in a segment
+		bool						bIsLastInPeriod;				//!< true if this is the last AU in the playing period.
 		bool						bIsSyncSample;					//!< true if this is a sync sample (keyframe)
 		bool						bIsDummyData;					//!< True if the decoder must be reset before decoding this AU.
 		bool						bTrackChangeDiscontinuity;		//!< True if this is the first AU after a track change.
@@ -155,6 +156,7 @@ namespace Electra
 			AUMemoryProvider = nullptr;
 			DropState = EDropState::None;
 			bIsFirstInSequence = false;
+			bIsLastInPeriod = false;
 			bIsSyncSample = false;
 			bIsDummyData = false;
 			bTrackChangeDiscontinuity = false;
diff --git a/Engine/Plugins/Media/ElectraPlayer/Source/ElectraPlayerRuntime/Private/Runtime/Utilities/UtilsMPEGAudio.cpp b/Engine/Plugins/Media/ElectraPlayer/Source/ElectraPlayerRuntime/Private/Runtime/Utilities/UtilsMPEGAudio.cpp
--- a/Engine/Plugins/Media/ElectraPlayer/Source/ElectraPlayerRuntime/Private/Runtime/Utilities/UtilsMPEGAudio.cpp
+++ b/Engine/Plugins/Media/ElectraPlayer/Source/ElectraPlayerRuntime/Private/Runtime/Utilities/UtilsMPEGAudio.cpp
@@ -12,7 +12,7 @@ namespace Electra
 		{
 			static uint32 GIndexToSampleRate[16] =
 			{
-			96000, 88200, 64000, 48000, 44100, 32000, 24000, 22050, 16000, 12000, 11025, 8000, 0, 0, 0, 0
+				96000, 88200, 64000, 48000, 44100, 32000, 24000, 22050, 16000, 12000, 11025, 8000, 0, 0, 0, 0
 			};
 
 			static uint32 GetAudioObjectType(FBitDataStream& Bitstream)
@@ -158,7 +158,9 @@ namespace Electra
 							}
 						}
 						if (bsp.GetRemainingBits() > 11 && bsp.GetBits(11) == 0x548)
+						{
 							PSSignal = bsp.GetBits(1);
+						}
 						break;
 					}
 					else
@@ -172,11 +174,6 @@ namespace Electra
 			MEDIA_UNUSED_VAR(remainingBits);
 			return true;
 		}
-
-
-
-
-
 	} // namespace MPEG
 } // namespace Electra
 
@@ -206,7 +203,9 @@ struct ADTSheader
 static bool GetADTSHeader(ADTSheader& header, const void* pData, int64 nDataSize)
 {
 	if (nDataSize < 7)
-		return(false);
+	{
+		return false;
+	}
 
 	FBitDataStream bs(pData, nDataSize);
 	header.mSyncWord = bs.GetBits(12);
@@ -225,8 +224,10 @@ static bool GetADTSHeader(ADTSheader& header, const void* pData, int64 nDataSize
 	header.mBufferFullness = bs.GetBits(11);
 	header.mNumFrames = bs.GetBits(2);
 	if (nDataSize >= 9 && !header.mProtAbsent)
+	{
 		header.mCRCifPresent = bs.GetBits(16);
-	return(header.mSyncWord == 0xfff && header.mLayer == 0);
+	}
+	return header.mSyncWord == 0xfff && header.mLayer == 0;
 }
 
 /*
