ElectraPlayer: fixed an overflow converting a string representing a double value into an internal time when the double has too many fractional digits

#rnx

diff --git a/Engine/Plugins/Media/ElectraPlayer/Source/ElectraPlayerRuntime/Private/Runtime/PlayerTime.cpp b/Engine/Plugins/Media/ElectraPlayer/Source/ElectraPlayerRuntime/Private/Runtime/PlayerTime.cpp
--- a/Engine/Plugins/Media/ElectraPlayer/Source/ElectraPlayerRuntime/Private/Runtime/PlayerTime.cpp
+++ b/Engine/Plugins/Media/ElectraPlayer/Source/ElectraPlayerRuntime/Private/Runtime/PlayerTime.cpp
@@ -122,7 +122,8 @@ namespace Electra
 				if (last0 != INDEX_NONE)
 				{
 					frc.MidInline(0, last0 + 1);
-					for (int32 i = 0; i < frc.Len(); ++i)
+					// Convert at most 7 fractional digits (giving us hundreds of nanoseconds (HNS))
+					for(int32 i = 0; i < frc.Len() && i<7; ++i)
 					{
 						Numerator = Numerator * 10 + (frc[i] - TCHAR('0'));
 						Denominator *= 10;
