ElectraPlayer: Adding player events to the metadata sample queue as they are received. Need to add a metadata sample sink with the facade to receive them. IMediaBinarySample now has a method to return a GUID to help identify the type of payload. Can be used to downcast the sample to a derived class (knowing its guid) that may offer additional methods for that type of metadata.

#rnx

diff --git a/Engine/Plugins/Media/ElectraPlayer/Source/ElectraPlayerPlugin/Private/ElectraPlayerPlugin.cpp b/Engine/Plugins/Media/ElectraPlayer/Source/ElectraPlayerPlugin/Private/ElectraPlayerPlugin.cpp
--- a/Engine/Plugins/Media/ElectraPlayer/Source/ElectraPlayerPlugin/Private/ElectraPlayerPlugin.cpp
+++ b/Engine/Plugins/Media/ElectraPlayer/Source/ElectraPlayerPlugin/Private/ElectraPlayerPlugin.cpp
@@ -9,6 +9,7 @@
 #include "ElectraPlayerPlugin.h"
 #include "IElectraPlayerRuntimeModule.h"
 #include "IElectraPlayerPluginModule.h"
+#include "IElectraMetadataSample.h"
 
 using namespace Electra;
 
@@ -225,6 +226,14 @@ void FElectraPlayerPlugin::FPlayerAdapterDelegate::PresentAudioFrame(const IAudi
 	}
 }
 
+void FElectraPlayerPlugin::FPlayerAdapterDelegate::PresentMetadataSample(const IElectraBinarySampleRef& InMetadataSample)
+{
+	if (TSharedPtr<FElectraPlayerPlugin, ESPMode::ThreadSafe> PinnedHost = Host.Pin())
+	{
+		PinnedHost->MediaSamples->AddMetadata(InMetadataSample);
+	}
+}
+
 
 bool FElectraPlayerPlugin::FPlayerAdapterDelegate::CanReceiveVideoSamples(int32 NumFrames)
 {
diff --git a/Engine/Plugins/Media/ElectraPlayer/Source/ElectraPlayerPlugin/Private/ElectraPlayerPlugin.h b/Engine/Plugins/Media/ElectraPlayer/Source/ElectraPlayerPlugin/Private/ElectraPlayerPlugin.h
--- a/Engine/Plugins/Media/ElectraPlayer/Source/ElectraPlayerPlugin/Private/ElectraPlayerPlugin.h
+++ b/Engine/Plugins/Media/ElectraPlayer/Source/ElectraPlayerPlugin/Private/ElectraPlayerPlugin.h
@@ -20,6 +20,8 @@
 
 class IElectraPlayerRuntimeModule;
 class FElectraPlayerResourceDelegate;
+class IElectraBinarySample;
+
 
 DECLARE_LOG_CATEGORY_EXTERN(LogElectraPlayerPlugin, Log, All);
 
@@ -163,6 +165,7 @@ private:
 		void OnAudioFlush() override;
 		void PresentVideoFrame(const FVideoDecoderOutputPtr& InVideoFrame) override;
 		void PresentAudioFrame(const IAudioDecoderOutputPtr& InAudioFrame) override;
+		void PresentMetadataSample(const IElectraBinarySampleRef& InMetadataSample) override;
 		bool CanReceiveVideoSamples(int32 NumFrames) override;
 		bool CanReceiveAudioSamples(int32 NumFrames) override;
 		void PrepareForDecoderShutdown() override;
diff --git a/Engine/Plugins/Media/ElectraPlayer/Source/ElectraPlayerRuntime/ElectraPlayerRuntime.Build.cs b/Engine/Plugins/Media/ElectraPlayer/Source/ElectraPlayerRuntime/ElectraPlayerRuntime.Build.cs
--- a/Engine/Plugins/Media/ElectraPlayer/Source/ElectraPlayerRuntime/ElectraPlayerRuntime.Build.cs
+++ b/Engine/Plugins/Media/ElectraPlayer/Source/ElectraPlayerRuntime/ElectraPlayerRuntime.Build.cs
@@ -21,7 +21,9 @@ namespace UnrealBuildTool.Rules
 					"HTTP",
 					"WebSockets",
 					"Json",
+					"Media",
 					"ElectraBase",
+					"ElectraSamples",
 					"ElectraCDM",
 					"XmlParser"
 				});
diff --git a/Engine/Plugins/Media/ElectraPlayer/Source/ElectraPlayerRuntime/Private/ElectraPlayer.cpp b/Engine/Plugins/Media/ElectraPlayer/Source/ElectraPlayerRuntime/Private/ElectraPlayer.cpp
--- a/Engine/Plugins/Media/ElectraPlayer/Source/ElectraPlayerRuntime/Private/ElectraPlayer.cpp
+++ b/Engine/Plugins/Media/ElectraPlayer/Source/ElectraPlayerRuntime/Private/ElectraPlayer.cpp
@@ -12,6 +12,8 @@
 #include "VideoDecoderResourceDelegate.h"
 #include "Utilities/Utilities.h"
 
+#include "IElectraMetadataSample.h"
+
 #include "Async/Async.h"
 
 #include "CoreGlobals.h"
@@ -72,6 +74,38 @@ static FString SanitizeMessage(FString InMessage)
 #endif
 }
 
+
+//-----------------------------------------------------------------------------
+
+class FElectraBinarySample : public IElectraBinarySample
+{
+public:
+	virtual ~FElectraBinarySample() = default;
+	virtual const void* GetData() override						{ return Data.GetData(); }
+	virtual FTimespan GetDuration() const override				{ return Duration; }
+	virtual uint32 GetSize() const override						{ return (uint32) Data.Num(); }
+	virtual FMediaTimeStamp GetTime() const override			{ return PresentationTime; }
+	virtual FGuid GetGUID() const override						{ return IElectraBinarySample::GetSampleTypeGUID(); }
+	virtual EOrigin GetOrigin() const override					{ return Origin; }
+	virtual EDispatchedMode GetDispatchedMode() const override	{ return DispatchedMode; }
+	virtual FString GetSchemeIdUri() const override				{ return SchemeIdUri; }
+	virtual FString GetValue() const override					{ return Value; }
+	virtual FString GetID() const override						{ return ID; }
+
+	TArray<uint8> Data;
+	FMediaTimeStamp PresentationTime;
+	FTimespan Duration;
+	EOrigin Origin;
+	EDispatchedMode DispatchedMode;
+	FString SchemeIdUri;
+	FString Value;
+	FString ID;
+};
+
+//-----------------------------------------------------------------------------
+
+
+
 //-----------------------------------------------------------------------------
 /**
  *	Construction of new player
@@ -209,10 +243,10 @@ bool FElectraPlayer::OpenInternal(const FString& Url, const FParamDict & PlayerO
 	NewPlayer->AdaptivePlayer->SetStaticResourceProviderCallback(StaticResourceProvider);
 	NewPlayer->AdaptivePlayer->SetVideoDecoderResourceDelegate(VideoDecoderResourceDelegate);
 
-	// Create a new media player event receiver and register it.
+	// Create a new media player event receiver and register it to receive all non player internal events as soon as they are received.
 	MediaPlayerEventReceiver = MakeSharedTS<FAEMSEventReceiver>();
 	MediaPlayerEventReceiver->GetEventReceivedDelegate().BindRaw(this, &FElectraPlayer::OnMediaPlayerEventReceived);
-	NewPlayer->AdaptivePlayer->AddAEMSReceiver(MediaPlayerEventReceiver, TEXT("*"), TEXT(""), IAdaptiveStreamingPlayerAEMSReceiver::EDispatchMode::OnStart);
+	NewPlayer->AdaptivePlayer->AddAEMSReceiver(MediaPlayerEventReceiver, TEXT("*"), TEXT(""), IAdaptiveStreamingPlayerAEMSReceiver::EDispatchMode::OnReceive);
 
 	if (!NewPlayer->AdaptivePlayer->Initialize(PlayerOptions))
 	{
@@ -1026,12 +1060,40 @@ bool FElectraPlayer::SelectTrack(EPlayerTrackType TrackType, int32 TrackIndex)
 
 void FElectraPlayer::OnMediaPlayerEventReceived(TSharedPtrTS<IAdaptiveStreamingPlayerAEMSEvent> InEvent, IAdaptiveStreamingPlayerAEMSReceiver::EDispatchMode InDispatchMode)
 {
+#if !UE_BUILD_SHIPPING
 	const TCHAR* const Origins[] = { TEXT("Playlist"), TEXT("Inband"), TEXT("TimedMetadata"), TEXT("???") };
 	UE_LOG(LogElectraPlayer, Log, TEXT("[%p][%p] %s event %s with \"%s\", \"%s\", \"%s\" PTS @ %.3f for %.3fs"), this, CurrentPlayer.Get(),
 		Origins[Electra::Utils::Min((int32)InEvent->GetOrigin(), (int32)UE_ARRAY_COUNT(Origins)-1)],
 		InDispatchMode==IAdaptiveStreamingPlayerAEMSReceiver::EDispatchMode::OnReceive?TEXT("received"):TEXT("started"),
 		*InEvent->GetSchemeIdUri(), *InEvent->GetValue(), *InEvent->GetID(),
 		InEvent->GetPresentationTime().GetAsSeconds(), InEvent->GetDuration().GetAsSeconds());
+#endif
+
+	if (TSharedPtr<IElectraPlayerAdapterDelegate, ESPMode::ThreadSafe> PinnedAdapterDelegate = AdapterDelegate.Pin())
+	{
+		// Create a binary media sample of our extended format and pass it up.
+		TSharedRef<FElectraBinarySample, ESPMode::ThreadSafe> Meta = MakeShared<FElectraBinarySample>();
+		switch(InDispatchMode)
+		{
+			default:
+			case IAdaptiveStreamingPlayerAEMSReceiver::EDispatchMode::OnReceive: Meta->DispatchedMode = FElectraBinarySample::EDispatchedMode::OnReceive;
+			case IAdaptiveStreamingPlayerAEMSReceiver::EDispatchMode::OnStart:	 Meta->DispatchedMode = FElectraBinarySample::EDispatchedMode::OnStart;
+		}
+		switch(InEvent->GetOrigin())
+		{
+			default:
+			case IAdaptiveStreamingPlayerAEMSEvent::EOrigin::TimedMetadata:		Meta->Origin = FElectraBinarySample::EOrigin::TimedMetadata;
+			case IAdaptiveStreamingPlayerAEMSEvent::EOrigin::EventStream:		Meta->Origin = FElectraBinarySample::EOrigin::EventStream;
+			case IAdaptiveStreamingPlayerAEMSEvent::EOrigin::InbandEventStream:	Meta->Origin = FElectraBinarySample::EOrigin::InbandEventStream;
+		}
+		Meta->Data = InEvent->GetMessageData();
+		Meta->SchemeIdUri = InEvent->GetSchemeIdUri();
+		Meta->Value = InEvent->GetValue();
+		Meta->ID = InEvent->GetID(),
+		Meta->Duration = InEvent->GetDuration().GetAsTimespan();
+		Meta->PresentationTime = FMediaTimeStamp(InEvent->GetPresentationTime().GetAsTimespan());
+		PinnedAdapterDelegate->PresentMetadataSample(Meta);
+	}
 }
 
 
diff --git a/Engine/Plugins/Media/ElectraPlayer/Source/ElectraPlayerRuntime/Public/IElectraPlayerInterface.h b/Engine/Plugins/Media/ElectraPlayer/Source/ElectraPlayerRuntime/Public/IElectraPlayerInterface.h
--- a/Engine/Plugins/Media/ElectraPlayer/Source/ElectraPlayerRuntime/Public/IElectraPlayerInterface.h
+++ b/Engine/Plugins/Media/ElectraPlayer/Source/ElectraPlayerRuntime/Public/IElectraPlayerInterface.h
@@ -12,6 +12,8 @@ class FVideoDecoderOutput;
 using FVideoDecoderOutputPtr = TSharedPtr<FVideoDecoderOutput, ESPMode::ThreadSafe>;
 class IAudioDecoderOutput;
 using IAudioDecoderOutputPtr = TSharedPtr<IAudioDecoderOutput, ESPMode::ThreadSafe>;
+class IElectraBinarySample;
+using IElectraBinarySampleRef = TSharedRef<IElectraBinarySample, ESPMode::ThreadSafe>;
 
 // ---------------------------------------------------------------------------------------------
 
@@ -66,6 +68,7 @@ public:
 	virtual void OnAudioFlush() = 0;
 	virtual void PresentVideoFrame(const FVideoDecoderOutputPtr& InVideoFrame) = 0;
 	virtual void PresentAudioFrame(const IAudioDecoderOutputPtr& InAudioFrame) = 0;
+	virtual void PresentMetadataSample(const IElectraBinarySampleRef& InMetadataSample) = 0;
 	virtual bool CanReceiveVideoSamples(int32 NumFrames) = 0;
 	virtual bool CanReceiveAudioSamples(int32 NumFrames) = 0;
 	virtual void PrepareForDecoderShutdown() = 0;
diff --git a/Engine/Plugins/Media/ElectraUtil/Source/ElectraBase/Public/PlayerTime.h b/Engine/Plugins/Media/ElectraUtil/Source/ElectraBase/Public/PlayerTime.h
--- a/Engine/Plugins/Media/ElectraUtil/Source/ElectraBase/Public/PlayerTime.h
+++ b/Engine/Plugins/Media/ElectraUtil/Source/ElectraBase/Public/PlayerTime.h
@@ -3,6 +3,7 @@
 #pragma once
 
 #include "CoreMinimal.h"
+#include "Misc/Timespan.h"
 #include <limits>
 
 namespace Electra
@@ -131,6 +132,19 @@ public:
 	//! Returns this time value in a custom timebase. Requires internal bigint conversion and is therefor SLOW!
 	int64 GetAsTimebase(uint32 CustomTimebase) const;
 
+	FTimespan GetAsTimespan() const
+	{
+		if (!bIsValid)
+		{
+			return FTimespan::MinValue();
+		}
+		if (IsInfinity())
+		{
+			return HNS >= 0 ? FTimespan::MaxValue() : FTimespan::MinValue();
+		}
+		return FTimespan(HNS);
+	}
+
 	FTimeValue& SetToInvalid()
 	{
 		HNS 		= 0;
diff --git a/Engine/Plugins/Media/ElectraUtil/Source/ElectraSamples/Public/IElectraMetadataSample.h b/Engine/Plugins/Media/ElectraUtil/Source/ElectraSamples/Public/IElectraMetadataSample.h
new file mode 100644
--- /dev/null
+++ b/Engine/Plugins/Media/ElectraUtil/Source/ElectraSamples/Public/IElectraMetadataSample.h
@@ -0,0 +1,57 @@
+// Copyright Epic Games, Inc. All Rights Reserved.
+#pragma once
+
+#include "Templates/SharedPointer.h"
+#include "Modules/ModuleInterface.h"
+#include "Misc/Timespan.h"
+#include "IMediaBinarySample.h"
+
+
+/**
+ * Extension of the media binary sample to carry Electra
+ * specific events.
+ */
+class IElectraBinarySample : public IMediaBinarySample
+{
+public:
+	static FGuid GetSampleTypeGUID()
+	{
+		static FGuid SampleTypeGUID(0xdfab9672, 0x754c459a, 0xa2831cb6, 0x332bab1d);
+		return SampleTypeGUID;
+	}
+
+	virtual ~IElectraBinarySample() = default;
+	virtual const void* GetData() = 0;
+	virtual FTimespan GetDuration() const = 0;
+	virtual uint32 GetSize() const = 0;
+	virtual FMediaTimeStamp GetTime() const = 0;
+	virtual FGuid GetGUID() const = 0;
+
+	enum class EOrigin
+	{
+		EventStream,
+		InbandEventStream,
+		TimedMetadata
+	};
+
+	enum class EDispatchedMode
+	{
+		OnReceive,
+		OnStart
+	};
+
+	// Returns where the event originated from.
+	virtual EOrigin GetOrigin() const = 0;
+
+	// Returns the type with which this event was dispatched.
+	virtual EDispatchedMode GetDispatchedMode() const = 0;
+
+	// Returns the event's scheme id URI
+	virtual FString GetSchemeIdUri() const = 0;
+
+	// Returns the event's value.
+	virtual FString GetValue() const = 0;
+
+	// Returns the event's ID.
+	virtual FString GetID() const = 0;
+};
diff --git a/Engine/Source/Runtime/Media/Public/IMediaBinarySample.h b/Engine/Source/Runtime/Media/Public/IMediaBinarySample.h
--- a/Engine/Source/Runtime/Media/Public/IMediaBinarySample.h
+++ b/Engine/Source/Runtime/Media/Public/IMediaBinarySample.h
@@ -6,6 +6,7 @@
 #include "Misc/Optional.h"
 #include "Misc/Timecode.h"
 #include "Misc/Timespan.h"
+#include "Misc/Guid.h"
 #include "IMediaTimeSource.h"
 
 
@@ -60,6 +61,15 @@ public:
 	 */
 	virtual TOptional<FTimecode> GetTimecode() const { return TOptional<FTimecode>(); }
 
+
+	/**
+	 * Get the GUID identifying the derived type that may implement additional,
+	 * type specific methods. If the GUID matches a type known type it is safe
+	 * to static cast this class to the derived type.
+	 * This is usually done to identify the format of the binary data carried here.
+	 */
+	virtual FGuid GetGUID() const { return FGuid(); }
+
 public:
 
 	/** Virtual destructor. */
