ElectraPlayer: added an explicit ESP:ThreadSafe keep the code identical between engine versions

#rnx

diff --git a/Engine/Plugins/Media/ElectraPlayer/Source/ElectraPlayerRuntime/Private/ElectraPlayer.cpp b/Engine/Plugins/Media/ElectraPlayer/Source/ElectraPlayerRuntime/Private/ElectraPlayer.cpp
--- a/Engine/Plugins/Media/ElectraPlayer/Source/ElectraPlayerRuntime/Private/ElectraPlayer.cpp
+++ b/Engine/Plugins/Media/ElectraPlayer/Source/ElectraPlayerRuntime/Private/ElectraPlayer.cpp
@@ -1072,7 +1072,7 @@ void FElectraPlayer::OnMediaPlayerEventReceived(TSharedPtrTS<IAdaptiveStreamingP
 	if (TSharedPtr<IElectraPlayerAdapterDelegate, ESPMode::ThreadSafe> PinnedAdapterDelegate = AdapterDelegate.Pin())
 	{
 		// Create a binary media sample of our extended format and pass it up.
-		TSharedRef<FElectraBinarySample, ESPMode::ThreadSafe> Meta = MakeShared<FElectraBinarySample>();
+		TSharedRef<FElectraBinarySample, ESPMode::ThreadSafe> Meta = MakeShared<FElectraBinarySample, ESPMode::ThreadSafe>();
 		switch(InDispatchMode)
 		{
 			default:
