ElectraPlayer: Slight API changes to return FString by const& instead of value; added missing break in case statement; changed default log verbosity down to Verbose to reduce log spam

#rnx

diff --git a/Engine/Plugins/Media/ElectraPlayer/Source/ElectraPlayerRuntime/Private/ElectraPlayer.cpp b/Engine/Plugins/Media/ElectraPlayer/Source/ElectraPlayerRuntime/Private/ElectraPlayer.cpp
--- a/Engine/Plugins/Media/ElectraPlayer/Source/ElectraPlayerRuntime/Private/ElectraPlayer.cpp
+++ b/Engine/Plugins/Media/ElectraPlayer/Source/ElectraPlayerRuntime/Private/ElectraPlayer.cpp
@@ -88,9 +88,9 @@ public:
 	virtual FGuid GetGUID() const override						{ return IElectraBinarySample::GetSampleTypeGUID(); }
 	virtual EOrigin GetOrigin() const override					{ return Origin; }
 	virtual EDispatchedMode GetDispatchedMode() const override	{ return DispatchedMode; }
-	virtual FString GetSchemeIdUri() const override				{ return SchemeIdUri; }
-	virtual FString GetValue() const override					{ return Value; }
-	virtual FString GetID() const override						{ return ID; }
+	virtual const FString& GetSchemeIdUri() const override		{ return SchemeIdUri; }
+	virtual const FString& GetValue() const override			{ return Value; }
+	virtual const FString& GetID() const override				{ return ID; }
 
 	TArray<uint8> Data;
 	FMediaTimeStamp PresentationTime;
@@ -1062,7 +1062,7 @@ void FElectraPlayer::OnMediaPlayerEventReceived(TSharedPtrTS<IAdaptiveStreamingP
 {
 #if !UE_BUILD_SHIPPING
 	const TCHAR* const Origins[] = { TEXT("Playlist"), TEXT("Inband"), TEXT("TimedMetadata"), TEXT("???") };
-	UE_LOG(LogElectraPlayer, Log, TEXT("[%p][%p] %s event %s with \"%s\", \"%s\", \"%s\" PTS @ %.3f for %.3fs"), this, CurrentPlayer.Get(),
+	UE_LOG(LogElectraPlayer, Verbose, TEXT("[%p][%p] %s event %s with \"%s\", \"%s\", \"%s\" PTS @ %.3f for %.3fs"), this, CurrentPlayer.Get(),
 		Origins[Electra::Utils::Min((int32)InEvent->GetOrigin(), (int32)UE_ARRAY_COUNT(Origins)-1)],
 		InDispatchMode==IAdaptiveStreamingPlayerAEMSReceiver::EDispatchMode::OnReceive?TEXT("received"):TEXT("started"),
 		*InEvent->GetSchemeIdUri(), *InEvent->GetValue(), *InEvent->GetID(),
@@ -1076,21 +1076,47 @@ void FElectraPlayer::OnMediaPlayerEventReceived(TSharedPtrTS<IAdaptiveStreamingP
 		switch(InDispatchMode)
 		{
 			default:
-			case IAdaptiveStreamingPlayerAEMSReceiver::EDispatchMode::OnReceive: Meta->DispatchedMode = FElectraBinarySample::EDispatchedMode::OnReceive;
-			case IAdaptiveStreamingPlayerAEMSReceiver::EDispatchMode::OnStart:	 Meta->DispatchedMode = FElectraBinarySample::EDispatchedMode::OnStart;
+			case IAdaptiveStreamingPlayerAEMSReceiver::EDispatchMode::OnReceive:
+			{
+				Meta->DispatchedMode = FElectraBinarySample::EDispatchedMode::OnReceive;
+				break;
+			}
+			case IAdaptiveStreamingPlayerAEMSReceiver::EDispatchMode::OnStart:
+			{
+				Meta->DispatchedMode = FElectraBinarySample::EDispatchedMode::OnStart;
+				break;
+			}
 		}
 		switch(InEvent->GetOrigin())
 		{
 			default:
-			case IAdaptiveStreamingPlayerAEMSEvent::EOrigin::TimedMetadata:		Meta->Origin = FElectraBinarySample::EOrigin::TimedMetadata;
-			case IAdaptiveStreamingPlayerAEMSEvent::EOrigin::EventStream:		Meta->Origin = FElectraBinarySample::EOrigin::EventStream;
-			case IAdaptiveStreamingPlayerAEMSEvent::EOrigin::InbandEventStream:	Meta->Origin = FElectraBinarySample::EOrigin::InbandEventStream;
+			case IAdaptiveStreamingPlayerAEMSEvent::EOrigin::TimedMetadata:		
+			{
+				Meta->Origin = FElectraBinarySample::EOrigin::TimedMetadata;
+				break;
+			}
+			case IAdaptiveStreamingPlayerAEMSEvent::EOrigin::EventStream:		
+			{
+				Meta->Origin = FElectraBinarySample::EOrigin::EventStream;
+				break;
+			}
+			case IAdaptiveStreamingPlayerAEMSEvent::EOrigin::InbandEventStream:	
+			{
+				Meta->Origin = FElectraBinarySample::EOrigin::InbandEventStream;
+				break;
+			}
 		}
 		Meta->Data = InEvent->GetMessageData();
 		Meta->SchemeIdUri = InEvent->GetSchemeIdUri();
 		Meta->Value = InEvent->GetValue();
 		Meta->ID = InEvent->GetID(),
 		Meta->Duration = InEvent->GetDuration().GetAsTimespan();
+		// A zero duration might cause the metadata sample fall through the cracks later
+		// so set it to a short 1ms instead.
+		if (Meta->Duration.IsZero())
+		{
+			Meta->Duration = FTimespan::FromMilliseconds(1);
+		}
 		Meta->PresentationTime = FMediaTimeStamp(InEvent->GetPresentationTime().GetAsTimespan());
 		PinnedAdapterDelegate->PresentMetadataSample(Meta);
 	}
diff --git a/Engine/Plugins/Media/ElectraUtil/Source/ElectraSamples/Public/IElectraMetadataSample.h b/Engine/Plugins/Media/ElectraUtil/Source/ElectraSamples/Public/IElectraMetadataSample.h
--- a/Engine/Plugins/Media/ElectraUtil/Source/ElectraSamples/Public/IElectraMetadataSample.h
+++ b/Engine/Plugins/Media/ElectraUtil/Source/ElectraSamples/Public/IElectraMetadataSample.h
@@ -47,11 +47,11 @@ public:
 	virtual EDispatchedMode GetDispatchedMode() const = 0;
 
 	// Returns the event's scheme id URI
-	virtual FString GetSchemeIdUri() const = 0;
+	virtual const FString& GetSchemeIdUri() const = 0;
 
 	// Returns the event's value.
-	virtual FString GetValue() const = 0;
+	virtual const FString& GetValue() const = 0;
 
 	// Returns the event's ID.
-	virtual FString GetID() const = 0;
+	virtual const FString& GetID() const = 0;
 };
