ElectraPlayer: added support for multiple language tracks in DASH and seamless switching between them; language tracks are switched automatically across periods to keep the same language if possible; using the first video representation's bitrate of the selected adaptation set to start playback at in absence of explicitly selected starting bitrate; limited frequency during Live playback with which to perform a time synchronization with the server

#rnx

diff --git a/Engine/Plugins/Media/ElectraPlayer/Source/ElectraPlayerPlugin/Private/ElectraPlayerPlugin.cpp b/Engine/Plugins/Media/ElectraPlayer/Source/ElectraPlayerPlugin/Private/ElectraPlayerPlugin.cpp
--- a/Engine/Plugins/Media/ElectraPlayer/Source/ElectraPlayerPlugin/Private/ElectraPlayerPlugin.cpp
+++ b/Engine/Plugins/Media/ElectraPlayer/Source/ElectraPlayerPlugin/Private/ElectraPlayerPlugin.cpp
@@ -326,7 +326,13 @@ bool FElectraPlayerPlugin::Open(const FString& Url, const IMediaOptions* Options
 	if (InPlayerOptions)
 	{
 		LocalPlaystartOptions.TimeOffset = InPlayerOptions->SeekTime;
-		LocalPlaystartOptions.AudioTrackIndex = InPlayerOptions->Tracks.Audio;
+		LocalPlaystartOptions.InitialAudioTrackAttributes.TrackIndexOverride = InPlayerOptions->Tracks.Audio;
+	}
+	FString InitialAudioLanguage = Options->GetMediaOption(TEXT("InitialAudioLanguage"), FString());
+	if (InitialAudioLanguage.Len())
+	{
+		LocalPlaystartOptions.InitialAudioTrackAttributes.Language_ISO639 = InitialAudioLanguage;
+		UE_LOG(LogElectraPlayerPlugin, Log, TEXT("[%p] IMediaPlayer::Open: Asking for initial audio language \"%s\""), this, *InitialAudioLanguage);
 	}
 
 	// Set up options to initialize the internal player with.
diff --git a/Engine/Plugins/Media/ElectraPlayer/Source/ElectraPlayerRuntime/Private/ElectraPlayer.cpp b/Engine/Plugins/Media/ElectraPlayer/Source/ElectraPlayerRuntime/Private/ElectraPlayer.cpp
--- a/Engine/Plugins/Media/ElectraPlayer/Source/ElectraPlayerRuntime/Private/ElectraPlayer.cpp
+++ b/Engine/Plugins/Media/ElectraPlayer/Source/ElectraPlayerRuntime/Private/ElectraPlayer.cpp
@@ -173,7 +173,7 @@ FElectraPlayer::~FElectraPlayer()
 	{
 		AsyncResourceReleaseNotification->Signal(ResourceFlags_OutputBuffers);
 	}
-	
+
 	FPlatformProcess::ReturnSynchEventToPool(WaitForPlayerDestroyedEvent);
 }
 
@@ -187,6 +187,7 @@ void FElectraPlayer::ClearToDefaultState()
 	SelectedQuality = 0;
 	SelectedVideoTrackIndex = -1;
 	SelectedAudioTrackIndex = -1;
+	bAudioTrackIndexDirty = true;
 	LastPresentedFrameDimension = FIntPoint::ZeroValue;
 	DeferredPlayerEvents.Empty();
 	MediaUrl.Empty();
@@ -248,10 +249,7 @@ bool FElectraPlayer::OpenInternal(const FString& Url, const FParamDict & PlayerO
 	MediaPlayerEventReceiver->GetEventReceivedDelegate().BindRaw(this, &FElectraPlayer::OnMediaPlayerEventReceived);
 	NewPlayer->AdaptivePlayer->AddAEMSReceiver(MediaPlayerEventReceiver, TEXT("*"), TEXT(""), IAdaptiveStreamingPlayerAEMSReceiver::EDispatchMode::OnReceive);
 
-	if (!NewPlayer->AdaptivePlayer->Initialize(PlayerOptions))
-	{
-		// FIXME: Can Initialize() fail? If so, we'd need to handle the error here.
-	}
+	NewPlayer->AdaptivePlayer->Initialize(PlayerOptions);
 
 	// Check for options that can be changed during playback and apply them at startup already.
 	// If a media source supports the MaxResolutionForMediaStreaming option then we can override the max resolution.
@@ -307,7 +305,7 @@ void FElectraPlayer::CloseInternal(bool bKillAfterClose)
 	State = EPlayerState::Closed;
 	MediaUrl = FString();
 	PlaystartOptions.TimeOffset.Reset();
-	PlaystartOptions.AudioTrackIndex.Reset();
+	PlaystartOptions.InitialAudioTrackAttributes.Reset();
 
 	// Next we detach ourselves from the renderers. This ensures we do not get any further data from them
 	// via OnVideoDecoded() and OnAudioDecoded(). It also means we do not get any calls to OnVideoFlush() and OnAudioFlush()
@@ -480,7 +478,7 @@ void FElectraPlayer::Tick(FTimespan DeltaTime, FTimespan Timecode)
 				}
 			}
 		}
-		
+
 		// Process accumulated player events.
 		HandleDeferredPlayerEvents();
 		if (bHasPendingError)
@@ -971,6 +969,44 @@ int32 FElectraPlayer::GetSelectedTrack(EPlayerTrackType TrackType) const
 	}
 	else if (TrackType == EPlayerTrackType::Audio)
 	{
+		/*
+			To reduce the overhead of this function we check for the track the underlying player has
+			actually selected only when we were told the tracks changed.
+
+			It is possible that the underlying player changes the track automatically as playback progresses.
+			For instance, when playing a DASH stream consisting of several periods the player needs to re-select
+			the audio stream when transitioning from one period into the next, which may change the index of
+			the selected track.
+		*/
+		if (bAudioTrackIndexDirty)
+		{
+			bAudioTrackIndexDirty = false;
+
+			if (NumTracksAudio == 0)
+			{
+				SelectedAudioTrackIndex = -1;
+			}
+			else
+			{
+				TSharedPtr<FInternalPlayerImpl, ESPMode::ThreadSafe> LockedPlayer = CurrentPlayer;
+				if (LockedPlayer.IsValid() && LockedPlayer->AdaptivePlayer.IsValid())
+				{
+					if (CurrentPlayer->AdaptivePlayer->IsTrackDeselected(Electra::EStreamType::Audio))
+					{
+						SelectedAudioTrackIndex = -1;
+					}
+					else
+					{
+						Electra::FStreamSelectionAttributes Attributes;
+						CurrentPlayer->AdaptivePlayer->GetSelectedTrackAttributes(Attributes, Electra::EStreamType::Audio);
+						if (Attributes.OverrideIndex.IsSet())
+						{
+							SelectedAudioTrackIndex = Attributes.OverrideIndex.GetValue();
+						}
+					}
+				}
+			}
+		}
 		return SelectedAudioTrackIndex;
 	}
 	return -1;
@@ -983,11 +1019,11 @@ FText FElectraPlayer::GetTrackDisplayName(EPlayerTrackType TrackType, int32 Trac
 	{
 		if (TrackType == EPlayerTrackType::Video)
 		{
-			return FText::FromString(FString::Printf(TEXT("Video Track ID %s"), *Meta->TrackID));
+			return FText::FromString(FString::Printf(TEXT("Video Track ID %s"), *Meta->ID));
 		}
 		else if (TrackType == EPlayerTrackType::Audio)
 		{
-			return FText::FromString(FString::Printf(TEXT("Audio Track ID %s"), *Meta->TrackID));
+			return FText::FromString(FString::Printf(TEXT("Audio Track ID %s"), *Meta->ID));
 		}
 	}
 	return FText();
@@ -1014,6 +1050,32 @@ FString FElectraPlayer::GetTrackName(EPlayerTrackType TrackType, int32 TrackInde
 	return TEXT("");
 }
 
+
+/**
+ * Selects a specified track for playback.
+ * 
+ * Note:
+ *   There is currently no concept of selecting a track based on metadata, only by index.
+ *   The idea being that before selecting a track by index the application needs to check
+ *   the metadata beforehand (eg. call GetTrackLanguage()) to figure out the index of the
+ *   track it wants to play.
+ * 
+ *   The underlying player however needs to select tracks based on metadata alone instead
+ *   of an index in case the track layout changes dynamically during playback.
+ *   For example, a part of the presentation could have both English and French audio,
+ *   followed by a part (say, an advertisement) that only has English audio, followed
+ *   by the continued regular part that has both. Without any user intervention the
+ *   player needs to automatically switch from French to English and back to French, or
+ *   index 1 -> 0 -> 1 (assuming French was the starting language of choice).
+ *   Indices are therefore meaningless to the underlying player.
+ * 
+ *   SelectTrack() is currently called implicitly by FMediaPlayerFacade::SelectDefaultTracks()
+ *   when EMediaEvent::TracksChanged is received. This is why this event is NOT sent out
+ *   in HandlePlayerEventTracksChanged() when the underlying player notifies us about a
+ *   change in track layout.
+ *   Other than the very first track selection made by the facade this method should only
+ *   be called from a direct user interaction.
+ */
 bool FElectraPlayer::SelectTrack(EPlayerTrackType TrackType, int32 TrackIndex)
 {
 	if (TrackType == EPlayerTrackType::Audio)
@@ -1021,20 +1083,35 @@ bool FElectraPlayer::SelectTrack(EPlayerTrackType TrackType, int32 TrackIndex)
 		// Select a track or deselect?
 		if (TrackIndex >= 0)
 		{
-			// Do a quick check if the track index is valid by checking the presence of the track metadata.
+			// Check if the track index exists by checking the presence of the track metadata.
+			// If for some reason the index is not valid the selection will not be changed.
 			TSharedPtr<Electra::FTrackMetadata, ESPMode::ThreadSafe> Meta = GetTrackStreamMetadata(EPlayerTrackType::Audio, TrackIndex);
 			if (Meta.IsValid())
 			{
-				PlaystartOptions.AudioTrackIndex = TrackIndex;
-				SelectedAudioTrackIndex = TrackIndex;
-
-				// Note: Selecting the audio track does not have an immediate effect at the moment.
-				//       The player does not switch tracks dynamically. We are still setting it here
-				//       because it will stick on the next seek, so we have at least that.
-				TSharedPtr<FInternalPlayerImpl, ESPMode::ThreadSafe> LockedPlayer = CurrentPlayer;
-				if (LockedPlayer.IsValid() && LockedPlayer->AdaptivePlayer.IsValid())
+				// Switch only when the track index has changed.
+				if (GetSelectedTrack(TrackType) != TrackIndex)
 				{
-					CurrentPlayer->AdaptivePlayer->SelectTrackByMetadata(Electra::EStreamType::Audio, *Meta);
+					Electra::FStreamSelectionAttributes AudioAttributes;
+					AudioAttributes.OverrideIndex = TrackIndex;
+
+					PlaystartOptions.InitialAudioTrackAttributes.TrackIndexOverride = TrackIndex;
+					if (!Meta->Kind.IsEmpty())
+					{
+						AudioAttributes.Kind = Meta->Kind;
+						PlaystartOptions.InitialAudioTrackAttributes.Kind = Meta->Kind;
+					}
+					if (!Meta->Language.IsEmpty())
+					{
+						AudioAttributes.Language_ISO639 = Meta->Language;
+						PlaystartOptions.InitialAudioTrackAttributes.Language_ISO639 = Meta->Language;
+					}
+					TSharedPtr<FInternalPlayerImpl, ESPMode::ThreadSafe> LockedPlayer = CurrentPlayer;
+					if (LockedPlayer.IsValid() && LockedPlayer->AdaptivePlayer.IsValid())
+					{
+						CurrentPlayer->AdaptivePlayer->SelectTrackByAttributes(Electra::EStreamType::Audio, AudioAttributes);
+					}
+
+					SelectedAudioTrackIndex = TrackIndex;
 				}
 				return true;
 			}
@@ -1042,8 +1119,8 @@ bool FElectraPlayer::SelectTrack(EPlayerTrackType TrackType, int32 TrackIndex)
 		else
 		{
 			// Deselect audio track.
-			PlaystartOptions.AudioTrackIndex = TrackIndex;
-			SelectedAudioTrackIndex = TrackIndex;
+			PlaystartOptions.InitialAudioTrackAttributes.TrackIndexOverride = -1;
+			SelectedAudioTrackIndex = -1;
 			TSharedPtr<FInternalPlayerImpl, ESPMode::ThreadSafe> LockedPlayer = CurrentPlayer;
 			if (LockedPlayer.IsValid() && LockedPlayer->AdaptivePlayer.IsValid())
 			{
@@ -1282,7 +1359,7 @@ void FElectraPlayer::HandleDeferredPlayerEvents()
 				break;
 			case FPlayerMetricEventBase::EType::DroppedAudioFrame:
 				HandlePlayerEventDroppedAudioFrame();
-				break; 
+				break;
 			default:
 				break;
 		}
@@ -1347,42 +1424,6 @@ void FElectraPlayer::HandlePlayerEventReceivedPlaylists()
 
 	UE_LOG(LogElectraPlayer, Log, TEXT("[%p][%p] Received initial stream playlists"), this, CurrentPlayer.Get());
 
-	// Set up the initial playback position
-	IAdaptiveStreamingPlayer::FSeekParam playParam;
-	// First we look at any potential time offset specified in the playstart options.
-	if (PlaystartOptions.TimeOffset.IsSet())
-	{
-		FTimespan Target;
-		CalculateTargetSeekTime(Target, PlaystartOptions.TimeOffset.GetValue());
-		playParam.Time.SetFromHNS(Target.GetTicks());
-	}
-	else
-	{
-		// Do not set a start time, let the player pick one.
-		//playParam.Time.SetToZero();
-	}
-
-	// Next, give a list of the seekable positions to the delegate and ask it if it wants to seek to one of them,
-	// overriding any potential time offset from above.
-	if (TSharedPtr<IElectraPlayerAdapterDelegate, ESPMode::ThreadSafe> PinnedAdapterDelegate = AdapterDelegate.Pin())
-	{
-		TSharedPtr<TArray<FTimespan>, ESPMode::ThreadSafe> SeekablePositions = MakeShared<TArray<FTimespan>, ESPMode::ThreadSafe>();
-		if (SeekablePositions.IsValid())
-		{
-			CurrentPlayer->AdaptivePlayer->GetSeekablePositions(*SeekablePositions);
-
-			// Check with the delegate if it wants to start somewhere else.
-			FVariantValue Result = PinnedAdapterDelegate->QueryOptions(IElectraPlayerAdapterDelegate::EOptionType::PlaystartPosFromSeekPositions, FVariantValue(SeekablePositions));
-			if (Result.IsValid())
-			{
-				check(Result.IsType(FVariantValue::EDataType::TypeInt64));
-				playParam.Time.SetFromHNS(Result.GetInt64());
-			}
-		}
-	}
-
-
-
 	Electra::FTimeRange MediaTimeline;
 	Electra::FTimeValue MediaDuration;
 	CurrentPlayer->AdaptivePlayer->GetTimelineRange(MediaTimeline);
@@ -1405,7 +1446,7 @@ void FElectraPlayer::HandlePlayerEventReceivedPlaylists()
 	{
 		for(int32 i=0; i<VideoStreamMetaData[0].StreamDetails.Num(); ++i)
 		{
-			UE_LOG(LogElectraPlayer, Log, TEXT("[%p][%p] Found %d * %d video stream at bitrate %d"), this, CurrentPlayer.Get(), 
+			UE_LOG(LogElectraPlayer, Log, TEXT("[%p][%p] Found %d * %d video stream at bitrate %d"), this, CurrentPlayer.Get(),
 												VideoStreamMetaData[0].StreamDetails[i].CodecInformation.GetResolution().Width,
 												VideoStreamMetaData[0].StreamDetails[i].CodecInformation.GetResolution().Height,
 												VideoStreamMetaData[0].StreamDetails[i].Bandwidth);
@@ -1427,34 +1468,47 @@ void FElectraPlayer::HandlePlayerEventReceivedPlaylists()
 	CurrentPlayer->AdaptivePlayer->GetTrackMetadata(AudioStreamMetaData, Electra::EStreamType::Audio);
 	NumTracksAudio = AudioStreamMetaData.Num();
 
-	// Select the desired audio track by hard metadata index.
-	if (NumTracksAudio)
+	// Set the initial audio track selection attributes.
+	Electra::FStreamSelectionAttributes InitialAudioAttributes;
+	InitialAudioAttributes.Kind = PlaystartOptions.InitialAudioTrackAttributes.Kind;
+	InitialAudioAttributes.Language_ISO639 = PlaystartOptions.InitialAudioTrackAttributes.Language_ISO639;
+	InitialAudioAttributes.OverrideIndex = PlaystartOptions.InitialAudioTrackAttributes.TrackIndexOverride;
+	CurrentPlayer->AdaptivePlayer->SetInitialStreamAttributes(Electra::EStreamType::Audio, InitialAudioAttributes);
+
+
+	// Set up the initial playback position
+	IAdaptiveStreamingPlayer::FSeekParam playParam;
+	// First we look at any potential time offset specified in the playstart options.
+	if (PlaystartOptions.TimeOffset.IsSet())
+	{
+		FTimespan Target;
+		CalculateTargetSeekTime(Target, PlaystartOptions.TimeOffset.GetValue());
+		playParam.Time.SetFromHNS(Target.GetTicks());
+	}
+	else
 	{
-		SelectedAudioTrackIndex = 0;
+		// Do not set a start time, let the player pick one.
+		//playParam.Time.SetToZero();
+	}
 
-		if (PlaystartOptions.AudioTrackIndex.IsSet())
+	// Next, give a list of the seekable positions to the delegate and ask it if it wants to seek to one of them,
+	// overriding any potential time offset from above.
+	if (TSharedPtr<IElectraPlayerAdapterDelegate, ESPMode::ThreadSafe> PinnedAdapterDelegate = AdapterDelegate.Pin())
+	{
+		TSharedPtr<TArray<FTimespan>, ESPMode::ThreadSafe> SeekablePositions = MakeShared<TArray<FTimespan>, ESPMode::ThreadSafe>();
+		if (SeekablePositions.IsValid())
 		{
-			int32 HardAudioTrackIndex = PlaystartOptions.AudioTrackIndex.GetValue();
-			if (HardAudioTrackIndex >= 0)
-			{
-				if (HardAudioTrackIndex >= NumTracksAudio)
-				{
-					HardAudioTrackIndex = 0;
-				}
-				SelectedAudioTrackIndex = HardAudioTrackIndex;
-				CurrentPlayer->AdaptivePlayer->SelectTrackByMetadata(Electra::EStreamType::Audio, AudioStreamMetaData[HardAudioTrackIndex]);
-			}
-			else
+			CurrentPlayer->AdaptivePlayer->GetSeekablePositions(*SeekablePositions);
+
+			// Check with the delegate if it wants to start somewhere else.
+			FVariantValue Result = PinnedAdapterDelegate->QueryOptions(IElectraPlayerAdapterDelegate::EOptionType::PlaystartPosFromSeekPositions, FVariantValue(SeekablePositions));
+			if (Result.IsValid())
 			{
-				SelectedAudioTrackIndex = -1;
-				CurrentPlayer->AdaptivePlayer->DeselectTrack(Electra::EStreamType::Audio);
+				check(Result.IsType(FVariantValue::EDataType::TypeInt64));
+				playParam.Time.SetFromHNS(Result.GetInt64());
 			}
 		}
 	}
-	else
-	{
-		SelectedAudioTrackIndex = -1;
-	}
 
 	// Trigger buffering at the intended start time.
 	CurrentPlayer->AdaptivePlayer->SeekTo(playParam);
@@ -1469,7 +1523,7 @@ void FElectraPlayer::HandlePlayerEventTracksChanged()
 	{
 		for(int32 i=0; i<VideoStreamMetaData[0].StreamDetails.Num(); ++i)
 		{
-			UE_LOG(LogElectraPlayer, Log, TEXT("[%p][%p] Found %d * %d video stream at bitrate %d"), this, CurrentPlayer.Get(), 
+			UE_LOG(LogElectraPlayer, Log, TEXT("[%p][%p] Found %d * %d video stream at bitrate %d"), this, CurrentPlayer.Get(),
 												VideoStreamMetaData[0].StreamDetails[i].CodecInformation.GetResolution().Width,
 												VideoStreamMetaData[0].StreamDetails[i].CodecInformation.GetResolution().Height,
 												VideoStreamMetaData[0].StreamDetails[i].Bandwidth);
@@ -1479,8 +1533,12 @@ void FElectraPlayer::HandlePlayerEventTracksChanged()
 	TArray<Electra::FTrackMetadata> AudioStreamMetaData;
 	CurrentPlayer->AdaptivePlayer->GetTrackMetadata(AudioStreamMetaData, Electra::EStreamType::Audio);
 	NumTracksAudio = AudioStreamMetaData.Num();
+	bAudioTrackIndexDirty = true;
 
-	DeferredEvents.Enqueue(IElectraPlayerAdapterDelegate::EPlayerEvent::TracksChanged);
+	/*
+	Note: Do not send TracksChanged here since this would trigger a re-selecting of the default tracks.
+		DeferredEvents.Enqueue(IElectraPlayerAdapterDelegate::EPlayerEvent::TracksChanged);
+	*/
 }
 
 
@@ -1742,7 +1800,7 @@ void FElectraPlayer::HandlePlayerEventVideoQualityChange(int32 NewBitrate, int32
 			if (VideoStreamMetaData[0].StreamDetails[i].Bandwidth == NewBitrate)
 			{
 				SelectedQuality = i;
-				Statistics.CurrentlyActivePlaylistURL = VideoStreamMetaData[0].TrackID;
+				Statistics.CurrentlyActivePlaylistURL = VideoStreamMetaData[0].ID;
 				Statistics.CurrentlyActiveResolutionWidth = VideoStreamMetaData[0].StreamDetails[i].CodecInformation.GetResolution().Width;
 				Statistics.CurrentlyActiveResolutionHeight = VideoStreamMetaData[0].StreamDetails[i].CodecInformation.GetResolution().Height;
 				break;
@@ -2353,10 +2411,7 @@ void FElectraPlayer::MediaStateOnEndReached()
 		// NOP
 		break;
 	}
-	DeferredEvents.Enqueue(IElectraPlayerAdapterDelegate::EPlayerEvent::TracksChanged);
 	State = EPlayerState::Stopped;
-
-	//Call player close()
 }
 
 void FElectraPlayer::MediaStateOnSeekFinished()
diff --git a/Engine/Plugins/Media/ElectraPlayer/Source/ElectraPlayerRuntime/Private/ElectraPlayer.h b/Engine/Plugins/Media/ElectraPlayer/Source/ElectraPlayerRuntime/Private/ElectraPlayer.h
--- a/Engine/Plugins/Media/ElectraPlayer/Source/ElectraPlayerRuntime/Private/ElectraPlayer.h
+++ b/Engine/Plugins/Media/ElectraPlayer/Source/ElectraPlayerRuntime/Private/ElectraPlayer.h
@@ -328,7 +328,8 @@ private:
 	int32											NumTracksVideo;
 	int32											SelectedQuality;
 	int32											SelectedVideoTrackIndex;
-	int32											SelectedAudioTrackIndex;
+	mutable int32									SelectedAudioTrackIndex;
+	mutable bool									bAudioTrackIndexDirty;
 
 	FIntPoint										LastPresentedFrameDimension;
 
diff --git a/Engine/Plugins/Media/ElectraPlayer/Source/ElectraPlayerRuntime/Private/Runtime/Demuxer/ParserISO14496-12.cpp b/Engine/Plugins/Media/ElectraPlayer/Source/ElectraPlayerRuntime/Private/Runtime/Demuxer/ParserISO14496-12.cpp
--- a/Engine/Plugins/Media/ElectraPlayer/Source/ElectraPlayerRuntime/Private/Runtime/Demuxer/ParserISO14496-12.cpp
+++ b/Engine/Plugins/Media/ElectraPlayer/Source/ElectraPlayerRuntime/Private/Runtime/Demuxer/ParserISO14496-12.cpp
@@ -9,6 +9,7 @@
 #include "Utilities/UtilsMPEG.h"
 #include "Utilities/UtilsMPEGAudio.h"
 #include "Utilities/UtilsMPEGVideo.h"
+#include "Utilities/ISO639-Map.h"
 #include "InfoLog.h"
 #include "Player/PlayerSessionServices.h"
 
@@ -1387,6 +1388,11 @@ private:
 			return HandlerType;
 		}
 
+		FString GetHandlerName() const
+		{
+			return NameUTF8;
+		}
+
 	private:
 		FMP4BoxHDLR() = delete;
 		FMP4BoxHDLR(const FMP4BoxHDLR&) = delete;
@@ -4629,7 +4635,7 @@ private:
 		virtual const ISegmentIndex* GetSegmentIndexByIndex(int32 Index) const override;
 		virtual const IEventMessage* GetEventMessageByIndex(int32 Index) const override;
 
-		virtual TSharedPtr<IAllTrackIterator, ESPMode::ThreadSafe> CreateAllTrackIteratorByFilePos(int64 InFromFilePos) const override;
+		virtual TSharedPtrTS<IAllTrackIterator> CreateAllTrackIteratorByFilePos(int64 InFromFilePos) const override;
 
 	private:
 		class FTrack;
@@ -4758,6 +4764,7 @@ private:
 			{
 			}
 			virtual uint32 GetID() const override;
+			virtual FString GetNameFromHandler() const override;
 			virtual FTimeFraction GetDuration() const override;
 			virtual ITrackIterator* CreateIterator() const override;
 			virtual const TArray<uint8>& GetCodecSpecificData() const override;
@@ -4824,9 +4831,9 @@ private:
 			//! Returns list of all iterators.
 			virtual void GetAllIterators(TArray<const ITrackIterator*>& OutIterators) const override;
 
-			TArray<TSharedPtr<ITrackIterator, ESPMode::ThreadSafe>>		TrackIterators;
-			TSharedPtr<ITrackIterator, ESPMode::ThreadSafe>				CurrentIterator;
-			TArray<const ITrackIterator*>								NewlyReachedEOS;
+			TArray<TSharedPtrTS<ITrackIterator>>		TrackIterators;
+			TSharedPtrTS<ITrackIterator>				CurrentIterator;
+			TArray<const ITrackIterator*>				NewlyReachedEOS;
 		};
 
 
@@ -4923,6 +4930,11 @@ private:
 	{
 		return TKHDBox ? TKHDBox->GetTrackID() : 0;
 	}
+	
+	FString FParserISO14496_12::FTrack::GetNameFromHandler() const
+	{
+		return HDLRBox ? HDLRBox->GetHandlerName() : FString();
+	}
 
 	FTimeFraction FParserISO14496_12::FTrack::GetDuration() const
 	{
@@ -5007,7 +5019,9 @@ private:
 
 	const FString FParserISO14496_12::FTrack::GetLanguage() const
 	{
-		return MDHDBox ? MDHDBox->GetLanguage() : FString(TEXT("und"));
+		FString Language = MDHDBox ? MDHDBox->GetLanguage() : FString(TEXT("und"));
+		// Try to map the ISO-639-2T language code to the shorter ISO-639-1 code if possible.
+		return ISO639::MapTo639_1(Language);
 	}
 
 	void FParserISO14496_12::FTrack::GetPSSHBoxes(TArray<TArray<uint8>>& OutBoxes, bool bFromMOOV, bool bFromMOOF) const
@@ -6448,9 +6462,9 @@ private:
 	}
 
 
-	TSharedPtr<IParserISO14496_12::IAllTrackIterator, ESPMode::ThreadSafe> FParserISO14496_12::CreateAllTrackIteratorByFilePos(int64 InFromFilePos) const
+	TSharedPtrTS<IParserISO14496_12::IAllTrackIterator> FParserISO14496_12::CreateAllTrackIteratorByFilePos(int64 InFromFilePos) const
 	{
-		TSharedPtr<FParserISO14496_12::FAllTrackIterator, ESPMode::ThreadSafe> ti = MakeShared<FParserISO14496_12::FAllTrackIterator, ESPMode::ThreadSafe>();
+		TSharedPtrTS<FParserISO14496_12::FAllTrackIterator> ti = MakeShared<FParserISO14496_12::FAllTrackIterator, ESPMode::ThreadSafe>();
 
 		int64 LowestFilePos = TNumericLimits<int64>::Max();
 		for(int32 nTrk = 0, nTrkMax = GetNumberOfTracks(); nTrk < nTrkMax; ++nTrk)
@@ -6458,7 +6472,7 @@ private:
 			const FTrack* Track = ParsedTrackInfo->GetTrackByIndex(nTrk);
 			check(Track);
 			FTrackIterator* TrkIt = static_cast<FTrackIterator*>(Track->CreateIterator());
-			TSharedPtr<ITrackIterator, ESPMode::ThreadSafe> SafeTrkIt(TrkIt);
+			TSharedPtrTS<ITrackIterator> SafeTrkIt(TrkIt);
 			ti->TrackIterators.Add(SafeTrkIt);
 			UEMediaError err = TrkIt->StartAtFirstInteral();
 			if (err == UEMEDIA_ERROR_OK)
diff --git a/Engine/Plugins/Media/ElectraPlayer/Source/ElectraPlayerRuntime/Private/Runtime/Demuxer/ParserISO14496-12.h b/Engine/Plugins/Media/ElectraPlayer/Source/ElectraPlayerRuntime/Private/Runtime/Demuxer/ParserISO14496-12.h
--- a/Engine/Plugins/Media/ElectraPlayer/Source/ElectraPlayerRuntime/Private/Runtime/Demuxer/ParserISO14496-12.h
+++ b/Engine/Plugins/Media/ElectraPlayer/Source/ElectraPlayerRuntime/Private/Runtime/Demuxer/ParserISO14496-12.h
@@ -215,6 +215,7 @@ namespace Electra
 			virtual ~ITrack() = default;
 
 			virtual uint32 GetID() const = 0;
+			virtual FString GetNameFromHandler() const = 0;
 			virtual FTimeFraction GetDuration() const = 0;
 
 			virtual ITrackIterator* CreateIterator() const = 0;
@@ -283,7 +284,7 @@ namespace Electra
 		};
 
 
-		virtual TSharedPtr<IAllTrackIterator, ESPMode::ThreadSafe> CreateAllTrackIteratorByFilePos(int64 InFromFilePos) const = 0;
+		virtual TSharedPtrTS<IAllTrackIterator> CreateAllTrackIteratorByFilePos(int64 InFromFilePos) const = 0;
 
 		virtual const ITrack* GetTrackByIndex(int32 Index) const = 0;
 		virtual const ITrack* GetTrackByTrackID(int32 TrackID) const = 0;
diff --git a/Engine/Plugins/Media/ElectraPlayer/Source/ElectraPlayerRuntime/Private/Runtime/Player/AdaptivePlayerOptionKeynames.h b/Engine/Plugins/Media/ElectraPlayer/Source/ElectraPlayerRuntime/Private/Runtime/Player/AdaptivePlayerOptionKeynames.h
--- a/Engine/Plugins/Media/ElectraPlayer/Source/ElectraPlayerRuntime/Private/Runtime/Player/AdaptivePlayerOptionKeynames.h
+++ b/Engine/Plugins/Media/ElectraPlayer/Source/ElectraPlayerRuntime/Private/Runtime/Player/AdaptivePlayerOptionKeynames.h
@@ -10,6 +10,8 @@ namespace Electra
 	const TCHAR* const OptionKeyInitialBitrate = TEXT("initial_bitrate");							//!< (int64) value indicating the bitrate to start with.
 	const TCHAR* const OptionKeyLiveSeekableEndOffset = TEXT("seekable_range_live_end_offset");		//!< (FTimeValue) value specifying how many seconds away from the Live media timeline the seekable range should end.
 
+	const TCHAR* const OptionKeyCurrentAvgStartingVideoBitrate = TEXT("current:avg_video_bitrate");
+
 } // namespace Electra
 
 
diff --git a/Engine/Plugins/Media/ElectraPlayer/Source/ElectraPlayerRuntime/Private/Runtime/Player/AdaptiveStreamingPlayer.cpp b/Engine/Plugins/Media/ElectraPlayer/Source/ElectraPlayerRuntime/Private/Runtime/Player/AdaptiveStreamingPlayer.cpp
--- a/Engine/Plugins/Media/ElectraPlayer/Source/ElectraPlayerRuntime/Private/Runtime/Player/AdaptiveStreamingPlayer.cpp
+++ b/Engine/Plugins/Media/ElectraPlayer/Source/ElectraPlayerRuntime/Private/Runtime/Player/AdaptiveStreamingPlayer.cpp
@@ -9,6 +9,7 @@
 #include "Player/PlayerLicenseKey.h"
 #include "Player/AdaptivePlayerOptionKeynames.h"
 #include "Utilities/Utilities.h"
+#include "Utilities/ISO639-Map.h"
 #include "Player/DRM/DRMManager.h"
 
 #include "HAL/LowLevelMemTracker.h"
@@ -35,6 +36,24 @@ void IAdaptiveStreamingPlayer::DebugHandle(void* pPlayer, void (*debugDrawPrintf
 //---------------------------------------------------------------------------------------------------------------------------------
 
 
+namespace
+{
+	inline int32 StreamTypeToArrayIndex(EStreamType StreamType)
+	{
+		switch(StreamType)
+		{
+			case EStreamType::Video: 
+				return 0;
+			case EStreamType::Audio: 
+				return 1;
+			case EStreamType::Subtitle: 
+				return 2;
+			default: 
+				return 3;
+		}
+	}
+}
+
 //-----------------------------------------------------------------------------
 /**
  * CTOR
@@ -45,6 +64,8 @@ FAdaptiveStreamingPlayer::FAdaptiveStreamingPlayer(const IAdaptiveStreamingPlaye
 
 	ExternalPlayerGUID		= InCreateParameters.ExternalPlayerGUID;
 
+	ManifestType			= EMediaFormatType::Unknown;
+
 	VideoRender.Renderer    = InCreateParameters.VideoRenderer;
 	AudioRender.Renderer    = InCreateParameters.AudioRenderer;
 	CurrentState   		    = EPlayerState::eState_Idle;
@@ -58,13 +79,17 @@ FAdaptiveStreamingPlayer::FAdaptiveStreamingPlayer(const IAdaptiveStreamingPlaye
 	LastBufferingState 		= EPlayerState::eState_Buffering;
 	bIsPlayStart   			= true;
 	bIsClosing 				= false;
+
 	StartAtTime.Reset();
 	bHaveVideoReader.Reset();
 	bHaveAudioReader.Reset();
 	bHaveTextReader.Reset();
 	RebufferDetectedAtPlayPos.SetToInvalid();
 
-	CurrentPlaybackSequenceID = 0;
+	CurrentPlaybackSequenceID[0] = 0;
+	CurrentPlaybackSequenceID[1] = 0;
+	CurrentPlaybackSequenceID[2] = 0;
+	CurrentPlaybackSequenceID[3] = 0;
 
 	DataAvailabilityStateVid.StreamType = EStreamType::Video;
 	DataAvailabilityStateVid.Availability = Metrics::FDataAvailabilityChange::EAvailability::DataNotAvailable;
@@ -174,10 +199,8 @@ void FAdaptiveStreamingPlayer::UpdateDiagnostics()
  * Initializes the player instance.
  *
  * @param Options
- *
- * @return
  */
-bool FAdaptiveStreamingPlayer::Initialize(const FParamDict& Options)
+void FAdaptiveStreamingPlayer::Initialize(const FParamDict& Options)
 {
 	PlayerOptions = Options;
 
@@ -196,7 +219,7 @@ bool FAdaptiveStreamingPlayer::Initialize(const FParamDict& Options)
 
 	// Set up video decoder resolution limits. As the media playlists are parsed the video streams will be
 	// compared against these limits and those that exceed the limit will not be considered for playback.
-	
+
 	// Maximum allowed vertical resolution specified?
 	if (PlayerOptions.HaveKey("max_resoY"))
 	{
@@ -229,7 +252,7 @@ bool FAdaptiveStreamingPlayer::Initialize(const FParamDict& Options)
 			PlayerConfig.H264LimitUpto30fps.MaxResolution.Height = Utils::Min(PlayerConfig.H264LimitUpto30fps.MaxResolution.Height, Capability.Height);
 			PlayerConfig.H264LimitAbove30fps.MaxResolution.Height = Utils::Min(PlayerConfig.H264LimitAbove30fps.MaxResolution.Height, Capability.Height);
 		}
-		
+
 		// If this is software decoding only and there is limit for this in place on Windows, apply it.
 		if (PlayerOptions.HaveKey("max_resoY_windows_software") && Capability.DecoderSupportType == IVideoDecoderH264::FStreamDecodeCapability::ESupported::SoftwareOnly)
 		{
@@ -247,8 +270,6 @@ bool FAdaptiveStreamingPlayer::Initialize(const FParamDict& Options)
 			PlayerConfig.H264LimitAbove30fps.MaxTierProfileLevel.Level = 10;
 		}
 	}
-
-	return true;
 }
 
 
@@ -266,7 +287,7 @@ void FAdaptiveStreamingPlayer::SetVideoDecoderResourceDelegate(const TSharedPtr<
 void FAdaptiveStreamingPlayer::AddMetricsReceiver(IAdaptiveStreamingPlayerMetrics* InMetricsReceiver)
 {
 	if (InMetricsReceiver)
-	{  
+	{
 		FMediaCriticalSection::ScopedLock lock(MetricListenerCriticalSection);
 		if (MetricListeners.Find(InMetricsReceiver) == INDEX_NONE)
 		{
@@ -357,8 +378,10 @@ void FAdaptiveStreamingPlayer::StopWorkerThread()
 
 
 
-
-
+//-----------------------------------------------------------------------------
+/**
+ * Starts loading the manifest / master playlist.
+ */
 void FAdaptiveStreamingPlayer::LoadManifest(const FString& InManifestURL)
 {
 	WorkerThread.SendLoadManifestMessage(InManifestURL, FString());
@@ -592,6 +615,52 @@ void FAdaptiveStreamingPlayer::GetTrackMetadata(TArray<FTrackMetadata>& OutTrack
 }
 
 
+#if 0
+//-----------------------------------------------------------------------------
+/**
+ * Returns the track metadata of the currently active track. If nothing is selected yet the TOptional<> will be unset.
+ *
+ * @param OutSelectedTrackMetadata
+ * @param StreamType
+ */
+void FAdaptiveStreamingPlayer::GetSelectedTrackMetadata(TOptional<FTrackMetadata>& OutSelectedTrackMetadata, EStreamType StreamType) const
+{
+}
+#endif
+
+void FAdaptiveStreamingPlayer::GetSelectedTrackAttributes(FStreamSelectionAttributes& OutAttributes, EStreamType StreamType) const
+{
+	TSharedPtrTS<FBufferSourceInfo> BufferInfo;
+
+	auto SetOutAttributes = [&](const TSharedPtrTS<FBufferSourceInfo>& BufferInfo, const FStreamSelectionAttributes& CurrentAttr) -> void
+	{
+		if (BufferInfo.IsValid())
+		{
+			OutAttributes.Kind = BufferInfo->Kind;
+			OutAttributes.Language_ISO639 = BufferInfo->Language;
+			OutAttributes.OverrideIndex = BufferInfo->HardIndex;
+		}
+		else
+		{
+			OutAttributes = CurrentAttr;
+		}
+	};
+
+	if (StreamType == EStreamType::Video)
+	{
+		SetOutAttributes(MultiStreamBufferVid.GetActiveOutputBufferInfo(), SelectedStreamAttributesVid);
+	}
+	else if (StreamType == EStreamType::Audio)
+	{
+		SetOutAttributes(MultiStreamBufferAud.GetActiveOutputBufferInfo(), SelectedStreamAttributesAud);
+	}
+	else if (StreamType == EStreamType::Subtitle)
+	{
+		SetOutAttributes(MultiStreamBufferTxt.GetActiveOutputBufferInfo(), SelectedStreamAttributesTxt);
+	}
+}
+
+
 //-----------------------------------------------------------------------------
 /**
  * Sets the highest bitrate when selecting a candidate stream.
@@ -618,6 +687,12 @@ void FAdaptiveStreamingPlayer::SetMaxResolution(int32 MaxWidth, int32 MaxHeight)
 }
 
 
+void FAdaptiveStreamingPlayer::SetInitialStreamAttributes(EStreamType StreamType, const FStreamSelectionAttributes& InitialSelection)
+{
+	WorkerThread.SendInitialStreamAttributeMessage(StreamType, InitialSelection);
+}
+
+#if 0
 //-----------------------------------------------------------------------------
 /**
  * Selects a track based from one of the array members returned by GetStreamMetadata().
@@ -629,6 +704,13 @@ void FAdaptiveStreamingPlayer::SelectTrackByMetadata(EStreamType StreamType, con
 {
 	WorkerThread.SendTrackSelectByMetadataMessage(StreamType, TrackMetadata);
 }
+#endif
+
+void FAdaptiveStreamingPlayer::SelectTrackByAttributes(EStreamType StreamType, const FStreamSelectionAttributes& Attributes)
+{
+	WorkerThread.SendTrackSelectByAttributeMessage(StreamType, Attributes);
+}
+
 
 //-----------------------------------------------------------------------------
 /**
@@ -641,6 +723,22 @@ void FAdaptiveStreamingPlayer::DeselectTrack(EStreamType StreamType)
 	WorkerThread.SendTrackDeselectMessage(StreamType);
 }
 
+bool FAdaptiveStreamingPlayer::IsTrackDeselected(EStreamType StreamType)
+{
+	switch(StreamType)
+	{
+		case EStreamType::Video:
+			return MultiStreamBufferVid.IsDeselected();
+		case EStreamType::Audio:
+			return MultiStreamBufferAud.IsDeselected();
+		case EStreamType::Subtitle:
+			return MultiStreamBufferTxt.IsDeselected();
+		default:
+			return true;
+	}
+}
+
+
 //-----------------------------------------------------------------------------
 /**
  * Starts the renderers.
@@ -1023,7 +1121,7 @@ bool FAdaptiveStreamingPlayer::OnFragmentAccessUnitReceived(FAccessUnit* pAccess
  * @param InStreamType
  * @param InStreamSourceInfo
  */
-void FAdaptiveStreamingPlayer::OnFragmentReachedEOS(EStreamType InStreamType, TSharedPtr<const FStreamSourceInfo, ESPMode::ThreadSafe> InStreamSourceInfo)
+void FAdaptiveStreamingPlayer::OnFragmentReachedEOS(EStreamType InStreamType, TSharedPtr<const FBufferSourceInfo, ESPMode::ThreadSafe> InStreamSourceInfo)
 {
 	if (!bIsClosing)
 	{
@@ -1194,49 +1292,76 @@ void FAdaptiveStreamingPlayer::WorkerThreadFN()
 					break;
 				}
 
-				case FWorkerThread::FMessage::EType::SelectTrackByMetadata:
+				case FWorkerThread::FMessage::EType::InitialStreamAttributes:
 				{
-					switch(msg.Data.TrackSelection.StreamType)
+					// Map the language in case it is not yet.
+					if (msg.Data.InitialStreamAttribute.InitialSelection.Language_ISO639.IsSet())
+					{
+						msg.Data.InitialStreamAttribute.InitialSelection.Language_ISO639 = ISO639::MapTo639_1(msg.Data.InitialStreamAttribute.InitialSelection.Language_ISO639.GetValue());
+					}
+					switch(msg.Data.InitialStreamAttribute.StreamType)
 					{
 						case EStreamType::Video:
-							InitialTrackSelectionVid = MakeShared<FTrackMetadata, ESPMode::ThreadSafe>(msg.Data.TrackSelection.TrackMetadata);
+							StreamSelectionAttributesVid = msg.Data.InitialStreamAttribute.InitialSelection;
 							break;
 						case EStreamType::Audio:
-							InitialTrackSelectionAud = MakeShared<FTrackMetadata, ESPMode::ThreadSafe>(msg.Data.TrackSelection.TrackMetadata);
-							// For media with multiplexed tracks we can switch to a different track pretty much immediately.
-							if (ManifestReader.IsValid() && ManifestReader->GetPlaylistType() == TEXT("mp4"))
-							{
-								// Only select audio for now.
-								if (InitialTrackSelectionAud.IsValid())
-								{
-									MultiStreamBufferAud.SelectTrackByID(InitialTrackSelectionAud->TrackID);
-								}
-							}
+							StreamSelectionAttributesAud = msg.Data.InitialStreamAttribute.InitialSelection;
 							break;
 						case EStreamType::Subtitle:
-							InitialTrackSelectionTxt = MakeShared<FTrackMetadata, ESPMode::ThreadSafe>(msg.Data.TrackSelection.TrackMetadata);
+							StreamSelectionAttributesTxt = msg.Data.InitialStreamAttribute.InitialSelection;
+							break;
+						default:
 							break;
 					}
 					break;
 				}
 
+				case FWorkerThread::FMessage::EType::SelectTrackByAttributes:
+				{
+					if (!bIsClosing)
+					{
+						// Map the language in case it is not yet.
+						if (msg.Data.TrackSelection.TrackAttributes.Language_ISO639.IsSet())
+						{
+							msg.Data.TrackSelection.TrackAttributes.Language_ISO639 = ISO639::MapTo639_1(msg.Data.TrackSelection.TrackAttributes.Language_ISO639.GetValue());
+						}
+						switch(msg.Data.TrackSelection.StreamType)
+						{
+							case EStreamType::Video:
+								PendingTrackSelectionVid = MakeSharedTS<FStreamSelectionAttributes>(msg.Data.TrackSelection.TrackAttributes);
+								MultiStreamBufferVid.Activate();
+								break;
+							case EStreamType::Audio:
+								PendingTrackSelectionAud = MakeSharedTS<FStreamSelectionAttributes>(msg.Data.TrackSelection.TrackAttributes);
+								MultiStreamBufferAud.Activate();
+								break;
+							case EStreamType::Subtitle:
+								PendingTrackSelectionTxt = MakeSharedTS<FStreamSelectionAttributes>(msg.Data.TrackSelection.TrackAttributes);
+								MultiStreamBufferTxt.Activate();
+								break;
+						}
+					}
+					break;
+				}
+
+				case FWorkerThread::FMessage::EType::SelectTrackByMetadata:
+				{
+					// Currently not used. May be used again later.
+					break;
+				}
+
 				case FWorkerThread::FMessage::EType::DeselectTrack:
 				{
 					switch(msg.Data.TrackSelection.StreamType)
 					{
 						case EStreamType::Video:
-							InitialTrackSelectionVid.Reset();
+							MultiStreamBufferVid.Deselect();
 							break;
 						case EStreamType::Audio:
-							InitialTrackSelectionAud.Reset();
-							// For media with multiplexed tracks we can switch to a different track pretty much immediately.
-							if (ManifestReader.IsValid() && ManifestReader->GetPlaylistType() == TEXT("mp4"))
-							{
-								MultiStreamBufferAud.Deselect();
-							}
+							MultiStreamBufferAud.Deselect();
 							break;
 						case EStreamType::Subtitle:
-							InitialTrackSelectionTxt.Reset();
+							MultiStreamBufferTxt.Deselect();
 							break;
 					}
 					break;
@@ -1245,8 +1370,9 @@ void FAdaptiveStreamingPlayer::WorkerThreadFN()
 				case FWorkerThread::FMessage::EType::FragmentOpen:
 				{
 					TSharedPtrTS<IStreamSegment> pRequest = msg.Data.StreamReader.Request;
+					EStreamType reqType = msg.Data.StreamReader.Request->GetType();
 					// Check that the request is for this current playback sequence and not an outdated one.
-					if (pRequest.IsValid() && pRequest->GetPlaybackSequenceID() == CurrentPlaybackSequenceID)
+					if (pRequest.IsValid() && pRequest->GetPlaybackSequenceID() == CurrentPlaybackSequenceID[StreamTypeToArrayIndex(reqType)])
 					{
 						DispatchBufferUtilizationEvent(msg.Data.StreamReader.Request->GetType() == EStreamType::Video ? EStreamType::Video : EStreamType::Audio);
 
@@ -1270,18 +1396,17 @@ void FAdaptiveStreamingPlayer::WorkerThreadFN()
 				case FWorkerThread::FMessage::EType::FragmentClose:
 				{
 					TSharedPtrTS<IStreamSegment> pRequest = msg.Data.StreamReader.Request;
+					EStreamType reqType = msg.Data.StreamReader.Request->GetType();
 					// Check that the request is for this current playback sequence and not an outdated one.
-					if (pRequest.IsValid() && pRequest->GetPlaybackSequenceID() == CurrentPlaybackSequenceID)
+					if (pRequest.IsValid() && pRequest->GetPlaybackSequenceID() == CurrentPlaybackSequenceID[StreamTypeToArrayIndex(reqType)])
 					{
-						EStreamType reqType = msg.Data.StreamReader.Request->GetType();
-
 						// Dispatch download event
 						DispatchSegmentDownloadedEvent(pRequest);
 
 						// Dispatch buffer utilization
 						DispatchBufferUtilizationEvent(reqType == EStreamType::Video ? EStreamType::Video : EStreamType::Audio);
 
-						// Dispatch avergae throughput, bandwidth and latency event for video segments only.
+						// Dispatch average throughput, bandwidth and latency event for video segments only.
 						if (reqType == EStreamType::Video)
 						{
 							DispatchEvent(FMetricEvent::ReportBandwidth(StreamSelector->GetAverageBandwidth(), StreamSelector->GetAverageThroughput(), StreamSelector->GetAverageLatency()));
@@ -1426,7 +1551,7 @@ void FAdaptiveStreamingPlayer::HandlePlayStateChanges()
 void FAdaptiveStreamingPlayer::HandleMetadataChanges()
 {
 	// The timeline can change dynamically. Refresh it on occasion.
-	if (Manifest.IsValid() && ManifestReader.IsValid() && ManifestReader->GetPlaylistType() == TEXT("dash"))
+	if (Manifest.IsValid() && ManifestType == EMediaFormatType::DASH)
 	{
 		PlaybackState.SetSeekableRange(Manifest->GetSeekableTimeRange());
 		PlaybackState.SetTimelineRange(Manifest->GetTotalTimeRange());
@@ -1467,7 +1592,7 @@ void FAdaptiveStreamingPlayer::HandleMetadataChanges()
 							DispatchEvent(FMetricEvent::ReportTracksChanged());
 						}
 					}
-					
+
 					// Inform AEMS handler of the period transition.
 					AEMSEventHandler->PlaybackPeriodTransition(UpcomingPeriods[i].ID, UpcomingPeriods[i].TimeRange);
 
@@ -1661,7 +1786,7 @@ void FAdaptiveStreamingPlayer::HandleNewBufferedData()
 			// When we are dealing with a single multiplexed stream and one buffer is blocked then essentially all buffers
 			// must be considered blocked since demuxing cannot continue.
 			// FIXME: do this more elegant somehow
-			if (ManifestReader.IsValid() && ManifestReader->GetPlaylistType() == TEXT("mp4"))
+			if (ManifestType == EMediaFormatType::ISOBMFF)
 			{
 				if (VideoBufferStats.StreamBuffer.bLastPushWasBlocked ||
 					AudioBufferStats.StreamBuffer.bLastPushWasBlocked)
@@ -1889,19 +2014,8 @@ void FAdaptiveStreamingPlayer::HandleNewOutputData()
 
 
 
-//-----------------------------------------------------------------------------
-/**
- * Handles pending media segment requests.
- */
-void FAdaptiveStreamingPlayer::HandlePendingMediaSegmentRequests()
+void FAdaptiveStreamingPlayer::InternalHandlePendingStartRequest(const FTimeValue& CurrentTime)
 {
-	if (CurrentState == EPlayerState::eState_Paused || CurrentState == EPlayerState::eState_Error)
-	{
-		return;
-	}
-
-	FTimeValue CurrentTime = MEDIAutcTime::Current();
-
 	// Is there a start request?
 	if (PendingStartRequest.IsValid())
 	{
@@ -1970,6 +2084,7 @@ void FAdaptiveStreamingPlayer::HandlePendingMediaSegmentRequests()
 							err.SetMessage("Could not locate the playback period to begin playback at.");
 							err.SetCode(INTERR_COULD_NOT_LOCATE_START_PERIOD);
 							PostError(err);
+							InitialPlayPeriod.Reset();
 							PendingStartRequest.Reset();
 							break;
 						}
@@ -1982,12 +2097,70 @@ void FAdaptiveStreamingPlayer::HandlePendingMediaSegmentRequests()
 				// Is it ready yet?
 				switch(InitialPlayPeriod->GetReadyState())
 				{
-					case IManifest::IPlayPeriod::EReadyState::NotReady:
+					case IManifest::IPlayPeriod::EReadyState::NotLoaded:
 					{
+						// If there are already pending track selections update the initial selections with them and clear the pending ones.
+						if (PendingTrackSelectionVid.IsValid())
+						{
+							StreamSelectionAttributesVid = *PendingTrackSelectionVid;
+							PendingTrackSelectionVid.Reset();
+						}
+						if (PendingTrackSelectionAud.IsValid())
+						{
+							StreamSelectionAttributesAud = *PendingTrackSelectionAud;
+							PendingTrackSelectionAud.Reset();
+						}
+						if (PendingTrackSelectionTxt.IsValid())
+						{
+							StreamSelectionAttributesTxt = *PendingTrackSelectionTxt;
+							PendingTrackSelectionTxt.Reset();
+						}
+
 						// Prepare the play period. This must immediately change the state away from NotReady
-						InitialPlayPeriod->SetStreamPreferences(StreamPreferences);
+						InitialPlayPeriod->SetStreamPreferences(EStreamType::Video, StreamSelectionAttributesVid);
+						InitialPlayPeriod->SetStreamPreferences(EStreamType::Audio, StreamSelectionAttributesAud);
+						InitialPlayPeriod->SetStreamPreferences(EStreamType::Subtitle, StreamSelectionAttributesTxt);
+						InitialPlayPeriod->Load();
+						break;
+					}
+					case IManifest::IPlayPeriod::EReadyState::Loading:
+					{
+						// Period is not yet ready. Check again on the next iteration.
+						break;
+					}
+					case IManifest::IPlayPeriod::EReadyState::Loaded:
+					{
+						// Period is loaded and configured according to the stream preferences.
+						// Prepare the initial quality/bitrate for playback.
+						int64 StartingBitrate = StreamSelector->GetAverageBandwidth();
+						if (StartingBitrate <= 0)
+						{
+							StartingBitrate = PlayerOptions.GetValue(OptionKeyInitialBitrate).SafeGetInt64(0);
+						}
+						if (StartingBitrate <= 0)
+						{
+							StartingBitrate = InitialPlayPeriod->GetDefaultStartingBitrate();
+							if (StartingBitrate <= 0)
+							{
+								StartingBitrate = 1000000;
+							}
+						}
+						StreamSelector->SetBandwidth(StartingBitrate);
+						// On the very first playback enforce the starting bitrate.
+						if (PendingStartRequest->bIsPlayStart)
+						{
+							StreamSelector->SetForcedNextBandwidth(StartingBitrate);
+						}
+
+						// Set the current average video bitrate in the player options for the period to retrieve if necessary.
+						PlayerOptions.SetOrUpdate(OptionKeyCurrentAvgStartingVideoBitrate, FVariantValue(StreamSelector->GetAverageBandwidth()));
+
 						InitialPlayPeriod->PrepareForPlay();
-						check(InitialPlayPeriod->GetReadyState() != IManifest::IPlayPeriod::EReadyState::NotReady);
+						break;
+					}
+					case IManifest::IPlayPeriod::EReadyState::Preparing:
+					{
+						// Period is not yet ready. Check again on the next iteration.
 						break;
 					}
 					case IManifest::IPlayPeriod::EReadyState::IsReady:
@@ -2010,10 +2183,40 @@ void FAdaptiveStreamingPlayer::HandlePendingMediaSegmentRequests()
 						StreamSelector->SetCurrentPlaybackPeriod(EStreamType::Audio, CurrentPlayPeriodAudio);
 
 						// Single stream mp4 files cannot switch streams. Failures need to retry the same file.
-						if (ManifestReader.IsValid() && ManifestReader->GetPlaylistType() == TEXT("mp4"))
+						if (ManifestType == EMediaFormatType::ISOBMFF)
 						{
 							StreamSelector->SetCanSwitchToAlternateStreams(false);
 						}
+						TSharedPtrTS<FBufferSourceInfo> BufferSourceInfoVid = CurrentPlayPeriodVideo->GetSelectedStreamBufferSourceInfo(EStreamType::Video);
+						TSharedPtrTS<FBufferSourceInfo> BufferSourceInfoAud = CurrentPlayPeriodAudio->GetSelectedStreamBufferSourceInfo(EStreamType::Audio);
+						TSharedPtrTS<FBufferSourceInfo> BufferSourceInfoTxt = CurrentPlayPeriodText->GetSelectedStreamBufferSourceInfo(EStreamType::Subtitle);
+						MultiStreamBufferVid.SelectTrackWhenAvailable(BufferSourceInfoVid);
+						MultiStreamBufferAud.SelectTrackWhenAvailable(BufferSourceInfoAud);
+						MultiStreamBufferTxt.SelectTrackWhenAvailable(BufferSourceInfoTxt);
+						if (BufferSourceInfoVid.IsValid())
+						{
+							SelectedStreamAttributesVid.UpdateWith(BufferSourceInfoVid->Kind, BufferSourceInfoVid->Language, BufferSourceInfoVid->HardIndex);
+							if (ManifestType != EMediaFormatType::ISOBMFF)
+							{
+								StreamSelectionAttributesVid.UpdateIfOverrideSet(BufferSourceInfoVid->Kind, BufferSourceInfoVid->Language);
+							}
+						}
+						if (BufferSourceInfoAud.IsValid())
+						{
+							SelectedStreamAttributesAud.UpdateWith(BufferSourceInfoAud->Kind, BufferSourceInfoAud->Language, BufferSourceInfoAud->HardIndex);
+							if (ManifestType != EMediaFormatType::ISOBMFF)
+							{
+								StreamSelectionAttributesAud.UpdateIfOverrideSet(BufferSourceInfoAud->Kind, BufferSourceInfoAud->Language);
+							}
+						}
+						if (BufferSourceInfoTxt.IsValid())
+						{
+							SelectedStreamAttributesTxt.UpdateWith(BufferSourceInfoTxt->Kind, BufferSourceInfoTxt->Language, BufferSourceInfoTxt->HardIndex);
+							if (ManifestType != EMediaFormatType::ISOBMFF)
+							{
+								StreamSelectionAttributesTxt.UpdateIfOverrideSet(BufferSourceInfoTxt->Kind, BufferSourceInfoTxt->Language);
+							}
+						}
 
 						// Apply any set resolution limit at the start of the new play period.
 						UpdateStreamResolutionLimit();
@@ -2024,7 +2227,7 @@ void FAdaptiveStreamingPlayer::HandlePendingMediaSegmentRequests()
 							case IManifest::EType::OnDemand:
 								StreamSelector->SetPresentationType(IAdaptiveStreamSelector::EMediaPresentationType::OnDemand);
 								break;
-// FIXME: Need to differentiate between "Live" and "Realtime" here!
+							// FIXME: Need to differentiate between "Live" and "Realtime" here!
 							default:
 							case IManifest::EType::Live:
 								StreamSelector->SetPresentationType(IAdaptiveStreamSelector::EMediaPresentationType::Realtime);
@@ -2104,11 +2307,6 @@ void FAdaptiveStreamingPlayer::HandlePendingMediaSegmentRequests()
 						}
 						break;
 					}
-					case IManifest::IPlayPeriod::EReadyState::Preparing:
-					{
-						// Period is not yet ready. Check again on the next iteration.
-						break;
-					}
 				}
 			}
 		}
@@ -2118,7 +2316,11 @@ void FAdaptiveStreamingPlayer::HandlePendingMediaSegmentRequests()
 			PendingStartRequest.Reset();
 		}
 	}
+}
+
 
+void FAdaptiveStreamingPlayer::InternalHandlePendingFirstSegmentRequest(const FTimeValue& CurrentTime)
+{
 	// Is there a play start initial segment download request?
 	if (PendingFirstSegmentRequest.IsValid())
 	{
@@ -2202,8 +2404,8 @@ void FAdaptiveStreamingPlayer::HandlePendingMediaSegmentRequests()
 				CompletedSegmentRequests.Add(AlreadyEndedStreams[i]->GetType(), AlreadyEndedStreams[i]);
 			}
 
-			// Kick off the request
-			IStreamReader::EAddResult ReqResult = StreamReaderHandler->AddRequest(CurrentPlaybackSequenceID, PendingFirstSegmentRequest);
+			// Kick off the request. Use the sequence ID from video. All sequence IDs have been set equal for start requests so it does not matter.
+			IStreamReader::EAddResult ReqResult = StreamReaderHandler->AddRequest(CurrentPlaybackSequenceID[0], PendingFirstSegmentRequest);
 			if (ReqResult != IStreamReader::EAddResult::Added)
 			{
 				FErrorDetail err;
@@ -2215,7 +2417,88 @@ void FAdaptiveStreamingPlayer::HandlePendingMediaSegmentRequests()
 		}
 		PendingFirstSegmentRequest.Reset();
 	}
+}
+
 
+void FAdaptiveStreamingPlayer::InternalHandleCompletedSegmentRequests(const FTimeValue& CurrentTime)
+{
+
+	// Check the streams that reached EOS. If all of them are done decide what to do, end playback or loop.
+	if (CompletedSegmentRequests.Num())
+	{
+		bool bAllAtEOS = true;
+		if (bHaveVideoReader.GetWithDefault(false))
+		{
+			if (!CompletedSegmentRequests.Contains(EStreamType::Video))
+			{
+				bAllAtEOS = false;
+			}
+			else
+			{
+				MultiStreamBufferVid.PushEndOfDataAll();
+			}
+		}
+		if (bHaveAudioReader.GetWithDefault(false))
+		{
+			if (!CompletedSegmentRequests.Contains(EStreamType::Audio))
+			{
+				bAllAtEOS = false;
+			}
+			else
+			{
+				MultiStreamBufferAud.PushEndOfDataAll();
+			}
+		}
+		if (bHaveTextReader.GetWithDefault(false))
+		{
+			if (!CompletedSegmentRequests.Contains(EStreamType::Subtitle))
+			{
+				bAllAtEOS = false;
+			}
+			else
+			{
+				MultiStreamBufferTxt.PushEndOfDataAll();
+			}
+		}
+		if (bAllAtEOS)
+		{
+			// Move the finished requests into a local work map, emptying the original.
+			TMultiMap<EStreamType, TSharedPtrTS<IStreamSegment>> LocalFinishedRequests(MoveTemp(CompletedSegmentRequests));
+
+			// Looping enabled?
+			if (CurrentLoopParam.bEnableLooping)
+			{
+				if (Manifest.IsValid())
+				{
+					FTimeRange Seekable = Manifest->GetSeekableTimeRange();
+
+					PendingStartRequest = MakeSharedTS<FPendingStartRequest>();
+					PendingStartRequest->SearchType = IManifest::ESearchType::Closest;
+					PendingStartRequest->StartAt.Time = Seekable.Start;
+					PendingStartRequest->bForLooping = true;
+					PendingStartRequest->FinishedRequests = MoveTemp(LocalFinishedRequests);
+				}
+			}
+		}
+	}
+}
+
+
+//-----------------------------------------------------------------------------
+/**
+ * Handles pending media segment requests.
+ */
+void FAdaptiveStreamingPlayer::HandlePendingMediaSegmentRequests()
+{
+	if (CurrentState == EPlayerState::eState_Error)
+	{
+		return;
+	}
+
+	FTimeValue CurrentTime = MEDIAutcTime::Current();
+	InternalHandlePendingStartRequest(CurrentTime);
+	InternalHandlePendingFirstSegmentRequest(CurrentTime);
+	InternalHandleSegmentTrackChanges(CurrentTime);
 
 	// Are there completed downloads for which we need to find the next segment?
 	TQueue<FPendingSegmentRequest> pendingRequests;
@@ -2226,29 +2509,28 @@ void FAdaptiveStreamingPlayer::HandlePendingMediaSegmentRequests()
 		pendingRequests.Dequeue(FinishedReq);
 
 		TSharedPtrTS<IManifest::IPlayPeriod> SegmentPlayPeriod;
-		if (FinishedReq.Request.IsValid())
+
+		EStreamType StreamType = FinishedReq.StreamType != EStreamType::Unsupported ? FinishedReq.StreamType : FinishedReq.Request.IsValid() ? FinishedReq.Request->GetType() : EStreamType::Unsupported;
+		switch(StreamType)
 		{
-			switch(FinishedReq.Request->GetType())
+			default:
 			{
-				default:
-				{
-					break;
-				}
-				case EStreamType::Video:
-				 {
-					SegmentPlayPeriod = CurrentPlayPeriodVideo;
-					break;
-				}
-				case EStreamType::Audio:
-				{
-					SegmentPlayPeriod = CurrentPlayPeriodAudio;
-					break;
-				}
-				case EStreamType::Subtitle:
+				break;
+			}
+			case EStreamType::Video:
 				{
-					SegmentPlayPeriod = CurrentPlayPeriodText;
-					break;
-				}
+				SegmentPlayPeriod = CurrentPlayPeriodVideo;
+				break;
+			}
+			case EStreamType::Audio:
+			{
+				SegmentPlayPeriod = CurrentPlayPeriodAudio;
+				break;
+			}
+			case EStreamType::Subtitle:
+			{
+				SegmentPlayPeriod = CurrentPlayPeriodText;
+				break;
 			}
 		}
 
@@ -2259,13 +2541,56 @@ void FAdaptiveStreamingPlayer::HandlePendingMediaSegmentRequests()
 		}
 		else if (SegmentPlayPeriod.IsValid())
 		{
-			// Did we move into a new period that needs to be made ready?
+			// Is this a startover for a new period?
+			if (FinishedReq.bStartOver && !FinishedReq.Period.IsValid())
+			{
+				// Locate the period the startover time is in.
+				TSharedPtrTS<IManifest::IPlayPeriod> StartoverPeriod;
+			// NOTE: This will NOT work with looping streams!
+				IManifest::FResult PeriodResult = Manifest->FindPlayPeriod(StartoverPeriod, FinishedReq.StartoverPosition, IManifest::ESearchType::After);
+				// When found or at the end of the stream
+				if (PeriodResult.GetType() == IManifest::FResult::EType::Found)
+				{
+					FinishedReq.Period = MoveTemp(StartoverPeriod);
+				}
+				else if (PeriodResult.GetType() == IManifest::FResult::EType::PastEOS)
+				{
+					// At the end of the stream we just don't do anything. Just ignore the switch.
+					continue;
+				}
+				else
+				{
+					FErrorDetail err;
+					err.SetFacility(Facility::EFacility::Player);
+					err.SetMessage("Could not locate startover period for track switch");
+					err.SetCode(INTERR_COULD_NOT_LOCATE_START_PERIOD);
+					PostError(err);
+					break;
+				}
+			}
+
+			// Did we move into a new or a startover period that needs to be made ready?
 			if (FinishedReq.Period.IsValid())
 			{
 				IManifest::IPlayPeriod::EReadyState PeriodState = FinishedReq.Period->GetReadyState();
-				if (PeriodState == IManifest::IPlayPeriod::EReadyState::NotReady)
+				if (PeriodState == IManifest::IPlayPeriod::EReadyState::NotLoaded)
+				{
+					FinishedReq.Period->SetStreamPreferences(EStreamType::Video, StreamSelectionAttributesVid);
+					FinishedReq.Period->SetStreamPreferences(EStreamType::Audio, StreamSelectionAttributesAud);
+					FinishedReq.Period->SetStreamPreferences(EStreamType::Subtitle, StreamSelectionAttributesTxt);
+					FinishedReq.Period->Load();
+					NextPendingSegmentRequests.Enqueue(MoveTemp(FinishedReq));
+					continue;
+				}
+				else if (PeriodState == IManifest::IPlayPeriod::EReadyState::Loading)
 				{
-					FinishedReq.Period->SetStreamPreferences(StreamPreferences);
+					NextPendingSegmentRequests.Enqueue(MoveTemp(FinishedReq));
+					continue;
+				}
+				else if (PeriodState == IManifest::IPlayPeriod::EReadyState::Loaded)
+				{
+					// Set the current average video bitrate in the player options for the period to retrieve if necessary.
+					PlayerOptions.SetOrUpdate(OptionKeyCurrentAvgStartingVideoBitrate, FVariantValue(StreamSelector->GetAverageBandwidth()));
 					FinishedReq.Period->PrepareForPlay();
 					NextPendingSegmentRequests.Enqueue(MoveTemp(FinishedReq));
 					continue;
@@ -2275,29 +2600,45 @@ void FAdaptiveStreamingPlayer::HandlePendingMediaSegmentRequests()
 					NextPendingSegmentRequests.Enqueue(MoveTemp(FinishedReq));
 					continue;
 				}
-				else
+				else if (PeriodState == IManifest::IPlayPeriod::EReadyState::IsReady)
 				{
-					if (FinishedReq.Request->GetType() == EStreamType::Video)
+					if (StreamType == EStreamType::Video)
 					{
 						SegmentPlayPeriod = FinishedReq.Period;
 						CurrentPlayPeriodVideo = FinishedReq.Period;
 						StreamSelector->SetCurrentPlaybackPeriod(EStreamType::Video, CurrentPlayPeriodVideo);
 						UpdateStreamResolutionLimit();
+						if (!FinishedReq.bStartOver)
+						{
+							MultiStreamBufferVid.AddUpcomingBuffer(SegmentPlayPeriod->GetSelectedStreamBufferSourceInfo(StreamType));
+						}
 					}
-					else if (FinishedReq.Request->GetType() == EStreamType::Audio)
+					else if (StreamType == EStreamType::Audio)
 					{
 						SegmentPlayPeriod = FinishedReq.Period;
 						CurrentPlayPeriodAudio = FinishedReq.Period;
 						StreamSelector->SetCurrentPlaybackPeriod(EStreamType::Audio, CurrentPlayPeriodAudio);
+						if (!FinishedReq.bStartOver)
+						{
+							MultiStreamBufferAud.AddUpcomingBuffer(SegmentPlayPeriod->GetSelectedStreamBufferSourceInfo(StreamType));
+						}
 					}
-					else if (FinishedReq.Request->GetType() == EStreamType::Subtitle)
+					else if (StreamType == EStreamType::Subtitle)
 					{
 						SegmentPlayPeriod = FinishedReq.Period;
 						CurrentPlayPeriodText = FinishedReq.Period;
 						//StreamSelector->SetCurrentPlaybackPeriod(EStreamType::Subtitle, CurrentPlayPeriodText);
+						if (!FinishedReq.bStartOver)
+						{
+							MultiStreamBufferTxt.AddUpcomingBuffer(SegmentPlayPeriod->GetSelectedStreamBufferSourceInfo(StreamType));
+						}
+					}
+					// Add this to the upcoming periods the play position will move into and metadata will need to be updated then
+					// unless this is a startover request.
+					if (!FinishedReq.bStartOver)
+					{
+						AddUpcomingPeriod(SegmentPlayPeriod);
 					}
-					// Add this to the upcoming periods the play position will move into and metadata will need to be updated then.
-					AddUpcomingPeriod(SegmentPlayPeriod);
 				}
 			}
 
@@ -2314,15 +2655,22 @@ void FAdaptiveStreamingPlayer::HandlePendingMediaSegmentRequests()
 				break;
 			}
 
-			TSharedPtrTS<IStreamSegment> pNext;
+			TSharedPtrTS<IStreamSegment> NextSegment;
 			IManifest::FResult Result;
-			if (Action == IAdaptiveStreamSelector::ESegmentAction::FetchNext)
+			if (!FinishedReq.bStartOver)
 			{
-				Result = SegmentPlayPeriod->GetNextSegment(pNext, FinishedReq.Request);
+				if (Action == IAdaptiveStreamSelector::ESegmentAction::FetchNext)
+				{
+					Result = SegmentPlayPeriod->GetNextSegment(NextSegment, FinishedReq.Request);
+				}
+				else if (Action == IAdaptiveStreamSelector::ESegmentAction::Retry || Action == IAdaptiveStreamSelector::ESegmentAction::Fill)
+				{
+					Result = SegmentPlayPeriod->GetRetrySegment(NextSegment, FinishedReq.Request, Action == IAdaptiveStreamSelector::ESegmentAction::Fill);
+				}
 			}
-			else if (Action == IAdaptiveStreamSelector::ESegmentAction::Retry || Action == IAdaptiveStreamSelector::ESegmentAction::Fill)
+			else
 			{
-				Result = SegmentPlayPeriod->GetRetrySegment(pNext, FinishedReq.Request, Action == IAdaptiveStreamSelector::ESegmentAction::Fill);
+				Result = SegmentPlayPeriod->GetContinuationSegment(NextSegment, StreamType, CurrentLoopState, FinishedReq.StartoverPosition, IManifest::ESearchType::After);
 			}
 			switch(Result.GetType())
 			{
@@ -2334,8 +2682,52 @@ void FAdaptiveStreamingPlayer::HandlePendingMediaSegmentRequests()
 				}
 				case IManifest::FResult::EType::Found:
 				{
-					pNext->SetExecutionDelay(ActionDelay);
-					IStreamReader::EAddResult ReqResult = StreamReaderHandler->AddRequest(CurrentPlaybackSequenceID, pNext);
+					// Switching tracks?
+					if (FinishedReq.bStartOver)
+					{
+						// Tell the multi stream buffer where to switch to and update the selection defaults.
+						TSharedPtrTS<FBufferSourceInfo> BufferSourceInfo = SegmentPlayPeriod->GetSelectedStreamBufferSourceInfo(StreamType);
+						switch(StreamType)
+						{
+							case EStreamType::Video:
+							{
+								MultiStreamBufferVid.SelectTrackWhenAvailable(BufferSourceInfo);
+								SelectedStreamAttributesVid.UpdateWith(BufferSourceInfo->Kind, BufferSourceInfo->Language, BufferSourceInfo->HardIndex);
+								if (ManifestType != EMediaFormatType::ISOBMFF)
+								{
+									StreamSelectionAttributesVid.ClearOverrideIndex();
+								}
+								break;
+							}
+							case EStreamType::Audio:
+							{
+								MultiStreamBufferAud.SelectTrackWhenAvailable(BufferSourceInfo);
+								SelectedStreamAttributesAud.UpdateWith(BufferSourceInfo->Kind, BufferSourceInfo->Language, BufferSourceInfo->HardIndex);
+								if (ManifestType != EMediaFormatType::ISOBMFF)
+								{
+									StreamSelectionAttributesAud.ClearOverrideIndex();
+								}
+								break;
+							}
+							case EStreamType::Subtitle:
+							{
+								MultiStreamBufferTxt.SelectTrackWhenAvailable(BufferSourceInfo);
+								SelectedStreamAttributesTxt.UpdateWith(BufferSourceInfo->Kind, BufferSourceInfo->Language, BufferSourceInfo->HardIndex);
+								if (ManifestType != EMediaFormatType::ISOBMFF)
+								{
+									StreamSelectionAttributesTxt.ClearOverrideIndex();
+								}
+								break;
+							}
+							default:
+							{
+								break;
+							}
+						}
+					}
+
+					NextSegment->SetExecutionDelay(ActionDelay);
+					IStreamReader::EAddResult ReqResult = StreamReaderHandler->AddRequest(CurrentPlaybackSequenceID[StreamTypeToArrayIndex(StreamType)], NextSegment);
 					if (ReqResult != IStreamReader::EAddResult::Added)
 					{
 						FErrorDetail err;
@@ -2403,63 +2795,70 @@ void FAdaptiveStreamingPlayer::HandlePendingMediaSegmentRequests()
 		}
 	}
 
-	// Check the streams that reached EOS. If all of them are done decide what to do, end playback or loop.
-	if (CompletedSegmentRequests.Num())
+	InternalHandleCompletedSegmentRequests(CurrentTime);
+}
+
+
+void FAdaptiveStreamingPlayer::InternalHandleSegmentTrackChanges(const FTimeValue& CurrentTime)
+{
+	if (ManifestType == EMediaFormatType::ISOBMFF)
 	{
-		bool bAllAtEOS = true;
-		if (bHaveVideoReader.GetWithDefault(false))
-		{
-			if (!CompletedSegmentRequests.Contains(EStreamType::Video))
-			{
-				bAllAtEOS = false;
-			}
-			else
-			{
-				MultiStreamBufferVid.PushEndOfDataAll();
-			}
-		}
-		if (bHaveAudioReader.GetWithDefault(false))
+		if (PendingTrackSelectionAud.IsValid() && CurrentPlayPeriodAudio.IsValid())
 		{
-			if (!CompletedSegmentRequests.Contains(EStreamType::Audio))
-			{
-				bAllAtEOS = false;
-			}
-			else
+			IManifest::IPlayPeriod::ETrackChangeResult Result = CurrentPlayPeriodAudio->ChangeTrackStreamPreference(EStreamType::Audio, *PendingTrackSelectionAud);
+			TSharedPtrTS<FBufferSourceInfo> BufferSourceInfo = CurrentPlayPeriodAudio->GetSelectedStreamBufferSourceInfo(EStreamType::Audio);
+			if (Result == IManifest::IPlayPeriod::ETrackChangeResult::Changed)
 			{
-				MultiStreamBufferAud.PushEndOfDataAll();
+				MultiStreamBufferAud.SelectTrackWhenAvailable(BufferSourceInfo);
+				StreamSelectionAttributesAud = *PendingTrackSelectionAud;
 			}
+			SelectedStreamAttributesAud.UpdateWith(BufferSourceInfo->Kind, BufferSourceInfo->Language, BufferSourceInfo->HardIndex);
+			PendingTrackSelectionAud.Reset();
 		}
-		if (bHaveTextReader.GetWithDefault(false))
+		// Ignore video and subtitle changes for now.
+		PendingTrackSelectionVid.Reset();
+		PendingTrackSelectionTxt.Reset();
+	}
+	else if (ManifestType == EMediaFormatType::DASH)
+	{
+		if (PendingStartRequest.IsValid() || PendingFirstSegmentRequest.IsValid())
 		{
-			if (!CompletedSegmentRequests.Contains(EStreamType::Subtitle))
-			{
-				bAllAtEOS = false;
-			}
-			else
-			{
-				MultiStreamBufferTxt.PushEndOfDataAll();
-			}
+			return;
 		}
-		if (bAllAtEOS)
-		{
-			// Move the finished requests into a local work map, emptying the original.
-			TMultiMap<EStreamType, TSharedPtrTS<IStreamSegment>> LocalFinishedRequests(MoveTemp(CompletedSegmentRequests));
 
-			// Looping enabled?
-			if (CurrentLoopParam.bEnableLooping)
+		if (PendingTrackSelectionAud.IsValid() && CurrentPlayPeriodAudio.IsValid())
+		{
+			IManifest::IPlayPeriod::ETrackChangeResult TrackChangeResult = CurrentPlayPeriodAudio->ChangeTrackStreamPreference(EStreamType::Audio, *PendingTrackSelectionAud);
+			if (TrackChangeResult == IManifest::IPlayPeriod::ETrackChangeResult::NewPeriodNeeded)
 			{
-				if (Manifest.IsValid())
+				// Make sure any finishing current requests will be ignored by increasing the sequence ID.
+				++CurrentPlaybackSequenceID[StreamTypeToArrayIndex(EStreamType::Audio)];
+				// Cancel any ongoing segment download now.
+				if (StreamReaderHandler)
 				{
-					FTimeRange Seekable = Manifest->GetSeekableTimeRange();
-
-					PendingStartRequest = MakeSharedTS<FPendingStartRequest>();
-					PendingStartRequest->SearchType = IManifest::ESearchType::Closest;
-					PendingStartRequest->StartAt.Time = Seekable.Start;
-					PendingStartRequest->bForLooping = true;
-					PendingStartRequest->FinishedRequests = MoveTemp(LocalFinishedRequests);
+					StreamReaderHandler->CancelRequest(EStreamType::Audio, true);
 				}
+
+				StreamSelectionAttributesAud = *PendingTrackSelectionAud;
+
+				FPendingSegmentRequest NextReq;
+				NextReq.bStartOver = true;
+				NextReq.StartoverPosition.Time = PlaybackState.GetPlayPosition();
+				NextReq.StreamType = EStreamType::Audio;
+				NextPendingSegmentRequests.Enqueue(MoveTemp(NextReq));
 			}
+			PendingTrackSelectionAud.Reset();
 		}
+
+		// Ignore video and subtitle changes for now.
+		PendingTrackSelectionVid.Reset();
+		PendingTrackSelectionTxt.Reset();
+	}
+	else
+	{
+		PendingTrackSelectionVid.Reset();
+		PendingTrackSelectionAud.Reset();
+		PendingTrackSelectionTxt.Reset();
 	}
 }
 
@@ -2905,69 +3304,35 @@ void FAdaptiveStreamingPlayer::CheckForStreamEnd()
  */
 bool FAdaptiveStreamingPlayer::InternalStartAt(const FSeekParam& NewPosition)
 {
-	++CurrentPlaybackSequenceID;
+	uint32 NextPlaybackID = Utils::Max(Utils::Max(CurrentPlaybackSequenceID[0], CurrentPlaybackSequenceID[1]), CurrentPlaybackSequenceID[2]) + 1;
+	CurrentPlaybackSequenceID[0] = NextPlaybackID;
+	CurrentPlaybackSequenceID[1] = NextPlaybackID;
+	CurrentPlaybackSequenceID[2] = NextPlaybackID;
 
 	StreamState = EStreamState::eStream_Running;
 
 	PlaybackState.SetHasEnded(false);
 
-	int64 StartingBitrate = StreamSelector->GetAverageBandwidth();
-	if (StartingBitrate <= 0)
-	{
-		StartingBitrate = PlayerOptions.GetValue(OptionKeyInitialBitrate).SafeGetInt64(0);
-	}
-	if (StartingBitrate <= 0)
-	{
-		StartingBitrate = Manifest->GetDefaultStartingBitrate();
-		if (StartingBitrate <= 0)
-		{
-			StartingBitrate = 500000;
-		}
-	}
-
 	// Initialize buffers
 	MultiStreamBufferVid.CapacitySet(PlayerConfig.StreamBufferConfigVideo);
 	MultiStreamBufferAud.CapacitySet(PlayerConfig.StreamBufferConfigAudio);
 	MultiStreamBufferTxt.CapacitySet(PlayerConfig.StreamBufferConfigText);
-	// Auto-select the first track that is pushing data into the buffer.
-	MultiStreamBufferVid.AutoselectFirstTrack();
-	MultiStreamBufferAud.AutoselectFirstTrack();
-	MultiStreamBufferTxt.AutoselectFirstTrack();
+	if (ManifestType == EMediaFormatType::ISOBMFF)
+	{
+		// mp4 has multiple tracks in a single stream that provide data all at the same time.
+		MultiStreamBufferVid.SetParallelTrackMode();
+		MultiStreamBufferAud.SetParallelTrackMode();
+		MultiStreamBufferTxt.SetParallelTrackMode();
+	}
 
 	// Update data availability states in case this wasn't done yet.
 	UpdateDataAvailabilityState(DataAvailabilityStateVid, Metrics::FDataAvailabilityChange::EAvailability::DataNotAvailable);
 	UpdateDataAvailabilityState(DataAvailabilityStateAud, Metrics::FDataAvailabilityChange::EAvailability::DataNotAvailable);
 	UpdateDataAvailabilityState(DataAvailabilityStateTxt, Metrics::FDataAvailabilityChange::EAvailability::DataNotAvailable);
 
-	// For an mp4 file we expect several audio tracks to be available and select the one matching the selected index, if there is one.
-	if (ManifestReader.IsValid() && ManifestReader->GetPlaylistType() == TEXT("mp4"))
-	{
-		// Only select audio for now.
-		if (InitialTrackSelectionAud.IsValid())
-		{
-			MultiStreamBufferAud.SelectTrackByID(InitialTrackSelectionAud->TrackID);
-		}
-// TODO/FIXME: Cannot start with audio disabled at the moment
-#if 0
-		else
-		{
-			mMultiStreamBufferAud.Deselect();
-		}
-#endif
-	}
-
 	// (Re-)configure the stream selector
-	StreamSelector->SetBandwidth(StartingBitrate);
 	StreamSelector->SetBandwidthCeiling(BitrateCeiling);
 	StreamSelector->SetMaxVideoResolution(VideoResolutionLimitWidth, VideoResolutionLimitHeight);
-	// On the very first playback enforce the starting bitrate.
-	if (bIsPlayStart)
-	{
-		StreamSelector->SetForcedNextBandwidth(StartingBitrate);
-	}
-
-	// From this point on any further start will not be the very first one any more.
-	bIsPlayStart = false;
 
 	// Create the stream reader
 	check(StreamReaderHandler == nullptr);
@@ -3001,7 +3366,8 @@ bool FAdaptiveStreamingPlayer::InternalStartAt(const FSeekParam& NewPosition)
 	PendingStartRequest = MakeSharedTS<FPendingStartRequest>();
 	PendingStartRequest->SearchType = IManifest::ESearchType::Closest;
 	PendingStartRequest->StartAt.Time = NewPosition.Time;
-	PendingStartRequest->InitialBandwidth.Set(StartingBitrate);
+	PendingStartRequest->bIsPlayStart = bIsPlayStart;
+
 
 	// The fragment should be the closest to the time stamp unless we are rebuffering in which case it should be the one we failed on.
 /*
@@ -3019,6 +3385,9 @@ bool FAdaptiveStreamingPlayer::InternalStartAt(const FSeekParam& NewPosition)
 	// Let the AEMS handler know that we are starting up.
 	AEMSEventHandler->PlaybackStartingUp();
 
+	// From this point on any further start will not be the very first one any more.
+	bIsPlayStart = false;
+
 	return true;
 }
 
@@ -3105,7 +3474,9 @@ void FAdaptiveStreamingPlayer::InternalStop(bool bHoldCurrentFrame)
 {
 	// Increase the playback sequence ID. This will cause all async segment request message
 	// to be discarded when they are received.
-	++CurrentPlaybackSequenceID;
+	++CurrentPlaybackSequenceID[0];
+	++CurrentPlaybackSequenceID[1];
+	++CurrentPlaybackSequenceID[2];
 
 	// Pause decoders to stop feeding them data.
 	// NOTE: A decoder may be in InputNeeded() / FeedDecoder() at this point!
@@ -3214,6 +3585,10 @@ void FAdaptiveStreamingPlayer::InternalClose()
 	CurrentLoopState.Reset();
 	NextLoopStates.Clear();
 
+	PendingTrackSelectionVid.Reset();
+	PendingTrackSelectionAud.Reset();
+	PendingTrackSelectionTxt.Reset();
+
 	// Flush all events from the AEMS handler
 	AEMSEventHandler->FlushEverything();
 
diff --git a/Engine/Plugins/Media/ElectraPlayer/Source/ElectraPlayerRuntime/Private/Runtime/Player/AdaptiveStreamingPlayer.h b/Engine/Plugins/Media/ElectraPlayer/Source/ElectraPlayerRuntime/Private/Runtime/Player/AdaptiveStreamingPlayer.h
--- a/Engine/Plugins/Media/ElectraPlayer/Source/ElectraPlayerRuntime/Private/Runtime/Player/AdaptiveStreamingPlayer.h
+++ b/Engine/Plugins/Media/ElectraPlayer/Source/ElectraPlayerRuntime/Private/Runtime/Player/AdaptiveStreamingPlayer.h
@@ -48,7 +48,7 @@ public:
 	virtual void AddMetricsReceiver(IAdaptiveStreamingPlayerMetrics* InMetricsReceiver) = 0;
 	virtual void RemoveMetricsReceiver(IAdaptiveStreamingPlayerMetrics* InMetricsReceiver) = 0;
 
-	
+
 	//-------------------------------------------------------------------------
 	// Application Event or Metadata Stream (AEMS) receiver
 	//   Please refer to ISO/IEC 23009-1:2019/DAM 1:2020(E)
@@ -59,13 +59,23 @@ public:
 	//-------------------------------------------------------------------------
 	// Initializes the player. Options may be passed in to affect behaviour.
 	//
-	virtual bool Initialize(const FParamDict& InOptions) = 0;
+	virtual void Initialize(const FParamDict& InOptions) = 0;
+
+	/**
+	 * Sets the attributes for the stream to start buffering for and playing.
+	 * This must be set before calling SeekTo() or LoadManifest() to have an immediate effect.
+	 * A best effort to match a stream with the given attributes will be made.
+	 * If there is no exact match the player will make an educated guess.
+	 * Any later call to SelectTrackByMetadata() internally overwrites these initial attributes
+	 * with those of the explicitly selected track.
+	 */
+	virtual void SetInitialStreamAttributes(EStreamType StreamType, const FStreamSelectionAttributes& InitialSelection) = 0;
 
 
 	//-------------------------------------------------------------------------
 	// Manifest loading functions
 	//
-	//! Issue a load and parse of the manifest/master playlist file.
+	//! Issue a load and parse of the manifest/master playlist file. Make initial stream selection choice by calling SetInitialStreamAttributes() beforehand.
 	virtual void LoadManifest(const FString& ManifestURL) = 0;
 
 
@@ -155,10 +165,14 @@ public:
 	//! Returns track metadata of the currently active play period.
 	virtual void GetTrackMetadata(TArray<FTrackMetadata>& OutTrackMetadata, EStreamType StreamType) const = 0;
 
+	//! Returns attributes of the currently selected track.
+	virtual void GetSelectedTrackAttributes(FStreamSelectionAttributes& OutAttributes, EStreamType StreamType) const = 0;
+
 
 	//-------------------------------------------------------------------------
 	// Manual stream selection functions
 	//
+
 	//! Sets the highest bitrate when selecting a candidate stream.
 	virtual void SetBitrateCeiling(int32 HighestSelectableBitrate) = 0;
 
@@ -166,13 +180,17 @@ public:
 	//! Setting both will limit on either width or height, whichever limits first.
 	virtual void SetMaxResolution(int32 MaxWidth, int32 MaxHeight) = 0;
 
-	//! Selects a track based from one of the array members returned by GetTrackMetadata().
-	//! NOTE: Presently this must be done before starting playback!
-	virtual void SelectTrackByMetadata(EStreamType StreamType, const FTrackMetadata& StreamMetadata) = 0;
+	//! Selects a track based on given attributes (which can be constructed from one of the array members returned by GetTrackMetadata()).
+	//! This selection will explicitly override the initial stream attributes set by SetInitialStreamAttributes() and be applied automatically for upcoming periods.
+	virtual void SelectTrackByAttributes(EStreamType StreamType, const FStreamSelectionAttributes& Attributes) = 0;
 
 	//! Deselect track. The stream will continue to stream to allow for immediate selection/activation but no data will be fed to the decoder.
 	virtual void DeselectTrack(EStreamType StreamType) = 0;
 
+	//! Returns true if the track stream of the specified type has beed deselected through DeselectTrack().
+	virtual bool IsTrackDeselected(EStreamType StreamType) = 0;
+
+
 	//-------------------------------------------------------------------------
 	// Platform specific functions
 	//
diff --git a/Engine/Plugins/Media/ElectraPlayer/Source/ElectraPlayerRuntime/Private/Runtime/Player/AdaptiveStreamingPlayerABR.cpp b/Engine/Plugins/Media/ElectraPlayer/Source/ElectraPlayerRuntime/Private/Runtime/Player/AdaptiveStreamingPlayerABR.cpp
--- a/Engine/Plugins/Media/ElectraPlayer/Source/ElectraPlayerRuntime/Private/Runtime/Player/AdaptiveStreamingPlayerABR.cpp
+++ b/Engine/Plugins/Media/ElectraPlayer/Source/ElectraPlayerRuntime/Private/Runtime/Player/AdaptiveStreamingPlayerABR.cpp
@@ -367,33 +367,37 @@ namespace Electra
 				TSharedPtrTS<ITimelineMediaAsset> Asset = CurrentPlayPeriodVideo->GetMediaAsset();
 				if (Asset.IsValid())
 				{
-					CurrentVideoAdaptationSetID.Empty();
+					CurrentVideoAdaptationSetID = InCurrentPlayPeriod->GetSelectedAdaptationSetID(InStreamType);
 					StreamInformationVideo.Empty();
 					// Get video stream information
-					if (Asset->GetNumberOfAdaptationSets(InStreamType) > 0)
+					if (!CurrentVideoAdaptationSetID.IsEmpty())
 					{
-						TSharedPtrTS<IPlaybackAssetAdaptationSet> Adapt = Asset->GetAdaptationSetByTypeAndIndex(InStreamType, 0);
-						if (Adapt.IsValid())
+						for(int32 nA=0, maxA=Asset->GetNumberOfAdaptationSets(InStreamType); nA<maxA; ++nA)
 						{
-							CurrentVideoAdaptationSetID = Adapt->GetUniqueIdentifier();
-							for(int32 i=0, iMax=Adapt->GetNumberOfRepresentations(); i<iMax; ++i)
+							TSharedPtrTS<IPlaybackAssetAdaptationSet> Adapt = Asset->GetAdaptationSetByTypeAndIndex(InStreamType, nA);
+							if (CurrentVideoAdaptationSetID.Equals(Adapt->GetUniqueIdentifier()))
 							{
-								TSharedPtrTS<IPlaybackAssetRepresentation> Repr = Adapt->GetRepresentationByIndex(i);
-								TSharedPtrTS<FStreamInformation> si = MakeSharedTS<FStreamInformation>();
-								si->AdaptationSetUniqueID = Adapt->GetUniqueIdentifier();
-								si->RepresentationUniqueID = Repr->GetUniqueIdentifier();
-								si->Bitrate = Repr->GetBitrate();
-								if (Repr->GetCodecInformation().IsVideoCodec())
+								for(int32 i=0, iMax=Adapt->GetNumberOfRepresentations(); i<iMax; ++i)
 								{
-									si->Resolution = Repr->GetCodecInformation().GetResolution();
+									TSharedPtrTS<IPlaybackAssetRepresentation> Repr = Adapt->GetRepresentationByIndex(i);
+									TSharedPtrTS<FStreamInformation> si = MakeSharedTS<FStreamInformation>();
+									si->AdaptationSetUniqueID = Adapt->GetUniqueIdentifier();
+									si->RepresentationUniqueID = Repr->GetUniqueIdentifier();
+									si->Bitrate = Repr->GetBitrate();
+									if (Repr->GetCodecInformation().IsVideoCodec())
+									{
+										si->Resolution = Repr->GetCodecInformation().GetResolution();
+									}
+									StreamInformationVideo.Push(si);
 								}
-								StreamInformationVideo.Push(si);
+								// Sort the representations by ascending bitrate
+								StreamInformationVideo.Sort([](const TSharedPtrTS<FStreamInformation>& a, const TSharedPtrTS<FStreamInformation>& b)
+								{
+									return a->Bitrate < b->Bitrate;
+								});
+									
+								break;
 							}
-							// Sort the representations by ascending bitrate
-							StreamInformationVideo.Sort([](const TSharedPtrTS<FStreamInformation>& a, const TSharedPtrTS<FStreamInformation>& b)
-							{
-								return a->Bitrate < b->Bitrate;
-							});
 						}
 					}
 				}
@@ -412,29 +416,33 @@ namespace Electra
 				TSharedPtrTS<ITimelineMediaAsset> Asset = CurrentPlayPeriodAudio->GetMediaAsset();
 				if (Asset.IsValid())
 				{
-					CurrentAudioAdaptationSetID.Empty();
+					CurrentAudioAdaptationSetID = InCurrentPlayPeriod->GetSelectedAdaptationSetID(InStreamType);
 					StreamInformationAudio.Empty();
 					// Get audio stream information
-					if (Asset->GetNumberOfAdaptationSets(InStreamType) > 0)
+					if (!CurrentAudioAdaptationSetID.IsEmpty())
 					{
-						TSharedPtrTS<IPlaybackAssetAdaptationSet> Adapt = Asset->GetAdaptationSetByTypeAndIndex(InStreamType, 0);
-						if (Adapt.IsValid())
+						for(int32 nA=0, maxA=Asset->GetNumberOfAdaptationSets(InStreamType); nA<maxA; ++nA)
 						{
-							CurrentAudioAdaptationSetID = Adapt->GetUniqueIdentifier();
-							for(int32 i=0, iMax=Adapt->GetNumberOfRepresentations(); i<iMax; ++i)
+							TSharedPtrTS<IPlaybackAssetAdaptationSet> Adapt = Asset->GetAdaptationSetByTypeAndIndex(InStreamType, nA);
+							if (CurrentAudioAdaptationSetID.Equals(Adapt->GetUniqueIdentifier()))
 							{
-								TSharedPtrTS<IPlaybackAssetRepresentation> Repr = Adapt->GetRepresentationByIndex(i);
-								TSharedPtrTS<FStreamInformation> si = MakeSharedTS<FStreamInformation>();
-								si->AdaptationSetUniqueID = Adapt->GetUniqueIdentifier();
-								si->RepresentationUniqueID = Repr->GetUniqueIdentifier();
-								si->Bitrate = Repr->GetBitrate();
-								StreamInformationAudio.Push(si);
+								for(int32 i=0, iMax=Adapt->GetNumberOfRepresentations(); i<iMax; ++i)
+								{
+									TSharedPtrTS<IPlaybackAssetRepresentation> Repr = Adapt->GetRepresentationByIndex(i);
+									TSharedPtrTS<FStreamInformation> si = MakeSharedTS<FStreamInformation>();
+									si->AdaptationSetUniqueID = Adapt->GetUniqueIdentifier();
+									si->RepresentationUniqueID = Repr->GetUniqueIdentifier();
+									si->Bitrate = Repr->GetBitrate();
+									StreamInformationAudio.Push(si);
+								}
+								// Sort the representations by ascending bitrate
+								StreamInformationAudio.Sort([](const TSharedPtrTS<FStreamInformation>& a, const TSharedPtrTS<FStreamInformation>& b)
+								{
+									return a->Bitrate > b->Bitrate;
+								});
+
+								break;
 							}
-							// Sort the representations by ascending bitrate
-							StreamInformationAudio.Sort([](const TSharedPtrTS<FStreamInformation>& a, const TSharedPtrTS<FStreamInformation>& b)
-							{
-								return a->Bitrate > b->Bitrate;
-							});
 						}
 					}
 				}
diff --git a/Engine/Plugins/Media/ElectraPlayer/Source/ElectraPlayerRuntime/Private/Runtime/Player/AdaptiveStreamingPlayerInternal.h b/Engine/Plugins/Media/ElectraPlayer/Source/ElectraPlayerRuntime/Private/Runtime/Player/AdaptiveStreamingPlayerInternal.h
--- a/Engine/Plugins/Media/ElectraPlayer/Source/ElectraPlayerRuntime/Private/Runtime/Player/AdaptiveStreamingPlayerInternal.h
+++ b/Engine/Plugins/Media/ElectraPlayer/Source/ElectraPlayerRuntime/Private/Runtime/Player/AdaptiveStreamingPlayerInternal.h
@@ -672,8 +672,9 @@ public:
 	virtual void AddAEMSReceiver(TWeakPtrTS<IAdaptiveStreamingPlayerAEMSReceiver> InReceiver, FString InForSchemeIdUri, FString InForValue, IAdaptiveStreamingPlayerAEMSReceiver::EDispatchMode InDispatchMode) override;
 	virtual void RemoveAEMSReceiver(TWeakPtrTS<IAdaptiveStreamingPlayerAEMSReceiver> InReceiver, FString InForSchemeIdUri, FString InForValue, IAdaptiveStreamingPlayerAEMSReceiver::EDispatchMode InDispatchMode) override;
 
-	virtual bool Initialize(const FParamDict& Options) override;
+	virtual void Initialize(const FParamDict& Options) override;
 
+	virtual void SetInitialStreamAttributes(EStreamType StreamType, const FStreamSelectionAttributes& InitialSelection) override;
 	virtual void LoadManifest(const FString& manifestURL) override;
 
 	virtual void SeekTo(const FSeekParam& NewPosition) override;
@@ -698,12 +699,16 @@ public:
 
 	virtual void GetLoopState(FPlayerLoopState& OutLoopState) const override;
 	virtual void GetTrackMetadata(TArray<FTrackMetadata>& OutTrackMetadata, EStreamType StreamType) const override;
+	//virtual void GetSelectedTrackMetadata(TOptional<FTrackMetadata>& OutSelectedTrackMetadata, EStreamType StreamType) const override;
+	virtual void GetSelectedTrackAttributes(FStreamSelectionAttributes& OutAttributes, EStreamType StreamType) const override;
 
 	virtual void SetBitrateCeiling(int32 highestSelectableBitrate) override;
 	virtual void SetMaxResolution(int32 MaxWidth, int32 MaxHeight) override;
 
-	virtual void SelectTrackByMetadata(EStreamType StreamType, const FTrackMetadata& StreamMetadata) override;
+	//virtual void SelectTrackByMetadata(EStreamType StreamType, const FTrackMetadata& StreamMetadata) override;
+	virtual void SelectTrackByAttributes(EStreamType StreamType, const FStreamSelectionAttributes& Attributes) override;
 	virtual void DeselectTrack(EStreamType StreamType) override;
+	virtual bool IsTrackDeselected(EStreamType StreamType) override;
 
 #if PLATFORM_ANDROID
 	virtual void Android_UpdateSurface(const TSharedPtr<IOptionPointerValueContainer>& Surface) override;
@@ -754,13 +759,13 @@ private:
 			case EPlayerState::eState_Idle:				return("Idle");
 			case EPlayerState::eState_ParsingManifest:	return("Parsing manifest");
 			case EPlayerState::eState_PreparingStreams:	return("Preparing streams");
-			case EPlayerState::eState_Ready:				return("Ready");
-			case EPlayerState::eState_Buffering:			return("Buffering");
+			case EPlayerState::eState_Ready:			return("Ready");
+			case EPlayerState::eState_Buffering:		return("Buffering");
 			case EPlayerState::eState_Playing:			return("Playing");
 			case EPlayerState::eState_Paused:			return("Paused");
 			case EPlayerState::eState_Rebuffering:		return("Rebuffering");
 			case EPlayerState::eState_Seeking:			return("Seeking");
-			case EPlayerState::eState_Error:				return("Error");
+			case EPlayerState::eState_Error:			return("Error");
 			default:									return("undefined");
 		}
 	}
@@ -958,7 +963,9 @@ private:
 				Close,
 				ChangeBitrate,
 				LimitResolution,
+				InitialStreamAttributes,
 				SelectTrackByMetadata,
+				SelectTrackByAttributes,
 				DeselectTrack,
 				// Player session message
 				PlayerSession,
@@ -1021,10 +1028,16 @@ private:
 					int32											Width;
 					int32											Height;
 				};
+				struct FInitialStreamSelect
+				{
+					EStreamType										StreamType;
+					FStreamSelectionAttributes						InitialSelection;
+				};
 				struct FMetadataTrackSelection
 				{
 					EStreamType										StreamType;
 					FTrackMetadata									TrackMetadata;
+					FStreamSelectionAttributes						TrackAttributes;
 				};
 
 				FLoadManifest				ManifestToLoad;
@@ -1035,6 +1048,7 @@ private:
 				FBitrate					Bitrate;
 				FSession					Session;
 				FResolution					Resolution;
+				FInitialStreamSelect		InitialStreamAttribute;
 				FMetadataTrackSelection		TrackSelection;
 			};
 			EType					Type;
@@ -1124,6 +1138,14 @@ private:
 			Msg.Data.Resolution.Height = Height;
 			WorkMessages.SendMessage(Msg);
 		}
+		void SendInitialStreamAttributeMessage(EStreamType StreamType, const FStreamSelectionAttributes& InitialSelection)
+		{
+			FMessage Msg;
+			Msg.Type = FMessage::EType::InitialStreamAttributes;
+			Msg.Data.InitialStreamAttribute.StreamType = StreamType;
+			Msg.Data.InitialStreamAttribute.InitialSelection = InitialSelection;
+			WorkMessages.SendMessage(Msg);
+		}
 		void SendTrackSelectByMetadataMessage(EStreamType StreamType, const FTrackMetadata& TrackMetadata)
 		{
 			FMessage Msg;
@@ -1132,6 +1154,14 @@ private:
 			Msg.Data.TrackSelection.TrackMetadata = TrackMetadata;
 			WorkMessages.SendMessage(Msg);
 		}
+		void SendTrackSelectByAttributeMessage(EStreamType StreamType, const FStreamSelectionAttributes& TrackAttributes)
+		{
+			FMessage Msg;
+			Msg.Type = FMessage::EType::SelectTrackByAttributes;
+			Msg.Data.TrackSelection.StreamType = StreamType;
+			Msg.Data.TrackSelection.TrackAttributes = TrackAttributes;
+			WorkMessages.SendMessage(Msg);
+		}
 		void SendTrackDeselectMessage(EStreamType StreamType)
 		{
 			FMessage Msg;
@@ -1150,7 +1180,7 @@ private:
 		FPlayStartPosition										StartAt;
 		IManifest::ESearchType									SearchType;
 		FTimeValue												RetryAtTime;
-		TMediaOptionalValue<int32>								InitialBandwidth;
+		bool													bIsPlayStart = false;
 		bool													bForLooping = false;
 		TMultiMap<EStreamType, TSharedPtrTS<IStreamSegment>>	FinishedRequests;
 	};
@@ -1238,9 +1268,12 @@ private:
 
 	struct FPendingSegmentRequest
 	{
-		TSharedPtrTS<IStreamSegment>			Request;
-		FTimeValue								AtTime;
-		TSharedPtrTS<IManifest::IPlayPeriod>	Period;			//!< Set if transitioning between periods. This is the new period that needs to be readied.
+		TSharedPtrTS<IStreamSegment>				Request;
+		FTimeValue									AtTime;
+		TSharedPtrTS<IManifest::IPlayPeriod>		Period;					//!< Set if transitioning between periods. This is the new period that needs to be readied.
+		bool										bStartOver = false;
+		EStreamType									StreamType = EStreamType::Unsupported;
+		FPlayStartPosition							StartoverPosition;
 	};
 
 	struct FUpcomingPeriod
@@ -1285,9 +1318,14 @@ private:
 	// Stream reader events
 	void OnFragmentOpen(TSharedPtrTS<IStreamSegment> pRequest) override;
 	bool OnFragmentAccessUnitReceived(FAccessUnit* pAccessUnit) override;
-	void OnFragmentReachedEOS(EStreamType InStreamType, TSharedPtr<const FStreamSourceInfo, ESPMode::ThreadSafe> InStreamSourceInfo) override;
+	void OnFragmentReachedEOS(EStreamType InStreamType, TSharedPtr<const FBufferSourceInfo, ESPMode::ThreadSafe> InStreamSourceInfo) override;
 	void OnFragmentClose(TSharedPtrTS<IStreamSegment> pRequest) override;
 
+	void InternalHandlePendingStartRequest(const FTimeValue& CurrentTime);
+	void InternalHandlePendingFirstSegmentRequest(const FTimeValue& CurrentTime);
+	void InternalHandleCompletedSegmentRequests(const FTimeValue& CurrentTime);
+	void InternalHandleSegmentTrackChanges(const FTimeValue& CurrentTime);
+
 	void UpdateDiagnostics();
 	void HandleNewBufferedData();
 	void HandleNewOutputData();
@@ -1341,7 +1379,6 @@ private:
 
 	FParamDict															PlayerOptions;
 	FPlaybackState														PlaybackState;
-	FStreamPreferences													StreamPreferences;
 	ISynchronizedUTCTime*												SynchronizedUTCTime;
 	IAdaptiveStreamingPlayerAEMSHandler*								AEMSEventHandler;
 
@@ -1355,6 +1392,7 @@ private:
 	TMediaQueueDynamic<TSharedPtrTS<FErrorDetail>>						ErrorQueue;
 
 	FString																ManifestURL;
+	EMediaFormatType													ManifestType;
 	TSharedPtrTS<IManifest>												Manifest;
 	TSharedPtrTS<IPlaylistReader>										ManifestReader;
 	TSharedPtrTS<FHTTPResourceRequest>									ManifestMimeTypeRequest;
@@ -1373,9 +1411,17 @@ private:
 	Metrics::FDataAvailabilityChange									DataAvailabilityStateAud;
 	Metrics::FDataAvailabilityChange									DataAvailabilityStateTxt;
 
-	TSharedPtr<FTrackMetadata, ESPMode::ThreadSafe>						InitialTrackSelectionVid;
-	TSharedPtr<FTrackMetadata, ESPMode::ThreadSafe>						InitialTrackSelectionAud;
-	TSharedPtr<FTrackMetadata, ESPMode::ThreadSafe>						InitialTrackSelectionTxt;
+	FStreamSelectionAttributes											StreamSelectionAttributesVid;
+	FStreamSelectionAttributes											StreamSelectionAttributesAud;
+	FStreamSelectionAttributes											StreamSelectionAttributesTxt;
+
+	FStreamSelectionAttributes											SelectedStreamAttributesVid;
+	FStreamSelectionAttributes											SelectedStreamAttributesAud;
+	FStreamSelectionAttributes											SelectedStreamAttributesTxt;
+
+	TSharedPtrTS<FStreamSelectionAttributes>							PendingTrackSelectionVid;
+	TSharedPtrTS<FStreamSelectionAttributes>							PendingTrackSelectionAud;
+	TSharedPtrTS<FStreamSelectionAttributes>							PendingTrackSelectionTxt;
 
 	EPlayerState														CurrentState;
 	EPipelineState														PipelineState;
@@ -1418,7 +1464,7 @@ private:
 	TMultiMap<EStreamType, TSharedPtrTS<IStreamSegment>>				CompletedSegmentRequests;
 	bool																bFirstSegmentRequestIsForLooping;
 
-	uint32																CurrentPlaybackSequenceID;
+	uint32																CurrentPlaybackSequenceID[4]; // 0=video, 1=audio, 2=subtitiles, 3=UNSUPPORTED
 
 	FVideoRenderer														VideoRender;
 	FAudioRenderer														AudioRender;
diff --git a/Engine/Plugins/Media/ElectraPlayer/Source/ElectraPlayerRuntime/Private/Runtime/Player/AdaptiveStreamingPlayerManifest.cpp b/Engine/Plugins/Media/ElectraPlayer/Source/ElectraPlayerRuntime/Private/Runtime/Player/AdaptiveStreamingPlayerManifest.cpp
--- a/Engine/Plugins/Media/ElectraPlayer/Source/ElectraPlayerRuntime/Private/Runtime/Player/AdaptiveStreamingPlayerManifest.cpp
+++ b/Engine/Plugins/Media/ElectraPlayer/Source/ElectraPlayerRuntime/Private/Runtime/Player/AdaptiveStreamingPlayerManifest.cpp
@@ -125,6 +125,7 @@ void FAdaptiveStreamingPlayer::InternalLoadManifest(const FString& URL, const FS
 	// Remember the original request URL since we may lose the fragment part in requests.
 	ManifestURL = URL;
 	ManifestMimeTypeRequest.Reset();
+	ManifestType = EMediaFormatType::Unknown;
 	if (CurrentState == EPlayerState::eState_Idle)
 	{
 		FString mimeType = MimeType;
@@ -154,14 +155,17 @@ void FAdaptiveStreamingPlayer::InternalLoadManifest(const FString& URL, const FS
 			if (mimeType == Playlist::MIMETypeHLS)
 			{
 				ManifestReader = IPlaylistReaderHLS::Create(this);
+				ManifestType = EMediaFormatType::HLS;
 			}
 			else if (mimeType == Playlist::MIMETypeMP4)
 			{
 				ManifestReader = IPlaylistReaderMP4::Create(this);
+				ManifestType = EMediaFormatType::ISOBMFF;
 			}
 			else if (mimeType == Playlist::MIMETypeDASH)
 			{
 				ManifestReader = IPlaylistReaderDASH::Create(this);
+				ManifestType = EMediaFormatType::DASH;
 			}
 			else
 			{
@@ -205,9 +209,7 @@ bool FAdaptiveStreamingPlayer::SelectManifest()
 	if (ManifestReader)
 	{
 		check(Manifest == nullptr);
-		if (ManifestReader->GetPlaylistType() == TEXT("hls") ||
-			ManifestReader->GetPlaylistType() == TEXT("mp4") ||
-			ManifestReader->GetPlaylistType() == TEXT("dash"))
+		if (ManifestType != EMediaFormatType::Unknown)
 		{
 			TArray<FTimespan> SeekablePositions;
 			TSharedPtrTS<IManifest> NewPresentation = ManifestReader->GetManifest();
diff --git a/Engine/Plugins/Media/ElectraPlayer/Source/ElectraPlayerRuntime/Private/Runtime/Player/DASH/ManifestBuilderDASH.cpp b/Engine/Plugins/Media/ElectraPlayer/Source/ElectraPlayerRuntime/Private/Runtime/Player/DASH/ManifestBuilderDASH.cpp
--- a/Engine/Plugins/Media/ElectraPlayer/Source/ElectraPlayerRuntime/Private/Runtime/Player/DASH/ManifestBuilderDASH.cpp
+++ b/Engine/Plugins/Media/ElectraPlayer/Source/ElectraPlayerRuntime/Private/Runtime/Player/DASH/ManifestBuilderDASH.cpp
@@ -23,6 +23,7 @@
 #include "Utilities/URLParser.h"
 #include "Utilities/Utilities.h"
 #include "Utilities/TimeUtilities.h"
+#include "Utilities/ISO639-Map.h"
 
 #include "Player/DRM/DRMManager.h"
 
@@ -1232,10 +1233,10 @@ void FManifestDASHInternal::PreparePeriodAdaptationSets(TSharedPtrTS<FPeriod> Pe
 			// Keeping the index helps locate the adaptation set in the enclosing period's adaptation set array.
 			// Even with MPD updates the adaptation sets are not permitted to change so in theory the index will
 			// be valid at all times.
-			AdaptationSet->IndexOfSelf = nAdapt;
+			AdaptationSet->UniqueSequentialSetIndex = nAdapt;
 			AdaptationSet->AdaptationSet = MPDAdaptationSet;
 			AdaptationSet->PAR = MPDAdaptationSet->GetPAR();
-			AdaptationSet->Language = MPDAdaptationSet->GetLanguage();
+			AdaptationSet->Language = ISO639::RFC5646To639_1(MPDAdaptationSet->GetLanguage());
 
 			// Content components are not supported.
 			if (MPDAdaptationSet->GetContentComponents().Num())
@@ -1317,12 +1318,16 @@ void FManifestDASHInternal::PreparePeriodAdaptationSets(TSharedPtrTS<FPeriod> Pe
 						LogMessage(PlayerSessionServices, IInfoLog::ELevel::Info, FString::Printf(TEXT("Unsupported Accessibility type \"%s\" found."), *Accessibility));
 					}
 				}
-				// Note: 608 captions are recognized but not 708 ones.
 				else if (MPDAdaptationSet->GetAccessibilities()[nAcc]->GetSchemeIdUri().Equals(TEXT("urn:scte:dash:cc:cea-608:2015")))
 				{
 					// We do not parse this out here. We prepend a "608:" prefix and take the value verbatim for now.
 					AdaptationSet->Accessibilities.Emplace(FString(TEXT("608:"))+MPDAdaptationSet->GetAccessibilities()[nAcc]->GetValue());
 				}
+				else if (MPDAdaptationSet->GetAccessibilities()[nAcc]->GetSchemeIdUri().Equals(TEXT("urn:scte:dash:cc:cea-708:2015")))
+				{
+					// We do not parse this out here. We prepend a "708:" prefix and take the value verbatim for now.
+					AdaptationSet->Accessibilities.Emplace(FString(TEXT("708:"))+MPDAdaptationSet->GetAccessibilities()[nAcc]->GetValue());
+				}
 			}
 
 			// We do not rely on the AdaptationSet@contentType, AdaptationSet@mimeType and/or AdaptationSet@codecs as these are often not set.
@@ -1573,7 +1578,6 @@ void FManifestDASHInternal::PreparePeriodAdaptationSets(TSharedPtrTS<FPeriod> Pe
 				AdaptationSet->bIsUsable = true;
 				Period->AdaptationSets.Emplace(AdaptationSet);
 			}
-
 		}
 		Period->SetHasBeenPrepared(true);
 	}
diff --git a/Engine/Plugins/Media/ElectraPlayer/Source/ElectraPlayerRuntime/Private/Runtime/Player/DASH/ManifestDASH.cpp b/Engine/Plugins/Media/ElectraPlayer/Source/ElectraPlayerRuntime/Private/Runtime/Player/DASH/ManifestDASH.cpp
--- a/Engine/Plugins/Media/ElectraPlayer/Source/ElectraPlayerRuntime/Private/Runtime/Player/DASH/ManifestDASH.cpp
+++ b/Engine/Plugins/Media/ElectraPlayer/Source/ElectraPlayerRuntime/Private/Runtime/Player/DASH/ManifestDASH.cpp
@@ -106,12 +106,18 @@ public:
 	//----------------------------------------------
 	// Methods from IManifest::IPlayPeriod
 	//
-	virtual void SetStreamPreferences(const FStreamPreferences& Preferences) override;
+	virtual void SetStreamPreferences(EStreamType ForStreamType, const FStreamSelectionAttributes& StreamAttributes) override;
 	virtual EReadyState GetReadyState() override;
+	virtual void Load() override;
 	virtual void PrepareForPlay() override;
+	virtual int64 GetDefaultStartingBitrate() const override;
+	virtual TSharedPtrTS<FBufferSourceInfo> GetSelectedStreamBufferSourceInfo(EStreamType StreamType) override;
+	virtual FString GetSelectedAdaptationSetID(EStreamType StreamType) override;
+	virtual ETrackChangeResult ChangeTrackStreamPreference(EStreamType ForStreamType, const FStreamSelectionAttributes& StreamAttributes) override;
 	virtual TSharedPtrTS<ITimelineMediaAsset> GetMediaAsset() const override;
 	virtual void SelectStream(const FString& AdaptationSetID, const FString& RepresentationID) override;
 	virtual IManifest::FResult GetStartingSegment(TSharedPtrTS<IStreamSegment>& OutSegment, const FPlayStartPosition& StartPosition, IManifest::ESearchType SearchType) override;
+	virtual IManifest::FResult GetContinuationSegment(TSharedPtrTS<IStreamSegment>& OutSegment, EStreamType StreamType, const FPlayerLoopState& LoopState, const FPlayStartPosition& StartPosition, IManifest::ESearchType SearchType) override;
 	virtual IManifest::FResult GetLoopingSegment(TSharedPtrTS<IStreamSegment>& OutSegment, FPlayerLoopState& InOutLoopState, const TMultiMap<EStreamType, TSharedPtrTS<IStreamSegment>>& InFinishedSegments, const FPlayStartPosition& StartPosition, IManifest::ESearchType SearchType) override;
 	virtual IManifest::FResult GetNextSegment(TSharedPtrTS<IStreamSegment>& OutSegment, TSharedPtrTS<const IStreamSegment> CurrentSegment) override;
 	virtual IManifest::FResult GetRetrySegment(TSharedPtrTS<IStreamSegment>& OutSegment, TSharedPtrTS<const IStreamSegment> CurrentSegment, bool bReplaceWithFillerData) override;
@@ -128,16 +134,22 @@ private:
 	{
 		SamePeriodNext,
 		SamePeriodRetry,
-		NextPeriod
+		SamePeriodStartOver,
+		NextPeriod,
 	};
 	IManifest::FResult GetNextOrRetrySegment(TSharedPtrTS<IStreamSegment>& OutSegment, TSharedPtrTS<const IStreamSegment> InCurrentSegment, ENextSegType InNextType);
 
 	bool PrepareDRM(const TArray<FManifestDASHInternal::FAdaptationSet::FContentProtection>& InContentProtections);
 
+	void SetupCommonSegmentRequestInfos(TSharedPtrTS<FStreamSegmentRequestFMP4DASH>& InOutSegmentRequest);
+
+	TSharedPtrTS<FManifestDASHInternal::FAdaptationSet> SelectAdaptationSetByAttributes(TSharedPtrTS<FBufferSourceInfo>& OutBufferSourceInfo, TSharedPtrTS<FManifestDASHInternal::FPeriod> Period, EStreamType StreamType, const FStreamSelectionAttributes& Attributes);
 
 	IPlayerSessionServices* PlayerSessionServices = nullptr;
-	EReadyState ReadyState = EReadyState::NotReady;
-	FStreamPreferences StreamPreferences;
+	EReadyState ReadyState = EReadyState::NotLoaded;
+	FStreamSelectionAttributes VideoStreamPreferences;
+	FStreamSelectionAttributes AudioStreamPreferences;
+	FStreamSelectionAttributes SubtitleStreamPreferences;
 
 	FTimeValue SegmentFetchDelay = FTimeValue::GetZero();
 
@@ -145,10 +157,15 @@ private:
 
 	FString ActiveVideoAdaptationSetID;
 	FString ActiveAudioAdaptationSetID;
+	FString ActiveSubtitleAdaptationSetID;
 
 	FString ActiveVideoRepresentationID;
 	FString ActiveAudioRepresentationID;
 
+	TSharedPtrTS<FBufferSourceInfo> SourceBufferInfoVideo;
+	TSharedPtrTS<FBufferSourceInfo> SourceBufferInfoAudio;
+	TSharedPtrTS<FBufferSourceInfo> SourceBufferInfoSubtitles;
+
 	TSharedPtr<ElectraCDM::IMediaCDMClient, ESPMode::ThreadSafe> DrmClient;
 };
 
@@ -236,11 +253,6 @@ void FManifestDASH::ClearDefaultStartTime()
 	}
 }
 
-int64 FManifestDASH::GetDefaultStartingBitrate() const
-{
-	return 2 * 1000 * 1000;
-}
-
 FTimeValue FManifestDASH::GetMinBufferTime() const
 {
 	TSharedPtrTS<FManifestDASHInternal> Manifest(CurrentManifest);
@@ -446,16 +458,18 @@ TSharedPtrTS<FManifestDASHInternal> FDASHPlayPeriod::GetCurrentManifest() const
 	TSharedPtrTS<IPlaylistReader> ManifestReader = PlayerSessionServices->GetManifestReader();
 	if (ManifestReader.IsValid())
 	{
-		check(ManifestReader->GetPlaylistType().Equals(TEXT("dash")));
 		IPlaylistReaderDASH* Reader = static_cast<IPlaylistReaderDASH*>(ManifestReader.Get());
 		return Reader->GetCurrentMPD();
 	}
 	return nullptr;
 }
 
-void FDASHPlayPeriod::SetStreamPreferences(const FStreamPreferences& Preferences)
+void FDASHPlayPeriod::SetStreamPreferences(EStreamType ForStreamType, const FStreamSelectionAttributes& StreamAttributes)
 {
-	StreamPreferences = Preferences;
+	if (ForStreamType == EStreamType::Audio)
+	{
+		AudioStreamPreferences = StreamAttributes;
+	}
 }
 
 IManifest::IPlayPeriod::EReadyState FDASHPlayPeriod::GetReadyState()
@@ -463,7 +477,91 @@ IManifest::IPlayPeriod::EReadyState FDASHPlayPeriod::GetReadyState()
 	return ReadyState;
 }
 
-void FDASHPlayPeriod::PrepareForPlay()
+TSharedPtrTS<FManifestDASHInternal::FAdaptationSet> FDASHPlayPeriod::SelectAdaptationSetByAttributes(TSharedPtrTS<FBufferSourceInfo>& OutBufferSourceInfo, TSharedPtrTS<FManifestDASHInternal::FPeriod> Period, EStreamType StreamType, const FStreamSelectionAttributes& Attributes)
+{
+	TSharedPtrTS<FManifestDASHInternal::FAdaptationSet> AS;
+	int32 NumAdaptationSets = Period->GetNumberOfAdaptationSets(StreamType);
+	if (NumAdaptationSets > 0)
+	{
+		int32 SelectedTypeIndex = 0;
+		// For video just select the first one for now.
+		if (StreamType == EStreamType::Video)
+		{
+			AS = StaticCastSharedPtr<FManifestDASHInternal::FAdaptationSet>(Period->GetAdaptationSetByTypeAndIndex(StreamType, 0));
+		}
+		else if (StreamType == EStreamType::Audio)
+		{
+			// Check for a matching language.
+			// For now we ignore the track kind.
+			if (Attributes.Language_ISO639.IsSet())
+			{
+				for(int32 i=0; i<NumAdaptationSets; ++i)
+				{
+					TSharedPtrTS<FManifestDASHInternal::FAdaptationSet> A;
+					A = StaticCastSharedPtr<FManifestDASHInternal::FAdaptationSet>(Period->GetAdaptationSetByTypeAndIndex(StreamType, i));
+					if (A->GetLanguage().Equals(Attributes.Language_ISO639.GetValue()))
+					{
+						SelectedTypeIndex = i;
+						AS = MoveTemp(A);
+						break;
+					}
+				}
+			}
+			// Matching language not found. Is there an explicit index given?
+			if (!AS.IsValid())
+			{
+				if (Attributes.OverrideIndex.IsSet() && Attributes.OverrideIndex.GetValue() < NumAdaptationSets)
+				{
+					SelectedTypeIndex = Attributes.OverrideIndex.GetValue();
+					AS = StaticCastSharedPtr<FManifestDASHInternal::FAdaptationSet>(Period->GetAdaptationSetByTypeAndIndex(StreamType, SelectedTypeIndex));
+				}
+			}
+			// Still nothing? Use the first one.
+			if (!AS.IsValid())
+			{
+				AS = StaticCastSharedPtr<FManifestDASHInternal::FAdaptationSet>(Period->GetAdaptationSetByTypeAndIndex(StreamType, 0));
+				SelectedTypeIndex = 0;
+			}
+		}
+		if (AS.IsValid())
+		{
+			FTrackMetadata tm;
+			AS->GetMetaData(tm, StreamType);
+			OutBufferSourceInfo = MakeSharedTS<FBufferSourceInfo>();
+			OutBufferSourceInfo->Kind = tm.Kind;
+			OutBufferSourceInfo->Language = tm.Language;
+			OutBufferSourceInfo->HardIndex = SelectedTypeIndex;
+			OutBufferSourceInfo->PeriodAdaptationSetID = Period->GetUniqueIdentifier() + TEXT("/") + AS->GetUniqueIdentifier();
+		}
+	}
+	return AS;
+}
+
+
+int64 FDASHPlayPeriod::GetDefaultStartingBitrate() const
+{
+	TSharedPtrTS<FManifestDASHInternal> Manifest = GetCurrentManifest();
+	TSharedPtrTS<FManifestDASHInternal::FPeriod> Period = Manifest.IsValid() ? Manifest->GetPeriodByUniqueID(PeriodID) : nullptr;
+	if (Period.IsValid())
+	{
+		TSharedPtrTS<FManifestDASHInternal::FAdaptationSet> VideoAS = Period->GetAdaptationSetByUniqueID(ActiveVideoAdaptationSetID);
+		if (VideoAS.IsValid() && VideoAS->GetNumberOfRepresentations())
+		{
+			TSharedPtrTS<IPlaybackAssetRepresentation> VideoRepr = VideoAS->GetRepresentationByIndex(0);
+			return VideoRepr->GetBitrate();
+		}
+
+		TSharedPtrTS<FManifestDASHInternal::FAdaptationSet> AudioAS = Period->GetAdaptationSetByUniqueID(ActiveAudioAdaptationSetID);
+		if (AudioAS.IsValid() && AudioAS->GetNumberOfRepresentations())
+		{
+			TSharedPtrTS<IPlaybackAssetRepresentation> AudioRepr = AudioAS->GetRepresentationByIndex(0);
+			return AudioRepr->GetBitrate();
+		}
+	}
+	return 2 * 1000 * 1000;
+}
+
+void FDASHPlayPeriod::Load()
 {
 	TSharedPtrTS<FManifestDASHInternal> Manifest = GetCurrentManifest();
 	TSharedPtrTS<FManifestDASHInternal::FPeriod> Period = Manifest.IsValid() ? Manifest->GetPeriodByUniqueID(PeriodID) : nullptr;
@@ -473,56 +571,120 @@ void FDASHPlayPeriod::PrepareForPlay()
 
 		// Prepare the adaptation sets and periods.
 		Manifest->PreparePeriodAdaptationSets(Period, false);
-	
+
 		// We need to select one adaptation set per stream type we wish to play.
-		int32 NumVideoAdaptationSets = Period->GetNumberOfAdaptationSets(EStreamType::Video);
-		if (NumVideoAdaptationSets > 0)
+		TSharedPtrTS<FManifestDASHInternal::FAdaptationSet> VideoAS = SelectAdaptationSetByAttributes(SourceBufferInfoVideo, Period, EStreamType::Video, VideoStreamPreferences);
+		if (VideoAS.IsValid())
 		{
-			// For now pick the first one!
-			TSharedPtrTS<IPlaybackAssetAdaptationSet> VideoAS = Period->GetAdaptationSetByTypeAndIndex(EStreamType::Video, 0);
 			ActiveVideoAdaptationSetID = VideoAS->GetUniqueIdentifier();
 
-			TSharedPtrTS<IPlaybackAssetRepresentation> VideoRepr = GetRepresentationFromAdaptationByMaxBandwidth(VideoAS, 2*1000*1000);
-			ActiveVideoRepresentationID = VideoRepr->GetUniqueIdentifier();
-			
 			// Add encryption schemes, if any.
 			const FManifestDASHInternal::FAdaptationSet* Adapt = static_cast<const FManifestDASHInternal::FAdaptationSet*>(VideoAS.Get());
 			ContentProtections.Append(Adapt->GetPossibleContentProtections());
 		}
 
-		int32 NumAudioAdaptationSets = Period->GetNumberOfAdaptationSets(EStreamType::Audio);
-		if (NumAudioAdaptationSets > 0)
+		TSharedPtrTS<FManifestDASHInternal::FAdaptationSet> AudioAS = SelectAdaptationSetByAttributes(SourceBufferInfoAudio, Period, EStreamType::Audio, AudioStreamPreferences);
+		if (AudioAS.IsValid())
 		{
-			// For now pick the first one!
-			TSharedPtrTS<IPlaybackAssetAdaptationSet> AudioAS = Period->GetAdaptationSetByTypeAndIndex(EStreamType::Audio, 0);
 			ActiveAudioAdaptationSetID = AudioAS->GetUniqueIdentifier();
 
-			TSharedPtrTS<IPlaybackAssetRepresentation> AudioRepr = GetRepresentationFromAdaptationByMaxBandwidth(AudioAS, 128 * 1000);
-			ActiveAudioRepresentationID = AudioRepr->GetUniqueIdentifier();
-
 			// Add encryption schemes, if any.
 			const FManifestDASHInternal::FAdaptationSet* Adapt = static_cast<const FManifestDASHInternal::FAdaptationSet*>(AudioAS.Get());
 			ContentProtections.Append(Adapt->GetPossibleContentProtections());
 		}
 
-		// Emit all <EventStream> events of the period to the AEMS event handler.
-		Manifest->SendEventsFromAllPeriodEventStreams(Period);
-
 		// Prepare the DRM system for decryption.
 		if (PrepareDRM(ContentProtections))
 		{
-			//ReadyState = EReadyState::Preparing;
-			ReadyState = EReadyState::IsReady;
+			ReadyState = EReadyState::Loaded;
 		}
 		else
 		{
 			// Set state to preparing to prevent the player from progressing while
 			// the posted error works its magic.
-			ReadyState = EReadyState::Preparing;
+			ReadyState = EReadyState::Loading;
+		}
+	}
+	else
+	{
+		ReadyState = EReadyState::Loading;
+	}
+}
+
+void FDASHPlayPeriod::PrepareForPlay()
+{
+	TSharedPtrTS<FManifestDASHInternal> Manifest = GetCurrentManifest();
+	TSharedPtrTS<FManifestDASHInternal::FPeriod> Period = Manifest.IsValid() ? Manifest->GetPeriodByUniqueID(PeriodID) : nullptr;
+	if (Period.IsValid())
+	{
+		TSharedPtrTS<FManifestDASHInternal::FAdaptationSet> VideoAS = Period->GetAdaptationSetByUniqueID(ActiveVideoAdaptationSetID);
+		if (VideoAS.IsValid())
+		{
+			// Get the current average video bitrate with some sensible default if it is not set.
+			int64 StartingBitrate = PlayerSessionServices->GetOptions().GetValue(OptionKeyCurrentAvgStartingVideoBitrate).SafeGetInt64(2*1000*1000);
+
+			TSharedPtrTS<IPlaybackAssetRepresentation> VideoRepr = GetRepresentationFromAdaptationByMaxBandwidth(VideoAS, (int32) StartingBitrate);
+			ActiveVideoRepresentationID = VideoRepr->GetUniqueIdentifier();
+		}
+
+		TSharedPtrTS<FManifestDASHInternal::FAdaptationSet> AudioAS = Period->GetAdaptationSetByUniqueID(ActiveAudioAdaptationSetID);
+		if (AudioAS.IsValid())
+		{
+			TSharedPtrTS<IPlaybackAssetRepresentation> AudioRepr = GetRepresentationFromAdaptationByMaxBandwidth(AudioAS, 256 * 1000);
+			ActiveAudioRepresentationID = AudioRepr->GetUniqueIdentifier();
 		}
+
+		// Emit all <EventStream> events of the period to the AEMS event handler.
+		Manifest->SendEventsFromAllPeriodEventStreams(Period);
+
+		ReadyState = EReadyState::IsReady;
 	}
+	else
+	{
+		ReadyState = EReadyState::Preparing;
+	}
+}
+
+TSharedPtrTS<FBufferSourceInfo> FDASHPlayPeriod::GetSelectedStreamBufferSourceInfo(EStreamType StreamType)
+{
+	switch(StreamType)
+	{
+		case EStreamType::Video:
+			return SourceBufferInfoVideo;
+		case EStreamType::Audio:
+			return SourceBufferInfoAudio;
+		case EStreamType::Subtitle:
+			return SourceBufferInfoSubtitles;
+		default:
+			return nullptr;
+	}
+}
+
+FString FDASHPlayPeriod::GetSelectedAdaptationSetID(EStreamType StreamType)
+{
+	switch(StreamType)
+	{
+		case EStreamType::Video:
+			return ActiveVideoAdaptationSetID;
+		case EStreamType::Audio:
+			return ActiveAudioAdaptationSetID;
+		case EStreamType::Subtitle:
+			return ActiveSubtitleAdaptationSetID;
+		default:
+			return FString();
+	}
+}
+
+
+IManifest::IPlayPeriod::ETrackChangeResult FDASHPlayPeriod::ChangeTrackStreamPreference(EStreamType ForStreamType, const FStreamSelectionAttributes& StreamAttributes)
+{
+	// We cannot check if the the new stream to be selected is already the one that is active because a track change is
+	// triggered at the current playback position, while we could already be in a later period!
+	// Checking against a later period makes no sense, so we have to forcibly start over.
+	return IManifest::IPlayPeriod::ETrackChangeResult::NewPeriodNeeded;
 }
 
+
 bool FDASHPlayPeriod::PrepareDRM(const TArray<FManifestDASHInternal::FAdaptationSet::FContentProtection>& InContentProtections)
 {
 	if (InContentProtections.Num())
@@ -588,6 +750,21 @@ void FDASHPlayPeriod::SelectStream(const FString& AdaptationSetID, const FString
 	}
 }
 
+void FDASHPlayPeriod::SetupCommonSegmentRequestInfos(TSharedPtrTS<FStreamSegmentRequestFMP4DASH>& InOutSegmentRequest)
+{
+	FManifestDASHInternal::FRepresentation* Repr = static_cast<FManifestDASHInternal::FRepresentation*>(InOutSegmentRequest->Representation.Get());
+
+	// Source buffer info
+	InOutSegmentRequest->SourceBufferInfo = InOutSegmentRequest->StreamType == EStreamType::Video ? SourceBufferInfoVideo :
+											InOutSegmentRequest->StreamType == EStreamType::Audio ? SourceBufferInfoAudio :
+											InOutSegmentRequest->StreamType == EStreamType::Subtitle ? SourceBufferInfoSubtitles : nullptr;
+
+	// Encryption stuff
+	InOutSegmentRequest->DrmClient = DrmClient;
+	InOutSegmentRequest->DrmMimeType = Repr->GetCodecInformation().GetMimeTypeWithCodecAndFeatures();
+}
+
+
 IManifest::FResult FDASHPlayPeriod::GetStartingSegment(TSharedPtrTS<IStreamSegment>& OutSegment, const FPlayStartPosition& StartPosition, IManifest::ESearchType SearchType)
 {
 	TSharedPtrTS<FManifestDASHInternal> Manifest = GetCurrentManifest();
@@ -680,7 +857,6 @@ IManifest::FResult FDASHPlayPeriod::GetStartingSegment(TSharedPtrTS<IStreamSegme
 				{
 					return IManifest::FResult(IManifest::FResult::EType::NotLoaded).SetErrorDetail(FErrorDetail().SetMessage("Entity loader disappeared"));
 				}
-				check(ManifestReader->GetPlaylistType().Equals(TEXT("dash")));
 				IPlaylistReaderDASH* Reader = static_cast<IPlaylistReaderDASH*>(ManifestReader.Get());
 				Reader->AddElementLoadRequests(RemoteElementLoadRequests);
 				bTryAgainLater = true;
@@ -701,6 +877,7 @@ IManifest::FResult FDASHPlayPeriod::GetStartingSegment(TSharedPtrTS<IStreamSegme
 				}
 				SegmentRequest->Segment = MoveTemp(SegmentInfo);
 				SegmentRequest->bIsEOSSegment = true;
+				SetupCommonSegmentRequestInfos(SegmentRequest);
 				StartSegmentRequest->DependentStreams.Emplace(MoveTemp(SegmentRequest));
 				bAnyStreamAtEOS = true;
 			}
@@ -772,10 +949,7 @@ IManifest::FResult FDASHPlayPeriod::GetStartingSegment(TSharedPtrTS<IStreamSegme
 					StartSegmentRequest->SAET = SegmentRequest->SAET;
 				}
 
-				// Encryption stuff
-				SegmentRequest->DrmClient = DrmClient;
-				SegmentRequest->DrmMimeType = Repr->GetCodecInformation().GetMimeTypeWithCodecAndFeatures();
-
+				SetupCommonSegmentRequestInfos(SegmentRequest);
 				StartSegmentRequest->DependentStreams.Emplace(MoveTemp(SegmentRequest));
 				bAllStreamsAtEOS = false;
 			}
@@ -804,6 +978,24 @@ IManifest::FResult FDASHPlayPeriod::GetStartingSegment(TSharedPtrTS<IStreamSegme
 }
 
 
+//-----------------------------------------------------------------------------
+/**
+ * Same as GetStartingSegment() except this is for a specific stream (video, audio, ...) only.
+ * To be used when a track (language) change is made and a new segment is needed at the current playback position.
+ */
+IManifest::FResult FDASHPlayPeriod::GetContinuationSegment(TSharedPtrTS<IStreamSegment>& OutSegment, EStreamType StreamType, const FPlayerLoopState& LoopState, const FPlayStartPosition& StartPosition, IManifest::ESearchType SearchType)
+{
+	// Create a dummy request we can use to pass into GetNextOrRetrySegment().
+	// Only set the values that that method requires.
+	auto DummyReq = MakeSharedTS<FStreamSegmentRequestFMP4DASH>();
+	DummyReq->StreamType = StreamType;
+	DummyReq->PeriodStart = StartPosition.Time;
+	DummyReq->PlayerLoopState = LoopState;
+	return GetNextOrRetrySegment(OutSegment, DummyReq, ENextSegType::SamePeriodStartOver);
+}
+
+
+
 IManifest::FResult FDASHPlayPeriod::GetNextOrRetrySegment(TSharedPtrTS<IStreamSegment>& OutSegment, TSharedPtrTS<const IStreamSegment> InCurrentSegment, ENextSegType InNextType)
 {
 	TSharedPtrTS<const FStreamSegmentRequestFMP4DASH> Current = StaticCastSharedPtr<const FStreamSegmentRequestFMP4DASH>(InCurrentSegment);
@@ -842,6 +1034,10 @@ IManifest::FResult FDASHPlayPeriod::GetNextOrRetrySegment(TSharedPtrTS<IStreamSe
 		return IManifest::FResult(IManifest::FResult::EType::NotFound).SetErrorDetail(FErrorDetail().SetMessage("No active stream found to get next segment for"));
 	}
 
+	FTimeValue AST = Manifest->GetAnchorTime();
+	bool bUsesAST = Manifest->UsesAST();
+	bool bIsStaticType = Manifest->IsStaticType() || Manifest->IsEventType();
+
 	FManifestDASHInternal::FSegmentInformation SegmentInfo;
 	FManifestDASHInternal::FSegmentSearchOption SearchOpt;
 	TArray<TWeakPtrTS<FMPDLoadRequestDASH>> RemoteElementLoadRequests;
@@ -860,15 +1056,29 @@ IManifest::FResult FDASHPlayPeriod::GetNextOrRetrySegment(TSharedPtrTS<IStreamSe
 		SearchOpt.PeriodDuration = Period->GetDuration();
 		SearchOpt.SearchType = IManifest::ESearchType::Closest;
 	}
+	else if (InNextType == ENextSegType::SamePeriodStartOver)
+	{
+		FTimeValue StartTime = Current->PeriodStart - AST;
+		if (StartTime < Period->GetStart())
+		{
+			StartTime = Period->GetStart();
+		}
+		else if (StartTime >= Period->GetEnd())
+		{
+			StartTime = Period->GetEnd();
+		}
+		StartTime -= Period->GetStart();
+
+		SearchOpt.PeriodLocalTime = StartTime;
+		SearchOpt.PeriodDuration = Period->GetDuration();
+		SearchOpt.SearchType = IManifest::ESearchType::Before;
+	}
 	else //if (InNextType == ENextSegType::NextPeriod)
 	{
 		SearchOpt.PeriodLocalTime.SetToZero();
 		SearchOpt.PeriodDuration = Period->GetDuration();
 		SearchOpt.SearchType = IManifest::ESearchType::Closest;
 	}
-	FTimeValue AST = Manifest->GetAnchorTime();
-	bool bUsesAST = Manifest->UsesAST();
-	bool bIsStaticType = Manifest->IsStaticType() || Manifest->IsEventType();
 	if (!SearchOpt.PeriodDuration.IsValid() || SearchOpt.PeriodDuration.IsPositiveInfinity())
 	{
 		SearchOpt.PeriodDuration = Manifest->GetLastPeriodEndTime() - AST;
@@ -883,7 +1093,6 @@ IManifest::FResult FDASHPlayPeriod::GetNextOrRetrySegment(TSharedPtrTS<IStreamSe
 		{
 			return IManifest::FResult(IManifest::FResult::EType::NotFound).SetErrorDetail(FErrorDetail().SetMessage("Entity loader disappeared"));
 		}
-		check(ManifestReader->GetPlaylistType().Equals(TEXT("dash")));
 		IPlaylistReaderDASH* Reader = static_cast<IPlaylistReaderDASH*>(ManifestReader.Get());
 		Reader->AddElementLoadRequests(RemoteElementLoadRequests);
 		return IManifest::FResult().RetryAfterMilliseconds(100);
@@ -905,7 +1114,6 @@ IManifest::FResult FDASHPlayPeriod::GetNextOrRetrySegment(TSharedPtrTS<IStreamSe
 		}
 
 		IPlaylistReaderDASH* Reader = static_cast<IPlaylistReaderDASH*>(PlayerSessionServices->GetManifestReader().Get());
-		check(Reader->GetPlaylistType().Equals(TEXT("dash")));
 		Reader->RequestMPDUpdate(IPlaylistReaderDASH::EMPDRequestType::GetLatestSegment);
 		return IManifest::FResult().RetryAfterMilliseconds(250);
 	}
@@ -940,8 +1148,8 @@ IManifest::FResult FDASHPlayPeriod::GetNextOrRetrySegment(TSharedPtrTS<IStreamSe
 		{
 			SegmentRequest->bInsertFillerData = true;
 		}
-		
-		if (InNextType != ENextSegType::NextPeriod)
+
+		if (InNextType == ENextSegType::SamePeriodNext || InNextType == ENextSegType::SamePeriodRetry)
 		{
 			// Because we are searching for the next segment we do not want any first access units to be truncated.
 			// We keep the current media local AU time for the case where with <SegmentTemplate> addressing we get greatly varying
@@ -964,11 +1172,8 @@ IManifest::FResult FDASHPlayPeriod::GetNextOrRetrySegment(TSharedPtrTS<IStreamSe
 			SegmentRequest->bWarnedAboutTimescale = Current->bWarnedAboutTimescale;
 		}
 
-		// Encryption stuff
-		SegmentRequest->DrmClient = DrmClient;
-		SegmentRequest->DrmMimeType = Repr->GetCodecInformation().GetMimeTypeWithCodecAndFeatures();
-
 		SegmentRequest->Segment = MoveTemp(SegmentInfo);
+		SetupCommonSegmentRequestInfos(SegmentRequest);
 		OutSegment = MoveTemp(SegmentRequest);
 		return IManifest::FResult(IManifest::FResult::EType::Found);
 	}
@@ -1051,7 +1256,7 @@ IManifest::FResult FDASHPlayPeriod::GetLoopingSegment(TSharedPtrTS<IStreamSegmen
 			LoopRequest->PlayerLoopState = InOutLoopState;
 			// We have to subtract all the values that are added to the access units from the loop base time
 			// to avoid adding them twice.
-			LoopRequest->PlayerLoopState.LoopBasetime -= LoopRequest->PeriodStart + LoopRequest->AST + LoopRequest->AdditionalAdjustmentTime;			
+			LoopRequest->PlayerLoopState.LoopBasetime -= LoopRequest->PeriodStart + LoopRequest->AST + LoopRequest->AdditionalAdjustmentTime;
 			// Set the loop state in the dependent streams as well.
 			for(int32 nDep=0, nDepMax=LoopRequest->DependentStreams.Num(); nDep<nDepMax; ++nDep)
 			{
diff --git a/Engine/Plugins/Media/ElectraPlayer/Source/ElectraPlayerRuntime/Private/Runtime/Player/DASH/ManifestDASH.h b/Engine/Plugins/Media/ElectraPlayer/Source/ElectraPlayerRuntime/Private/Runtime/Player/DASH/ManifestDASH.h
--- a/Engine/Plugins/Media/ElectraPlayer/Source/ElectraPlayerRuntime/Private/Runtime/Player/DASH/ManifestDASH.h
+++ b/Engine/Plugins/Media/ElectraPlayer/Source/ElectraPlayerRuntime/Private/Runtime/Player/DASH/ManifestDASH.h
@@ -26,7 +26,6 @@ public:
 	virtual FTimeValue GetDuration() const override;
 	virtual FTimeValue GetDefaultStartTime() const override;
 	virtual void ClearDefaultStartTime() override;
-	virtual int64 GetDefaultStartingBitrate() const override;
 	virtual FTimeValue GetMinBufferTime() const override;
 	virtual void GetTrackMetadata(TArray<FTrackMetadata>& OutMetadata, EStreamType StreamType) const override;
 	virtual void UpdateDynamicRefetchCounter() override;
diff --git a/Engine/Plugins/Media/ElectraPlayer/Source/ElectraPlayerRuntime/Private/Runtime/Player/DASH/PlaylistReaderDASH.cpp b/Engine/Plugins/Media/ElectraPlayer/Source/ElectraPlayerRuntime/Private/Runtime/Player/DASH/PlaylistReaderDASH.cpp
--- a/Engine/Plugins/Media/ElectraPlayer/Source/ElectraPlayerRuntime/Private/Runtime/Player/DASH/PlaylistReaderDASH.cpp
+++ b/Engine/Plugins/Media/ElectraPlayer/Source/ElectraPlayerRuntime/Private/Runtime/Player/DASH/PlaylistReaderDASH.cpp
@@ -140,6 +140,7 @@ private:
 
 	// Time synchronization.
 	TArray<TSharedPtrTS<FDashMPD_DescriptorType>>			AttemptedTimesyncDescriptors;
+	FTimeValue												NextTimeSyncTime;
 	bool													bTimeSyncInProgress = false;
 	bool													bInitialTimeSyncedToMPDDateTimeHeader = false;
 
@@ -575,6 +576,9 @@ void FPlaylistReaderDASH::TriggerTimeSynchronization()
 		return;
 	}
 
+	// We're doing a sync now. Clear the next time.
+	NextTimeSyncTime.SetToInvalid();
+
 	// Time sync is only necessary when there is either an MPD@availabilityStartTime or MPD@type is dynamic.
 	if (Manifest.IsValid() && (Manifest->UsesAST() || !Manifest->IsStaticType()))
 	{
@@ -690,6 +694,24 @@ void FPlaylistReaderDASH::TriggerTimeSynchronization()
 					LogMessage(IInfoLog::ELevel::Warning, FString::Printf(TEXT("No supported <UTCTiming> element found in MPD, but time was synchronized to the MPD's HTTP response Date header. This may not be accurate enough.")));
 				}
 			}
+
+			// Do not re-sync the time too often.
+			const FTimeValue MinResyncTimeInterval(120.0);		// not sooner than every 120 seconds.
+			const FTimeValue MaxResyncTimeInterval(1800.0);	// at least once every 30 minutes.
+
+			FTimeValue ResyncTimeInterval(MinResyncTimeInterval);
+			// How frequently does the MPD update?
+			FTimeValue mup = Manifest->GetMinimumUpdatePeriod();
+			if (mup.IsValid() && mup > MinResyncTimeInterval)
+			{
+				ResyncTimeInterval = mup;
+			}
+			if (ResyncTimeInterval > MaxResyncTimeInterval)
+			{
+				ResyncTimeInterval = MaxResyncTimeInterval;
+			}
+
+			NextTimeSyncTime = PlayerSessionServices->GetSynchronizedUTCTime()->GetTime() + ResyncTimeInterval;
 		}
 	}
 }
@@ -705,6 +727,8 @@ void FPlaylistReaderDASH::WorkerThread()
 	bInbandEventByStreamType[1] = false;
 	bInbandEventByStreamType[2] = false;
 
+	NextTimeSyncTime.SetToInvalid();
+
 	// Register event callbacks
 	PlayerSessionServices->GetAEMSEventHandler()->AddAEMSReceiver(AsShared(), DASH::Schemes::ManifestEvents::Scheme_urn_mpeg_dash_event_2012, TEXT(""), IAdaptiveStreamingPlayerAEMSReceiver::EDispatchMode::OnStart, false);
 	PlayerSessionServices->GetAEMSEventHandler()->AddAEMSReceiver(AsShared(), DASH::Schemes::ManifestEvents::Scheme_urn_mpeg_dash_event_callback_2015, TEXT("1"), IAdaptiveStreamingPlayerAEMSReceiver::EDispatchMode::OnStart, false);
@@ -738,6 +762,12 @@ void FPlaylistReaderDASH::WorkerThread()
 
 		// Check if the MPD must be updated.
 		CheckForMPDUpdate();
+
+		// Time sync
+		if (NextTimeSyncTime.IsValid() && Now >= NextTimeSyncTime)
+		{
+			TriggerTimeSynchronization();
+		}
 	}
 
 	// Unregister event callbacks
@@ -1089,7 +1119,10 @@ void FPlaylistReaderDASH::ManifestDownloadCompleted(FResourceLoadRequestPtr Requ
 					NewManifest->SetURLFragmentComponents(MoveTemp(URLFragmentComponents));
 				}
 				Manifest = NewManifest;
-				TriggerTimeSynchronization();
+
+				// Set the next time sync time to zero to cause a time sync as soon as possible.
+				NextTimeSyncTime.SetToZero();
+
 				if (Manifest.IsValid())
 				{
 					// Note: This is not like the standard defines FetchTime. It should be the time the server is processing the request
@@ -1227,7 +1260,6 @@ void FPlaylistReaderDASH::ManifestUpdateDownloadCompleted(FResourceLoadRequestPt
 					// Also update in the external manifest we handed out to the player.
 					PlayerManifest->UpdateInternalManifest(Manifest);
 
-					TriggerTimeSynchronization();
 					if (Manifest.IsValid())
 					{
 						Manifest->GetMPDRoot()->SetFetchTime(FetchTime);
diff --git a/Engine/Plugins/Media/ElectraPlayer/Source/ElectraPlayerRuntime/Private/Runtime/Player/DASH/PlaylistReaderDASH_Internal.h b/Engine/Plugins/Media/ElectraPlayer/Source/ElectraPlayerRuntime/Private/Runtime/Player/DASH/PlaylistReaderDASH_Internal.h
--- a/Engine/Plugins/Media/ElectraPlayer/Source/ElectraPlayerRuntime/Private/Runtime/Player/DASH/PlaylistReaderDASH_Internal.h
+++ b/Engine/Plugins/Media/ElectraPlayer/Source/ElectraPlayerRuntime/Private/Runtime/Player/DASH/PlaylistReaderDASH_Internal.h
@@ -315,6 +315,134 @@ public:
 			return nullptr;
 		}
 
+
+		void MapRoleAccessibilityToHTML5(FTrackMetadata& InOutMetadata, EStreamType StreamType) const
+		{
+			/*
+				Role: "main", "alternate", "supplementary", "commentary", "dub", "emergency", "caption", "subtitle", "sign" or "description"
+				Accessibility: "sign", "caption", "description", "enhanced-audio-intelligibility", or starts with "608:"/"708:" followed by the Value
+			*/
+			auto IsCEAService = [=]() -> bool
+			{
+				for(auto &Acc : Accessibilities)
+				{
+					if (Acc.StartsWith(TEXT("608:")) || Acc.StartsWith(TEXT("708:")))
+					{
+						return true;
+					}
+				}
+				return false;
+			};
+
+			// See: https://dev.w3.org/html5/html-sourcing-inband-tracks/#mpegdash
+			bool bMain = Roles.Contains(TEXT("main"));
+			bool bAlternate = Roles.Contains(TEXT("alternate"));
+			bool bSupplementary = Roles.Contains(TEXT("supplementary"));
+			bool bCommentary = Roles.Contains(TEXT("commentary"));
+			bool bDub = Roles.Contains(TEXT("dub"));
+			bool bEmergency = Roles.Contains(TEXT("emergency"));
+			bool bCaption = Roles.Contains(TEXT("caption"));
+			bool bSubtitle = Roles.Contains(TEXT("subtitle"));
+			bool bSign = Roles.Contains(TEXT("sign"));
+			bool bDescription = Roles.Contains(TEXT("description"));
+			bool bIsCEAService = IsCEAService();
+			if (StreamType == EStreamType::Video || StreamType == EStreamType::Audio)
+			{
+				/*
+					"alternative": if the role is "alternate" but not also "main" or "commentary", or "dub"
+					"captions": if the role is "caption" and also "main"
+					"descriptions": if the role is "description" and also "supplementary"
+					"main": if the role is "main" but not also "caption", "subtitle", or "dub"
+					"main-desc": if the role is "main" and also "description"
+					"sign": not used
+					"subtitles": if the role is "subtitle" and also "main"
+					"translation": if the role is "dub" and also "main"
+					"commentary": if the role is "commentary" but not also "main"
+					"": otherwise
+				*/
+				if (bMain && !(bCaption || bSubtitle || bDub))
+				{
+					InOutMetadata.Kind = TEXT("main");
+				}
+				else if (bMain && bDescription)
+				{
+					InOutMetadata.Kind = TEXT("main-desc");
+				}
+				else if (bAlternate && !(bMain || bCommentary || bDub))
+				{
+					InOutMetadata.Kind = TEXT("alternative");
+				}
+				else if (bSubtitle && bMain)
+				{
+					InOutMetadata.Kind = TEXT("subtitles");
+				}
+				else if (bCaption && bMain)
+				{
+					InOutMetadata.Kind = TEXT("captions");
+				}
+				else if (bDescription && bSupplementary)
+				{
+					InOutMetadata.Kind = TEXT("descriptions");
+				}
+				else if (bDub && bMain)
+				{
+					InOutMetadata.Kind = TEXT("translation");
+				}
+				else if (bCommentary && !bMain)
+				{
+					InOutMetadata.Kind = TEXT("commentary");
+				}
+			}
+			else if (StreamType == EStreamType::Subtitle)
+			{
+				/*
+					Is an ISOBMFF CEA 608 or 708 caption service: "captions".
+					"captions": if the Role descriptor's value is "caption"
+					"subtitles": if the Role descriptor's value is "subtitle"
+					"metadata": otherwise
+				*/
+				if (bIsCEAService || bCaption)
+				{
+					InOutMetadata.Kind = TEXT("captions");
+				}
+				else if (bSubtitle)
+				{
+					InOutMetadata.Kind = TEXT("subtitles");
+				}
+				else
+				{
+					InOutMetadata.Kind = TEXT("metadata");
+				}
+				// TODO: ID and language (complicated by the CEA services)
+				check(!"need to add ID and language handling");
+				InOutMetadata.ID = TEXT("CC1");
+				InOutMetadata.Language = TEXT("en");
+			}
+		}
+
+
+		void GetMetaData(FTrackMetadata& OutMetadata, EStreamType StreamType) const
+		{
+			OutMetadata.ID = GetUniqueIdentifier();
+			OutMetadata.Language = GetLanguage();
+			OutMetadata.HighestBandwidth = GetMaxBandwidth();
+			OutMetadata.HighestBandwidthCodec = GetCodec();
+			// Map role and accessibility. Do this last since this is allowed to overwrite ID and Language
+			MapRoleAccessibilityToHTML5(OutMetadata, StreamType);
+			const TArray<TSharedPtrTS<FRepresentation>>& Reprs = GetRepresentations();
+			for(int32 j=0; j<Reprs.Num(); ++j)
+			{
+				if (Reprs[j]->CanBePlayed())
+				{
+					FStreamMetadata sd;
+					sd.Bandwidth = Reprs[j]->GetBitrate();
+					sd.CodecInformation = Reprs[j]->GetCodecInformation();
+					OutMetadata.StreamDetails.Emplace(MoveTemp(sd));
+				}
+			}
+		}
+
+
 		struct FContentProtection
 		{
 			TSharedPtrTS<FDashMPD_DescriptorType> Descriptor;
@@ -344,7 +472,7 @@ public:
 		//
 		virtual FString GetUniqueIdentifier() const override
 		{
-			return FString::Printf(TEXT("%d"), IndexOfSelf);
+			return FString::Printf(TEXT("%d"), UniqueSequentialSetIndex);
 		}
 		virtual FString GetListOfCodecs() const override
 		{
@@ -380,7 +508,7 @@ public:
 		FTimeFraction PAR;
 		FString Language;
 		int32 MaxBandwidth = 0;
-		int32 IndexOfSelf = 0;
+		int32 UniqueSequentialSetIndex = 0;
 		bool bIsUsable = false;
 		bool bIsEnabled = true;
 		// Encryption related
@@ -516,21 +644,7 @@ public:
 				if (Adapt && Adapt->GetIsUsable())
 				{
 					FTrackMetadata tm;
-					tm.TrackID = Adapt->GetUniqueIdentifier();
-					tm.Language = Adapt->GetLanguage();
-					tm.HighestBandwidth = Adapt->GetMaxBandwidth();
-					tm.HighestBandwidthCodec = Adapt->GetCodec();
-					const TArray<TSharedPtrTS<FRepresentation>>& Reprs = Adapt->GetRepresentations();
-					for(int32 j=0; j<Reprs.Num(); ++j)
-					{
-						if (Reprs[j]->CanBePlayed())
-						{
-							FStreamMetadata sd;
-							sd.Bandwidth = Reprs[j]->GetBitrate();
-							sd.CodecInformation = Reprs[j]->GetCodecInformation();
-							tm.StreamDetails.Emplace(MoveTemp(sd));
-						}
-					}
+					Adapt->GetMetaData(tm, OfStreamType);
 					OutMetadata.Emplace(MoveTemp(tm));
 				}
 			}
diff --git a/Engine/Plugins/Media/ElectraPlayer/Source/ElectraPlayerRuntime/Private/Runtime/Player/DASH/StreamReaderFMP4DASH.cpp b/Engine/Plugins/Media/ElectraPlayer/Source/ElectraPlayerRuntime/Private/Runtime/Player/DASH/StreamReaderFMP4DASH.cpp
--- a/Engine/Plugins/Media/ElectraPlayer/Source/ElectraPlayerRuntime/Private/Runtime/Player/DASH/StreamReaderFMP4DASH.cpp
+++ b/Engine/Plugins/Media/ElectraPlayer/Source/ElectraPlayerRuntime/Private/Runtime/Player/DASH/StreamReaderFMP4DASH.cpp
@@ -132,7 +132,9 @@ UEMediaError FStreamReaderFMP4DASH::Create(IPlayerSessionServices* InPlayerSessi
 		StreamHandlers[i].Parameters		   = InCreateParam;
 		StreamHandlers[i].bTerminate		   = false;
 		StreamHandlers[i].bRequestCanceled     = false;
+		StreamHandlers[i].bSilentCancellation  = false;
 		StreamHandlers[i].bHasErrored   	   = false;
+		StreamHandlers[i].IsIdleSignal.Signal();
 
 		StreamHandlers[i].ThreadSetName(i==0?"ElectraPlayer::fmp4 Video":"ElectraPlayer::fmp4 Audio");
 		StreamHandlers[i].ThreadStart(Electra::MakeDelegate(&StreamHandlers[i], &FStreamHandler::WorkerThread));
@@ -149,7 +151,7 @@ void FStreamReaderFMP4DASH::Close()
 		for(int32 i=0; i<FMEDIA_STATIC_ARRAY_COUNT(StreamHandlers); ++i)
 		{
 			StreamHandlers[i].bTerminate = true;
-			StreamHandlers[i].Cancel();
+			StreamHandlers[i].Cancel(true);
 			StreamHandlers[i].SignalWork();
 		}
 		// Wait until they finished.
@@ -206,24 +208,39 @@ IStreamReader::EAddResult FStreamReaderFMP4DASH::AddRequest(uint32 CurrentPlayba
 			return IStreamReader::EAddResult::Error;
 		}
 		// Is the handler busy?
-		if (Handler->CurrentRequest.IsValid())
+		bool bIsIdle = Handler->IsIdleSignal.WaitTimeout(1000 * 1000);
+		if (!bIsIdle)
 		{
 			ErrorDetail.SetMessage(FString::Printf(TEXT("The handler for this stream type is busy!?")));
 			return IStreamReader::EAddResult::Error;
 		}
 
 		Request->SetPlaybackSequenceID(CurrentPlaybackSequenceID);
+		Handler->bRequestCanceled = false;
+		Handler->bSilentCancellation = false;
 		Handler->CurrentRequest = Request;
 		Handler->SignalWork();
 	}
 	return IStreamReader::EAddResult::Added;
 }
 
+void FStreamReaderFMP4DASH::CancelRequest(EStreamType StreamType, bool bSilent)
+{
+	if (StreamType == EStreamType::Video)
+	{
+		StreamHandlers[0].Cancel(bSilent);
+	}
+	else if (StreamType == EStreamType::Audio)
+	{
+		StreamHandlers[1].Cancel(bSilent);
+	}
+}
+
 void FStreamReaderFMP4DASH::CancelRequests()
 {
 	for(int32 i=0; i<FMEDIA_STATIC_ARRAY_COUNT(StreamHandlers); ++i)
 	{
-		StreamHandlers[i].Cancel();
+		StreamHandlers[i].Cancel(false);
 	}
 }
 
@@ -241,8 +258,9 @@ FStreamReaderFMP4DASH::FStreamHandler::~FStreamHandler()
 	// NOTE: The thread will have been terminated by the enclosing FStreamReaderFMP4DASH's Close() method!
 }
 
-void FStreamReaderFMP4DASH::FStreamHandler::Cancel()
+void FStreamReaderFMP4DASH::FStreamHandler::Cancel(bool bSilent)
 {
+	bSilentCancellation = bSilent;
 	bRequestCanceled = true;
 }
 
@@ -264,6 +282,7 @@ void FStreamReaderFMP4DASH::FStreamHandler::WorkerThread()
 		{
 			if (CurrentRequest.IsValid())
 			{
+				IsIdleSignal.Reset();
 				if (!bRequestCanceled)
 				{
 					HandleRequest();
@@ -272,8 +291,10 @@ void FStreamReaderFMP4DASH::FStreamHandler::WorkerThread()
 				{
 					CurrentRequest.Reset();
 				}
+				IsIdleSignal.Signal();
 			}
 			bRequestCanceled = false;
+			bSilentCancellation = false;
 		}
 	}
 	StreamSelector.Reset();
@@ -806,6 +827,7 @@ void FStreamReaderFMP4DASH::FStreamHandler::HandleRequest()
 									AccessUnit->bIsSyncSample = TrackIterator->IsSyncSample();
 									AccessUnit->bIsDummyData = false;
 									AccessUnit->AUCodecData = CSD;
+									AccessUnit->BufferSourceInfo = Request->SourceBufferInfo;
 									AccessUnit->PlayerLoopState = PlayerLoopState;
 
 									AccessUnit->DropState = FAccessUnit::EDropState::None;
@@ -1108,6 +1130,7 @@ void FStreamReaderFMP4DASH::FStreamHandler::HandleRequest()
 			}
 
 			AccessUnit->ESType = Request->GetType();
+			AccessUnit->BufferSourceInfo = Request->SourceBufferInfo;
 			AccessUnit->PlayerLoopState = PlayerLoopState;
 			AccessUnit->Duration = DefaultDuration;
 			AccessUnit->AUSize = 0;
@@ -1241,7 +1264,10 @@ void FStreamReaderFMP4DASH::FStreamHandler::HandleRequest()
 		}
 	}
 
-	StreamSelector->ReportDownloadEnd(ds);
+	if (!bSilentCancellation)
+	{
+		StreamSelector->ReportDownloadEnd(ds);
+	}
 
 	// Remember the next expected timestamp.
 	CurrentRequest->NextLargestExpectedTimestamp = NextExpectedDTS;
@@ -1253,7 +1279,10 @@ void FStreamReaderFMP4DASH::FStreamHandler::HandleRequest()
 	MP4Parser.Reset();
 	SegmentEventsFound.Empty();
 
-	Parameters.EventListener->OnFragmentClose(FinishedRequest);
+	if (!bSilentCancellation)
+	{
+		Parameters.EventListener->OnFragmentClose(FinishedRequest);
+	}
 }
 
 
diff --git a/Engine/Plugins/Media/ElectraPlayer/Source/ElectraPlayerRuntime/Private/Runtime/Player/DASH/StreamReaderFMP4DASH.h b/Engine/Plugins/Media/ElectraPlayer/Source/ElectraPlayerRuntime/Private/Runtime/Player/DASH/StreamReaderFMP4DASH.h
--- a/Engine/Plugins/Media/ElectraPlayer/Source/ElectraPlayerRuntime/Private/Runtime/Player/DASH/StreamReaderFMP4DASH.h
+++ b/Engine/Plugins/Media/ElectraPlayer/Source/ElectraPlayerRuntime/Private/Runtime/Player/DASH/StreamReaderFMP4DASH.h
@@ -62,6 +62,7 @@ public:
 
 
 	// Internal work variables
+	TSharedPtrTS<FBufferSourceInfo>							SourceBufferInfo;
 	FPlayerLoopState										PlayerLoopState;
 	int32													NumOverallRetries = 0;								//!< Number of retries for this _segment_ across all possible quality levels and CDNs.
 	uint32													CurrentPlaybackSequenceID = ~0U;					//!< Set by the player before adding the request to the stream reader.
@@ -90,6 +91,9 @@ public:
 	//! Adds a request to read from a stream
 	virtual EAddResult AddRequest(uint32 CurrentPlaybackSequenceID, TSharedPtrTS<IStreamSegment> Request) override;
 
+	//! Cancels any ongoing requests of the given stream type. Silent cancellation will not notify OnFragmentClose() or OnFragmentReachedEOS(). 
+	virtual void CancelRequest(EStreamType StreamType, bool bSilent) override;
+
 	//! Cancels all pending requests.
 	virtual void CancelRequests() override;
 
@@ -120,8 +124,10 @@ private:
 		IStreamReader::CreateParam								Parameters;
 		TSharedPtrTS<FStreamSegmentRequestFMP4DASH>				CurrentRequest;
 		FMediaSemaphore											WorkSignal;
+		FMediaEvent												IsIdleSignal;
 		volatile bool											bTerminate = false;
 		volatile bool											bRequestCanceled = false;
+		volatile bool											bSilentCancellation = false;
 		volatile bool											bHasErrored = false;
 		bool													bAbortedByABR = false;
 		bool													bAllowEarlyEmitting = false;
@@ -146,7 +152,7 @@ private:
 
 		FStreamHandler();
 		virtual ~FStreamHandler();
-		void Cancel();
+		void Cancel(bool bSilent);
 		void SignalWork();
 		void WorkerThread();
 		void HandleRequest();
diff --git a/Engine/Plugins/Media/ElectraPlayer/Source/ElectraPlayerRuntime/Private/Runtime/Player/HLS/ManifestBuilderHLS.cpp b/Engine/Plugins/Media/ElectraPlayer/Source/ElectraPlayerRuntime/Private/Runtime/Player/HLS/ManifestBuilderHLS.cpp
--- a/Engine/Plugins/Media/ElectraPlayer/Source/ElectraPlayerRuntime/Private/Runtime/Player/HLS/ManifestBuilderHLS.cpp
+++ b/Engine/Plugins/Media/ElectraPlayer/Source/ElectraPlayerRuntime/Private/Runtime/Player/HLS/ManifestBuilderHLS.cpp
@@ -395,7 +395,7 @@ FErrorDetail FManifestBuilderHLS::SetupRenditions(FManifestHLSInternal* Manifest
 		{
 			Adapt = new FPlaybackAssetAdaptationSetHLS;
 			Adapt->UniqueIdentifier = It.Key();
-			Adapt->Metadata.TrackID = Adapt->UniqueIdentifier;
+			Adapt->Metadata.ID = Adapt->UniqueIdentifier;
 			TSharedPtrTS<FPlaybackAssetAdaptationSetHLS> IAdapt(Adapt);
 			Asset->AudioAdaptationSets.Push(IAdapt);
 		}
@@ -698,7 +698,7 @@ FErrorDetail FManifestBuilderHLS::SetupVariants(FManifestHLSInternal* Manifest,
 					Adapt = new FPlaybackAssetAdaptationSetHLS;
 					Asset->VideoAdaptationSet = MakeShareable(Adapt);
 					Adapt->UniqueIdentifier = TEXT("$video$");
-					Adapt->Metadata.TrackID = Adapt->UniqueIdentifier;
+					Adapt->Metadata.ID = Adapt->UniqueIdentifier;
 				}
 				Adapt = static_cast<FPlaybackAssetAdaptationSetHLS*>(Asset->VideoAdaptationSet.Get());
 
@@ -801,7 +801,7 @@ FErrorDetail FManifestBuilderHLS::SetupVariants(FManifestHLSInternal* Manifest,
 					{
 						Adapt = new FPlaybackAssetAdaptationSetHLS;
 						Adapt->UniqueIdentifier = TEXT("$audio$");
-						Adapt->Metadata.TrackID = Adapt->UniqueIdentifier;
+						Adapt->Metadata.ID = Adapt->UniqueIdentifier;
 						TSharedPtrTS<FPlaybackAssetAdaptationSetHLS> IAdapt(Adapt);
 						Asset->AudioAdaptationSets.Push(IAdapt);
 					}
diff --git a/Engine/Plugins/Media/ElectraPlayer/Source/ElectraPlayerRuntime/Private/Runtime/Player/HLS/ManifestHLS.cpp b/Engine/Plugins/Media/ElectraPlayer/Source/ElectraPlayerRuntime/Private/Runtime/Player/HLS/ManifestHLS.cpp
--- a/Engine/Plugins/Media/ElectraPlayer/Source/ElectraPlayerRuntime/Private/Runtime/Player/HLS/ManifestHLS.cpp
+++ b/Engine/Plugins/Media/ElectraPlayer/Source/ElectraPlayerRuntime/Private/Runtime/Player/HLS/ManifestHLS.cpp
@@ -24,14 +24,21 @@ public:
 	FPlayPeriodHLS(IPlayerSessionServices* SessionServices, IPlaylistReaderHLS* PlaylistReader, TSharedPtrTS<FManifestHLSInternal> Manifest);
 	virtual ~FPlayPeriodHLS();
 
-	virtual void SetStreamPreferences(const FStreamPreferences& Preferences) override;
+	virtual void SetStreamPreferences(EStreamType ForStreamType, const FStreamSelectionAttributes& StreamAttributes) override;
 
 	virtual EReadyState GetReadyState() override;
+	virtual void Load() override;
 	virtual void PrepareForPlay() override;
+	virtual int64 GetDefaultStartingBitrate() const override;
+
+	virtual TSharedPtrTS<FBufferSourceInfo> GetSelectedStreamBufferSourceInfo(EStreamType StreamType) override;
+	virtual FString GetSelectedAdaptationSetID(EStreamType StreamType) override;
+	virtual ETrackChangeResult ChangeTrackStreamPreference(EStreamType ForStreamType, const FStreamSelectionAttributes& StreamAttributes) override;
 
 	// TODO: need to provide metadata (duration, streams, languages, etc.)
 
 	virtual IManifest::FResult GetStartingSegment(TSharedPtrTS<IStreamSegment>& OutSegment, const FPlayStartPosition& StartPosition, IManifest::ESearchType SearchType) override;
+	virtual IManifest::FResult GetContinuationSegment(TSharedPtrTS<IStreamSegment>& OutSegment, EStreamType StreamType, const FPlayerLoopState& LoopState, const FPlayStartPosition& StartPosition, IManifest::ESearchType SearchType) override;
 	virtual IManifest::FResult GetNextSegment(TSharedPtrTS<IStreamSegment>& OutSegment, TSharedPtrTS<const IStreamSegment> CurrentSegment) override;
 	virtual IManifest::FResult GetRetrySegment(TSharedPtrTS<IStreamSegment>& OutSegment, TSharedPtrTS<const IStreamSegment> CurrentSegment, bool bReplaceWithFillerData) override;
 	virtual IManifest::FResult GetLoopingSegment(TSharedPtrTS<IStreamSegment>& OutSegment, FPlayerLoopState& InOutLoopState, const TMultiMap<EStreamType, TSharedPtrTS<IStreamSegment>>& InFinishedSegments, const FPlayStartPosition& StartPosition, IManifest::ESearchType SearchType) override;
@@ -80,6 +87,9 @@ private:
 
 	uint32										ActiveVideoUniqueID;
 	uint32										ActiveAudioUniqueID;
+
+	TSharedPtrTS<FBufferSourceInfo>				CurrentSourceBufferInfoVideo;
+	TSharedPtrTS<FBufferSourceInfo>				CurrentSourceBufferInfoAudio;
 };
 
 
@@ -156,21 +166,6 @@ void FManifestHLS::ClearDefaultStartTime()
 
 
 
-/**
- * Returns the bitrate of the default stream (usually the first one specified).
- *
- * @return
- */
-int64 FManifestHLS::GetDefaultStartingBitrate() const
-{
-	FManifestHLSInternal::ScopedLockPlaylists lock(InternalManifest);
-	if (InternalManifest->VariantStreams.Num())
-	{
-		return InternalManifest->VariantStreams[0]->Bandwidth;
-	}
-	return 0;
-}
-
 /**
  * Returns track metadata. For period based presentations the streams can be different per period in which case the metadata of the first period is returned.
  *
@@ -269,7 +264,7 @@ FPlayPeriodHLS::FPlayPeriodHLS(IPlayerSessionServices* InSessionServices, IPlayl
 {
 	check(PlaylistReader);
 
-	CurrentReadyState   			  = IManifest::IPlayPeriod::EReadyState::NotReady;
+	CurrentReadyState = IManifest::IPlayPeriod::EReadyState::NotLoaded;
 
 	// Set the active video and audio stream IDs to 0, which means none are selected.
 	ActiveVideoUniqueID = 0;
@@ -292,12 +287,58 @@ void FPlayPeriodHLS::LogMessage(IInfoLog::ELevel Level, const FString& Message)
 /**
  * Sets stream preferences.
  *
- * @param Preferences
+ * @param ForStreamType
+ * @param StreamAttributes
  */
-void FPlayPeriodHLS::SetStreamPreferences(const FStreamPreferences& Preferences)
+void FPlayPeriodHLS::SetStreamPreferences(EStreamType ForStreamType, const FStreamSelectionAttributes& StreamAttributes)
+{
+}
+
+
+/**
+ * Returns the bitrate of the default stream (usually the first one specified).
+ *
+ * @return
+ */
+int64 FPlayPeriodHLS::GetDefaultStartingBitrate() const
+{
+	FManifestHLSInternal::ScopedLockPlaylists lock(InternalManifest);
+	if (InternalManifest->VariantStreams.Num())
+	{
+		return InternalManifest->VariantStreams[0]->Bandwidth;
+	}
+	return 0;
+}
+
+TSharedPtrTS<FBufferSourceInfo> FPlayPeriodHLS::GetSelectedStreamBufferSourceInfo(EStreamType StreamType)
 {
+	if (StreamType == EStreamType::Video)
+	{
+		return CurrentSourceBufferInfoVideo;
+	}
+	else if (StreamType == EStreamType::Audio)
+	{
+		return CurrentSourceBufferInfoAudio;
+	}
+	return nullptr;
 }
 
+FString FPlayPeriodHLS::GetSelectedAdaptationSetID(EStreamType StreamType)
+{
+	int32 nA = GetMediaAsset()->GetNumberOfAdaptationSets(StreamType);
+	if (nA)
+	{
+		return GetMediaAsset()->GetAdaptationSetByTypeAndIndex(StreamType, 0)->GetUniqueIdentifier();
+	}
+	return FString();
+}
+
+
+IManifest::IPlayPeriod::ETrackChangeResult FPlayPeriodHLS::ChangeTrackStreamPreference(EStreamType ForStreamType, const FStreamSelectionAttributes& StreamAttributes)
+{
+	// Currently track switching is not implemented.
+	return IManifest::IPlayPeriod::ETrackChangeResult::NotChanged;
+}
 
 
 /**
@@ -312,6 +353,11 @@ IManifest::IPlayPeriod::EReadyState FPlayPeriodHLS::GetReadyState()
 
 
 
+void FPlayPeriodHLS::Load()
+{
+	CurrentReadyState = IManifest::IPlayPeriod::EReadyState::Loaded;
+}
+
 /**
  * Prepares the period for playback.
  */
@@ -359,6 +405,13 @@ void FPlayPeriodHLS::PrepareForPlay()
 	InternalManifest->SelectActiveStreamID(ActiveVideoUniqueID, OldVideoUniqueID);
 	InternalManifest->SelectActiveStreamID(ActiveAudioUniqueID, OldAudioUniqueID);
 
+	// Set up source buffer information for video and audio.
+	// These are currently dummies since we do not support track switching yet.
+	CurrentSourceBufferInfoVideo = MakeSharedTS<FBufferSourceInfo>();
+	CurrentSourceBufferInfoAudio = MakeSharedTS<FBufferSourceInfo>();
+	CurrentSourceBufferInfoVideo->PeriodAdaptationSetID = TEXT("video.0");
+	CurrentSourceBufferInfoAudio->PeriodAdaptationSetID = TEXT("audio.0");
+
 	CurrentReadyState = IManifest::IPlayPeriod::EReadyState::IsReady;
 }
 
@@ -499,10 +552,15 @@ IManifest::FResult FPlayPeriodHLS::FindSegment(TSharedPtrTS<FStreamSegmentReques
 	Req->bHasEncryptedSegments     = InStream->bHasEncryptedSegments;
 	if (StreamType == EStreamType::Video)
 	{
+		Req->SourceBufferInfo = CurrentSourceBufferInfoVideo;
 		Req->Bitrate				   = InPlaylist->GetBitrate();
 		check(InternalManifest->BandwidthToQualityIndex.Find(Req->Bitrate) != nullptr);
 		Req->QualityLevel   		   = InternalManifest->BandwidthToQualityIndex[Req->Bitrate];
 	}
+	else if (StreamType == EStreamType::Audio)
+	{
+		Req->SourceBufferInfo = CurrentSourceBufferInfoAudio;
+	}
 
 	const TArray<FManifestHLSInternal::FMediaStream::FMediaSegment>& SegmentList = InStream->SegmentList;
 	if (SegmentList.Num())
@@ -882,6 +940,20 @@ IManifest::FResult FPlayPeriodHLS::GetStartingSegment(TSharedPtrTS<IStreamSegmen
 	}
 }
 
+//-----------------------------------------------------------------------------
+/**
+ * Same as GetStartingSegment() except this is for a specific stream (video, audio, ...) only.
+ * To be used when a track (language) change is made and a new segment is needed at the current playback position.
+ */
+IManifest::FResult FPlayPeriodHLS::GetContinuationSegment(TSharedPtrTS<IStreamSegment>& OutSegment, EStreamType StreamType, const FPlayerLoopState& LoopState, const FPlayStartPosition& StartPosition, IManifest::ESearchType SearchType)
+{
+	// Not currently supported.
+	// We could call GetStartingSegment() and return the segment request of the stream type that was found there
+	// but that has a start time matched to the video which may not be desirable.
+	return IManifest::FResult(IManifest::FResult::EType::NotFound);
+}
+
+
 
 void FPlayPeriodHLS::IncreaseSegmentFetchDelay(const FTimeValue& IncreaseAmount)
 {
diff --git a/Engine/Plugins/Media/ElectraPlayer/Source/ElectraPlayerRuntime/Private/Runtime/Player/HLS/ManifestHLS.h b/Engine/Plugins/Media/ElectraPlayer/Source/ElectraPlayerRuntime/Private/Runtime/Player/HLS/ManifestHLS.h
--- a/Engine/Plugins/Media/ElectraPlayer/Source/ElectraPlayerRuntime/Private/Runtime/Player/HLS/ManifestHLS.h
+++ b/Engine/Plugins/Media/ElectraPlayer/Source/ElectraPlayerRuntime/Private/Runtime/Player/HLS/ManifestHLS.h
@@ -24,7 +24,6 @@ public:
 	virtual FTimeValue GetDuration() const override;
 	virtual FTimeValue GetDefaultStartTime() const override;
 	virtual void ClearDefaultStartTime() override;
-	virtual int64 GetDefaultStartingBitrate() const override;
 	virtual FTimeValue GetMinBufferTime() const override;
 	virtual void GetTrackMetadata(TArray<FTrackMetadata>& OutMetadata, EStreamType StreamType) const override;
 	virtual void UpdateDynamicRefetchCounter() override;
diff --git a/Engine/Plugins/Media/ElectraPlayer/Source/ElectraPlayerRuntime/Private/Runtime/Player/HLS/StreamReaderHLSfmp4.cpp b/Engine/Plugins/Media/ElectraPlayer/Source/ElectraPlayerRuntime/Private/Runtime/Player/HLS/StreamReaderHLSfmp4.cpp
--- a/Engine/Plugins/Media/ElectraPlayer/Source/ElectraPlayerRuntime/Private/Runtime/Player/HLS/StreamReaderHLSfmp4.cpp
+++ b/Engine/Plugins/Media/ElectraPlayer/Source/ElectraPlayerRuntime/Private/Runtime/Player/HLS/StreamReaderHLSfmp4.cpp
@@ -140,6 +140,7 @@ UEMediaError FStreamReaderHLSfmp4::Create(IPlayerSessionServices* InPlayerSessio
 		StreamHandlers[i].Parameters		   = InCreateParam;
 		StreamHandlers[i].bTerminate		   = false;
 		StreamHandlers[i].bRequestCanceled     = false;
+		StreamHandlers[i].bSilentCancellation  = false;
 		StreamHandlers[i].bHasErrored   	   = false;
 
 		StreamHandlers[i].ThreadSetName(i==0?"ElectraPlayer::fmp4 Video":"ElectraPlayer::fmp4 Audio");
@@ -157,7 +158,7 @@ void FStreamReaderHLSfmp4::Close()
 		for(int32 i=0; i<FMEDIA_STATIC_ARRAY_COUNT(StreamHandlers); ++i)
 		{
 			StreamHandlers[i].bTerminate = true;
-			StreamHandlers[i].Cancel();
+			StreamHandlers[i].Cancel(true);
 			StreamHandlers[i].SignalWork();
 		}
 		// Wait until they finished.
@@ -257,6 +258,8 @@ IStreamReader::EAddResult FStreamReaderHLSfmp4::AddRequest(uint32 CurrentPlaybac
 		// Only add the request if this is not an EOD segment.
 		if (!Request2->bIsEOSSegment)
 		{
+			Handler2->bRequestCanceled = false;
+			Handler2->bSilentCancellation = false;
 			Handler2->CurrentRequest = Request2;
 			Handler2->SignalWork();
 		}
@@ -265,17 +268,32 @@ IStreamReader::EAddResult FStreamReaderHLSfmp4::AddRequest(uint32 CurrentPlaybac
 	// Only add the request if this is not an EOD segment.
 	if (!Request->bIsEOSSegment)
 	{
+		Handler->bRequestCanceled = false;
+		Handler->bSilentCancellation = false;
 		Handler->CurrentRequest = Request;
 		Handler->SignalWork();
 	}
 	return IStreamReader::EAddResult::Added;
 }
 
+void FStreamReaderHLSfmp4::CancelRequest(EStreamType StreamType, bool bSilent)
+{
+	if (StreamType == EStreamType::Video)
+	{
+		StreamHandlers[0].Cancel(bSilent);
+	}
+	else if (StreamType == EStreamType::Audio)
+	{
+		StreamHandlers[1].Cancel(bSilent);
+	}
+}
+
+
 void FStreamReaderHLSfmp4::CancelRequests()
 {
 	for(int32 i=0; i<FMEDIA_STATIC_ARRAY_COUNT(StreamHandlers); ++i)
 	{
-		StreamHandlers[i].Cancel();
+		StreamHandlers[i].Cancel(false);
 	}
 }
 
@@ -290,6 +308,7 @@ FStreamReaderHLSfmp4::FStreamHandler::FStreamHandler()
 	PlayerSessionService = nullptr;
 	bTerminate  		 = false;
 	bRequestCanceled	 = false;
+	bSilentCancellation  = false;
 	bAbortedByABR   	 = false;
 	bHasErrored 		 = false;
 	NumMOOFBoxesFound    = 0;
@@ -301,8 +320,9 @@ FStreamReaderHLSfmp4::FStreamHandler::~FStreamHandler()
 	// NOTE: The thread will have been terminated by the enclosing FStreamReaderHLSfmp4's Close() method!
 }
 
-void FStreamReaderHLSfmp4::FStreamHandler::Cancel()
+void FStreamReaderHLSfmp4::FStreamHandler::Cancel(bool bSilent)
 {
+	bSilentCancellation = bSilent;
 	bRequestCanceled = true;
 }
 
@@ -335,6 +355,7 @@ void FStreamReaderHLSfmp4::FStreamHandler::WorkerThread()
 				}
 			}
 			bRequestCanceled = false;
+			bSilentCancellation = false;
 		}
 	}
 	StreamSelector.Reset();
@@ -911,7 +932,7 @@ void FStreamReaderHLSfmp4::FStreamHandler::HandleRequest()
 										AccessUnit->DTS = DTS + TimeMappingOffset + LoopTimestampOffset;
 										AccessUnit->PTS = PTS + TimeMappingOffset + LoopTimestampOffset;
 
-										//AccessUnit->StreamSourceInfo = ....;
+										AccessUnit->BufferSourceInfo = Request->SourceBufferInfo;
 										AccessUnit->PlayerLoopState = PlayerLoopState;
 
 										// There should not be any gaps!
@@ -1119,6 +1140,7 @@ void FStreamReaderHLSfmp4::FStreamHandler::HandleRequest()
 				AccessUnit->AUSize = 0;
 				AccessUnit->AUData = nullptr;
 				AccessUnit->bIsDummyData = true;
+				AccessUnit->BufferSourceInfo = Request->SourceBufferInfo;
 				AccessUnit->PlayerLoopState = PlayerLoopState;
 				if (CSD.IsValid() && CSD->CodecSpecificData.Num())
 				{
@@ -1197,7 +1219,10 @@ void FStreamReaderHLSfmp4::FStreamHandler::HandleRequest()
 		ds.ThroughputBps = ds.TimeToDownload > 0.0 ? 8 * ds.NumBytesDownloaded / ds.TimeToDownload : 0;
 	}
 
-	StreamSelector->ReportDownloadEnd(ds);
+	if (!bSilentCancellation)
+	{
+		StreamSelector->ReportDownloadEnd(ds);
+	}
 
 	// Remember the next expected timestamp.
 	CurrentRequest->NextLargestExpectedTimestamp = NextExpectedDTS;
@@ -1211,7 +1236,10 @@ void FStreamReaderHLSfmp4::FStreamHandler::HandleRequest()
 	MP4Parser.Reset();
 	Decrypter.Reset();
 
-	Parameters.EventListener->OnFragmentClose(FinishedRequest);
+	if (!bSilentCancellation)
+	{
+		Parameters.EventListener->OnFragmentClose(FinishedRequest);
+	}
 }
 
 
diff --git a/Engine/Plugins/Media/ElectraPlayer/Source/ElectraPlayerRuntime/Private/Runtime/Player/HLS/StreamReaderHLSfmp4.h b/Engine/Plugins/Media/ElectraPlayer/Source/ElectraPlayerRuntime/Private/Runtime/Player/HLS/StreamReaderHLSfmp4.h
--- a/Engine/Plugins/Media/ElectraPlayer/Source/ElectraPlayerRuntime/Private/Runtime/Player/HLS/StreamReaderHLSfmp4.h
+++ b/Engine/Plugins/Media/ElectraPlayer/Source/ElectraPlayerRuntime/Private/Runtime/Player/HLS/StreamReaderHLSfmp4.h
@@ -67,6 +67,8 @@ public:
 
 	bool																		bHasEncryptedSegments;
 
+	TSharedPtrTS<FBufferSourceInfo>												SourceBufferInfo;
+
 	TSharedPtrTS<IInitSegmentCacheHLS> 											InitSegmentCache;
 	TSharedPtrTS<const FManifestHLSInternal::FMediaStream::FInitSegmentInfo>	InitSegmentInfo;
 
@@ -103,6 +105,9 @@ public:
 	//! Adds a request to read from a stream
 	virtual EAddResult AddRequest(uint32 CurrentPlaybackSequenceID, TSharedPtrTS<IStreamSegment> Request) override;
 
+	//! Cancels any ongoing requests of the given stream type. Silent cancellation will not notify OnFragmentClose() or OnFragmentReachedEOS(). 
+	virtual void CancelRequest(EStreamType StreamType, bool bSilent) override;
+
 	//! Cancels all pending requests.
 	virtual void CancelRequests() override;
 
@@ -193,6 +198,7 @@ private:
 		FMediaSemaphore											WorkSignal;
 		volatile bool											bTerminate;
 		volatile bool											bRequestCanceled;
+		volatile bool											bSilentCancellation;
 		volatile bool											bHasErrored;
 		bool													bAbortedByABR;
 		bool													bAllowEarlyEmitting;
@@ -216,7 +222,7 @@ private:
 
 		FStreamHandler();
 		virtual ~FStreamHandler();
-		void Cancel();
+		void Cancel(bool bSilent);
 		void SignalWork();
 		void WorkerThread();
 		void HandleRequest();
diff --git a/Engine/Plugins/Media/ElectraPlayer/Source/ElectraPlayerRuntime/Private/Runtime/Player/Manifest.h b/Engine/Plugins/Media/ElectraPlayer/Source/ElectraPlayerRuntime/Private/Runtime/Player/Manifest.h
--- a/Engine/Plugins/Media/ElectraPlayer/Source/ElectraPlayerRuntime/Private/Runtime/Player/Manifest.h
+++ b/Engine/Plugins/Media/ElectraPlayer/Source/ElectraPlayerRuntime/Private/Runtime/Player/Manifest.h
@@ -4,6 +4,7 @@
 #include "PlayerCore.h"
 #include "PlayerTime.h"
 #include "StreamTypes.h"
+#include "StreamAccessUnitBuffer.h"
 #include "HTTP/HTTPManager.h"
 #include "Player/PlaybackTimeline.h"
 #include "Player/AdaptiveStreamingPlayerMetrics.h"
@@ -22,6 +23,14 @@ namespace Electra
 
 	class IStreamReader;
 
+	enum class EMediaFormatType
+	{
+		Unknown,
+		ISOBMFF,					// mp4
+		HLS,						// Apple HLS (HTTP Live Streaming)
+		DASH						// MPEG DASH
+	};
+
 
 	struct FPlayStartPosition
 	{
@@ -225,17 +234,12 @@ namespace Electra
 		 * If the presentation has no preferred start time an invalid value is returned.
 		 */
 		virtual FTimeValue GetDefaultStartTime() const = 0;
-		
+
 		/**
 		 * Clears the internal default start time so it will not be used again.
 		 */
 		virtual void ClearDefaultStartTime() = 0;
 
-
-
-		//! Returns the bitrate of the default stream (usually the first one specified).
-		virtual int64 GetDefaultStartingBitrate() const = 0;
-
 		//! Returns track metadata. For period based presentations the streams can be different per period in which case the metadata of the first period is returned.
 		virtual void GetTrackMetadata(TArray<FTrackMetadata>& OutMetadata, EStreamType StreamType) const = 0;
 
@@ -265,17 +269,50 @@ namespace Electra
 		public:
 			virtual ~IPlayPeriod() = default;
 
-			virtual void SetStreamPreferences(const FStreamPreferences& Preferences) = 0;
+			virtual void SetStreamPreferences(EStreamType ForStreamType, const FStreamSelectionAttributes& StreamAttributes) = 0;
 
 			enum class EReadyState
 			{
-				NotReady,
+				NotLoaded,
+				Loading,
+				Loaded,
 				Preparing,
 				IsReady,
 			};
 			virtual EReadyState GetReadyState() = 0;
+			virtual void Load() = 0;
 			virtual void PrepareForPlay() = 0;
 
+			//! Returns the bitrate of the default stream (usually the first one specified).
+			virtual int64 GetDefaultStartingBitrate() const = 0;
+
+			/**
+			 * Returns the selected stream's buffer source information for the specified stream type
+			 * after setting the stream preferences via SetStreamPreferences() and a following PrepareToPlay().
+			 * Valid only when GetReadyState() returns IsReady.
+			 * If a nullptr is returned there is no stream for this type.
+			 */
+			virtual TSharedPtrTS<FBufferSourceInfo> GetSelectedStreamBufferSourceInfo(EStreamType StreamType) = 0;
+
+			virtual FString GetSelectedAdaptationSetID(EStreamType StreamType) = 0;
+
+			enum class ETrackChangeResult
+			{
+				Changed,
+				NotChanged,
+				NewPeriodNeeded
+			};
+
+			/**
+			 * Changes the specified track over to one of the passed attributes.
+			 * If Changed is returned the new selection has been made active and the new FBufferSourceInfo can be fetched
+			 * with GetSelectedStreamBufferSourceInfo(). If NotChanged is returned no change has been made, either because
+			 * the track is already selected or no (better) match could be made.
+			 * If NewPeriodNeeded is returned this play period cannot perform track changing on the fly and a new period
+			 * with new stream preferences must be created instead.
+			 */
+			virtual ETrackChangeResult ChangeTrackStreamPreference(EStreamType ForStreamType, const FStreamSelectionAttributes& StreamAttributes) = 0;
+
 			/**
 			 * Returns the media asset for this play period.
 			 *
@@ -304,6 +341,12 @@ namespace Electra
 			virtual FResult GetStartingSegment(TSharedPtrTS<IStreamSegment>& OutSegment, const FPlayStartPosition& StartPosition, ESearchType SearchType) = 0;
 
 
+			/**
+			 * Same as GetStartingSegment() except this is for a specific stream (video, audio, ...) only.
+			 * To be used when a track (language) change is made and a new segment is needed at the current playback position.
+			 */
+			virtual FResult GetContinuationSegment(TSharedPtrTS<IStreamSegment>& OutSegment, EStreamType StreamType, const FPlayerLoopState& LoopState, const FPlayStartPosition& StartPosition, ESearchType SearchType) = 0;
+
 			/**
 			 * Sets up a starting segment request to loop playback to.
 			 * The streams selected through SelectStream() will be used.
@@ -343,7 +386,7 @@ namespace Electra
 			/**
 			 * Called by the ABR to increase the delay in fetching the next segment in case the segment returned a 404 when fetched at
 			 * the announced availability time. This may reduce 404's on the next segment fetches.
-			 * 
+			 *
 			 * @param IncreaseAmount
 			 */
 			virtual void IncreaseSegmentFetchDelay(const FTimeValue& IncreaseAmount) = 0;
diff --git a/Engine/Plugins/Media/ElectraPlayer/Source/ElectraPlayerRuntime/Private/Runtime/Player/PlayerStreamReader.h b/Engine/Plugins/Media/ElectraPlayer/Source/ElectraPlayerRuntime/Private/Runtime/Player/PlayerStreamReader.h
--- a/Engine/Plugins/Media/ElectraPlayer/Source/ElectraPlayerRuntime/Private/Runtime/Player/PlayerStreamReader.h
+++ b/Engine/Plugins/Media/ElectraPlayer/Source/ElectraPlayerRuntime/Private/Runtime/Player/PlayerStreamReader.h
@@ -12,7 +12,7 @@ namespace Electra
 {
 	class IAccessUnitMemoryProvider;
 	struct FAccessUnit;
-	struct FStreamSourceInfo;
+	struct FBufferSourceInfo;
 
 
 	/**
@@ -29,7 +29,7 @@ namespace Electra
 		public:
 			virtual ~StreamReaderEventListener() = default;
 			virtual bool OnFragmentAccessUnitReceived(FAccessUnit* pAccessUnit) = 0;
-			virtual void OnFragmentReachedEOS(EStreamType InStreamType, TSharedPtr<const FStreamSourceInfo, ESPMode::ThreadSafe> InStreamSourceInfo) = 0;
+			virtual void OnFragmentReachedEOS(EStreamType InStreamType, TSharedPtr<const FBufferSourceInfo, ESPMode::ThreadSafe> InStreamSourceInfo) = 0;
 			virtual void OnFragmentOpen(TSharedPtrTS<IStreamSegment> pRequest) = 0;
 			virtual void OnFragmentClose(TSharedPtrTS<IStreamSegment> pRequest) = 0;
 		};
@@ -55,6 +55,9 @@ namespace Electra
 		//! Adds a request to read from a stream
 		virtual EAddResult AddRequest(uint32 CurrentPlaybackSequenceID, TSharedPtrTS<IStreamSegment> Request) = 0;
 
+		//! Cancels any ongoing requests of the given stream type. Silent cancellation will not notify OnFragmentClose() or OnFragmentReachedEOS(). 
+		virtual void CancelRequest(EStreamType StreamType, bool bSilent) = 0;
+
 		//! Cancels all pending requests.
 		virtual void CancelRequests() = 0;
 	};
diff --git a/Engine/Plugins/Media/ElectraPlayer/Source/ElectraPlayerRuntime/Private/Runtime/Player/mp4/ManifestMP4.cpp b/Engine/Plugins/Media/ElectraPlayer/Source/ElectraPlayerRuntime/Private/Runtime/Player/mp4/ManifestMP4.cpp
--- a/Engine/Plugins/Media/ElectraPlayer/Source/ElectraPlayerRuntime/Private/Runtime/Player/mp4/ManifestMP4.cpp
+++ b/Engine/Plugins/Media/ElectraPlayer/Source/ElectraPlayerRuntime/Private/Runtime/Player/mp4/ManifestMP4.cpp
@@ -87,27 +87,6 @@ IManifest::EType FManifestMP4Internal::GetPresentationType() const
 }
 
 
-//-----------------------------------------------------------------------------
-/**
- * Returns the starting bitrate.
- *
- * This is merely informational and not strictly required.
- * If fetching of the moov box provided us with the total size of the mp4 file
- * we will use that divided by the duration.
- *
- * @return
- */
-int64 FManifestMP4Internal::GetDefaultStartingBitrate() const
-{
-	FTimeValue dur = GetDuration();
-	if (ConnectionInfo.ContentLength > 0 && dur.IsValid() && dur > FTimeValue::GetZero())
-	{
-		return (int64)( ConnectionInfo.ContentLength * 8 / dur.GetAsSeconds() );
-	}
-	return 0;
-}
-
-
 //-----------------------------------------------------------------------------
 /**
  * Returns track metadata.
@@ -191,7 +170,7 @@ IManifest::FResult FManifestMP4Internal::FindNextPlayPeriod(TSharedPtrTS<IPlayPe
  */
 FManifestMP4Internal::FPlayPeriodMP4::FPlayPeriodMP4(TSharedPtrTS<FManifestMP4Internal::FTimelineAssetMP4> InMediaAsset)
 	: MediaAsset(InMediaAsset)
-	, bIsReady(false)
+	, CurrentReadyState(IManifest::IPlayPeriod::EReadyState::NotLoaded)
 {
 }
 
@@ -209,14 +188,33 @@ FManifestMP4Internal::FPlayPeriodMP4::~FPlayPeriodMP4()
 /**
  * Sets stream playback preferences for this playback period.
  *
- * @param InPreferences
+ * @param ForStreamType
+ * @param StreamAttributes
  */
-void FManifestMP4Internal::FPlayPeriodMP4::SetStreamPreferences(const FStreamPreferences& InPreferences)
+void FManifestMP4Internal::FPlayPeriodMP4::SetStreamPreferences(EStreamType ForStreamType, const FStreamSelectionAttributes& StreamAttributes)
 {
-	Preferences = InPreferences;
+	if (ForStreamType == EStreamType::Audio)
+	{
+		AudioPreferences = StreamAttributes;
+	}
 }
 
 
+//-----------------------------------------------------------------------------
+/**
+ * Returns the starting bitrate.
+ *
+ * This is merely informational and not strictly required.
+ * If fetching of the moov box provided us with the total size of the mp4 file
+ * we will use that divided by the duration.
+ *
+ * @return
+ */
+int64 FManifestMP4Internal::FPlayPeriodMP4::GetDefaultStartingBitrate() const
+{
+	return 2000000;
+}
+
 //-----------------------------------------------------------------------------
 /**
  * Returns the ready state of this playback period.
@@ -225,10 +223,15 @@ void FManifestMP4Internal::FPlayPeriodMP4::SetStreamPreferences(const FStreamPre
  */
 IManifest::IPlayPeriod::EReadyState FManifestMP4Internal::FPlayPeriodMP4::GetReadyState()
 {
-	return bIsReady ? IManifest::IPlayPeriod::EReadyState::IsReady : IManifest::IPlayPeriod::EReadyState::NotReady;
+	return CurrentReadyState;
 }
 
 
+void FManifestMP4Internal::FPlayPeriodMP4::Load()
+{
+	CurrentReadyState = IManifest::IPlayPeriod::EReadyState::Loaded;
+}
+
 //-----------------------------------------------------------------------------
 /**
  * Prepares the playback period for playback.
@@ -237,10 +240,154 @@ IManifest::IPlayPeriod::EReadyState FManifestMP4Internal::FPlayPeriodMP4::GetRea
  */
 void FManifestMP4Internal::FPlayPeriodMP4::PrepareForPlay()
 {
-	bIsReady = true;
+	SelectedVideoMetadata.Reset();
+	SelectedAudioMetadata.Reset();
+	SelectedSubtitleMetadata.Reset();
+	VideoBufferSourceInfo.Reset();
+	AudioBufferSourceInfo.Reset();
+	SubtitleBufferSourceInfo.Reset();
+	SelectInitialStream(EStreamType::Video);
+	SelectInitialStream(EStreamType::Audio);
+	SelectInitialStream(EStreamType::Subtitle);
+	CurrentReadyState = IManifest::IPlayPeriod::EReadyState::IsReady;
+}
+
+
+TSharedPtrTS<FBufferSourceInfo> FManifestMP4Internal::FPlayPeriodMP4::GetSelectedStreamBufferSourceInfo(EStreamType StreamType)
+{
+	return StreamType == EStreamType::Video ? VideoBufferSourceInfo :
+		   StreamType == EStreamType::Audio ? AudioBufferSourceInfo :
+		   StreamType == EStreamType::Subtitle ? SubtitleBufferSourceInfo : nullptr;
+}
+
+FString FManifestMP4Internal::FPlayPeriodMP4::GetSelectedAdaptationSetID(EStreamType StreamType)
+{
+	switch(StreamType)
+	{
+		case EStreamType::Video:
+			return SelectedVideoMetadata.IsValid() ? SelectedVideoMetadata->ID : FString();
+		case EStreamType::Audio:
+			return SelectedAudioMetadata.IsValid() ? SelectedAudioMetadata->ID : FString();
+		case EStreamType::Subtitle:
+			return SelectedSubtitleMetadata.IsValid() ? SelectedSubtitleMetadata->ID : FString();
+		default:
+			return FString();
+	}
 }
 
 
+IManifest::IPlayPeriod::ETrackChangeResult FManifestMP4Internal::FPlayPeriodMP4::ChangeTrackStreamPreference(EStreamType StreamType, const FStreamSelectionAttributes& StreamAttributes)
+{
+	TSharedPtrTS<FTrackMetadata> Metadata = SelectMetadataForAttributes(StreamType, StreamAttributes);
+	if (Metadata.IsValid())
+	{
+		if (StreamType == EStreamType::Video)
+		{
+			if (!(SelectedVideoMetadata.IsValid() && Metadata->Equals(*SelectedVideoMetadata)))
+			{
+				SelectedVideoMetadata = Metadata;
+				MakeBufferSourceInfoFromMetadata(StreamType, VideoBufferSourceInfo, SelectedVideoMetadata);
+				return IManifest::IPlayPeriod::ETrackChangeResult::Changed;
+			}
+		}
+		else if (StreamType == EStreamType::Audio)
+		{
+			if (!(SelectedAudioMetadata.IsValid() && Metadata->Equals(*SelectedAudioMetadata)))
+			{
+				SelectedAudioMetadata = Metadata;
+				MakeBufferSourceInfoFromMetadata(StreamType, AudioBufferSourceInfo, SelectedAudioMetadata);
+				return IManifest::IPlayPeriod::ETrackChangeResult::Changed;
+			}
+		}
+		else if (StreamType == EStreamType::Subtitle)
+		{
+			if (!(SelectedSubtitleMetadata.IsValid() && Metadata->Equals(*SelectedSubtitleMetadata)))
+			{
+				SelectedSubtitleMetadata = Metadata;
+				MakeBufferSourceInfoFromMetadata(StreamType, SubtitleBufferSourceInfo, SelectedSubtitleMetadata);
+				return IManifest::IPlayPeriod::ETrackChangeResult::Changed;
+			}
+		}
+	}
+	return IManifest::IPlayPeriod::ETrackChangeResult::NotChanged;
+}
+
+void FManifestMP4Internal::FPlayPeriodMP4::SelectInitialStream(EStreamType StreamType)
+{
+	if (StreamType == EStreamType::Video)
+	{
+		SelectedVideoMetadata = SelectMetadataForAttributes(StreamType, VideoPreferences);
+		MakeBufferSourceInfoFromMetadata(StreamType, VideoBufferSourceInfo, SelectedVideoMetadata);
+	}
+	else if (StreamType == EStreamType::Audio)
+	{
+		SelectedAudioMetadata = SelectMetadataForAttributes(StreamType, AudioPreferences);
+		MakeBufferSourceInfoFromMetadata(StreamType, AudioBufferSourceInfo, SelectedAudioMetadata);
+	}
+	else if (StreamType == EStreamType::Subtitle)
+	{
+		SelectedSubtitleMetadata = SelectMetadataForAttributes(StreamType, SubtitlePreferences);
+		MakeBufferSourceInfoFromMetadata(StreamType, SubtitleBufferSourceInfo, SelectedSubtitleMetadata);
+	}
+}
+
+TSharedPtrTS<FTrackMetadata> FManifestMP4Internal::FPlayPeriodMP4::SelectMetadataForAttributes(EStreamType StreamType, const FStreamSelectionAttributes& InAttributes)
+{
+	TSharedPtrTS<FTimelineAssetMP4> Asset = MediaAsset.Pin();
+	if (Asset.IsValid())
+	{
+		TArray<FTrackMetadata> Metadata;
+		Asset->GetMetaData(Metadata, StreamType);
+		// Is there a fixed index to be used?
+		if (InAttributes.OverrideIndex.IsSet() && InAttributes.OverrideIndex.GetValue() < Metadata.Num())
+		{
+			// Use this.
+			return MakeSharedTS<FTrackMetadata>(Metadata[InAttributes.OverrideIndex.GetValue()]);
+		}
+		if (Metadata.Num())
+		{
+			// We do not look at the 'kind' here, only the language.
+			// Set the first track as default in case we do not find the one we're looking for.
+			if (InAttributes.Language_ISO639.IsSet())
+			{
+				for(auto &Meta : Metadata)
+				{
+					if (Meta.Language.Equals(InAttributes.Language_ISO639.GetValue()))
+					{
+						return MakeSharedTS<FTrackMetadata>(Meta);
+					}
+				}
+			}
+			return MakeSharedTS<FTrackMetadata>(Metadata[0]);
+		}
+	}
+	return nullptr;
+}
+
+void FManifestMP4Internal::FPlayPeriodMP4::MakeBufferSourceInfoFromMetadata(EStreamType StreamType, TSharedPtrTS<FBufferSourceInfo>& OutBufferSourceInfo, TSharedPtrTS<FTrackMetadata> InMetadata)
+{
+	if (InMetadata.IsValid())
+	{
+		OutBufferSourceInfo = MakeSharedTS<FBufferSourceInfo>();
+		OutBufferSourceInfo->Kind = InMetadata->Kind;
+		OutBufferSourceInfo->Language = InMetadata->Language;
+		TSharedPtrTS<FTimelineAssetMP4> Asset = MediaAsset.Pin();
+		OutBufferSourceInfo->PeriodAdaptationSetID = Asset->GetUniqueIdentifier() + TEXT(".") + InMetadata->ID;
+		TArray<FTrackMetadata> Metadata;
+		Asset->GetMetaData(Metadata, StreamType);
+		for(int32 i=0; i<Metadata.Num(); ++i)
+		{
+			if (Metadata[i].Equals(*InMetadata))
+			{
+				OutBufferSourceInfo->HardIndex = i;
+				break;
+			}
+		}
+	}
+}
+
+
+
 //-----------------------------------------------------------------------------
 /**
  * Returns the timeline media asset. We have a weak pointer to it only to
@@ -289,6 +436,17 @@ IManifest::FResult FManifestMP4Internal::FPlayPeriodMP4::GetStartingSegment(TSha
 	return ma.IsValid() ? ma->GetStartingSegment(OutSegment, StartPosition, SearchType, -1) : IManifest::FResult(IManifest::FResult::EType::NotFound);
 }
 
+//-----------------------------------------------------------------------------
+/**
+ * Same as GetStartingSegment() except this is for a specific stream (video, audio, ...) only.
+ * To be used when a track (language) change is made and a new segment is needed at the current playback position.
+ */
+IManifest::FResult FManifestMP4Internal::FPlayPeriodMP4::GetContinuationSegment(TSharedPtrTS<IStreamSegment>& OutSegment, EStreamType StreamType, const FPlayerLoopState& LoopState, const FPlayStartPosition& StartPosition, ESearchType SearchType)
+{
+	// Not supported
+	return IManifest::FResult(IManifest::FResult::EType::NotFound);
+}
+
 
 //-----------------------------------------------------------------------------
 /**
@@ -314,7 +472,7 @@ IManifest::FResult FManifestMP4Internal::FPlayPeriodMP4::GetLoopingSegment(TShar
 /**
  * Called by the ABR to increase the delay in fetching the next segment in case the segment returned a 404 when fetched at
  * the announced availability time. This may reduce 404's on the next segment fetches.
- * 
+ *
  * @param IncreaseAmount
  */
 void FManifestMP4Internal::FPlayPeriodMP4::IncreaseSegmentFetchDelay(const FTimeValue& IncreaseAmount)
@@ -856,7 +1014,7 @@ FErrorDetail FManifestMP4Internal::FAdaptationSetMP4::CreateFrom(const IParserIS
 	if (err.IsOK())
 	{
 		CodecRFC6381	 = Representation->GetCodecInformation().GetCodecSpecifierRFC6381();
-		UniqueIdentifier = FString("adaptation.") + Representation->GetUniqueIdentifier();
+		UniqueIdentifier = Representation->GetUniqueIdentifier();
 		Language		 = InTrack->GetLanguage();
 	}
 	return err;
@@ -874,6 +1032,8 @@ FErrorDetail FManifestMP4Internal::FRepresentationMP4::CreateFrom(const IParserI
 	// NOTE: This *MUST* be just a number since it gets parsed back out from a string into a number later! Do *NOT* prepend/append any string literals!!
 	UniqueIdentifier = LexToString(InTrack->GetID());
 
+	Name = InTrack->GetNameFromHandler();
+
 	// Get bitrate from the average or max bitrate as stored in the track. If not stored it will be 0.
 	Bitrate = InTrack->GetBitrateInfo().AvgBitrate ? InTrack->GetBitrateInfo().AvgBitrate : InTrack->GetBitrateInfo().MaxBitrate;
 
diff --git a/Engine/Plugins/Media/ElectraPlayer/Source/ElectraPlayerRuntime/Private/Runtime/Player/mp4/ManifestMP4.h b/Engine/Plugins/Media/ElectraPlayer/Source/ElectraPlayerRuntime/Private/Runtime/Player/mp4/ManifestMP4.h
--- a/Engine/Plugins/Media/ElectraPlayer/Source/ElectraPlayerRuntime/Private/Runtime/Player/mp4/ManifestMP4.h
+++ b/Engine/Plugins/Media/ElectraPlayer/Source/ElectraPlayerRuntime/Private/Runtime/Player/mp4/ManifestMP4.h
@@ -58,7 +58,6 @@ public:
 	{ return FTimeValue::GetInvalid(); }
 	virtual void ClearDefaultStartTime() override
 	{ }
-	virtual int64 GetDefaultStartingBitrate() const override;
 	virtual void GetTrackMetadata(TArray<FTrackMetadata>& OutMetadata, EStreamType StreamType) const override;
 	virtual FTimeValue GetMinBufferTime() const override;
 	virtual void UpdateDynamicRefetchCounter() override;
@@ -86,9 +85,13 @@ public:
 		{ return 0; }
 		virtual bool CanBePlayed() const override
 		{ return true; }
+
+		const FString& GetName() const
+		{ return Name; }
 	private:
 		FStreamCodecInformation		CodecInformation;
 		FString						UniqueIdentifier;
+		FString						Name;
 		int32						Bitrate;
 
 		TArray<uint8>				CodecSpecificData;
@@ -195,8 +198,9 @@ public:
 				if (AdaptSet.IsValid())
 				{
 					FTrackMetadata tm;
-					tm.TrackID = AdaptSet->GetUniqueIdentifier();
+					tm.ID = AdaptSet->GetUniqueIdentifier();
 					tm.Language = AdaptSet->GetLanguage();
+					tm.Kind = i==0 ? TEXT("main") : TEXT("translation");
 
 					for(int32 j=0, jMax=AdaptSet->GetNumberOfRepresentations(); j<jMax; ++j)
 					{
@@ -211,6 +215,8 @@ public:
 							tm.HighestBandwidth = sd.Bandwidth;
 							tm.HighestBandwidthCodec = sd.CodecInformation;
 
+							tm.Label = Repr->GetName();
+
 							tm.StreamDetails.Emplace(MoveTemp(sd));
 						}
 					}
@@ -249,22 +255,40 @@ public:
 	public:
 		FPlayPeriodMP4(TSharedPtrTS<FTimelineAssetMP4> InMediaAsset);
 		virtual ~FPlayPeriodMP4();
-		virtual void SetStreamPreferences(const FStreamPreferences& Preferences) override;
+		virtual void SetStreamPreferences(EStreamType ForStreamType, const FStreamSelectionAttributes& StreamAttributes) override;
 		virtual EReadyState GetReadyState() override;
+		virtual void Load() override;
 		virtual void PrepareForPlay() override;
+		virtual int64 GetDefaultStartingBitrate() const override;
+		virtual TSharedPtrTS<FBufferSourceInfo> GetSelectedStreamBufferSourceInfo(EStreamType StreamType) override;
+		virtual FString GetSelectedAdaptationSetID(EStreamType StreamType) override;
+		virtual ETrackChangeResult ChangeTrackStreamPreference(EStreamType ForStreamType, const FStreamSelectionAttributes& StreamAttributes) override;
 		virtual TSharedPtrTS<ITimelineMediaAsset> GetMediaAsset() const override;
 		virtual void SelectStream(const FString& AdaptationSetID, const FString& RepresentationID) override;
 		virtual FResult GetStartingSegment(TSharedPtrTS<IStreamSegment>& OutSegment, const FPlayStartPosition& StartPosition, ESearchType SearchType) override;
+		virtual FResult GetContinuationSegment(TSharedPtrTS<IStreamSegment>& OutSegment, EStreamType StreamType, const FPlayerLoopState& LoopState, const FPlayStartPosition& StartPosition, ESearchType SearchType) override;
 		virtual FResult GetNextSegment(TSharedPtrTS<IStreamSegment>& OutSegment, TSharedPtrTS<const IStreamSegment> CurrentSegment) override;
 		virtual FResult GetRetrySegment(TSharedPtrTS<IStreamSegment>& OutSegment, TSharedPtrTS<const IStreamSegment> CurrentSegment, bool bReplaceWithFillerData) override;
 		virtual FResult GetLoopingSegment(TSharedPtrTS<IStreamSegment>& OutSegment, FPlayerLoopState& InOutLoopState, const TMultiMap<EStreamType, TSharedPtrTS<IStreamSegment>>& InFinishedSegments, const FPlayStartPosition& StartPosition, ESearchType SearchType) override;
 		virtual void IncreaseSegmentFetchDelay(const FTimeValue& IncreaseAmount) override;
 		virtual void GetSegmentInformation(TArray<FSegmentInformation>& OutSegmentInformation, FTimeValue& OutAverageSegmentDuration, TSharedPtrTS<const IStreamSegment> CurrentSegment, const FTimeValue& LookAheadTime, const FString& AdaptationSetID, const FString& RepresentationID) override;
+
 	private:
+		void SelectInitialStream(EStreamType StreamType);
+		TSharedPtrTS<FTrackMetadata> SelectMetadataForAttributes(EStreamType StreamType, const FStreamSelectionAttributes& InAttributes);
+		void MakeBufferSourceInfoFromMetadata(EStreamType StreamType, TSharedPtrTS<FBufferSourceInfo>& OutBufferSourceInfo, TSharedPtrTS<FTrackMetadata> InMetadata);
+
 		TWeakPtrTS<FTimelineAssetMP4>			MediaAsset;
-		FStreamPreferences						Preferences;
-		FParamDict								Options;
-		bool									bIsReady;
+		FStreamSelectionAttributes				VideoPreferences;
+		FStreamSelectionAttributes				AudioPreferences;
+		FStreamSelectionAttributes				SubtitlePreferences;
+		TSharedPtrTS<FTrackMetadata>			SelectedVideoMetadata;
+		TSharedPtrTS<FTrackMetadata>			SelectedAudioMetadata;
+		TSharedPtrTS<FTrackMetadata>			SelectedSubtitleMetadata;
+		TSharedPtrTS<FBufferSourceInfo>			VideoBufferSourceInfo;
+		TSharedPtrTS<FBufferSourceInfo>			AudioBufferSourceInfo;
+		TSharedPtrTS<FBufferSourceInfo>			SubtitleBufferSourceInfo;
+		EReadyState								CurrentReadyState;
 	};
 
 	void LogMessage(IInfoLog::ELevel Level, const FString& Message);
diff --git a/Engine/Plugins/Media/ElectraPlayer/Source/ElectraPlayerRuntime/Private/Runtime/Player/mp4/StreamReaderMP4.cpp b/Engine/Plugins/Media/ElectraPlayer/Source/ElectraPlayerRuntime/Private/Runtime/Player/mp4/StreamReaderMP4.cpp
--- a/Engine/Plugins/Media/ElectraPlayer/Source/ElectraPlayerRuntime/Private/Runtime/Player/mp4/StreamReaderMP4.cpp
+++ b/Engine/Plugins/Media/ElectraPlayer/Source/ElectraPlayerRuntime/Private/Runtime/Player/mp4/StreamReaderMP4.cpp
@@ -133,7 +133,7 @@ UEMediaError FStreamReaderMP4::Create(IPlayerSessionServices* InPlayerSessionSer
 	{
 		return UEMEDIA_ERROR_BAD_ARGUMENTS;
 	}
-	
+
 	PlayerSessionServices = InPlayerSessionService;
 	Parameters = InCreateParam;
 	bTerminate = false;
@@ -192,6 +192,11 @@ IStreamReader::EAddResult FStreamReaderMP4::AddRequest(uint32 CurrentPlaybackSeq
 	return EAddResult::Added;
 }
 
+void FStreamReaderMP4::CancelRequest(EStreamType StreamType, bool bSilent)
+{
+	// No-op.
+}
+
 void FStreamReaderMP4::CancelRequests()
 {
 	bRequestCanceled = true;
@@ -271,7 +276,7 @@ void FStreamReaderMP4::WorkerThread()
 // FIXME: If looping to somewhere within the stream instead of the beginning at zero the offset here must be made relative to the DTS of the first sample we demux.
 //        Otherwise there would be a large jump ahead in time!
 		FTimeValue LoopTimestampOffset = Request->PlayerLoopState.LoopBasetime;
-		TSharedPtr<const FPlayerLoopState, ESPMode::ThreadSafe>		PlayerLoopState = MakeShared<const FPlayerLoopState, ESPMode::ThreadSafe>(Request->PlayerLoopState);
+		TSharedPtrTS<const FPlayerLoopState> PlayerLoopState = MakeSharedTS<const FPlayerLoopState>(Request->PlayerLoopState);
 
 
 		// Get the list of all the tracks that have been selected in the asset.
@@ -280,11 +285,13 @@ void FStreamReaderMP4::WorkerThread()
 		struct FPlaylistTrackMetadata
 		{
 			EStreamType		Type;
+			FString			Kind;
 			FString			Language;
 			FString			PeriodID;
 			FString			AdaptationSetID;
 			FString			RepresentationID;
 			int32			Bitrate;
+			int32			Index;
 		};
 		TMap<uint32, FPlaylistTrackMetadata>	SelectedTrackMap;
 		const EStreamType TypesOfSupportedTracks[] = { EStreamType::Video, EStreamType::Audio, EStreamType::Subtitle };
@@ -303,14 +310,16 @@ void FStreamReaderMP4::WorkerThread()
 					// Note: By definition the representations unique identifier is a string of the numeric track ID and can thus be parsed back into a number.
 					FString ReprID = Representation->GetUniqueIdentifier();
 					uint32 TrackId;
-					LexFromString(TrackId, *ReprID); 
+					LexFromString(TrackId, *ReprID);
 					FPlaylistTrackMetadata tmd;
 					tmd.Type			 = TypesOfSupportedTracks[nStrType];
 					tmd.Language		 = Language;
+					tmd.Kind			 = nRepr == 0 ? TEXT("main") : TEXT("translation");	// perhaps this should come from querying the metadata instead of doing this over here...
 					tmd.PeriodID		 = Request->MediaAsset->GetUniqueIdentifier();
 					tmd.AdaptationSetID  = AdaptID;
 					tmd.RepresentationID = ReprID;
 					tmd.Bitrate 		 = Representation->GetBitrate();
+					tmd.Index			 = nAdapt;
 					SelectedTrackMap.Emplace(TrackId, tmd);
 				}
 			}
@@ -407,7 +416,7 @@ void FStreamReaderMP4::WorkerThread()
 		FTimeValue DurationSuccessfullyRead(FTimeValue::GetZero());
 		FTimeValue NextLargestExpectedTimestamp(FTimeValue::GetZero());
 		bool bDone = false;
-		TSharedPtr<IParserISO14496_12::IAllTrackIterator, ESPMode::ThreadSafe> AllTrackIterator = TimelineAsset->GetMoovBoxParser()->CreateAllTrackIteratorByFilePos(Request->FileStartOffset);
+		TSharedPtrTS<IParserISO14496_12::IAllTrackIterator> AllTrackIterator = TimelineAsset->GetMoovBoxParser()->CreateAllTrackIteratorByFilePos(Request->FileStartOffset);
 		while(!bDone && !HasErrored() && !HasBeenAborted() && !bTerminate)
 		{
 			auto UpdateSelectedTrack = [&SelectedTrackMap](const IParserISO14496_12::ITrackIterator* trkIt, TMap<uint32, FSelectedTrackData>& ActiveTrks) -> FSelectedTrackData&
@@ -417,9 +426,9 @@ void FStreamReaderMP4::WorkerThread()
 
 				// Check if this track ID is already in our map of active tracks.
 				FSelectedTrackData& st = ActiveTrks.FindOrAdd(tkid);
-				if (!st.StreamSourceInfo.IsValid())
+				if (!st.BufferSourceInfo.IsValid())
 				{
-					auto meta = MakeShared<FStreamSourceInfo, ESPMode::ThreadSafe>();
+					auto meta = MakeShared<FBufferSourceInfo, ESPMode::ThreadSafe>();
 
 					// Check if this track is in the list of selected tracks.
 					const FPlaylistTrackMetadata* SelectedTrackMetadata = SelectedTrackMap.Find(tkid);
@@ -427,13 +436,12 @@ void FStreamReaderMP4::WorkerThread()
 					{
 						st.bIsSelectedTrack = true;
 						st.StreamType = SelectedTrackMetadata->Type;
-						//meta->Role			 = SelectedTrackMetadata->;
-						meta->Language  		 = SelectedTrackMetadata->Language;
-						meta->PeriodID  		 = SelectedTrackMetadata->PeriodID;
-						meta->AdaptationSetID    = SelectedTrackMetadata->AdaptationSetID;
-						meta->RepresentationID   = SelectedTrackMetadata->RepresentationID;
+						meta->Kind = SelectedTrackMetadata->Kind;
+						meta->Language = SelectedTrackMetadata->Language;
+						meta->PeriodAdaptationSetID = SelectedTrackMetadata->PeriodID + TEXT(".") + SelectedTrackMetadata->AdaptationSetID;
+						meta->HardIndex = SelectedTrackMetadata->Index;
 					}
-					st.StreamSourceInfo = MoveTemp(meta);
+					st.BufferSourceInfo = MoveTemp(meta);
 				}
 				if (!st.CSD.IsValid())
 				{
@@ -458,7 +466,7 @@ void FStreamReaderMP4::WorkerThread()
 				// Is this a track that is selected and we are interested in?
 				if (SelectedTrack.bIsSelectedTrack)
 				{
-					Parameters.EventListener->OnFragmentReachedEOS(SelectedTrack.StreamType, SelectedTrack.StreamSourceInfo);
+					Parameters.EventListener->OnFragmentReachedEOS(SelectedTrack.StreamType, SelectedTrack.BufferSourceInfo);
 				}
 			}
 
@@ -546,7 +554,7 @@ void FStreamReaderMP4::WorkerThread()
 						AccessUnit->AUData = AccessUnit->AllocatePayloadBuffer(AccessUnit->AUSize);
 
 						// Set the associated stream metadata
-						AccessUnit->StreamSourceInfo = SelectedTrack.StreamSourceInfo;
+						AccessUnit->BufferSourceInfo = SelectedTrack.BufferSourceInfo;
 						AccessUnit->PlayerLoopState = PlayerLoopState;
 
 						SelectedTrack.bIsFirstInSequence = false;
diff --git a/Engine/Plugins/Media/ElectraPlayer/Source/ElectraPlayerRuntime/Private/Runtime/Player/mp4/StreamReaderMP4.h b/Engine/Plugins/Media/ElectraPlayer/Source/ElectraPlayerRuntime/Private/Runtime/Player/mp4/StreamReaderMP4.h
--- a/Engine/Plugins/Media/ElectraPlayer/Source/ElectraPlayerRuntime/Private/Runtime/Player/mp4/StreamReaderMP4.h
+++ b/Engine/Plugins/Media/ElectraPlayer/Source/ElectraPlayerRuntime/Private/Runtime/Player/mp4/StreamReaderMP4.h
@@ -87,6 +87,7 @@ public:
 	virtual UEMediaError Create(IPlayerSessionServices* PlayerSessionService, const CreateParam &createParam) override;
 	virtual void Close() override;
 	virtual EAddResult AddRequest(uint32 CurrentPlaybackSequenceID, TSharedPtrTS<IStreamSegment> Request) override;
+	virtual void CancelRequest(EStreamType StreamType, bool bSilent) override;
 	virtual void CancelRequests() override;
 
 private:
@@ -148,14 +149,14 @@ private:
 			DurationSuccessfullyRead.SetToZero();
 			DurationSuccessfullyDelivered.SetToZero();
 		}
-		TSharedPtrTS<FAccessUnit::CodecData>					CSD;
-		TSharedPtr<FStreamSourceInfo, ESPMode::ThreadSafe>		StreamSourceInfo;
-		EStreamType												StreamType;
-		bool													bGotKeyframe;
-		bool													bIsSelectedTrack;
-		bool													bIsFirstInSequence;
-		FTimeValue												DurationSuccessfullyRead;
-		FTimeValue												DurationSuccessfullyDelivered;
+		TSharedPtrTS<FAccessUnit::CodecData>	CSD;
+		TSharedPtrTS<FBufferSourceInfo>			BufferSourceInfo;
+		EStreamType								StreamType;
+		bool									bGotKeyframe;
+		bool									bIsSelectedTrack;
+		bool									bIsFirstInSequence;
+		FTimeValue								DurationSuccessfullyRead;
+		FTimeValue								DurationSuccessfullyDelivered;
 	};
 
 	CreateParam									Parameters;
diff --git a/Engine/Plugins/Media/ElectraPlayer/Source/ElectraPlayerRuntime/Private/Runtime/StreamAccessUnitBuffer.cpp b/Engine/Plugins/Media/ElectraPlayer/Source/ElectraPlayerRuntime/Private/Runtime/StreamAccessUnitBuffer.cpp
--- a/Engine/Plugins/Media/ElectraPlayer/Source/ElectraPlayerRuntime/Private/Runtime/StreamAccessUnitBuffer.cpp
+++ b/Engine/Plugins/Media/ElectraPlayer/Source/ElectraPlayerRuntime/Private/Runtime/StreamAccessUnitBuffer.cpp
@@ -2,28 +2,24 @@
 
 #include "StreamAccessUnitBuffer.h"
 
-
 namespace Electra
 {
-
-
-
 	FMultiTrackAccessUnitBuffer::FMultiTrackAccessUnitBuffer()
 	{
 		EmptyBuffer = CreateNewBuffer();
 		LastPoppedDTS.SetToInvalid();
 		LastPoppedPTS.SetToInvalid();
-		bAutoselectFirstTrack = false;
 		bIsDeselected = false;
 		bEndOfData = false;
 		bLastPushWasBlocked = false;
+		bIsParallelTrackMode = false;
 	}
 
-	FMultiTrackAccessUnitBuffer::FAccessUnitBufferPtr FMultiTrackAccessUnitBuffer::CreateNewBuffer()
+	TSharedPtrTS<FAccessUnitBuffer> FMultiTrackAccessUnitBuffer::CreateNewBuffer()
 	{
 		// All buffers here are internally unbounded as we need them to accept all incoming data.
 		// Limits are used on the active buffer only.
-		FAccessUnitBufferPtr NewBuffer = MakeShared<FAccessUnitBuffer, ESPMode::ThreadSafe>();
+		TSharedPtrTS<FAccessUnitBuffer> NewBuffer = MakeSharedTS<FAccessUnitBuffer>();
 		NewBuffer->CapacitySet(FAccessUnitBuffer::FConfiguration(1024 << 20, 3600.0));
 		return NewBuffer;
 	}
@@ -33,63 +29,73 @@ namespace Electra
 		PurgeAll();
 	}
 
+	void FMultiTrackAccessUnitBuffer::SetParallelTrackMode()
+	{
+		bIsParallelTrackMode = true;
+	}
+
+
 	void FMultiTrackAccessUnitBuffer::CapacitySet(const FAccessUnitBuffer::FConfiguration& Config)
 	{
-		PrimaryBufferConfiguration = Config;
+		BufferConfiguration = Config;
 	}
 
 	bool FMultiTrackAccessUnitBuffer::Push(FAccessUnit*& AU)
 	{
-		FString TrackID;
-		if (AU->StreamSourceInfo.IsValid())
-		{
-			TrackID = AU->StreamSourceInfo->AdaptationSetID;
-		}
+		check(AU->BufferSourceInfo.IsValid());
+
+		FString TrackID = AU->BufferSourceInfo->PeriodAdaptationSetID;
 		// Do we have a buffer for this track ID yet?
 		AccessLock.Lock();
 		if (!TrackBuffers.Contains(TrackID))
 		{
 			TrackBuffers.Emplace(TrackID, CreateNewBuffer());
 		}
-		FAccessUnitBufferPtr TrackBuffer = TrackBuffers[TrackID];
+		TSharedPtrTS<FAccessUnitBuffer> TrackBuffer = TrackBuffers[TrackID];
 
-		if (bAutoselectFirstTrack)
-		{
-			bAutoselectFirstTrack = false;
-			ActiveOutputID = TrackID;
-		}
+		ActivateBuffer(true);
 
 		// Pushing data to either buffer means that there is data and we are not at EOD here any more.
 		// This does not necessarily mean that the selected track is not at EOD. Just that some track is not.
 		bEndOfData = false;
+
+
+		bool bIsSwitchOverAU = false;
+		for(auto &SwitchBuf : SwitchOverBufferChain)
+		{
+			if (SwitchBuf.Info->PeriodAdaptationSetID.Equals(TrackID))
+			{
+				bIsSwitchOverAU = true;
+				break;
+			}
+		}
+
+		// Get the combined buffer utilization of enqueued buffers.
+		FAccessUnitBuffer::FExternalBufferInfo EnqueuedBufferInfo;
+		GetEnqueuedBufferInfo(EnqueuedBufferInfo, bIsSwitchOverAU);
+
 		AccessLock.Unlock();
-		// If this is the active buffer or no active buffer has been set yet we use the configured buffer limit in pushing.
-		const FAccessUnitBuffer::FConfiguration* PushConfig = ActiveOutputID.IsEmpty() || ActiveOutputID.Equals(TrackID) ? &PrimaryBufferConfiguration : nullptr;
-		bool bWasPushed = TrackBuffer->Push(AU, PushConfig);
+
+		bool bWasPushed = TrackBuffer->Push(AU, &BufferConfiguration, &EnqueuedBufferInfo);
 		bLastPushWasBlocked = TrackBuffer->WasLastPushBlocked();
+
 		return bWasPushed;
 	}
 
-	void FMultiTrackAccessUnitBuffer::PushEndOfDataFor(TSharedPtr<const FStreamSourceInfo, ESPMode::ThreadSafe> InStreamSourceInfo)
+	void FMultiTrackAccessUnitBuffer::PushEndOfDataFor(TSharedPtrTS<const FBufferSourceInfo> InStreamSourceInfo)
 	{
-		FString TrackID;
-		if (InStreamSourceInfo.IsValid())
-		{
-			TrackID = InStreamSourceInfo->AdaptationSetID;
-		}
+		check(InStreamSourceInfo.IsValid());
+		
+		FString TrackID = InStreamSourceInfo->PeriodAdaptationSetID;
 		// Do we have a buffer for this track ID yet?
 		AccessLock.Lock();
 		if (!TrackBuffers.Contains(TrackID))
 		{
 			TrackBuffers.Emplace(TrackID, CreateNewBuffer());
 		}
-		FAccessUnitBufferPtr TrackBuffer = TrackBuffers[TrackID];
+		TSharedPtrTS<FAccessUnitBuffer> TrackBuffer = TrackBuffers[TrackID];
 
-		if (bAutoselectFirstTrack)
-		{
-			bAutoselectFirstTrack = false;
-			ActiveOutputID = TrackID;
-		}
+		ActivateBuffer(true);
 
 		AccessLock.Unlock();
 		TrackBuffer->PushEndOfData();
@@ -100,25 +106,33 @@ namespace Electra
 		FMediaCriticalSection::ScopedLock lock(AccessLock);
 		bEndOfData = true;
 		// Push an end-of-data into all tracks.
-		for (auto& It : TrackBuffers)
+		for(auto& It : TrackBuffers)
 		{
 			It.Value->PushEndOfData();
 		}
 	}
 
+	void FMultiTrackAccessUnitBuffer::Clear()
+	{
+		LastPoppedDTS.SetToInvalid();
+		LastPoppedPTS.SetToInvalid();
+		bEndOfData = false;
+		bLastPushWasBlocked = false;
+		ActiveOutputBufferInfo.Reset();
+		LastPoppedBufferInfo.Reset();
+		UpcomingBufferChain.Empty();
+		SwitchOverBufferChain.Empty();
+	}
+
 	void FMultiTrackAccessUnitBuffer::Flush()
 	{
 		FMediaCriticalSection::ScopedLock lock(AccessLock);
 		// Flush all existing buffers but keep them in the track map.
-		for (auto& It : TrackBuffers)
+		for(auto& It : TrackBuffers)
 		{
 			It.Value->Flush();
 		}
-		LastPoppedDTS.SetToInvalid();
-		LastPoppedPTS.SetToInvalid();
-		bEndOfData = false;
-		bLastPushWasBlocked = false;
-		LastPoppedBufferID.Empty();
+		Clear();
 	}
 
 	void FMultiTrackAccessUnitBuffer::PurgeAll()
@@ -126,117 +140,339 @@ namespace Electra
 		FMediaCriticalSection::ScopedLock lock(AccessLock);
 		// Get rid of all buffers
 		TrackBuffers.Empty();
-		LastPoppedDTS.SetToInvalid();
-		LastPoppedPTS.SetToInvalid();
-		bEndOfData = false;
-		bLastPushWasBlocked = false;
-		LastPoppedBufferID.Empty();
+		Clear();
 	}
 
-	void FMultiTrackAccessUnitBuffer::AutoselectFirstTrack()
+
+	void FMultiTrackAccessUnitBuffer::SelectTrackWhenAvailable(TSharedPtrTS<FBufferSourceInfo> InBufferSourceInfo)
 	{
-		FMediaCriticalSection::ScopedLock lock(AccessLock);
-		bAutoselectFirstTrack = true;
-		bIsDeselected = false;
-	}
+		if (InBufferSourceInfo.IsValid())
+		{
+			FMediaCriticalSection::ScopedLock lock(AccessLock);
+			if (SwitchOverBufferChain.Num() == 0)
+			{
+				FQueuedBuffer Next;
+				Next.Info = MoveTemp(InBufferSourceInfo);
+				SwitchOverBufferChain.Empty();
+				SwitchOverBufferChain.Emplace(MoveTemp(Next));
+			}
+			else if (!SwitchOverBufferChain.Last().Info->PeriodAdaptationSetID.Equals(InBufferSourceInfo->PeriodAdaptationSetID))
+			{
+				/*
+					We have to check if what we are to switch to is what is already active.
+					This can happen on quick track changes where we are currently on track1, then a change to track2 was issued
+					but immediately chaged back to track1. So: track1 -> track2 -> track1
+					In this case we would remove the request for track2 and see track1 -> track1. Since the ID of this track is
+					still the same we would append the data coming in from the new track1 request to the already existing track1
+					buffer. This is bad because segments are streamed in starting at their beginning and we would thus put data
+					into the buffer that is already there.
+					To not make things overly complex we will do the following if the switch-to track is the currently active one:
+					We rename the active one by appending some suffix. This works because the current stream download was already
+					canceled and no new data is coming in from that stream. We need to make sure we change the ID in the current
+					track map (the key) as well as the source info. The source info is maintained as a shared pointer, so changing
+					it there means it is automatically changed everywhere that pointer is shared!
+				*/
+
+				if (ActiveOutputBufferInfo.IsValid() && ActiveOutputBufferInfo->PeriodAdaptationSetID.Equals(InBufferSourceInfo->PeriodAdaptationSetID))
+				{
+					// The ID may be the same but the TSharedPtr MUST NOT be.
+					// If it were the same pointer we would be renaming the new segment buffer and info as well!
+					check(ActiveOutputBufferInfo != InBufferSourceInfo);
 
+					// First let's rename the buffer in the track buffer map. Double check that it is actually there.
+					FString NewName = InBufferSourceInfo->PeriodAdaptationSetID + TEXT("$1");
+					if (TrackBuffers.Contains(InBufferSourceInfo->PeriodAdaptationSetID))
+					{
+						TSharedPtrTS<FAccessUnitBuffer> TrackBuffer = TrackBuffers[ActiveOutputBufferInfo->PeriodAdaptationSetID];
+						TrackBuffers.Remove(ActiveOutputBufferInfo->PeriodAdaptationSetID);
+						TrackBuffers.Emplace(NewName, TrackBuffer);
+					}
+					// Then change in the name in the active output buffer info.
+					ActiveOutputBufferInfo->PeriodAdaptationSetID = NewName;
 
-	void FMultiTrackAccessUnitBuffer::SelectTrackByID(const FString& TrackID)
+					// Finally, if the LastPoppedBufferInfo is set the it SHOULD be the same as the ActiveOutputBufferInfo,
+					// ie: check(LastPoppedBufferInfo == ActiveOutputBufferInfo);
+					// But in case it is not (due to future code changes) let's update it as well.
+					if (LastPoppedBufferInfo.IsValid() && LastPoppedBufferInfo->PeriodAdaptationSetID.Equals(InBufferSourceInfo->PeriodAdaptationSetID))
+					{
+						LastPoppedBufferInfo->PeriodAdaptationSetID = NewName;
+					}
+				}
+
+				FQueuedBuffer Next;
+				Next.Info = MoveTemp(InBufferSourceInfo);
+				SwitchOverBufferChain.Empty();
+				SwitchOverBufferChain.Emplace(MoveTemp(Next));
+			}
+		}
+	}
+
+	void FMultiTrackAccessUnitBuffer::AddUpcomingBuffer(TSharedPtrTS<FBufferSourceInfo> InBufferSourceInfo)
 	{
 		FMediaCriticalSection::ScopedLock lock(AccessLock);
-		if (TrackID != ActiveOutputID)
+		// Add to the current or the switch-over chain?
+		if (SwitchOverBufferChain.Num())
 		{
-			ActiveOutputID = TrackID;
+			if (!SwitchOverBufferChain.Last().Info->PeriodAdaptationSetID.Equals(InBufferSourceInfo->PeriodAdaptationSetID))
+			{
+				FQueuedBuffer Next;
+				Next.Info = MoveTemp(InBufferSourceInfo);
+				SwitchOverBufferChain.Emplace(MoveTemp(Next));
+			}
 		}
-		bAutoselectFirstTrack = false;
+		else
+		{
+			if (UpcomingBufferChain.Num() == 0 || !UpcomingBufferChain.Last().Info->PeriodAdaptationSetID.Equals(InBufferSourceInfo->PeriodAdaptationSetID))
+			{
+				FQueuedBuffer Next;
+				Next.Info = MoveTemp(InBufferSourceInfo);
+				UpcomingBufferChain.Emplace(MoveTemp(Next));
+			}
+		}
+	}
+
+
+	void FMultiTrackAccessUnitBuffer::Activate()
+	{
 		bIsDeselected = false;
 	}
 
 	void FMultiTrackAccessUnitBuffer::Deselect()
 	{
+		// The buffer remains selected in order to track its fullness etc.
+		// Deselected only means that the access unit may not get sent to the decoder
+		// or be replaced by one that produces no output.
 		bIsDeselected = true;
-		// Note: We do not change ActiveOutputID or bAutoselectFirstTrack.
-		//       The buffer actually remains selected in order to track its fullness etc.
-		//       Deselected only means that the access unit may not get sent to the decoder
-		//       or be replaced by one that produces no output.
 	}
 
-	bool FMultiTrackAccessUnitBuffer::IsDeselected() const
+	bool FMultiTrackAccessUnitBuffer::IsDeselected()
 	{
 		return bIsDeselected;
 	}
 
-	FTimeValue FMultiTrackAccessUnitBuffer::GetLastPoppedPTS() const
+	FTimeValue FMultiTrackAccessUnitBuffer::GetLastPoppedPTS()
 	{
 		FMediaCriticalSection::ScopedLock lock(AccessLock);
 		return LastPoppedPTS;
 	}
 
 
-	TSharedPtr<FAccessUnitBuffer, ESPMode::ThreadSafe> FMultiTrackAccessUnitBuffer::GetSelectedTrackBuffer()
+	TSharedPtrTS<FAccessUnitBuffer> FMultiTrackAccessUnitBuffer::GetSelectedTrackBuffer()
 	{
-		FAccessUnitBufferPtr TrackBuffer;
-		AccessLock.Lock();
-		if (TrackBuffers.Contains(ActiveOutputID))
+		TSharedPtrTS<FAccessUnitBuffer> TrackBuffer;
+		FMediaCriticalSection::ScopedLock lock(AccessLock);
+		if (ActiveOutputBufferInfo.IsValid() && TrackBuffers.Contains(ActiveOutputBufferInfo->PeriodAdaptationSetID))
 		{
-			TrackBuffer = TrackBuffers[ActiveOutputID];
+			TrackBuffer = TrackBuffers[ActiveOutputBufferInfo->PeriodAdaptationSetID];
 		}
 		else
 		{
 			TrackBuffer = EmptyBuffer;
 		}
-		AccessLock.Unlock();
 		return TrackBuffer;
 	}
 
-	TSharedPtr<const FAccessUnitBuffer, ESPMode::ThreadSafe> FMultiTrackAccessUnitBuffer::GetSelectedTrackBuffer() const
+	void FMultiTrackAccessUnitBuffer::GetStats(FAccessUnitBufferInfo& OutStats)
 	{
-		FAccessUnitBufferPtr TrackBuffer;
-		AccessLock.Lock();
-		if (TrackBuffers.Contains(ActiveOutputID))
+		FMediaCriticalSection::ScopedLock lock(AccessLock);
+		// Stats are always returned with the configuration for one buffer, not
+		// any summation of enqueued buffers.
+		TSharedPtrTS<const FAccessUnitBuffer> Buf = GetSelectedTrackBuffer();
+		Buf->GetStats(OutStats, &BufferConfiguration);
+		if (Buf == EmptyBuffer)
 		{
-			TrackBuffer = TrackBuffers[ActiveOutputID];
+			OutStats.bEndOfData = bEndOfData;
 		}
 		else
 		{
-			TrackBuffer = EmptyBuffer;
+			// Add the stats of any enqueued buffers
+			for(int32 i=0; i<UpcomingBufferChain.Num(); ++i)
+			{
+				if (UpcomingBufferChain[i].Buffer.IsValid())
+				{
+					FAccessUnitBufferInfo NextStats;
+					UpcomingBufferChain[i].Buffer->GetStats(NextStats);
+					OutStats.PushedDuration += NextStats.PushedDuration;
+					OutStats.PlayableDuration += NextStats.PlayableDuration;
+					OutStats.CurrentMemInUse += NextStats.CurrentMemInUse;
+					OutStats.NumCurrentAccessUnits += NextStats.NumCurrentAccessUnits;
+					// With chained buffers the end of data and last push blocked flags are only
+					// relevant on the last buffer in the chain.
+					OutStats.bEndOfData = NextStats.bEndOfData;
+					OutStats.bLastPushWasBlocked = NextStats.bLastPushWasBlocked;
+				}
+			}
 		}
-		AccessLock.Unlock();
-		return TrackBuffer;
 	}
 
-	void FMultiTrackAccessUnitBuffer::GetStats(FAccessUnitBufferInfo& OutStats)
+	void FMultiTrackAccessUnitBuffer::ChangeOver()
 	{
-		// Stats are always returned with our configuration.
-		TSharedPtr<const FAccessUnitBuffer, ESPMode::ThreadSafe> Buf = GetSelectedTrackBuffer();
-		Buf->GetStats(OutStats, &PrimaryBufferConfiguration);
-		if (Buf == EmptyBuffer)
+		ActiveOutputBufferInfo = SwitchOverBufferChain[0].Info;
+		SwitchOverBufferChain.RemoveAt(0);
+		UpcomingBufferChain = MoveTemp(SwitchOverBufferChain);
+		RemoveUnusedBuffers();
+	}
+
+	void FMultiTrackAccessUnitBuffer::RemoveUnusedBuffers()
+	{
+		if (!bIsParallelTrackMode)
 		{
-			OutStats.bEndOfData = bEndOfData;
+			auto IsReferencesInChain = [&](const FString& ID, TArray<FQueuedBuffer>& BufferChain) -> bool
+			{
+				for(auto &Next : BufferChain)
+				{
+					if (Next.Info.IsValid() && Next.Info->PeriodAdaptationSetID.Equals(ID))
+					{
+						return true;
+					}
+				}
+				return false;
+			};
+
+			TArray<FString> UnusedBuffers;
+			for(auto &Buffer : TrackBuffers)
+			{
+				const FString& ID = Buffer.Key;
+				bool bIsReferenced = (ActiveOutputBufferInfo.IsValid() && ActiveOutputBufferInfo->PeriodAdaptationSetID.Equals(ID)) || IsReferencesInChain(ID, UpcomingBufferChain) || IsReferencesInChain(ID, SwitchOverBufferChain);
+				if (!bIsReferenced)
+				{
+					UnusedBuffers.Emplace(ID);
+				}
+			}
+			for(auto &Buffer : UnusedBuffers)
+			{
+				TrackBuffers.Remove(Buffer);
+			}
+		}
+	}
+
+	void FMultiTrackAccessUnitBuffer::ActivateBuffer(bool bIsPushing)
+	{
+		auto UpdateBufferChainBuffers = [&](TArray<FQueuedBuffer>& BufferChain) -> void
+		{
+			for(auto &Next : BufferChain)
+			{
+				if (!Next.Buffer.IsValid() && TrackBuffers.Contains(Next.Info->PeriodAdaptationSetID))
+				{
+					Next.Buffer = TrackBuffers[Next.Info->PeriodAdaptationSetID];
+				}
+			}
+		};
+
+
+		// Update the upcoming buffer chains for GetStats().
+		UpdateBufferChainBuffers(UpcomingBufferChain);
+		UpdateBufferChainBuffers(SwitchOverBufferChain);
+
+		if (!ActiveOutputBufferInfo.IsValid())
+		{
+			check(UpcomingBufferChain.IsEmpty());
+			if (SwitchOverBufferChain.Num())
+			{
+				ChangeOver();
+			}
+		}
+		// Upon pushing data to the buffer we only need to activate the track and update the buffer chains
+		// for statistics. Buffer switching is handled when popping data.
+		if (bIsPushing)
+		{
+			return;
+		}
+
+		if (SwitchOverBufferChain.Num())
+		{
+			// In parallel track mode we switch immediately as the individual buffers are
+			// time synchronized already.
+			if (bIsParallelTrackMode)
+			{
+				ChangeOver();
+			}
+			else
+			{
+				FTimeValue PoppedDTS, PoppedPTS;
+				FTimeValue PopUntilDTS, PopUntilPTS;
+				PopUntilDTS = LastPoppedDTS;
+				PopUntilPTS = LastPoppedPTS;
+				bool bSwitch = false;
+				if (PopUntilDTS.IsValid() && PopUntilPTS.IsValid())
+				{
+					for(int32 nBuf=0; nBuf<SwitchOverBufferChain.Num(); ++nBuf)
+					{
+						if (SwitchOverBufferChain[nBuf].Buffer.IsValid())
+						{
+							// Remove AUs from the buffer that are outdated.
+							SwitchOverBufferChain[nBuf].Buffer->DiscardUntil(PopUntilDTS, PopUntilPTS, PoppedDTS, PoppedPTS);
+
+							// Does it contain the time we want to switch at?
+							if (SwitchOverBufferChain[nBuf].Buffer->ContainsPTS(PopUntilPTS))
+							{
+								bSwitch = true;
+								break;
+							}
+							// If the buffer is now empty with another non-empty one following we can remove this one.
+							if (SwitchOverBufferChain[nBuf].Buffer->Num() == 0 && nBuf+1 < SwitchOverBufferChain.Num() && SwitchOverBufferChain[nBuf+1].Buffer.IsValid() && SwitchOverBufferChain[nBuf+1].Buffer->Num())
+							{
+								SwitchOverBufferChain.RemoveAt(nBuf);
+								--nBuf;
+							}
+						}
+						else
+						{
+							// If this buffer is not valid (yet) then the following ones are not either.
+							break;
+						}
+					}
+				}
+				else
+				{
+					// When no data was popped from the buffer yet do an immediate switch.
+					bSwitch = true;
+				}
+				if (bSwitch)
+				{
+					ChangeOver();
+				}
+			}
+		}
+
+		// Moving into the next enqueued buffer?
+		if (UpcomingBufferChain.Num())
+		{
+			TSharedPtrTS<FAccessUnitBuffer> ActiveBuffer = GetSelectedTrackBuffer();
+			if (ActiveBuffer->Num() == 0 && UpcomingBufferChain[0].Buffer.IsValid() && UpcomingBufferChain[0].Buffer->Num())
+			{
+				ActiveOutputBufferInfo = UpcomingBufferChain[0].Info;
+				UpcomingBufferChain.RemoveAt(0);
+				RemoveUnusedBuffers();
+			}
 		}
 	}
 
+
 	bool FMultiTrackAccessUnitBuffer::Pop(FAccessUnit*& OutAU)
 	{
 		// Note: We assume the access lock is held by the caller!
-		TSharedPtr<FAccessUnitBuffer, ESPMode::ThreadSafe> Buf = FMultiTrackAccessUnitBuffer::GetSelectedTrackBuffer();
+		ActivateBuffer(false);
+
+		TSharedPtrTS<FAccessUnitBuffer> Buf = FMultiTrackAccessUnitBuffer::GetSelectedTrackBuffer();
 		if (Buf->Num())
 		{
 			bool bDidPop = Buf->Pop(OutAU);
 			if (bDidPop && OutAU)
 			{
 				// Did we just pop from a different buffer than last time?
-				if (LastPoppedBufferID != ActiveOutputID)
+				if (LastPoppedBufferInfo != ActiveOutputBufferInfo)
 				{
 					// FIXME: We may need to discard everything that is not tagged as a sync sample to ensure proper stream switching.
 					OutAU->bTrackChangeDiscontinuity = true;
+
+					// Remember from which buffer we popped.
+					LastPoppedBufferInfo = ActiveOutputBufferInfo;
 				}
-				// Remember from which buffer we popped.
-				LastPoppedBufferID = ActiveOutputID;
 				LastPoppedDTS = OutAU->DTS;
 				LastPoppedPTS = OutAU->PTS;
 				// With this AU being popped off now we will also pop off all AUs that are now obsolete in the other tracks.
-				for (auto& It : TrackBuffers)
+				for(auto& It : TrackBuffers)
 				{
 					if (It.Value != Buf)
 					{
@@ -257,9 +493,9 @@ namespace Electra
 	void FMultiTrackAccessUnitBuffer::PopDiscardUntil(FTimeValue UntilTime)
 	{
 		// Note: We assume the access lock is held by the caller!
-		for (auto& It : TrackBuffers)
+		for(auto& It : TrackBuffers)
 		{
-			FAccessUnitBufferPtr Buf = It.Value;
+			TSharedPtrTS<FAccessUnitBuffer> Buf = It.Value;
 			FTimeValue PoppedDTS, PoppedPTS;
 			Buf->DiscardUntil(UntilTime, UntilTime, PoppedDTS, PoppedPTS);
 			if (PoppedDTS.IsValid() && LastPoppedDTS.IsValid() && PoppedDTS > LastPoppedDTS)
@@ -274,20 +510,58 @@ namespace Electra
 	}
 
 
-	bool FMultiTrackAccessUnitBuffer::IsEODFlagSet() const
+	bool FMultiTrackAccessUnitBuffer::IsEODFlagSet()
 	{
 		// Note: We assume the access lock is held by the caller!
-		TSharedPtr<const FAccessUnitBuffer, ESPMode::ThreadSafe> Buf = GetSelectedTrackBuffer();
+		TSharedPtrTS<FAccessUnitBuffer> Buf = GetSelectedTrackBuffer();
 		return Buf == EmptyBuffer ? bEndOfData : Buf->IsEODFlagSet();
 	}
 
-	int32 FMultiTrackAccessUnitBuffer::Num() const
+	int32 FMultiTrackAccessUnitBuffer::Num()
 	{
 		// Note: We assume the access lock is held by the caller!
-		return (int32)GetSelectedTrackBuffer()->Num();
+		int64 TotalNum = GetSelectedTrackBuffer()->Num();
+		// Add the enqueued buffers
+		for(int32 i=0; i<UpcomingBufferChain.Num(); ++i)
+		{
+			if (UpcomingBufferChain[i].Buffer.IsValid())
+			{
+				TotalNum += UpcomingBufferChain[i].Buffer->Num();
+			}
+		}
+		return (int32) TotalNum;
+	}
+
+	void FMultiTrackAccessUnitBuffer::GetEnqueuedBufferInfo(FAccessUnitBuffer::FExternalBufferInfo& OutInfo, bool bForSwitchOverChain)
+	{
+		if (!bForSwitchOverChain)
+		{
+			TSharedPtrTS<FAccessUnitBuffer> Buf = GetSelectedTrackBuffer();
+			OutInfo.DataSize = Buf->AllocatedSize();
+			OutInfo.Duration = Buf->GetPlayableDuration();
+			for(int32 i=0; i<UpcomingBufferChain.Num(); ++i)
+			{
+				if (UpcomingBufferChain[i].Buffer.IsValid())
+				{
+					OutInfo.DataSize += UpcomingBufferChain[i].Buffer->AllocatedSize();
+					OutInfo.Duration += UpcomingBufferChain[i].Buffer->GetPlayableDuration();
+				}
+			}
+		}
+		else
+		{
+			for(auto &SwitchBuf : SwitchOverBufferChain)
+			{
+				if (SwitchBuf.Buffer.IsValid())
+				{
+					OutInfo.DataSize += SwitchBuf.Buffer->AllocatedSize();
+					OutInfo.Duration += SwitchBuf.Buffer->GetPlayableDuration();
+				}
+			}
+		}
 	}
 
-	bool FMultiTrackAccessUnitBuffer::WasLastPushBlocked() const
+	bool FMultiTrackAccessUnitBuffer::WasLastPushBlocked()
 	{
 		return bLastPushWasBlocked;
 	}
diff --git a/Engine/Plugins/Media/ElectraPlayer/Source/ElectraPlayerRuntime/Private/Runtime/StreamAccessUnitBuffer.h b/Engine/Plugins/Media/ElectraPlayer/Source/ElectraPlayerRuntime/Private/Runtime/StreamAccessUnitBuffer.h
--- a/Engine/Plugins/Media/ElectraPlayer/Source/ElectraPlayerRuntime/Private/Runtime/StreamAccessUnitBuffer.h
+++ b/Engine/Plugins/Media/ElectraPlayer/Source/ElectraPlayerRuntime/Private/Runtime/StreamAccessUnitBuffer.h
@@ -29,14 +29,18 @@ namespace Electra
 
 
 
-
-	struct FStreamSourceInfo
+	/**
+	 * Information into which buffer the AU data needs to be placed.
+	 */
+	struct FBufferSourceInfo
 	{
+		// Identifies the period and track (adaptation set) this data is originating from.
+		FString		PeriodAdaptationSetID;
+		// Partial track metadata. See FTrackMetadata.
+		FString		Kind;
 		FString		Language;
-		FString		Role;
-		FString		PeriodID;
-		FString		AdaptationSetID;
-		FString		RepresentationID;
+		// Internal hard index, used for multiplexed streams.
+		int32		HardIndex = -1;
 	};
 
 	struct FAccessUnit
@@ -73,8 +77,8 @@ namespace Electra
 		bool						bIsDummyData;					//!< True if the decoder must be reset before decoding this AU.
 		bool						bTrackChangeDiscontinuity;		//!< True if this is the first AU after a track change.
 
-		TSharedPtr<const FStreamSourceInfo, ESPMode::ThreadSafe>	StreamSourceInfo;
-		TSharedPtr<const FPlayerLoopState, ESPMode::ThreadSafe>		PlayerLoopState;
+		TSharedPtrTS<const FBufferSourceInfo>	BufferSourceInfo;
+		TSharedPtrTS<const FPlayerLoopState>	PlayerLoopState;
 
 		// Reserved for decoders.
 		// FIXME: move this out of here and have the decoders track what they're doing with the AU.
@@ -227,6 +231,12 @@ namespace Electra
 			int64		MaxDataSize;
 		};
 
+		struct FExternalBufferInfo
+		{
+			FTimeValue	Duration = FTimeValue::GetZero();
+			int64		DataSize = 0;
+		};
+
 		FAccessUnitBuffer()
 			: FrontDTS(FTimeValue::GetInvalid())
 			, PushedDuration(FTimeValue::GetZero())
@@ -239,7 +249,7 @@ namespace Electra
 
 		~FAccessUnitBuffer()
 		{
-			while (AccessUnits.Num())
+			while(AccessUnits.Num())
 			{
 				FAccessUnit::Release(AccessUnits.Pop());
 			}
@@ -259,12 +269,10 @@ namespace Electra
 			return CurrentMemInUse;
 		}
 
-		//! Returns both the number of access units and the amount of allocated memory.
-		void SizeAndAllocation(int64& OutNumAUs, int64& OutAllocatedMem) const
+		//! Returns the amount of playable duration.
+		FTimeValue GetPlayableDuration() const
 		{
-			FMediaCriticalSection::ScopedLock Lock(AccessLock);
-			OutNumAUs = AccessUnits.Num();
-			OutAllocatedMem = CurrentMemInUse;
+			return PlayableDuration;
 		}
 
 		//! Returns all vital statistics.
@@ -292,14 +300,15 @@ namespace Electra
 		}
 
 		//! Adds an access unit to the FIFO. Returns true if successful, false if the FIFO has insufficient free space.
-		bool Push(FAccessUnit*& AU, const FConfiguration* Limit = nullptr)
+		bool Push(FAccessUnit*& AU, const FConfiguration* Limit = nullptr, const FExternalBufferInfo* ExternalInfo = nullptr)
 		{
 			FMediaCriticalSection::ScopedLock Lock(AccessLock);
 			// Pushing new data unconditionally clears the EOD flag even if the buffer is currently full.
 			// The attempt to push implies there will be more data.
 			bEndOfData = false;
 			Limit = Limit ? Limit : &Config;
-			if (CanPush(AU, Limit))
+			ExternalInfo = ExternalInfo ? ExternalInfo : &ZeroExternalInfo;
+			if (CanPush(AU, Limit, ExternalInfo))
 			{
 				bLastPushWasBlocked = false;
 				AccessUnits.Push(AU);
@@ -352,7 +361,7 @@ namespace Electra
 					else
 					{
 						FrontDTS.SetToInvalid();
-						for (uint32 i = 1; i < AccessUnits.Num(); ++i)
+						for(uint32 i = 1; i < AccessUnits.Num(); ++i)
 						{
 							if (AccessUnits[i]->DropState == FAccessUnit::EDropState::None)
 							{
@@ -399,10 +408,32 @@ namespace Electra
 			}
 		}
 
+		//!
+		bool ContainsPTS(const FTimeValue& InPTS) const
+		{
+			FMediaCriticalSection::ScopedLock Lock(AccessLock);
+			if (AccessUnits.Num())
+			{
+				return AccessUnits.FrontRef()->PTS <= InPTS && InPTS < AccessUnits.BackRef()->PTS + AccessUnits.BackRef()->Duration;
+			}
+			return false;
+		}
+
+		bool ContainsFuturePTS(const FTimeValue& InPTS) const
+		{
+			FMediaCriticalSection::ScopedLock Lock(AccessLock);
+			if (AccessUnits.Num())
+			{
+				return InPTS <= AccessUnits.BackRef()->PTS + AccessUnits.BackRef()->Duration;
+			}
+			return false;
+		}
+
+
 		//! Discards data that has both its DTS and PTS less than the provided ones.
-		void DiscardUntil(FTimeValue NextValidDTS, FTimeValue NextValidPTS, FTimeValue& OutPoppedDTS, FTimeValue& OutPoppedPTS)
+		void DiscardUntil(const FTimeValue& NextValidDTS, const FTimeValue& NextValidPTS, FTimeValue& OutPoppedDTS, FTimeValue& OutPoppedPTS)
 		{
-			while (true)
+			while(true)
 			{
 				FAccessUnit* NextAU = nullptr;
 				if (Peek(NextAU))
@@ -420,6 +451,7 @@ namespace Electra
 						}
 						else
 						{
+							FAccessUnit::Release(NextAU);
 							break;
 						}
 					}
@@ -450,7 +482,7 @@ namespace Electra
 		void Flush()
 		{
 			FMediaCriticalSection::ScopedLock Lock(AccessLock);
-			while (AccessUnits.Num())
+			while(AccessUnits.Num())
 			{
 				FAccessUnit::Release(AccessUnits.Pop());
 				NumInSemaphore.TryToObtain();
@@ -503,23 +535,34 @@ namespace Electra
 
 	private:
 		//! Checks if an access unit can be pushed to the FIFO. Returns true if successful, false if the FIFO has insufficient free space.
-		bool CanPush(const FAccessUnit* AU, const FConfiguration* Limit = nullptr)
+		bool CanPush(const FAccessUnit* AU, const FConfiguration* Limit, const FExternalBufferInfo* ExternalInfo)
 		{
-			FMediaCriticalSection::ScopedLock Lock(AccessLock);
-			Limit = Limit ? Limit : &Config;
 			check(Limit->MaxDuration > FTimeValue::GetZero());
 			check(AU->Duration.IsValid() && !AU->Duration.IsInfinity());
-			// Memory ok?
-			if (AU->TotalMemSize() + CurrentMemInUse <= Limit->MaxDataSize)
+			if (ExternalInfo == nullptr)
+			{
+				// Memory ok?
+				if (AU->TotalMemSize() + CurrentMemInUse <= Limit->MaxDataSize)
+				{
+					// Max allowed duration ok?
+					return PlayableDuration.IsValid() && PlayableDuration + AU->Duration > Limit->MaxDuration ? false : true;
+				}
+			}
+			else
 			{
-				// Max allowed duration ok?
-				return PlayableDuration.IsValid() && PlayableDuration + AU->Duration > Limit->MaxDuration ? false : true;
+				// Memory ok?
+				if (AU->TotalMemSize() + ExternalInfo->DataSize <= Limit->MaxDataSize)
+				{
+					// Max allowed duration ok?
+					return AU->Duration + ExternalInfo->Duration > Limit->MaxDuration ? false : true;
+				}
 			}
 			return false;
 		}
 
-		mutable FMediaCriticalSection							AccessLock;
+		mutable FMediaCriticalSection				AccessLock;
 		FConfiguration								Config;
+		FExternalBufferInfo							ZeroExternalInfo;
 		TMediaQueueDynamicNoLock<FAccessUnit*>		AccessUnits;
 		FMediaSemaphore								NumInSemaphore;
 		FTimeValue									FrontDTS;					//!< DTS of first AU in buffer
@@ -542,18 +585,21 @@ namespace Electra
 	public:
 		FMultiTrackAccessUnitBuffer();
 		~FMultiTrackAccessUnitBuffer();
+		void SetParallelTrackMode();
 		void CapacitySet(const FAccessUnitBuffer::FConfiguration& Config);
-		void AutoselectFirstTrack();
-		void SelectTrackByID(const FString& TrackID);
+		void SelectTrackWhenAvailable(TSharedPtrTS<FBufferSourceInfo> InBufferSourceInfo);
+		void AddUpcomingBuffer(TSharedPtrTS<FBufferSourceInfo> InBufferSourceInfo);
+		void Activate();
 		void Deselect();
-		//FString GetSelectedTrackID() const;
 		bool Push(FAccessUnit*& AU);
-		void PushEndOfDataFor(TSharedPtr<const FStreamSourceInfo, ESPMode::ThreadSafe> InStreamSourceInfo);
+		void PushEndOfDataFor(TSharedPtrTS<const FBufferSourceInfo> InStreamSourceInfo);
 		void PushEndOfDataAll();
 		void Flush();
 		void GetStats(FAccessUnitBufferInfo& OutStats);
-		bool IsDeselected() const;
-		FTimeValue GetLastPoppedPTS() const;
+		bool IsDeselected();
+		FTimeValue GetLastPoppedPTS();
+		TSharedPtrTS<FBufferSourceInfo> GetActiveOutputBufferInfo() const
+		{ return ActiveOutputBufferInfo; }
 
 		// Helper class to lock the AU buffer
 		class FScopedLock : private TMediaNoncopyable<FScopedLock>
@@ -575,31 +621,42 @@ namespace Electra
 
 		bool Pop(FAccessUnit*& OutAU);
 		void PopDiscardUntil(FTimeValue UntilTime);
-		bool IsEODFlagSet() const;
-		int32 Num() const;
-		bool WasLastPushBlocked() const;
+		bool IsEODFlagSet();
+		int32 Num();
+		bool WasLastPushBlocked();
 
 	private:
-		using FAccessUnitBufferPtr = TSharedPtr<FAccessUnitBuffer, ESPMode::ThreadSafe>;
 
-		void PurgeAll();
+		struct FQueuedBuffer
+		{
+			TSharedPtrTS<FBufferSourceInfo> Info;
+			TSharedPtrTS<FAccessUnitBuffer> Buffer;
+		};
 
-		FAccessUnitBufferPtr CreateNewBuffer();
-		TSharedPtr<FAccessUnitBuffer, ESPMode::ThreadSafe> GetSelectedTrackBuffer();
-		TSharedPtr<const FAccessUnitBuffer, ESPMode::ThreadSafe> GetSelectedTrackBuffer() const;
-
-		mutable FMediaCriticalSection			AccessLock;
-		FAccessUnitBuffer::FConfiguration		PrimaryBufferConfiguration;			//!< Configuration of the primary (active) buffer. All other buffers are unbounded to ensure no stalling of the demuxer.
-		FAccessUnitBufferPtr					EmptyBuffer;						//!< An empty buffer
-		FString									ActiveOutputID;						//!< Track ID of the buffer that is feeding the decoder.
-		FString									LastPoppedBufferID;					//!< Buffer ID from which an AU was popped last.
-		TMap<FString, FAccessUnitBufferPtr>		TrackBuffers;						//!< Map of track buffers. One per track ID.
-		FTimeValue								LastPoppedDTS;
-		FTimeValue								LastPoppedPTS;
-		bool									bAutoselectFirstTrack;
-		bool									bIsDeselected;
-		bool									bEndOfData;
-		bool									bLastPushWasBlocked;
+		void PurgeAll();
+		void Clear();
+
+		TSharedPtrTS<FAccessUnitBuffer> CreateNewBuffer();
+		TSharedPtrTS<FAccessUnitBuffer> GetSelectedTrackBuffer();
+		void ActivateBuffer(bool bIsPushing);
+		void ChangeOver();
+		void RemoveUnusedBuffers();
+		void GetEnqueuedBufferInfo(FAccessUnitBuffer::FExternalBufferInfo& OutInfo, bool bForSwitchOverChain);
+
+		FMediaCriticalSection							AccessLock;
+		FAccessUnitBuffer::FConfiguration				BufferConfiguration;				//!< Buffer configuration.
+		TMap<FString, TSharedPtrTS<FAccessUnitBuffer>>	TrackBuffers;						//!< Map of track buffers. One per track ID.
+		TArray<FQueuedBuffer>							UpcomingBufferChain;
+		TArray<FQueuedBuffer>							SwitchOverBufferChain;
+		TSharedPtrTS<FAccessUnitBuffer>					EmptyBuffer;						//!< An empty buffer
+		TSharedPtrTS<FBufferSourceInfo>					ActiveOutputBufferInfo;
+		TSharedPtrTS<FBufferSourceInfo>					LastPoppedBufferInfo;
+		FTimeValue										LastPoppedDTS;
+		FTimeValue										LastPoppedPTS;
+		bool											bIsDeselected;
+		bool											bEndOfData;
+		bool											bLastPushWasBlocked;
+		bool											bIsParallelTrackMode;
 	};
 
 
diff --git a/Engine/Plugins/Media/ElectraPlayer/Source/ElectraPlayerRuntime/Private/Runtime/StreamTypes.h b/Engine/Plugins/Media/ElectraPlayer/Source/ElectraPlayerRuntime/Private/Runtime/StreamTypes.h
--- a/Engine/Plugins/Media/ElectraPlayer/Source/ElectraPlayerRuntime/Private/Runtime/StreamTypes.h
+++ b/Engine/Plugins/Media/ElectraPlayer/Source/ElectraPlayerRuntime/Private/Runtime/StreamTypes.h
@@ -519,18 +519,25 @@ namespace Electra
 
 	/**
 	 * Metadata per track type.
+	 *
+	 * See: https://dev.w3.org/html5/html-sourcing-inband-tracks/
+	 *
+	 * While this may not represent every possible "role" the presentation format may offer this is a
+	 * representation that covers most of the use cases and can be applied to a variety of formats.
 	 */
 	struct FTrackMetadata
 	{
-		FString TrackID;
-		FString Role;
-		FString Language;
+		FString ID;
+		FString Kind;
+		FString Language;		// ISO-639-1 language code
+		FString Label;
+
 		TArray<FStreamMetadata> StreamDetails;
 		FStreamCodecInformation HighestBandwidthCodec;
 		int32 HighestBandwidth = 0;
 		bool Equals(const FTrackMetadata& Other) const
 		{
-			bool bEquals = TrackID.Equals(Other.TrackID) && Role.Equals(Other.Role) && Language.Equals(Other.Language);
+			bool bEquals = ID.Equals(Other.ID) && Kind.Equals(Other.Kind) && Language.Equals(Other.Language);
 			if (bEquals)
 			{
 				if (StreamDetails.Num() == Other.StreamDetails.Num())
@@ -551,16 +558,63 @@ namespace Electra
 
 
 	/**
-	 * Temporary!
-	 * Used to convey the user's stream preferences (ie. language and other things) to the player.
+	 * Stream selection attributes.
+	 * See FTrackMetadata comments.
 	 */
-	struct FStreamPreferences
+	struct FStreamSelectionAttributes
 	{
-		FParamDict	Unused;							//!< This is only here as a placeholder for now. Nothing writes to or reads from this.
+		// Primarily used for audio selection. Should typically be set to "main" or left unset.
+		TOptional<FString> Kind;
+
+		// Used for audio and subtitles or captions. The language must be a valid two letter ISO-639 (-1, -2 or -3) code
+		// for which a two-letter ISO-639-1 equivalent exists. This is the equivalent of the primary language subtag of
+		// RFC-4646 and of RFC-5646
+		TOptional<FString> Language_ISO639;
+
+		// Rarely used. Unconditionally selects a track by its index where the index is a sequential numbering from [0..n)
+		// of the tracks as they are found. If the index is invalid the selection rules for kind and language are applied.
+		TOptional<int32> OverrideIndex;
+
+		void Reset()
+		{
+			Kind.Reset();
+			Language_ISO639.Reset();
+			OverrideIndex.Reset();
+		}
+		void UpdateWith(const FString& InKind, const FString& InLanguage, int32 InOverrideIndex)
+		{
+			Reset();
+			if (!InKind.IsEmpty())
+			{
+				Kind = InKind;
+			}
+			if (!InLanguage.IsEmpty())
+			{
+				Language_ISO639 = InLanguage;
+			}
+			if (InOverrideIndex >= 0)
+			{
+				OverrideIndex = InOverrideIndex;
+			}
+		}
+		void UpdateIfOverrideSet(const FString& InKind, const FString& InLanguage)
+		{
+			if (OverrideIndex.IsSet())
+			{
+				ClearOverrideIndex();
+				Kind = InKind;
+				Language_ISO639 = InLanguage;
+			}
+		}
+		void ClearOverrideIndex()
+		{
+			OverrideIndex.Reset();
+		}
 	};
 
 
 
+
 	struct FPlayerLoopState
 	{
 		FPlayerLoopState()
diff --git a/Engine/Plugins/Media/ElectraPlayer/Source/ElectraPlayerRuntime/Private/Runtime/Utilities/ISO639-Map.cpp b/Engine/Plugins/Media/ElectraPlayer/Source/ElectraPlayerRuntime/Private/Runtime/Utilities/ISO639-Map.cpp
new file mode 100644
--- /dev/null
+++ b/Engine/Plugins/Media/ElectraPlayer/Source/ElectraPlayerRuntime/Private/Runtime/Utilities/ISO639-Map.cpp
@@ -0,0 +1,278 @@
+// Copyright Epic Games, Inc. All Rights Reserved.
+
+#include "Utilities/ISO639-Map.h"
+
+namespace Electra
+{
+	namespace ISO639
+	{
+		struct FCodeEntry
+		{
+			const TCHAR* _2B;		// ISO-639-2B
+			const TCHAR* _3;		// ISO-639-3 (same as ISO 639-2T)
+			const TCHAR* _1;		// ISO-639-1
+		};
+
+		/*
+			This map was generated as given by https://en.wikipedia.org/wiki/List_of_ISO_639-2_codes
+
+			Elements for which no 2 letter ISO-639-1 exists have been removed.
+		*/
+
+		static TArray<FCodeEntry> CodeMap =
+		{
+			{TEXT("aar"), TEXT("aar"), TEXT("aa")},	// Afar
+			{TEXT("abk"), TEXT("abk"), TEXT("ab")},	// Abkhazian
+			{TEXT("afr"), TEXT("afr"), TEXT("af")},	// Afrikaans
+			{TEXT("aka"), TEXT("aka"), TEXT("ak")},	// Akan
+			{TEXT("amh"), TEXT("amh"), TEXT("am")},	// Amharic
+			{TEXT("ara"), TEXT("ara"), TEXT("ar")},	// Arabic
+			{TEXT("arg"), TEXT("arg"), TEXT("an")},	// Aragonese
+			{TEXT("asm"), TEXT("asm"), TEXT("as")},	// Assamese
+			{TEXT("ava"), TEXT("ava"), TEXT("av")},	// Avaric
+			{TEXT("ave"), TEXT("ave"), TEXT("ae")},	// Avestan
+			{TEXT("aym"), TEXT("aym"), TEXT("ay")},	// Aymara
+			{TEXT("aze"), TEXT("aze"), TEXT("az")},	// Azerbaijani
+			{TEXT("bak"), TEXT("bak"), TEXT("ba")},	// Bashkir
+			{TEXT("bam"), TEXT("bam"), TEXT("bm")},	// Bambara
+			{TEXT("bel"), TEXT("bel"), TEXT("be")},	// Belarusian
+			{TEXT("ben"), TEXT("ben"), TEXT("bn")},	// Bengali
+			{TEXT("bis"), TEXT("bis"), TEXT("bi")},	// Bislama
+			{TEXT("tib"), TEXT("bod"), TEXT("bo")},	// Tibetan
+			{TEXT("bos"), TEXT("bos"), TEXT("bs")},	// Bosnian
+			{TEXT("bre"), TEXT("bre"), TEXT("br")},	// Breton
+			{TEXT("bul"), TEXT("bul"), TEXT("bg")},	// Bulgarian
+			{TEXT("cat"), TEXT("cat"), TEXT("ca")},	// Catalan; Valencian
+			{TEXT("cze"), TEXT("ces"), TEXT("cs")},	// Czech
+			{TEXT("cha"), TEXT("cha"), TEXT("ch")},	// Chamorro
+			{TEXT("che"), TEXT("che"), TEXT("ce")},	// Chechen
+			{TEXT("chu"), TEXT("chu"), TEXT("cu")},	// Church Slavic; Old Slavonic; Church Slavonic; Old Bulgarian; Old Church Slavonic
+			{TEXT("chv"), TEXT("chv"), TEXT("cv")},	// Chuvash
+			{TEXT("cor"), TEXT("cor"), TEXT("kw")},	// Cornish
+			{TEXT("cos"), TEXT("cos"), TEXT("co")},	// Corsican
+			{TEXT("cre"), TEXT("cre"), TEXT("cr")},	// Cree
+			{TEXT("wel"), TEXT("cym"), TEXT("cy")},	// Welsh
+			{TEXT("dan"), TEXT("dan"), TEXT("da")},	// Danish
+			{TEXT("ger"), TEXT("deu"), TEXT("de")},	// German
+			{TEXT("div"), TEXT("div"), TEXT("dv")},	// Divehi; Dhivehi; Maldivian
+			{TEXT("dzo"), TEXT("dzo"), TEXT("dz")},	// Dzongkha
+			{TEXT("gre"), TEXT("ell"), TEXT("el")},	// Greek, Modern (1453–)
+			{TEXT("eng"), TEXT("eng"), TEXT("en")},	// English
+			{TEXT("epo"), TEXT("epo"), TEXT("eo")},	// Esperanto
+			{TEXT("est"), TEXT("est"), TEXT("et")},	// Estonian
+			{TEXT("baq"), TEXT("eus"), TEXT("eu")},	// Basque
+			{TEXT("ewe"), TEXT("ewe"), TEXT("ee")},	// Ewe
+			{TEXT("fao"), TEXT("fao"), TEXT("fo")},	// Faroese
+			{TEXT("per"), TEXT("fas"), TEXT("fa")},	// Persian
+			{TEXT("fij"), TEXT("fij"), TEXT("fj")},	// Fijian
+			{TEXT("fin"), TEXT("fin"), TEXT("fi")},	// Finnish
+			{TEXT("fre"), TEXT("fra"), TEXT("fr")},	// French
+			{TEXT("fry"), TEXT("fry"), TEXT("fy")},	// Western Frisian
+			{TEXT("ful"), TEXT("ful"), TEXT("ff")},	// Fulah
+			{TEXT("gla"), TEXT("gla"), TEXT("gd")},	// Gaelic; Scottish Gaelic
+			{TEXT("gle"), TEXT("gle"), TEXT("ga")},	// Irish
+			{TEXT("glg"), TEXT("glg"), TEXT("gl")},	// Galician
+			{TEXT("glv"), TEXT("glv"), TEXT("gv")},	// Manx
+			{TEXT("grn"), TEXT("grn"), TEXT("gn")},	// Guarani
+			{TEXT("guj"), TEXT("guj"), TEXT("gu")},	// Gujarati
+			{TEXT("hat"), TEXT("hat"), TEXT("ht")},	// Haitian; Haitian Creole
+			{TEXT("hau"), TEXT("hau"), TEXT("ha")},	// Hausa
+			{TEXT("heb"), TEXT("heb"), TEXT("he")},	// Hebrew
+			{TEXT("her"), TEXT("her"), TEXT("hz")},	// Herero
+			{TEXT("hin"), TEXT("hin"), TEXT("hi")},	// Hindi
+			{TEXT("hmo"), TEXT("hmo"), TEXT("ho")},	// Hiri Motu
+			{TEXT("hrv"), TEXT("hrv"), TEXT("hr")},	// Croatian
+			{TEXT("hun"), TEXT("hun"), TEXT("hu")},	// Hungarian
+			{TEXT("arm"), TEXT("hye"), TEXT("hy")},	// Armenian
+			{TEXT("ibo"), TEXT("ibo"), TEXT("ig")},	// Igbo
+			{TEXT("ido"), TEXT("ido"), TEXT("io")},	// Ido
+			{TEXT("iii"), TEXT("iii"), TEXT("ii")},	// Sichuan Yi; Nuosu
+			{TEXT("iku"), TEXT("iku"), TEXT("iu")},	// Inuktitut
+			{TEXT("ile"), TEXT("ile"), TEXT("ie")},	// Interlingue; Occidental
+			{TEXT("ina"), TEXT("ina"), TEXT("ia")},	// Interlingua (International Auxiliary Language Association)
+			{TEXT("ind"), TEXT("ind"), TEXT("id")},	// Indonesian
+			{TEXT("ipk"), TEXT("ipk"), TEXT("ik")},	// Inupiaq
+			{TEXT("ice"), TEXT("isl"), TEXT("is")},	// Icelandic
+			{TEXT("ita"), TEXT("ita"), TEXT("it")},	// Italian
+			{TEXT("jav"), TEXT("jav"), TEXT("jv")},	// Javanese
+			{TEXT("jpn"), TEXT("jpn"), TEXT("ja")},	// Japanese
+			{TEXT("kal"), TEXT("kal"), TEXT("kl")},	// Kalaallisut; Greenlandic
+			{TEXT("kan"), TEXT("kan"), TEXT("kn")},	// Kannada
+			{TEXT("kas"), TEXT("kas"), TEXT("ks")},	// Kashmiri
+			{TEXT("geo"), TEXT("kat"), TEXT("ka")},	// Georgian
+			{TEXT("kau"), TEXT("kau"), TEXT("kr")},	// Kanuri
+			{TEXT("kaz"), TEXT("kaz"), TEXT("kk")},	// Kazakh
+			{TEXT("khm"), TEXT("khm"), TEXT("km")},	// Central Khmer
+			{TEXT("kik"), TEXT("kik"), TEXT("ki")},	// Kikuyu; Gikuyu
+			{TEXT("kin"), TEXT("kin"), TEXT("rw")},	// Kinyarwanda
+			{TEXT("kir"), TEXT("kir"), TEXT("ky")},	// Kirghiz; Kyrgyz
+			{TEXT("kom"), TEXT("kom"), TEXT("kv")},	// Komi
+			{TEXT("kon"), TEXT("kon"), TEXT("kg")},	// Kongo
+			{TEXT("kor"), TEXT("kor"), TEXT("ko")},	// Korean
+			{TEXT("kua"), TEXT("kua"), TEXT("kj")},	// Kuanyama; Kwanyama
+			{TEXT("kur"), TEXT("kur"), TEXT("ku")},	// Kurdish
+			{TEXT("lao"), TEXT("lao"), TEXT("lo")},	// Lao
+			{TEXT("lat"), TEXT("lat"), TEXT("la")},	// Latin
+			{TEXT("lav"), TEXT("lav"), TEXT("lv")},	// Latvian
+			{TEXT("lim"), TEXT("lim"), TEXT("li")},	// Limburgan; Limburger; Limburgish
+			{TEXT("lin"), TEXT("lin"), TEXT("ln")},	// Lingala
+			{TEXT("lit"), TEXT("lit"), TEXT("lt")},	// Lithuanian
+			{TEXT("ltz"), TEXT("ltz"), TEXT("lb")},	// Luxembourgish; Letzeburgesch
+			{TEXT("lub"), TEXT("lub"), TEXT("lu")},	// Luba-Katanga
+			{TEXT("lug"), TEXT("lug"), TEXT("lg")},	// Ganda
+			{TEXT("mah"), TEXT("mah"), TEXT("mh")},	// Marshallese
+			{TEXT("mal"), TEXT("mal"), TEXT("ml")},	// Malayalam
+			{TEXT("mar"), TEXT("mar"), TEXT("mr")},	// Marathi
+			{TEXT("mac"), TEXT("mkd"), TEXT("mk")},	// Macedonian
+			{TEXT("mlg"), TEXT("mlg"), TEXT("mg")},	// Malagasy
+			{TEXT("mlt"), TEXT("mlt"), TEXT("mt")},	// Maltese
+			{TEXT("mon"), TEXT("mon"), TEXT("mn")},	// Mongolian
+			{TEXT("mao"), TEXT("mri"), TEXT("mi")},	// Māori
+			{TEXT("may"), TEXT("msa"), TEXT("ms")},	// Malay
+			{TEXT("bur"), TEXT("mya"), TEXT("my")},	// Burmese
+			{TEXT("nau"), TEXT("nau"), TEXT("na")},	// Nauru
+			{TEXT("nav"), TEXT("nav"), TEXT("nv")},	// Navajo; Navaho
+			{TEXT("nbl"), TEXT("nbl"), TEXT("nr")},	// Ndebele, South; South Ndebele
+			{TEXT("nde"), TEXT("nde"), TEXT("nd")},	// Ndebele, North; North Ndebele
+			{TEXT("ndo"), TEXT("ndo"), TEXT("ng")},	// Ndonga
+			{TEXT("nep"), TEXT("nep"), TEXT("ne")},	// Nepali
+			{TEXT("dut"), TEXT("nld"), TEXT("nl")},	// Dutch; Flemish
+			{TEXT("nno"), TEXT("nno"), TEXT("nn")},	// Norwegian Nynorsk; Nynorsk, Norwegian
+			{TEXT("nob"), TEXT("nob"), TEXT("nb")},	// Bokmål, Norwegian; Norwegian Bokmål
+			{TEXT("nor"), TEXT("nor"), TEXT("no")},	// Norwegian
+			{TEXT("nya"), TEXT("nya"), TEXT("ny")},	// Chichewa; Chewa; Nyanja
+			{TEXT("oci"), TEXT("oci"), TEXT("oc")},	// Occitan (post 1500)
+			{TEXT("oji"), TEXT("oji"), TEXT("oj")},	// Ojibwa
+			{TEXT("ori"), TEXT("ori"), TEXT("or")},	// Oriya
+			{TEXT("orm"), TEXT("orm"), TEXT("om")},	// Oromo
+			{TEXT("oss"), TEXT("oss"), TEXT("os")},	// Ossetian; Ossetic
+			{TEXT("pan"), TEXT("pan"), TEXT("pa")},	// Panjabi; Punjabi
+			{TEXT("pli"), TEXT("pli"), TEXT("pi")},	// Pali
+			{TEXT("pol"), TEXT("pol"), TEXT("pl")},	// Polish
+			{TEXT("por"), TEXT("por"), TEXT("pt")},	// Portuguese
+			{TEXT("pus"), TEXT("pus"), TEXT("ps")},	// Pushto; Pashto
+			{TEXT("que"), TEXT("que"), TEXT("qu")},	// Quechua
+			{TEXT("roh"), TEXT("roh"), TEXT("rm")},	// Romansh
+			{TEXT("rum"), TEXT("ron"), TEXT("ro")},	// Romanian; Moldavian; Moldovan
+			{TEXT("run"), TEXT("run"), TEXT("rn")},	// Rundi
+			{TEXT("rus"), TEXT("rus"), TEXT("ru")},	// Russian
+			{TEXT("sag"), TEXT("sag"), TEXT("sg")},	// Sango
+			{TEXT("san"), TEXT("san"), TEXT("sa")},	// Sanskrit
+			{TEXT("sin"), TEXT("sin"), TEXT("si")},	// Sinhala; Sinhalese
+			{TEXT("slo"), TEXT("slk"), TEXT("sk")},	// Slovak
+			{TEXT("slv"), TEXT("slv"), TEXT("sl")},	// Slovenian
+			{TEXT("sme"), TEXT("sme"), TEXT("se")},	// Northern Sami
+			{TEXT("smo"), TEXT("smo"), TEXT("sm")},	// Samoan
+			{TEXT("sna"), TEXT("sna"), TEXT("sn")},	// Shona
+			{TEXT("snd"), TEXT("snd"), TEXT("sd")},	// Sindhi
+			{TEXT("som"), TEXT("som"), TEXT("so")},	// Somali
+			{TEXT("sot"), TEXT("sot"), TEXT("st")},	// Sotho, Southern
+			{TEXT("spa"), TEXT("spa"), TEXT("es")},	// Spanish; Castilian
+			{TEXT("alb"), TEXT("sqi"), TEXT("sq")},	// Albanian
+			{TEXT("srd"), TEXT("srd"), TEXT("sc")},	// Sardinian
+			{TEXT("srp"), TEXT("srp"), TEXT("sr")},	// Serbian
+			{TEXT("ssw"), TEXT("ssw"), TEXT("ss")},	// Swati
+			{TEXT("sun"), TEXT("sun"), TEXT("su")},	// Sundanese
+			{TEXT("swa"), TEXT("swa"), TEXT("sw")},	// Swahili
+			{TEXT("swe"), TEXT("swe"), TEXT("sv")},	// Swedish
+			{TEXT("tah"), TEXT("tah"), TEXT("ty")},	// Tahitian
+			{TEXT("tam"), TEXT("tam"), TEXT("ta")},	// Tamil
+			{TEXT("tat"), TEXT("tat"), TEXT("tt")},	// Tatar
+			{TEXT("tel"), TEXT("tel"), TEXT("te")},	// Telugu
+			{TEXT("tgk"), TEXT("tgk"), TEXT("tg")},	// Tajik
+			{TEXT("tgl"), TEXT("tgl"), TEXT("tl")},	// Tagalog
+			{TEXT("tha"), TEXT("tha"), TEXT("th")},	// Thai
+			{TEXT("tir"), TEXT("tir"), TEXT("ti")},	// Tigrinya
+			{TEXT("ton"), TEXT("ton"), TEXT("to")},	// Tonga (Tonga Islands)
+			{TEXT("tsn"), TEXT("tsn"), TEXT("tn")},	// Tswana
+			{TEXT("tso"), TEXT("tso"), TEXT("ts")},	// Tsonga
+			{TEXT("tuk"), TEXT("tuk"), TEXT("tk")},	// Turkmen
+			{TEXT("tur"), TEXT("tur"), TEXT("tr")},	// Turkish
+			{TEXT("twi"), TEXT("twi"), TEXT("tw")},	// Twi
+			{TEXT("uig"), TEXT("uig"), TEXT("ug")},	// Uighur; Uyghur
+			{TEXT("ukr"), TEXT("ukr"), TEXT("uk")},	// Ukrainian
+			{TEXT("urd"), TEXT("urd"), TEXT("ur")},	// Urdu
+			{TEXT("uzb"), TEXT("uzb"), TEXT("uz")},	// Uzbek
+			{TEXT("ven"), TEXT("ven"), TEXT("ve")},	// Venda
+			{TEXT("vie"), TEXT("vie"), TEXT("vi")},	// Vietnamese
+			{TEXT("vol"), TEXT("vol"), TEXT("vo")},	// Volapük
+			{TEXT("wln"), TEXT("wln"), TEXT("wa")},	// Walloon
+			{TEXT("wol"), TEXT("wol"), TEXT("wo")},	// Wolof
+			{TEXT("xho"), TEXT("xho"), TEXT("xh")},	// Xhosa
+			{TEXT("yid"), TEXT("yid"), TEXT("yi")},	// Yiddish
+			{TEXT("yor"), TEXT("yor"), TEXT("yo")},	// Yoruba
+			{TEXT("zha"), TEXT("zha"), TEXT("za")},	// Zhuang; Chuang
+			{TEXT("chi"), TEXT("zho"), TEXT("zh")},	// Chinese
+			{TEXT("zul"), TEXT("zul"), TEXT("zu")}	// Zulu
+		};
+		
+		// Undefined
+		static const TCHAR* const CodeUnd(TEXT("und"));
+
+		const TCHAR* Get639_1(const FString& InFrom639_1_2_3)
+		{
+			// Already a 2 letter code?
+			if (InFrom639_1_2_3.Len() == 2)
+			{
+				// Look it up in the table if it exists.
+				int32 Index =  CodeMap.IndexOfByPredicate([&](const FCodeEntry& Entry) { return InFrom639_1_2_3.Equals(Entry._1); });
+				if (Index != INDEX_NONE)
+				{
+					return CodeMap[Index]._1;
+				}
+			}
+			else if (InFrom639_1_2_3.Len() == 3)
+			{
+				// Three letter code, either ISO-639-3 or ISO-639-2B. Look it up in either table.
+				int32 Index =  CodeMap.IndexOfByPredicate([&](const FCodeEntry& Entry) { return InFrom639_1_2_3.Equals(Entry._3); });
+				if (Index == INDEX_NONE)
+				{
+					Index =  CodeMap.IndexOfByPredicate([&](const FCodeEntry& Entry) { return InFrom639_1_2_3.Equals(Entry._2B); });
+				}
+				if (Index != INDEX_NONE)
+				{
+					return CodeMap[Index]._1;
+				}
+				else
+				{
+					/*
+						The standard includes some codes for special situations:
+							mis, for "uncoded languages"
+							mul, for "multiple languages"
+							qaa-qtz, a range reserved for local use
+							und, for "undetermined"
+							zxx, for "no linguistic content; not applicable"
+					*/
+					if (InFrom639_1_2_3.Equals(TEXT("mis")) || InFrom639_1_2_3.Equals(TEXT("zxx")) || InFrom639_1_2_3.Equals(TEXT("mul")))
+					{
+						// Map those to "und", even though this is not a 639-1 code!
+						return CodeUnd;
+					}
+				}
+			}
+			return nullptr;
+		}
+
+
+		FString MapTo639_1(const FString& InFrom639_1_2_3)
+		{
+			const TCHAR* Mapped = Get639_1(InFrom639_1_2_3);
+			return Mapped ? Mapped : InFrom639_1_2_3;
+		}
+
+
+		FString RFC5646To639_1(const FString& InFromRFC5646)
+		{
+			FString s(InFromRFC5646);
+			// We only look at the primary language tag before the first '-',
+			// which by definition happens to be the "shortest ISO 639 code".
+			int32 DashPos = INDEX_NONE;
+			if (s.FindChar(TCHAR('-'), DashPos))
+			{
+				s.LeftInline(DashPos);
+			}
+			return MapTo639_1(s);
+		}
+
+	}
+}
diff --git a/Engine/Plugins/Media/ElectraPlayer/Source/ElectraPlayerRuntime/Private/Runtime/Utilities/ISO639-Map.h b/Engine/Plugins/Media/ElectraPlayer/Source/ElectraPlayerRuntime/Private/Runtime/Utilities/ISO639-Map.h
new file mode 100644
--- /dev/null
+++ b/Engine/Plugins/Media/ElectraPlayer/Source/ElectraPlayerRuntime/Private/Runtime/Utilities/ISO639-Map.h
@@ -0,0 +1,21 @@
+// Copyright Epic Games, Inc. All Rights Reserved.
+
+#pragma once
+
+#include <CoreMinimal.h>
+
+
+
+namespace Electra
+{
+
+	namespace ISO639
+	{
+		const TCHAR* Get639_1(const FString& InFrom639_1_2_3);
+		FString MapTo639_1(const FString& InFrom639_1_2_3);
+
+		FString RFC5646To639_1(const FString& InFromRFC5646);
+	}
+
+}
+
diff --git a/Engine/Plugins/Media/ElectraPlayer/Source/ElectraPlayerRuntime/Public/IElectraPlayerInterface.h b/Engine/Plugins/Media/ElectraPlayer/Source/ElectraPlayerRuntime/Public/IElectraPlayerInterface.h
--- a/Engine/Plugins/Media/ElectraPlayer/Source/ElectraPlayerRuntime/Public/IElectraPlayerInterface.h
+++ b/Engine/Plugins/Media/ElectraPlayer/Source/ElectraPlayerRuntime/Public/IElectraPlayerInterface.h
@@ -97,10 +97,23 @@ public:
 
 	// -------- PlayerAdapter (Plugin/Native) API
 
+	struct FStreamSelectionAttributes
+	{
+		TOptional<FString> Kind;
+		TOptional<FString> Language_ISO639;
+		TOptional<int32> TrackIndexOverride;
+		void Reset()
+		{
+			Kind.Reset();
+			Language_ISO639.Reset();
+			TrackIndexOverride.Reset();
+		}
+	};
+
 	struct FPlaystartOptions
 	{
 		TOptional<FTimespan>		TimeOffset;
-		TOptional<int32>			AudioTrackIndex;
+		FStreamSelectionAttributes	InitialAudioTrackAttributes;
 		TOptional<int32>			MaxVerticalStreamResolution;
 		TOptional<int32>			MaxBandwidthForStreaming;
 	};
diff --git a/Engine/Plugins/Media/ElectraUtil/Source/ElectraBase/Public/ParameterDictionary.h b/Engine/Plugins/Media/ElectraUtil/Source/ElectraBase/Public/ParameterDictionary.h
--- a/Engine/Plugins/Media/ElectraUtil/Source/ElectraBase/Public/ParameterDictionary.h
+++ b/Engine/Plugins/Media/ElectraUtil/Source/ElectraBase/Public/ParameterDictionary.h
@@ -155,6 +155,10 @@ public:
 	void Set(const FString& Key, const FVariantValue& Value);
 	bool HaveKey(const FString& Key) const;
 	const FVariantValue& GetValue(const FString& Key) const;
+	void Remove(const FString& Key)
+	{ Dictionary.Remove(Key); }
+	void SetOrUpdate(const FString& Key, const FVariantValue& Value)
+	{ Dictionary.Add(Key, Value); }
 
 	void GetKeysStartingWith(const FString& StartsWith, TArray<FString>& Keys);
 private:
