ElectraPlayer: added option to not preload the media to allow selecting the default tracks that might be different from what is being preloaded; checking if audio stream switches are necessary to perform.

#rnx

diff --git a/Engine/Plugins/Media/ElectraPlayer/Source/ElectraPlayerPlugin/Private/ElectraPlayerPlugin.cpp b/Engine/Plugins/Media/ElectraPlayer/Source/ElectraPlayerPlugin/Private/ElectraPlayerPlugin.cpp
--- a/Engine/Plugins/Media/ElectraPlayer/Source/ElectraPlayerPlugin/Private/ElectraPlayerPlugin.cpp
+++ b/Engine/Plugins/Media/ElectraPlayer/Source/ElectraPlayerPlugin/Private/ElectraPlayerPlugin.cpp
@@ -334,6 +334,13 @@ bool FElectraPlayerPlugin::Open(const FString& Url, const IMediaOptions* Options
 		LocalPlaystartOptions.InitialAudioTrackAttributes.Language_ISO639 = InitialAudioLanguage;
 		UE_LOG(LogElectraPlayerPlugin, Log, TEXT("[%p] IMediaPlayer::Open: Asking for initial audio language \"%s\""), this, *InitialAudioLanguage);
 	}
+	bool bNoPreloading = Options->GetMediaOption(TEXT("ElectraNoPreloading"), (bool)false);
+	if (bNoPreloading)
+	{
+		LocalPlaystartOptions.bDoNotPreload = true;
+		UE_LOG(LogElectraPlayerPlugin, Log, TEXT("[%p] IMediaPlayer::Open: No preloading after opening media"), this);
+	}
+
 
 	// Set up options to initialize the internal player with.
 	Electra::FParamDict PlayerOptions;
diff --git a/Engine/Plugins/Media/ElectraPlayer/Source/ElectraPlayerRuntime/Private/ElectraPlayer.cpp b/Engine/Plugins/Media/ElectraPlayer/Source/ElectraPlayerRuntime/Private/ElectraPlayer.cpp
--- a/Engine/Plugins/Media/ElectraPlayer/Source/ElectraPlayerRuntime/Private/ElectraPlayer.cpp
+++ b/Engine/Plugins/Media/ElectraPlayer/Source/ElectraPlayerRuntime/Private/ElectraPlayer.cpp
@@ -188,6 +188,7 @@ void FElectraPlayer::ClearToDefaultState()
 	SelectedVideoTrackIndex = -1;
 	SelectedAudioTrackIndex = -1;
 	bAudioTrackIndexDirty = true;
+	bInitialSeekPerformed = false;
 	LastPresentedFrameDimension = FIntPoint::ZeroValue;
 	DeferredPlayerEvents.Empty();
 	MediaUrl.Empty();
@@ -198,7 +199,7 @@ void FElectraPlayer::ClearToDefaultState()
 /**
  *	Open player
  */
-bool FElectraPlayer::OpenInternal(const FString& Url, const FParamDict & PlayerOptions, const FPlaystartOptions & InPlaystartOptions)
+bool FElectraPlayer::OpenInternal(const FString& Url, const FParamDict& PlayerOptions, const FPlaystartOptions& InPlaystartOptions)
 {
 	LLM_SCOPE(ELLMTag::ElectraPlayer);
 	CSV_EVENT(ElectraPlayer, TEXT("Open"));
@@ -796,6 +797,7 @@ bool FElectraPlayer::SetRate(float Rate)
 		{
 			if (CurrentPlayer->AdaptivePlayer->IsPaused() || !CurrentPlayer->AdaptivePlayer->IsPlaying())
 			{
+				TriggerFirstSeekIfNecessary();
 				CurrentPlayer->AdaptivePlayer->Resume();
 			}
 		}
@@ -804,6 +806,52 @@ bool FElectraPlayer::SetRate(float Rate)
 	return false;
 }
 
+void FElectraPlayer::TriggerFirstSeekIfNecessary()
+{
+	if (!bInitialSeekPerformed)
+	{
+		bInitialSeekPerformed = true;
+
+		// Set up the initial playback position
+		IAdaptiveStreamingPlayer::FSeekParam playParam;
+		// First we look at any potential time offset specified in the playstart options.
+		if (PlaystartOptions.TimeOffset.IsSet())
+		{
+			FTimespan Target;
+			CalculateTargetSeekTime(Target, PlaystartOptions.TimeOffset.GetValue());
+			playParam.Time.SetFromHNS(Target.GetTicks());
+		}
+		else
+		{
+			// Do not set a start time, let the player pick one.
+			//playParam.Time.SetToZero();
+		}
+
+		// Next, give a list of the seekable positions to the delegate and ask it if it wants to seek to one of them,
+		// overriding any potential time offset from above.
+		if (TSharedPtr<IElectraPlayerAdapterDelegate, ESPMode::ThreadSafe> PinnedAdapterDelegate = AdapterDelegate.Pin())
+		{
+			TSharedPtr<TArray<FTimespan>, ESPMode::ThreadSafe> SeekablePositions = MakeShared<TArray<FTimespan>, ESPMode::ThreadSafe>();
+			if (SeekablePositions.IsValid())
+			{
+				CurrentPlayer->AdaptivePlayer->GetSeekablePositions(*SeekablePositions);
+
+				// Check with the delegate if it wants to start somewhere else.
+				FVariantValue Result = PinnedAdapterDelegate->QueryOptions(IElectraPlayerAdapterDelegate::EOptionType::PlaystartPosFromSeekPositions, FVariantValue(SeekablePositions));
+				if (Result.IsValid())
+				{
+					check(Result.IsType(FVariantValue::EDataType::TypeInt64));
+					playParam.Time.SetFromHNS(Result.GetInt64());
+				}
+			}
+		}
+
+		// Trigger buffering at the intended start time.
+		CurrentPlayer->AdaptivePlayer->SeekTo(playParam);
+	}
+}
+
+
 void FElectraPlayer::CalculateTargetSeekTime(FTimespan& OutTargetTime, const FTimespan& InTime)
 {
 	if (CurrentPlayer.Get())
@@ -865,6 +913,7 @@ bool FElectraPlayer::Seek(const FTimespan& Time)
 		CalculateTargetSeekTime(Target, Time);
 		Electra::IAdaptiveStreamingPlayer::FSeekParam seek;
 		seek.Time.SetFromHNS(Target.GetTicks());
+		bInitialSeekPerformed = true;
 		CurrentPlayer->AdaptivePlayer->SeekTo(seek);
 		return true;
 	}
@@ -980,8 +1029,6 @@ int32 FElectraPlayer::GetSelectedTrack(EPlayerTrackType TrackType) const
 		*/
 		if (bAudioTrackIndexDirty)
 		{
-			bAudioTrackIndexDirty = false;
-
 			if (NumTracksAudio == 0)
 			{
 				SelectedAudioTrackIndex = -1;
@@ -994,6 +1041,7 @@ int32 FElectraPlayer::GetSelectedTrack(EPlayerTrackType TrackType) const
 					if (CurrentPlayer->AdaptivePlayer->IsTrackDeselected(Electra::EStreamType::Audio))
 					{
 						SelectedAudioTrackIndex = -1;
+						bAudioTrackIndexDirty = false;
 					}
 					else
 					{
@@ -1002,6 +1050,7 @@ int32 FElectraPlayer::GetSelectedTrack(EPlayerTrackType TrackType) const
 						if (Attributes.OverrideIndex.IsSet())
 						{
 							SelectedAudioTrackIndex = Attributes.OverrideIndex.GetValue();
+							bAudioTrackIndexDirty = false;
 						}
 					}
 				}
@@ -1475,43 +1524,11 @@ void FElectraPlayer::HandlePlayerEventReceivedPlaylists()
 	InitialAudioAttributes.OverrideIndex = PlaystartOptions.InitialAudioTrackAttributes.TrackIndexOverride;
 	CurrentPlayer->AdaptivePlayer->SetInitialStreamAttributes(Electra::EStreamType::Audio, InitialAudioAttributes);
 
-
-	// Set up the initial playback position
-	IAdaptiveStreamingPlayer::FSeekParam playParam;
-	// First we look at any potential time offset specified in the playstart options.
-	if (PlaystartOptions.TimeOffset.IsSet())
+	// Trigger preloading unless forbidden.
+	if (PlaystartOptions.bDoNotPreload == false)
 	{
-		FTimespan Target;
-		CalculateTargetSeekTime(Target, PlaystartOptions.TimeOffset.GetValue());
-		playParam.Time.SetFromHNS(Target.GetTicks());
+		TriggerFirstSeekIfNecessary();
 	}
-	else
-	{
-		// Do not set a start time, let the player pick one.
-		//playParam.Time.SetToZero();
-	}
-
-	// Next, give a list of the seekable positions to the delegate and ask it if it wants to seek to one of them,
-	// overriding any potential time offset from above.
-	if (TSharedPtr<IElectraPlayerAdapterDelegate, ESPMode::ThreadSafe> PinnedAdapterDelegate = AdapterDelegate.Pin())
-	{
-		TSharedPtr<TArray<FTimespan>, ESPMode::ThreadSafe> SeekablePositions = MakeShared<TArray<FTimespan>, ESPMode::ThreadSafe>();
-		if (SeekablePositions.IsValid())
-		{
-			CurrentPlayer->AdaptivePlayer->GetSeekablePositions(*SeekablePositions);
-
-			// Check with the delegate if it wants to start somewhere else.
-			FVariantValue Result = PinnedAdapterDelegate->QueryOptions(IElectraPlayerAdapterDelegate::EOptionType::PlaystartPosFromSeekPositions, FVariantValue(SeekablePositions));
-			if (Result.IsValid())
-			{
-				check(Result.IsType(FVariantValue::EDataType::TypeInt64));
-				playParam.Time.SetFromHNS(Result.GetInt64());
-			}
-		}
-	}
-
-	// Trigger buffering at the intended start time.
-	CurrentPlayer->AdaptivePlayer->SeekTo(playParam);
 }
 
 void FElectraPlayer::HandlePlayerEventTracksChanged()
diff --git a/Engine/Plugins/Media/ElectraPlayer/Source/ElectraPlayerRuntime/Private/ElectraPlayer.h b/Engine/Plugins/Media/ElectraPlayer/Source/ElectraPlayerRuntime/Private/ElectraPlayer.h
--- a/Engine/Plugins/Media/ElectraPlayer/Source/ElectraPlayerRuntime/Private/ElectraPlayer.h
+++ b/Engine/Plugins/Media/ElectraPlayer/Source/ElectraPlayerRuntime/Private/ElectraPlayer.h
@@ -317,6 +317,7 @@ private:
 	bool MediaStateOnPause();
 	void MediaStateOnEndReached();
 	void MediaStateOnSeekFinished();
+	void TriggerFirstSeekIfNecessary();
 
 	TSharedPtr<FTrackMetadata, ESPMode::ThreadSafe> GetTrackStreamMetadata(EPlayerTrackType TrackType, int32 TrackIndex) const;
 
@@ -331,6 +332,8 @@ private:
 	mutable int32									SelectedAudioTrackIndex;
 	mutable bool									bAudioTrackIndexDirty;
 
+	bool											bInitialSeekPerformed;
+
 	FIntPoint										LastPresentedFrameDimension;
 
 
diff --git a/Engine/Plugins/Media/ElectraPlayer/Source/ElectraPlayerRuntime/Private/Runtime/Player/AdaptiveStreamingPlayer.cpp b/Engine/Plugins/Media/ElectraPlayer/Source/ElectraPlayerRuntime/Private/Runtime/Player/AdaptiveStreamingPlayer.cpp
--- a/Engine/Plugins/Media/ElectraPlayer/Source/ElectraPlayerRuntime/Private/Runtime/Player/AdaptiveStreamingPlayer.cpp
+++ b/Engine/Plugins/Media/ElectraPlayer/Source/ElectraPlayerRuntime/Private/Runtime/Player/AdaptiveStreamingPlayer.cpp
@@ -2828,24 +2828,27 @@ void FAdaptiveStreamingPlayer::InternalHandleSegmentTrackChanges(const FTimeValu
 
 		if (PendingTrackSelectionAud.IsValid() && CurrentPlayPeriodAudio.IsValid())
 		{
-			IManifest::IPlayPeriod::ETrackChangeResult TrackChangeResult = CurrentPlayPeriodAudio->ChangeTrackStreamPreference(EStreamType::Audio, *PendingTrackSelectionAud);
-			if (TrackChangeResult == IManifest::IPlayPeriod::ETrackChangeResult::NewPeriodNeeded)
+			if (!SelectedStreamAttributesAud.IsCompatibleWith(*PendingTrackSelectionAud))
 			{
-				// Make sure any finishing current requests will be ignored by increasing the sequence ID.
-				++CurrentPlaybackSequenceID[StreamTypeToArrayIndex(EStreamType::Audio)];
-				// Cancel any ongoing segment download now.
-				if (StreamReaderHandler)
+				IManifest::IPlayPeriod::ETrackChangeResult TrackChangeResult = CurrentPlayPeriodAudio->ChangeTrackStreamPreference(EStreamType::Audio, *PendingTrackSelectionAud);
+				if (TrackChangeResult == IManifest::IPlayPeriod::ETrackChangeResult::NewPeriodNeeded)
 				{
-					StreamReaderHandler->CancelRequest(EStreamType::Audio, true);
-				}
+					// Make sure any finishing current requests will be ignored by increasing the sequence ID.
+					++CurrentPlaybackSequenceID[StreamTypeToArrayIndex(EStreamType::Audio)];
+					// Cancel any ongoing segment download now.
+					if (StreamReaderHandler)
+					{
+						StreamReaderHandler->CancelRequest(EStreamType::Audio, true);
+					}
 
-				StreamSelectionAttributesAud = *PendingTrackSelectionAud;
+					StreamSelectionAttributesAud = *PendingTrackSelectionAud;
 
-				FPendingSegmentRequest NextReq;
-				NextReq.bStartOver = true;
-				NextReq.StartoverPosition.Time = PlaybackState.GetPlayPosition();
-				NextReq.StreamType = EStreamType::Audio;
-				NextPendingSegmentRequests.Enqueue(MoveTemp(NextReq));
+					FPendingSegmentRequest NextReq;
+					NextReq.bStartOver = true;
+					NextReq.StartoverPosition.Time = PlaybackState.GetPlayPosition();
+					NextReq.StreamType = EStreamType::Audio;
+					NextPendingSegmentRequests.Enqueue(MoveTemp(NextReq));
+				}
 			}
 			PendingTrackSelectionAud.Reset();
 		}
diff --git a/Engine/Plugins/Media/ElectraPlayer/Source/ElectraPlayerRuntime/Private/Runtime/StreamAccessUnitBuffer.cpp b/Engine/Plugins/Media/ElectraPlayer/Source/ElectraPlayerRuntime/Private/Runtime/StreamAccessUnitBuffer.cpp
--- a/Engine/Plugins/Media/ElectraPlayer/Source/ElectraPlayerRuntime/Private/Runtime/StreamAccessUnitBuffer.cpp
+++ b/Engine/Plugins/Media/ElectraPlayer/Source/ElectraPlayerRuntime/Private/Runtime/StreamAccessUnitBuffer.cpp
@@ -366,7 +366,7 @@ namespace Electra
 
 		if (!ActiveOutputBufferInfo.IsValid())
 		{
-			check(UpcomingBufferChain.IsEmpty());
+			check(UpcomingBufferChain.Num() == 0);
 			if (SwitchOverBufferChain.Num())
 			{
 				ChangeOver();
diff --git a/Engine/Plugins/Media/ElectraPlayer/Source/ElectraPlayerRuntime/Private/Runtime/StreamTypes.h b/Engine/Plugins/Media/ElectraPlayer/Source/ElectraPlayerRuntime/Private/Runtime/StreamTypes.h
--- a/Engine/Plugins/Media/ElectraPlayer/Source/ElectraPlayerRuntime/Private/Runtime/StreamTypes.h
+++ b/Engine/Plugins/Media/ElectraPlayer/Source/ElectraPlayerRuntime/Private/Runtime/StreamTypes.h
@@ -575,6 +575,21 @@ namespace Electra
 		// of the tracks as they are found. If the index is invalid the selection rules for kind and language are applied.
 		TOptional<int32> OverrideIndex;
 
+		bool IsCompatibleWith(const FStreamSelectionAttributes& Other)
+		{
+			FString Kind1 = Kind.IsSet() ? Kind.GetValue() : FString();
+			FString Kind2 = Other.Kind.IsSet() ? Other.Kind.GetValue() : FString();
+
+			FString Lang1 = Language_ISO639.IsSet() ? Language_ISO639.GetValue() : FString();
+			FString Lang2 = Other.Language_ISO639.IsSet() ? Other.Language_ISO639.GetValue() : FString();
+
+			if (Kind1.IsEmpty() || Kind2.IsEmpty() || Kind1.Equals(Kind2))
+			{
+				return Lang1.Equals(Lang2);
+			}
+			return false;
+		}
+
 		void Reset()
 		{
 			Kind.Reset();
diff --git a/Engine/Plugins/Media/ElectraPlayer/Source/ElectraPlayerRuntime/Public/IElectraPlayerInterface.h b/Engine/Plugins/Media/ElectraPlayer/Source/ElectraPlayerRuntime/Public/IElectraPlayerInterface.h
--- a/Engine/Plugins/Media/ElectraPlayer/Source/ElectraPlayerRuntime/Public/IElectraPlayerInterface.h
+++ b/Engine/Plugins/Media/ElectraPlayer/Source/ElectraPlayerRuntime/Public/IElectraPlayerInterface.h
@@ -116,6 +116,7 @@ public:
 		FStreamSelectionAttributes	InitialAudioTrackAttributes;
 		TOptional<int32>			MaxVerticalStreamResolution;
 		TOptional<int32>			MaxBandwidthForStreaming;
+		bool						bDoNotPreload = false;
 	};
 
 	virtual bool OpenInternal(const FString& Url, const Electra::FParamDict& PlayerOptions, const FPlaystartOptions& InPlaystartOptions) = 0;
