ElectraPlayer: fixed selected audio track index for HLS playback

#rnx

diff --git a/Engine/Plugins/Media/ElectraPlayer/Source/ElectraPlayerRuntime/Private/Runtime/Player/AdaptiveStreamingPlayer.cpp b/Engine/Plugins/Media/ElectraPlayer/Source/ElectraPlayerRuntime/Private/Runtime/Player/AdaptiveStreamingPlayer.cpp
--- a/Engine/Plugins/Media/ElectraPlayer/Source/ElectraPlayerRuntime/Private/Runtime/Player/AdaptiveStreamingPlayer.cpp
+++ b/Engine/Plugins/Media/ElectraPlayer/Source/ElectraPlayerRuntime/Private/Runtime/Player/AdaptiveStreamingPlayer.cpp
@@ -2819,6 +2819,20 @@ void FAdaptiveStreamingPlayer::InternalHandleSegmentTrackChanges(const FTimeValu
 		PendingTrackSelectionVid.Reset();
 		PendingTrackSelectionTxt.Reset();
 	}
+	else if (ManifestType == EMediaFormatType::HLS)
+	{
+		if (PendingTrackSelectionAud.IsValid() && CurrentPlayPeriodAudio.IsValid())
+		{
+			TSharedPtrTS<FBufferSourceInfo> BufferSourceInfo = CurrentPlayPeriodAudio->GetSelectedStreamBufferSourceInfo(EStreamType::Audio);
+			MultiStreamBufferAud.SelectTrackWhenAvailable(BufferSourceInfo);
+			StreamSelectionAttributesAud = *PendingTrackSelectionAud;
+			SelectedStreamAttributesAud.UpdateWith(BufferSourceInfo->Kind, BufferSourceInfo->Language, BufferSourceInfo->HardIndex);
+			PendingTrackSelectionAud.Reset();
+		}
+		// Ignore video and subtitle changes for now.
+		PendingTrackSelectionVid.Reset();
+		PendingTrackSelectionTxt.Reset();
+	}
 	else if (ManifestType == EMediaFormatType::DASH)
 	{
 		if (PendingStartRequest.IsValid() || PendingFirstSegmentRequest.IsValid())
diff --git a/Engine/Plugins/Media/ElectraPlayer/Source/ElectraPlayerRuntime/Private/Runtime/Player/HLS/ManifestHLS.cpp b/Engine/Plugins/Media/ElectraPlayer/Source/ElectraPlayerRuntime/Private/Runtime/Player/HLS/ManifestHLS.cpp
--- a/Engine/Plugins/Media/ElectraPlayer/Source/ElectraPlayerRuntime/Private/Runtime/Player/HLS/ManifestHLS.cpp
+++ b/Engine/Plugins/Media/ElectraPlayer/Source/ElectraPlayerRuntime/Private/Runtime/Player/HLS/ManifestHLS.cpp
@@ -410,7 +410,9 @@ void FPlayPeriodHLS::PrepareForPlay()
 	CurrentSourceBufferInfoVideo = MakeSharedTS<FBufferSourceInfo>();
 	CurrentSourceBufferInfoAudio = MakeSharedTS<FBufferSourceInfo>();
 	CurrentSourceBufferInfoVideo->PeriodAdaptationSetID = TEXT("video.0");
+	CurrentSourceBufferInfoVideo->HardIndex = 0;
 	CurrentSourceBufferInfoAudio->PeriodAdaptationSetID = TEXT("audio.0");
+	CurrentSourceBufferInfoAudio->HardIndex = 0;
 
 	CurrentReadyState = IManifest::IPlayPeriod::EReadyState::IsReady;
 }
