ElectraPlayer: improved detection of end-of-stream in case an elementary stream (usually audio) ends early and the facade stops consuming samples altogether, preventing the other elementary stream from reaching its end.

#rnx

diff --git a/Engine/Plugins/Media/ElectraPlayer/Source/ElectraPlayerRuntime/Private/Runtime/Player/AdaptiveStreamingPlayer.cpp b/Engine/Plugins/Media/ElectraPlayer/Source/ElectraPlayerRuntime/Private/Runtime/Player/AdaptiveStreamingPlayer.cpp
--- a/Engine/Plugins/Media/ElectraPlayer/Source/ElectraPlayerRuntime/Private/Runtime/Player/AdaptiveStreamingPlayer.cpp
+++ b/Engine/Plugins/Media/ElectraPlayer/Source/ElectraPlayerRuntime/Private/Runtime/Player/AdaptiveStreamingPlayer.cpp
@@ -184,10 +184,14 @@ FErrorDetail FAdaptiveStreamingPlayer::GetError() const
  */
 void FAdaptiveStreamingPlayer::UpdateDiagnostics()
 {
+	int64 tNow = MEDIAutcTime::CurrentMSec();
 	DiagnosticsCriticalSection.Lock();
 	MultiStreamBufferVid.GetStats(VideoBufferStats.StreamBuffer);
 	MultiStreamBufferAud.GetStats(AudioBufferStats.StreamBuffer);
 	MultiStreamBufferTxt.GetStats(TextBufferStats.StreamBuffer);
+	VideoBufferStats.UpdateStalledDuration(tNow);
+	AudioBufferStats.UpdateStalledDuration(tNow);
+	TextBufferStats.UpdateStalledDuration(tNow);
 	DiagnosticsCriticalSection.Unlock();
 }
 
@@ -1444,26 +1448,40 @@ void FAdaptiveStreamingPlayer::WorkerThreadFN()
 		{
 			CSV_SCOPED_TIMING_STAT(ElectraPlayer, AdaptiveStreamingPlayer_Worker);
 
-			if (RenderClock->IsRunning())
+			// Update the play position we return to the interested caller.
+			if (!PlaybackState.GetHasEnded())
 			{
-				FTimeValue playPos = RenderClock->GetInterpolatedRenderTime(IMediaRenderClock::ERendererType::Video);
-				PlaybackState.SetPlayPosition(playPos);
-				// When we reach the time at which the next loop occurred pop it off the queue and update the playback state loop state.
-				while(NextLoopStates.Num())
+				if (RenderClock->IsRunning())
 				{
-					FTimeValue loopAtTime = NextLoopStates.FrontRef().LoopBasetime;
-					if (playPos >= loopAtTime)
+					FTimeValue playPos = RenderClock->GetInterpolatedRenderTime(IMediaRenderClock::ERendererType::Video);
+					PlaybackState.SetPlayPosition(playPos);
+					// When we reach the time at which the next loop occurred pop it off the queue and update the playback state loop state.
+					while(NextLoopStates.Num())
 					{
-						FPlayerLoopState newLoopState = NextLoopStates.Pop();
-						PlaybackState.SetLoopState(newLoopState);
-						DispatchEvent(FMetricEvent::ReportJumpInPlayPosition(playPos, playPos, Metrics::ETimeJumpReason::Looping));
-					}
-					else
-					{
-						break;
+						FTimeValue loopAtTime = NextLoopStates.FrontRef().LoopBasetime;
+						if (playPos >= loopAtTime)
+						{
+							FPlayerLoopState newLoopState = NextLoopStates.Pop();
+							PlaybackState.SetLoopState(newLoopState);
+							DispatchEvent(FMetricEvent::ReportJumpInPlayPosition(playPos, playPos, Metrics::ETimeJumpReason::Looping));
+						}
+						else
+						{
+							break;
+						}
 					}
 				}
 			}
+			else
+			{
+				// When playback has ended we lock the position to the end of the timeline.
+				// This is only in case the application checks for the play position to reach the end of the timeline
+				// ie. calling GetPlayPosition() to compare against the end of GetTimelineRange()
+				// instead of using the dedicated HasEnded() method.
+				FTimeRange TimelineRange;
+				PlaybackState.GetTimelineRange(TimelineRange);
+				PlaybackState.SetPlayPosition(TimelineRange.End);
+			}
 
 			// Handle state changes to match the user request.
 			HandlePlayStateChanges();
@@ -3271,6 +3289,8 @@ void FAdaptiveStreamingPlayer::CheckForStreamEnd()
 				bool bEndVid = bHaveVideoReader.IsSet();
 				bool bEndAud = bHaveAudioReader.IsSet();
 				bool bEndTxt = bHaveTextReader.IsSet();
+				int64 VidStalled = 0;
+				int64 AudStalled = 0;
 
 				DiagnosticsCriticalSection.Lock();
 				FBufferStats vidStats = VideoBufferStats;
@@ -3282,6 +3302,7 @@ void FAdaptiveStreamingPlayer::CheckForStreamEnd()
 				{
 					// All buffers at end of data?
 					bEndVid = (vidStats.StreamBuffer.bEndOfData && vidStats.DecoderInputBuffer.bEODSignaled && vidStats.DecoderOutputBuffer.bEODreached);
+					VidStalled = vidStats.StreamBuffer.bEndOfData ? vidStats.GetStalledDurationMillisec() : 0;
 				}
 
 				// Check for end of audio stream
@@ -3289,12 +3310,24 @@ void FAdaptiveStreamingPlayer::CheckForStreamEnd()
 				{
 					// All buffers at end of data?
 					bEndAud = (audStats.StreamBuffer.bEndOfData && audStats.DecoderInputBuffer.bEODSignaled && audStats.DecoderOutputBuffer.bEODreached);
+					AudStalled = audStats.StreamBuffer.bEndOfData ? audStats.GetStalledDurationMillisec() : 0;
 				}
 
 				// Text stream
 				// ...
 
 
+				// If either primary stream has reliably ended do a check if the other may have stalled because the application is no longer
+				// consuming decoder output, which will prevent the other stream from ending.
+				if (bEndVid || bEndAud)
+				{
+					int64 OtherStallTime = bEndAud ? VidStalled : AudStalled;
+					if (OtherStallTime > 500)
+					{
+						bEndVid = bEndAud = true;
+					}
+				}
+
 				// Everything at EOD?
 				if (bEndVid && bEndAud && bEndTxt)
 				{
@@ -3342,6 +3375,10 @@ bool FAdaptiveStreamingPlayer::InternalStartAt(const FSeekParam& NewPosition)
 		MultiStreamBufferTxt.SetParallelTrackMode();
 	}
 
+	VideoBufferStats.Clear();
+	AudioBufferStats.Clear();
+	TextBufferStats.Clear();
+
 	// Update data availability states in case this wasn't done yet.
 	UpdateDataAvailabilityState(DataAvailabilityStateVid, Metrics::FDataAvailabilityChange::EAvailability::DataNotAvailable);
 	UpdateDataAvailabilityState(DataAvailabilityStateAud, Metrics::FDataAvailabilityChange::EAvailability::DataNotAvailable);
diff --git a/Engine/Plugins/Media/ElectraPlayer/Source/ElectraPlayerRuntime/Private/Runtime/Player/AdaptiveStreamingPlayerInternal.h b/Engine/Plugins/Media/ElectraPlayer/Source/ElectraPlayerRuntime/Private/Runtime/Player/AdaptiveStreamingPlayerInternal.h
--- a/Engine/Plugins/Media/ElectraPlayer/Source/ElectraPlayerRuntime/Private/Runtime/Player/AdaptiveStreamingPlayerInternal.h
+++ b/Engine/Plugins/Media/ElectraPlayer/Source/ElectraPlayerRuntime/Private/Runtime/Player/AdaptiveStreamingPlayerInternal.h
@@ -1187,15 +1187,66 @@ private:
 
 	struct FBufferStats
 	{
+		struct FStallMonitor
+		{
+			int64 DurationMillisec;
+			int64 PreviousCheckTime;
+			bool bPreviousState;
+			FStallMonitor()
+			{
+				Clear();
+			}
+			void Clear()
+			{
+				DurationMillisec = 0;
+				PreviousCheckTime = 0;
+				bPreviousState = false;
+			}
+			void Update(int64 tNowMillisec, bool bInCurrentStallState)
+			{
+				if (bInCurrentStallState)
+				{
+					if (!bPreviousState)
+					{
+						PreviousCheckTime = tNowMillisec;
+						DurationMillisec = 0;
+					}
+					else
+					{
+						DurationMillisec = tNowMillisec - PreviousCheckTime;
+					}
+				}
+				else
+				{
+					DurationMillisec = 0;
+				}
+				bPreviousState = bInCurrentStallState;
+			}
+			int64 GetStalledDurationMillisec() const
+			{
+				return DurationMillisec;
+			}
+		};
+
 		void Clear()
 		{
 			StreamBuffer.Clear();
 			DecoderInputBuffer.Clear();
 			DecoderOutputBuffer.Clear();
+			DecoderOutputStalledMonitor.Clear();
+		}
+		void UpdateStalledDuration(int64 tNowMillisec)
+		{
+			DecoderOutputStalledMonitor.Update(tNowMillisec, DecoderOutputBuffer.bOutputStalled);
+		}
+		int64 GetStalledDurationMillisec() const
+		{
+			return DecoderOutputStalledMonitor.GetStalledDurationMillisec();
 		}
 		FAccessUnitBufferInfo								StreamBuffer;
 		IAccessUnitBufferListener::FBufferStats				DecoderInputBuffer;
 		IDecoderOutputBufferListener::FDecodeReadyStats		DecoderOutputBuffer;
+		FStallMonitor										DecoderOutputStalledMonitor;
 	};
 
 	struct FPrerollVars
