ElectraPlayer: cleanup and some speculative code changes trying to hunt down a rare crash on player close

#rnx

diff --git a/Engine/Plugins/Media/ElectraPlayer/Source/ElectraPlayerPlugin/Private/ElectraPlayerPlugin.cpp b/Engine/Plugins/Media/ElectraPlayer/Source/ElectraPlayerPlugin/Private/ElectraPlayerPlugin.cpp
--- a/Engine/Plugins/Media/ElectraPlayer/Source/ElectraPlayerPlugin/Private/ElectraPlayerPlugin.cpp
+++ b/Engine/Plugins/Media/ElectraPlayer/Source/ElectraPlayerPlugin/Private/ElectraPlayerPlugin.cpp
@@ -97,7 +97,8 @@ FElectraPlayerPlugin::~FElectraPlayerPlugin()
 
 Electra::FVariantValue FElectraPlayerPlugin::FPlayerAdapterDelegate::QueryOptions(EOptionType Type, const Electra::FVariantValue & Param)
 {
-	if (TSharedPtr<FElectraPlayerPlugin, ESPMode::ThreadSafe> PinnedHost = Host.Pin())
+	TSharedPtr<FElectraPlayerPlugin, ESPMode::ThreadSafe> PinnedHost = Host.Pin();
+	if (PinnedHost.IsValid())
 	{
 		FScopeLock lock(&PinnedHost->CallbackPointerLock);
 		if (PinnedHost->OptionInterface)
@@ -106,17 +107,19 @@ Electra::FVariantValue FElectraPlayerPlugin::FPlayerAdapterDelegate::QueryOption
 			{
 				case EOptionType::MaxVerticalStreamResolution:
 				{
-					return FVariantValue((int64)PinnedHost->OptionInterface->GetMediaOption(TEXT("MaxResolutionForMediaStreaming"), (int64)0));
+					static const FName MaxResolutionOptionKey = TEXT("MaxResolutionForMediaStreaming");
+					return FVariantValue((int64)PinnedHost->OptionInterface->GetMediaOption(MaxResolutionOptionKey, (int64)0));
 				}
 
 				case EOptionType::MaxBandwidthForStreaming:
 				{
-					return FVariantValue((int64)PinnedHost->OptionInterface->GetMediaOption(TEXT("ElectraMaxStreamingBandwidth"), (int64)0));
+					static const FName MaxBandwidthOptionKey = TEXT("ElectraMaxStreamingBandwidth");
+					return FVariantValue((int64)PinnedHost->OptionInterface->GetMediaOption(MaxBandwidthOptionKey, (int64)0));
 				}
 
 				case EOptionType::PlayListData:
 				{
-					const FName PlaylistOptionKey = TEXT("ElectraGetPlaylistData");
+					static const FName PlaylistOptionKey = TEXT("ElectraGetPlaylistData");
 					if (PinnedHost->OptionInterface->HasMediaOption(PlaylistOptionKey))
 					{
 						check(Param.IsType(FVariantValue::EDataType::TypeFString));
@@ -127,7 +130,7 @@ Electra::FVariantValue FElectraPlayerPlugin::FPlayerAdapterDelegate::QueryOption
 
 				case EOptionType::LicenseKeyData:
 				{
-					const FName LicenseKeyDataOptionKey = TEXT("ElectraGetLicenseKeyData");
+					static const FName LicenseKeyDataOptionKey = TEXT("ElectraGetLicenseKeyData");
 					if (PinnedHost->OptionInterface->HasMediaOption(LicenseKeyDataOptionKey))
 					{
 						check(Param.IsType(FVariantValue::EDataType::TypeFString));
@@ -138,7 +141,7 @@ Electra::FVariantValue FElectraPlayerPlugin::FPlayerAdapterDelegate::QueryOption
 
 				case EOptionType::PlaystartPosFromSeekPositions:
 				{
-					const FName PlaystartOptionKey = TEXT("ElectraGetPlaystartPosFromSeekPositions");
+					static const FName PlaystartOptionKey = TEXT("ElectraGetPlaystartPosFromSeekPositions");
 					if (PinnedHost->OptionInterface->HasMediaOption(PlaystartOptionKey))
 					{
 						check(Param.IsType(FVariantValue::EDataType::TypeSharedPointer));
@@ -170,7 +173,8 @@ Electra::FVariantValue FElectraPlayerPlugin::FPlayerAdapterDelegate::QueryOption
 
 void FElectraPlayerPlugin::FPlayerAdapterDelegate::SendMediaEvent(EPlayerEvent Event)
 {
-	if (TSharedPtr<FElectraPlayerPlugin, ESPMode::ThreadSafe> PinnedHost = Host.Pin())
+	TSharedPtr<FElectraPlayerPlugin, ESPMode::ThreadSafe> PinnedHost = Host.Pin();
+	if (PinnedHost.IsValid())
 	{
 		FScopeLock lock(&PinnedHost->CallbackPointerLock);
 		if (PinnedHost->EventSink)
@@ -183,31 +187,34 @@ void FElectraPlayerPlugin::FPlayerAdapterDelegate::SendMediaEvent(EPlayerEvent E
 
 void FElectraPlayerPlugin::FPlayerAdapterDelegate::OnVideoFlush()
 {
-	if (TSharedPtr<FElectraPlayerPlugin, ESPMode::ThreadSafe> PinnedHost = Host.Pin())
+	TSharedPtr<FElectraPlayerPlugin, ESPMode::ThreadSafe> PinnedHost = Host.Pin();
+	if (PinnedHost.IsValid())
 	{
 		TRange<FTimespan> AllTime(FTimespan::MinValue(), FTimespan::MaxValue());
 		TSharedPtr<IMediaTextureSample, ESPMode::ThreadSafe> FlushSample;
 		while (PinnedHost->GetSamples().FetchVideo(AllTime, FlushSample))
-			;
+		{ }
 	}
 }
 
 
 void FElectraPlayerPlugin::FPlayerAdapterDelegate::OnAudioFlush()
 {
-	if (TSharedPtr<FElectraPlayerPlugin, ESPMode::ThreadSafe> PinnedHost = Host.Pin())
+	TSharedPtr<FElectraPlayerPlugin, ESPMode::ThreadSafe> PinnedHost = Host.Pin();
+	if (PinnedHost.IsValid())
 	{
 		TRange<FTimespan> AllTime(FTimespan::MinValue(), FTimespan::MaxValue());
 		TSharedPtr<IMediaAudioSample, ESPMode::ThreadSafe> FlushSample;
 		while (PinnedHost->GetSamples().FetchAudio(AllTime, FlushSample))
-			;
+		{ }
 	}
 }
 
 
 void FElectraPlayerPlugin::FPlayerAdapterDelegate::PresentVideoFrame(const FVideoDecoderOutputPtr & InVideoFrame)
 {
-	if (TSharedPtr<FElectraPlayerPlugin, ESPMode::ThreadSafe> PinnedHost = Host.Pin())
+	TSharedPtr<FElectraPlayerPlugin, ESPMode::ThreadSafe> PinnedHost = Host.Pin();
+	if (PinnedHost.IsValid())
 	{
 		FElectraTextureSampleRef TextureSample = PinnedHost->OutputTexturePool.AcquireShared();
 		TextureSample->Initialize(InVideoFrame.Get());
@@ -218,7 +225,8 @@ void FElectraPlayerPlugin::FPlayerAdapterDelegate::PresentVideoFrame(const FVide
 
 void FElectraPlayerPlugin::FPlayerAdapterDelegate::PresentAudioFrame(const IAudioDecoderOutputPtr& InAudioFrame)
 {
-	if (TSharedPtr<FElectraPlayerPlugin, ESPMode::ThreadSafe> PinnedHost = Host.Pin())
+	TSharedPtr<FElectraPlayerPlugin, ESPMode::ThreadSafe> PinnedHost = Host.Pin();
+	if (PinnedHost.IsValid())
 	{
 		TSharedRef<FElectraPlayerAudioSample, ESPMode::ThreadSafe> AudioSample = PinnedHost->OutputAudioPool.AcquireShared();
 		AudioSample->Initialize(InAudioFrame);
@@ -228,7 +236,8 @@ void FElectraPlayerPlugin::FPlayerAdapterDelegate::PresentAudioFrame(const IAudi
 
 void FElectraPlayerPlugin::FPlayerAdapterDelegate::PresentMetadataSample(const IElectraBinarySampleRef& InMetadataSample)
 {
-	if (TSharedPtr<FElectraPlayerPlugin, ESPMode::ThreadSafe> PinnedHost = Host.Pin())
+	TSharedPtr<FElectraPlayerPlugin, ESPMode::ThreadSafe> PinnedHost = Host.Pin();
+	if (PinnedHost.IsValid())
 	{
 		PinnedHost->MediaSamples->AddMetadata(InMetadataSample);
 	}
@@ -237,7 +246,8 @@ void FElectraPlayerPlugin::FPlayerAdapterDelegate::PresentMetadataSample(const I
 
 bool FElectraPlayerPlugin::FPlayerAdapterDelegate::CanReceiveVideoSamples(int32 NumFrames)
 {
-	if (TSharedPtr<FElectraPlayerPlugin, ESPMode::ThreadSafe> PinnedHost = Host.Pin())
+	TSharedPtr<FElectraPlayerPlugin, ESPMode::ThreadSafe> PinnedHost = Host.Pin();
+	if (PinnedHost.IsValid())
 	{
 		return PinnedHost->MediaSamples->CanReceiveVideoSamples(NumFrames);
 	}
@@ -247,7 +257,8 @@ bool FElectraPlayerPlugin::FPlayerAdapterDelegate::CanReceiveVideoSamples(int32
 
 bool FElectraPlayerPlugin::FPlayerAdapterDelegate::CanReceiveAudioSamples(int32 NumFrames)
 {
-	if (TSharedPtr<FElectraPlayerPlugin, ESPMode::ThreadSafe> PinnedHost = Host.Pin())
+	TSharedPtr<FElectraPlayerPlugin, ESPMode::ThreadSafe> PinnedHost = Host.Pin();
+	if (PinnedHost.IsValid())
 	{
 		return PinnedHost->MediaSamples->CanReceiveAudioSamples(NumFrames);
 	}
@@ -256,7 +267,8 @@ bool FElectraPlayerPlugin::FPlayerAdapterDelegate::CanReceiveAudioSamples(int32
 
 void FElectraPlayerPlugin::FPlayerAdapterDelegate::PrepareForDecoderShutdown()
 {
-	if (TSharedPtr<FElectraPlayerPlugin, ESPMode::ThreadSafe> PinnedHost = Host.Pin())
+	TSharedPtr<FElectraPlayerPlugin, ESPMode::ThreadSafe> PinnedHost = Host.Pin();
+	if (PinnedHost.IsValid())
 	{
 		PinnedHost->OutputTexturePool.PrepareForDecoderShutdown();
 	}
@@ -271,7 +283,8 @@ FString FElectraPlayerPlugin::FPlayerAdapterDelegate::GetVideoAdapterName() cons
 
 IElectraPlayerResourceDelegate* FElectraPlayerPlugin::FPlayerAdapterDelegate::GetResourceDelegate() const
 {
-	if (TSharedPtr<FElectraPlayerPlugin, ESPMode::ThreadSafe> PinnedHost = Host.Pin())
+	TSharedPtr<FElectraPlayerPlugin, ESPMode::ThreadSafe> PinnedHost = Host.Pin();
+	if (PinnedHost.IsValid())
 	{
 		return PinnedHost->PlayerResourceDelegate.Get();
 	}
diff --git a/Engine/Plugins/Media/ElectraPlayer/Source/ElectraPlayerRuntime/Private/ElectraPlayer.cpp b/Engine/Plugins/Media/ElectraPlayer/Source/ElectraPlayerRuntime/Private/ElectraPlayer.cpp
--- a/Engine/Plugins/Media/ElectraPlayer/Source/ElectraPlayerRuntime/Private/ElectraPlayer.cpp
+++ b/Engine/Plugins/Media/ElectraPlayer/Source/ElectraPlayerRuntime/Private/ElectraPlayer.cpp
@@ -139,7 +139,7 @@ FElectraPlayer::FElectraPlayer(const TSharedPtr<IElectraPlayerAdapterDelegate, E
 	State = EPlayerState::Closed;
 	Status = EPlayerStatus::None;
 	bAllowKillAfterCloseEvent = false;
-	bCloseInProgress = false;
+	bPlayerHasClosed = false;
 	bHasPendingError = false;
 	AnalyticsInstanceEventCount = 0;
 	NumQueuedAnalyticEvents = 0;
@@ -210,7 +210,7 @@ bool FElectraPlayer::OpenInternal(const FString& Url, const FParamDict& PlayerOp
 	bAllowKillAfterCloseEvent = false;
 	State = EPlayerState::Closed;
 	Status = EPlayerStatus::None;
-	bCloseInProgress = false;
+	bPlayerHasClosed = false;
 	bHasPendingError = false;
 
 	// Start statistics with a clean slate.
@@ -265,7 +265,7 @@ bool FElectraPlayer::OpenInternal(const FString& Url, const FParamDict& PlayerOp
 	}
 
 	// Set the player member variable to the new player now that it has been fully initialized.
-	CurrentPlayer = NewPlayer;
+	CurrentPlayer = MoveTemp(NewPlayer);
 
 	// Issue load of the playlist.
 	CurrentPlayer->AdaptivePlayer->LoadManifest(MediaUrl);
@@ -282,12 +282,12 @@ void FElectraPlayer::CloseInternal(bool bKillAfterClose)
 	LLM_SCOPE(ELLMTag::ElectraPlayer);
 
 	PlayerLock.Lock();
-	if (bCloseInProgress || !CurrentPlayer.IsValid())
+	if (bPlayerHasClosed || !CurrentPlayer.IsValid())
 	{
 		PlayerLock.Unlock();
 		return;
 	}
-	bCloseInProgress = true;
+	bPlayerHasClosed = true;
 	WaitForPlayerDestroyedEvent->Reset();
 	PlayerLock.Unlock();
 
@@ -322,7 +322,8 @@ void FElectraPlayer::CloseInternal(bool bKillAfterClose)
 		Player->RendererAudio->DetachPlayer();
 	}
 
-	if (TSharedPtr<IElectraPlayerAdapterDelegate, ESPMode::ThreadSafe> PinnedAdapterDelegate = AdapterDelegate.Pin())
+	TSharedPtr<IElectraPlayerAdapterDelegate, ESPMode::ThreadSafe> PinnedAdapterDelegate = AdapterDelegate.Pin();
+	if (PinnedAdapterDelegate.IsValid())
 	{
 		PinnedAdapterDelegate->PrepareForDecoderShutdown();
 	}
@@ -367,7 +368,6 @@ void FElectraPlayer::CloseInternal(bool bKillAfterClose)
 
 	// Clear out the player instance now.
 	CurrentPlayer.Reset();
-	bCloseInProgress = false;
 
 	// Kick off asynchronous closing now.
 	FInternalPlayerImpl::DoCloseAsync(MoveTemp(Player), AsyncResourceReleaseNotification);
@@ -443,18 +443,19 @@ void FElectraPlayer::Tick(FTimespan DeltaTime, FTimespan Timecode)
 	// Handle the internal player, if we have one.
 	PlayerLock.Lock();
 	TSharedPtr<FInternalPlayerImpl, ESPMode::ThreadSafe> Player = CurrentPlayer;
-	if (!bCloseInProgress && CurrentPlayer.IsValid() && State != EPlayerState::Error)
+	if (!bPlayerHasClosed && Player.IsValid() && State != EPlayerState::Error)
 	{
-		if (CurrentPlayer->RendererVideo.IsValid())
+		if (Player->RendererVideo.IsValid())
 		{
-			CurrentPlayer->RendererVideo->TickInput(DeltaTime, Timecode);
+			Player->RendererVideo->TickInput(DeltaTime, Timecode);
 		}
 
 		// Handle static resource fetch requests.
 		StaticResourceProvider->ProcessPendingStaticResourceRequests();
 
 		// Check for option changes
-		if (TSharedPtr<IElectraPlayerAdapterDelegate, ESPMode::ThreadSafe> PinnedAdapterDelegate = AdapterDelegate.Pin())
+		TSharedPtr<IElectraPlayerAdapterDelegate, ESPMode::ThreadSafe> PinnedAdapterDelegate = AdapterDelegate.Pin();
+		if (PinnedAdapterDelegate.IsValid())
 		{
 			FVariantValue Value = PinnedAdapterDelegate->QueryOptions(IElectraPlayerAdapterDelegate::EOptionType::MaxVerticalStreamResolution);
 			if (Value.IsValid())
@@ -463,8 +464,8 @@ void FElectraPlayer::Tick(FTimespan DeltaTime, FTimespan Timecode)
 				if (NewVerticalStreamResolution != PlaystartOptions.MaxVerticalStreamResolution.Get(0))
 				{
 					PlaystartOptions.MaxVerticalStreamResolution = NewVerticalStreamResolution;
-					UE_LOG(LogElectraPlayer, Log, TEXT("[%p][%p] Limiting max vertical resolution to %d"), this, CurrentPlayer.Get(), (int32)NewVerticalStreamResolution);
-					CurrentPlayer->AdaptivePlayer->SetMaxResolution(0, (int32)NewVerticalStreamResolution);
+					UE_LOG(LogElectraPlayer, Log, TEXT("[%p][%p] Limiting max vertical resolution to %d"), this, Player.Get(), (int32)NewVerticalStreamResolution);
+					Player->AdaptivePlayer->SetMaxResolution(0, (int32)NewVerticalStreamResolution);
 				}
 			}
 			Value = PinnedAdapterDelegate->QueryOptions(IElectraPlayerAdapterDelegate::EOptionType::MaxBandwidthForStreaming);
@@ -474,8 +475,8 @@ void FElectraPlayer::Tick(FTimespan DeltaTime, FTimespan Timecode)
 				if (NewBandwidthForStreaming != PlaystartOptions.MaxBandwidthForStreaming.Get(0))
 				{
 					PlaystartOptions.MaxBandwidthForStreaming = NewBandwidthForStreaming;
-					UE_LOG(LogElectraPlayer, Log, TEXT("[%p][%p] Limiting max streaming bandwidth to %d bps"), this, CurrentPlayer.Get(), (int32)NewBandwidthForStreaming);
-					CurrentPlayer->AdaptivePlayer->SetBitrateCeiling((int32)NewBandwidthForStreaming);
+					UE_LOG(LogElectraPlayer, Log, TEXT("[%p][%p] Limiting max streaming bandwidth to %d bps"), this, Player.Get(), (int32)NewBandwidthForStreaming);
+					Player->AdaptivePlayer->SetBitrateCeiling((int32)NewBandwidthForStreaming);
 				}
 			}
 		}
@@ -505,7 +506,8 @@ void FElectraPlayer::Tick(FTimespan DeltaTime, FTimespan Timecode)
 	PlayerLock.Unlock();
 
 	// Forward enqueued session events. We do this even with no current internal player to ensure all pending events are sent and none are lost.
-	if (TSharedPtr<IElectraPlayerAdapterDelegate, ESPMode::ThreadSafe> PinnedAdapterDelegate = AdapterDelegate.Pin())
+	TSharedPtr<IElectraPlayerAdapterDelegate, ESPMode::ThreadSafe> PinnedAdapterDelegate = AdapterDelegate.Pin();
+	if (PinnedAdapterDelegate.IsValid())
 	{
 		IElectraPlayerAdapterDelegate::EPlayerEvent Event;
 		while (DeferredEvents.Dequeue(Event))
@@ -542,7 +544,8 @@ void FElectraPlayer::OnVideoFlush()
 	TSharedPtr<FInternalPlayerImpl, ESPMode::ThreadSafe> Player = CurrentPlayer;
 	if (Player.IsValid())
 	{
-		if (TSharedPtr<IElectraPlayerAdapterDelegate, ESPMode::ThreadSafe> PinnedAdapterDelegate = AdapterDelegate.Pin())
+		TSharedPtr<IElectraPlayerAdapterDelegate, ESPMode::ThreadSafe> PinnedAdapterDelegate = AdapterDelegate.Pin();
+		if (PinnedAdapterDelegate.IsValid())
 		{
 			PinnedAdapterDelegate->OnVideoFlush();
 			// The event sink (player facade) needs to prepare for all new data.
@@ -575,7 +578,8 @@ void FElectraPlayer::OnAudioFlush()
 	TSharedPtr<FInternalPlayerImpl, ESPMode::ThreadSafe> Player = CurrentPlayer;
 	if (Player.IsValid())
 	{
-		if (TSharedPtr<IElectraPlayerAdapterDelegate, ESPMode::ThreadSafe> PinnedAdapterDelegate = AdapterDelegate.Pin())
+		TSharedPtr<IElectraPlayerAdapterDelegate, ESPMode::ThreadSafe> PinnedAdapterDelegate = AdapterDelegate.Pin();
+		if (PinnedAdapterDelegate.IsValid())
 		{
 			PinnedAdapterDelegate->OnAudioFlush();
 		}
@@ -587,7 +591,8 @@ void FElectraPlayer::OnVideoRenderingStarted()
 	TSharedPtr<FInternalPlayerImpl, ESPMode::ThreadSafe> Player = CurrentPlayer;
 	if (Player.IsValid())
 	{
-		if (TSharedPtr<IElectraPlayerAdapterDelegate, ESPMode::ThreadSafe> PinnedAdapterDelegate = AdapterDelegate.Pin())
+		TSharedPtr<IElectraPlayerAdapterDelegate, ESPMode::ThreadSafe> PinnedAdapterDelegate = AdapterDelegate.Pin();
+		if (PinnedAdapterDelegate.IsValid())
 		{
 			PinnedAdapterDelegate->SendMediaEvent(IElectraPlayerAdapterDelegate::EPlayerEvent::Internal_RenderClockStart);
 		}
@@ -598,7 +603,8 @@ void FElectraPlayer::OnVideoRenderingStopped()
 	TSharedPtr<FInternalPlayerImpl, ESPMode::ThreadSafe> Player = CurrentPlayer;
 	if (Player.IsValid())
 	{
-		if (TSharedPtr<IElectraPlayerAdapterDelegate, ESPMode::ThreadSafe> PinnedAdapterDelegate = AdapterDelegate.Pin())
+		TSharedPtr<IElectraPlayerAdapterDelegate, ESPMode::ThreadSafe> PinnedAdapterDelegate = AdapterDelegate.Pin();
+		if (PinnedAdapterDelegate.IsValid())
 		{
 			PinnedAdapterDelegate->SendMediaEvent(IElectraPlayerAdapterDelegate::EPlayerEvent::Internal_RenderClockStop);
 		}
@@ -626,10 +632,11 @@ void FElectraPlayer::OnAudioRenderingStopped()
  */
 bool FElectraPlayer::PresentVideoFrame(const FVideoDecoderOutputPtr& InVideoFrame)
 {
-	if (TSharedPtr<IElectraPlayerAdapterDelegate, ESPMode::ThreadSafe> PinnedAdapterDelegate = AdapterDelegate.Pin())
+	TSharedPtr<IElectraPlayerAdapterDelegate, ESPMode::ThreadSafe> PinnedAdapterDelegate = AdapterDelegate.Pin();
+	if (PinnedAdapterDelegate.IsValid())
 	{
-	  PinnedAdapterDelegate->PresentVideoFrame(InVideoFrame);
-	  LastPresentedFrameDimension = InVideoFrame->GetOutputDim();
+		PinnedAdapterDelegate->PresentVideoFrame(InVideoFrame);
+		LastPresentedFrameDimension = InVideoFrame->GetOutputDim();
 	}
 	return true;
 }
@@ -642,7 +649,8 @@ bool FElectraPlayer::PresentVideoFrame(const FVideoDecoderOutputPtr& InVideoFram
  */
 bool FElectraPlayer::PresentAudioFrame(const IAudioDecoderOutputPtr& DecoderOutput)
 {
-	if (TSharedPtr<IElectraPlayerAdapterDelegate, ESPMode::ThreadSafe> PinnedAdapterDelegate = AdapterDelegate.Pin())
+	TSharedPtr<IElectraPlayerAdapterDelegate, ESPMode::ThreadSafe> PinnedAdapterDelegate = AdapterDelegate.Pin();
+	if (PinnedAdapterDelegate.IsValid())
 	{
 		PinnedAdapterDelegate->PresentAudioFrame(DecoderOutput);
 	}
@@ -655,7 +663,8 @@ bool FElectraPlayer::PresentAudioFrame(const IAudioDecoderOutputPtr& DecoderOutp
 */
 void FElectraPlayer::DropOldFramesFromPresentationQueue()
 {
-	if (TSharedPtr<IElectraPlayerAdapterDelegate, ESPMode::ThreadSafe> PinnedAdapterDelegate = AdapterDelegate.Pin())
+	TSharedPtr<IElectraPlayerAdapterDelegate, ESPMode::ThreadSafe> PinnedAdapterDelegate = AdapterDelegate.Pin();
+	if (PinnedAdapterDelegate.IsValid())
 	{
 		// We ask the event sink (player facade) to trigger this, as we don't have good enough timing info
 		PinnedAdapterDelegate->SendMediaEvent(IElectraPlayerAdapterDelegate::EPlayerEvent::Internal_PurgeVideoSamplesHint);
@@ -669,7 +678,8 @@ void FElectraPlayer::DropOldFramesFromPresentationQueue()
 bool FElectraPlayer::CanPresentVideoFrames(uint64 NumFrames)
 {
 	DropOldFramesFromPresentationQueue();
-	if (TSharedPtr<IElectraPlayerAdapterDelegate, ESPMode::ThreadSafe> PinnedAdapterDelegate = AdapterDelegate.Pin())
+	TSharedPtr<IElectraPlayerAdapterDelegate, ESPMode::ThreadSafe> PinnedAdapterDelegate = AdapterDelegate.Pin();
+	if (PinnedAdapterDelegate.IsValid())
 	{
 		return PinnedAdapterDelegate->CanReceiveVideoSamples(NumFrames);
 	}
@@ -678,7 +688,8 @@ bool FElectraPlayer::CanPresentVideoFrames(uint64 NumFrames)
 
 bool FElectraPlayer::CanPresentAudioFrames(uint64 NumFrames)
 {
-	if (TSharedPtr<IElectraPlayerAdapterDelegate, ESPMode::ThreadSafe> PinnedAdapterDelegate = AdapterDelegate.Pin())
+	TSharedPtr<IElectraPlayerAdapterDelegate, ESPMode::ThreadSafe> PinnedAdapterDelegate = AdapterDelegate.Pin();
+	if (PinnedAdapterDelegate.IsValid())
 	{
 		return PinnedAdapterDelegate->CanReceiveAudioSamples(NumFrames);
 	}
@@ -829,7 +840,8 @@ void FElectraPlayer::TriggerFirstSeekIfNecessary()
 
 		// Next, give a list of the seekable positions to the delegate and ask it if it wants to seek to one of them,
 		// overriding any potential time offset from above.
-		if (TSharedPtr<IElectraPlayerAdapterDelegate, ESPMode::ThreadSafe> PinnedAdapterDelegate = AdapterDelegate.Pin())
+		TSharedPtr<IElectraPlayerAdapterDelegate, ESPMode::ThreadSafe> PinnedAdapterDelegate = AdapterDelegate.Pin();
+		if (PinnedAdapterDelegate.IsValid())
 		{
 			TSharedPtr<TArray<FTimespan>, ESPMode::ThreadSafe> SeekablePositions = MakeShared<TArray<FTimespan>, ESPMode::ThreadSafe>();
 			if (SeekablePositions.IsValid())
@@ -928,11 +940,11 @@ TSharedPtr<Electra::FTrackMetadata, ESPMode::ThreadSafe> FElectraPlayer::GetTrac
 		TArray<Electra::FTrackMetadata> TrackMetaData;
 		if (TrackType == EPlayerTrackType::Video)
 		{
-			CurrentPlayer->AdaptivePlayer->GetTrackMetadata(TrackMetaData, Electra::EStreamType::Video);
+			LockedPlayer->AdaptivePlayer->GetTrackMetadata(TrackMetaData, Electra::EStreamType::Video);
 		}
 		else if (TrackType == EPlayerTrackType::Audio)
 		{
-			CurrentPlayer->AdaptivePlayer->GetTrackMetadata(TrackMetaData, Electra::EStreamType::Audio);
+			LockedPlayer->AdaptivePlayer->GetTrackMetadata(TrackMetaData, Electra::EStreamType::Audio);
 		}
 		if (TrackIndex >= 0 && TrackIndex < TrackMetaData.Num())
 		{
@@ -1038,7 +1050,7 @@ int32 FElectraPlayer::GetSelectedTrack(EPlayerTrackType TrackType) const
 				TSharedPtr<FInternalPlayerImpl, ESPMode::ThreadSafe> LockedPlayer = CurrentPlayer;
 				if (LockedPlayer.IsValid() && LockedPlayer->AdaptivePlayer.IsValid())
 				{
-					if (CurrentPlayer->AdaptivePlayer->IsTrackDeselected(Electra::EStreamType::Audio))
+					if (LockedPlayer->AdaptivePlayer->IsTrackDeselected(Electra::EStreamType::Audio))
 					{
 						SelectedAudioTrackIndex = -1;
 						bAudioTrackIndexDirty = false;
@@ -1046,7 +1058,7 @@ int32 FElectraPlayer::GetSelectedTrack(EPlayerTrackType TrackType) const
 					else
 					{
 						Electra::FStreamSelectionAttributes Attributes;
-						CurrentPlayer->AdaptivePlayer->GetSelectedTrackAttributes(Attributes, Electra::EStreamType::Audio);
+						LockedPlayer->AdaptivePlayer->GetSelectedTrackAttributes(Attributes, Electra::EStreamType::Audio);
 						if (Attributes.OverrideIndex.IsSet())
 						{
 							SelectedAudioTrackIndex = Attributes.OverrideIndex.GetValue();
@@ -1157,7 +1169,7 @@ bool FElectraPlayer::SelectTrack(EPlayerTrackType TrackType, int32 TrackIndex)
 					TSharedPtr<FInternalPlayerImpl, ESPMode::ThreadSafe> LockedPlayer = CurrentPlayer;
 					if (LockedPlayer.IsValid() && LockedPlayer->AdaptivePlayer.IsValid())
 					{
-						CurrentPlayer->AdaptivePlayer->SelectTrackByAttributes(Electra::EStreamType::Audio, AudioAttributes);
+						LockedPlayer->AdaptivePlayer->SelectTrackByAttributes(Electra::EStreamType::Audio, AudioAttributes);
 					}
 
 					SelectedAudioTrackIndex = TrackIndex;
@@ -1173,7 +1185,7 @@ bool FElectraPlayer::SelectTrack(EPlayerTrackType TrackType, int32 TrackIndex)
 			TSharedPtr<FInternalPlayerImpl, ESPMode::ThreadSafe> LockedPlayer = CurrentPlayer;
 			if (LockedPlayer.IsValid() && LockedPlayer->AdaptivePlayer.IsValid())
 			{
-				CurrentPlayer->AdaptivePlayer->DeselectTrack(Electra::EStreamType::Audio);
+				LockedPlayer->AdaptivePlayer->DeselectTrack(Electra::EStreamType::Audio);
 			}
 			return true;
 		}
@@ -1195,7 +1207,8 @@ void FElectraPlayer::OnMediaPlayerEventReceived(TSharedPtrTS<IAdaptiveStreamingP
 		InEvent->GetPresentationTime().GetAsSeconds(), InEvent->GetDuration().GetAsSeconds());
 #endif
 
-	if (TSharedPtr<IElectraPlayerAdapterDelegate, ESPMode::ThreadSafe> PinnedAdapterDelegate = AdapterDelegate.Pin())
+	TSharedPtr<IElectraPlayerAdapterDelegate, ESPMode::ThreadSafe> PinnedAdapterDelegate = AdapterDelegate.Pin();
+	if (PinnedAdapterDelegate.IsValid())
 	{
 		// Create a binary media sample of our extended format and pass it up.
 		TSharedRef<FElectraBinarySample, ESPMode::ThreadSafe> Meta = MakeShared<FElectraBinarySample, ESPMode::ThreadSafe>();
@@ -2117,11 +2130,11 @@ void FElectraPlayer::UpdatePlayEndStatistics()
 		return;
 	}
 
-	double PlayPos = CurrentPlayer->AdaptivePlayer->GetPlayPosition().GetAsSeconds();
+	double PlayPos = LockedPlayer->AdaptivePlayer->GetPlayPosition().GetAsSeconds();
 	Electra::FTimeRange MediaTimeline;
 	Electra::FTimeValue MediaDuration;
-	CurrentPlayer->AdaptivePlayer->GetTimelineRange(MediaTimeline);
-	MediaDuration = CurrentPlayer->AdaptivePlayer->GetDuration();
+	LockedPlayer->AdaptivePlayer->GetTimelineRange(MediaTimeline);
+	MediaDuration = LockedPlayer->AdaptivePlayer->GetDuration();
 	// Update statistics
 	FScopeLock Lock(&StatisticsLock);
 	if (Statistics.PlayPosAtStart >= 0.0 && Statistics.PlayPosAtEnd < 0.0)
@@ -2146,7 +2159,8 @@ void FElectraPlayer::LogStatistics()
 		SegsPerStream += FString::Printf(TEXT("%d : %u\n"), pair.Key, pair.Value);
 	}
 
-	if (TSharedPtr<IElectraPlayerAdapterDelegate, ESPMode::ThreadSafe> PinnedAdapterDelegate = AdapterDelegate.Pin())
+	TSharedPtr<IElectraPlayerAdapterDelegate, ESPMode::ThreadSafe> PinnedAdapterDelegate = AdapterDelegate.Pin();
+	if (PinnedAdapterDelegate.IsValid())
 	{
 		UE_LOG(LogElectraPlayer, Log, TEXT(
 			"[%p][%p] Electra player statistics:\n"\
@@ -2459,7 +2473,8 @@ void FElectraPlayer::FAdaptiveStreamingPlayerResourceProvider::ProcessPendingSta
 			  InOutRequest->GetResourceType() == Electra::IAdaptiveStreamingPlayerResourceRequest::EPlaybackResourceType::LicenseKey);
 		if (InOutRequest->GetResourceType() == Electra::IAdaptiveStreamingPlayerResourceRequest::EPlaybackResourceType::Playlist)
 		{
-			if (TSharedPtr<IElectraPlayerAdapterDelegate, ESPMode::ThreadSafe> PinnedAdapterDelegate = AdapterDelegate.Pin())
+			TSharedPtr<IElectraPlayerAdapterDelegate, ESPMode::ThreadSafe> PinnedAdapterDelegate = AdapterDelegate.Pin();
+			if (PinnedAdapterDelegate.IsValid())
 			{
 				FVariantValue Value = PinnedAdapterDelegate->QueryOptions(IElectraPlayerAdapterDelegate::EOptionType::PlayListData, FVariantValue(InOutRequest->GetResourceURL()));
 				if (Value.IsValid())
@@ -2486,7 +2501,8 @@ void FElectraPlayer::FAdaptiveStreamingPlayerResourceProvider::ProcessPendingSta
 		}
 		else if (InOutRequest->GetResourceType() == Electra::IAdaptiveStreamingPlayerResourceRequest::EPlaybackResourceType::LicenseKey)
 		{
-			if (TSharedPtr<IElectraPlayerAdapterDelegate, ESPMode::ThreadSafe> PinnedAdapterDelegate = AdapterDelegate.Pin())
+			TSharedPtr<IElectraPlayerAdapterDelegate, ESPMode::ThreadSafe> PinnedAdapterDelegate = AdapterDelegate.Pin();
+			if (PinnedAdapterDelegate.IsValid())
 			{
 				FVariantValue Value = PinnedAdapterDelegate->QueryOptions(IElectraPlayerAdapterDelegate::EOptionType::LicenseKeyData, FVariantValue(InOutRequest->GetResourceURL()));
 				if (Value.IsValid())
diff --git a/Engine/Plugins/Media/ElectraPlayer/Source/ElectraPlayerRuntime/Private/ElectraPlayer.h b/Engine/Plugins/Media/ElectraPlayer/Source/ElectraPlayerRuntime/Private/ElectraPlayer.h
--- a/Engine/Plugins/Media/ElectraPlayer/Source/ElectraPlayerRuntime/Private/ElectraPlayer.h
+++ b/Engine/Plugins/Media/ElectraPlayer/Source/ElectraPlayerRuntime/Private/ElectraPlayer.h
@@ -352,7 +352,7 @@ private:
 	TAtomic<EPlayerState>							State;
 	TAtomic<EPlayerStatus>							Status;
 
-	TAtomic<bool>									bCloseInProgress;
+	TAtomic<bool>									bPlayerHasClosed;
 	TAtomic<bool>									bHasPendingError;
 
 	bool											bAllowKillAfterCloseEvent;
