ElectraPlayer: added perf stat macros

#rnx

diff --git a/Engine/Plugins/Media/ElectraPlayer/Source/ElectraPlayerRuntime/Private/Runtime/Decoder/Android/AudioDecoderAAC_Android.cpp b/Engine/Plugins/Media/ElectraPlayer/Source/ElectraPlayerRuntime/Private/Runtime/Decoder/Android/AudioDecoderAAC_Android.cpp
--- a/Engine/Plugins/Media/ElectraPlayer/Source/ElectraPlayerRuntime/Private/Runtime/Decoder/Android/AudioDecoderAAC_Android.cpp
+++ b/Engine/Plugins/Media/ElectraPlayer/Source/ElectraPlayerRuntime/Private/Runtime/Decoder/Android/AudioDecoderAAC_Android.cpp
@@ -11,9 +11,14 @@
 #include "Utilities/StringHelpers.h"
 #include "DecoderErrors_Android.h"
 #include "HAL/LowLevelMemTracker.h"
+#include "ElectraPlayerPrivate.h"
 
 #include "AudioDecoderAAC_JavaWrapper_Android.h"
 
+DECLARE_CYCLE_STAT(TEXT("FAudioDecoderAAC::Decode()"), STAT_ElectraPlayer_AudioAACDecode, STATGROUP_ElectraPlayer);
+DECLARE_CYCLE_STAT(TEXT("FAudioDecoderAAC::ConvertOutput()"), STAT_ElectraPlayer_AudioAACConvertOutput, STATGROUP_ElectraPlayer);
+
+
 namespace Electra
 {
 
@@ -572,6 +577,8 @@ bool FAudioDecoderAAC::Decode(FAccessUnit* InAccessUnit)
 		// Need a new output buffer?
 		if (CurrentOutputBuffer == nullptr)
 		{
+			SCOPE_CYCLE_COUNTER(STAT_ElectraPlayer_AudioAACDecode);
+			CSV_SCOPED_TIMING_STAT(ElectraPlayer, AudioAACDecode);
 			UEMediaError bufResult = Renderer->AcquireBuffer(CurrentOutputBuffer, 0, BufferAcquireOptions);
 			check(bufResult == UEMEDIA_ERROR_OK || bufResult == UEMEDIA_ERROR_INSUFFICIENT_DATA);
 			if (bufResult != UEMEDIA_ERROR_OK && bufResult != UEMEDIA_ERROR_INSUFFICIENT_DATA)
@@ -627,6 +634,8 @@ bool FAudioDecoderAAC::Decode(FAccessUnit* InAccessUnit)
 				result = DecoderInstance->DequeueOutputBuffer(OutputBufferInfo, 1000 * 2);
 				if (result == 0)
 				{
+					SCOPE_CYCLE_COUNTER(STAT_ElectraPlayer_AudioAACConvertOutput);
+					CSV_SCOPED_TIMING_STAT(ElectraPlayer, AudioAACConvertOutput);
 					if (OutputBufferInfo.BufferIndex >= 0)
 					{
 						int32 OutputByteCount    = OutputBufferInfo.Size;
@@ -705,6 +714,8 @@ bool FAudioDecoderAAC::Decode(FAccessUnit* InAccessUnit)
 			}
 			else
 			{
+				SCOPE_CYCLE_COUNTER(STAT_ElectraPlayer_AudioAACConvertOutput);
+				CSV_SCOPED_TIMING_STAT(ElectraPlayer, AudioAACConvertOutput);
 				// Check if we are in dummy decode mode already. If not we need to flush the decoder if we have one.
 				if (!bInDummyDecodeMode)
 				{
@@ -798,6 +809,8 @@ void FAudioDecoderAAC::WorkerThread()
 		// Notify the buffer listener that we will now be needing an AU for our input buffer.
 		if (!bError && InputBufferListener && AccessUnits.Num() == 0)
 		{
+			SCOPE_CYCLE_COUNTER(STAT_ElectraPlayer_AudioAACDecode);
+			CSV_SCOPED_TIMING_STAT(ElectraPlayer, AudioAACDecode);
 			FAccessUnitBufferInfo	sin;
 			IAccessUnitBufferListener::FBufferStats	stats;
 			AccessUnits.GetStats(sin);
@@ -821,7 +834,9 @@ void FAudioDecoderAAC::WorkerThread()
 
 			// Check if the format has changed such that we need to destroy and re-create the decoder.
 			if (CurrentAccessUnit->bTrackChangeDiscontinuity || IsDifferentFormat(CurrentAccessUnit) || (CurrentAccessUnit->bIsDummyData && !bInDummyDecodeMode))
-				{
+			{
+				SCOPE_CYCLE_COUNTER(STAT_ElectraPlayer_AudioAACDecode);
+				CSV_SCOPED_TIMING_STAT(ElectraPlayer, AudioAACDecode);
 				if (DecoderInstance.IsValid())
 				{
 					InternalDecoderDestroy();
@@ -835,6 +850,8 @@ void FAudioDecoderAAC::WorkerThread()
 			// Parse the CSD into a configuration record.
 			if (!ConfigRecord.IsValid() && !bError)
 			{
+				SCOPE_CYCLE_COUNTER(STAT_ElectraPlayer_AudioAACDecode);
+				CSV_SCOPED_TIMING_STAT(ElectraPlayer, AudioAACDecode);
 				if (CurrentAccessUnit->AUCodecData.IsValid())
 				{
 					CurrentCodecData = CurrentAccessUnit->AUCodecData;
@@ -855,6 +872,8 @@ void FAudioDecoderAAC::WorkerThread()
 				// Need to create a decoder instance?
 				if (!DecoderInstance.IsValid() && !bError)
 				{
+					SCOPE_CYCLE_COUNTER(STAT_ElectraPlayer_AudioAACDecode);
+					CSV_SCOPED_TIMING_STAT(ElectraPlayer, AudioAACDecode);
 					if (!InternalDecoderCreate())
 					{
 						bError = true;
@@ -893,6 +912,8 @@ void FAudioDecoderAAC::WorkerThread()
 		// Flush?
 		if (FlushDecoderSignal.IsSignaled())
 		{
+			SCOPE_CYCLE_COUNTER(STAT_ElectraPlayer_AudioAACDecode);
+			CSV_SCOPED_TIMING_STAT(ElectraPlayer, AudioAACDecode);
 			ReturnUnusedOutputBuffer();
 			AccessUnits.Flush();
 			FlushDecoder();
diff --git a/Engine/Plugins/Media/ElectraPlayer/Source/ElectraPlayerRuntime/Private/Runtime/Decoder/Android/VideoDecoderH264_Android.cpp b/Engine/Plugins/Media/ElectraPlayer/Source/ElectraPlayerRuntime/Private/Runtime/Decoder/Android/VideoDecoderH264_Android.cpp
--- a/Engine/Plugins/Media/ElectraPlayer/Source/ElectraPlayerRuntime/Private/Runtime/Decoder/Android/VideoDecoderH264_Android.cpp
+++ b/Engine/Plugins/Media/ElectraPlayer/Source/ElectraPlayerRuntime/Private/Runtime/Decoder/Android/VideoDecoderH264_Android.cpp
@@ -11,6 +11,7 @@
 #include "DecoderErrors_Android.h"
 #include "HAL/LowLevelMemTracker.h"
 #include "Android/AndroidPlatformMisc.h"
+#include "ElectraPlayerPrivate.h"
 
 #include "VideoDecoderH264_JavaWrapper_Android.h"
 #include "MediaVideoDecoderOutputAndroid.h"
@@ -19,6 +20,11 @@
 #include "Android/AndroidPlatform.h"
 #include "Android/AndroidJava.h"
 
+
+DECLARE_CYCLE_STAT(TEXT("FVideoDecoderH264::Decode()"), STAT_ElectraPlayer_VideoH264Decode, STATGROUP_ElectraPlayer);
+DECLARE_CYCLE_STAT(TEXT("FVideoDecoderH264::ConvertOutput()"), STAT_ElectraPlayer_VideoH264ConvertOutput, STATGROUP_ElectraPlayer);
+
+
 namespace Electra
 {
 
@@ -1501,6 +1507,8 @@ void FVideoDecoderH264::WorkerThread()
 		// Notify optional buffer listener that we will now be needing an AU for our input buffer.
 		if (!CurrentAccessUnit && !bError && InputBufferListener && AccessUnitBuffer.Num() == 0)
 		{
+			SCOPE_CYCLE_COUNTER(STAT_ElectraPlayer_VideoH264Decode);
+			CSV_SCOPED_TIMING_STAT(ElectraPlayer, VideoH264Decode);
 			FAccessUnitBufferInfo	sin;
 			IAccessUnitBufferListener::FBufferStats	stats;
 			AccessUnitBuffer.GetStats(sin);
@@ -1520,6 +1528,9 @@ void FVideoDecoderH264::WorkerThread()
 		// Try to pull output. Do this first to make room in the decoder for new input data.
 		if (!bError)
 		{
+			SCOPE_CYCLE_COUNTER(STAT_ElectraPlayer_VideoH264ConvertOutput);
+			CSV_SCOPED_TIMING_STAT(ElectraPlayer, VideoH264ConvertOutput);
+
 			ProcessReadyOutputBuffersToSurface();
 
 			// Any new data available?
@@ -1586,6 +1597,8 @@ void FVideoDecoderH264::WorkerThread()
 		{
 			if (AccessUnitBuffer.WaitForData(1000 * 5))
 			{
+				SCOPE_CYCLE_COUNTER(STAT_ElectraPlayer_VideoH264Decode);
+				CSV_SCOPED_TIMING_STAT(ElectraPlayer, VideoH264Decode);
 				// When there is data, even and especially after a previous EOD, we are no longer done and idling.
 				if (bDone)
 				{
@@ -1696,6 +1709,8 @@ void FVideoDecoderH264::WorkerThread()
 			{
 				if (DecoderState == EDecoderState::Create)
 				{
+					SCOPE_CYCLE_COUNTER(STAT_ElectraPlayer_VideoH264Decode);
+					CSV_SCOPED_TIMING_STAT(ElectraPlayer, VideoH264Decode);
 					DecoderState = EDecoderState::Ready;
 					if (InternalDecoderCreate())
 					{
@@ -1711,6 +1726,8 @@ void FVideoDecoderH264::WorkerThread()
 				}
 				else if (DecoderState == EDecoderState::Flush)
 				{
+					SCOPE_CYCLE_COUNTER(STAT_ElectraPlayer_VideoH264Decode);
+					CSV_SCOPED_TIMING_STAT(ElectraPlayer, VideoH264Decode);
 					if (DecoderInstance.IsValid())
 					{
 						DecoderInstance->Flush();
@@ -1758,6 +1775,8 @@ void FVideoDecoderH264::WorkerThread()
 					}
 					case EInputState::FeedEOS:
 					{
+						SCOPE_CYCLE_COUNTER(STAT_ElectraPlayer_VideoH264Decode);
+						CSV_SCOPED_TIMING_STAT(ElectraPlayer, VideoH264Decode);
 						bBlockedOnInput = false;
 						DrainDecoder();
 						break;
@@ -1791,6 +1810,8 @@ void FVideoDecoderH264::WorkerThread()
 		// Flush?
 		if (FlushDecoderSignal.IsSignaled())
 		{
+			SCOPE_CYCLE_COUNTER(STAT_ElectraPlayer_VideoH264Decode);
+			CSV_SCOPED_TIMING_STAT(ElectraPlayer, VideoH264Decode);
 			// Have to destroy the decoder!
 			InternalDecoderDestroy();
 			AccessUnitBuffer.Flush();
diff --git a/Engine/Plugins/Media/ElectraPlayer/Source/ElectraPlayerRuntime/Private/Runtime/Decoder/Apple/AudioDecoderAAC_Apple.cpp b/Engine/Plugins/Media/ElectraPlayer/Source/ElectraPlayerRuntime/Private/Runtime/Decoder/Apple/AudioDecoderAAC_Apple.cpp
--- a/Engine/Plugins/Media/ElectraPlayer/Source/ElectraPlayerRuntime/Private/Runtime/Decoder/Apple/AudioDecoderAAC_Apple.cpp
+++ b/Engine/Plugins/Media/ElectraPlayer/Source/ElectraPlayerRuntime/Private/Runtime/Decoder/Apple/AudioDecoderAAC_Apple.cpp
@@ -15,12 +15,16 @@
 #include "Utilities/StringHelpers.h"
 #include "DecoderErrors_Apple.h"
 #include "HAL/LowLevelMemTracker.h"
+#include "ElectraPlayerPrivate.h"
 
 #include <AudioToolbox/AudioToolbox.h>
 
 
 #define MAGIC_INSUFFICIENT_INPUT_MARKER		4711
 
+DECLARE_CYCLE_STAT(TEXT("FAudioDecoderAAC::Decode()"), STAT_ElectraPlayer_AudioAACDecode, STATGROUP_ElectraPlayer);
+DECLARE_CYCLE_STAT(TEXT("FAudioDecoderAAC::ConvertOutput()"), STAT_ElectraPlayer_AudioAACConvertOutput, STATGROUP_ElectraPlayer);
+
 namespace Electra
 {
 
@@ -655,6 +659,8 @@ bool FAudioDecoderAAC::Decode(FAccessUnit* InAccessUnit)
 		// Need a new output buffer?
 		if (CurrentOutputBuffer == nullptr)
 		{
+			SCOPE_CYCLE_COUNTER(STAT_ElectraPlayer_AudioAACDecode);
+			CSV_SCOPED_TIMING_STAT(ElectraPlayer, AudioAACDecode);
 			UEMediaError bufResult = Renderer->AcquireBuffer(CurrentOutputBuffer, 0, BufferAcquireOptions);
 			check(bufResult == UEMEDIA_ERROR_OK || bufResult == UEMEDIA_ERROR_INSUFFICIENT_DATA);
 			if (bufResult != UEMEDIA_ERROR_OK && bufResult != UEMEDIA_ERROR_INSUFFICIENT_DATA)
@@ -715,6 +721,8 @@ bool FAudioDecoderAAC::Decode(FAccessUnit* InAccessUnit)
 
 					if (InOutPacketSize)
 					{
+						SCOPE_CYCLE_COUNTER(STAT_ElectraPlayer_AudioAACConvertOutput);
+						CSV_SCOPED_TIMING_STAT(ElectraPlayer, AudioAACConvertOutput);
 						NumSamplesProduced += InOutPacketSize;
 
 						int32 OutputByteCount = InOutPacketSize * NumberOfChannels * sizeof(int16);
@@ -748,6 +756,8 @@ bool FAudioDecoderAAC::Decode(FAccessUnit* InAccessUnit)
 
 				if (NumSamplesProduced)
 				{
+					SCOPE_CYCLE_COUNTER(STAT_ElectraPlayer_AudioAACConvertOutput);
+					CSV_SCOPED_TIMING_STAT(ElectraPlayer, AudioAACConvertOutput);
 					FTimeValue Duration;
 					Duration.SetFromND(NumSamplesProduced, SamplingRate);
 
@@ -765,10 +775,11 @@ bool FAudioDecoderAAC::Decode(FAccessUnit* InAccessUnit)
 
 					CurrentPTS += Duration;
 				}
-
 			}
 			else
 			{
+				SCOPE_CYCLE_COUNTER(STAT_ElectraPlayer_AudioAACConvertOutput);
+				CSV_SCOPED_TIMING_STAT(ElectraPlayer, AudioAACConvertOutput);
 				// Check if we are in dummy decode mode already.
 				if (!bInDummyDecodeMode)
 				{
@@ -850,7 +861,7 @@ bool FAudioDecoderAAC::IsDifferentFormat(const FAccessUnit* Data)
  */
 void FAudioDecoderAAC::WorkerThread()
 {
-	LLM_SCOPE(ELLMTag::MediaStreaming);
+	LLM_SCOPE(ELLMTag::ElectraPlayer);
 
 	FAccessUnit*	CurrentAccessUnit = nullptr;
 	bool			bError			  = false;
@@ -869,6 +880,8 @@ void FAudioDecoderAAC::WorkerThread()
 		// Notify the buffer listener that we will now be needing an AU for our input buffer.
 		if (!bError && InputBufferListener && AccessUnits.Num() == 0)
 		{
+			SCOPE_CYCLE_COUNTER(STAT_ElectraPlayer_AudioAACDecode);
+			CSV_SCOPED_TIMING_STAT(ElectraPlayer, AudioAACDecode);
 			FAccessUnitBufferInfo	sin;
 			IAccessUnitBufferListener::FBufferStats	stats;
 			AccessUnits.GetStats(sin);
@@ -892,7 +905,9 @@ void FAudioDecoderAAC::WorkerThread()
 
 			// Check if the format has changed such that we need to destroy and re-create the decoder.
 			if (CurrentAccessUnit->bTrackChangeDiscontinuity || IsDifferentFormat(CurrentAccessUnit) || (CurrentAccessUnit->bIsDummyData && !bInDummyDecodeMode))
-				{
+			{
+				SCOPE_CYCLE_COUNTER(STAT_ElectraPlayer_AudioAACDecode);
+				CSV_SCOPED_TIMING_STAT(ElectraPlayer, AudioAACDecode);
 				if (DecoderInstance)
 				{
 					InternalDecoderDestroy();
@@ -906,6 +921,8 @@ void FAudioDecoderAAC::WorkerThread()
 			// Parse the CSD into a configuration record.
 			if (!ConfigRecord.IsValid() && !bError)
 			{
+				SCOPE_CYCLE_COUNTER(STAT_ElectraPlayer_AudioAACDecode);
+				CSV_SCOPED_TIMING_STAT(ElectraPlayer, AudioAACDecode);
 				if (CurrentAccessUnit->AUCodecData.IsValid())
 				{
 					CurrentCodecData = CurrentAccessUnit->AUCodecData;
@@ -964,6 +981,8 @@ void FAudioDecoderAAC::WorkerThread()
 		// Flush?
 		if (FlushDecoderSignal.IsSignaled())
 		{
+			SCOPE_CYCLE_COUNTER(STAT_ElectraPlayer_AudioAACDecode);
+			CSV_SCOPED_TIMING_STAT(ElectraPlayer, AudioAACDecode);
 			ReturnUnusedOutputBuffer();
 			AccessUnits.Flush();
 
diff --git a/Engine/Plugins/Media/ElectraPlayer/Source/ElectraPlayerRuntime/Private/Runtime/Decoder/Apple/VideoDecoderH264_Apple.cpp b/Engine/Plugins/Media/ElectraPlayer/Source/ElectraPlayerRuntime/Private/Runtime/Decoder/Apple/VideoDecoderH264_Apple.cpp
--- a/Engine/Plugins/Media/ElectraPlayer/Source/ElectraPlayerRuntime/Private/Runtime/Decoder/Apple/VideoDecoderH264_Apple.cpp
+++ b/Engine/Plugins/Media/ElectraPlayer/Source/ElectraPlayerRuntime/Private/Runtime/Decoder/Apple/VideoDecoderH264_Apple.cpp
@@ -14,11 +14,15 @@
 #include "Utilities/UtilsMPEGVideo.h"
 #include "DecoderErrors_Apple.h"
 #include "HAL/LowLevelMemTracker.h"
+#include "ElectraPlayerPrivate.h"
 #include "MediaVideoDecoderOutputApple.h"
 #include "Renderer/RendererVideo.h"
 
 #include <VideoToolbox/VideoToolbox.h>
 
+DECLARE_CYCLE_STAT(TEXT("FVideoDecoderH264::Decode()"), STAT_ElectraPlayer_VideoH264Decode, STATGROUP_ElectraPlayer);
+DECLARE_CYCLE_STAT(TEXT("FVideoDecoderH264::ConvertOutput()"), STAT_ElectraPlayer_VideoH264ConvertOutput, STATGROUP_ElectraPlayer);
+
 namespace Electra
 {
 
@@ -49,7 +53,7 @@ public:
 	virtual void SetReadyBufferListener(IDecoderOutputBufferListener* Listener) override;
 
 	virtual void SetRenderer(TSharedPtr<IMediaRenderer, ESPMode::ThreadSafe> InRenderer) override;
-    
+
     virtual void SetResourceDelegate(const TSharedPtr<IVideoDecoderResourceDelegate, ESPMode::ThreadSafe>& ResourceDelegate) override;
 
 	virtual IAccessUnitBufferInterface::EAUpushResult AUdataPushAU(FAccessUnit* AccessUnit) override;
@@ -244,7 +248,7 @@ private:
 	IPlayerSessionServices*								PlayerSessionServices;
 
 	TSharedPtr<IMediaRenderer, ESPMode::ThreadSafe>		Renderer;
-    
+
     TWeakPtr<IVideoDecoderResourceDelegate, ESPMode::ThreadSafe> ResourceDelegate;
 
 	FAccessUnitBuffer									AccessUnitBuffer;
@@ -1293,6 +1297,8 @@ FVideoDecoderH264::EDecodeResult FVideoDecoderH264::Decode(FAccessUnit* AccessUn
  */
 void FVideoDecoderH264::WorkerThread()
 {
+	LLM_SCOPE(ELLMTag::ElectraPlayer);
+
 	bool bDone  = false;
     bool bGotLastSequenceAU = false;
 
@@ -1316,6 +1322,8 @@ void FVideoDecoderH264::WorkerThread()
 		// Notify optional buffer listener that we will now be needing an AU for our input buffer.
 		if (!bError && InputBufferListener && AccessUnitBuffer.Num() == 0)
 		{
+			SCOPE_CYCLE_COUNTER(STAT_ElectraPlayer_VideoH264Decode);
+			CSV_SCOPED_TIMING_STAT(ElectraPlayer, VideoH264Decode);
 			FAccessUnitBufferInfo	sin;
 			IAccessUnitBufferListener::FBufferStats	stats;
 			AccessUnitBuffer.GetStats(sin);
@@ -1391,6 +1399,9 @@ void FVideoDecoderH264::WorkerThread()
                 bGotLastSequenceAU = AccessUnit->bIsLastInPeriod;
 				if (bStreamFormatChanged || !DecoderHandle)
 				{
+					SCOPE_CYCLE_COUNTER(STAT_ElectraPlayer_VideoH264Decode);
+					CSV_SCOPED_TIMING_STAT(ElectraPlayer, VideoH264Decode);
+
 					CMFormatDescriptionRef NewFormatDescr = nullptr;
 					if (CreateFormatDescription(NewFormatDescr, AccessUnit))
 					{
@@ -1486,6 +1497,9 @@ void FVideoDecoderH264::WorkerThread()
 		// Flush?
 		if (FlushDecoderSignal.IsSignaled())
 		{
+			SCOPE_CYCLE_COUNTER(STAT_ElectraPlayer_VideoH264Decode);
+			CSV_SCOPED_TIMING_STAT(ElectraPlayer, VideoH264Decode);
+
 			// Have to destroy the decoder!
 			InternalDecoderDestroy();
 			FlushPendingImages();
@@ -1516,6 +1530,9 @@ void FVideoDecoderH264::WorkerThread()
  */
 void FVideoDecoderH264::ProcessOutput(bool bFlush)
 {
+	SCOPE_CYCLE_COUNTER(STAT_ElectraPlayer_VideoH264ConvertOutput);
+	CSV_SCOPED_TIMING_STAT(ElectraPlayer, VideoH264ConvertOutput);
+
 	// Pull any frames we can get from in-PTS-order array...
 	while(1)
 	{
diff --git a/Engine/Plugins/Media/ElectraPlayer/Source/ElectraPlayerRuntime/Private/Runtime/Decoder/DX/AudioDecoderAAC_DX.cpp b/Engine/Plugins/Media/ElectraPlayer/Source/ElectraPlayerRuntime/Private/Runtime/Decoder/DX/AudioDecoderAAC_DX.cpp
--- a/Engine/Plugins/Media/ElectraPlayer/Source/ElectraPlayerRuntime/Private/Runtime/Decoder/DX/AudioDecoderAAC_DX.cpp
+++ b/Engine/Plugins/Media/ElectraPlayer/Source/ElectraPlayerRuntime/Private/Runtime/Decoder/DX/AudioDecoderAAC_DX.cpp
@@ -57,6 +57,9 @@ if (FAILED(res))						\
 /*********************************************************************************************************************/
 /*********************************************************************************************************************/
 
+DECLARE_CYCLE_STAT(TEXT("FAudioDecoderAAC::Decode()"), STAT_ElectraPlayer_AudioAACDecode, STATGROUP_ElectraPlayer);
+DECLARE_CYCLE_STAT(TEXT("FAudioDecoderAAC::ConvertOutput()"), STAT_ElectraPlayer_AudioAACConvertOutput, STATGROUP_ElectraPlayer);
+
 
 #define AACDEC_PCM_SAMPLE_SIZE  2048
 #define AACDEC_MAX_CHANNELS		6
@@ -635,8 +638,6 @@ namespace Electra
 		}
 		else
 		{
-			CSV_SCOPED_TIMING_STAT(ElectraPlayer, FAudioDecoderAAC_Worker);
-
 			// Create the input sample.
 			TRefCountPtr<IMFMediaBuffer>	InputSampleBuffer;
 			BYTE* pbNewBuffer = nullptr;
@@ -644,6 +645,9 @@ namespace Electra
 			DWORD					dwSize = 0;
 			LONGLONG				llSampleTime = 0;
 
+			SCOPE_CYCLE_COUNTER(STAT_ElectraPlayer_AudioAACDecode);
+			CSV_SCOPED_TIMING_STAT(ElectraPlayer, AudioAACDecode);
+
 			VERIFY_HR(MFCreateSample(InputSample.GetInitReference()), "Failed to create audio decoder input sample", ERRCODE_INTERNAL_COULD_NOT_CREATE_INPUT_SAMPLE);
 			VERIFY_HR(MFCreateMemoryBuffer((DWORD)InSize, InputSampleBuffer.GetInitReference()), "Failed to create audio decoder input sample memory buffer", ERRCODE_INTERNAL_COULD_NOT_CREATE_INPUTBUFFER);
 			VERIFY_HR(InputSample->AddBuffer(InputSampleBuffer.GetReference()), "Failed to set audio decoder input buffer with sample", ERRCODE_INTERNAL_COULD_NOT_ADD_INPUT_BUFFER_TO_SAMPLE);
@@ -666,7 +670,8 @@ namespace Electra
 			}
 			if (!CurrentDecoderOutputBuffer.IsValid())
 			{
-				CSV_SCOPED_TIMING_STAT(ElectraPlayer, FAudioDecoderAAC_Worker);
+				SCOPE_CYCLE_COUNTER(STAT_ElectraPlayer_AudioAACDecode);
+				CSV_SCOPED_TIMING_STAT(ElectraPlayer, AudioAACDecode);
 
 				if (!CreateDecoderOutputBuffer())
 				{
@@ -678,18 +683,19 @@ namespace Electra
 			CurrentDecoderOutputBuffer->PrepareForProcess();
 
 			{
-				CSV_SCOPED_TIMING_STAT(ElectraPlayer, FAudioDecoderAAC_Worker);
+				SCOPE_CYCLE_COUNTER(STAT_ElectraPlayer_AudioAACDecode);
+				CSV_SCOPED_TIMING_STAT(ElectraPlayer, AudioAACDecode);
 				res = DecoderTransform->ProcessOutput(0, 1, &CurrentDecoderOutputBuffer->mOutputBuffer, &dwStatus);
 			}
 
 			if (res == MF_E_TRANSFORM_NEED_MORE_INPUT)
 			{
-				CSV_SCOPED_TIMING_STAT(ElectraPlayer, FAudioDecoderAAC_Worker);
-
 				// Flushing / draining?
 				if (bFlush)
 				{
 					// Yes. This means we have received all pending output and are done now.
+					SCOPE_CYCLE_COUNTER(STAT_ElectraPlayer_AudioAACDecode);
+					CSV_SCOPED_TIMING_STAT(ElectraPlayer, AudioAACDecode);
 					if (!bFlushOnly)
 					{
 						// After a drain issue a flush.
@@ -716,6 +722,8 @@ namespace Electra
 			}
 			else if (res == MF_E_TRANSFORM_STREAM_CHANGE)
 			{
+				SCOPE_CYCLE_COUNTER(STAT_ElectraPlayer_AudioAACDecode);
+				CSV_SCOPED_TIMING_STAT(ElectraPlayer, AudioAACDecode);
 				// Update output type.
 				if (!SetDecoderOutputType())
 				{
@@ -742,8 +750,7 @@ namespace Electra
 				VERIFY_HR(DecodedOutputSample->ConvertToContiguousBuffer(DecodedLinearOutputBuffer.GetInitReference()), "Failed to convert audio decoder output sample to contiguous buffer", ERRCODE_INTERNAL_COULD_NOT_MAKE_CONTIGUOUS_OUTPUT_BUFFER);
 				VERIFY_HR(DecodedLinearOutputBuffer->GetCurrentLength(&dwBufferLen), "Failed to get audio decoder output buffer current length", ERRCODE_INTERNAL_COULD_NOT_GET_OUTPUT_BUFFER_LENGTH);
 				VERIFY_HR(DecodedLinearOutputBuffer->Lock(&pDecompressedData, &dwMaxBufferLen, &dwBufferLen), "Failed to lock audio decoder output buffer", ERRCODE_INTERNAL_COULD_NOT_LOCK_OUTPUT_BUFFER)
-
-					VERIFY_HR(MFCreateWaveFormatExFromMFMediaType(CurrentOutputMediaType.GetReference(), &OutputWaveFormatPtr, &WaveFormatSize, MFWaveFormatExConvertFlag_Normal), "Failed to create audio decoder output buffer format info", ERRCODE_INTERNAL_COULD_NOT_CREATE_OUTPUT_BUFFER_FORMAT_INFO);
+				VERIFY_HR(MFCreateWaveFormatExFromMFMediaType(CurrentOutputMediaType.GetReference(), &OutputWaveFormatPtr, &WaveFormatSize, MFWaveFormatExConvertFlag_Normal), "Failed to create audio decoder output buffer format info", ERRCODE_INTERNAL_COULD_NOT_CREATE_OUTPUT_BUFFER_FORMAT_INFO);
 				FMemory::Memcpy(&OutputWaveFormat, OutputWaveFormatPtr, sizeof(OutputWaveFormat));
 				CoTaskMemFree(OutputWaveFormatPtr);
 				OutputWaveFormatPtr = nullptr;
@@ -766,7 +773,8 @@ namespace Electra
 
 					if (CurrentRenderOutputBuffer == nullptr)
 					{
-						CSV_SCOPED_TIMING_STAT(ElectraPlayer, FAudioDecoderAAC_Worker);
+						SCOPE_CYCLE_COUNTER(STAT_ElectraPlayer_AudioAACDecode);
+						CSV_SCOPED_TIMING_STAT(ElectraPlayer, AudioAACDecode);
 
 						UEMediaError bufResult = Renderer->AcquireBuffer(CurrentRenderOutputBuffer, 0, BufferAcquireOptions);
 						check(bufResult == UEMEDIA_ERROR_OK || bufResult == UEMEDIA_ERROR_INSUFFICIENT_DATA);
@@ -780,7 +788,8 @@ namespace Electra
 					NotifyReadyBufferListener(bHaveAvailSmpBlk);
 					if (bHaveAvailSmpBlk)
 					{
-						CSV_SCOPED_TIMING_STAT(ElectraPlayer, FAudioDecoderAAC_Worker);
+						SCOPE_CYCLE_COUNTER(STAT_ElectraPlayer_AudioAACConvertOutput);
+						CSV_SCOPED_TIMING_STAT(ElectraPlayer, AudioAACConvertOutput);
 
 						FTimeValue OutputSampleDuration;
 						OutputSampleDuration.SetFromND(nSamplesProduced, OutputWaveFormat.nSamplesPerSec);
@@ -876,7 +885,8 @@ namespace Electra
 			// Notify optional buffer listener that we will now be needing an AU for our input buffer.
 			if (!bError && InputBufferListener && AccessUnits.Num() == 0)
 			{
-				CSV_SCOPED_TIMING_STAT(ElectraPlayer, FAudioDecoderAAC_Worker);
+				SCOPE_CYCLE_COUNTER(STAT_ElectraPlayer_AudioAACDecode);
+				CSV_SCOPED_TIMING_STAT(ElectraPlayer, AudioAACDecode);
 
 				FAccessUnitBufferInfo	sin;
 				IAccessUnitBufferListener::FBufferStats	stats;
@@ -953,7 +963,8 @@ namespace Electra
 				// Need to create a decoder instance?
 				if (!DecoderTransform.IsValid() && !bError)
 				{
-					CSV_SCOPED_TIMING_STAT(ElectraPlayer, FAudioDecoderAAC_Worker);
+					SCOPE_CYCLE_COUNTER(STAT_ElectraPlayer_AudioAACDecode);
+					CSV_SCOPED_TIMING_STAT(ElectraPlayer, AudioAACDecode);
 
 					// Can't create a decoder based on dummy data.
 					if (!pData->bIsDummyData)
@@ -1003,7 +1014,8 @@ namespace Electra
 						// Need a new output buffer?
 						if (CurrentRenderOutputBuffer == nullptr)
 						{
-							CSV_SCOPED_TIMING_STAT(ElectraPlayer, FAudioDecoderAAC_Worker);
+							SCOPE_CYCLE_COUNTER(STAT_ElectraPlayer_AudioAACDecode);
+							CSV_SCOPED_TIMING_STAT(ElectraPlayer, AudioAACDecode);
 
 							UEMediaError bufResult = Renderer->AcquireBuffer(CurrentRenderOutputBuffer, 0, BufferAcquireOptions);
 							check(bufResult == UEMEDIA_ERROR_OK || bufResult == UEMEDIA_ERROR_INSUFFICIENT_DATA);
@@ -1018,7 +1030,8 @@ namespace Electra
 						NotifyReadyBufferListener(bHaveAvailSmpBlk);
 						if (bHaveAvailSmpBlk)
 						{
-							CSV_SCOPED_TIMING_STAT(ElectraPlayer, FAudioDecoderAAC_Worker);
+							SCOPE_CYCLE_COUNTER(STAT_ElectraPlayer_AudioAACConvertOutput);
+							CSV_SCOPED_TIMING_STAT(ElectraPlayer, AudioAACConvertOutput);
 
 							// Clear to silence
 							SIZE_T CurrentRenderOutputBufferSize = (SIZE_T)CurrentRenderOutputBuffer->GetBufferProperties().GetValue("size").GetInt64();
diff --git a/Engine/Plugins/Media/ElectraPlayer/Source/ElectraPlayerRuntime/Private/Runtime/Decoder/DX/VideoDecoderH264_DX.cpp b/Engine/Plugins/Media/ElectraPlayer/Source/ElectraPlayerRuntime/Private/Runtime/Decoder/DX/VideoDecoderH264_DX.cpp
--- a/Engine/Plugins/Media/ElectraPlayer/Source/ElectraPlayerRuntime/Private/Runtime/Decoder/DX/VideoDecoderH264_DX.cpp
+++ b/Engine/Plugins/Media/ElectraPlayer/Source/ElectraPlayerRuntime/Private/Runtime/Decoder/DX/VideoDecoderH264_DX.cpp
@@ -58,6 +58,11 @@ static const int32 MaxHWDecodeH = 1088;
 
 }
 
+
+DECLARE_CYCLE_STAT(TEXT("FVideoDecoderH264::Decode()"), STAT_ElectraPlayer_VideoH264Decode, STATGROUP_ElectraPlayer);
+DECLARE_CYCLE_STAT(TEXT("FVideoDecoderH264::ConvertOutput()"), STAT_ElectraPlayer_VideoH264ConvertOutput, STATGROUP_ElectraPlayer);
+
+
 namespace Electra {
 
 IVideoDecoderH264::FSystemConfiguration			FVideoDecoderH264::SystemConfig;
@@ -867,7 +872,10 @@ bool FVideoDecoderH264::AcquireOutputBuffer()
 
 bool FVideoDecoderH264::ConvertDecodedImage(const TRefCountPtr<IMFSample>& DecodedOutputSample)
 {
-	FParamDict* OutputBufferSampleProperties = new FParamDict();
+	SCOPE_CYCLE_COUNTER(STAT_ElectraPlayer_VideoH264ConvertOutput);
+	CSV_SCOPED_TIMING_STAT(ElectraPlayer, VideoH264ConvertOutput);
+
+	TUniquePtr<FParamDict> OutputBufferSampleProperties = MakeUnique<FParamDict>();
 	if (!OutputBufferSampleProperties)
 	{
 		return false;
@@ -959,14 +967,14 @@ bool FVideoDecoderH264::ConvertDecodedImage(const TRefCountPtr<IMFSample>& Decod
 
 	if (CurrentRenderOutputBuffer != nullptr)
 	{
-		if (!SetupDecodeOutputData(FIntPoint(videoArea.Area.cx, videoArea.Area.cy), DecodedOutputSample, OutputBufferSampleProperties))
+		if (!SetupDecodeOutputData(FIntPoint(videoArea.Area.cx, videoArea.Area.cy), DecodedOutputSample, OutputBufferSampleProperties.Get()))
 		{
 			return false;
 		}
 
 		Renderer->ReturnBuffer(CurrentRenderOutputBuffer, true, *OutputBufferSampleProperties);
 		CurrentRenderOutputBuffer = nullptr;
-
+		OutputBufferSampleProperties.Release();
 	}
 	return true;
 }
@@ -1010,7 +1018,8 @@ bool FVideoDecoderH264::Decode(FAccessUnit* AccessUnit, bool bResolutionChanged)
 		DWORD					dwSize = 0;
 		LONGLONG				llSampleTime = 0;
 
-		CSV_SCOPED_TIMING_STAT(ElectraPlayer, FVideoDecoderH264_Worker);
+		SCOPE_CYCLE_COUNTER(STAT_ElectraPlayer_VideoH264Decode);
+		CSV_SCOPED_TIMING_STAT(ElectraPlayer, VideoH264Decode);
 
 		VERIFY_HR(MFCreateSample(InputSample.GetInitReference()), "Failed to create video decoder input sample", ERRCODE_INTERNAL_COULD_NOT_CREATE_INPUT_SAMPLE);
 		VERIFY_HR(MFCreateMemoryBuffer((DWORD) AccessUnit->AUSize, InputSampleBuffer.GetInitReference()), "Failed to create video decoder input sample memory buffer", ERRCODE_INTERNAL_COULD_NOT_CREATE_INPUTBUFFER);
@@ -1061,7 +1070,8 @@ bool FVideoDecoderH264::Decode(FAccessUnit* AccessUnit, bool bResolutionChanged)
 
 		if (!CurrentDecoderOutputBuffer.IsValid())
 		{
-			CSV_SCOPED_TIMING_STAT(ElectraPlayer, FVideoDecoderH264_Worker);
+			SCOPE_CYCLE_COUNTER(STAT_ElectraPlayer_VideoH264Decode);
+			CSV_SCOPED_TIMING_STAT(ElectraPlayer, VideoH264Decode);
 
 			if (!CreateDecoderOutputBuffer())
 			{
@@ -1079,6 +1089,8 @@ bool FVideoDecoderH264::Decode(FAccessUnit* AccessUnit, bool bResolutionChanged)
 			// Draining for resolution change?
 			if (bResolutionChanged || bDecoderFlushPending)
 			{
+				SCOPE_CYCLE_COUNTER(STAT_ElectraPlayer_VideoH264Decode);
+				CSV_SCOPED_TIMING_STAT(ElectraPlayer, VideoH264Decode);
 				bDecoderFlushPending = false;
 				CurrentSampleInfo = NewSampleInfo;
 				SetupBufferAcquisitionProperties();
@@ -1094,6 +1106,8 @@ bool FVideoDecoderH264::Decode(FAccessUnit* AccessUnit, bool bResolutionChanged)
 			{
 				// Yes. This means we have received all pending output and are done now.
 				// Issue a flush for completeness sake.
+				SCOPE_CYCLE_COUNTER(STAT_ElectraPlayer_VideoH264Decode);
+				CSV_SCOPED_TIMING_STAT(ElectraPlayer, VideoH264Decode);
 				VERIFY_HR(DecoderTransform->ProcessMessage(MFT_MESSAGE_COMMAND_FLUSH, 0), "Failed to issue video decoder flush command", ERRCODE_INTERNAL_COULD_NOT_SET_DECODER_FLUSHCOMMAND);
 				// And start over.
 				VERIFY_HR(DecoderTransform->ProcessMessage(MFT_MESSAGE_NOTIFY_BEGIN_STREAMING, 0), "Failed to set video decoder stream begin", ERRCODE_INTERNAL_COULD_NOT_SET_DECODER_BEGIN);
@@ -1104,8 +1118,6 @@ bool FVideoDecoderH264::Decode(FAccessUnit* AccessUnit, bool bResolutionChanged)
 			}
 			else if (InputSample.IsValid())
 			{
-				CSV_SCOPED_TIMING_STAT(ElectraPlayer, FVideoDecoderH264_Worker);
-
 				VERIFY_HR(DecoderTransform->ProcessInput(0, InputSample.GetReference(), 0), "Failed to process video decoder input", ERRCODE_INTERNAL_COULD_NOT_PROCESS_INPUT);
 				// Used this sample. Have no further input data for now, but continue processing to produce output if possible.
 				InputSample = nullptr;
@@ -1119,6 +1131,8 @@ bool FVideoDecoderH264::Decode(FAccessUnit* AccessUnit, bool bResolutionChanged)
 		}
 		else if (res == MF_E_TRANSFORM_STREAM_CHANGE)
 		{
+			SCOPE_CYCLE_COUNTER(STAT_ElectraPlayer_VideoH264Decode);
+			CSV_SCOPED_TIMING_STAT(ElectraPlayer, VideoH264Decode);
 			//if((CurrentDecoderOutputBuffer->mOutputBuffer.dwFlags & MFT_OUTPUT_DATA_BUFFER_FORMAT_CHANGE) != 0)
 			{
 				// Update output type.
@@ -1146,13 +1160,12 @@ bool FVideoDecoderH264::Decode(FAccessUnit* AccessUnit, bool bResolutionChanged)
 					}
 				}
 
-				CSV_SCOPED_TIMING_STAT(ElectraPlayer, FVideoDecoderH264_Worker);
-
 				// Check if to be flushed now. This may have resulted in AcquireOutputBuffer() to return with no buffer!
 				if (FlushDecoderSignal.IsSignaled())
 				{
 					break;
 				}
+
 				if (!ConvertDecodedImage(DecodedOutputSample))
 				{
 					return false;
@@ -1305,7 +1318,8 @@ void FVideoDecoderH264::WorkerThread()
 		// Notify optional buffer listener that we will now be needing an AU for our input buffer.
 		if (!bError && InputBufferListener && AccessUnits.Num() == 0)
 		{
-			CSV_SCOPED_TIMING_STAT(ElectraPlayer, FVideoDecoderH264_Worker);
+			SCOPE_CYCLE_COUNTER(STAT_ElectraPlayer_VideoH264Decode);
+			CSV_SCOPED_TIMING_STAT(ElectraPlayer, VideoH264Decode);
 
 			FAccessUnitBufferInfo	sin;
 			IAccessUnitBufferListener::FBufferStats	stats;
@@ -1347,7 +1361,8 @@ void FVideoDecoderH264::WorkerThread()
 				if (!CurrentAccessUnit->bIsDummyData)
 				{
 					{
-					CSV_SCOPED_TIMING_STAT(ElectraPlayer, FVideoDecoderH264_Worker);
+					SCOPE_CYCLE_COUNTER(STAT_ElectraPlayer_VideoH264Decode);
+					CSV_SCOPED_TIMING_STAT(ElectraPlayer, VideoH264Decode);
 					PrepareAU(CurrentAccessUnit);
 					}
 
@@ -1365,7 +1380,8 @@ void FVideoDecoderH264::WorkerThread()
 
 						if (!bHaveDecoder)
 						{
-							CSV_SCOPED_TIMING_STAT(ElectraPlayer, FVideoDecoderH264_Worker);
+							SCOPE_CYCLE_COUNTER(STAT_ElectraPlayer_VideoH264Decode);
+							CSV_SCOPED_TIMING_STAT(ElectraPlayer, VideoH264Decode);
 							DecoderCreate();
 							bNeedInitialReconfig = true;
 						}
@@ -1439,7 +1455,8 @@ void FVideoDecoderH264::WorkerThread()
 		// Flush?
 		if (FlushDecoderSignal.IsSignaled())
 		{
-			CSV_SCOPED_TIMING_STAT(ElectraPlayer, FVideoDecoderH264_Worker);
+			SCOPE_CYCLE_COUNTER(STAT_ElectraPlayer_VideoH264Decode);
+			CSV_SCOPED_TIMING_STAT(ElectraPlayer, VideoH264Decode);
 
 			// Destroy and re-create.
 			InternalDecoderDestroy();
diff --git a/Engine/Plugins/Media/ElectraPlayer/Source/ElectraPlayerRuntime/Private/Runtime/Player/AdaptiveStreamingPlayer.cpp b/Engine/Plugins/Media/ElectraPlayer/Source/ElectraPlayerRuntime/Private/Runtime/Player/AdaptiveStreamingPlayer.cpp
--- a/Engine/Plugins/Media/ElectraPlayer/Source/ElectraPlayerRuntime/Private/Runtime/Player/AdaptiveStreamingPlayer.cpp
+++ b/Engine/Plugins/Media/ElectraPlayer/Source/ElectraPlayerRuntime/Private/Runtime/Player/AdaptiveStreamingPlayer.cpp
@@ -15,6 +15,9 @@
 #include "HAL/LowLevelMemTracker.h"
 
 
+DECLARE_CYCLE_STAT(TEXT("FAdaptiveStreamingPlayer::WorkerThread"), STAT_ElectraPlayer_AdaptiveWorker, STATGROUP_ElectraPlayer);
+DECLARE_CYCLE_STAT(TEXT("FAdaptiveStreamingPlayer::EventThread"), STAT_ElectraPlayer_EventWorker, STATGROUP_ElectraPlayer);
+
 
 namespace Electra
 {
@@ -1193,6 +1196,7 @@ void FAdaptiveStreamingPlayer::WorkerThreadFN()
 		FWorkerThread::FMessage msg;
 		if (WorkerThread.WorkMessages.ReceiveMessage(msg, 1000 * 20))
 		{
+			SCOPE_CYCLE_COUNTER(STAT_ElectraPlayer_AdaptiveWorker);
 			CSV_SCOPED_TIMING_STAT(ElectraPlayer, AdaptiveStreamingPlayer_Worker);
 
 			// While closed ignore all but the quit message.
@@ -1446,6 +1450,7 @@ void FAdaptiveStreamingPlayer::WorkerThreadFN()
 
 		if (!bIsClosing)
 		{
+			SCOPE_CYCLE_COUNTER(STAT_ElectraPlayer_AdaptiveWorker);
 			CSV_SCOPED_TIMING_STAT(ElectraPlayer, AdaptiveStreamingPlayer_Worker);
 
 			// Update the play position we return to the interested caller.
@@ -3859,6 +3864,7 @@ void FAdaptiveStreamingPlayerEventHandler::StopWorkerThread()
 
 void FAdaptiveStreamingPlayerEventHandler::WorkerThreadFN()
 {
+	LLM_SCOPE(ELLMTag::ElectraPlayer);
 	while(1)
 	{
 		TSharedPtrTS<FMetricEvent> pEvt = EventQueue.ReceiveMessage();
@@ -3867,93 +3873,98 @@ void FAdaptiveStreamingPlayerEventHandler::WorkerThreadFN()
 			break;
 		}
 
-		// Get the player that sends the event. If it no longer exists ignore the event.
-		TSharedPtr<FAdaptiveStreamingPlayer, ESPMode::ThreadSafe>	Player = pEvt->Player.Pin();
-		if (Player.IsValid())
 		{
-			// Call listeners under lock. They are not allowed to add or remove themselves or other listeners.
-			Player->LockMetricsReceivers();
-			const TArray<IAdaptiveStreamingPlayerMetrics*, TInlineAllocator<4>>& MetricListeners = Player->GetMetricsReceivers();
-			for(int32 i=0; i<MetricListeners.Num(); ++i)
+			SCOPE_CYCLE_COUNTER(STAT_ElectraPlayer_EventWorker);
+			CSV_SCOPED_TIMING_STAT(ElectraPlayer, AdaptiveStreamingPlayer_Event);
+
+			// Get the player that sends the event. If it no longer exists ignore the event.
+			TSharedPtr<FAdaptiveStreamingPlayer, ESPMode::ThreadSafe>	Player = pEvt->Player.Pin();
+			if (Player.IsValid())
 			{
-				switch(pEvt->Type)
+				// Call listeners under lock. They are not allowed to add or remove themselves or other listeners.
+				Player->LockMetricsReceivers();
+				const TArray<IAdaptiveStreamingPlayerMetrics*, TInlineAllocator<4>>& MetricListeners = Player->GetMetricsReceivers();
+				for(int32 i=0; i<MetricListeners.Num(); ++i)
 				{
-					case FMetricEvent::EType::OpenSource:
-						MetricListeners[i]->ReportOpenSource(pEvt->Param.URL);
-						break;
-					case FMetricEvent::EType::ReceivedMasterPlaylist:
-						MetricListeners[i]->ReportReceivedMasterPlaylist(pEvt->Param.URL);
-						break;
-					case FMetricEvent::EType::ReceivedPlaylists:
-						MetricListeners[i]->ReportReceivedPlaylists();
-						break;
-					case FMetricEvent::EType::TracksChanged:
-						MetricListeners[i]->ReportTracksChanged();
-						break;
-					case FMetricEvent::EType::BufferingStart:
-						MetricListeners[i]->ReportBufferingStart(pEvt->Param.BufferingReason);
-						break;
-					case FMetricEvent::EType::BufferingEnd:
-						MetricListeners[i]->ReportBufferingEnd(pEvt->Param.BufferingReason);
-						break;
-					case FMetricEvent::EType::Bandwidth:
-						MetricListeners[i]->ReportBandwidth(pEvt->Param.Bandwidth.EffectiveBps, pEvt->Param.Bandwidth.ThroughputBps, pEvt->Param.Bandwidth.Latency);
-						break;
-					case FMetricEvent::EType::BufferUtilization:
-						MetricListeners[i]->ReportBufferUtilization(pEvt->Param.BufferStats);
-						break;
-					case FMetricEvent::EType::PlaylistDownload:
-						MetricListeners[i]->ReportPlaylistDownload(pEvt->Param.PlaylistStats);
-						break;
-					case FMetricEvent::EType::SegmentDownload:
-						MetricListeners[i]->ReportSegmentDownload(pEvt->Param.SegmentStats);
-						break;
-					case FMetricEvent::EType::LicenseKey:
-						MetricListeners[i]->ReportLicenseKey(pEvt->Param.LicenseKeyStats);
-						break;
-					case FMetricEvent::EType::DataAvailabilityChange:
-						MetricListeners[i]->ReportDataAvailabilityChange(pEvt->Param.DataAvailability);
-						break;
-					case FMetricEvent::EType::VideoQualityChange:
-						MetricListeners[i]->ReportVideoQualityChange(pEvt->Param.QualityChange.NewBitrate, pEvt->Param.QualityChange.PrevBitrate, pEvt->Param.QualityChange.bIsDrastic);
-						break;
-					case FMetricEvent::EType::PrerollStart:
-						MetricListeners[i]->ReportPrerollStart();
-						break;
-					case FMetricEvent::EType::PrerollEnd:
-						MetricListeners[i]->ReportPrerollEnd();
-						break;
-					case FMetricEvent::EType::PlaybackStart:
-						MetricListeners[i]->ReportPlaybackStart();
-						break;
-					case FMetricEvent::EType::PlaybackPaused:
-						MetricListeners[i]->ReportPlaybackPaused();
-						break;
-					case FMetricEvent::EType::PlaybackResumed:
-						MetricListeners[i]->ReportPlaybackResumed();
-						break;
-					case FMetricEvent::EType::PlaybackEnded:
-						MetricListeners[i]->ReportPlaybackEnded();
-						break;
-					case FMetricEvent::EType::PlaybackJumped:
-						MetricListeners[i]->ReportJumpInPlayPosition(pEvt->Param.TimeJump.ToNewTime, pEvt->Param.TimeJump.FromTime, pEvt->Param.TimeJump.Reason);
-						break;
-					case FMetricEvent::EType::PlaybackStopped:
-						MetricListeners[i]->ReportPlaybackStopped();
-						break;
-					case FMetricEvent::EType::Errored:
-						MetricListeners[i]->ReportError(pEvt->Param.ErrorDetail.GetPrintable());
-						break;
-					case FMetricEvent::EType::LogMessage:
-						MetricListeners[i]->ReportLogMessage(pEvt->Param.LogMessage.Level, pEvt->Param.LogMessage.Message, pEvt->Param.LogMessage.AtMillis);
-						break;
+					switch(pEvt->Type)
+					{
+						case FMetricEvent::EType::OpenSource:
+							MetricListeners[i]->ReportOpenSource(pEvt->Param.URL);
+							break;
+						case FMetricEvent::EType::ReceivedMasterPlaylist:
+							MetricListeners[i]->ReportReceivedMasterPlaylist(pEvt->Param.URL);
+							break;
+						case FMetricEvent::EType::ReceivedPlaylists:
+							MetricListeners[i]->ReportReceivedPlaylists();
+							break;
+						case FMetricEvent::EType::TracksChanged:
+							MetricListeners[i]->ReportTracksChanged();
+							break;
+						case FMetricEvent::EType::BufferingStart:
+							MetricListeners[i]->ReportBufferingStart(pEvt->Param.BufferingReason);
+							break;
+						case FMetricEvent::EType::BufferingEnd:
+							MetricListeners[i]->ReportBufferingEnd(pEvt->Param.BufferingReason);
+							break;
+						case FMetricEvent::EType::Bandwidth:
+							MetricListeners[i]->ReportBandwidth(pEvt->Param.Bandwidth.EffectiveBps, pEvt->Param.Bandwidth.ThroughputBps, pEvt->Param.Bandwidth.Latency);
+							break;
+						case FMetricEvent::EType::BufferUtilization:
+							MetricListeners[i]->ReportBufferUtilization(pEvt->Param.BufferStats);
+							break;
+						case FMetricEvent::EType::PlaylistDownload:
+							MetricListeners[i]->ReportPlaylistDownload(pEvt->Param.PlaylistStats);
+							break;
+						case FMetricEvent::EType::SegmentDownload:
+							MetricListeners[i]->ReportSegmentDownload(pEvt->Param.SegmentStats);
+							break;
+						case FMetricEvent::EType::LicenseKey:
+							MetricListeners[i]->ReportLicenseKey(pEvt->Param.LicenseKeyStats);
+							break;
+						case FMetricEvent::EType::DataAvailabilityChange:
+							MetricListeners[i]->ReportDataAvailabilityChange(pEvt->Param.DataAvailability);
+							break;
+						case FMetricEvent::EType::VideoQualityChange:
+							MetricListeners[i]->ReportVideoQualityChange(pEvt->Param.QualityChange.NewBitrate, pEvt->Param.QualityChange.PrevBitrate, pEvt->Param.QualityChange.bIsDrastic);
+							break;
+						case FMetricEvent::EType::PrerollStart:
+							MetricListeners[i]->ReportPrerollStart();
+							break;
+						case FMetricEvent::EType::PrerollEnd:
+							MetricListeners[i]->ReportPrerollEnd();
+							break;
+						case FMetricEvent::EType::PlaybackStart:
+							MetricListeners[i]->ReportPlaybackStart();
+							break;
+						case FMetricEvent::EType::PlaybackPaused:
+							MetricListeners[i]->ReportPlaybackPaused();
+							break;
+						case FMetricEvent::EType::PlaybackResumed:
+							MetricListeners[i]->ReportPlaybackResumed();
+							break;
+						case FMetricEvent::EType::PlaybackEnded:
+							MetricListeners[i]->ReportPlaybackEnded();
+							break;
+						case FMetricEvent::EType::PlaybackJumped:
+							MetricListeners[i]->ReportJumpInPlayPosition(pEvt->Param.TimeJump.ToNewTime, pEvt->Param.TimeJump.FromTime, pEvt->Param.TimeJump.Reason);
+							break;
+						case FMetricEvent::EType::PlaybackStopped:
+							MetricListeners[i]->ReportPlaybackStopped();
+							break;
+						case FMetricEvent::EType::Errored:
+							MetricListeners[i]->ReportError(pEvt->Param.ErrorDetail.GetPrintable());
+							break;
+						case FMetricEvent::EType::LogMessage:
+							MetricListeners[i]->ReportLogMessage(pEvt->Param.LogMessage.Level, pEvt->Param.LogMessage.Message, pEvt->Param.LogMessage.AtMillis);
+							break;
+					}
 				}
+				Player->UnlockMetricsReceivers();
+			}
+			if (pEvt->EventSignal)
+			{
+				pEvt->EventSignal->Signal();
 			}
-			Player->UnlockMetricsReceivers();
-		}
-		if (pEvt->EventSignal)
-		{
-			pEvt->EventSignal->Signal();
 		}
 	}
 }
diff --git a/Engine/Plugins/Media/ElectraPlayer/Source/ElectraPlayerRuntime/Private/Runtime/Player/DASH/ManifestBuilderDASH.cpp b/Engine/Plugins/Media/ElectraPlayer/Source/ElectraPlayerRuntime/Private/Runtime/Player/DASH/ManifestBuilderDASH.cpp
--- a/Engine/Plugins/Media/ElectraPlayer/Source/ElectraPlayerRuntime/Private/Runtime/Player/DASH/ManifestBuilderDASH.cpp
+++ b/Engine/Plugins/Media/ElectraPlayer/Source/ElectraPlayerRuntime/Private/Runtime/Player/DASH/ManifestBuilderDASH.cpp
@@ -1,6 +1,7 @@
 // Copyright Epic Games, Inc. All Rights Reserved.
 
 #include "PlayerCore.h"
+#include "ElectraPlayerPrivate.h"
 #include "Player/Manifest.h"
 #include "Player/PlaybackTimeline.h"
 
@@ -43,6 +44,9 @@
 #define ERRCODE_DASH_MPD_BUILDER_BAD_PERIOD_DURATION				203
 
 
+DECLARE_CYCLE_STAT(TEXT("FManifestDASHInternal::Build"), STAT_ElectraPlayer_FManifestDASHInternal_Build, STATGROUP_ElectraPlayer);
+DECLARE_CYCLE_STAT(TEXT("FManifestDASHInternal::XLink"), STAT_ElectraPlayer_FManifestDASHInternal_XLink, STATGROUP_ElectraPlayer);
+
 
 namespace Electra
 {
@@ -1027,6 +1031,9 @@ FErrorDetail FManifestDASHInternal::Build(IPlayerSessionServices* InPlayerSessio
 
 FErrorDetail FManifestDASHInternal::BuildAfterInitialRemoteElementDownload()
 {
+	SCOPE_CYCLE_COUNTER(STAT_ElectraPlayer_FManifestDASHInternal_Build);
+	CSV_SCOPED_TIMING_STAT(ElectraPlayer, FManifestDASHInternal_Build);
+
 	FErrorDetail Error;
 
 	bool bWarnedPresentationDuration = false;
@@ -1221,6 +1228,10 @@ void FManifestDASHInternal::PreparePeriodAdaptationSets(TSharedPtrTS<FPeriod> Pe
 		{
 			return;
 		}
+
+		SCOPE_CYCLE_COUNTER(STAT_ElectraPlayer_FManifestDASHInternal_Build);
+		CSV_SCOPED_TIMING_STAT(ElectraPlayer, FManifestDASHInternal_Build);
+
 		Period->AdaptationSets.Empty();
 		const TArray<TSharedPtrTS<FDashMPD_AdaptationSetType>>& MPDAdaptationSets = MPDPeriod->GetAdaptationSets();
 		for(int32 nAdapt=0, nAdaptMax=MPDAdaptationSets.Num(); nAdapt<nAdaptMax; ++nAdapt)
@@ -1715,6 +1726,9 @@ void FManifestDASHInternal::SendEventsFromAllPeriodEventStreams(TSharedPtrTS<FPe
  */
 FErrorDetail FManifestDASHInternal::ResolveInitialRemoteElementRequest(TSharedPtrTS<FMPDLoadRequestDASH> RequestResponse, FString XMLResponse, bool bSuccess)
 {
+	SCOPE_CYCLE_COUNTER(STAT_ElectraPlayer_FManifestDASHInternal_XLink);
+	CSV_SCOPED_TIMING_STAT(ElectraPlayer, FManifestDASHInternal_XLink);
+
 	// Because this is intended solely for initial MPD entities we do not need to worry about anyone accessing our internal structures
 	// while we update them. The player proper has not been informed yet that the initial manifest is ready for use.
 	FErrorDetail Error;
diff --git a/Engine/Plugins/Media/ElectraPlayer/Source/ElectraPlayerRuntime/Private/Runtime/Player/DASH/ManifestDASH.cpp b/Engine/Plugins/Media/ElectraPlayer/Source/ElectraPlayerRuntime/Private/Runtime/Player/DASH/ManifestDASH.cpp
--- a/Engine/Plugins/Media/ElectraPlayer/Source/ElectraPlayerRuntime/Private/Runtime/Player/DASH/ManifestDASH.cpp
+++ b/Engine/Plugins/Media/ElectraPlayer/Source/ElectraPlayerRuntime/Private/Runtime/Player/DASH/ManifestDASH.cpp
@@ -1,6 +1,7 @@
 // Copyright Epic Games, Inc. All Rights Reserved.
 
 #include "PlayerCore.h"
+#include "ElectraPlayerPrivate.h"
 #include "ManifestDASH.h"
 #include "ManifestBuilderDASH.h"
 #include "PlaylistReaderDASH.h"
@@ -19,6 +20,9 @@
 #include "Player/DRM/DRMManager.h"
 
 
+DECLARE_CYCLE_STAT(TEXT("FRepresentation::FindSegment"), STAT_ElectraPlayer_DASH_FindSegment, STATGROUP_ElectraPlayer);
+DECLARE_CYCLE_STAT(TEXT("FDASHPlayPeriod::GetSegmentInformation"), STAT_ElectraPlayer_DASH_GetSegmentInformation, STATGROUP_ElectraPlayer);
+
 namespace Electra
 {
 
@@ -1279,6 +1283,9 @@ void FDASHPlayPeriod::IncreaseSegmentFetchDelay(const FTimeValue& IncreaseAmount
 
 void FDASHPlayPeriod::GetSegmentInformation(TArray<FSegmentInformation>& OutSegmentInformation, FTimeValue& OutAverageSegmentDuration, TSharedPtrTS<const IStreamSegment> CurrentSegment, const FTimeValue& LookAheadTime, const FString& AdaptationSetID, const FString& RepresentationID)
 {
+	SCOPE_CYCLE_COUNTER(STAT_ElectraPlayer_DASH_GetSegmentInformation);
+	CSV_SCOPED_TIMING_STAT(ElectraPlayer, DASH_GetSegmentInformation);
+
 	OutSegmentInformation.Empty();
 	OutAverageSegmentDuration.SetToInvalid();
 
@@ -1802,6 +1809,9 @@ void FManifestDASHInternal::FRepresentation::CollectInbandEventStreams(IPlayerSe
 
 FManifestDASHInternal::FRepresentation::ESearchResult FManifestDASHInternal::FRepresentation::FindSegment(IPlayerSessionServices* InPlayerSessionServices, FSegmentInformation& OutSegmentInfo, TArray<TWeakPtrTS<FMPDLoadRequestDASH>>& OutRemoteElementLoadRequests, const FSegmentSearchOption& InSearchOptions)
 {
+	SCOPE_CYCLE_COUNTER(STAT_ElectraPlayer_DASH_FindSegment);
+	CSV_SCOPED_TIMING_STAT(ElectraPlayer, DASH_FindSegment);
+
 	/*
 		Note: We use the DASH-IF-IOP specification and timing model. This is more strict than the general DASH standard and removes ambiguities
 		      and otherwise conflicting information.
diff --git a/Engine/Plugins/Media/ElectraPlayer/Source/ElectraPlayerRuntime/Private/Runtime/Player/DASH/ManifestParserDASH.cpp b/Engine/Plugins/Media/ElectraPlayer/Source/ElectraPlayerRuntime/Private/Runtime/Player/DASH/ManifestParserDASH.cpp
--- a/Engine/Plugins/Media/ElectraPlayer/Source/ElectraPlayerRuntime/Private/Runtime/Player/DASH/ManifestParserDASH.cpp
+++ b/Engine/Plugins/Media/ElectraPlayer/Source/ElectraPlayerRuntime/Private/Runtime/Player/DASH/ManifestParserDASH.cpp
@@ -1,6 +1,7 @@
 // Copyright Epic Games, Inc. All Rights Reserved.
 
 #include "PlayerCore.h"
+#include "ElectraPlayerPrivate.h"
 #include "Player/Manifest.h"
 #include "Player/PlaybackTimeline.h"
 
@@ -28,6 +29,8 @@
 #define ERRCODE_DASH_MPD_PARSER_BAD_TIME_VALUE						200
 #define ERRCODE_DASH_MPD_PARSER_BAD_ELEMENT_VALUE					201
 
+DECLARE_CYCLE_STAT(TEXT("FManifestParserDASH::BuildFromMPD"), STAT_ElectraPlayer_DASH_BuildFromMPD, STATGROUP_ElectraPlayer);
+
 
 namespace Electra
 {
@@ -494,6 +497,9 @@ void FManifestParserDASH::LogMessage(IInfoLog::ELevel Level, const FString& Mess
 
 FErrorDetail FManifestParserDASH::BuildFromMPD(TCHAR* InOutMPDXML, const TCHAR* InExpectedRootElement)
 {
+	SCOPE_CYCLE_COUNTER(STAT_ElectraPlayer_DASH_BuildFromMPD);
+	CSV_SCOPED_TIMING_STAT(ElectraPlayer, DASH_BuildFromMPD);
+
 	FText ErrorMsg;
 	int32 ErrorLine = 0;
 	const TCHAR* NoFile = TEXT("");
diff --git a/Engine/Plugins/Media/ElectraPlayer/Source/ElectraPlayerRuntime/Private/Runtime/Player/DASH/PlaylistReaderDASH.cpp b/Engine/Plugins/Media/ElectraPlayer/Source/ElectraPlayerRuntime/Private/Runtime/Player/DASH/PlaylistReaderDASH.cpp
--- a/Engine/Plugins/Media/ElectraPlayer/Source/ElectraPlayerRuntime/Private/Runtime/Player/DASH/PlaylistReaderDASH.cpp
+++ b/Engine/Plugins/Media/ElectraPlayer/Source/ElectraPlayerRuntime/Private/Runtime/Player/DASH/PlaylistReaderDASH.cpp
@@ -29,6 +29,9 @@
 #define ERRCODE_DASH_MPD_REMOTE_ENTITY_FAILED				5
 
 
+DECLARE_CYCLE_STAT(TEXT("FPlaylistReaderDASH_WorkerThread"), STAT_ElectraPlayer_DASH_PlaylistWorker, STATGROUP_ElectraPlayer);
+
+
 //#define DO_NOT_PERFORM_CONDITIONAL_GET
 
 namespace Electra
@@ -721,7 +724,6 @@ void FPlaylistReaderDASH::TriggerTimeSynchronization()
 void FPlaylistReaderDASH::WorkerThread()
 {
 	LLM_SCOPE(ELLMTag::ElectraPlayer);
-	CSV_SCOPED_TIMING_STAT(ElectraPlayer, MPDReaderDASH_Worker);
 
 	bInbandEventByStreamType[0] = false;
 	bInbandEventByStreamType[1] = false;
@@ -754,19 +756,23 @@ void FPlaylistReaderDASH::WorkerThread()
 		{
 			break;
 		}
-		FTimeValue Now = PlayerSessionServices->GetSynchronizedUTCTime()->GetTime();
+		{
+			SCOPE_CYCLE_COUNTER(STAT_ElectraPlayer_DASH_PlaylistWorker);
+			CSV_SCOPED_TIMING_STAT(ElectraPlayer, DASH_PlaylistWorker);
+			FTimeValue Now = PlayerSessionServices->GetSynchronizedUTCTime()->GetTime();
 
-		HandleCompletedRequests(Now);
+			HandleCompletedRequests(Now);
 
-		ExecutePendingRequests(Now);
+			ExecutePendingRequests(Now);
 
-		// Check if the MPD must be updated.
-		CheckForMPDUpdate();
+			// Check if the MPD must be updated.
+			CheckForMPDUpdate();
 
-		// Time sync
-		if (NextTimeSyncTime.IsValid() && Now >= NextTimeSyncTime)
-		{
-			TriggerTimeSynchronization();
+			// Time sync
+			if (NextTimeSyncTime.IsValid() && Now >= NextTimeSyncTime)
+			{
+				TriggerTimeSynchronization();
+			}
 		}
 	}
 
diff --git a/Engine/Plugins/Media/ElectraPlayer/Source/ElectraPlayerRuntime/Private/Runtime/Player/DASH/StreamReaderFMP4DASH.cpp b/Engine/Plugins/Media/ElectraPlayer/Source/ElectraPlayerRuntime/Private/Runtime/Player/DASH/StreamReaderFMP4DASH.cpp
--- a/Engine/Plugins/Media/ElectraPlayer/Source/ElectraPlayerRuntime/Private/Runtime/Player/DASH/StreamReaderFMP4DASH.cpp
+++ b/Engine/Plugins/Media/ElectraPlayer/Source/ElectraPlayerRuntime/Private/Runtime/Player/DASH/StreamReaderFMP4DASH.cpp
@@ -21,6 +21,8 @@
 #define INTERNAL_ERROR_INIT_SEGMENT_PARSE_ERROR						2
 
 
+DECLARE_CYCLE_STAT(TEXT("FStreamReaderFMP4DASH::HandleRequest"), STAT_ElectraPlayer_DASH_StreamReader, STATGROUP_ElectraPlayer);
+
 namespace Electra
 {
 
@@ -443,6 +445,9 @@ FErrorDetail FStreamReaderFMP4DASH::FStreamHandler::GetInitSegment(TSharedPtrTS<
 		return CreateError(FString::Printf(TEXT("Init segment download error: %s"),	*ci->StatusInfo.ErrorDetail.GetMessage()), INTERNAL_ERROR_INIT_SEGMENT_DOWNLOAD_ERROR);
 	}
 
+	SCOPE_CYCLE_COUNTER(STAT_ElectraPlayer_DASH_StreamReader);
+	CSV_SCOPED_TIMING_STAT(ElectraPlayer, DASH_StreamReader);
+
 	FMP4StaticDataReader StaticDataReader;
 	StaticDataReader.SetParseData(LoadReq->Request->GetResponseBuffer());
 	TSharedPtrTS<IParserISO14496_12> Init = IParserISO14496_12::CreateParser();
@@ -482,6 +487,8 @@ FErrorDetail FStreamReaderFMP4DASH::FStreamHandler::GetInitSegment(TSharedPtrTS<
 
 void FStreamReaderFMP4DASH::FStreamHandler::CheckForInbandDASHEvents()
 {
+	SCOPE_CYCLE_COUNTER(STAT_ElectraPlayer_DASH_StreamReader);
+	CSV_SCOPED_TIMING_STAT(ElectraPlayer, DASH_StreamReader);
 	bool bHasInbandEvent = false;
 	if (!CurrentRequest->bIsEOSSegment)
 	{
@@ -505,6 +512,8 @@ void FStreamReaderFMP4DASH::FStreamHandler::CheckForInbandDASHEvents()
 
 void FStreamReaderFMP4DASH::FStreamHandler::HandleEventMessages()
 {
+	SCOPE_CYCLE_COUNTER(STAT_ElectraPlayer_DASH_StreamReader);
+	CSV_SCOPED_TIMING_STAT(ElectraPlayer, DASH_StreamReader);
 	// Are there 'emsg' boxes we need to handle?
 	if (MP4Parser->GetNumberOfEventMessages())
 	{
@@ -589,8 +598,6 @@ void FStreamReaderFMP4DASH::FStreamHandler::HandleEventMessages()
 
 void FStreamReaderFMP4DASH::FStreamHandler::HandleRequest()
 {
-//CSV_SCOPED_TIMING_STAT(ElectraPlayer, StreamReaderMP4_Worker);
-
 	// Get the request into a local shared pointer to hold on to it.
 	TSharedPtrTS<FStreamSegmentRequestFMP4DASH> Request = CurrentRequest;
 	TSharedPtr<const FPlayerLoopState, ESPMode::ThreadSafe> PlayerLoopState = MakeShared<const FPlayerLoopState, ESPMode::ThreadSafe>(Request->PlayerLoopState);
@@ -767,8 +774,12 @@ void FStreamReaderFMP4DASH::FStreamHandler::HandleRequest()
 				UEMediaError parseError = MP4Parser->ParseHeader(this, this, PlayerSessionService, MP4InitSegment.Get());
 				if (parseError == UEMEDIA_ERROR_OK)
 				{
-					Request->Segment.bSawLMSG = MP4Parser->HasBrand(IParserISO14496_12::BrandType_lmsg);
-					parseError = MP4Parser->PrepareTracks(MP4InitSegment);
+					{
+						SCOPE_CYCLE_COUNTER(STAT_ElectraPlayer_DASH_StreamReader);
+						CSV_SCOPED_TIMING_STAT(ElectraPlayer, DASH_StreamReader);
+						Request->Segment.bSawLMSG = MP4Parser->HasBrand(IParserISO14496_12::BrandType_lmsg);
+						parseError = MP4Parser->PrepareTracks(MP4InitSegment);
+					}
 					if (parseError == UEMEDIA_ERROR_OK)
 					{
 						// For the time being we only want to have a single track in the movie segments.
@@ -918,6 +929,8 @@ void FStreamReaderFMP4DASH::FStreamHandler::HandleRequest()
 											ElectraCDM::ECDMError DecryptResult = ElectraCDM::ECDMError::Failure;
 											if (Decrypter->GetState() == ElectraCDM::ECDMState::Ready)
 											{
+												SCOPE_CYCLE_COUNTER(STAT_ElectraPlayer_DASH_StreamReader);
+												CSV_SCOPED_TIMING_STAT(ElectraPlayer, DASH_StreamReader);
 												DecryptResult = Decrypter->DecryptInPlace((uint8*) AccessUnit->AUData, (int32) AccessUnit->AUSize, SampleEncryptionInfo);
 											}
 											if (DecryptResult != ElectraCDM::ECDMError::Success)
@@ -972,6 +985,8 @@ void FStreamReaderFMP4DASH::FStreamHandler::HandleRequest()
 									// Shall we pass on any AUs we already read?
 									if (bAllowEarlyEmitting)
 									{
+										SCOPE_CYCLE_COUNTER(STAT_ElectraPlayer_DASH_StreamReader);
+										CSV_SCOPED_TIMING_STAT(ElectraPlayer, DASH_StreamReader);
 										while(AccessUnitFIFO.Num() && !HasReadBeenAborted())
 										{
 											FAccessUnit* pNext = AccessUnitFIFO.FrontRef();
@@ -1251,6 +1266,8 @@ void FStreamReaderFMP4DASH::FStreamHandler::HandleRequest()
 	{
 		if (SegmentEventsFound.IndexOfByPredicate([](const TSharedPtrTS<DASH::FPlayerEvent>& This){return This->GetSchemeIdUri().Equals(DASH::Schemes::ManifestEvents::Scheme_urn_mpeg_dash_event_2012);}) == INDEX_NONE)
 		{
+			SCOPE_CYCLE_COUNTER(STAT_ElectraPlayer_DASH_StreamReader);
+			CSV_SCOPED_TIMING_STAT(ElectraPlayer, DASH_StreamReader);
 			TSharedPtrTS<DASH::FPlayerEvent> NewEvent = MakeSharedTS<DASH::FPlayerEvent>();
 			NewEvent->SetOrigin(IAdaptiveStreamingPlayerAEMSEvent::EOrigin::InbandEventStream);
 			NewEvent->SetSchemeIdUri(DASH::Schemes::ManifestEvents::Scheme_urn_mpeg_dash_event_2012);
@@ -1327,6 +1344,8 @@ int64 FStreamReaderFMP4DASH::FStreamHandler::ReadData(void* IntoBuffer, int64 Nu
 				ds.ABRState.ProgressDecision = StreamSelectorDecision;
 				if ((StreamSelectorDecision.Flags & FABRDownloadProgressDecision::EDecisionFlags::eABR_EmitPartialData) != 0)
 				{
+					SCOPE_CYCLE_COUNTER(STAT_ElectraPlayer_DASH_StreamReader);
+					CSV_SCOPED_TIMING_STAT(ElectraPlayer, DASH_StreamReader);
 					bAllowEarlyEmitting = true;
 					// Deliver all enqueued AUs right now. Unless the request also gets aborted we could be stuck
 					// in here for a while longer.
@@ -1371,6 +1390,8 @@ int64 FStreamReaderFMP4DASH::FStreamHandler::ReadData(void* IntoBuffer, int64 Nu
 		}
 		else
 		{
+			SCOPE_CYCLE_COUNTER(STAT_ElectraPlayer_DASH_StreamReader);
+			CSV_SCOPED_TIMING_STAT(ElectraPlayer, DASH_StreamReader);
 			SourceBuffer.Lock();
 			if (SourceBuffer.Num() >= ReadBuffer.ParsePos + NumBytesToRead)
 			{
diff --git a/Engine/Plugins/Media/ElectraPlayer/Source/ElectraPlayerRuntime/Private/Runtime/Player/HLS/ManifestBuilderHLS.cpp b/Engine/Plugins/Media/ElectraPlayer/Source/ElectraPlayerRuntime/Private/Runtime/Player/HLS/ManifestBuilderHLS.cpp
--- a/Engine/Plugins/Media/ElectraPlayer/Source/ElectraPlayerRuntime/Private/Runtime/Player/HLS/ManifestBuilderHLS.cpp
+++ b/Engine/Plugins/Media/ElectraPlayer/Source/ElectraPlayerRuntime/Private/Runtime/Player/HLS/ManifestBuilderHLS.cpp
@@ -1,6 +1,7 @@
 // Copyright Epic Games, Inc. All Rights Reserved.
 
 #include "PlayerCore.h"
+#include "ElectraPlayerPrivate.h"
 #include "Player/Manifest.h"
 #include "Player/PlaybackTimeline.h"
 
@@ -50,6 +51,7 @@
 #define ERRCODE_HLS_BUILDER_RENDITION_NOT_FOUND_IN_GROUP		19
 #define ERRCODE_HLS_BUILDER_REFERENCED_GROUP_NOT_DEFINED		20
 
+DECLARE_CYCLE_STAT(TEXT("FManifestBuilderHLS::Build"), STAT_ElectraPlayer_FManifestHLS_Build, STATGROUP_ElectraPlayer);
 
 namespace Electra
 {
@@ -209,6 +211,9 @@ void FManifestBuilderHLS::LogMessage(IInfoLog::ELevel Level, const FString& Mess
 
 FErrorDetail FManifestBuilderHLS::BuildFromMasterPlaylist(TSharedPtrTS<FManifestHLSInternal>& OutHLSPlaylist, const HLSPlaylistParser::FPlaylist& Playlist, const FPlaylistLoadRequestHLS& SourceRequest, const HTTP::FConnectionInfo* ConnectionInfo)
 {
+	SCOPE_CYCLE_COUNTER(STAT_ElectraPlayer_FManifestHLS_Build);
+	CSV_SCOPED_TIMING_STAT(ElectraPlayer, FManifestHLS_Build);
+
 	FErrorDetail Error;
 
 	if (Playlist.Type != HLSPlaylistParser::EPlaylistType::Master)
@@ -1275,6 +1280,9 @@ void FManifestBuilderHLS::SetVariantPlaylistFailure(TSharedPtrTS<FManifestHLSInt
 
 FErrorDetail FManifestBuilderHLS::UpdateFromVariantPlaylist(TSharedPtrTS<FManifestHLSInternal> InOutHLSPlaylist, const HLSPlaylistParser::FPlaylist& VariantPlaylist, const FPlaylistLoadRequestHLS& SourceRequest, const HTTP::FConnectionInfo* ConnectionInfo, uint32 ResponseCRC)
 {
+	SCOPE_CYCLE_COUNTER(STAT_ElectraPlayer_FManifestHLS_Build);
+	CSV_SCOPED_TIMING_STAT(ElectraPlayer, FManifestHLS_Build);
+
 	// Find the variant or rendition this playlist updates.
 	TSharedPtrTS<FManifestHLSInternal::FPlaylistBase> Playlist;
 	InOutHLSPlaylist->LockPlaylists();
diff --git a/Engine/Plugins/Media/ElectraPlayer/Source/ElectraPlayerRuntime/Private/Runtime/Player/HLS/ManifestHLS.cpp b/Engine/Plugins/Media/ElectraPlayer/Source/ElectraPlayerRuntime/Private/Runtime/Player/HLS/ManifestHLS.cpp
--- a/Engine/Plugins/Media/ElectraPlayer/Source/ElectraPlayerRuntime/Private/Runtime/Player/HLS/ManifestHLS.cpp
+++ b/Engine/Plugins/Media/ElectraPlayer/Source/ElectraPlayerRuntime/Private/Runtime/Player/HLS/ManifestHLS.cpp
@@ -1,6 +1,7 @@
 // Copyright Epic Games, Inc. All Rights Reserved.
 
 #include "PlayerCore.h"
+#include "ElectraPlayerPrivate.h"
 #include "ManifestHLS.h"
 #include "ManifestBuilderHLS.h"
 #include "PlaylistReaderHLS.h"
@@ -12,6 +13,10 @@
 #include "Utilities/URLParser.h"
 
 
+DECLARE_CYCLE_STAT(TEXT("FPlayPeriodHLS::FindSegment"), STAT_ElectraPlayer_HLS_FindSegment, STATGROUP_ElectraPlayer);
+DECLARE_CYCLE_STAT(TEXT("FPlayPeriodHLS::GetSegmentInformation"), STAT_ElectraPlayer_HLS_GetSegmentInformation, STATGROUP_ElectraPlayer);
+
+
 namespace Electra
 {
 
@@ -521,6 +526,10 @@ IManifest::FResult FPlayPeriodHLS::GetMediaStreamForID(TSharedPtrTS<FManifestHLS
  */
 IManifest::FResult FPlayPeriodHLS::FindSegment(TSharedPtrTS<FStreamSegmentRequestHLSfmp4>& OutRequest, TSharedPtrTS<FManifestHLSInternal::FPlaylistBase> InPlaylist, TSharedPtrTS<FManifestHLSInternal::FMediaStream> InStream, uint32 StreamUniqueID, EStreamType StreamType, const FPlayPeriodHLS::FSegSearchParam& SearchParam, IManifest::ESearchType SearchType)
 {
+	SCOPE_CYCLE_COUNTER(STAT_ElectraPlayer_HLS_FindSegment);
+	CSV_SCOPED_TIMING_STAT(ElectraPlayer, HLS_FindSegment);
+
+
 	// VOD or EVENT playlist?
 ///////	if (InStream->PlaylistType == FManifestHLSInternal::FMediaStream::EPlaylistType::VOD || InStream->PlaylistType == FManifestHLSInternal::FMediaStream::EPlaylistType::Event)
 
@@ -1195,6 +1204,9 @@ void FPlayPeriodHLS::RefreshBlacklistState()
  */
 void FPlayPeriodHLS::GetSegmentInformation(TArray<FSegmentInformation>& OutSegmentInformation, FTimeValue& OutAverageSegmentDuration, TSharedPtrTS<const IStreamSegment> InCurrentSegment, const FTimeValue& LookAheadTime, const FString& AdaptationSetID, const FString& RepresentationID)
 {
+	SCOPE_CYCLE_COUNTER(STAT_ElectraPlayer_HLS_GetSegmentInformation);
+	CSV_SCOPED_TIMING_STAT(ElectraPlayer, HLS_GetSegmentInformation);
+
 	OutSegmentInformation.Empty();
 	OutAverageSegmentDuration.SetToInvalid();
 
diff --git a/Engine/Plugins/Media/ElectraPlayer/Source/ElectraPlayerRuntime/Private/Runtime/Player/HLS/PlaylistReaderHLS.cpp b/Engine/Plugins/Media/ElectraPlayer/Source/ElectraPlayerRuntime/Private/Runtime/Player/HLS/PlaylistReaderHLS.cpp
--- a/Engine/Plugins/Media/ElectraPlayer/Source/ElectraPlayerRuntime/Private/Runtime/Player/HLS/PlaylistReaderHLS.cpp
+++ b/Engine/Plugins/Media/ElectraPlayer/Source/ElectraPlayerRuntime/Private/Runtime/Player/HLS/PlaylistReaderHLS.cpp
@@ -24,6 +24,7 @@
 #define ERRCODE_HLS_VARIANT_PLAYLIST_DOWNLOAD_FAILED		5
 #define ERRCODE_HLS_VARIANT_PLAYLIST_PARSING_FAILED			6
 
+DECLARE_CYCLE_STAT(TEXT("FPlaylistReaderHLS_WorkerThread"), STAT_ElectraPlayer_HLS_PlaylistWorker, STATGROUP_ElectraPlayer);
 
 
 namespace Electra
@@ -457,7 +458,6 @@ void FPlaylistReaderHLS::CheckForPlaylistUpdate(const FTimeValue& TimeNow)
 void FPlaylistReaderHLS::WorkerThread()
 {
 	LLM_SCOPE(ELLMTag::ElectraPlayer);
-	CSV_SCOPED_TIMING_STAT(ElectraPlayer, PlaylistReaderHLS_Worker);
 
 	Builder = IManifestBuilderHLS::Create(PlayerSessionServices);
 
@@ -475,21 +475,27 @@ void FPlaylistReaderHLS::WorkerThread()
 		{
 			break;
 		}
-		FTimeValue Now = PlayerSessionServices->GetSynchronizedUTCTime()->GetTime();
 
-		// Check for completed static resource requests.
-		// NOTE: This needs to be done first in order to move the completed ones being moved onto the completed list to be
-		//       handled immediately in the following HandleCompletedPlaylistDownloads() call!!
-		HandleStaticRequestCompletions(Now);
+		{
+			SCOPE_CYCLE_COUNTER(STAT_ElectraPlayer_HLS_PlaylistWorker);
+			CSV_SCOPED_TIMING_STAT(ElectraPlayer, HLS_PlaylistWorker);
+
+			FTimeValue Now = PlayerSessionServices->GetSynchronizedUTCTime()->GetTime();
 
-		// Check completed playlist downloads
-		HandleCompletedPlaylistDownloads(Now);
+			// Check for completed static resource requests.
+			// NOTE: This needs to be done first in order to move the completed ones being moved onto the completed list to be
+			//       handled immediately in the following HandleCompletedPlaylistDownloads() call!!
+			HandleStaticRequestCompletions(Now);
 
-		// Execute enqueued playlist download requests.
-		HandleEnqueuedPlaylistDownloads(Now);
+			// Check completed playlist downloads
+			HandleCompletedPlaylistDownloads(Now);
 
-		// Check which playlists need to be reloaded.
-		CheckForPlaylistUpdate(Now);
+			// Execute enqueued playlist download requests.
+			HandleEnqueuedPlaylistDownloads(Now);
+
+			// Check which playlists need to be reloaded.
+			CheckForPlaylistUpdate(Now);
+		}
 	}
 
 
diff --git a/Engine/Plugins/Media/ElectraPlayer/Source/ElectraPlayerRuntime/Private/Runtime/Player/HLS/StreamReaderHLSfmp4.cpp b/Engine/Plugins/Media/ElectraPlayer/Source/ElectraPlayerRuntime/Private/Runtime/Player/HLS/StreamReaderHLSfmp4.cpp
--- a/Engine/Plugins/Media/ElectraPlayer/Source/ElectraPlayerRuntime/Private/Runtime/Player/HLS/StreamReaderHLSfmp4.cpp
+++ b/Engine/Plugins/Media/ElectraPlayer/Source/ElectraPlayerRuntime/Private/Runtime/Player/HLS/StreamReaderHLSfmp4.cpp
@@ -13,6 +13,8 @@
 #define INTERNAL_ERROR_INIT_SEGMENT_LICENSEKEY_ERROR				10
 
 
+DECLARE_CYCLE_STAT(TEXT("FStreamReaderHLSfmp4_HandleRequest"), STAT_ElectraPlayer_HLS_StreamReader, STATGROUP_ElectraPlayer);
+
 
 namespace Electra
 {
@@ -588,6 +590,9 @@ FStreamReaderHLSfmp4::FStreamHandler::EInitSegmentResult FStreamReaderHLSfmp4::F
 			// Note: It is only safe to access the connection info when the HTTP request has completed or the request been removed.
 			PlayerSessionService->GetHTTPManager()->RemoveRequest(HTTP, false);
 			Request->ConnectionInfo = HTTP->ConnectionInfo;
+
+			SCOPE_CYCLE_COUNTER(STAT_ElectraPlayer_HLS_StreamReader);
+			CSV_SCOPED_TIMING_STAT(ElectraPlayer, HLS_StreamReader);
 			if (!HTTP->ConnectionInfo.StatusInfo.ErrorDetail.IsError())
 			{
 // TODO: If encrypted we must now decrypt it!!
@@ -633,8 +638,6 @@ FStreamReaderHLSfmp4::FStreamHandler::EInitSegmentResult FStreamReaderHLSfmp4::F
 
 void FStreamReaderHLSfmp4::FStreamHandler::HandleRequest()
 {
-//CSV_SCOPED_TIMING_STAT(ElectraPlayer, StreamReaderMP4_Worker);
-
 	UEMediaError								Error;
 	bool										bIsEmptyFillerSegment = false;
 
@@ -865,7 +868,11 @@ void FStreamReaderHLSfmp4::FStreamHandler::HandleRequest()
 					UEMediaError parseError = MP4Parser->ParseHeader(this, this, PlayerSessionService, MP4InitSegment.Get());
 					if (parseError == UEMEDIA_ERROR_OK)
 					{
-						parseError = MP4Parser->PrepareTracks(MP4InitSegment);
+						{
+							SCOPE_CYCLE_COUNTER(STAT_ElectraPlayer_HLS_StreamReader);
+							CSV_SCOPED_TIMING_STAT(ElectraPlayer, HLS_StreamReader);
+							parseError = MP4Parser->PrepareTracks(MP4InitSegment);
+						}
 						if (parseError == UEMEDIA_ERROR_OK)
 						{
 							// For the time being we only want to have a single track in the movie segments.
@@ -964,6 +971,8 @@ void FStreamReaderHLSfmp4::FStreamHandler::HandleRequest()
 										// Shall we pass on any AUs we already read?
 										if (bAllowEarlyEmitting)
 										{
+											SCOPE_CYCLE_COUNTER(STAT_ElectraPlayer_HLS_StreamReader);
+											CSV_SCOPED_TIMING_STAT(ElectraPlayer, HLS_StreamReader);
 											while(AccessUnitFIFO.Num() && !HasReadBeenAborted())
 											{
 												FAccessUnit* pNext = AccessUnitFIFO.FrontRef();
@@ -1284,6 +1293,9 @@ int64 FStreamReaderHLSfmp4::FStreamHandler::ReadData(void* IntoBuffer, int64 Num
 				ds.ABRState.ProgressDecision = StreamSelectorDecision;
 				if ((StreamSelectorDecision.Flags & FABRDownloadProgressDecision::EDecisionFlags::eABR_EmitPartialData) != 0)
 				{
+					SCOPE_CYCLE_COUNTER(STAT_ElectraPlayer_HLS_StreamReader);
+					CSV_SCOPED_TIMING_STAT(ElectraPlayer, HLS_StreamReader);
+
 					bAllowEarlyEmitting = true;
 					// Deliver all enqueued AUs right now. Unless the request also gets aborted we could be stuck
 					// in here for a while longer.
@@ -1349,6 +1361,9 @@ int64 FStreamReaderHLSfmp4::FStreamHandler::ReadData(void* IntoBuffer, int64 Num
 			}
 			else
 			{
+				SCOPE_CYCLE_COUNTER(STAT_ElectraPlayer_HLS_StreamReader);
+				CSV_SCOPED_TIMING_STAT(ElectraPlayer, HLS_StreamReader);
+
 				SourceBuffer.Lock();
 				// Check the available size. If the read was aborted there may not be enough in here as the wait got released early.
 				if (SourceBuffer.Num() >= RequiredEncryptedSize)
@@ -1439,6 +1454,8 @@ int64 FStreamReaderHLSfmp4::FStreamHandler::ReadData(void* IntoBuffer, int64 Num
 			}
 			else
 			{
+				SCOPE_CYCLE_COUNTER(STAT_ElectraPlayer_HLS_StreamReader);
+				CSV_SCOPED_TIMING_STAT(ElectraPlayer, HLS_StreamReader);
 				SourceBuffer.Lock();
 				if (SourceBuffer.Num() >= ReadBuffer.ParsePos + NumBytesToRead)
 				{
diff --git a/Engine/Plugins/Media/ElectraPlayer/Source/ElectraPlayerRuntime/Private/Runtime/Player/mp4/ManifestMP4.cpp b/Engine/Plugins/Media/ElectraPlayer/Source/ElectraPlayerRuntime/Private/Runtime/Player/mp4/ManifestMP4.cpp
--- a/Engine/Plugins/Media/ElectraPlayer/Source/ElectraPlayerRuntime/Private/Runtime/Player/mp4/ManifestMP4.cpp
+++ b/Engine/Plugins/Media/ElectraPlayer/Source/ElectraPlayerRuntime/Private/Runtime/Player/mp4/ManifestMP4.cpp
@@ -18,6 +18,9 @@
 #define ERRCODE_MANIFEST_MP4_STARTSEGMENT_NOT_FOUND		1
 
 
+DECLARE_CYCLE_STAT(TEXT("FPlayPeriodMP4::FindSegment"), STAT_ElectraPlayer_MP4_FindSegment, STATGROUP_ElectraPlayer);
+DECLARE_CYCLE_STAT(TEXT("FPlayPeriodMP4::GetSegmentInformation"), STAT_ElectraPlayer_MP4_GetSegmentInformation, STATGROUP_ElectraPlayer);
+
 namespace Electra
 {
 
@@ -669,6 +672,9 @@ void FManifestMP4Internal::FTimelineAssetMP4::LimitSegmentDownloadSize(TSharedPt
 
 IManifest::FResult FManifestMP4Internal::FTimelineAssetMP4::GetStartingSegment(TSharedPtrTS<IStreamSegment>& OutSegment, const FPlayStartPosition& StartPosition, ESearchType SearchType, int64 AtAbsoluteFilePos)
 {
+	SCOPE_CYCLE_COUNTER(STAT_ElectraPlayer_MP4_FindSegment);
+	CSV_SCOPED_TIMING_STAT(ElectraPlayer, MP4_FindSegment);
+
 // TODO: If there is a SIDX box we will look in there.
 
 	// Look at the actual tracks. If there is video search there first for a keyframe/IDR frame.
@@ -991,6 +997,9 @@ IManifest::FResult FManifestMP4Internal::FTimelineAssetMP4::GetLoopingSegment(TS
 
 void FManifestMP4Internal::FTimelineAssetMP4::GetSegmentInformation(TArray<IManifest::IPlayPeriod::FSegmentInformation>& OutSegmentInformation, FTimeValue& OutAverageSegmentDuration, TSharedPtrTS<const IStreamSegment> CurrentSegment, const FTimeValue& LookAheadTime, const FString& AdaptationSetID, const FString& RepresentationID)
 {
+	SCOPE_CYCLE_COUNTER(STAT_ElectraPlayer_MP4_GetSegmentInformation);
+	CSV_SCOPED_TIMING_STAT(ElectraPlayer, MP4_GetSegmentInformation);
+
 	// This is not expected to be called. And if it does we return a dummy entry.
 	OutAverageSegmentDuration.SetFromSeconds(60.0);
 	OutSegmentInformation.Empty();
diff --git a/Engine/Plugins/Media/ElectraPlayer/Source/ElectraPlayerRuntime/Private/Runtime/Player/mp4/PlaylistReaderMP4.cpp b/Engine/Plugins/Media/ElectraPlayer/Source/ElectraPlayerRuntime/Private/Runtime/Player/mp4/PlaylistReaderMP4.cpp
--- a/Engine/Plugins/Media/ElectraPlayer/Source/ElectraPlayerRuntime/Private/Runtime/Player/mp4/PlaylistReaderMP4.cpp
+++ b/Engine/Plugins/Media/ElectraPlayer/Source/ElectraPlayerRuntime/Private/Runtime/Player/mp4/PlaylistReaderMP4.cpp
@@ -18,6 +18,8 @@
 #define ERRCODE_MP4_INVALID_FILE							1
 
 
+DECLARE_CYCLE_STAT(TEXT("FPlaylistReaderMP4_WorkerThread"), STAT_ElectraPlayer_MP4_PlaylistWorker, STATGROUP_ElectraPlayer);
+
 
 namespace Electra
 {
@@ -332,7 +334,9 @@ void FPlaylistReaderMP4::ReadNextChunk(int64 InFromOffset, int64 ChunkSize)
 void FPlaylistReaderMP4::WorkerThread()
 {
 	LLM_SCOPE(ELLMTag::ElectraPlayer);
-	CSV_SCOPED_TIMING_STAT(ElectraPlayer, PlaylistReaderMP4_Worker);
+
+	SCOPE_CYCLE_COUNTER(STAT_ElectraPlayer_MP4_PlaylistWorker);
+	CSV_SCOPED_TIMING_STAT(ElectraPlayer, MP4_PlaylistWorker);
 
 	MP4Parser = IParserISO14496_12::CreateParser();
 	UEMediaError parseError = MP4Parser->ParseHeader(this, this, PlayerSessionServices, nullptr);
diff --git a/Engine/Plugins/Media/ElectraPlayer/Source/ElectraPlayerRuntime/Private/Runtime/Player/mp4/StreamReaderMP4.cpp b/Engine/Plugins/Media/ElectraPlayer/Source/ElectraPlayerRuntime/Private/Runtime/Player/mp4/StreamReaderMP4.cpp
--- a/Engine/Plugins/Media/ElectraPlayer/Source/ElectraPlayerRuntime/Private/Runtime/Player/mp4/StreamReaderMP4.cpp
+++ b/Engine/Plugins/Media/ElectraPlayer/Source/ElectraPlayerRuntime/Private/Runtime/Player/mp4/StreamReaderMP4.cpp
@@ -15,7 +15,7 @@
 #include "Player/mp4/StreamReaderMP4.h"
 
 
-
+DECLARE_CYCLE_STAT(TEXT("FStreamReaderMP4_HandleRequest"), STAT_ElectraPlayer_MP4_StreamReader, STATGROUP_ElectraPlayer);
 
 
 namespace Electra
@@ -251,50 +251,40 @@ void FStreamReaderMP4::HTTPUpdateStats(const FTimeValue& CurrentTime, const IEle
 	}
 }
 
-void FStreamReaderMP4::WorkerThread()
+void FStreamReaderMP4::HandleRequest()
 {
-	LLM_SCOPE(ELLMTag::ElectraPlayer);
-	while(!bTerminate)
-	{
-		WorkSignal.WaitAndReset();
-		if (bTerminate)
-		{
-			break;
-		}
-
-		TSharedPtrTS<FStreamSegmentRequestMP4>	Request = CurrentRequest;
-		if (!Request.IsValid())
-		{
-			continue;
-		}
+	TSharedPtrTS<FStreamSegmentRequestMP4>	Request = CurrentRequest;
 
-		FManifestMP4Internal::FTimelineAssetMP4* TimelineAsset = static_cast<FManifestMP4Internal::FTimelineAssetMP4*>(Request->MediaAsset.Get());
+	FManifestMP4Internal::FTimelineAssetMP4* TimelineAsset = static_cast<FManifestMP4Internal::FTimelineAssetMP4*>(Request->MediaAsset.Get());
 
-		// Clear the active track map.
-		ActiveTrackMap.Reset();
+	// Clear the active track map.
+	ActiveTrackMap.Reset();
 
 // FIXME: If looping to somewhere within the stream instead of the beginning at zero the offset here must be made relative to the DTS of the first sample we demux.
 //        Otherwise there would be a large jump ahead in time!
-		FTimeValue LoopTimestampOffset = Request->PlayerLoopState.LoopBasetime;
-		TSharedPtrTS<const FPlayerLoopState> PlayerLoopState = MakeSharedTS<const FPlayerLoopState>(Request->PlayerLoopState);
+	FTimeValue LoopTimestampOffset = Request->PlayerLoopState.LoopBasetime;
+	TSharedPtrTS<const FPlayerLoopState> PlayerLoopState = MakeSharedTS<const FPlayerLoopState>(Request->PlayerLoopState);
 
 
-		// Get the list of all the tracks that have been selected in the asset.
-		// This does not mean their data will be _used_ for playback, only that the track is usable by the player
-		// with regards to type and codec.
-		struct FPlaylistTrackMetadata
-		{
-			EStreamType		Type;
-			FString			Kind;
-			FString			Language;
-			FString			PeriodID;
-			FString			AdaptationSetID;
-			FString			RepresentationID;
-			int32			Bitrate;
-			int32			Index;
-		};
-		TMap<uint32, FPlaylistTrackMetadata>	SelectedTrackMap;
-		const EStreamType TypesOfSupportedTracks[] = { EStreamType::Video, EStreamType::Audio, EStreamType::Subtitle };
+	// Get the list of all the tracks that have been selected in the asset.
+	// This does not mean their data will be _used_ for playback, only that the track is usable by the player
+	// with regards to type and codec.
+	struct FPlaylistTrackMetadata
+	{
+		EStreamType		Type;
+		FString			Kind;
+		FString			Language;
+		FString			PeriodID;
+		FString			AdaptationSetID;
+		FString			RepresentationID;
+		int32			Bitrate;
+		int32			Index;
+	};
+	TMap<uint32, FPlaylistTrackMetadata>	SelectedTrackMap;
+	const EStreamType TypesOfSupportedTracks[] = { EStreamType::Video, EStreamType::Audio, EStreamType::Subtitle };
+	{
+		SCOPE_CYCLE_COUNTER(STAT_ElectraPlayer_MP4_StreamReader);
+		CSV_SCOPED_TIMING_STAT(ElectraPlayer, MP4_StreamReader);
 		for(int32 nStrType=0; nStrType<UE_ARRAY_COUNT(TypesOfSupportedTracks); ++nStrType)
 		{
 			int32 NumAdapt = Request->MediaAsset->GetNumberOfAdaptationSets(TypesOfSupportedTracks[nStrType]);
@@ -324,139 +314,148 @@ void FStreamReaderMP4::WorkerThread()
 				}
 			}
 		}
+	}
 
 
-		Metrics::FSegmentDownloadStats& ds = Request->DownloadStats;
-		ds.StatsID = FMediaInterlockedIncrement(UniqueDownloadID);
+	Metrics::FSegmentDownloadStats& ds = Request->DownloadStats;
+	ds.StatsID = FMediaInterlockedIncrement(UniqueDownloadID);
 
-		check(Request->PrimaryTrackIterator.IsValid() && Request->PrimaryTrackIterator->GetTrack());
-		uint32 PrimaryTrackID = Request->PrimaryTrackIterator->GetTrack()->GetID();
-		const FPlaylistTrackMetadata* PrimaryTrackMetadata = SelectedTrackMap.Find(PrimaryTrackID);
-		check(PrimaryTrackMetadata);
-		if (PrimaryTrackMetadata)
-		{
-			ds.MediaAssetID 	= PrimaryTrackMetadata->PeriodID;
-			ds.AdaptationSetID  = PrimaryTrackMetadata->AdaptationSetID;
-			ds.RepresentationID = PrimaryTrackMetadata->RepresentationID;
-			ds.Bitrate  		= PrimaryTrackMetadata->Bitrate;
-		}
+	check(Request->PrimaryTrackIterator.IsValid() && Request->PrimaryTrackIterator->GetTrack());
+	uint32 PrimaryTrackID = Request->PrimaryTrackIterator->GetTrack()->GetID();
+	const FPlaylistTrackMetadata* PrimaryTrackMetadata = SelectedTrackMap.Find(PrimaryTrackID);
+	check(PrimaryTrackMetadata);
+	if (PrimaryTrackMetadata)
+	{
+		ds.MediaAssetID 	= PrimaryTrackMetadata->PeriodID;
+		ds.AdaptationSetID  = PrimaryTrackMetadata->AdaptationSetID;
+		ds.RepresentationID = PrimaryTrackMetadata->RepresentationID;
+		ds.Bitrate  		= PrimaryTrackMetadata->Bitrate;
+	}
 
-		ds.FailureReason.Empty();
-		ds.bWasSuccessful      = true;
-		ds.bWasAborted  	   = false;
-		ds.bDidTimeout  	   = false;
-		ds.HTTPStatusCode      = 0;
-		ds.StreamType   	   = Request->GetType();
-		ds.SegmentType  	   = Metrics::ESegmentType::Media;
-		ds.PresentationTime    = (Request->FirstPTS + LoopTimestampOffset).GetAsSeconds();
-		ds.Duration 		   = Request->SegmentDuration.GetAsSeconds();
-		ds.DurationDownloaded  = 0.0;
-		ds.DurationDelivered   = 0.0;
-		ds.TimeToFirstByte     = 0.0;
-		ds.TimeToDownload      = 0.0;
-		ds.ByteSize 		   = -1;
-		ds.NumBytesDownloaded  = 0;
-		ds.ThroughputBps	   = 0;
-		ds.bInsertedFillerData = false;
-		ds.URL  			   = TimelineAsset->GetMediaURL();
-		ds.bIsMissingSegment   = false;
-		ds.bParseFailure	   = false;
-		ds.RetryNumber  	   = Request->NumOverallRetries;
-		ds.ABRState.Reset();
-
-		Parameters.EventListener->OnFragmentOpen(Request);
-
-		TSharedPtrTS<IElectraHttpManager::FProgressListener>	ProgressListener;
-		ProgressListener = MakeSharedTS<IElectraHttpManager::FProgressListener>();
-		ProgressListener->CompletionDelegate = Electra::MakeDelegate(this, &FStreamReaderMP4::HTTPCompletionCallback);
-		ProgressListener->ProgressDelegate   = Electra::MakeDelegate(this, &FStreamReaderMP4::HTTPProgressCallback);
-
-		ReadBuffer.Reset();
-		ReadBuffer.ReceiveBuffer = MakeSharedTS<IElectraHttpManager::FReceiveBuffer>();
-		// Set the receive buffer to an okay-ish size. Too small and the file I/O may block too often and get too slow.
-		ReadBuffer.ReceiveBuffer->Buffer.Reserve(4 << 20);
-		ReadBuffer.ReceiveBuffer->bEnableRingbuffer = true;
-		ReadBuffer.SetCurrentPos(Request->FileStartOffset);
-		TSharedPtrTS<IElectraHttpManager::FRequest> HTTP(new IElectraHttpManager::FRequest);
-		HTTP->Parameters.URL				= TimelineAsset->GetMediaURL();
-		HTTP->Parameters.Range.Start		= Request->FileStartOffset;
-		HTTP->Parameters.Range.EndIncluding = Request->FileEndOffset;
-		// Explicit range?
-		int64 NumRequestedBytes = HTTP->Parameters.Range.GetNumberOfBytes();
-		int32 SubRequestSize = 0;
-		if (NumRequestedBytes > 0)
+	ds.FailureReason.Empty();
+	ds.bWasSuccessful      = true;
+	ds.bWasAborted  	   = false;
+	ds.bDidTimeout  	   = false;
+	ds.HTTPStatusCode      = 0;
+	ds.StreamType   	   = Request->GetType();
+	ds.SegmentType  	   = Metrics::ESegmentType::Media;
+	ds.PresentationTime    = (Request->FirstPTS + LoopTimestampOffset).GetAsSeconds();
+	ds.Duration 		   = Request->SegmentDuration.GetAsSeconds();
+	ds.DurationDownloaded  = 0.0;
+	ds.DurationDelivered   = 0.0;
+	ds.TimeToFirstByte     = 0.0;
+	ds.TimeToDownload      = 0.0;
+	ds.ByteSize 		   = -1;
+	ds.NumBytesDownloaded  = 0;
+	ds.ThroughputBps	   = 0;
+	ds.bInsertedFillerData = false;
+	ds.URL  			   = TimelineAsset->GetMediaURL();
+	ds.bIsMissingSegment   = false;
+	ds.bParseFailure	   = false;
+	ds.RetryNumber  	   = Request->NumOverallRetries;
+	ds.ABRState.Reset();
+
+	Parameters.EventListener->OnFragmentOpen(Request);
+
+	TSharedPtrTS<IElectraHttpManager::FProgressListener>	ProgressListener;
+	ProgressListener = MakeSharedTS<IElectraHttpManager::FProgressListener>();
+	ProgressListener->CompletionDelegate = Electra::MakeDelegate(this, &FStreamReaderMP4::HTTPCompletionCallback);
+	ProgressListener->ProgressDelegate   = Electra::MakeDelegate(this, &FStreamReaderMP4::HTTPProgressCallback);
+
+	ReadBuffer.Reset();
+	ReadBuffer.ReceiveBuffer = MakeSharedTS<IElectraHttpManager::FReceiveBuffer>();
+	// Set the receive buffer to an okay-ish size. Too small and the file I/O may block too often and get too slow.
+	ReadBuffer.ReceiveBuffer->Buffer.Reserve(4 << 20);
+	ReadBuffer.ReceiveBuffer->bEnableRingbuffer = true;
+	ReadBuffer.SetCurrentPos(Request->FileStartOffset);
+	TSharedPtrTS<IElectraHttpManager::FRequest> HTTP(new IElectraHttpManager::FRequest);
+	HTTP->Parameters.URL				= TimelineAsset->GetMediaURL();
+	HTTP->Parameters.Range.Start		= Request->FileStartOffset;
+	HTTP->Parameters.Range.EndIncluding = Request->FileEndOffset;
+	// Explicit range?
+	int64 NumRequestedBytes = HTTP->Parameters.Range.GetNumberOfBytes();
+	int32 SubRequestSize = 0;
+	if (NumRequestedBytes > 0)
+	{
+		if (Request->bIsFirstSegment && !Request->bIsLastSegment)
 		{
-			if (Request->bIsFirstSegment && !Request->bIsLastSegment)
-			{
-				SubRequestSize = 512 << 10;
-			}
-			else if (Request->bIsFirstSegment && Request->bIsLastSegment)
-			{
-				SubRequestSize = 2 << 20;
-			}
+			SubRequestSize = 512 << 10;
 		}
-		else
+		else if (Request->bIsFirstSegment && Request->bIsLastSegment)
 		{
-			if (Request->SegmentInternalSize < 0)
-			{
-				SubRequestSize = 2 << 20;
-			}
+			SubRequestSize = 2 << 20;
 		}
-		if (SubRequestSize)
+	}
+	else
+	{
+		if (Request->SegmentInternalSize < 0)
 		{
-			HTTP->Parameters.SubRangeRequestSize = SubRequestSize;
+			SubRequestSize = 2 << 20;
 		}
+	}
+	if (SubRequestSize)
+	{
+		HTTP->Parameters.SubRangeRequestSize = SubRequestSize;
+	}
 
-		HTTP->ReceiveBuffer 				= ReadBuffer.ReceiveBuffer;
-		HTTP->ProgressListener  			= ProgressListener;
-		PlayerSessionServices->GetHTTPManager()->AddRequest(HTTP, false);
+	HTTP->ReceiveBuffer 				= ReadBuffer.ReceiveBuffer;
+	HTTP->ProgressListener  			= ProgressListener;
+	PlayerSessionServices->GetHTTPManager()->AddRequest(HTTP, false);
 
 
-		FTimeValue DurationSuccessfullyDelivered(FTimeValue::GetZero());
-		FTimeValue DurationSuccessfullyRead(FTimeValue::GetZero());
-		FTimeValue NextLargestExpectedTimestamp(FTimeValue::GetZero());
-		bool bDone = false;
-		TSharedPtrTS<IParserISO14496_12::IAllTrackIterator> AllTrackIterator = TimelineAsset->GetMoovBoxParser()->CreateAllTrackIteratorByFilePos(Request->FileStartOffset);
-		while(!bDone && !HasErrored() && !HasBeenAborted() && !bTerminate)
+	FTimeValue DurationSuccessfullyDelivered(FTimeValue::GetZero());
+	FTimeValue DurationSuccessfullyRead(FTimeValue::GetZero());
+	FTimeValue NextLargestExpectedTimestamp(FTimeValue::GetZero());
+	bool bDone = false;
+	TSharedPtrTS<IParserISO14496_12::IAllTrackIterator> AllTrackIterator;
+	{
+		SCOPE_CYCLE_COUNTER(STAT_ElectraPlayer_MP4_StreamReader);
+		CSV_SCOPED_TIMING_STAT(ElectraPlayer, MP4_StreamReader);
+		AllTrackIterator = TimelineAsset->GetMoovBoxParser()->CreateAllTrackIteratorByFilePos(Request->FileStartOffset);
+	}
+	while(!bDone && !HasErrored() && !HasBeenAborted() && !bTerminate)
+	{
+		auto UpdateSelectedTrack = [&SelectedTrackMap](const IParserISO14496_12::ITrackIterator* trkIt, TMap<uint32, FSelectedTrackData>& ActiveTrks) -> FSelectedTrackData&
 		{
-			auto UpdateSelectedTrack = [&SelectedTrackMap](const IParserISO14496_12::ITrackIterator* trkIt, TMap<uint32, FSelectedTrackData>& ActiveTrks) -> FSelectedTrackData&
-			{
-				const IParserISO14496_12::ITrack* Track = trkIt->GetTrack();
-				uint32 tkid = Track->GetID();
+			const IParserISO14496_12::ITrack* Track = trkIt->GetTrack();
+			uint32 tkid = Track->GetID();
 
-				// Check if this track ID is already in our map of active tracks.
-				FSelectedTrackData& st = ActiveTrks.FindOrAdd(tkid);
-				if (!st.BufferSourceInfo.IsValid())
-				{
-					auto meta = MakeShared<FBufferSourceInfo, ESPMode::ThreadSafe>();
+			// Check if this track ID is already in our map of active tracks.
+			FSelectedTrackData& st = ActiveTrks.FindOrAdd(tkid);
+			if (!st.BufferSourceInfo.IsValid())
+			{
+				auto meta = MakeShared<FBufferSourceInfo, ESPMode::ThreadSafe>();
 
-					// Check if this track is in the list of selected tracks.
-					const FPlaylistTrackMetadata* SelectedTrackMetadata = SelectedTrackMap.Find(tkid);
-					if (SelectedTrackMetadata)
-					{
-						st.bIsSelectedTrack = true;
-						st.StreamType = SelectedTrackMetadata->Type;
-						meta->Kind = SelectedTrackMetadata->Kind;
-						meta->Language = SelectedTrackMetadata->Language;
-						meta->PeriodAdaptationSetID = SelectedTrackMetadata->PeriodID + TEXT(".") + SelectedTrackMetadata->AdaptationSetID;
-						meta->HardIndex = SelectedTrackMetadata->Index;
-					}
-					st.BufferSourceInfo = MoveTemp(meta);
-				}
-				if (!st.CSD.IsValid())
+				// Check if this track is in the list of selected tracks.
+				const FPlaylistTrackMetadata* SelectedTrackMetadata = SelectedTrackMap.Find(tkid);
+				if (SelectedTrackMetadata)
 				{
-					TSharedPtrTS<FAccessUnit::CodecData> CSD(new FAccessUnit::CodecData);
-					CSD->CodecSpecificData = Track->GetCodecSpecificData();
-					CSD->RawCSD			   = Track->GetCodecSpecificDataRAW();
-					CSD->ParsedInfo		   = Track->GetCodecInformation();
-					st.CSD = MoveTemp(CSD);
+					st.bIsSelectedTrack = true;
+					st.StreamType = SelectedTrackMetadata->Type;
+					meta->Kind = SelectedTrackMetadata->Kind;
+					meta->Language = SelectedTrackMetadata->Language;
+					meta->PeriodAdaptationSetID = SelectedTrackMetadata->PeriodID + TEXT(".") + SelectedTrackMetadata->AdaptationSetID;
+					meta->HardIndex = SelectedTrackMetadata->Index;
 				}
-				return st;
-			};
+				st.BufferSourceInfo = MoveTemp(meta);
+			}
+			if (!st.CSD.IsValid())
+			{
+				TSharedPtrTS<FAccessUnit::CodecData> CSD(new FAccessUnit::CodecData);
+				CSD->CodecSpecificData = Track->GetCodecSpecificData();
+				CSD->RawCSD			   = Track->GetCodecSpecificDataRAW();
+				CSD->ParsedInfo		   = Track->GetCodecInformation();
+				st.CSD = MoveTemp(CSD);
+			}
+			return st;
+		};
 
-			// Handle all the new tracks that have reached EOS while iterating. We do this first here to
-			// handle the tracks that hit EOS before reaching the intended start position.
-			TArray<const IParserISO14496_12::ITrackIterator*> TracksAtEOS;
+		// Handle all the new tracks that have reached EOS while iterating. We do this first here to
+		// handle the tracks that hit EOS before reaching the intended start position.
+		TArray<const IParserISO14496_12::ITrackIterator*> TracksAtEOS;
+		{
+			SCOPE_CYCLE_COUNTER(STAT_ElectraPlayer_MP4_StreamReader);
+			CSV_SCOPED_TIMING_STAT(ElectraPlayer, MP4_StreamReader);
 			AllTrackIterator->GetNewEOSTracks(TracksAtEOS);
 			AllTrackIterator->ClearNewEOSTracks();
 			for(int32 nTrk=0; nTrk<TracksAtEOS.Num(); ++nTrk)
@@ -469,184 +468,185 @@ void FStreamReaderMP4::WorkerThread()
 					Parameters.EventListener->OnFragmentReachedEOS(SelectedTrack.StreamType, SelectedTrack.BufferSourceInfo);
 				}
 			}
+		}
 
-			// Handle current track iterator
-			const IParserISO14496_12::ITrackIterator* TrackIt = AllTrackIterator->Current();
-			if (TrackIt)
+		// Handle current track iterator
+		const IParserISO14496_12::ITrackIterator* TrackIt = AllTrackIterator->Current();
+		if (TrackIt)
+		{
+			FSelectedTrackData& SelectedTrack = UpdateSelectedTrack(TrackIt, ActiveTrackMap);
+
+			// Get the sample properties
+			uint32 SampleNumber    = TrackIt->GetSampleNumber();
+			int64 DTS   		   = TrackIt->GetDTS();
+			int64 PTS   		   = TrackIt->GetPTS();
+			int64 Duration  	   = TrackIt->GetDuration();
+			uint32 Timescale	   = TrackIt->GetTimescale();
+			bool bIsSyncSample     = TrackIt->IsSyncSample();
+			int64 SampleSize	   = TrackIt->GetSampleSize();
+			int64 SampleFileOffset = TrackIt->GetSampleFileOffset();
+
+			// Remember at which file position we are currently at. In case of failure this is where we will retry.
+			Request->CurrentIteratorBytePos = SampleFileOffset;
+
+			// Do we need to skip over some data?
+			if (SampleFileOffset > ReadBuffer.GetCurrentPos())
 			{
-				FSelectedTrackData& SelectedTrack = UpdateSelectedTrack(TrackIt, ActiveTrackMap);
-
-				// Get the sample properties
-				uint32 SampleNumber    = TrackIt->GetSampleNumber();
-				int64 DTS   		   = TrackIt->GetDTS();
-				int64 PTS   		   = TrackIt->GetPTS();
-				int64 Duration  	   = TrackIt->GetDuration();
-				uint32 Timescale	   = TrackIt->GetTimescale();
-				bool bIsSyncSample     = TrackIt->IsSyncSample();
-				int64 SampleSize	   = TrackIt->GetSampleSize();
-				int64 SampleFileOffset = TrackIt->GetSampleFileOffset();
-
-				// Remember at which file position we are currently at. In case of failure this is where we will retry.
-				Request->CurrentIteratorBytePos = SampleFileOffset;
-
-				// Do we need to skip over some data?
-				if (SampleFileOffset > ReadBuffer.GetCurrentPos())
+				int32 NumBytesToSkip = SampleFileOffset - ReadBuffer.GetCurrentPos();
+				int64 nr = ReadBuffer.ReadTo(nullptr, NumBytesToSkip);
+				if (nr != NumBytesToSkip)
 				{
-					int32 NumBytesToSkip = SampleFileOffset - ReadBuffer.GetCurrentPos();
-					int64 nr = ReadBuffer.ReadTo(nullptr, NumBytesToSkip);
-					if (nr != NumBytesToSkip)
-					{
-						bDone = true;
-						break;
-					}
-				}
-				else if (SampleFileOffset < ReadBuffer.GetCurrentPos())
-				{
-					ds.bParseFailure  = true;
-					Request->ConnectionInfo.StatusInfo.ErrorDetail.SetMessage(FString::Printf(TEXT("Segment parse error. Sample offset %lld for sample #%u in track %u is before the current read position at %lld"), (long long int)SampleFileOffset, SampleNumber, TrackIt->GetTrack()->GetID(), (long long int)ReadBuffer.GetCurrentPos()));
-					bHasErrored = true;
+					bDone = true;
 					break;
 				}
+			}
+			else if (SampleFileOffset < ReadBuffer.GetCurrentPos())
+			{
+				ds.bParseFailure  = true;
+				Request->ConnectionInfo.StatusInfo.ErrorDetail.SetMessage(FString::Printf(TEXT("Segment parse error. Sample offset %lld for sample #%u in track %u is before the current read position at %lld"), (long long int)SampleFileOffset, SampleNumber, TrackIt->GetTrack()->GetID(), (long long int)ReadBuffer.GetCurrentPos()));
+				bHasErrored = true;
+				break;
+			}
 
-				// Do we read the sample because the track is selected or do we discard it?
-				if (SelectedTrack.bIsSelectedTrack)
+			// Do we read the sample because the track is selected or do we discard it?
+			if (SelectedTrack.bIsSelectedTrack)
+			{
+				// Is this a sync sample?
+				if (bIsSyncSample && !SelectedTrack.bGotKeyframe)
 				{
-					// Is this a sync sample?
-					if (bIsSyncSample && !SelectedTrack.bGotKeyframe)
-					{
-						SelectedTrack.bGotKeyframe = true;
-					}
-					// Do we need to skip samples from this track until we reach a sync sample?
-					bool bSkipUntilSyncSample = !SelectedTrack.bGotKeyframe && !Request->bIsContinuationSegment;
+					SelectedTrack.bGotKeyframe = true;
+				}
+				// Do we need to skip samples from this track until we reach a sync sample?
+				bool bSkipUntilSyncSample = !SelectedTrack.bGotKeyframe && !Request->bIsContinuationSegment;
 
-					// Create an access unit.
-					FAccessUnit *AccessUnit = FAccessUnit::Create(Parameters.MemoryProvider);
-					if (AccessUnit)
+				// Create an access unit.
+				FAccessUnit *AccessUnit = FAccessUnit::Create(Parameters.MemoryProvider);
+				if (AccessUnit)
+				{
+					AccessUnit->ESType = SelectedTrack.StreamType;
+					AccessUnit->PTS.SetFromND(PTS, Timescale);
+					AccessUnit->DTS.SetFromND(DTS, Timescale);
+					AccessUnit->Duration.SetFromND(Duration, Timescale);
+					AccessUnit->AUSize = (uint32) SampleSize;
+					AccessUnit->AUCodecData = SelectedTrack.CSD;
+					AccessUnit->DropState = FAccessUnit::EDropState::None;
+					// If this is a continuation then we must not tag samples as being too early.
+					if (!Request->bIsContinuationSegment)
 					{
-						AccessUnit->ESType = SelectedTrack.StreamType;
-						AccessUnit->PTS.SetFromND(PTS, Timescale);
-						AccessUnit->DTS.SetFromND(DTS, Timescale);
-						AccessUnit->Duration.SetFromND(Duration, Timescale);
-						AccessUnit->AUSize = (uint32) SampleSize;
-						AccessUnit->AUCodecData = SelectedTrack.CSD;
-						AccessUnit->DropState = FAccessUnit::EDropState::None;
-						// If this is a continuation then we must not tag samples as being too early.
-						if (!Request->bIsContinuationSegment)
+						if (AccessUnit->DTS < Request->FirstPTS)
 						{
-							if (AccessUnit->DTS < Request->FirstPTS)
-							{
-								AccessUnit->DropState |= FAccessUnit::EDropState::DtsTooEarly;
-							}
-							if (AccessUnit->PTS < Request->FirstPTS)
-							{
-								AccessUnit->DropState |= FAccessUnit::EDropState::PtsTooEarly;
-							}
+							AccessUnit->DropState |= FAccessUnit::EDropState::DtsTooEarly;
 						}
-						// FIXME: if we only want to read a partial segment we could set drop state based on the sample being 'too late'.
+						if (AccessUnit->PTS < Request->FirstPTS)
+						{
+							AccessUnit->DropState |= FAccessUnit::EDropState::PtsTooEarly;
+						}
+					}
+					// FIXME: if we only want to read a partial segment we could set drop state based on the sample being 'too late'.
 
-						// Apply timestamp offsets for looping after checking the timestamp limits.
-						AccessUnit->PTS += LoopTimestampOffset;
-						AccessUnit->DTS += LoopTimestampOffset;
+					// Apply timestamp offsets for looping after checking the timestamp limits.
+					AccessUnit->PTS += LoopTimestampOffset;
+					AccessUnit->DTS += LoopTimestampOffset;
 
-						AccessUnit->bIsFirstInSequence = SelectedTrack.bIsFirstInSequence;
-						AccessUnit->bIsSyncSample = bIsSyncSample;
-						AccessUnit->bIsDummyData = false;
-						AccessUnit->AUData = AccessUnit->AllocatePayloadBuffer(AccessUnit->AUSize);
+					AccessUnit->bIsFirstInSequence = SelectedTrack.bIsFirstInSequence;
+					AccessUnit->bIsSyncSample = bIsSyncSample;
+					AccessUnit->bIsDummyData = false;
+					AccessUnit->AUData = AccessUnit->AllocatePayloadBuffer(AccessUnit->AUSize);
 
-						// Set the associated stream metadata
-						AccessUnit->BufferSourceInfo = SelectedTrack.BufferSourceInfo;
-						AccessUnit->PlayerLoopState = PlayerLoopState;
+					// Set the associated stream metadata
+					AccessUnit->BufferSourceInfo = SelectedTrack.BufferSourceInfo;
+					AccessUnit->PlayerLoopState = PlayerLoopState;
 
-						SelectedTrack.bIsFirstInSequence = false;
+					SelectedTrack.bIsFirstInSequence = false;
 
-						int64 nr = ReadBuffer.ReadTo(AccessUnit->AUData, (int32)SampleSize);
-						if (nr == SampleSize)
-						{
-							SelectedTrack.DurationSuccessfullyRead += AccessUnit->Duration;
+					int64 nr = ReadBuffer.ReadTo(AccessUnit->AUData, (int32)SampleSize);
+					if (nr == SampleSize)
+					{
+						SelectedTrack.DurationSuccessfullyRead += AccessUnit->Duration;
 
-							//LogMessage(IInfoLog::ELevel::Info, FString::Printf("[%u] %4u: DTS=%lld PTS=%lld dur=%lld sync=%d; %lld bytes @ %lld", tkid, SampleNumber, (long long int)AccessUnit->mDTS.GetAsMicroseconds(), (long long int)AccessUnit->mPTS.GetAsMicroseconds(), (long long int)AccessUnit->mDuration.GetAsMicroseconds(), bIsSyncSample?1:0, (long long int)SampleSize, (long long int)SampleFileOffset));
+						//LogMessage(IInfoLog::ELevel::Info, FString::Printf("[%u] %4u: DTS=%lld PTS=%lld dur=%lld sync=%d; %lld bytes @ %lld", tkid, SampleNumber, (long long int)AccessUnit->mDTS.GetAsMicroseconds(), (long long int)AccessUnit->mPTS.GetAsMicroseconds(), (long long int)AccessUnit->mDuration.GetAsMicroseconds(), bIsSyncSample?1:0, (long long int)SampleSize, (long long int)SampleFileOffset));
 
-							// Keep track of the next expected sample PTS and remember the largest value of all tracks.
-							FTimeValue NextExpectedPTS = AccessUnit->PTS + AccessUnit->Duration;
-							if (NextExpectedPTS > NextLargestExpectedTimestamp)
-							{
-								NextLargestExpectedTimestamp = NextExpectedPTS;
-							}
+						// Keep track of the next expected sample PTS and remember the largest value of all tracks.
+						FTimeValue NextExpectedPTS = AccessUnit->PTS + AccessUnit->Duration;
+						if (NextExpectedPTS > NextLargestExpectedTimestamp)
+						{
+							NextLargestExpectedTimestamp = NextExpectedPTS;
+						}
 
-							bool bSentOff = false;
-							while(!bSentOff && !HasBeenAborted() && !bTerminate)
+						bool bSentOff = false;
+						while(!bSentOff && !HasBeenAborted() && !bTerminate)
+						{
+							if (Parameters.EventListener->OnFragmentAccessUnitReceived(AccessUnit))
 							{
-								if (Parameters.EventListener->OnFragmentAccessUnitReceived(AccessUnit))
-								{
-									SelectedTrack.DurationSuccessfullyDelivered += AccessUnit->Duration;
-									bSentOff = true;
-									AccessUnit = nullptr;
-
-									// Since we have delivered this access unit, if we are detecting an error now we need to then
-									// retry on the _next_ AU and not this one again!
-									Request->CurrentIteratorBytePos = SampleFileOffset + SampleSize;
-								}
-								else
-								{
-									FMediaRunnable::SleepMicroseconds(1000 * 10);
-								}
-							}
-
-							// Release the AU if we still have it.
-							FAccessUnit::Release(AccessUnit);
-							AccessUnit = nullptr;
+								SelectedTrack.DurationSuccessfullyDelivered += AccessUnit->Duration;
+								bSentOff = true;
+								AccessUnit = nullptr;
 
-							// For error handling, if we managed to get additional data we reset the retry count.
-							if (ds.RetryNumber && SelectedTrack.DurationSuccessfullyRead.GetAsSeconds() > 2.0)
+								// Since we have delivered this access unit, if we are detecting an error now we need to then
+								// retry on the _next_ AU and not this one again!
+								Request->CurrentIteratorBytePos = SampleFileOffset + SampleSize;
+							}
+							else
 							{
-								ds.RetryNumber = 0;
-								Request->NumOverallRetries = 0;
+								FMediaRunnable::SleepMicroseconds(1000 * 10);
 							}
 						}
-						else
+
+						// Release the AU if we still have it.
+						FAccessUnit::Release(AccessUnit);
+						AccessUnit = nullptr;
+
+						// For error handling, if we managed to get additional data we reset the retry count.
+						if (ds.RetryNumber && SelectedTrack.DurationSuccessfullyRead.GetAsSeconds() > 2.0)
 						{
-							// Did not get the number of bytes we needed. Either because of a read error or because we got aborted.
-							FAccessUnit::Release(AccessUnit);
-							AccessUnit = nullptr;
-							bDone = true;
-							break;
+							ds.RetryNumber = 0;
+							Request->NumOverallRetries = 0;
 						}
 					}
 					else
 					{
-						// TODO: Throw OOM error
+						// Did not get the number of bytes we needed. Either because of a read error or because we got aborted.
+						FAccessUnit::Release(AccessUnit);
+						AccessUnit = nullptr;
+						bDone = true;
+						break;
 					}
 				}
+				else
+				{
+					// TODO: Throw OOM error
+				}
 			}
-			else
-			{
-				break;
-			}
-			AllTrackIterator->Next();
 		}
+		else
+		{
+			break;
+		}
+		AllTrackIterator->Next();
+	}
 
-		// Remove the download request.
-		ProgressListener.Reset();
-		PlayerSessionServices->GetHTTPManager()->RemoveRequest(HTTP, false);
-		Request->ConnectionInfo = HTTP->ConnectionInfo;
-		HTTP.Reset();
+	// Remove the download request.
+	ProgressListener.Reset();
+	PlayerSessionServices->GetHTTPManager()->RemoveRequest(HTTP, false);
+	Request->ConnectionInfo = HTTP->ConnectionInfo;
+	HTTP.Reset();
 
-		// Remember the next largest timestamp from all tracks.
-		Request->NextLargestExpectedTimestamp = NextLargestExpectedTimestamp;
+	// Remember the next largest timestamp from all tracks.
+	Request->NextLargestExpectedTimestamp = NextLargestExpectedTimestamp;
 
 
-		// Set downloaded and delivered duration from the primary track.
-		FSelectedTrackData& PrimaryTrack = ActiveTrackMap.FindOrAdd(PrimaryTrackID);
-		DurationSuccessfullyRead	  = PrimaryTrack.DurationSuccessfullyRead;
-		DurationSuccessfullyDelivered = PrimaryTrack.DurationSuccessfullyDelivered;
+	// Set downloaded and delivered duration from the primary track.
+	FSelectedTrackData& PrimaryTrack = ActiveTrackMap.FindOrAdd(PrimaryTrackID);
+	DurationSuccessfullyRead	  = PrimaryTrack.DurationSuccessfullyRead;
+	DurationSuccessfullyDelivered = PrimaryTrack.DurationSuccessfullyDelivered;
 
-		// Set up remaining download stat fields.
+	// Set up remaining download stat fields.
 // Note: currently commented out because of UE-88612.
 //       This must be reinstated once we set failure reasons in the loop above so they won't get replaced by this!
 //		if (ds.FailureReason.length() == 0)
-		{
-			ds.FailureReason = Request->ConnectionInfo.StatusInfo.ErrorDetail.GetMessage();
-		}
+	{
+		ds.FailureReason = Request->ConnectionInfo.StatusInfo.ErrorDetail.GetMessage();
+	}
 //		if (bAbortedByABR)
 //		{
 //			// If aborted set the reason as the download failure.
@@ -654,30 +654,46 @@ void FStreamReaderMP4::WorkerThread()
 //		}
 //		ds.bWasAborted  	  = bAbortedByABR;
 //		ds.bWasSuccessful     = !bHasErrored && !bAbortedByABR;
-		ds.bWasSuccessful     = !bHasErrored;
-		ds.URL  			  = Request->ConnectionInfo.EffectiveURL;
-		ds.HTTPStatusCode     = Request->ConnectionInfo.StatusInfo.HTTPStatus;
-		ds.DurationDownloaded = DurationSuccessfullyRead.GetAsSeconds();
-		ds.DurationDelivered  = DurationSuccessfullyDelivered.GetAsSeconds();
-		ds.TimeToFirstByte    = Request->ConnectionInfo.TimeUntilFirstByte;
-		ds.TimeToDownload     = (Request->ConnectionInfo.RequestEndTime - Request->ConnectionInfo.RequestStartTime).GetAsSeconds();
-		ds.ByteSize 		  = Request->ConnectionInfo.ContentLength;
-		ds.NumBytesDownloaded = Request->ConnectionInfo.BytesReadSoFar;
-		ds.ThroughputBps	  = Request->ConnectionInfo.Throughput.GetThroughput();
-		if (ds.ThroughputBps == 0)
-		{
-			ds.ThroughputBps = ds.TimeToDownload > 0.0 ? 8 * ds.NumBytesDownloaded / ds.TimeToDownload : 0;
-		}
+	ds.bWasSuccessful     = !bHasErrored;
+	ds.URL  			  = Request->ConnectionInfo.EffectiveURL;
+	ds.HTTPStatusCode     = Request->ConnectionInfo.StatusInfo.HTTPStatus;
+	ds.DurationDownloaded = DurationSuccessfullyRead.GetAsSeconds();
+	ds.DurationDelivered  = DurationSuccessfullyDelivered.GetAsSeconds();
+	ds.TimeToFirstByte    = Request->ConnectionInfo.TimeUntilFirstByte;
+	ds.TimeToDownload     = (Request->ConnectionInfo.RequestEndTime - Request->ConnectionInfo.RequestStartTime).GetAsSeconds();
+	ds.ByteSize 		  = Request->ConnectionInfo.ContentLength;
+	ds.NumBytesDownloaded = Request->ConnectionInfo.BytesReadSoFar;
+	ds.ThroughputBps	  = Request->ConnectionInfo.Throughput.GetThroughput();
+	if (ds.ThroughputBps == 0)
+	{
+		ds.ThroughputBps = ds.TimeToDownload > 0.0 ? 8 * ds.NumBytesDownloaded / ds.TimeToDownload : 0;
+	}
 
-		ActiveTrackMap.Reset();
+	ActiveTrackMap.Reset();
 
-		// Reset the current request so another one can be added immediately when we call OnFragmentClose()
-		CurrentRequest.Reset();
-		PlayerSessionServices->GetStreamSelector()->ReportDownloadEnd(ds);
-		Parameters.EventListener->OnFragmentClose(Request);
-	}
+	// Reset the current request so another one can be added immediately when we call OnFragmentClose()
+	CurrentRequest.Reset();
+	PlayerSessionServices->GetStreamSelector()->ReportDownloadEnd(ds);
+	Parameters.EventListener->OnFragmentClose(Request);
 }
 
+void FStreamReaderMP4::WorkerThread()
+{
+	LLM_SCOPE(ELLMTag::ElectraPlayer);
+	while(!bTerminate)
+	{
+		WorkSignal.WaitAndReset();
+		if (bTerminate)
+		{
+			break;
+		}
+		TSharedPtrTS<FStreamSegmentRequestMP4>	Request = CurrentRequest;
+		if (Request.IsValid())
+		{
+			HandleRequest();
+		}
+	}
+}
 
 
 int32 FStreamReaderMP4::FReadBuffer::ReadTo(void* ToBuffer, int32 NumBytes)
@@ -688,6 +704,9 @@ int32 FStreamReaderMP4::FReadBuffer::ReadTo(void* ToBuffer, int32 NumBytes)
 	// Do we have enough data in the ringbuffer to satisfy the read?
 	if (SourceBuffer.Num() >= NumBytes)
 	{
+		SCOPE_CYCLE_COUNTER(STAT_ElectraPlayer_MP4_StreamReader);
+		CSV_SCOPED_TIMING_STAT(ElectraPlayer, MP4_StreamReader);
+
 		// Yes. Get the data and return.
 		int32 NumGot = SourceBuffer.PopData(OutputBuffer, NumBytes);
 		check(NumGot == NumBytes);
@@ -711,7 +730,12 @@ int32 FStreamReaderMP4::FReadBuffer::ReadTo(void* ToBuffer, int32 NumBytes)
 			}
 
 			// Get whatever amount of data is currently available to free up the buffer for receiving more data.
-			int32 NumGot = SourceBuffer.PopData(OutputBuffer, NumBytesToGo);
+			int32 NumGot;
+			{
+				SCOPE_CYCLE_COUNTER(STAT_ElectraPlayer_MP4_StreamReader);
+				CSV_SCOPED_TIMING_STAT(ElectraPlayer, MP4_StreamReader);
+				NumGot = SourceBuffer.PopData(OutputBuffer, NumBytesToGo);
+			}
 			if ((NumBytesToGo -= NumGot) > 0)
 			{
 				if (OutputBuffer)
diff --git a/Engine/Plugins/Media/ElectraPlayer/Source/ElectraPlayerRuntime/Private/Runtime/Player/mp4/StreamReaderMP4.h b/Engine/Plugins/Media/ElectraPlayer/Source/ElectraPlayerRuntime/Private/Runtime/Player/mp4/StreamReaderMP4.h
--- a/Engine/Plugins/Media/ElectraPlayer/Source/ElectraPlayerRuntime/Private/Runtime/Player/mp4/StreamReaderMP4.h
+++ b/Engine/Plugins/Media/ElectraPlayer/Source/ElectraPlayerRuntime/Private/Runtime/Player/mp4/StreamReaderMP4.h
@@ -96,6 +96,7 @@ private:
 	int32 HTTPProgressCallback(const IElectraHttpManager::FRequest* InRequest);
 	void HTTPCompletionCallback(const IElectraHttpManager::FRequest* InRequest);
 	void HTTPUpdateStats(const FTimeValue& CurrentTime, const IElectraHttpManager::FRequest* Request);
+	void HandleRequest();
 
 	bool HasBeenAborted() const;
 	bool HasErrored() const;
