ElectraPlayer: moved metadata events from player into plugin to avoid dependencies to media and mediasamples

#rnx

diff --git a/Engine/Plugins/Media/ElectraPlayer/Source/ElectraPlayerPlugin/Private/ElectraPlayerPlugin.cpp b/Engine/Plugins/Media/ElectraPlayer/Source/ElectraPlayerPlugin/Private/ElectraPlayerPlugin.cpp
--- a/Engine/Plugins/Media/ElectraPlayer/Source/ElectraPlayerPlugin/Private/ElectraPlayerPlugin.cpp
+++ b/Engine/Plugins/Media/ElectraPlayer/Source/ElectraPlayerPlugin/Private/ElectraPlayerPlugin.cpp
@@ -11,6 +11,8 @@
 #include "IElectraPlayerPluginModule.h"
 #include "IElectraMetadataSample.h"
 
+#include "MediaMetaDataDecoderOutput.h"
+
 using namespace Electra;
 
 //-----------------------------------------------------------------------------
@@ -95,6 +97,77 @@ FElectraPlayerPlugin::~FElectraPlayerPlugin()
 
 //-----------------------------------------------------------------------------
 
+class FElectraBinarySample : public IElectraBinarySample
+{
+public:
+	virtual ~FElectraBinarySample() = default;
+	virtual const void* GetData() override						{ return Metadata->GetData(); }
+	virtual uint32 GetSize() const override						{ return Metadata->GetSize(); }
+	virtual FGuid GetGUID() const override						{ return IElectraBinarySample::GetSampleTypeGUID(); }
+	virtual const FString& GetSchemeIdUri() const override		{ return Metadata->GetSchemeIdUri(); }
+	virtual const FString& GetValue() const override			{ return Metadata->GetValue(); }
+	virtual const FString& GetID() const override				{ return Metadata->GetID(); }
+
+	virtual EDispatchedMode GetDispatchedMode() const override
+	{
+		switch(Metadata->GetDispatchedMode())
+		{
+			default:
+			case IMetaDataDecoderOutput::EDispatchedMode::OnReceive:
+			{
+				return FElectraBinarySample::EDispatchedMode::OnReceive;
+			}
+			case IMetaDataDecoderOutput::EDispatchedMode::OnStart:
+			{
+				return FElectraBinarySample::EDispatchedMode::OnStart;
+			}
+		}
+	}
+	
+	virtual EOrigin GetOrigin() const override
+	{
+		switch(Metadata->GetOrigin())
+		{
+			default:
+			case IMetaDataDecoderOutput::EOrigin::TimedMetadata:		
+			{
+				return FElectraBinarySample::EOrigin::TimedMetadata;
+			}
+			case IMetaDataDecoderOutput::EOrigin::EventStream:		
+			{
+				return FElectraBinarySample::EOrigin::EventStream;
+			}
+			case IMetaDataDecoderOutput::EOrigin::InbandEventStream:	
+			{
+				return FElectraBinarySample::EOrigin::InbandEventStream;
+			}
+		}
+	}
+
+	virtual FMediaTimeStamp GetTime() const override
+	{ 
+		FDecoderTimeStamp ts = Metadata->GetTime(); 
+		return FMediaTimeStamp(ts.Time, ts.SequenceIndex);
+	}
+	
+	virtual FTimespan GetDuration() const override
+	{
+		FTimespan Duration = Metadata->GetDuration();
+		// A zero duration might cause the metadata sample fall through the cracks later
+		// so set it to a short 1ms instead.
+		if (Duration.IsZero())
+		{
+			Duration = FTimespan::FromMilliseconds(1);
+		}
+		return Duration;
+	}
+
+	IMetaDataDecoderOutputPtr Metadata;
+};
+
+//-----------------------------------------------------------------------------
+
+
 Electra::FVariantValue FElectraPlayerPlugin::FPlayerAdapterDelegate::QueryOptions(EOptionType Type, const Electra::FVariantValue & Param)
 {
 	TSharedPtr<FElectraPlayerPlugin, ESPMode::ThreadSafe> PinnedHost = Host.Pin();
@@ -234,12 +307,15 @@ void FElectraPlayerPlugin::FPlayerAdapterDelegate::PresentAudioFrame(const IAudi
 	}
 }
 
-void FElectraPlayerPlugin::FPlayerAdapterDelegate::PresentMetadataSample(const IElectraBinarySampleRef& InMetadataSample)
+void FElectraPlayerPlugin::FPlayerAdapterDelegate::PresentMetadataSample(const IMetaDataDecoderOutputPtr& InMetadataFrame)
 {
 	TSharedPtr<FElectraPlayerPlugin, ESPMode::ThreadSafe> PinnedHost = Host.Pin();
 	if (PinnedHost.IsValid())
 	{
-		PinnedHost->MediaSamples->AddMetadata(InMetadataSample);
+		// Create a binary media sample of our extended format and pass it up.
+		TSharedRef<FElectraBinarySample, ESPMode::ThreadSafe> MetaDataSample = MakeShared<FElectraBinarySample, ESPMode::ThreadSafe>();
+		MetaDataSample->Metadata = InMetadataFrame;
+		PinnedHost->MediaSamples->AddMetadata(MetaDataSample);
 	}
 }
 
diff --git a/Engine/Plugins/Media/ElectraPlayer/Source/ElectraPlayerPlugin/Private/ElectraPlayerPlugin.h b/Engine/Plugins/Media/ElectraPlayer/Source/ElectraPlayerPlugin/Private/ElectraPlayerPlugin.h
--- a/Engine/Plugins/Media/ElectraPlayer/Source/ElectraPlayerPlugin/Private/ElectraPlayerPlugin.h
+++ b/Engine/Plugins/Media/ElectraPlayer/Source/ElectraPlayerPlugin/Private/ElectraPlayerPlugin.h
@@ -20,7 +20,6 @@
 
 class IElectraPlayerRuntimeModule;
 class FElectraPlayerResourceDelegate;
-class IElectraBinarySample;
 
 
 DECLARE_LOG_CATEGORY_EXTERN(LogElectraPlayerPlugin, Log, All);
@@ -165,7 +164,7 @@ private:
 		void OnAudioFlush() override;
 		void PresentVideoFrame(const FVideoDecoderOutputPtr& InVideoFrame) override;
 		void PresentAudioFrame(const IAudioDecoderOutputPtr& InAudioFrame) override;
-		void PresentMetadataSample(const IElectraBinarySampleRef& InMetadataSample) override;
+		void PresentMetadataSample(const IMetaDataDecoderOutputPtr& InMetadataSample) override;
 		bool CanReceiveVideoSamples(int32 NumFrames) override;
 		bool CanReceiveAudioSamples(int32 NumFrames) override;
 		void PrepareForDecoderShutdown() override;
diff --git a/Engine/Plugins/Media/ElectraPlayer/Source/ElectraPlayerRuntime/ElectraPlayerRuntime.Build.cs b/Engine/Plugins/Media/ElectraPlayer/Source/ElectraPlayerRuntime/ElectraPlayerRuntime.Build.cs
--- a/Engine/Plugins/Media/ElectraPlayer/Source/ElectraPlayerRuntime/ElectraPlayerRuntime.Build.cs
+++ b/Engine/Plugins/Media/ElectraPlayer/Source/ElectraPlayerRuntime/ElectraPlayerRuntime.Build.cs
@@ -21,9 +21,7 @@ namespace UnrealBuildTool.Rules
 					"HTTP",
 					"WebSockets",
 					"Json",
-					"Media",
 					"ElectraBase",
-					"ElectraSamples",
 					"ElectraCDM",
 					"XmlParser"
 				});
diff --git a/Engine/Plugins/Media/ElectraPlayer/Source/ElectraPlayerRuntime/Private/ElectraPlayer.cpp b/Engine/Plugins/Media/ElectraPlayer/Source/ElectraPlayerRuntime/Private/ElectraPlayer.cpp
--- a/Engine/Plugins/Media/ElectraPlayer/Source/ElectraPlayerRuntime/Private/ElectraPlayer.cpp
+++ b/Engine/Plugins/Media/ElectraPlayer/Source/ElectraPlayerRuntime/Private/ElectraPlayer.cpp
@@ -9,11 +9,10 @@
 #include "Player/AdaptiveStreamingPlayer.h"
 #include "Renderer/RendererVideo.h"
 #include "Renderer/RendererAudio.h"
+#include "MediaMetaDataDecoderOutput.h"
 #include "VideoDecoderResourceDelegate.h"
 #include "Utilities/Utilities.h"
 
-#include "IElectraMetadataSample.h"
-
 #include "Async/Async.h"
 
 #include "CoreGlobals.h"
@@ -77,15 +76,14 @@ static FString SanitizeMessage(FString InMessage)
 
 //-----------------------------------------------------------------------------
 
-class FElectraBinarySample : public IElectraBinarySample
+class FMetaDataDecoderOutput : public IMetaDataDecoderOutput
 {
 public:
-	virtual ~FElectraBinarySample() = default;
+	virtual ~FMetaDataDecoderOutput() = default;
 	virtual const void* GetData() override						{ return Data.GetData(); }
 	virtual FTimespan GetDuration() const override				{ return Duration; }
 	virtual uint32 GetSize() const override						{ return (uint32) Data.Num(); }
-	virtual FMediaTimeStamp GetTime() const override			{ return PresentationTime; }
-	virtual FGuid GetGUID() const override						{ return IElectraBinarySample::GetSampleTypeGUID(); }
+	virtual FDecoderTimeStamp GetTime() const override			{ return PresentationTime; }
 	virtual EOrigin GetOrigin() const override					{ return Origin; }
 	virtual EDispatchedMode GetDispatchedMode() const override	{ return DispatchedMode; }
 	virtual const FString& GetSchemeIdUri() const override		{ return SchemeIdUri; }
@@ -93,7 +91,7 @@ public:
 	virtual const FString& GetID() const override				{ return ID; }
 
 	TArray<uint8> Data;
-	FMediaTimeStamp PresentationTime;
+	FDecoderTimeStamp PresentationTime;
 	FTimespan Duration;
 	EOrigin Origin;
 	EDispatchedMode DispatchedMode;
@@ -1211,18 +1209,18 @@ void FElectraPlayer::OnMediaPlayerEventReceived(TSharedPtrTS<IAdaptiveStreamingP
 	if (PinnedAdapterDelegate.IsValid())
 	{
 		// Create a binary media sample of our extended format and pass it up.
-		TSharedRef<FElectraBinarySample, ESPMode::ThreadSafe> Meta = MakeShared<FElectraBinarySample, ESPMode::ThreadSafe>();
+		TSharedPtr<FMetaDataDecoderOutput, ESPMode::ThreadSafe> Meta = MakeShared<FMetaDataDecoderOutput, ESPMode::ThreadSafe>();
 		switch(InDispatchMode)
 		{
 			default:
 			case IAdaptiveStreamingPlayerAEMSReceiver::EDispatchMode::OnReceive:
 			{
-				Meta->DispatchedMode = FElectraBinarySample::EDispatchedMode::OnReceive;
+				Meta->DispatchedMode = FMetaDataDecoderOutput::EDispatchedMode::OnReceive;
 				break;
 			}
 			case IAdaptiveStreamingPlayerAEMSReceiver::EDispatchMode::OnStart:
 			{
-				Meta->DispatchedMode = FElectraBinarySample::EDispatchedMode::OnStart;
+				Meta->DispatchedMode = FMetaDataDecoderOutput::EDispatchedMode::OnStart;
 				break;
 			}
 		}
@@ -1231,17 +1229,17 @@ void FElectraPlayer::OnMediaPlayerEventReceived(TSharedPtrTS<IAdaptiveStreamingP
 			default:
 			case IAdaptiveStreamingPlayerAEMSEvent::EOrigin::TimedMetadata:		
 			{
-				Meta->Origin = FElectraBinarySample::EOrigin::TimedMetadata;
+				Meta->Origin = FMetaDataDecoderOutput::EOrigin::TimedMetadata;
 				break;
 			}
 			case IAdaptiveStreamingPlayerAEMSEvent::EOrigin::EventStream:		
 			{
-				Meta->Origin = FElectraBinarySample::EOrigin::EventStream;
+				Meta->Origin = FMetaDataDecoderOutput::EOrigin::EventStream;
 				break;
 			}
 			case IAdaptiveStreamingPlayerAEMSEvent::EOrigin::InbandEventStream:	
 			{
-				Meta->Origin = FElectraBinarySample::EOrigin::InbandEventStream;
+				Meta->Origin = FMetaDataDecoderOutput::EOrigin::InbandEventStream;
 				break;
 			}
 		}
@@ -1250,13 +1248,7 @@ void FElectraPlayer::OnMediaPlayerEventReceived(TSharedPtrTS<IAdaptiveStreamingP
 		Meta->Value = InEvent->GetValue();
 		Meta->ID = InEvent->GetID(),
 		Meta->Duration = InEvent->GetDuration().GetAsTimespan();
-		// A zero duration might cause the metadata sample fall through the cracks later
-		// so set it to a short 1ms instead.
-		if (Meta->Duration.IsZero())
-		{
-			Meta->Duration = FTimespan::FromMilliseconds(1);
-		}
-		Meta->PresentationTime = FMediaTimeStamp(InEvent->GetPresentationTime().GetAsTimespan());
+		Meta->PresentationTime = FDecoderTimeStamp(InEvent->GetPresentationTime().GetAsTimespan(), 0);
 		PinnedAdapterDelegate->PresentMetadataSample(Meta);
 	}
 }
diff --git a/Engine/Plugins/Media/ElectraPlayer/Source/ElectraPlayerRuntime/Private/ElectraPlayer.h b/Engine/Plugins/Media/ElectraPlayer/Source/ElectraPlayerRuntime/Private/ElectraPlayer.h
--- a/Engine/Plugins/Media/ElectraPlayer/Source/ElectraPlayerRuntime/Private/ElectraPlayer.h
+++ b/Engine/Plugins/Media/ElectraPlayer/Source/ElectraPlayerRuntime/Private/ElectraPlayer.h
@@ -17,6 +17,9 @@ class FVideoDecoderOutput;
 using FVideoDecoderOutputPtr = TSharedPtr<FVideoDecoderOutput, ESPMode::ThreadSafe>;
 class IAudioDecoderOutput;
 using IAudioDecoderOutputPtr = TSharedPtr<IAudioDecoderOutput, ESPMode::ThreadSafe>;
+class IMetaDataDecoderOutput;
+using IMetaDataDecoderOutputPtr = TSharedPtr<IMetaDataDecoderOutput, ESPMode::ThreadSafe>;
+
 
 namespace Electra
 {
diff --git a/Engine/Plugins/Media/ElectraPlayer/Source/ElectraPlayerRuntime/Public/IElectraPlayerInterface.h b/Engine/Plugins/Media/ElectraPlayer/Source/ElectraPlayerRuntime/Public/IElectraPlayerInterface.h
--- a/Engine/Plugins/Media/ElectraPlayer/Source/ElectraPlayerRuntime/Public/IElectraPlayerInterface.h
+++ b/Engine/Plugins/Media/ElectraPlayer/Source/ElectraPlayerRuntime/Public/IElectraPlayerInterface.h
@@ -12,8 +12,8 @@ class FVideoDecoderOutput;
 using FVideoDecoderOutputPtr = TSharedPtr<FVideoDecoderOutput, ESPMode::ThreadSafe>;
 class IAudioDecoderOutput;
 using IAudioDecoderOutputPtr = TSharedPtr<IAudioDecoderOutput, ESPMode::ThreadSafe>;
-class IElectraBinarySample;
-using IElectraBinarySampleRef = TSharedRef<IElectraBinarySample, ESPMode::ThreadSafe>;
+class IMetaDataDecoderOutput;
+using IMetaDataDecoderOutputPtr = TSharedPtr<IMetaDataDecoderOutput, ESPMode::ThreadSafe>;
 
 // ---------------------------------------------------------------------------------------------
 
@@ -68,7 +68,7 @@ public:
 	virtual void OnAudioFlush() = 0;
 	virtual void PresentVideoFrame(const FVideoDecoderOutputPtr& InVideoFrame) = 0;
 	virtual void PresentAudioFrame(const IAudioDecoderOutputPtr& InAudioFrame) = 0;
-	virtual void PresentMetadataSample(const IElectraBinarySampleRef& InMetadataSample) = 0;
+	virtual void PresentMetadataSample(const IMetaDataDecoderOutputPtr& InMetadataSample) = 0;
 	virtual bool CanReceiveVideoSamples(int32 NumFrames) = 0;
 	virtual bool CanReceiveAudioSamples(int32 NumFrames) = 0;
 	virtual void PrepareForDecoderShutdown() = 0;
diff --git a/Engine/Plugins/Media/ElectraUtil/Source/ElectraBase/Public/MediaMetaDataDecoderOutput.h b/Engine/Plugins/Media/ElectraUtil/Source/ElectraBase/Public/MediaMetaDataDecoderOutput.h
new file mode 100644
--- /dev/null
+++ b/Engine/Plugins/Media/ElectraUtil/Source/ElectraBase/Public/MediaMetaDataDecoderOutput.h
@@ -0,0 +1,49 @@
+// Copyright Epic Games, Inc. All Rights Reserved.
+#pragma once
+
+#include <CoreMinimal.h>
+#include <Misc/Timespan.h>
+#include <Templates/SharedPointer.h>
+
+#include "MediaDecoderOutput.h"
+
+class IMetaDataDecoderOutput
+{
+public:
+	virtual ~IMetaDataDecoderOutput() = default;
+
+	virtual const void* GetData() = 0;
+	virtual FTimespan GetDuration() const = 0;
+	virtual uint32 GetSize() const = 0;
+	virtual FDecoderTimeStamp GetTime() const = 0;
+
+	enum class EOrigin
+	{
+		EventStream,
+		InbandEventStream,
+		TimedMetadata
+	};
+
+	enum class EDispatchedMode
+	{
+		OnReceive,
+		OnStart
+	};
+
+	// Returns where the event originated from.
+	virtual EOrigin GetOrigin() const = 0;
+
+	// Returns the type with which this event was dispatched.
+	virtual EDispatchedMode GetDispatchedMode() const = 0;
+
+	// Returns the event's scheme id URI
+	virtual const FString& GetSchemeIdUri() const = 0;
+
+	// Returns the event's value.
+	virtual const FString& GetValue() const = 0;
+
+	// Returns the event's ID.
+	virtual const FString& GetID() const = 0;
+};
+
+using IMetaDataDecoderOutputPtr = TSharedPtr<IMetaDataDecoderOutput, ESPMode::ThreadSafe>;
