ElectraPlayer: fixed an issue with end-of-stream detection when only one type of elementary stream is present.

#rnx

diff --git a/Engine/Plugins/Media/ElectraPlayer/Source/ElectraPlayerRuntime/Private/Runtime/Player/AdaptiveStreamingPlayer.cpp b/Engine/Plugins/Media/ElectraPlayer/Source/ElectraPlayerRuntime/Private/Runtime/Player/AdaptiveStreamingPlayer.cpp
--- a/Engine/Plugins/Media/ElectraPlayer/Source/ElectraPlayerRuntime/Private/Runtime/Player/AdaptiveStreamingPlayer.cpp
+++ b/Engine/Plugins/Media/ElectraPlayer/Source/ElectraPlayerRuntime/Private/Runtime/Player/AdaptiveStreamingPlayer.cpp
@@ -3290,6 +3290,9 @@ void FAdaptiveStreamingPlayer::CheckForStreamEnd()
 		{
 			if (StreamState == EStreamState::eStream_Running)
 			{
+				const bool bHaveVid = bHaveVideoReader.GetWithDefault(false);
+				const bool bHaveAud = bHaveAudioReader.GetWithDefault(false);
+				const bool bHaveTxt = bHaveTextReader.GetWithDefault(false);
 				// For simplicity we assume all streams are done if we know whether or not they exist.
 				bool bEndVid = bHaveVideoReader.IsSet();
 				bool bEndAud = bHaveAudioReader.IsSet();
@@ -3303,7 +3306,7 @@ void FAdaptiveStreamingPlayer::CheckForStreamEnd()
 				DiagnosticsCriticalSection.Unlock();
 
 				// Check for end of video stream
-				if (StreamReaderHandler && bHaveVideoReader.GetWithDefault(false))
+				if (StreamReaderHandler && bHaveVid)
 				{
 					// All buffers at end of data?
 					bEndVid = (vidStats.StreamBuffer.bEndOfData && vidStats.DecoderInputBuffer.bEODSignaled && vidStats.DecoderOutputBuffer.bEODreached);
@@ -3311,7 +3314,7 @@ void FAdaptiveStreamingPlayer::CheckForStreamEnd()
 				}
 
 				// Check for end of audio stream
-				if (StreamReaderHandler && bHaveAudioReader.GetWithDefault(false))
+				if (StreamReaderHandler && bHaveAud)
 				{
 					// All buffers at end of data?
 					bEndAud = (audStats.StreamBuffer.bEndOfData && audStats.DecoderInputBuffer.bEODSignaled && audStats.DecoderOutputBuffer.bEODreached);
@@ -3324,7 +3327,7 @@ void FAdaptiveStreamingPlayer::CheckForStreamEnd()
 
 				// If either primary stream has reliably ended do a check if the other may have stalled because the application is no longer
 				// consuming decoder output, which will prevent the other stream from ending.
-				if (bEndVid || bEndAud)
+				if ((bHaveVid && bHaveAud) && (bEndVid || bEndAud))
 				{
 					int64 OtherStallTime = bEndAud ? VidStalled : AudStalled;
 					if (OtherStallTime > 500)
