ElectraPlayer: Added support for HEVC decoder where available

#rnx

diff --git a/Engine/Plugins/Media/ElectraPlayer/Source/ElectraPlayerRuntime/ElectraPlayerRuntime.Build.cs b/Engine/Plugins/Media/ElectraPlayer/Source/ElectraPlayerRuntime/ElectraPlayerRuntime.Build.cs
--- a/Engine/Plugins/Media/ElectraPlayer/Source/ElectraPlayerRuntime/ElectraPlayerRuntime.Build.cs
+++ b/Engine/Plugins/Media/ElectraPlayer/Source/ElectraPlayerRuntime/ElectraPlayerRuntime.Build.cs
@@ -54,8 +54,9 @@ namespace UnrealBuildTool.Rules
 
                 PrivateDefinitions.Add("ELECTRA_ENABLE_MFDECODER");
                 PrivateDefinitions.Add("ELECTRA_ENABLE_SWDECODE");
+				PrivateDefinitions.Add("ELECTRA_PLATFORM_HAS_H265_DECODER=1");
 
-                AddEngineThirdPartyPrivateStaticDependencies(Target, "WinHttp");
+				AddEngineThirdPartyPrivateStaticDependencies(Target, "WinHttp");
 
 				if (WinSupportsDX9())
 				{
@@ -102,6 +103,8 @@ namespace UnrealBuildTool.Rules
 								"QuartzCore"
 				});
 
+				PrivateDefinitions.Add("ELECTRA_PLATFORM_HAS_H265_DECODER=1");
+
 				PublicIncludePaths.Add("$(ModuleDir)/Public/Apple");
 
 				PrivateIncludePaths.Add("ElectraPlayerRuntime/Private/Runtime/Decoder/Apple");
@@ -118,6 +121,8 @@ namespace UnrealBuildTool.Rules
 								"QuartzCore"
 				});
 
+				PrivateDefinitions.Add("ELECTRA_PLATFORM_HAS_H265_DECODER=1");
+
 				PublicIncludePaths.Add("$(ModuleDir)/Public/Apple");
 
 				PrivateIncludePaths.Add("ElectraPlayerRuntime/Private/Runtime/Decoder/Apple");
diff --git a/Engine/Plugins/Media/ElectraPlayer/Source/ElectraPlayerRuntime/Private/ElectraPlayerRuntimeModule.cpp b/Engine/Plugins/Media/ElectraPlayer/Source/ElectraPlayerRuntime/Private/ElectraPlayerRuntimeModule.cpp
--- a/Engine/Plugins/Media/ElectraPlayer/Source/ElectraPlayerRuntime/Private/ElectraPlayerRuntimeModule.cpp
+++ b/Engine/Plugins/Media/ElectraPlayer/Source/ElectraPlayerRuntime/Private/ElectraPlayerRuntimeModule.cpp
@@ -11,6 +11,7 @@
 
 #include "Renderer/RendererAudio.h"
 #include "Decoder/VideoDecoderH264.h"
+#include "Decoder/VideoDecoderH265.h"
 #include "Decoder/AudioDecoderAAC.h"
 
 #include "ElectraPlayerPrivate.h"
@@ -80,7 +81,9 @@ public:
 		// H.264 decoder
 		IVideoDecoderH264::FSystemConfiguration SysCfgH264;
 		IVideoDecoderH264::Startup(SysCfgH264);
-
+		// H.265 decoder
+		FParamDict NoOptions;
+		IVideoDecoderH265::Startup(NoOptions);
 		// AAC decoder
 		IAudioDecoderAAC::FSystemConfiguration SysCfgAAC;
 		IAudioDecoderAAC::Startup(SysCfgAAC);
@@ -94,6 +97,7 @@ public:
 		{
 			Electra::WaitForAllPlayersToHaveTerminated();
 			IAudioDecoderAAC::Shutdown();
+			IVideoDecoderH265::Shutdown();
 			IVideoDecoderH264::Shutdown();
 			FElectraRendererAudio::Shutdown();
 			FElectraRendererVideo::Shutdown();
diff --git a/Engine/Plugins/Media/ElectraPlayer/Source/ElectraPlayerRuntime/Private/Runtime/Decoder/Android/VideoDecoderH264_Android.cpp b/Engine/Plugins/Media/ElectraPlayer/Source/ElectraPlayerRuntime/Private/Runtime/Decoder/Android/VideoDecoderH264_Android.cpp
--- a/Engine/Plugins/Media/ElectraPlayer/Source/ElectraPlayerRuntime/Private/Runtime/Decoder/Android/VideoDecoderH264_Android.cpp
+++ b/Engine/Plugins/Media/ElectraPlayer/Source/ElectraPlayerRuntime/Private/Runtime/Decoder/Android/VideoDecoderH264_Android.cpp
@@ -44,6 +44,7 @@ public:
 
 	virtual void Open(const FInstanceConfiguration& InConfig) override;
 	virtual void Close() override;
+	virtual void DrainForCodecChange() override;
 
 	virtual void SetMaximumDecodeCapability(int32 MaxWidth, int32 MaxHeight, int32 MaxProfile, int32 MaxProfileLevel, const FParamDict& AdditionalOptions) override;
 
@@ -474,6 +475,18 @@ void FVideoDecoderH264::Close()
 }
 
 
+//-----------------------------------------------------------------------------
+/**
+ * Drains the decoder of all enqueued input and ends it, after which the decoder must send an FDecoderMessage to the player
+ * to signal completion.
+ */
+void FVideoDecoderH264::DrainForCodecChange()
+{
+	// Draining is currently not implemented. Send the player message nonetheless so the player will proceed with Close() and termination.
+	PlayerSessionServices->SendMessageToPlayer(FDecoderMessage::Create(FDecoderMessage::EReason::DrainingFinished, this, EStreamType::Video, FStreamCodecInformation::ECodec::H264));
+}
+
+
 //-----------------------------------------------------------------------------
 /**
  * Sets a new decoder limit.
diff --git a/Engine/Plugins/Media/ElectraPlayer/Source/ElectraPlayerRuntime/Private/Runtime/Decoder/Apple/VideoDecoderH264_Apple.cpp b/Engine/Plugins/Media/ElectraPlayer/Source/ElectraPlayerRuntime/Private/Runtime/Decoder/Apple/VideoDecoderH264_Apple.cpp
--- a/Engine/Plugins/Media/ElectraPlayer/Source/ElectraPlayerRuntime/Private/Runtime/Decoder/Apple/VideoDecoderH264_Apple.cpp
+++ b/Engine/Plugins/Media/ElectraPlayer/Source/ElectraPlayerRuntime/Private/Runtime/Decoder/Apple/VideoDecoderH264_Apple.cpp
@@ -26,9 +26,6 @@ DECLARE_CYCLE_STAT(TEXT("FVideoDecoderH264::ConvertOutput()"), STAT_ElectraPlaye
 namespace Electra
 {
 
-const uint32 NumImagesHoldBackForPTSOrdering = 5; // Number of frames held back to ensure proper PTS-ordering of decoder output
-const uint32 MaxImagesHoldBackForPTSOrdering = 5; // Maximum number of frames to be held in the buffer before we stall the decoder
-
 /**
  * H264 video decoder class implementation.
 **/
@@ -45,6 +42,7 @@ public:
 
 	virtual void Open(const FInstanceConfiguration& InConfig) override;
 	virtual void Close() override;
+	virtual void DrainForCodecChange() override;
 
 	virtual void SetMaximumDecodeCapability(int32 MaxWidth, int32 MaxHeight, int32 MaxProfile, int32 MaxProfileLevel, const FParamDict& AdditionalOptions) override;
 
@@ -99,6 +97,12 @@ private:
 	void LogMessage(IInfoLog::ELevel Level, const FString& Message);
 
 private:
+	enum
+	{
+		NumImagesHoldBackForPTSOrdering = 5,	// Number of frames held back to ensure proper PTS-ordering of decoder output
+		MaxImagesHoldBackForPTSOrdering = 5		// Maximum number of frames to be held in the buffer before we stall the decoder
+	};
+
 	void DecodeCallback(void* pSrcRef, OSStatus status, VTDecodeInfoFlags infoFlags, CVImageBufferRef imageBuffer, CMTime presentationTimeStamp, CMTime presentationDuration);
 	static void _DecodeCallback(void* pUser, void* pSrcRef, OSStatus status, VTDecodeInfoFlags infoFlags, CVImageBufferRef imageBuffer, CMTime presentationTimeStamp, CMTime presentationDuration)
 	{
@@ -253,6 +257,7 @@ private:
 
 	FAccessUnitBuffer									AccessUnitBuffer;
 	FAccessUnitBuffer									ReplayAccessUnitBuffer;
+	bool												bDrainForCodecChange;
 
 	FMediaCriticalSection								ListenerMutex;
 	IAccessUnitBufferListener*							InputBufferListener;
@@ -430,7 +435,7 @@ FVideoDecoderH264::FVideoDecoderH264()
 	: FMediaThread("ElectraPlayer::H264 decoder")
 	, bThreadStarted(false)
 	, PlayerSessionServices(nullptr)
-	, Renderer(nullptr)
+	, bDrainForCodecChange(false)
 	, InputBufferListener(nullptr)
 	, ReadyBufferListener(nullptr)
 	, DecoderHandle(nullptr)
@@ -484,9 +489,9 @@ void FVideoDecoderH264::SetReadyBufferListener(IDecoderOutputBufferListener* Lis
  *
  * @return
  */
-void FVideoDecoderH264::SetPlayerSessionServices(IPlayerSessionServices* SessionServices)
+void FVideoDecoderH264::SetPlayerSessionServices(IPlayerSessionServices* InSessionServices)
 {
-	PlayerSessionServices = SessionServices;
+	PlayerSessionServices = InSessionServices;
 }
 
 
@@ -494,7 +499,7 @@ void FVideoDecoderH264::SetPlayerSessionServices(IPlayerSessionServices* Session
 /**
  * Opens a decoder instance
  *
- * @param config
+ * @param InConfig
  *
  * @return
  */
@@ -521,6 +526,17 @@ void FVideoDecoderH264::Close()
 }
 
 
+//-----------------------------------------------------------------------------
+/**
+ * Drains the decoder of all enqueued input and ends it, after which the decoder must send an FDecoderMessage to the player
+ * to signal completion.
+ */
+void FVideoDecoderH264::DrainForCodecChange()
+{
+	bDrainForCodecChange = true;
+}
+
+
 //-----------------------------------------------------------------------------
 /**
  * Sets a new decoder limit.
@@ -932,13 +948,38 @@ void FVideoDecoderH264::PrepareAU(FAccessUnit* AccessUnit, bool& bOutIsIDR, bool
 				// TODO: we might need to set aside any SEI messages carrying 608 or 708 caption data.
 			}
 
-	// FIXME: We may need to remove SPS and PPS from here in case the stream data is avc3 data.
-	//        Not sure if the decoder will actually deal with those or if they need to be strictly passed
-	//        on the side and not inband.
-
 			uint32 naluLen = MEDIA_FROM_BIG_ENDIAN(*CurrentNALU) + 4;
 			CurrentNALU = Electra::AdvancePointer(CurrentNALU, naluLen);
 		}
+
+		// Is there codec specific data?
+		if (AccessUnit->AUCodecData.IsValid())
+		{
+			// Yes.
+			if (bOutIsIDR || AccessUnit->bIsSyncSample || AccessUnit->bIsFirstInSequence)
+			{
+				// Have to re-allocate the AU memory to preprend the codec data
+				uint64 nb = AccessUnit->AUSize + AccessUnit->AUCodecData->CodecSpecificData.Num();
+				void* pD = AccessUnit->AllocatePayloadBuffer(nb);
+				void* pP = pD;
+				FMemory::Memcpy(pP, AccessUnit->AUCodecData->CodecSpecificData.GetData(), AccessUnit->AUCodecData->CodecSpecificData.Num());
+				pP = Electra::AdvancePointer(pP, AccessUnit->AUCodecData->CodecSpecificData.Num());
+				FMemory::Memcpy(pP, AccessUnit->AUData, AccessUnit->AUSize);
+				AccessUnit->AdoptNewPayloadBuffer(pD, nb);
+				// New codec data makes this AU non-discardable.
+				bOutIsDiscardable = false;
+
+				// Get the NALUs from the CSD.
+				TArray<MPEG::FNaluInfo>	NALUs;
+				MPEG::ParseBitstreamForNALUs(NALUs, AccessUnit->AUCodecData->CodecSpecificData.GetData(), AccessUnit->AUCodecData->CodecSpecificData.Num());
+				// Replace the startcodes in the CSD with length values
+				for(int32 i=0; i<NALUs.Num(); ++i)
+				{
+					uint8* NALU = (uint8*)Electra::AdvancePointer(pD, NALUs[i].Offset);
+					*(uint32*)NALU = MEDIA_TO_BIG_ENDIAN((uint32)NALUs[i].Size);
+				}
+			}
+		}
 	}
 }
 
@@ -1319,179 +1360,195 @@ void FVideoDecoderH264::WorkerThread()
 
 	while(!TerminateThreadSignal.IsSignaled())
 	{
-		// Notify optional buffer listener that we will now be needing an AU for our input buffer.
-		if (!bError && InputBufferListener && AccessUnitBuffer.Num() == 0)
+		if (!bDrainForCodecChange)
 		{
-			SCOPE_CYCLE_COUNTER(STAT_ElectraPlayer_VideoH264Decode);
-			CSV_SCOPED_TIMING_STAT(ElectraPlayer, VideoH264Decode);
-			FAccessUnitBufferInfo	sin;
-			IAccessUnitBufferListener::FBufferStats	stats;
-			AccessUnitBuffer.GetStats(sin);
-			stats.NumAUsAvailable  = sin.NumCurrentAccessUnits;
-			stats.NumBytesInBuffer = sin.CurrentMemInUse;
-			stats.MaxBytesOfBuffer = sin.MaxDataSize;
-			stats.bEODSignaled     = sin.bEndOfData;
-			stats.bEODReached      = sin.bEndOfData && sin.NumCurrentAccessUnits == 0;
-			ListenerMutex.Lock();
-			if (InputBufferListener)
+			// Notify optional buffer listener that we will now be needing an AU for our input buffer.
+			if (!bError && InputBufferListener && AccessUnitBuffer.Num() == 0)
 			{
-				InputBufferListener->DecoderInputNeeded(stats);
-			}
-			ListenerMutex.Unlock();
-		}
-		// Wait for data.
-		bool bHaveData = AccessUnitBuffer.WaitForData(1000 * 5);
-
-		// Only send new data to the decoder if we know we got enough room (to avoid accumulating too many frames in our internal PTS-sort queue)
-		bool bTooManyImagesWaiting = false;
-		if (bHaveData)
-		{
-			// When there is data, even and especially after a previous EOD, we are no longer done and idling.
-			if (bDone)
-			{
-				bDone = false;
+				SCOPE_CYCLE_COUNTER(STAT_ElectraPlayer_VideoH264Decode);
+				CSV_SCOPED_TIMING_STAT(ElectraPlayer, VideoH264Decode);
+				FAccessUnitBufferInfo	sin;
+				IAccessUnitBufferListener::FBufferStats	stats;
+				AccessUnitBuffer.GetStats(sin);
+				stats.NumAUsAvailable  = sin.NumCurrentAccessUnits;
+				stats.NumBytesInBuffer = sin.CurrentMemInUse;
+				stats.MaxBytesOfBuffer = sin.MaxDataSize;
+				stats.bEODSignaled     = sin.bEndOfData;
+				stats.bEODReached      = sin.bEndOfData && sin.NumCurrentAccessUnits == 0;
+				ListenerMutex.Lock();
+				if (InputBufferListener)
+				{
+					InputBufferListener->DecoderInputNeeded(stats);
+				}
+				ListenerMutex.Unlock();
 			}
+			// Wait for data.
+			bool bHaveData = AccessUnitBuffer.WaitForData(1000 * 5);
 
-			ReadyImageMutex.Lock();
-			bTooManyImagesWaiting = (ReadyImages.Num() > MaxImagesHoldBackForPTSOrdering);
-			ReadyImageMutex.Unlock();
-			if (bTooManyImagesWaiting)
+			// Only send new data to the decoder if we know we got enough room (to avoid accumulating too many frames in our internal PTS-sort queue)
+			bool bTooManyImagesWaiting = false;
+			if (bHaveData)
 			{
-				// Signal blockage and wait for it to clear up before going back to normal decoder work...
-				NotifyReadyBufferListener(false);
+				// When there is data, even and especially after a previous EOD, we are no longer done and idling.
+				if (bDone)
+				{
+					bDone = false;
+				}
 
-				while(!TerminateThreadSignal.IsSignaled() && !FlushDecoderSignal.IsSignaled())
+				ReadyImageMutex.Lock();
+				bTooManyImagesWaiting = (ReadyImages.Num() > MaxImagesHoldBackForPTSOrdering);
+				ReadyImageMutex.Unlock();
+				if (bTooManyImagesWaiting)
 				{
-					ProcessOutput();
+					// Signal blockage and wait for it to clear up before going back to normal decoder work...
+					NotifyReadyBufferListener(false);
 
-					ReadyImageMutex.Lock();
-					bTooManyImagesWaiting = (ReadyImages.Num() > MaxImagesHoldBackForPTSOrdering);
-					ReadyImageMutex.Unlock();
-					if (!bTooManyImagesWaiting)
+					while(!TerminateThreadSignal.IsSignaled() && !FlushDecoderSignal.IsSignaled())
 					{
-						break;
-					}
+						ProcessOutput();
+
+						ReadyImageMutex.Lock();
+						bTooManyImagesWaiting = (ReadyImages.Num() > MaxImagesHoldBackForPTSOrdering);
+						ReadyImageMutex.Unlock();
+						if (!bTooManyImagesWaiting)
+						{
+							break;
+						}
 
-					FMediaRunnable::SleepMilliseconds(20);
+						FMediaRunnable::SleepMilliseconds(20);
+					}
 				}
 			}
-		}
 
-		if (bHaveData && !bTooManyImagesWaiting)
-		{
-			FAccessUnit* AccessUnit = nullptr;
-			bool bOk = AccessUnitBuffer.Pop(AccessUnit);
-			MEDIA_UNUSED_VAR(bOk);
-			check(bOk);
-
-			if (!AccessUnit->bIsDummyData)
+			if (bHaveData && !bTooManyImagesWaiting)
 			{
-				bool bIsKeyframe, bIsDiscardable;
-				PrepareAU(AccessUnit, bIsKeyframe, bIsDiscardable);
-				// An IDR frame means we can start decoding there, so we can purge any accumulated replay AUs.
-				if (bIsKeyframe)
-				{
-					ReplayAccessUnitBuffer.Flush();
-				}
+				FAccessUnit* AccessUnit = nullptr;
+				bool bOk = AccessUnitBuffer.Pop(AccessUnit);
+				MEDIA_UNUSED_VAR(bOk);
+				check(bOk);
 
-				bool bStreamFormatChanged = CurrentStreamFormatInfo.IsDifferentFrom(AccessUnit) || bGotLastSequenceAU;
-				bool bNeedNewDecoder = false;
-                bGotLastSequenceAU = AccessUnit->bIsLastInPeriod;
-				if (bStreamFormatChanged || !DecoderHandle)
+				if (!AccessUnit->bIsDummyData)
 				{
-					SCOPE_CYCLE_COUNTER(STAT_ElectraPlayer_VideoH264Decode);
-					CSV_SCOPED_TIMING_STAT(ElectraPlayer, VideoH264Decode);
+					bool bIsKeyframe, bIsDiscardable;
+					PrepareAU(AccessUnit, bIsKeyframe, bIsDiscardable);
+					// An IDR frame means we can start decoding there, so we can purge any accumulated replay AUs.
+					if (bIsKeyframe)
+					{
+						ReplayAccessUnitBuffer.Flush();
+					}
 
-					CMFormatDescriptionRef NewFormatDescr = nullptr;
-					if (CreateFormatDescription(NewFormatDescr, AccessUnit))
+					bool bStreamFormatChanged = CurrentStreamFormatInfo.IsDifferentFrom(AccessUnit) || bGotLastSequenceAU;
+					bool bNeedNewDecoder = false;
+					bGotLastSequenceAU = AccessUnit->bIsLastInPeriod;
+					if (bStreamFormatChanged || !DecoderHandle)
 					{
-						bNeedNewDecoder = DecoderHandle == nullptr || !DecoderHandle->IsCompatibleWith(NewFormatDescr);
-						if (bNeedNewDecoder)
+						SCOPE_CYCLE_COUNTER(STAT_ElectraPlayer_VideoH264Decode);
+						CSV_SCOPED_TIMING_STAT(ElectraPlayer, VideoH264Decode);
+
+						CMFormatDescriptionRef NewFormatDescr = nullptr;
+						if (CreateFormatDescription(NewFormatDescr, AccessUnit))
 						{
-							FlushDecoder();
-							InternalDecoderDestroy();
-							if (InternalDecoderCreate(NewFormatDescr))
+							bNeedNewDecoder = DecoderHandle == nullptr || !DecoderHandle->IsCompatibleWith(NewFormatDescr);
+							if (bNeedNewDecoder)
 							{
-								// Ok
+								FlushDecoder();
+								InternalDecoderDestroy();
+								if (InternalDecoderCreate(NewFormatDescr))
+								{
+									// Ok
+								}
+								else
+								{
+									bError = true;
+								}
 							}
-							else
+							CFRelease(NewFormatDescr);
+						}
+						else
+						{
+							bError = true;
+						}
+					}
+
+					// Decode
+					if (!bError && DecoderHandle && DecoderHandle->DecompressionSession)
+					{
+						EDecodeResult DecRes = Decode(AccessUnit, false);
+
+						// Process any output we might have pending
+						ProcessOutput();
+
+						// Decode ok?
+						if (DecRes == EDecodeResult::Ok)
+						{
+							// Yes, add to the replay buffer if it is not a discardable access unit.
+							if (!bIsDiscardable)
 							{
-								bError = true;
+								AccessUnit->AddRef();
+								if (!ReplayAccessUnitBuffer.Push(AccessUnit))
+								{
+									// FIXME: Is this cause for a playback error? For now we just forget about the replay AU and take any possible decoding artefacts.
+									FAccessUnit::Release(AccessUnit);
+								}
 							}
 						}
-						CFRelease(NewFormatDescr);
+						// Lost the decoder session?
+						else if (DecRes == EDecodeResult::SessionLost)
+						{
+							// Did not produce output, release semaphore again.
+							RecreateDecoderSession();
+							// Retry this AU on the new decoder session.
+							continue;
+						}
+						// Decode failed
+						else
+						{
+							bError = true;
+						}
 					}
-					else
+				}
+				else
+				{
+					if (!DecodeDummy(AccessUnit))
 					{
 						bError = true;
 					}
 				}
 
-				// Decode
-				if (!bError && DecoderHandle && DecoderHandle->DecompressionSession)
-				{
-                    EDecodeResult DecRes = Decode(AccessUnit, false);
-
-					// Process any output we might have pending
-					ProcessOutput();
-
-                    // Decode ok?
-                    if (DecRes == EDecodeResult::Ok)
-                    {
-                        // Yes, add to the replay buffer if it is not a discardable access unit.
-                        if (!bIsDiscardable)
-                        {
-                            AccessUnit->AddRef();
-                            if (!ReplayAccessUnitBuffer.Push(AccessUnit))
-                            {
-                                // FIXME: Is this cause for a playback error? For now we just forget about the replay AU and take any possible decoding artefacts.
-                                FAccessUnit::Release(AccessUnit);
-                            }
-                        }
-                    }
-                    // Lost the decoder session?
-                    else if (DecRes == EDecodeResult::SessionLost)
-                    {
-                        // Did not produce output, release semaphore again.
-                        RecreateDecoderSession();
-                        // Retry this AU on the new decoder session.
-                        continue;
-                    }
-                    // Decode failed
-                    else
-                    {
-                        bError = true;
-                    }
-				}
+				FAccessUnit::Release(AccessUnit);
+				AccessUnit = nullptr;
 			}
 			else
 			{
-				if (!DecodeDummy(AccessUnit))
+				ProcessOutput();
+
+				// No data. Is the buffer at EOD?
+				if (AccessUnitBuffer.IsEndOfData())
 				{
-					bError = true;
+					NotifyReadyBufferListener(true);
+					// Are we done yet?
+					if (!bDone && !bError)
+					{
+						bError = !FlushDecoder();
+					}
+					bDone = true;
+					FMediaRunnable::SleepMilliseconds(10);
+					ProcessOutput(true);
 				}
 			}
-
-			FAccessUnit::Release(AccessUnit);
-			AccessUnit = nullptr;
 		}
 		else
 		{
 			ProcessOutput();
-
-			// No data. Is the buffer at EOD?
-			if (AccessUnitBuffer.IsEndOfData())
+			NotifyReadyBufferListener(true);
+			// Are we done yet?
+			if (!bDone && !bError)
 			{
-				NotifyReadyBufferListener(true);
-				// Are we done yet?
-				if (!bDone && !bError)
-				{
-					bError = !FlushDecoder();
-				}
-				bDone = true;
-				FMediaRunnable::SleepMilliseconds(10);
-				ProcessOutput(true);
+				bError = !FlushDecoder();
 			}
+			bDone = true;
+			ProcessOutput(true);
+			break;
 		}
 
 		// Flush?
@@ -1523,6 +1580,19 @@ void FVideoDecoderH264::WorkerThread()
 	AccessUnitBuffer.CapacitySet(0);
 	ReplayAccessUnitBuffer.Flush();
 	ReplayAccessUnitBuffer.CapacitySet(0);
+	if (bDrainForCodecChange)
+	{
+		// Notify the player that we have finished draining.
+		PlayerSessionServices->SendMessageToPlayer(FDecoderMessage::Create(FDecoderMessage::EReason::DrainingFinished, this, EStreamType::Video, FStreamCodecInformation::ECodec::H264));
+		// We need to wait to get terminated. Also check if flushing is requested and acknowledge if it is.
+		while(!TerminateThreadSignal.IsSignaled())
+		{
+			if (FlushDecoderSignal.WaitTimeoutAndReset(1000 * 10))
+			{
+				DecoderFlushedSignal.Signal();
+			}
+		}
+	}
 }
 
 //-----------------------------------------------------------------------------
diff --git a/Engine/Plugins/Media/ElectraPlayer/Source/ElectraPlayerRuntime/Private/Runtime/Decoder/Apple/VideoDecoderH265_Apple.cpp b/Engine/Plugins/Media/ElectraPlayer/Source/ElectraPlayerRuntime/Private/Runtime/Decoder/Apple/VideoDecoderH265_Apple.cpp
new file mode 100644
--- /dev/null
+++ b/Engine/Plugins/Media/ElectraPlayer/Source/ElectraPlayerRuntime/Private/Runtime/Decoder/Apple/VideoDecoderH265_Apple.cpp
@@ -0,0 +1,1688 @@
+// Copyright Epic Games, Inc. All Rights Reserved.
+
+#include "HAL/Platform.h"
+
+#if PLATFORM_MAC || PLATFORM_IOS || PLATFORM_TVOS
+
+#include "PlayerCore.h"
+
+#include "StreamAccessUnitBuffer.h"
+#include "Decoder/VideoDecoderH265.h"
+#include "Renderer/RendererBase.h"
+#include "Player/PlayerSessionServices.h"
+#include "Utilities/StringHelpers.h"
+#include "Utilities/UtilsMPEGVideo.h"
+#include "DecoderErrors_Apple.h"
+#include "HAL/LowLevelMemTracker.h"
+#include "ElectraPlayerPrivate.h"
+#include "MediaVideoDecoderOutputApple.h"
+#include "Renderer/RendererVideo.h"
+
+#if ELECTRA_PLATFORM_HAS_H265_DECODER
+
+#include <VideoToolbox/VideoToolbox.h>
+
+DECLARE_CYCLE_STAT(TEXT("FVideoDecoderH265::Decode()"), STAT_ElectraPlayer_VideoH265Decode, STATGROUP_ElectraPlayer);
+DECLARE_CYCLE_STAT(TEXT("FVideoDecoderH265::ConvertOutput()"), STAT_ElectraPlayer_VideoH265ConvertOutput, STATGROUP_ElectraPlayer);
+
+namespace Electra
+{
+
+/**
+ * H265 video decoder class implementation.
+**/
+class FVideoDecoderH265 : public IVideoDecoderH265, public FMediaThread
+{
+public:
+	static bool Startup(const FParamDict& Options);
+	static void Shutdown();
+
+	FVideoDecoderH265();
+	virtual ~FVideoDecoderH265();
+
+	virtual void SetPlayerSessionServices(IPlayerSessionServices* SessionServices) override;
+
+	virtual void Open(const FInstanceConfiguration& InConfig) override;
+	virtual void Close() override;
+	virtual void DrainForCodecChange() override;
+
+	virtual void SetMaximumDecodeCapability(int32 MaxTier, int32 MaxWidth, int32 MaxHeight, int32 MaxProfile, int32 MaxProfileLevel, const FParamDict& AdditionalOptions) override;
+
+	virtual void SetAUInputBufferListener(IAccessUnitBufferListener* Listener) override;
+
+	virtual void SetReadyBufferListener(IDecoderOutputBufferListener* Listener) override;
+
+	virtual void SetRenderer(TSharedPtr<IMediaRenderer, ESPMode::ThreadSafe> InRenderer) override;
+
+    virtual void SetResourceDelegate(const TSharedPtr<IVideoDecoderResourceDelegate, ESPMode::ThreadSafe>& ResourceDelegate) override;
+
+	virtual IAccessUnitBufferInterface::EAUpushResult AUdataPushAU(FAccessUnit* AccessUnit) override;
+	virtual void AUdataPushEOD() override;
+	virtual void AUdataFlushEverything() override;
+
+private:
+	bool InternalDecoderCreate(CMFormatDescriptionRef InputFormatDescription);
+	void InternalDecoderDestroy();
+	void RecreateDecoderSession();
+	void StartThread();
+	void StopThread();
+	void WorkerThread();
+
+	bool CreateDecodedImagePool();
+	void DestroyDecodedImagePool();
+
+	void NotifyReadyBufferListener(bool bHaveOutput);
+
+	bool AcquireOutputBuffer(IMediaRenderer::IBuffer*& RenderOutputBuffer);
+
+	void PrepareAU(FAccessUnit* AccessUnit, bool& bOutIsIDR, bool& bOutIsDiscardable);
+
+	struct FDecodedImage;
+	void ProcessOutput(bool bFlush = false);
+
+private:
+	bool FlushDecoder();
+	void ClearInDecoderInfos();
+	void FlushPendingImages();
+
+	enum EDecodeResult
+	{
+		Ok,
+		Fail,
+		SessionLost
+	};
+
+	EDecodeResult Decode(FAccessUnit* AccessUnit, bool bRecreatingSession);
+	bool DecodeDummy(FAccessUnit* AccessUnit);
+
+	void PostError(int32_t ApiReturnValue, const FString& Message, uint16 Code, UEMediaError Error = UEMEDIA_ERROR_OK);
+	void LogMessage(IInfoLog::ELevel Level, const FString& Message);
+
+private:
+	enum
+	{
+		NumImagesHoldBackForPTSOrdering = 5,	// Number of frames held back to ensure proper PTS-ordering of decoder output
+		MaxImagesHoldBackForPTSOrdering = 5		// Maximum number of frames to be held in the buffer before we stall the decoder
+	};
+
+	void DecodeCallback(void* pSrcRef, OSStatus status, VTDecodeInfoFlags infoFlags, CVImageBufferRef imageBuffer, CMTime presentationTimeStamp, CMTime presentationDuration);
+	static void _DecodeCallback(void* pUser, void* pSrcRef, OSStatus status, VTDecodeInfoFlags infoFlags, CVImageBufferRef imageBuffer, CMTime presentationTimeStamp, CMTime presentationDuration)
+	{
+		static_cast<FVideoDecoderH265*>(pUser)->DecodeCallback(pSrcRef, status, infoFlags, imageBuffer, presentationTimeStamp, presentationDuration);
+	}
+
+	bool CreateFormatDescription(CMFormatDescriptionRef& OutFormatDescription, const FAccessUnit* InputAccessUnit);
+
+	struct FDecoderFormatInfo
+	{
+		void Reset()
+		{
+			CurrentCodecData.Reset();
+		}
+		bool IsDifferentFrom(const FAccessUnit* InputAccessUnit);
+
+		TSharedPtrTS<const FAccessUnit::CodecData>		CurrentCodecData;
+	};
+
+	struct FDecoderHandle
+	{
+		FDecoderHandle()
+			: FormatDescription(nullptr)
+			, DecompressionSession(nullptr)
+		{
+		}
+		~FDecoderHandle()
+		{
+			Close();
+		}
+
+		void Close()
+		{
+			if (DecompressionSession)
+			{
+				// Note: Do not finish delayed frames here. We don't want them and have to destroy the session immediately!
+				//VTDecompressionSessionFinishDelayedFrames(DecompressionSession);
+				VTDecompressionSessionWaitForAsynchronousFrames(DecompressionSession);
+				VTDecompressionSessionInvalidate(DecompressionSession);
+				CFRelease(DecompressionSession);
+				DecompressionSession = nullptr;
+			}
+			if (FormatDescription)
+			{
+				CFRelease(FormatDescription);
+				FormatDescription = nullptr;
+			}
+		}
+
+		bool IsCompatibleWith(CMFormatDescriptionRef NewFormatDescription)
+		{
+			if (DecompressionSession)
+			{
+				Boolean bIsCompatible = VTDecompressionSessionCanAcceptFormatDescription(DecompressionSession, NewFormatDescription);
+				return bIsCompatible;
+			}
+			return false;
+		}
+		CMFormatDescriptionRef		FormatDescription;
+		VTDecompressionSessionRef	DecompressionSession;
+	};
+
+	struct FInDecoderInfo : public TSharedFromThis<FInDecoderInfo, ESPMode::ThreadSafe>
+	{
+		FStreamCodecInformation		ParsedInfo;
+		FTimeValue					PTS;
+		FTimeValue					DTS;
+		FTimeValue					Duration;
+	};
+
+	struct FDecodedImage
+	{
+		FDecodedImage()
+			: ImageBufferRef(nullptr)
+		{
+		}
+
+		FDecodedImage(const FDecodedImage& rhs)
+			: ImageBufferRef(nullptr)
+		{
+			InternalCopy(rhs);
+		}
+		FDecodedImage& operator=(const FDecodedImage& rhs)
+		{
+			if (this != &rhs)
+			{
+				InternalCopy(rhs);
+			}
+			return *this;
+		}
+
+		bool operator<(const FDecodedImage& rhs) const
+		{
+			return(SourceInfo->PTS < rhs.SourceInfo->PTS);
+		}
+
+		void SetImageBufferRef(CVImageBufferRef InImageBufferRef)
+		{
+			if (ImageBufferRef)
+			{
+				CFRelease(ImageBufferRef);
+				ImageBufferRef = nullptr;
+			}
+			if (InImageBufferRef)
+			{
+				CFRetain(InImageBufferRef);
+				ImageBufferRef = InImageBufferRef;
+			}
+		}
+		CVImageBufferRef GetImageBufferRef()
+		{
+			return ImageBufferRef;
+		}
+
+		CVImageBufferRef ReleaseImageBufferRef()
+		{
+			CVImageBufferRef ref = ImageBufferRef;
+			ImageBufferRef = nullptr;
+			return ref;
+		}
+
+		~FDecodedImage()
+		{
+			if (ImageBufferRef)
+			{
+				CFRelease(ImageBufferRef);
+			}
+		}
+		TSharedPtr<FInDecoderInfo, ESPMode::ThreadSafe>	SourceInfo;
+	private:
+		void InternalCopy(const FDecodedImage& rhs)
+		{
+			SourceInfo = rhs.SourceInfo;
+			SetImageBufferRef(rhs.ImageBufferRef);
+		}
+
+		CVImageBufferRef			ImageBufferRef;
+	};
+
+	FInstanceConfiguration								Config;
+
+	FMediaEvent											TerminateThreadSignal;
+	FMediaEvent											FlushDecoderSignal;
+	FMediaEvent											DecoderFlushedSignal;
+	bool												bThreadStarted;
+
+	IPlayerSessionServices*								PlayerSessionServices;
+
+	TSharedPtr<IMediaRenderer, ESPMode::ThreadSafe>		Renderer;
+
+    TWeakPtr<IVideoDecoderResourceDelegate, ESPMode::ThreadSafe> ResourceDelegate;
+
+	FAccessUnitBuffer									AccessUnitBuffer;
+	FAccessUnitBuffer									ReplayAccessUnitBuffer;
+	bool												bDrainForCodecChange;
+
+	FMediaCriticalSection								ListenerMutex;
+	IAccessUnitBufferListener*							InputBufferListener;
+	IDecoderOutputBufferListener*						ReadyBufferListener;
+
+	FDecoderFormatInfo									CurrentStreamFormatInfo;
+	FDecoderHandle*										DecoderHandle;
+	TArray<TSharedPtr<FInDecoderInfo, ESPMode::ThreadSafe>> InDecoderInfos;
+	FMediaCriticalSection								InDecoderInfoMutex;
+
+	int32												MaxDecodeBufferSize;
+	bool												bError;
+
+	FMediaCriticalSection								ReadyImageMutex;
+	TArray<FDecodedImage>								ReadyImages;
+};
+
+/***************************************************************************************************************************************************/
+/***************************************************************************************************************************************************/
+/***************************************************************************************************************************************************/
+
+bool IVideoDecoderH265::Startup(const FParamDict& Options)
+{
+	return FVideoDecoderH265::Startup(Options);
+}
+
+void IVideoDecoderH265::Shutdown()
+{
+	FVideoDecoderH265::Shutdown();
+}
+
+bool IVideoDecoderH265::GetStreamDecodeCapability(FStreamDecodeCapability& OutResult, const FStreamDecodeCapability& InStreamParameter)
+{
+	return false;
+}
+
+IVideoDecoderH265::FInstanceConfiguration::FInstanceConfiguration()
+	: MaxDecodedFrames(8)
+{
+}
+
+IVideoDecoderH265* IVideoDecoderH265::Create()
+{
+	return new FVideoDecoderH265;
+}
+
+
+/***************************************************************************************************************************************************/
+/***************************************************************************************************************************************************/
+/***************************************************************************************************************************************************/
+
+class FElectraPlayerVideoDecoderOutputH265Apple : public FVideoDecoderOutputApple
+{
+public:
+	FElectraPlayerVideoDecoderOutputH265Apple()
+		: ImageBufferRef(nullptr)
+	{
+	}
+
+	~FElectraPlayerVideoDecoderOutputH265Apple()
+	{
+		if (ImageBufferRef)
+		{
+			CFRelease(ImageBufferRef);
+		}
+	}
+
+    void Initialize(CVImageBufferRef InImageBufferRef, FParamDict* InParamDict)
+	{
+		FVideoDecoderOutputApple::Initialize(InParamDict);
+
+		if (ImageBufferRef)
+		{
+			CFRelease(ImageBufferRef);
+		}
+		ImageBufferRef = InImageBufferRef;
+		if (ImageBufferRef)
+		{
+			CFRetain(ImageBufferRef);
+
+			Stride = CVPixelBufferGetBytesPerRow(ImageBufferRef);
+		}
+	}
+
+	void SetOwner(const TSharedPtr<IDecoderOutputOwner, ESPMode::ThreadSafe>& InOwningRenderer) override
+	{
+		OwningRenderer = InOwningRenderer;
+	}
+
+	void ShutdownPoolable() override
+	{
+		// release image buffer (we currently realloc it every time anyway)
+        if (ImageBufferRef)
+        {
+            CFRelease(ImageBufferRef);
+            ImageBufferRef = nullptr;
+        }
+
+		TSharedPtr<IDecoderOutputOwner, ESPMode::ThreadSafe> lockedVideoRenderer = OwningRenderer.Pin();
+		if (lockedVideoRenderer.IsValid())
+		{
+			lockedVideoRenderer->SampleReleasedToPool(GetDuration());
+		}
+	}
+
+	virtual uint32 GetStride() const override
+	{
+		return Stride;
+	}
+
+	virtual CVImageBufferRef GetImageBuffer() const override
+	{
+		return ImageBufferRef;
+	}
+
+private:
+	uint32 Stride;
+
+	CVImageBufferRef ImageBufferRef;
+
+	// We hold a weak reference to the video renderer. During destruction the video renderer could be destroyed while samples are still out there..
+	TWeakPtr<IDecoderOutputOwner, ESPMode::ThreadSafe> OwningRenderer;
+};
+
+/***************************************************************************************************************************************************/
+/***************************************************************************************************************************************************/
+/***************************************************************************************************************************************************/
+
+//-----------------------------------------------------------------------------
+/**
+ * Decoder system startup
+ *
+ * @param Options
+ *
+ * @return
+ */
+bool FVideoDecoderH265::Startup(const FParamDict& Options)
+{
+	return true;
+}
+
+
+//-----------------------------------------------------------------------------
+/**
+ * Decoder system shutdown.
+ */
+void FVideoDecoderH265::Shutdown()
+{
+}
+
+
+//-----------------------------------------------------------------------------
+/**
+ * Constructor
+ */
+FVideoDecoderH265::FVideoDecoderH265()
+	: FMediaThread("ElectraPlayer::H265 decoder")
+	, bThreadStarted(false)
+	, PlayerSessionServices(nullptr)
+	, bDrainForCodecChange(false)
+	, InputBufferListener(nullptr)
+	, ReadyBufferListener(nullptr)
+	, DecoderHandle(nullptr)
+	, MaxDecodeBufferSize(0)
+	, bError(false)
+{
+}
+
+
+//-----------------------------------------------------------------------------
+/**
+ * Destructor
+ */
+FVideoDecoderH265::~FVideoDecoderH265()
+{
+	Close();
+}
+
+
+//-----------------------------------------------------------------------------
+/**
+ * Sets an AU input buffer listener.
+ *
+ * @param Listener
+ */
+void FVideoDecoderH265::SetAUInputBufferListener(IAccessUnitBufferListener* Listener)
+{
+	FMediaCriticalSection::ScopedLock lock(ListenerMutex);
+	InputBufferListener = Listener;
+}
+
+
+//-----------------------------------------------------------------------------
+/**
+ * Sets a buffer-ready listener.
+ *
+ * @param Listener
+ */
+void FVideoDecoderH265::SetReadyBufferListener(IDecoderOutputBufferListener* Listener)
+{
+	FMediaCriticalSection::ScopedLock lock(ListenerMutex);
+	ReadyBufferListener = Listener;
+}
+
+
+//-----------------------------------------------------------------------------
+/**
+ * Sets the owning player's session service interface.
+ *
+ * @param InSessionServices
+ *
+ * @return
+ */
+void FVideoDecoderH265::SetPlayerSessionServices(IPlayerSessionServices* InSessionServices)
+{
+	PlayerSessionServices = InSessionServices;
+}
+
+
+//-----------------------------------------------------------------------------
+/**
+ * Opens a decoder instance
+ *
+ * @param InConfig
+ *
+ * @return
+ */
+void FVideoDecoderH265::Open(const FInstanceConfiguration& InConfig)
+{
+	Config = InConfig;
+
+	// Set a large enough size to hold a single access unit. Since we will be asking for a new AU on
+	// demand there is no need to overly restrict ourselves here. But it must be large enough to
+	// hold at least the largest expected access unit.
+	AccessUnitBuffer.CapacitySet(FAccessUnitBuffer::FConfiguration(16 << 20, 60.0));
+
+	StartThread();
+}
+
+
+//-----------------------------------------------------------------------------
+/**
+ * Closes the decoder instance.
+ */
+void FVideoDecoderH265::Close()
+{
+	StopThread();
+}
+
+
+//-----------------------------------------------------------------------------
+/**
+ * Drains the decoder of all enqueued input and ends it, after which the decoder must send an FDecoderMessage to the player
+ * to signal completion.
+ */
+void FVideoDecoderH265::DrainForCodecChange()
+{
+	bDrainForCodecChange = true;
+}
+
+
+//-----------------------------------------------------------------------------
+/**
+ * Sets a new decoder limit.
+ * As soon as a new sequence starts matching this limit the decoder will be
+ * destroyed and recreated to conserve memory.
+ * The assumption is that the data being streamed will not exceed this limit,
+ * but if it does, any access unit requiring a better decoder will take
+ * precedence and force a decoder capable of decoding it!
+ *
+ * @param MaxTier
+ * @param MaxWidth
+ * @param MaxHeight
+ * @param MaxProfile
+ * @param MaxProfileLevel
+ * @param AdditionalOptions
+ */
+void FVideoDecoderH265::SetMaximumDecodeCapability(int32 MaxTier, int32 MaxWidth, int32 MaxHeight, int32 MaxProfile, int32 MaxProfileLevel, const FParamDict& AdditionalOptions)
+{
+	// Not implemented
+}
+
+
+//-----------------------------------------------------------------------------
+/**
+ * Sets a new renderer.
+ *
+ * @param InRenderer
+ */
+void FVideoDecoderH265::SetRenderer(TSharedPtr<IMediaRenderer, ESPMode::ThreadSafe> InRenderer)
+{
+	Renderer = InRenderer;
+}
+
+void FVideoDecoderH265::SetResourceDelegate(const TSharedPtr<IVideoDecoderResourceDelegate, ESPMode::ThreadSafe>& InResourceDelegate)
+{
+    ResourceDelegate = InResourceDelegate;
+}
+
+//-----------------------------------------------------------------------------
+/**
+ * Creates and runs the decoder thread.
+ */
+void FVideoDecoderH265::StartThread()
+{
+	ThreadStart(Electra::MakeDelegate(this, &FVideoDecoderH265::WorkerThread));
+	bThreadStarted = true;
+}
+
+
+//-----------------------------------------------------------------------------
+/**
+ * Stops the decoder thread.
+ */
+void FVideoDecoderH265::StopThread()
+{
+	if (bThreadStarted)
+	{
+		TerminateThreadSignal.Signal();
+		ThreadWaitDone();
+		bThreadStarted = false;
+	}
+}
+
+
+//-----------------------------------------------------------------------------
+/**
+ * Posts an error to the session service error listeners.
+ *
+ * @param ApiReturnValue
+ * @param Message
+ * @param Code
+ * @param Error
+ */
+void FVideoDecoderH265::PostError(int32_t ApiReturnValue, const FString& Message, uint16 Code, UEMediaError Error)
+{
+	check(PlayerSessionServices);
+	if (PlayerSessionServices)
+	{
+		FErrorDetail err;
+		err.SetError(Error != UEMEDIA_ERROR_OK ? Error : UEMEDIA_ERROR_DETAIL);
+		err.SetFacility(Facility::EFacility::H265Decoder);
+		err.SetCode(Code);
+		err.SetMessage(Message);
+		err.SetPlatformMessage(FString::Printf(TEXT("%d (0x%08x)"), ApiReturnValue, ApiReturnValue));
+		PlayerSessionServices->PostError(err);
+	}
+}
+
+
+//-----------------------------------------------------------------------------
+/**
+ * Sends a log message to the session service log.
+ *
+ * @param Level
+ * @param Message
+ */
+void FVideoDecoderH265::LogMessage(IInfoLog::ELevel Level, const FString& Message)
+{
+	if (PlayerSessionServices)
+	{
+		PlayerSessionServices->PostLog(Facility::EFacility::H265Decoder, Level, Message);
+	}
+}
+
+
+//-----------------------------------------------------------------------------
+/**
+ * Create a pool of decoded images for the decoder.
+ *
+ * @return
+ */
+bool FVideoDecoderH265::CreateDecodedImagePool()
+{
+	check(Renderer);
+	FParamDict poolOpts;
+
+	poolOpts.Set("num_buffers", FVariantValue((int64) Config.MaxDecodedFrames));
+
+	UEMediaError Error = Renderer->CreateBufferPool(poolOpts);
+	check(Error == UEMEDIA_ERROR_OK);
+
+	MaxDecodeBufferSize = (int32) Renderer->GetBufferPoolProperties().GetValue("max_buffers").GetInt64();
+
+	if (Error != UEMEDIA_ERROR_OK)
+	{
+		PostError(0, "Failed to create image pool", ERRCODE_INTERNAL_APPLE_COULD_NOT_CREATE_IMAGE_POOL, Error);
+	}
+
+	return Error == UEMEDIA_ERROR_OK;
+}
+
+
+//-----------------------------------------------------------------------------
+/**
+ * Destroys the pool of decoded images.
+ */
+void FVideoDecoderH265::DestroyDecodedImagePool()
+{
+	Renderer->ReleaseBufferPool();
+}
+
+
+//-----------------------------------------------------------------------------
+/**
+ * Called to receive a new input access unit for decoding.
+ *
+ * @param AccessUnit
+ */
+IAccessUnitBufferInterface::EAUpushResult FVideoDecoderH265::AUdataPushAU(FAccessUnit* AccessUnit)
+{
+	AccessUnit->AddRef();
+	bool bOk = AccessUnitBuffer.Push(AccessUnit);
+	if (!bOk)
+	{
+		FAccessUnit::Release(AccessUnit);
+	}
+
+	return bOk ? IAccessUnitBufferInterface::EAUpushResult::Ok : IAccessUnitBufferInterface::EAUpushResult::Full;
+}
+
+
+//-----------------------------------------------------------------------------
+/**
+ * "Pushes" an End Of Data marker indicating no further access units will be added.
+ */
+void FVideoDecoderH265::AUdataPushEOD()
+{
+	AccessUnitBuffer.PushEndOfData();
+}
+
+
+//-----------------------------------------------------------------------------
+/**
+ * Flushes the decoder and clears the input access unit buffer.
+ */
+void FVideoDecoderH265::AUdataFlushEverything()
+{
+	FlushDecoderSignal.Signal();
+	DecoderFlushedSignal.WaitAndReset();
+}
+
+
+//-----------------------------------------------------------------------------
+/**
+ * Create a decoder instance.
+ *
+ * @param InputFormatDescription
+ *
+ * @return true if successful, false on error
+ */
+bool FVideoDecoderH265::InternalDecoderCreate(CMFormatDescriptionRef InputFormatDescription)
+{
+	check(DecoderHandle == nullptr);
+	check(InputFormatDescription != nullptr);
+	DecoderHandle = new FDecoderHandle;
+	CFRetain(InputFormatDescription);
+	DecoderHandle->FormatDescription = InputFormatDescription;
+
+	VTDecompressionOutputCallbackRecord CallbackRecord;
+	CallbackRecord.decompressionOutputCallback = _DecodeCallback;
+	CallbackRecord.decompressionOutputRefCon   = this;
+
+	// Output image format configuration
+	CFMutableDictionaryRef OutputImageFormat = CFDictionaryCreateMutable(nullptr, 3, &kCFTypeDictionaryKeyCallBacks, &kCFTypeDictionaryValueCallBacks);
+	// Choice of: kCVPixelBufferOpenGLCompatibilityKey (all)  kCVPixelBufferOpenGLESCompatibilityKey (iOS only)   kCVPixelBufferMetalCompatibilityKey (all)
+#if PLATFORM_MAC
+	int pxfmt = kCVPixelFormatType_32BGRA;
+	CFNumberRef PixelFormat = CFNumberCreate(kCFAllocatorDefault, kCFNumberSInt32Type, &pxfmt);
+	CFDictionarySetValue(OutputImageFormat, kCVPixelBufferPixelFormatTypeKey, PixelFormat);
+	CFRelease(PixelFormat);
+	CFDictionarySetValue(OutputImageFormat, kCVPixelBufferMetalCompatibilityKey, kCFBooleanTrue);
+#elif PLATFORM_IOS
+	CFDictionarySetValue(OutputImageFormat, kCVPixelBufferOpenGLESCompatibilityKey, kCFBooleanFalse);
+	CFDictionarySetValue(OutputImageFormat, kCVPixelBufferMetalCompatibilityKey, kCFBooleanTrue);
+	int pxfmt = kCVPixelFormatType_32BGRA;
+	CFNumberRef PixelFormat = CFNumberCreate(kCFAllocatorDefault, kCFNumberSInt32Type, &pxfmt);
+	CFDictionarySetValue(OutputImageFormat, kCVPixelBufferPixelFormatTypeKey, PixelFormat);
+	CFRelease(PixelFormat);
+#else
+	#error "Should not get here. Check platform checks at the top of the file."
+#endif
+
+	// Session configuration
+	CFMutableDictionaryRef SessionConfiguration = CFDictionaryCreateMutable(nullptr, 1, &kCFTypeDictionaryKeyCallBacks, &kCFTypeDictionaryValueCallBacks);
+	// Ask for hardware decoding
+	//	CFDictionarySetValue(SessionConfiguration, kVTVideoDecoderSpecification_EnableHardwareAcceleratedVideoDecoder, kCFBooleanTrue);
+
+	OSStatus res = VTDecompressionSessionCreate(kCFAllocatorDefault, InputFormatDescription, SessionConfiguration, OutputImageFormat, &CallbackRecord, &DecoderHandle->DecompressionSession);
+	CFRelease(SessionConfiguration);
+	CFRelease(OutputImageFormat);
+	if (res != 0)
+	{
+		PostError(res, "Failed to create video decoder", ERRCODE_INTERNAL_APPLE_COULD_NOT_CREATE_VIDEO_DECODER);
+		return false;
+	}
+	return true;
+}
+
+
+//-----------------------------------------------------------------------------
+/**
+ * Destroys the current decoder instance.
+ */
+void FVideoDecoderH265::InternalDecoderDestroy()
+{
+	if (DecoderHandle)
+	{
+		DecoderHandle->Close();
+		delete DecoderHandle;
+		DecoderHandle = nullptr;
+	}
+}
+
+
+//-----------------------------------------------------------------------------
+/**
+ * Creates a new decoder and runs all AUs since the last IDR frame through
+ * without producing decoded images.
+ * Used to continue decoding after the application has been backgrounded
+ * which results in a decoder session loss when resumed.
+ * Since decoding cannot resume at arbitrary points in the stream everything
+ * from the last IDR frame needs to be decoded again.
+ * To speed this up AUs that have no dependencies are not added to the replay data.
+ */
+void FVideoDecoderH265::RecreateDecoderSession()
+{
+	// Destroy existing decoder. It is of no use any longer.
+	InternalDecoderDestroy();
+	// Likewise anything that was still pending we also no longer need.
+	ClearInDecoderInfos();
+
+	// Do we have replay data?
+	if (ReplayAccessUnitBuffer.Num())
+	{
+		TMediaQueueDynamicNoLock<FAccessUnit *>	ReprocessedAUs;
+
+		bool bDone = false;
+		FAccessUnit* AccessUnit = nullptr;
+		bool bFirst = true;
+		while(!bError && !bDone)
+		{
+			AccessUnit = nullptr;
+			if (ReplayAccessUnitBuffer.Pop(AccessUnit))
+			{
+				check(AccessUnit != nullptr);
+				ReprocessedAUs.Push(AccessUnit);
+				// Create the format description from the first replay AU.
+				if (bFirst)
+				{
+					bFirst = false;
+					CMFormatDescriptionRef NewFormatDescr = nullptr;
+					if (CreateFormatDescription(NewFormatDescr, AccessUnit))
+					{
+						if (InternalDecoderCreate(NewFormatDescr))
+						{
+							// Ok
+						}
+						else
+						{
+							bError = true;
+						}
+						CFRelease(NewFormatDescr);
+					}
+					else
+					{
+						bError = true;
+					}
+				}
+				// Decode
+				EDecodeResult DecRes = Decode(AccessUnit, true);
+				// On failure or yet another loss of decoder session, leave...
+				if (DecRes != EDecodeResult::Ok)
+				{
+					bDone = true;
+				}
+			}
+			else
+			{
+				bDone = true;
+			}
+		}
+		// Even in case of an error we need to get all replay AUs into our processed FIFO and
+		// from there back into the replay buffer. We may need them again and they need to
+		// stay in the original order.
+		while(ReplayAccessUnitBuffer.Pop(AccessUnit))
+		{
+			ReprocessedAUs.Push(AccessUnit);
+		}
+		while(!ReprocessedAUs.IsEmpty())
+		{
+			AccessUnit = ReprocessedAUs.Pop();
+			ReplayAccessUnitBuffer.Push(AccessUnit);
+		}
+	// Flush the decoder to get it idle and discard any accumulated source infos.
+	FlushDecoder();
+	}
+}
+
+
+//-----------------------------------------------------------------------------
+/**
+ * Notifies the buffer-ready listener that we will now be producing output data
+ * and need an output buffer.
+ * This call is not intended to create or obtain an output buffer. It is merely
+ * indicating that new output will be produced.
+ *
+ * @param bHaveOutput
+ */
+void FVideoDecoderH265::NotifyReadyBufferListener(bool bHaveOutput)
+{
+	if (ReadyBufferListener)
+	{
+		IDecoderOutputBufferListener::FDecodeReadyStats stats;
+		stats.MaxDecodedElementsReady = MaxDecodeBufferSize;
+		stats.NumElementsInDecoder    = InDecoderInfos.Num();
+		stats.bOutputStalled		  = !bHaveOutput;
+		stats.bEODreached   		  = AccessUnitBuffer.IsEndOfData() && stats.NumDecodedElementsReady == 0 && stats.NumElementsInDecoder == 0;
+		ListenerMutex.Lock();
+		if (ReadyBufferListener)
+		{
+			ReadyBufferListener->DecoderOutputReady(stats);
+		}
+		ListenerMutex.Unlock();
+	}
+}
+
+
+//-----------------------------------------------------------------------------
+/**
+ * Prepares the passed access unit to be sent to the decoder.
+ *
+ * @param AccessUnit
+ * @param bOutIsIDR
+ * @param bOutIsDiscardable
+ */
+void FVideoDecoderH265::PrepareAU(FAccessUnit* AccessUnit, bool& bOutIsIDR, bool& bOutIsDiscardable)
+{
+	bOutIsIDR = false;
+	bOutIsDiscardable = false;
+	if (!AccessUnit->bHasBeenPrepared)
+	{
+		AccessUnit->bHasBeenPrepared = true;
+
+		// Process NALUs
+		// Note: With VideoToolbox the data must be in avcc format, that is no startcodes but the
+		//       length of each NALU. We already have that so there is not a lot to do here.
+		uint32* CurrentNALU = (uint32 *)AccessUnit->AUData;
+		uint32* LastNALU    = (uint32 *)Electra::AdvancePointer(CurrentNALU, AccessUnit->AUSize);
+		while(CurrentNALU < LastNALU)
+		{
+			uint8 nut = *(const uint8 *)(CurrentNALU + 1);
+			nut >>= 1;
+			// IDR frame?
+			if (nut == 19 /*IDR_W_RADL*/ || nut == 20 /*IDR_N_LP*/)
+			{
+                bOutIsIDR = true;
+			}
+
+			uint32 naluLen = MEDIA_FROM_BIG_ENDIAN(*CurrentNALU) + 4;
+			CurrentNALU = Electra::AdvancePointer(CurrentNALU, naluLen);
+		}
+
+		// Is there codec specific data?
+		if (AccessUnit->AUCodecData.IsValid())
+		{
+			// Yes.
+			if (bOutIsIDR || AccessUnit->bIsSyncSample || AccessUnit->bIsFirstInSequence)
+			{
+				// Have to re-allocate the AU memory to preprend the codec data
+				uint64 nb = AccessUnit->AUSize + AccessUnit->AUCodecData->CodecSpecificData.Num();
+				void* pD = AccessUnit->AllocatePayloadBuffer(nb);
+				void* pP = pD;
+				FMemory::Memcpy(pP, AccessUnit->AUCodecData->CodecSpecificData.GetData(), AccessUnit->AUCodecData->CodecSpecificData.Num());
+				pP = Electra::AdvancePointer(pP, AccessUnit->AUCodecData->CodecSpecificData.Num());
+				FMemory::Memcpy(pP, AccessUnit->AUData, AccessUnit->AUSize);
+				AccessUnit->AdoptNewPayloadBuffer(pD, nb);
+				// New codec data makes this AU non-discardable.
+				bOutIsDiscardable = false;
+
+				// Get the NALUs from the CSD.
+				TArray<MPEG::FNaluInfo>	NALUs;
+				MPEG::ParseBitstreamForNALUs(NALUs, AccessUnit->AUCodecData->CodecSpecificData.GetData(), AccessUnit->AUCodecData->CodecSpecificData.Num());
+				// Replace the startcodes in the CSD with length values
+				for(int32 i=0; i<NALUs.Num(); ++i)
+				{
+					uint8* NALU = (uint8*)Electra::AdvancePointer(pD, NALUs[i].Offset);
+					*(uint32*)NALU = MEDIA_TO_BIG_ENDIAN((uint32)NALUs[i].Size);
+				}
+			}
+		}
+	}
+}
+
+
+//-----------------------------------------------------------------------------
+/**
+ * Gets an output buffer from the video renderer.
+ *
+ * @param RenderOutputBuffer
+ *
+ * @return true if successful, false on error
+ */
+bool FVideoDecoderH265::AcquireOutputBuffer(IMediaRenderer::IBuffer*& RenderOutputBuffer)
+{
+	RenderOutputBuffer = nullptr;
+	FParamDict BufferAcquireOptions;
+	while(!TerminateThreadSignal.IsSignaled() && !FlushDecoderSignal.IsSignaled())
+	{
+		UEMediaError bufResult = Renderer->AcquireBuffer(RenderOutputBuffer, 0, BufferAcquireOptions);
+		check(bufResult == UEMEDIA_ERROR_OK || bufResult == UEMEDIA_ERROR_INSUFFICIENT_DATA);
+		if (bufResult != UEMEDIA_ERROR_OK && bufResult != UEMEDIA_ERROR_INSUFFICIENT_DATA)
+		{
+			PostError(0, "Failed to acquire sample buffer", ERRCODE_INTERNAL_APPLE_COULD_NOT_GET_OUTPUT_BUFFER, bufResult);
+			return false;
+		}
+		bool bHaveAvailSmpBlk = RenderOutputBuffer != nullptr;
+		NotifyReadyBufferListener(bHaveAvailSmpBlk);
+		if (bHaveAvailSmpBlk)
+		{
+			break;
+		}
+		else
+		{
+			// No available buffer. Sleep for a bit. Can't sleep on a signal since we have to check two: abort and flush
+			FMediaRunnable::SleepMilliseconds(5);
+		}
+	}
+	return true;
+}
+
+
+//-----------------------------------------------------------------------------
+/**
+ * Flushes the decoder and passes all pending frames to the renderer.
+ *
+ * @return true if successful, false otherwise.
+ */
+bool FVideoDecoderH265::FlushDecoder()
+{
+	if (DecoderHandle && DecoderHandle->DecompressionSession)
+	{
+		VTDecompressionSessionFinishDelayedFrames(DecoderHandle->DecompressionSession);
+		VTDecompressionSessionWaitForAsynchronousFrames(DecoderHandle->DecompressionSession);
+		// Push out all images we got from the decoder from our internal PTS sorting facility
+		ProcessOutput(true);
+	}
+	ClearInDecoderInfos();
+	return true;
+}
+
+
+//-----------------------------------------------------------------------------
+/**
+ * Clears out any leftover decoder source infos that were not picked by decoded frames.
+ * This must be called after having flushed the decoder to ensure no delayed decode
+ * callbacks will arrive.
+ */
+void FVideoDecoderH265::ClearInDecoderInfos()
+{
+	FMediaCriticalSection::ScopedLock lock(InDecoderInfoMutex);
+	InDecoderInfos.Empty();
+}
+
+
+//-----------------------------------------------------------------------------
+/**
+ * Creates a dummy output image that is not to be displayed and has no image data.
+ * Dummy access units are created when stream data is missing to ensure the data
+ * pipeline does not run dry and exhibits no gaps in the timeline.
+ *
+ * @param AccessUnit
+ *
+ * @return true if successful, false otherwise
+ */
+bool FVideoDecoderH265::DecodeDummy(FAccessUnit* AccessUnit)
+{
+	check(AccessUnit->Duration.IsValid());
+
+    FDecodedImage NextImage;
+    NextImage.SourceInfo = TSharedPtr<FInDecoderInfo, ESPMode::ThreadSafe>(new FInDecoderInfo);
+    NextImage.SourceInfo->DTS = AccessUnit->DTS;
+    NextImage.SourceInfo->PTS = AccessUnit->PTS;
+    NextImage.SourceInfo->Duration = AccessUnit->Duration;
+    if (AccessUnit->AUCodecData.IsValid())
+    {
+        NextImage.SourceInfo->ParsedInfo = AccessUnit->AUCodecData->ParsedInfo;
+    }
+
+	ReadyImageMutex.Lock();
+	ReadyImages.Add(NextImage);
+	ReadyImages.Sort();
+	ReadyImageMutex.Unlock();
+
+    return true;
+}
+
+//-----------------------------------------------------------------------------
+/**
+ * Flushes all images not passed to the renderer from our list.
+ */
+void FVideoDecoderH265::FlushPendingImages()
+{
+	// Clear out the map. This implicitly drops all image refcounts.
+	ReadyImageMutex.Lock();
+	ReadyImages.Empty();
+	ReadyImageMutex.Unlock();
+}
+
+//-----------------------------------------------------------------------------
+/**
+ * Checks if the codec specific data has changed.
+ *
+ * @return false if the format is still the same, true if it has changed.
+ */
+bool FVideoDecoderH265::FDecoderFormatInfo::IsDifferentFrom(const FAccessUnit* InputAccessUnit)
+{
+	if (InputAccessUnit->AUCodecData.IsValid() && InputAccessUnit->AUCodecData.Get() != CurrentCodecData.Get())
+	{
+		CurrentCodecData = InputAccessUnit->AUCodecData;
+		return true;
+	}
+	return false;
+}
+
+
+//-----------------------------------------------------------------------------
+/**
+ * Creates a source format description to be passed into the video decoder from
+ * the codec specific data (CSD) attached to the input access unit.
+ *
+ * @param OutFormatDescription
+ * @param InputAccessUnit
+ *
+ * @return true if successful, false otherwise.
+ */
+bool FVideoDecoderH265::CreateFormatDescription(CMFormatDescriptionRef& OutFormatDescription, const FAccessUnit* InputAccessUnit)
+{
+	if (InputAccessUnit && InputAccessUnit->AUCodecData.IsValid() && InputAccessUnit->AUCodecData->CodecSpecificData.Num())
+	{
+		// Get the NALUs from the CSD.
+		TArray<MPEG::FNaluInfo>	NALUs;
+		MPEG::ParseBitstreamForNALUs(NALUs, InputAccessUnit->AUCodecData->CodecSpecificData.GetData(), InputAccessUnit->AUCodecData->CodecSpecificData.Num());
+		if (NALUs.Num())
+		{
+			int32 NumRecords = NALUs.Num();
+			if (NumRecords)
+			{
+				uint8_t const* * DataPointers = new uint8_t const* [NumRecords];
+				SIZE_T*          DataSizes    = new SIZE_T [NumRecords];
+				CFDictionaryRef NoExtras = nullptr;
+				for(int32 i=0; i<NumRecords; ++i)
+				{
+					DataPointers[i] = Electra::AdvancePointer(InputAccessUnit->AUCodecData->CodecSpecificData.GetData(), NALUs[i].Offset + NALUs[i].UnitLength);
+					DataSizes[i]    = NALUs[i].Size;
+				}
+				OSStatus res = CMVideoFormatDescriptionCreateFromHEVCParameterSets(kCFAllocatorDefault, NumRecords, DataPointers, DataSizes, 4, NoExtras, &OutFormatDescription);
+				delete [] DataPointers;
+				delete [] DataSizes;
+				if (res == 0)
+				{
+					return true;
+				}
+				else
+				{
+					if (OutFormatDescription)
+					{
+						CFRelease(OutFormatDescription);
+						OutFormatDescription = nullptr;
+					}
+					PostError(res, "Failed to create video format description from CSD", ERRCODE_INTERNAL_APPLE_BAD_VIDEO_CSD);
+					return false;
+				}
+			}
+		}
+		PostError(0, "Failed to create video format description from CSD", ERRCODE_INTERNAL_APPLE_BAD_VIDEO_CSD);
+		return false;
+	}
+	PostError(0, "Cannot create video format description from empty CSD", ERRCODE_INTERNAL_APPLE_NO_VIDEO_CSD);
+	return false;
+}
+
+
+//-----------------------------------------------------------------------------
+/**
+ * Callback from the video decoder when a new decoded image is ready.
+ *
+ * @param pSrcRef
+ * @param status
+ * @param infoFlags
+ * @param imageBuffer
+ * @param presentationTimeStamp
+ * @param presentationDuration
+ */
+void FVideoDecoderH265::DecodeCallback(void* pSrcRef, OSStatus status, VTDecodeInfoFlags infoFlags, CVImageBufferRef imageBuffer, CMTime presentationTimeStamp, CMTime presentationDuration)
+{
+	// Remove the source info even if there ultimately was a decode error or if the frame was dropped.
+	TSharedPtr<FInDecoderInfo, ESPMode::ThreadSafe> SourceInfo;
+	InDecoderInfoMutex.Lock();
+	int32 NumCurrentDecodeInfos = InDecoderInfos.Num();
+	for(int32 i=0; i<NumCurrentDecodeInfos; ++i)
+	{
+		if (InDecoderInfos[i].Get() == pSrcRef)
+		{
+			SourceInfo = InDecoderInfos[i];
+			InDecoderInfos.RemoveSingle(SourceInfo);
+			break;
+		}
+	}
+	InDecoderInfoMutex.Unlock();
+
+	if (!SourceInfo.IsValid())
+	{
+		LogMessage(IInfoLog::ELevel::Error, FString::Printf(TEXT("FVideoDecoderH265::DecodeCallback(): No source info found for decoded srcref %p in %d pending infos (OSStatus %d, infoFlags %d)"), pSrcRef, NumCurrentDecodeInfos, (int32)status, (int32)infoFlags));
+	}
+
+	if (status == 0)
+	{
+		if (imageBuffer != nullptr && (infoFlags & kVTDecodeInfo_FrameDropped) == 0 && SourceInfo.IsValid())
+		{
+			FDecodedImage NextImage;
+			NextImage.SourceInfo = SourceInfo;
+			NextImage.SetImageBufferRef(imageBuffer);
+
+			// Recall decoded frame for later processing
+			// (we do all processing of output on the decoder thread)
+ 			ReadyImageMutex.Lock();
+			ReadyImages.Add(NextImage);
+			ReadyImages.Sort();
+			ReadyImageMutex.Unlock();
+		}
+	}
+	else
+	{
+		bError = true;
+		PostError(status, "Failed to decode video", ERRCODE_INTERNAL_APPLE_FAILED_TO_DECODE_VIDEO);
+	}
+}
+
+
+//-----------------------------------------------------------------------------
+/**
+ * Sends an access unit to the decoder for decoding.
+ *
+ * @param AccessUnit
+ * @param bRecreatingSession
+ *
+ * @return
+ */
+FVideoDecoderH265::EDecodeResult FVideoDecoderH265::Decode(FAccessUnit* AccessUnit, bool bRecreatingSession)
+{
+	if (!DecoderHandle || !DecoderHandle->DecompressionSession)
+	{
+		return EDecodeResult::Fail;
+	}
+
+	check(AccessUnit->Duration.IsValid());
+	check(AccessUnit->PTS.IsValid());
+	check(AccessUnit->DTS.IsValid());
+
+	// Create a memory block for the access unit. We pass it the memory we already have and ensure that by setting the block allocator
+	// to kCFAllocatorNull no one will attempt to deallocate the memory!
+	CMBlockBufferRef AUDataBlock = nullptr;
+	SIZE_T AUDataSize = AccessUnit->AUSize;
+	OSStatus res = CMBlockBufferCreateWithMemoryBlock(nullptr, AccessUnit->AUData, AUDataSize, kCFAllocatorNull, nullptr, 0, AUDataSize, 0, &AUDataBlock);
+	if (res)
+	{
+		PostError(res, "Failed to create video data block buffer", ERRCODE_INTERNAL_APPLE_FAILED_TO_CREATE_BLOCK_BUFFER);
+		return EDecodeResult::Fail;
+	}
+
+	// Set up the timing info with DTS, PTS and duration.
+	CMSampleTimingInfo TimingInfo;
+	const int64_t HNS_PER_S = 10000000;
+	TimingInfo.decodeTimeStamp       = CMTimeMake(AccessUnit->DTS.GetAsHNS(), HNS_PER_S);
+	TimingInfo.presentationTimeStamp = CMTimeMake(AccessUnit->PTS.GetAsHNS(), HNS_PER_S);
+	TimingInfo.duration              = CMTimeMake(AccessUnit->Duration.GetAsHNS(), HNS_PER_S);
+
+	CMSampleBufferRef SampleBufferRef = nullptr;
+	res = CMSampleBufferCreate(kCFAllocatorDefault, AUDataBlock, true, nullptr, nullptr, DecoderHandle->FormatDescription, 1, 1, &TimingInfo, 1, &AUDataSize, &SampleBufferRef);
+	// The buffer is now held by the sample, so we release our ref count.
+	CFRelease(AUDataBlock);
+	if (res)	// see CMSampleBuffer for kCMSampleBufferError_AllocationFailed and such.
+	{
+		PostError(res, "Failed to create video sample buffer", ERRCODE_INTERNAL_APPLE_FAILED_TO_CREATE_SAMPLE_BUFFER);
+		return EDecodeResult::Fail;
+	}
+
+	TSharedPtr<FInDecoderInfo, ESPMode::ThreadSafe> SourceInfo(new FInDecoderInfo);
+	SourceInfo->DTS = AccessUnit->DTS;
+	SourceInfo->PTS = AccessUnit->PTS;
+	SourceInfo->Duration = AccessUnit->Duration;
+	if (AccessUnit->AUCodecData.IsValid())
+	{
+		SourceInfo->ParsedInfo = AccessUnit->AUCodecData->ParsedInfo;
+	}
+	InDecoderInfoMutex.Lock();
+	InDecoderInfos.Push(SourceInfo);
+	InDecoderInfoMutex.Unlock();
+
+	// Decode
+/*
+	kVTDecodeFrame_EnableAsynchronousDecompression = 1<<0,
+	kVTDecodeFrame_DoNotOutputFrame = 1<<1,
+	kVTDecodeFrame_1xRealTimePlayback = 1<<2,
+	kVTDecodeFrame_EnableTemporalProcessing = 1<<3,
+*/
+// TODO: set the proper flags. This may require a larger intermediate output queue. Needs experimenting.
+	VTDecodeFrameFlags DecodeFlags = kVTDecodeFrame_EnableAsynchronousDecompression | kVTDecodeFrame_EnableTemporalProcessing;
+	if (bRecreatingSession)
+	{
+		DecodeFlags = kVTDecodeFrame_DoNotOutputFrame;
+	}
+	VTDecodeInfoFlags  InfoFlags = 0;
+/*
+	kVTDecodeInfo_Asynchronous = 1UL << 0,
+	kVTDecodeInfo_FrameDropped = 1UL << 1,
+	kVTDecodeInfo_ImageBufferModifiable = 1UL << 2,
+*/
+// FIXME: This will require access to the decoder that should be granted by some arbitrator.
+//        Especially for iOS/padOS/tvos where the app can be backgrounded and the decoder potentially becoming inaccessible
+	res = VTDecompressionSessionDecodeFrame(DecoderHandle->DecompressionSession, SampleBufferRef, DecodeFlags, SourceInfo.Get(), &InfoFlags);
+	CFRelease(SampleBufferRef);
+	if (res == 0)
+	{
+		// Ok.
+		return EDecodeResult::Ok;
+	}
+	else
+	{
+		InDecoderInfoMutex.Lock();
+		InDecoderInfos.RemoveSingle(SourceInfo);
+		InDecoderInfoMutex.Unlock();
+		if (res == kVTInvalidSessionErr)
+		{
+			// Lost the decoder session due to being backgrounded and returning to the foreground.
+			return EDecodeResult::SessionLost;
+		}
+		PostError(res, "Failed to decode video frame", ERRCODE_INTERNAL_APPLE_FAILED_TO_DECODE_VIDEO);
+		return EDecodeResult::Fail;
+	}
+}
+
+
+//-----------------------------------------------------------------------------
+/**
+ * H265 video decoder main threaded decode loop
+ */
+void FVideoDecoderH265::WorkerThread()
+{
+	LLM_SCOPE(ELLMTag::ElectraPlayer);
+
+	bool bDone  = false;
+    bool bGotLastSequenceAU = false;
+
+	bError = false;
+
+	// If the application is suspended and resumed we will lose out decoder session and cannot continue
+	// decoding at the point of interruption. We have to keep the access units from the previous IDR frame
+	// and run them through a new decoder session again, discarding all output, after which we can resume
+	// decoding from the point of interruption.
+	// Configure the playback AU buffer to be large enough to hold the access units to be replayed.
+	ReplayAccessUnitBuffer.CapacitySet(FAccessUnitBuffer::FConfiguration(64 << 20, 360.0));
+
+	// Create decoded image pool.
+	if (!CreateDecodedImagePool())
+	{
+		bError = true;
+	}
+
+	while(!TerminateThreadSignal.IsSignaled())
+	{
+		if (!bDrainForCodecChange)
+		{
+			// Notify optional buffer listener that we will now be needing an AU for our input buffer.
+			if (!bError && InputBufferListener && AccessUnitBuffer.Num() == 0)
+			{
+				SCOPE_CYCLE_COUNTER(STAT_ElectraPlayer_VideoH265Decode);
+				CSV_SCOPED_TIMING_STAT(ElectraPlayer, VideoH265Decode);
+				FAccessUnitBufferInfo	sin;
+				IAccessUnitBufferListener::FBufferStats	stats;
+				AccessUnitBuffer.GetStats(sin);
+				stats.NumAUsAvailable  = sin.NumCurrentAccessUnits;
+				stats.NumBytesInBuffer = sin.CurrentMemInUse;
+				stats.MaxBytesOfBuffer = sin.MaxDataSize;
+				stats.bEODSignaled     = sin.bEndOfData;
+				stats.bEODReached      = sin.bEndOfData && sin.NumCurrentAccessUnits == 0;
+				ListenerMutex.Lock();
+				if (InputBufferListener)
+				{
+					InputBufferListener->DecoderInputNeeded(stats);
+				}
+				ListenerMutex.Unlock();
+			}
+			// Wait for data.
+			bool bHaveData = AccessUnitBuffer.WaitForData(1000 * 5);
+
+			// Only send new data to the decoder if we know we got enough room (to avoid accumulating too many frames in our internal PTS-sort queue)
+			bool bTooManyImagesWaiting = false;
+			if (bHaveData)
+			{
+				// When there is data, even and especially after a previous EOD, we are no longer done and idling.
+				if (bDone)
+				{
+					bDone = false;
+				}
+
+				ReadyImageMutex.Lock();
+				bTooManyImagesWaiting = (ReadyImages.Num() > MaxImagesHoldBackForPTSOrdering);
+				ReadyImageMutex.Unlock();
+				if (bTooManyImagesWaiting)
+				{
+					// Signal blockage and wait for it to clear up before going back to normal decoder work...
+					NotifyReadyBufferListener(false);
+
+					while(!TerminateThreadSignal.IsSignaled() && !FlushDecoderSignal.IsSignaled())
+					{
+						ProcessOutput();
+
+						ReadyImageMutex.Lock();
+						bTooManyImagesWaiting = (ReadyImages.Num() > MaxImagesHoldBackForPTSOrdering);
+						ReadyImageMutex.Unlock();
+						if (!bTooManyImagesWaiting)
+						{
+							break;
+						}
+
+						FMediaRunnable::SleepMilliseconds(20);
+					}
+				}
+			}
+
+			if (bHaveData && !bTooManyImagesWaiting)
+			{
+				FAccessUnit* AccessUnit = nullptr;
+				bool bOk = AccessUnitBuffer.Pop(AccessUnit);
+				MEDIA_UNUSED_VAR(bOk);
+				check(bOk);
+
+				if (!AccessUnit->bIsDummyData)
+				{
+					bool bIsKeyframe, bIsDiscardable;
+					PrepareAU(AccessUnit, bIsKeyframe, bIsDiscardable);
+					// An IDR frame means we can start decoding there, so we can purge any accumulated replay AUs.
+					if (bIsKeyframe)
+					{
+						ReplayAccessUnitBuffer.Flush();
+					}
+
+					bool bStreamFormatChanged = CurrentStreamFormatInfo.IsDifferentFrom(AccessUnit) || bGotLastSequenceAU;
+					bool bNeedNewDecoder = false;
+					bGotLastSequenceAU = AccessUnit->bIsLastInPeriod;
+					if (bStreamFormatChanged || !DecoderHandle)
+					{
+						SCOPE_CYCLE_COUNTER(STAT_ElectraPlayer_VideoH265Decode);
+						CSV_SCOPED_TIMING_STAT(ElectraPlayer, VideoH265Decode);
+
+						CMFormatDescriptionRef NewFormatDescr = nullptr;
+						if (CreateFormatDescription(NewFormatDescr, AccessUnit))
+						{
+							bNeedNewDecoder = DecoderHandle == nullptr || !DecoderHandle->IsCompatibleWith(NewFormatDescr);
+							if (bNeedNewDecoder)
+							{
+								FlushDecoder();
+								InternalDecoderDestroy();
+								if (InternalDecoderCreate(NewFormatDescr))
+								{
+									// Ok
+								}
+								else
+								{
+									bError = true;
+								}
+							}
+							CFRelease(NewFormatDescr);
+						}
+						else
+						{
+							bError = true;
+						}
+					}
+
+					// Decode
+					if (!bError && DecoderHandle && DecoderHandle->DecompressionSession)
+					{
+						EDecodeResult DecRes = Decode(AccessUnit, false);
+
+						// Process any output we might have pending
+						ProcessOutput();
+
+						// Decode ok?
+						if (DecRes == EDecodeResult::Ok)
+						{
+							// Yes, add to the replay buffer if it is not a discardable access unit.
+							if (!bIsDiscardable)
+							{
+								AccessUnit->AddRef();
+								if (!ReplayAccessUnitBuffer.Push(AccessUnit))
+								{
+									// FIXME: Is this cause for a playback error? For now we just forget about the replay AU and take any possible decoding artefacts.
+									FAccessUnit::Release(AccessUnit);
+								}
+							}
+						}
+						// Lost the decoder session?
+						else if (DecRes == EDecodeResult::SessionLost)
+						{
+							// Did not produce output, release semaphore again.
+							RecreateDecoderSession();
+							// Retry this AU on the new decoder session.
+							continue;
+						}
+						// Decode failed
+						else
+						{
+							bError = true;
+						}
+					}
+				}
+				else
+				{
+					if (!DecodeDummy(AccessUnit))
+					{
+						bError = true;
+					}
+				}
+
+				FAccessUnit::Release(AccessUnit);
+				AccessUnit = nullptr;
+			}
+			else
+			{
+				ProcessOutput();
+
+				// No data. Is the buffer at EOD?
+				if (AccessUnitBuffer.IsEndOfData())
+				{
+					NotifyReadyBufferListener(true);
+					// Are we done yet?
+					if (!bDone && !bError)
+					{
+						bError = !FlushDecoder();
+					}
+					bDone = true;
+					FMediaRunnable::SleepMilliseconds(10);
+					ProcessOutput(true);
+				}
+			}
+		}
+		else
+		{
+			ProcessOutput();
+			NotifyReadyBufferListener(true);
+			// Are we done yet?
+			if (!bDone && !bError)
+			{
+				bError = !FlushDecoder();
+			}
+			bDone = true;
+			ProcessOutput(true);
+			break;
+		}
+
+		// Flush?
+		if (FlushDecoderSignal.IsSignaled())
+		{
+			SCOPE_CYCLE_COUNTER(STAT_ElectraPlayer_VideoH265Decode);
+			CSV_SCOPED_TIMING_STAT(ElectraPlayer, VideoH265Decode);
+
+			// Have to destroy the decoder!
+			InternalDecoderDestroy();
+			FlushPendingImages();
+			ClearInDecoderInfos();
+			AccessUnitBuffer.Flush();
+			ReplayAccessUnitBuffer.Flush();
+
+			FlushDecoderSignal.Reset();
+			DecoderFlushedSignal.Signal();
+
+			// Reset done state.
+			bDone = false;
+		}
+	}
+
+	InternalDecoderDestroy();
+	FlushPendingImages();
+	DestroyDecodedImagePool();
+	ClearInDecoderInfos();
+	AccessUnitBuffer.Flush();
+	AccessUnitBuffer.CapacitySet(0);
+	ReplayAccessUnitBuffer.Flush();
+	ReplayAccessUnitBuffer.CapacitySet(0);
+	if (bDrainForCodecChange)
+	{
+		// Notify the player that we have finished draining.
+		PlayerSessionServices->SendMessageToPlayer(FDecoderMessage::Create(FDecoderMessage::EReason::DrainingFinished, this, EStreamType::Video, FStreamCodecInformation::ECodec::H265));
+		// We need to wait to get terminated. Also check if flushing is requested and acknowledge if it is.
+		while(!TerminateThreadSignal.IsSignaled())
+		{
+			if (FlushDecoderSignal.WaitTimeoutAndReset(1000 * 10))
+			{
+				DecoderFlushedSignal.Signal();
+			}
+		}
+	}
+}
+
+//-----------------------------------------------------------------------------
+/**
+ */
+void FVideoDecoderH265::ProcessOutput(bool bFlush)
+{
+	SCOPE_CYCLE_COUNTER(STAT_ElectraPlayer_VideoH265ConvertOutput);
+	CSV_SCOPED_TIMING_STAT(ElectraPlayer, VideoH265ConvertOutput);
+
+	// Pull any frames we can get from in-PTS-order array...
+	while(1)
+	{
+		ReadyImageMutex.Lock();
+		if (ReadyImages.Num() < (bFlush ? 1 : NumImagesHoldBackForPTSOrdering))
+		{
+			ReadyImageMutex.Unlock();
+			break;
+		}
+
+		if (!Renderer->CanReceiveOutputFrames(1))
+		{
+			NotifyReadyBufferListener(false);
+			ReadyImageMutex.Unlock();
+			break;
+		}
+
+		FDecodedImage NextImage(ReadyImages[0]);
+		ReadyImages.RemoveAt(0);
+
+		ReadyImageMutex.Unlock();
+
+		// Get an output buffer from the renderer to pass the image to.
+		IMediaRenderer::IBuffer* RenderOutputBuffer = nullptr;
+		if (AcquireOutputBuffer(RenderOutputBuffer))
+		{
+			if (RenderOutputBuffer)
+			{
+                double PixelAspectRatio = 1.0;
+
+                // Note that the ImageBuffer reference below might be null if this is a dummy frame!
+                CVImageBufferRef ImageBufferRef = NextImage.ReleaseImageBufferRef();
+
+				FParamDict* OutputBufferSampleProperties = new FParamDict();
+				if (ImageBufferRef)
+				{
+					// Start with a safe 1:1 aspect ratio assumption.
+					long ax = 1;
+					long ay = 1;
+					// Get the initial aspect ratio from the parsed codec info.
+					if (NextImage.SourceInfo->ParsedInfo.GetAspectRatio().IsSet())
+					{
+						ax = NextImage.SourceInfo->ParsedInfo.GetAspectRatio().Width;
+						ay = NextImage.SourceInfo->ParsedInfo.GetAspectRatio().Height;
+					}
+					// If there is aspect ratio information on the image itself and it's valid, use that instead.
+					NSDictionary* Dict = (NSDictionary*)CVBufferGetAttachments(ImageBufferRef, kCVAttachmentMode_ShouldPropagate);
+					if (Dict)
+					{
+						NSDictionary* AspectDict = (NSDictionary*)Dict[(__bridge NSString*)kCVImageBufferPixelAspectRatioKey];
+						if (AspectDict)
+						{
+							NSNumber* hs = (NSNumber*)AspectDict[(__bridge NSString*)kCVImageBufferPixelAspectRatioHorizontalSpacingKey];
+							NSNumber* vs = (NSNumber*)AspectDict[(__bridge NSString*)kCVImageBufferPixelAspectRatioVerticalSpacingKey];
+							if (hs && vs)
+							{
+								long parx = [hs longValue];
+								long pary = [vs longValue];
+								if (parx && pary)
+								{
+									ax = parx;
+									ay = pary;
+								}
+							}
+						}
+					}
+					PixelAspectRatio = (double)ax / (double)ay;
+
+					OutputBufferSampleProperties->Set("width",        FVariantValue((int64) NextImage.SourceInfo->ParsedInfo.GetResolution().Width));
+					OutputBufferSampleProperties->Set("height",       FVariantValue((int64) NextImage.SourceInfo->ParsedInfo.GetResolution().Height));
+					OutputBufferSampleProperties->Set("crop_left",    FVariantValue((int64) 0));
+					OutputBufferSampleProperties->Set("crop_right",   FVariantValue((int64) 0));
+					OutputBufferSampleProperties->Set("crop_top",     FVariantValue((int64) 0));
+					OutputBufferSampleProperties->Set("crop_bottom",  FVariantValue((int64) 0));
+					OutputBufferSampleProperties->Set("aspect_ratio", FVariantValue((double) PixelAspectRatio));
+					OutputBufferSampleProperties->Set("aspect_w",     FVariantValue((int64) ax));
+					OutputBufferSampleProperties->Set("aspect_h",     FVariantValue((int64) ay));
+					OutputBufferSampleProperties->Set("fps_num",      FVariantValue((int64) 0 ));
+					OutputBufferSampleProperties->Set("fps_denom",    FVariantValue((int64) 0 ));
+					OutputBufferSampleProperties->Set("pixelfmt",     FVariantValue((int64)EPixelFormat::PF_B8G8R8A8));
+				}
+				else
+				{
+					OutputBufferSampleProperties->Set("is_dummy", FVariantValue(true));
+				}
+				OutputBufferSampleProperties->Set("pts",          FVariantValue(NextImage.SourceInfo->PTS));
+				OutputBufferSampleProperties->Set("duration",     FVariantValue(NextImage.SourceInfo->Duration));
+
+				TSharedPtr<FElectraPlayerVideoDecoderOutputH265Apple, ESPMode::ThreadSafe> DecoderOutput = RenderOutputBuffer->GetBufferProperties().GetValue("texture").GetSharedPointer<FElectraPlayerVideoDecoderOutputH265Apple>();
+
+				DecoderOutput->Initialize(ImageBufferRef, OutputBufferSampleProperties);
+                if (ImageBufferRef)
+                {
+                    CFRelease(ImageBufferRef);
+                }
+
+				// Return the buffer to the renderer.
+				Renderer->ReturnBuffer(RenderOutputBuffer, true, *OutputBufferSampleProperties);
+			}
+		}
+	}
+}
+
+} // namespace Electra
+
+// -----------------------------------------------------------------------------------------------------------------------------
+
+#endif // ELECTRA_PLATFORM_HAS_H265_DECODER
+
+#endif
+
diff --git a/Engine/Plugins/Media/ElectraPlayer/Source/ElectraPlayerRuntime/Private/Runtime/Decoder/DX/AudioDecoderAAC_DX.cpp b/Engine/Plugins/Media/ElectraPlayer/Source/ElectraPlayerRuntime/Private/Runtime/Decoder/DX/AudioDecoderAAC_DX.cpp
--- a/Engine/Plugins/Media/ElectraPlayer/Source/ElectraPlayerRuntime/Private/Runtime/Decoder/DX/AudioDecoderAAC_DX.cpp
+++ b/Engine/Plugins/Media/ElectraPlayer/Source/ElectraPlayerRuntime/Private/Runtime/Decoder/DX/AudioDecoderAAC_DX.cpp
@@ -23,6 +23,7 @@
 #include "Windows/AllowWindowsPlatformTypes.h"
 #include "Windows/WindowsHWrapper.h"
 #include "HAL/LowLevelMemTracker.h"
+#include "Misc/CoreDelegates.h"
 
 THIRD_PARTY_INCLUDES_START
 #include "mftransform.h"
@@ -151,10 +152,18 @@ namespace Electra
 
 		bool IsDifferentFormat(const FAccessUnit* Data);
 
+		void HandleApplicationHasEnteredForeground();
+		void HandleApplicationWillEnterBackground();
+
 		FInstanceConfiguration													Config;
 
 		FAccessUnitBuffer														AccessUnits;
 
+		FDelegateHandle															ApplicationSuspendedDelegate;
+		FDelegateHandle															ApplicationResumeDelegate;
+		FMediaEvent																ApplicationRunningSignal;
+		FMediaEvent																ApplicationSuspendConfirmedSignal;
+
 		FMediaEvent																TerminateThreadEvent;
 		FMediaEvent																FlushDecoderEvent;
 		FMediaEvent																DecoderFlushedEvent;
@@ -860,6 +869,28 @@ namespace Electra
 	}
 
 
+	//-----------------------------------------------------------------------------
+	/**
+	 * Application has entered foreground.
+	 */
+	void FAudioDecoderAAC::HandleApplicationHasEnteredForeground()
+	{
+		ApplicationRunningSignal.Signal();
+	}
+
+	
+	//-----------------------------------------------------------------------------
+	/**
+	 * Application goes into background.
+	 */
+	void FAudioDecoderAAC::HandleApplicationWillEnterBackground()
+	{
+		ApplicationSuspendConfirmedSignal.Reset();
+		ApplicationRunningSignal.Reset();
+		ApplicationSuspendConfirmedSignal.WaitTimeoutAndReset(1000 * 500);
+	}
+
+	
 	//-----------------------------------------------------------------------------
 	/**
 	 * AAC audio decoder main threaded decode loop
@@ -868,6 +899,11 @@ namespace Electra
 	{
 		LLM_SCOPE(ELLMTag::ElectraPlayer);
 
+		ApplicationRunningSignal.Signal();
+		ApplicationSuspendConfirmedSignal.Reset();
+		ApplicationSuspendedDelegate = FCoreDelegates::ApplicationWillEnterBackgroundDelegate.AddRaw(this, &FAudioDecoderAAC::HandleApplicationWillEnterBackground);
+		ApplicationResumeDelegate = FCoreDelegates::ApplicationHasEnteredForegroundDelegate.AddRaw(this, &FAudioDecoderAAC::HandleApplicationHasEnteredForeground);
+
 		bool bError = false;
 		bool bInDummyDecodeMode = false;
 
@@ -882,6 +918,16 @@ namespace Electra
 
 		while (!TerminateThreadEvent.IsSignaled())
 		{
+			// If in background, wait until we get activated again.
+			if (!ApplicationRunningSignal.IsSignaled())
+			{
+				ApplicationSuspendConfirmedSignal.Signal();
+				while(!ApplicationRunningSignal.WaitTimeout(100 * 1000) && !TerminateThreadEvent.IsSignaled())
+				{
+				}
+				continue;
+			}
+
 			// Notify optional buffer listener that we will now be needing an AU for our input buffer.
 			if (!bError && InputBufferListener && AccessUnits.Num() == 0)
 			{
@@ -1134,6 +1180,17 @@ namespace Electra
 
 		CurrentPTS.SetToInvalid();
 		NextExpectedPTS.SetToInvalid();
+
+		if (ApplicationSuspendedDelegate.IsValid())
+		{
+			FCoreDelegates::ApplicationWillEnterBackgroundDelegate.Remove(ApplicationSuspendedDelegate);
+			ApplicationSuspendedDelegate.Reset();
+		}
+		if (ApplicationResumeDelegate.IsValid())
+		{
+			FCoreDelegates::ApplicationHasEnteredForegroundDelegate.Remove(ApplicationResumeDelegate);
+			ApplicationResumeDelegate.Reset();
+		}
 	}
 
 
diff --git a/Engine/Plugins/Media/ElectraPlayer/Source/ElectraPlayerRuntime/Private/Runtime/Decoder/DX/MediaFoundationGUIDs.h b/Engine/Plugins/Media/ElectraPlayer/Source/ElectraPlayerRuntime/Private/Runtime/Decoder/DX/MediaFoundationGUIDs.h
new file mode 100644
--- /dev/null
+++ b/Engine/Plugins/Media/ElectraPlayer/Source/ElectraPlayerRuntime/Private/Runtime/Decoder/DX/MediaFoundationGUIDs.h
@@ -0,0 +1,36 @@
+// Copyright Epic Games, Inc. All Rights Reserved.
+
+#pragma once
+
+#define INITGUID
+#include <guiddef.h>
+
+namespace Electra
+{
+
+#if (WINVER < _WIN32_WINNT_WIN8)
+//const GUID MF_SA_D3D11_AWARE = { 0x206b4fc8, 0xfcf9, 0x4c51, { 0xaf, 0xe3, 0x97, 0x64, 0x36, 0x9e, 0x33, 0xa0 } };
+DEFINE_GUID(MF_SA_D3D11_AWARE, 0x206b4fc8, 0xfcf9, 0x4c51, 0xaf, 0xe3, 0x97, 0x64, 0x36, 0x9e, 0x33, 0xa0);
+#endif
+
+
+// Define the H.264 class GUID ourselves. Using the CLSID_CMSH264DecoderMFT from wmcodecdsp.h requires linking with some lib again...
+//const GUID MFTmsH264Decoder					= { 0x62CE7E72, 0x4C71, 0x4D20, { 0xB1, 0x5D, 0x45, 0x28, 0x31, 0xA8, 0x7D, 0x9D } };
+DEFINE_GUID(MFTmsH264Decoder, 0x62CE7E72, 0x4C71, 0x4D20, 0xB1, 0x5D, 0x45, 0x28, 0x31, 0xA8, 0x7D, 0x9D);
+
+//const GUID CODECAPI_AVLowLatencyMode		= { 0x9c27891a, 0xed7a, 0x40e1, { 0x88, 0xe8, 0xb2, 0x27, 0x27, 0xa0, 0x24, 0xee } };
+DEFINE_GUID(CODECAPI_AVLowLatencyMode, 0x9c27891a, 0xed7a, 0x40e1, 0x88, 0xe8, 0xb2, 0x27, 0x27, 0xa0, 0x24, 0xee);
+
+//const GUID MF_MT_DECODER_USE_MAX_RESOLUTION	= { 0x4c547c24, 0xaf9a, 0x4f38, { 0x96, 0xad, 0x97, 0x87, 0x73, 0xcf, 0x53, 0xe7 } };
+DEFINE_GUID(MF_MT_DECODER_USE_MAX_RESOLUTION, 0x4c547c24, 0xaf9a, 0x4f38, 0x96, 0xad, 0x97, 0x87, 0x73, 0xcf, 0x53, 0xe7);
+
+
+//DEFINE_GUID(MF_MT_DEFAULT_STRIDE,    0x644b4e48, 0x1e02, 0x4516, 0xb0, 0xeb, 0xc0, 0x1c, 0xa9, 0xd4, 0x9a, 0xc6);
+//static const GUID AVDecVideoMaxCodedWidth   = { 0x5ae557b8, 0x77af, 0x41f5, { 0x9f, 0xa6, 0x4d, 0xb2, 0xfe, 0x1d, 0x4b, 0xca } };
+//static const GUID AVDecVideoMaxCodedHeight  = { 0x7262a16a, 0xd2dc, 0x4e75, { 0x9b, 0xa8, 0x65, 0xc0, 0xc6, 0xd3, 0x2b, 0x13 } };
+
+//#define MFSampleExtension_ClosedCaption_CEA708_MAX_SIZE	256
+//DEFINE_GUID(MFSampleExtension_ClosedCaption_CEA708, 0x26f09068, 0xe744, 0x47dc, 0xaa,0x03,0xdb,0xf2,0x04,0x03,0xbd,0xe6);
+//DEFINE_GUID(MF_USER_DATA_PAYLOAD,	0xd1d4985d, 0xdc92, 0x457a, 0xb3,0xa0,0x65,0x1a,0x33,0xa3,0x10,0x47);
+
+}
diff --git a/Engine/Plugins/Media/ElectraPlayer/Source/ElectraPlayerRuntime/Private/Runtime/Decoder/DX/VideoDecoderH264_DX.cpp b/Engine/Plugins/Media/ElectraPlayer/Source/ElectraPlayerRuntime/Private/Runtime/Decoder/DX/VideoDecoderH264_DX.cpp
--- a/Engine/Plugins/Media/ElectraPlayer/Source/ElectraPlayerRuntime/Private/Runtime/Decoder/DX/VideoDecoderH264_DX.cpp
+++ b/Engine/Plugins/Media/ElectraPlayer/Source/ElectraPlayerRuntime/Private/Runtime/Decoder/DX/VideoDecoderH264_DX.cpp
@@ -21,6 +21,7 @@
 #include "Windows/AllowWindowsPlatformTypes.h"
 #include "Windows/WindowsHWrapper.h"
 #include "HAL/LowLevelMemTracker.h"
+#include "Misc/CoreDelegates.h"
 
 THIRD_PARTY_INCLUDES_START
 #include "mftransform.h"
@@ -29,35 +30,12 @@ THIRD_PARTY_INCLUDES_START
 #include "mfidl.h"
 THIRD_PARTY_INCLUDES_END
 
+#include "Decoder/DX/MediaFoundationGUIDs.h"
 #include "Windows/HideWindowsPlatformTypes.h"
 
 #include "DecoderErrors_DX.h"
 #include "VideoDecoderH264_DX.h"
 
-namespace
-{
-
-static const int32 MaxHWDecodeW = 1920;
-static const int32 MaxHWDecodeH = 1088;
-
-//DEFINE_GUID(MF_MT_DEFAULT_STRIDE,    0x644b4e48, 0x1e02, 0x4516, 0xb0, 0xeb, 0xc0, 0x1c, 0xa9, 0xd4, 0x9a, 0xc6);
-//static const GUID AVDecVideoMaxCodedWidth   = { 0x5ae557b8, 0x77af, 0x41f5, { 0x9f, 0xa6, 0x4d, 0xb2, 0xfe, 0x1d, 0x4b, 0xca } };
-//static const GUID AVDecVideoMaxCodedHeight  = { 0x7262a16a, 0xd2dc, 0x4e75, { 0x9b, 0xa8, 0x65, 0xc0, 0xc6, 0xd3, 0x2b, 0x13 } };
-
-//#define MFSampleExtension_ClosedCaption_CEA708_MAX_SIZE	256
-//DEFINE_GUID(MFSampleExtension_ClosedCaption_CEA708, 0x26f09068, 0xe744, 0x47dc, 0xaa,0x03,0xdb,0xf2,0x04,0x03,0xbd,0xe6);
-//DEFINE_GUID(MF_USER_DATA_PAYLOAD,	0xd1d4985d, 0xdc92, 0x457a, 0xb3,0xa0,0x65,0x1a,0x33,0xa3,0x10,0x47);
-
-// `MF_LOW_LATENCY` is defined in "mfapi.h" for >= WIN8
-// UE4 supports lower Windows versions at the moment and so `WINVER` is < `_WIN32_WINNT_WIN8`
-// to be able to use `MF_LOW_LATENCY` with default UE4 build we define it ourselves and check actual
-// Windows version in runtime
-#if (WINVER < _WIN32_WINNT_WIN8)
-	const GUID MF_LOW_LATENCY = { 0x9c27891a, 0xed7a, 0x40e1,{ 0x88, 0xe8, 0xb2, 0x27, 0x27, 0xa0, 0x24, 0xee } };
-#endif
-
-}
-
 
 DECLARE_CYCLE_STAT(TEXT("FVideoDecoderH264::Decode()"), STAT_ElectraPlayer_VideoH264Decode, STATGROUP_ElectraPlayer);
 DECLARE_CYCLE_STAT(TEXT("FVideoDecoderH264::ConvertOutput()"), STAT_ElectraPlayer_VideoH264ConvertOutput, STATGROUP_ElectraPlayer);
@@ -72,6 +50,14 @@ bool FVideoDecoderH264::bDidCheckHWSupport = false;
 bool FVideoDecoderH264::bIsHWSupported = false;
 #endif
 
+namespace ConstantsMF264
+{
+
+static const int32 MaxHWDecodeW = 1920;
+static const int32 MaxHWDecodeH = 1088;
+
+}
+
 /*********************************************************************************************************************/
 /*********************************************************************************************************************/
 /*********************************************************************************************************************/
@@ -148,6 +134,7 @@ void FVideoDecoderH264::Shutdown()
 FVideoDecoderH264::FVideoDecoderH264()
 	: FMediaThread("Electra::H264 decoder")
 	, bThreadStarted(false)
+	, bDrainForCodecChange(false)
 	, PlayerSessionServices(nullptr)
 	, Renderer(nullptr)
 	, CurrentAccessUnit(nullptr)
@@ -248,6 +235,17 @@ void FVideoDecoderH264::Close()
 }
 
 
+//-----------------------------------------------------------------------------
+/**
+ * Drains the decoder of all enqueued input and ends it, after which the decoder must send an FDecoderMessage to the player
+ * to signal completion.
+ */
+void FVideoDecoderH264::DrainForCodecChange()
+{
+	bDrainForCodecChange = true;
+}
+
+
 //-----------------------------------------------------------------------------
 /**
  * Sets a new decoder limit.
@@ -582,8 +580,8 @@ bool FVideoDecoderH264::DecoderSetInputType()
 	int32 configW, configH;
 	if (bIsHardwareAccelerated)
 	{
-		configW = MaxHWDecodeW;
-		configH = MaxHWDecodeH;
+		configW = ConstantsMF264::MaxHWDecodeW;
+		configH = ConstantsMF264::MaxHWDecodeH;
 	}
 	else
 	{
@@ -1040,6 +1038,11 @@ bool FVideoDecoderH264::Decode(FAccessUnit* AccessUnit, bool bResolutionChanged)
 
 	while(!TerminateThreadSignal.IsSignaled())
 	{
+		// Tick output buffer pool to handle state progression of samples in flight
+		Renderer->TickOutputBufferPool();
+		// Tick platform code in case something needs doing
+		PlatformTick();
+
 #ifdef ELECTRA_ENABLE_SWDECODE
 		// For SW decoding we need to get a buffer from the renderer right now.
 		if (!bIsHardwareAccelerated)
@@ -1277,6 +1280,32 @@ void FVideoDecoderH264::TestHardwareDecoding()
 }
 
 
+
+//-----------------------------------------------------------------------------
+/**
+ * Application has entered foreground.
+ */
+void FVideoDecoderH264::HandleApplicationHasEnteredForeground()
+{
+	ApplicationRunningSignal.Signal();
+}
+
+
+//-----------------------------------------------------------------------------
+/**
+ * Application goes into background.
+ */
+void FVideoDecoderH264::HandleApplicationWillEnterBackground()
+{
+	ApplicationSuspendConfirmedSignal.Reset();
+	ApplicationRunningSignal.Reset();
+	if (!ApplicationSuspendConfirmedSignal.WaitTimeoutAndReset(1000 * 500))
+	{
+		UE_LOG(LogElectraPlayer, Log, TEXT("FVideoDecoderH264(%p): Did not suspend in time"), this);
+	}
+}
+
+
 //-----------------------------------------------------------------------------
 /**
  * H264 video decoder main threaded decode loop
@@ -1285,6 +1314,11 @@ void FVideoDecoderH264::WorkerThread()
 {
 	LLM_SCOPE(ELLMTag::ElectraPlayer);
 
+	ApplicationRunningSignal.Signal();
+	ApplicationSuspendConfirmedSignal.Reset();
+	ApplicationSuspendedDelegate = FCoreDelegates::ApplicationWillEnterBackgroundDelegate.AddRaw(this, &FVideoDecoderH264::HandleApplicationWillEnterBackground);
+	ApplicationResumeDelegate = FCoreDelegates::ApplicationHasEnteredForegroundDelegate.AddRaw(this, &FVideoDecoderH264::HandleApplicationHasEnteredForeground);
+
 	CurrentRenderOutputBuffer     = nullptr;
 	CurrentAccessUnit   		  = nullptr;
 	bHaveDecoder				  = false;
@@ -1315,141 +1349,161 @@ void FVideoDecoderH264::WorkerThread()
 	bool bNeedInitialReconfig = true;
 	while(!TerminateThreadSignal.IsSignaled())
 	{
-		// Notify optional buffer listener that we will now be needing an AU for our input buffer.
-		if (!bError && InputBufferListener && AccessUnits.Num() == 0)
+		// If in background, wait until we get activated again.
+		if (!ApplicationRunningSignal.IsSignaled())
 		{
-			SCOPE_CYCLE_COUNTER(STAT_ElectraPlayer_VideoH264Decode);
-			CSV_SCOPED_TIMING_STAT(ElectraPlayer, VideoH264Decode);
-
-			FAccessUnitBufferInfo	sin;
-			IAccessUnitBufferListener::FBufferStats	stats;
-			AccessUnits.GetStats(sin);
-			stats.NumAUsAvailable  = sin.NumCurrentAccessUnits;
-			stats.NumBytesInBuffer = sin.CurrentMemInUse;
-			stats.MaxBytesOfBuffer = sin.MaxDataSize;
-			stats.bEODSignaled     = sin.bEndOfData;
-			stats.bEODReached      = sin.bEndOfData && sin.NumCurrentAccessUnits == 0;
-			ListenerMutex.Lock();
-			if (InputBufferListener)
+			UE_LOG(LogElectraPlayer, Log, TEXT("FVideoDecoderH264(%p): OnSuspending"), this);
+			ApplicationSuspendConfirmedSignal.Signal();
+			while(!ApplicationRunningSignal.WaitTimeout(100 * 1000) && !TerminateThreadSignal.IsSignaled())
 			{
-				InputBufferListener->DecoderInputNeeded(stats);
 			}
-			ListenerMutex.Unlock();
+			UE_LOG(LogElectraPlayer, Log, TEXT("FVideoDecoderH264(%p): OnResuming"), this);
+			continue;
 		}
-		// Wait for data.
-		bool bHaveData = AccessUnits.WaitForData(1000 * 5);
-		// When there is data, even and especially after a previous EOD, we are no longer done and idling.
-		if (bHaveData && bDone)
+
+		if (!bDrainForCodecChange)
 		{
-			bDone = false;
-		}
+			// Notify optional buffer listener that we will now be needing an AU for our input buffer.
+			if (!bError && InputBufferListener && AccessUnits.Num() == 0)
+			{
+				SCOPE_CYCLE_COUNTER(STAT_ElectraPlayer_VideoH264Decode);
+				CSV_SCOPED_TIMING_STAT(ElectraPlayer, VideoH264Decode);
 
-		// Tick output buffer pool to handle state progression of samples in flight
-		Renderer->TickOutputBufferPool();
-		// Tick platform code in case something needs doing
-		PlatformTick();
+				FAccessUnitBufferInfo	sin;
+				IAccessUnitBufferListener::FBufferStats	stats;
+				AccessUnits.GetStats(sin);
+				stats.NumAUsAvailable  = sin.NumCurrentAccessUnits;
+				stats.NumBytesInBuffer = sin.CurrentMemInUse;
+				stats.MaxBytesOfBuffer = sin.MaxDataSize;
+				stats.bEODSignaled     = sin.bEndOfData;
+				stats.bEODReached      = sin.bEndOfData && sin.NumCurrentAccessUnits == 0;
+				ListenerMutex.Lock();
+				if (InputBufferListener)
+				{
+					InputBufferListener->DecoderInputNeeded(stats);
+				}
+				ListenerMutex.Unlock();
+			}
+			// Wait for data.
+			bool bHaveData = AccessUnits.WaitForData(1000 * 5);
+			// When there is data, even and especially after a previous EOD, we are no longer done and idling.
+			if (bHaveData && bDone)
+			{
+				bDone = false;
+			}
 
-		if (bHaveData && Renderer->CanReceiveOutputFrames(1)) // note: the check for full output buffers ignores the NumFramesInDecoder value as we don't seem to receive new data if we don't feed new input
-		{
-			CurrentAccessUnit = nullptr;
-			bool bOk = AccessUnits.Pop(CurrentAccessUnit);
-			MEDIA_UNUSED_VAR(bOk);
-			check(bOk);
+			// Tick output buffer pool to handle state progression of samples in flight
+			Renderer->TickOutputBufferPool();
+			// Tick platform code in case something needs doing
+			PlatformTick();
 
-			if (!bError)
+			if (bHaveData && Renderer->CanReceiveOutputFrames(1)) // note: the check for full output buffers ignores the NumFramesInDecoder value as we don't seem to receive new data if we don't feed new input
 			{
-				if (!CurrentAccessUnit->bIsDummyData)
-				{
-					{
-					SCOPE_CYCLE_COUNTER(STAT_ElectraPlayer_VideoH264Decode);
-					CSV_SCOPED_TIMING_STAT(ElectraPlayer, VideoH264Decode);
-					PrepareAU(CurrentAccessUnit);
-					}
+				CurrentAccessUnit = nullptr;
+				bool bOk = AccessUnits.Pop(CurrentAccessUnit);
+				MEDIA_UNUSED_VAR(bOk);
+				check(bOk);
 
-					// Decode
-					if (!bError)
+				if (!bError)
+				{
+					if (!CurrentAccessUnit->bIsDummyData)
 					{
-						// Update the buffer acquisition after a change in streams?
-						bool bFormatChangedJustNow = CurrentSampleInfo.IsDifferentFromOtherVideo(NewSampleInfo) || bGotLastSequenceAU;
-
-						if (bInDummyDecodeMode)
 						{
-							bFormatChangedJustNow = true;
-							bInDummyDecodeMode = false;
+						SCOPE_CYCLE_COUNTER(STAT_ElectraPlayer_VideoH264Decode);
+						CSV_SCOPED_TIMING_STAT(ElectraPlayer, VideoH264Decode);
+						PrepareAU(CurrentAccessUnit);
 						}
 
-						if (!bHaveDecoder)
+						// Decode
+						if (!bError)
 						{
-							SCOPE_CYCLE_COUNTER(STAT_ElectraPlayer_VideoH264Decode);
-							CSV_SCOPED_TIMING_STAT(ElectraPlayer, VideoH264Decode);
-							DecoderCreate();
-							bNeedInitialReconfig = true;
-						}
+							// Update the buffer acquisition after a change in streams?
+							bool bFormatChangedJustNow = CurrentSampleInfo.IsDifferentFromOtherVideo(NewSampleInfo) || bGotLastSequenceAU;
 
-						if (bFormatChangedJustNow || bNeedInitialReconfig)
-						{
-							if (bNeedInitialReconfig)
+							if (bInDummyDecodeMode)
+							{
+								bFormatChangedJustNow = true;
+								bInDummyDecodeMode = false;
+							}
+
+							if (!bHaveDecoder)
+							{
+								SCOPE_CYCLE_COUNTER(STAT_ElectraPlayer_VideoH264Decode);
+								CSV_SCOPED_TIMING_STAT(ElectraPlayer, VideoH264Decode);
+								DecoderCreate();
+								bNeedInitialReconfig = true;
+							}
+
+							if (bFormatChangedJustNow || bNeedInitialReconfig)
 							{
-								bDecoderFlushPending = true;
+								if (bNeedInitialReconfig)
+								{
+									bDecoderFlushPending = true;
+								}
+								if (!Decode(nullptr, bNeedInitialReconfig ? false : true))
+								{
+									// FIXME: Is this error fatal or safe to ignore?
+								}
+								CurrentSampleInfo = NewSampleInfo;
+								SetupBufferAcquisitionProperties();
 							}
-							if (!Decode(nullptr, bNeedInitialReconfig ? false : true))
+							bNeedInitialReconfig = false;
+							bGotLastSequenceAU = CurrentAccessUnit->bIsLastInPeriod;
+							if (!Decode(CurrentAccessUnit, false))
 							{
-								// FIXME: Is this error fatal or safe to ignore?
+								bError = true;
 							}
-							CurrentSampleInfo = NewSampleInfo;
-							SetupBufferAcquisitionProperties();
 						}
-						bNeedInitialReconfig = false;
-						bGotLastSequenceAU = CurrentAccessUnit->bIsLastInPeriod;
-						if (!Decode(CurrentAccessUnit, false))
+					}
+					else
+					{
+						if (!bInDummyDecodeMode)
+						{
+							bInDummyDecodeMode = true;
+							Decode(nullptr, true);
+						}
+
+						if (!DecodeDummy(CurrentAccessUnit))
 						{
 							bError = true;
 						}
 					}
 				}
-				else
-				{
-					if (!bInDummyDecodeMode)
-					{
-						bInDummyDecodeMode = true;
-						Decode(nullptr, true);
-					}
 
-					if (!DecodeDummy(CurrentAccessUnit))
-					{
-						bError = true;
-					}
-				}
+				// If decoding was aborted and the AU not handed off, delete it now.
+				FAccessUnit::Release(CurrentAccessUnit);
+				CurrentAccessUnit = nullptr;
 			}
-
-			// If decoding was aborted and the AU not handed off, delete it now.
-			FAccessUnit::Release(CurrentAccessUnit);
-			CurrentAccessUnit = nullptr;
-		}
-		else
-		{
-			if (!bHaveData)
+			else
 			{
-				// No data. Is the buffer at EOD?
-				if (AccessUnits.IsEndOfData())
+				if (!bHaveData)
 				{
-					NotifyReadyBufferListener(true);
-					// Are we done yet?
-					if (!bDone && !bError)
+					// No data. Is the buffer at EOD?
+					if (AccessUnits.IsEndOfData())
 					{
-						bError = !Decode(nullptr, false);
-						NumFramesInDecoder = 0;
+						NotifyReadyBufferListener(true);
+						// Are we done yet?
+						if (!bDone && !bError)
+						{
+							bError = !Decode(nullptr, false);
+							NumFramesInDecoder = 0;
+						}
+						bDone = true;
+						FMediaRunnable::SleepMilliseconds(10);
 					}
-					bDone = true;
-					FMediaRunnable::SleepMilliseconds(10);
+				}
+				else
+				{
+					// We could not decode new data as the output buffers are full... notify listener & wait a bit...
+					NotifyReadyBufferListener(false);
+					FMediaRunnable::SleepMilliseconds(5);
 				}
 			}
-			else
-			{
-				// We could not decode new data as the output buffers are full... notify listener & wait a bit...
-				NotifyReadyBufferListener(false);
-				FMediaRunnable::SleepMilliseconds(5);
-			}
+		}
+		else
+		{
+			bError = !Decode(nullptr, false);
+			break;
 		}
 
 		// Flush?
@@ -1480,6 +1534,31 @@ void FVideoDecoderH264::WorkerThread()
 	DestroyDecodedImagePool();
 	AccessUnits.Flush();
 	AccessUnits.CapacitySet(0);
+
+	if (ApplicationSuspendedDelegate.IsValid())
+	{
+		FCoreDelegates::ApplicationWillEnterBackgroundDelegate.Remove(ApplicationSuspendedDelegate);
+		ApplicationSuspendedDelegate.Reset();
+	}
+	if (ApplicationResumeDelegate.IsValid())
+	{
+		FCoreDelegates::ApplicationHasEnteredForegroundDelegate.Remove(ApplicationResumeDelegate);
+		ApplicationResumeDelegate.Reset();
+	}
+
+	if (bDrainForCodecChange)
+	{
+		// Notify the player that we have finished draining.
+		PlayerSessionServices->SendMessageToPlayer(FDecoderMessage::Create(FDecoderMessage::EReason::DrainingFinished, this, EStreamType::Video, FStreamCodecInformation::ECodec::H264));
+		// We need to wait to get terminated. Also check if flushing is requested and acknowledge if it is.
+		while(!TerminateThreadSignal.IsSignaled())
+		{
+			if (FlushDecoderSignal.WaitTimeoutAndReset(1000 * 10))
+			{
+				DecoderFlushedSignal.Signal();
+			}
+		}
+	}
 }
 
 } // namespace Electra
