ElectraPlayer: Added support for HEVC decoder where available

diff --git a/Engine/Plugins/Media/ElectraPlayer/Source/ElectraPlayerRuntime/Private/Runtime/Decoder/DX/VideoDecoderH264_DX.h b/Engine/Plugins/Media/ElectraPlayer/Source/ElectraPlayerRuntime/Private/Runtime/Decoder/DX/VideoDecoderH264_DX.h
--- a/Engine/Plugins/Media/ElectraPlayer/Source/ElectraPlayerRuntime/Private/Runtime/Decoder/DX/VideoDecoderH264_DX.h
+++ b/Engine/Plugins/Media/ElectraPlayer/Source/ElectraPlayerRuntime/Private/Runtime/Decoder/DX/VideoDecoderH264_DX.h
@@ -3,12 +3,6 @@
 #pragma once
 
 namespace Electra {
-
-	// Define the H.264 class GUID ourselves. Using the CLSID_CMSH264DecoderMFT from wmcodecdsp.h requires linking with some lib again...
-	static const GUID MFTmsH264Decoder = { 0x62CE7E72, 0x4C71, 0x4D20, { 0xB1, 0x5D, 0x45, 0x28, 0x31, 0xA8, 0x7D, 0x9D } };
-
-	static const GUID CODECAPI_AVLowLatencyMode = { 0x9c27891a, 0xed7a, 0x40e1, { 0x88, 0xe8, 0xb2, 0x27, 0x27, 0xa0, 0x24, 0xee } };
-
 	/**
 	 * H264 video decoder class implementation.
 	**/
@@ -29,6 +23,7 @@ namespace Electra {
 
 		virtual void Open(const FInstanceConfiguration& InConfig) override;
 		virtual void Close() override;
+		virtual void DrainForCodecChange() override;
 
 		virtual void SetMaximumDecodeCapability(int32 MaxWidth, int32 MaxHeight, int32 MaxProfile, int32 MaxProfileLevel, const FParamDict& AdditionalOptions) override;
 
@@ -123,6 +118,9 @@ namespace Electra {
 		bool Configure();
 		bool StartStreaming();
 
+		void HandleApplicationHasEnteredForeground();
+		void HandleApplicationWillEnterBackground();
+
 		// Per platform specialization
 		virtual bool InternalDecoderCreate() = 0;
 		virtual bool CreateDecoderOutputBuffer() = 0;
@@ -132,10 +130,16 @@ namespace Electra {
 
 		FInstanceConfiguration								Config;
 
+		FDelegateHandle										ApplicationSuspendedDelegate;
+		FDelegateHandle										ApplicationResumeDelegate;
+		FMediaEvent											ApplicationRunningSignal;
+		FMediaEvent											ApplicationSuspendConfirmedSignal;
+
 		FMediaEvent											TerminateThreadSignal;
 		FMediaEvent											FlushDecoderSignal;
 		FMediaEvent											DecoderFlushedSignal;
 		bool												bThreadStarted;
+		bool												bDrainForCodecChange;
 
 		IPlayerSessionServices*								PlayerSessionServices;
 
diff --git a/Engine/Plugins/Media/ElectraPlayer/Source/ElectraPlayerRuntime/Private/Runtime/Decoder/DX/VideoDecoderH265_DX.cpp b/Engine/Plugins/Media/ElectraPlayer/Source/ElectraPlayerRuntime/Private/Runtime/Decoder/DX/VideoDecoderH265_DX.cpp
new file mode 100644
--- /dev/null
+++ b/Engine/Plugins/Media/ElectraPlayer/Source/ElectraPlayerRuntime/Private/Runtime/Decoder/DX/VideoDecoderH265_DX.cpp
@@ -0,0 +1,1342 @@
+// Copyright Epic Games, Inc. All Rights Reserved.
+
+#ifdef ELECTRA_ENABLE_MFDECODER
+
+#include "PlayerCore.h"
+#include "ElectraPlayerPrivate.h"
+
+#include "ElectraPlayerPrivate_Platform.h"
+
+#include "StreamAccessUnitBuffer.h"
+#include "Decoder/VideoDecoderH265.h"
+#include "Renderer/RendererBase.h"
+#include "Renderer/RendererVideo.h"
+#include "Player/PlayerSessionServices.h"
+#include "Utilities/StringHelpers.h"
+#include "Utilities/UtilsMPEGVideo.h"
+
+#if ELECTRA_PLATFORM_HAS_H265_DECODER
+
+/*********************************************************************************************************************/
+/*********************************************************************************************************************/
+/*********************************************************************************************************************/
+
+#include "Windows/AllowWindowsPlatformTypes.h"
+#include "Windows/WindowsHWrapper.h"
+#include "HAL/LowLevelMemTracker.h"
+#include "Misc/CoreDelegates.h"
+
+THIRD_PARTY_INCLUDES_START
+#include "mftransform.h"
+#include "mfapi.h"
+#include "mferror.h"
+#include "mfidl.h"
+THIRD_PARTY_INCLUDES_END
+
+#include "Decoder/DX/MediaFoundationGUIDs.h"
+#include "Windows/HideWindowsPlatformTypes.h"
+
+#include "DecoderErrors_DX.h"
+#include "VideoDecoderH265_DX.h"
+
+
+DECLARE_CYCLE_STAT(TEXT("FVideoDecoderH265::Decode()"), STAT_ElectraPlayer_VideoH265Decode, STATGROUP_ElectraPlayer);
+DECLARE_CYCLE_STAT(TEXT("FVideoDecoderH265::ConvertOutput()"), STAT_ElectraPlayer_VideoH265ConvertOutput, STATGROUP_ElectraPlayer);
+
+
+namespace Electra {
+
+/*********************************************************************************************************************/
+/*********************************************************************************************************************/
+/*********************************************************************************************************************/
+
+namespace ConstantsMF265
+{
+// Note: keep those set at 1080p.
+static const int32 MaxHWDecodeW = 1920;
+static const int32 MaxHWDecodeH = 1088;
+}
+
+bool IVideoDecoderH265::Startup(const FParamDict& Options)
+{
+	return FVideoDecoderH265::Startup(Options);
+}
+
+void IVideoDecoderH265::Shutdown()
+{
+	FVideoDecoderH265::Shutdown();
+}
+
+bool IVideoDecoderH265::GetStreamDecodeCapability(FStreamDecodeCapability& OutResult, const FStreamDecodeCapability& InStreamParameter)
+{
+	return FVideoDecoderH265::GetStreamDecodeCapability(OutResult, InStreamParameter);
+}
+
+IVideoDecoderH265::FInstanceConfiguration::FInstanceConfiguration()
+	: MaxDecodedFrames(8)														// FIXME: how many do we need exactly?
+{
+}
+
+/***************************************************************************************************************************************************/
+/***************************************************************************************************************************************************/
+/***************************************************************************************************************************************************/
+
+//-----------------------------------------------------------------------------
+/**
+ * Decoder system startup
+ *
+ * @param Options
+ *
+ * @return
+ */
+bool FVideoDecoderH265::Startup(const FParamDict& Options)
+{
+	return true;
+}
+
+
+//-----------------------------------------------------------------------------
+/**
+ * Decoder system shutdown.
+ */
+void FVideoDecoderH265::Shutdown()
+{
+}
+
+
+//-----------------------------------------------------------------------------
+/**
+ * Constructor
+ */
+FVideoDecoderH265::FVideoDecoderH265()
+	: FMediaThread("Electra::H265 decoder")
+	, bThreadStarted(false)
+	, bDrainForCodecChange(false)
+	, PlayerSessionServices(nullptr)
+	, Renderer(nullptr)
+	, CurrentAccessUnit(nullptr)
+	, InputBufferListener(nullptr)
+	, ReadyBufferListener(nullptr)
+	, bIsHardwareAccelerated(true)
+	, NumFramesInDecoder(0)
+	, bError(false)
+	, CurrentRenderOutputBuffer(nullptr)
+	, bHaveDecoder(false)
+	, MaxDecodeBufferSize(0)
+{
+}
+
+
+//-----------------------------------------------------------------------------
+/**
+ * Destructor
+ */
+FVideoDecoderH265::~FVideoDecoderH265()
+{
+	Close();
+}
+
+
+//-----------------------------------------------------------------------------
+/**
+ * Sets an AU input buffer listener.
+ *
+ * @param InListener
+ */
+void FVideoDecoderH265::SetAUInputBufferListener(IAccessUnitBufferListener* InListener)
+{
+	FMediaCriticalSection::ScopedLock lock(ListenerMutex);
+	InputBufferListener = InListener;
+}
+
+
+//-----------------------------------------------------------------------------
+/**
+ * Sets a buffer-ready listener.
+ *
+ * @param InListener
+ */
+void FVideoDecoderH265::SetReadyBufferListener(IDecoderOutputBufferListener* InListener)
+{
+	FMediaCriticalSection::ScopedLock lock(ListenerMutex);
+	ReadyBufferListener = InListener;
+}
+
+
+//-----------------------------------------------------------------------------
+/**
+ * Sets the owning player's session service interface.
+ *
+ * @param InSessionServices
+ *
+ * @return
+ */
+void FVideoDecoderH265::SetPlayerSessionServices(IPlayerSessionServices* SessionServices)
+{
+	PlayerSessionServices = SessionServices;
+}
+
+
+//-----------------------------------------------------------------------------
+/**
+ * Opens a decoder instance
+ *
+ * @param InConfig
+ *
+ * @return
+ */
+void FVideoDecoderH265::Open(const IVideoDecoderH265::FInstanceConfiguration& InConfig)
+{
+	Config = InConfig;
+
+	MaxDecodeDim = FIntPoint(Config.MaxFrameWidth, Config.MaxFrameHeight);
+
+	// Set a large enough size to hold a single access unit. Since we will be asking for a new AU on
+	// demand there is no need to overly restrict ourselves here. But it must be large enough to
+	// hold at least the largest expected access unit.
+	AccessUnits.CapacitySet(FAccessUnitBuffer::FConfiguration(16 << 20, 60.0));
+
+	StartThread();
+}
+
+
+//-----------------------------------------------------------------------------
+/**
+ * Closes the decoder instance.
+ */
+void FVideoDecoderH265::Close()
+{
+	StopThread();
+}
+
+
+//-----------------------------------------------------------------------------
+/**
+ * Drains the decoder of all enqueued input and ends it, after which the decoder must send an FDecoderMessage to the player
+ * to signal completion.
+ */
+void FVideoDecoderH265::DrainForCodecChange()
+{
+	bDrainForCodecChange = true;
+}
+
+
+//-----------------------------------------------------------------------------
+/**
+ * Sets a new decoder limit.
+ * As soon as a new sequence starts matching this limit the decoder will be
+ * destroyed and recreated to conserve memory.
+ * The assumption is that the data being streamed will not exceed this limit,
+ * but if it does, any access unit requiring a better decoder will take
+ * precendence and force a decoder capable of decoding it!
+ *
+ * @param MaxTier
+ * @param MaxWidth
+ * @param MaxHeight
+ * @param MaxProfile
+ * @param MaxProfileLevel
+ * @param AdditionalOptions
+ */
+void FVideoDecoderH265::SetMaximumDecodeCapability(int32 MaxTier, int32 MaxWidth, int32 MaxHeight, int32 MaxProfile, int32 MaxProfileLevel, const FParamDict& AdditionalOptions)
+{
+	// Recall max resolution so we can allocate properly sized sample buffers
+	MaxDecodeDim = FIntPoint(MaxWidth, MaxHeight);
+}
+
+
+//-----------------------------------------------------------------------------
+/**
+ * Sets a new renderer.
+ *
+ * @param InRenderer
+ */
+void FVideoDecoderH265::SetRenderer(TSharedPtr<IMediaRenderer, ESPMode::ThreadSafe> InRenderer)
+{
+	Renderer = InRenderer;
+}
+
+
+void FVideoDecoderH265::SetResourceDelegate(const TSharedPtr<IVideoDecoderResourceDelegate, ESPMode::ThreadSafe>& InResourceDelegate)
+{
+	ResourceDelegate = InResourceDelegate;
+}
+
+//-----------------------------------------------------------------------------
+/**
+ * Creates and runs the decoder thread.
+ */
+void FVideoDecoderH265::StartThread()
+{
+	ThreadStart(Electra::MakeDelegate(this, &FVideoDecoderH265::WorkerThread));
+	bThreadStarted = true;
+}
+
+
+//-----------------------------------------------------------------------------
+/**
+ * Stops the decoder thread.
+ */
+void FVideoDecoderH265::StopThread()
+{
+	if (bThreadStarted)
+	{
+		TerminateThreadSignal.Signal();
+		ThreadWaitDone();
+		bThreadStarted = false;
+	}
+}
+
+
+//-----------------------------------------------------------------------------
+/**
+ * Posts an error to the session service error listeners.
+ *
+ * @param ApiReturnValue
+ * @param Message
+ * @param Code
+ * @param Error
+ */
+void FVideoDecoderH265::PostError(int32 ApiReturnValue, const FString& Message, uint16 Code, UEMediaError Error)
+{
+	bError = true;
+	check(PlayerSessionServices);
+	if (PlayerSessionServices)
+	{
+		FErrorDetail err;
+		err.SetError(Error != UEMEDIA_ERROR_OK ? Error : UEMEDIA_ERROR_DETAIL);
+		err.SetFacility(Facility::EFacility::H265Decoder);
+		err.SetCode(Code);
+		err.SetMessage(Message);
+		err.SetPlatformMessage(FString::Printf(TEXT("%d (0x%08x)"), ApiReturnValue, ApiReturnValue));
+		PlayerSessionServices->PostError(err);
+	}
+}
+
+
+//-----------------------------------------------------------------------------
+/**
+ * Sends a log message to the session service log.
+ *
+ * @param Level
+ * @param Message
+ */
+void FVideoDecoderH265::LogMessage(IInfoLog::ELevel Level, const FString& Message)
+{
+	if (PlayerSessionServices)
+	{
+		PlayerSessionServices->PostLog(Facility::EFacility::H265Decoder, Level, Message);
+	}
+}
+
+
+//-----------------------------------------------------------------------------
+/**
+ * Create a pool of decoded images for the decoder.
+ *
+ * @return
+ */
+bool FVideoDecoderH265::CreateDecodedImagePool()
+{
+	check(Renderer);
+	FParamDict poolOpts;
+
+	poolOpts.Set("num_buffers", FVariantValue((int64) Config.MaxDecodedFrames));
+	if (!bIsHardwareAccelerated)
+	{
+		poolOpts.Set("sw_texture", FVariantValue(true));
+	}
+	UEMediaError Error = Renderer->CreateBufferPool(poolOpts);
+	check(Error == UEMEDIA_ERROR_OK);
+
+	MaxDecodeBufferSize = (int32) Renderer->GetBufferPoolProperties().GetValue("max_buffers").GetInt64();
+
+	if (Error != UEMEDIA_ERROR_OK)
+	{
+		PostError(0, "Failed to create image pool", ERRCODE_INTERNAL_COULD_NOT_CREATE_IMAGE_POOL, Error);
+	}
+
+	return Error == UEMEDIA_ERROR_OK;
+}
+
+
+//-----------------------------------------------------------------------------
+/**
+ * Destroys the pool of decoded images.
+ */
+void FVideoDecoderH265::DestroyDecodedImagePool()
+{
+	Renderer->ReleaseBufferPool();
+}
+
+
+//-----------------------------------------------------------------------------
+/**
+ * Called to receive a new input access unit for decoding.
+ *
+ * @param InAccessUnit
+ */
+IAccessUnitBufferInterface::EAUpushResult FVideoDecoderH265::AUdataPushAU(FAccessUnit* InAccessUnit)
+{
+	InAccessUnit->AddRef();
+	bool bOk = AccessUnits.Push(InAccessUnit);
+	if (!bOk)
+	{
+		FAccessUnit::Release(InAccessUnit);
+	}
+
+	return bOk ? IAccessUnitBufferInterface::EAUpushResult::Ok : IAccessUnitBufferInterface::EAUpushResult::Full;
+}
+
+
+//-----------------------------------------------------------------------------
+/**
+ * "Pushes" an End Of Data marker indicating no further access units will be added.
+ */
+void FVideoDecoderH265::AUdataPushEOD()
+{
+	AccessUnits.PushEndOfData();
+}
+
+
+//-----------------------------------------------------------------------------
+/**
+ * Flushes the decoder and clears the input access unit buffer.
+ */
+void FVideoDecoderH265::AUdataFlushEverything()
+{
+	FlushDecoderSignal.Signal();
+	DecoderFlushedSignal.WaitAndReset();
+}
+
+
+//-----------------------------------------------------------------------------
+/**
+ *	Reconfigures decoder for software mode.
+ */
+bool FVideoDecoderH265::ReconfigureForSwDecoding(FString Reason)
+{
+	PostError(S_OK, Reason, ERRCODE_INTERNAL_INCOMPATIBLE_DECODER);
+	return false;
+}
+
+
+//-----------------------------------------------------------------------------
+/**
+ * Configures the decoder input and output media types.
+ *
+ * @return true if successful, false on error
+ */
+bool FVideoDecoderH265::Configure()
+{
+	// Setup media input type
+	if (!DecoderSetInputType())
+	{
+		return false;
+	}
+	// Setup media output type
+	if (!DecoderSetOutputType())
+	{
+		return false;
+	}
+	// Verify status
+	if (!DecoderVerifyStatus())
+	{
+		return false;
+	}
+	return true;
+}
+
+
+//-----------------------------------------------------------------------------
+/**
+ * Starts decoder processing.
+ *
+ * @return true if successful, false on error
+ */
+bool FVideoDecoderH265::StartStreaming()
+{
+	HRESULT res;
+	// Start it!
+	VERIFY_HR(DecoderTransform->ProcessMessage(MFT_MESSAGE_NOTIFY_BEGIN_STREAMING, 0), "Failed to notify video decoder of stream begin", ERRCODE_INTERNAL_COULD_NOT_SET_DECODER_BEGIN);
+	VERIFY_HR(DecoderTransform->ProcessMessage(MFT_MESSAGE_NOTIFY_START_OF_STREAM, 0), "Failed to notify video decoder of start start", ERRCODE_INTERNAL_COULD_NOT_SET_DECODER_START);
+	return true;
+}
+
+
+//-----------------------------------------------------------------------------
+/**
+ * Configures the decoder input media type.
+ *
+ * @return true if successful, false on error
+ */
+bool FVideoDecoderH265::DecoderSetInputType()
+{
+	TRefCountPtr<IMFMediaType>	InputMediaType;
+	HRESULT					res;
+
+	// https://docs.microsoft.com/en-us/windows/win32/medfound/h-265---hevc-video-decoder
+	VERIFY_HR(MFCreateMediaType(InputMediaType.GetInitReference()), "Failed to create video decoder input media type", ERRCODE_INTERNAL_COULD_NOT_CREATE_INPUT_MEDIA_TYPE);
+	VERIFY_HR(InputMediaType->SetGUID(MF_MT_MAJOR_TYPE, MFMediaType_Video), "Failed to set video decoder input media type major", ERRCODE_INTERNAL_COULD_NOT_CREATE_INPUT_MEDIA_TYPE_MAJOR);
+	VERIFY_HR(InputMediaType->SetGUID(MF_MT_SUBTYPE, MFVideoFormat_HEVC), "Failed to set video decoder input media sub type", ERRCODE_INTERNAL_COULD_NOT_CREATE_INPUT_MEDIA_TYPE_SUBTYPE);
+
+	int32 configW, configH;
+	if (bIsHardwareAccelerated)
+	{
+		configW = ConstantsMF265::MaxHWDecodeW;
+		configH = ConstantsMF265::MaxHWDecodeH;
+	}
+	else
+	{
+		configW = CurrentSampleInfo.GetResolution().Width;
+		configH = CurrentSampleInfo.GetResolution().Height;
+	}
+
+	if (CurrentAccessUnit && CurrentAccessUnit->AUCodecData.IsValid() && CurrentAccessUnit->AUCodecData->RawCSD.Num())
+	{
+		VERIFY_HR(InputMediaType->SetBlob(MF_MT_USER_DATA, CurrentAccessUnit->AUCodecData->RawCSD.GetData(), (UINT32) CurrentAccessUnit->AUCodecData->RawCSD.Num()), "Failed to set video decoder input media type codec blob", ERRCODE_INTERNAL_COULD_NOT_SET_INPUT_MEDIA_TYPE_CODEC_BLOB);
+	}
+
+	VERIFY_HR(MFSetAttributeSize(InputMediaType, MF_MT_FRAME_SIZE, configW, configH), "Failed to set video decoder input media type resolution", ERRCODE_INTERNAL_COULD_NOT_CREATE_INPUT_MEDIA_TYPE_RESOLUTION);
+
+	VERIFY_HR(DecoderTransform->SetInputType(0, InputMediaType, 0), "Failed to set video decoder input type", ERRCODE_INTERNAL_COULD_NOT_SET_INPUT_MEDIA_TYPE)
+
+	return true;
+}
+
+
+//-----------------------------------------------------------------------------
+/**
+ * Configures the decoder output media type.
+ *
+ * @return true if successful, false on error
+ */
+bool FVideoDecoderH265::DecoderSetOutputType()
+{
+	TRefCountPtr<IMFMediaType>	OutputMediaType;
+	GUID					OutputMediaMajorType;
+	GUID					OutputMediaSubtype;
+	HRESULT					res;
+
+	// Supposedly calling GetOutputAvailableType() returns following output media subtypes:
+	// MFVideoFormat_NV12, MFVideoFormat_YV12, MFVideoFormat_IYUV, MFVideoFormat_I420, MFVideoFormat_YUY2
+	for(int32 TypeIndex=0; ; ++TypeIndex)
+	{
+		VERIFY_HR(DecoderTransform->GetOutputAvailableType(0, TypeIndex, OutputMediaType.GetInitReference()), "Failed to get video decoder available output type", ERRCODE_INTERNAL_COULD_NOT_GET_OUTPUT_MEDIA_TYPE);
+		VERIFY_HR(OutputMediaType->GetGUID(MF_MT_MAJOR_TYPE, &OutputMediaMajorType), "Failed to get video decoder available output media type", ERRCODE_INTERNAL_COULD_NOT_GET_OUTPUT_MEDIA_TYPE_MAJOR);
+		VERIFY_HR(OutputMediaType->GetGUID(MF_MT_SUBTYPE, &OutputMediaSubtype), "Failed to get video decoder available output media subtype", ERRCODE_INTERNAL_COULD_NOT_GET_OUTPUT_MEDIA_TYPE_SUBTYPE);
+		if (OutputMediaMajorType == MFMediaType_Video && OutputMediaSubtype == MFVideoFormat_NV12)
+		{
+			VERIFY_HR(DecoderTransform->SetOutputType(0, OutputMediaType, 0), "Failed to set video decoder output type", ERRCODE_INTERNAL_COULD_NOT_SET_OUTPUT_MEDIA_TYPE);
+			CurrentOutputMediaType = OutputMediaType;
+			return true;
+		}
+	}
+	PostError(S_OK, "Failed to set video decoder output type to desired format", ERRCODE_INTERNAL_COULD_NOT_SET_OUTPUT_DESIRED_MEDIA_TYPE);
+	return false;
+}
+
+
+//-----------------------------------------------------------------------------
+/**
+ * Checks that the decoder is in the state we think it should be and good to go.
+ *
+ * @return true if successful, false on error
+ */
+bool FVideoDecoderH265::DecoderVerifyStatus()
+{
+	HRESULT		res;
+	DWORD		NumInputStreams;
+	DWORD		NumOutputStreams;
+
+	VERIFY_HR(DecoderTransform->GetStreamCount(&NumInputStreams, &NumOutputStreams), "Failed to get video decoder stream count", ERRCODE_INTERNAL_COULD_NOT_GET_STREAM_COUNT);
+	if (NumInputStreams != 1 || NumOutputStreams != 1)
+	{
+		PostError(S_OK, FString::Printf(TEXT("Unexpected number of streams: input %lu, output %lu"), NumInputStreams, NumOutputStreams), ERRCODE_INTERNAL_UNEXPECTED_NUMBER_OF_STREAMS);
+		return false;
+	}
+
+	DWORD DecoderStatus = 0;
+	VERIFY_HR(DecoderTransform->GetInputStatus(0, &DecoderStatus), "Failed to get video decoder input status", ERRCODE_INTERNAL_COULD_NOT_GET_INPUT_STATUS);
+	if (MFT_INPUT_STATUS_ACCEPT_DATA != DecoderStatus)
+	{
+		PostError(S_OK, FString::Printf(TEXT("Decoder doesn't accept data, status %lu"), DecoderStatus), ERRCODE_INTERNAL_DECODER_NOT_ACCEPTING_DATA);
+		return false;
+	}
+
+	VERIFY_HR(DecoderTransform->GetOutputStreamInfo(0, &DecoderOutputStreamInfo), "Failed to get video decoder output stream info", ERRCODE_INTERNAL_COULD_NOT_GET_OUTPUT_STREAM_INFO);
+	if (!(DecoderOutputStreamInfo.dwFlags & MFT_OUTPUT_STREAM_FIXED_SAMPLE_SIZE))
+	{
+		return ReconfigureForSwDecoding(TEXT("Incompatible H.265 decoder: fixed sample size expected"));
+	}
+	if (!(DecoderOutputStreamInfo.dwFlags & MFT_OUTPUT_STREAM_WHOLE_SAMPLES))
+	{
+		return ReconfigureForSwDecoding(TEXT("Incompatible H.265 decoder: whole samples expected"));
+	}
+	if (bIsHardwareAccelerated && !(DecoderOutputStreamInfo.dwFlags & MFT_OUTPUT_STREAM_PROVIDES_SAMPLES))
+	{
+		// theoretically we can handle this situation with H/W decoder, but we can't reproduce it locally for testing so we aren't sure if H/W
+		// decoder would work in this case
+		return ReconfigureForSwDecoding(TEXT("Incompatible H.265 decoder: h/w accelerated decoder is expected to provide output samples"));
+	}
+	if (!bIsHardwareAccelerated && (DecoderOutputStreamInfo.dwFlags & MFT_OUTPUT_STREAM_PROVIDES_SAMPLES))
+	{
+		PostError(S_OK, "Incompatible H.265 decoder: s/w decoder is expected to require preallocated output samples", ERRCODE_INTERNAL_INCOMPATIBLE_DECODER);
+		return false;
+	}
+	return true;
+}
+
+
+//-----------------------------------------------------------------------------
+/**
+ * Destroys the current decoder instance.
+ */
+void FVideoDecoderH265::InternalDecoderDestroy()
+{
+	DecoderTransform = nullptr;
+	CurrentOutputMediaType = nullptr;
+	bHaveDecoder = false;
+}
+
+
+//-----------------------------------------------------------------------------
+/**
+ * Notifies the buffer-ready listener that we will now be producing output data
+ * and need an output buffer.
+ * This call is not intended to create or obtain an output buffer. It is merely
+ * indicating that new output will be produced.
+ *
+ * @param bHaveOutput
+ */
+void FVideoDecoderH265::NotifyReadyBufferListener(bool bHaveOutput)
+{
+	if (ReadyBufferListener)
+	{
+		IDecoderOutputBufferListener::FDecodeReadyStats stats;
+		stats.MaxDecodedElementsReady = MaxDecodeBufferSize;
+		stats.NumElementsInDecoder    = NumFramesInDecoder;
+		stats.bOutputStalled		  = !bHaveOutput;
+		stats.bEODreached   		  = AccessUnits.IsEndOfData() && stats.NumDecodedElementsReady == 0 && stats.NumElementsInDecoder == 0;
+		ListenerMutex.Lock();
+		if (ReadyBufferListener)
+		{
+			ReadyBufferListener->DecoderOutputReady(stats);
+		}
+		ListenerMutex.Unlock();
+	}
+}
+
+
+
+void FVideoDecoderH265::PrepareAU(FAccessUnit* AccessUnit)
+{
+	if (!AccessUnit->bHasBeenPrepared)
+	{
+		AccessUnit->bHasBeenPrepared = true;
+
+		// Process NALUs
+		bool bIsDiscardable = false;
+		bool bIsKeyframe = false;
+		//AccessUnit->mIsDiscardable = 1;
+		// Replace the NALU lengths with the startcode.
+		uint32* NALU = (uint32*)AccessUnit->AUData;
+		uint32* End  = (uint32*)Electra::AdvancePointer(NALU, AccessUnit->AUSize);
+		while(NALU < End)
+		{
+			uint8 nut = *(const uint8 *)(NALU + 1);
+			nut >>= 1;
+			// IDR frame?
+			if (nut == 19 /*IDR_W_RADL*/ || nut == 20 /*IDR_N_LP*/)
+			{
+				bIsKeyframe = true;
+			}
+
+			uint32 naluLen = MEDIA_FROM_BIG_ENDIAN(*NALU) + 4;
+			*NALU = MEDIA_TO_BIG_ENDIAN(0x00000001U);
+			NALU = Electra::AdvancePointer(NALU, naluLen);
+		}
+
+		// Is there codec specific data?
+		if (AccessUnit->AUCodecData.IsValid())
+		{
+			// Yes.
+			if (bIsKeyframe || AccessUnit->bIsSyncSample || AccessUnit->bIsFirstInSequence)
+			{
+				NewSampleInfo = AccessUnit->AUCodecData->ParsedInfo;
+
+				// Have to re-allocate the AU memory to preprend the codec data
+				uint64 nb = AccessUnit->AUSize + AccessUnit->AUCodecData->CodecSpecificData.Num();
+				void* PayloadBuffer = AccessUnit->AllocatePayloadBuffer(nb);
+				check(PayloadBuffer);
+				void* PayloadData = PayloadBuffer;
+				FMemory::Memcpy(PayloadData, AccessUnit->AUCodecData->CodecSpecificData.GetData(), AccessUnit->AUCodecData->CodecSpecificData.Num());
+				PayloadData = Electra::AdvancePointer(PayloadData, AccessUnit->AUCodecData->CodecSpecificData.Num());
+				FMemory::Memcpy(PayloadData, AccessUnit->AUData, AccessUnit->AUSize);
+				AccessUnit->AdoptNewPayloadBuffer(PayloadBuffer, nb);
+				// New codec data makes this AU non-discardable.
+				bIsDiscardable = false;
+			}
+		}
+	}
+}
+
+
+void FVideoDecoderH265::SetupBufferAcquisitionProperties()
+{
+	BufferAcquireOptions.Clear();
+	int32 w = CurrentSampleInfo.GetResolution().Width;
+	int32 h = CurrentSampleInfo.GetResolution().Height;
+	BufferAcquireOptions.Set("width",  FVariantValue((int64) w));
+	BufferAcquireOptions.Set("height", FVariantValue((int64) h));
+}
+
+
+
+void FVideoDecoderH265::ReturnUnusedFrame()
+{
+	CurrentDecoderOutputBuffer.Reset();
+	if (CurrentRenderOutputBuffer)
+	{
+		Electra::FParamDict DummySampleProperties;
+		Renderer->ReturnBuffer(CurrentRenderOutputBuffer, false, DummySampleProperties);
+		CurrentRenderOutputBuffer = nullptr;
+	}
+}
+
+
+bool FVideoDecoderH265::AcquireOutputBuffer()
+{
+	// If we still have one we're done.
+	if (CurrentRenderOutputBuffer)
+	{
+		return true;
+	}
+	while(!TerminateThreadSignal.IsSignaled())
+	{
+		if (FlushDecoderSignal.IsSignaled())
+		{
+			break;
+		}
+
+		UEMediaError bufResult = Renderer->AcquireBuffer(CurrentRenderOutputBuffer, 0, BufferAcquireOptions);
+		check(bufResult == UEMEDIA_ERROR_OK || bufResult == UEMEDIA_ERROR_INSUFFICIENT_DATA);
+		if (bufResult != UEMEDIA_ERROR_OK && bufResult != UEMEDIA_ERROR_INSUFFICIENT_DATA)
+		{
+			PostError(S_OK, "Failed to acquire sample buffer", ERRCODE_INTERNAL_COULD_NOT_GET_SAMPLE_BUFFER, bufResult);
+			return false;
+		}
+
+		bool bHaveOutputBuffer = CurrentRenderOutputBuffer != nullptr;
+		NotifyReadyBufferListener(bHaveOutputBuffer);
+		if (bHaveOutputBuffer)
+		{
+			break;
+		}
+		else
+		{
+			// No available buffer. Sleep for a bit on the flush signal. This gets set before the termination signal.
+			FlushDecoderSignal.WaitTimeout(1000 * 5);
+		}
+	}
+	return true;
+
+}
+
+
+bool FVideoDecoderH265::ConvertDecodedImage(const TRefCountPtr<IMFSample>& DecodedOutputSample)
+{
+	SCOPE_CYCLE_COUNTER(STAT_ElectraPlayer_VideoH265ConvertOutput);
+	CSV_SCOPED_TIMING_STAT(ElectraPlayer, VideoH265ConvertOutput);
+
+	TUniquePtr<FParamDict> OutputBufferSampleProperties = MakeUnique<FParamDict>();
+	if (!OutputBufferSampleProperties)
+	{
+		return false;
+	}
+
+	HRESULT					res;
+	LONGLONG				llTimeStamp = 0;
+	MFVideoArea				videoArea = {};
+	UINT32					uiPanScanEnabled = 0;
+	FTimeValue				TimeStamp, Duration;
+
+	//PVS reports warning 773 (possible memory leak) for OutputBufferSampleProperties in the line below. It's a false positive as passing it into SetupDecodeOutputData() later on will pass on ownership implicitly.
+	VERIFY_HR(DecodedOutputSample->GetSampleTime(&llTimeStamp), "Failed to get video decoder output sample timestamp", ERRCODE_INTERNAL_COULD_NOT_GET_OUTPUT_SAMPLE_TIME); //-V773
+	TimeStamp.SetFromHNS((int64) llTimeStamp);
+	OutputBufferSampleProperties->Set("pts", FVariantValue(TimeStamp));
+
+	VERIFY_HR(DecodedOutputSample->GetSampleDuration(&llTimeStamp), "Failed to get video decoder output sample duration", ERRCODE_INTERNAL_COULD_NOT_GET_OUTPUT_SAMPLE_DURATION);
+	Duration.SetFromHNS((int64) llTimeStamp);
+	OutputBufferSampleProperties->Set("duration", FVariantValue(Duration));
+
+	if (SUCCEEDED(res = CurrentOutputMediaType->GetUINT32(MF_MT_PAN_SCAN_ENABLED, &uiPanScanEnabled)) && uiPanScanEnabled)
+	{
+		res = CurrentOutputMediaType->GetBlob(MF_MT_PAN_SCAN_APERTURE, (UINT8*)&videoArea, sizeof(MFVideoArea), nullptr);
+	}
+
+	UINT32	dwInputWidth = 0;
+	UINT32	dwInputHeight = 0;
+	res = MFGetAttributeSize(CurrentOutputMediaType.GetReference(), MF_MT_FRAME_SIZE, &dwInputWidth, &dwInputHeight);
+	check(SUCCEEDED(res));
+	res = CurrentOutputMediaType->GetBlob(MF_MT_MINIMUM_DISPLAY_APERTURE, (UINT8*)&videoArea, sizeof(MFVideoArea), nullptr);
+	if (FAILED(res))
+	{
+		res = CurrentOutputMediaType->GetBlob(MF_MT_GEOMETRIC_APERTURE, (UINT8*)&videoArea, sizeof(MFVideoArea), nullptr);
+		if (FAILED(res))
+		{
+			videoArea.OffsetX.fract = 0;
+			videoArea.OffsetX.value = 0;
+			videoArea.OffsetY.fract = 0;
+			videoArea.OffsetY.value = 0;
+			videoArea.Area.cx = dwInputWidth;
+			videoArea.Area.cy = dwInputHeight;
+		}
+	}
+
+	// Width and height is the cropped area!
+	OutputBufferSampleProperties->Set("width",  FVariantValue((int64) videoArea.Area.cx));
+	OutputBufferSampleProperties->Set("height", FVariantValue((int64) videoArea.Area.cy));
+
+	if (!bIsHardwareAccelerated || (FDXDeviceInfo::s_DXDeviceInfo->DxVersion == FDXDeviceInfo::ED3DVersion::Version11Win8))
+	{
+		// Cropping is applied to immediate decoder output and never visible to outside code
+		OutputBufferSampleProperties->Set("crop_left", FVariantValue((int64)0));
+		OutputBufferSampleProperties->Set("crop_right", FVariantValue((int64)0));
+		OutputBufferSampleProperties->Set("crop_top", FVariantValue((int64)0));
+		OutputBufferSampleProperties->Set("crop_bottom", FVariantValue((int64)0));
+	}
+	else
+	{
+		// Note: WMF decoder will always crop only at the lower, right borders
+		OutputBufferSampleProperties->Set("crop_left", FVariantValue((int64)0));
+		OutputBufferSampleProperties->Set("crop_right", FVariantValue((int64)(dwInputWidth - videoArea.Area.cx)));
+		OutputBufferSampleProperties->Set("crop_top", FVariantValue((int64)0));
+		OutputBufferSampleProperties->Set("crop_bottom", FVariantValue((int64)(dwInputHeight - videoArea.Area.cy)));
+	}
+
+	// Try to get the stride. Defaults to 0 should it not be obtainable.
+	UINT32 stride = MFGetAttributeUINT32(CurrentOutputMediaType, MF_MT_DEFAULT_STRIDE, 0);
+	OutputBufferSampleProperties->Set("pitch",  FVariantValue((int64) stride));
+
+	// Try to get the frame rate ratio
+	UINT32 num=0, denom=0;
+	MFGetAttributeRatio(CurrentOutputMediaType, MF_MT_FRAME_RATE, &num, &denom);
+	OutputBufferSampleProperties->Set("fps_num",   FVariantValue((int64) num ));
+	OutputBufferSampleProperties->Set("fps_denom", FVariantValue((int64) denom ));
+
+	// Try to get the pixel aspect ratio
+	num=0, denom=0;
+	MFGetAttributeRatio(CurrentOutputMediaType, MF_MT_PIXEL_ASPECT_RATIO, &num, &denom);
+	if (!num || !denom)
+	{
+		num   = 1;
+		denom = 1;
+	}
+	OutputBufferSampleProperties->Set("aspect_ratio", FVariantValue((double) num / (double)denom));
+	OutputBufferSampleProperties->Set("aspect_w",     FVariantValue((int64) num));
+	OutputBufferSampleProperties->Set("aspect_h",     FVariantValue((int64) denom));
+
+	OutputBufferSampleProperties->Set("pixelfmt",	  FVariantValue((int64)EPixelFormat::PF_NV12));
+
+	if (CurrentRenderOutputBuffer != nullptr)
+	{
+		if (!SetupDecodeOutputData(FIntPoint(videoArea.Area.cx, videoArea.Area.cy), DecodedOutputSample, OutputBufferSampleProperties.Get()))
+		{
+			return false;
+		}
+
+		Renderer->ReturnBuffer(CurrentRenderOutputBuffer, true, *OutputBufferSampleProperties);
+		CurrentRenderOutputBuffer = nullptr;
+		OutputBufferSampleProperties.Release();
+	}
+	return true;
+}
+
+
+bool FVideoDecoderH265::Decode(FAccessUnit* AccessUnit)
+{
+	// No decoder, nothing to do.
+	if (DecoderTransform.GetReference() == nullptr)
+	{
+		return true;
+	}
+
+	HRESULT res;
+	TRefCountPtr<IMFSample>	InputSample;
+	bool bFlush = AccessUnit == nullptr;
+
+	if (bFlush)
+	{
+		VERIFY_HR(DecoderTransform->ProcessMessage(MFT_MESSAGE_NOTIFY_END_OF_STREAM, 0), "Failed to set video decoder end of stream notification", ERRCODE_INTERNAL_COULD_NOT_SET_DECODER_ENDOFSTREAM);
+		VERIFY_HR(DecoderTransform->ProcessMessage(MFT_MESSAGE_COMMAND_DRAIN, 0), "Failed to issue video decoder drain command", ERRCODE_INTERNAL_COULD_NOT_SET_DECODER_DRAINCOMMAND);
+	}
+	else
+	{
+		// Create the input sample.
+		TRefCountPtr<IMFMediaBuffer> InputSampleBuffer;
+		BYTE*					pbNewBuffer = nullptr;
+		DWORD					dwMaxBufferSize = 0;
+		DWORD					dwSize = 0;
+		LONGLONG				llSampleTime = 0;
+
+		SCOPE_CYCLE_COUNTER(STAT_ElectraPlayer_VideoH265Decode);
+		CSV_SCOPED_TIMING_STAT(ElectraPlayer, VideoH265Decode);
+
+		VERIFY_HR(MFCreateSample(InputSample.GetInitReference()), "Failed to create video decoder input sample", ERRCODE_INTERNAL_COULD_NOT_CREATE_INPUT_SAMPLE);
+		VERIFY_HR(MFCreateMemoryBuffer((DWORD) AccessUnit->AUSize, InputSampleBuffer.GetInitReference()), "Failed to create video decoder input sample memory buffer", ERRCODE_INTERNAL_COULD_NOT_CREATE_INPUTBUFFER);
+		VERIFY_HR(InputSample->AddBuffer(InputSampleBuffer.GetReference()), "Failed to set video decoder input buffer with sample", ERRCODE_INTERNAL_COULD_NOT_ADD_INPUT_BUFFER_TO_SAMPLE);
+		VERIFY_HR(InputSampleBuffer->Lock(&pbNewBuffer, &dwMaxBufferSize, &dwSize), "Failed to lock video decoder input sample buffer", ERRCODE_INTERNAL_COULD_NOT_LOCK_INPUT_BUFFER);
+		FMemory::Memcpy(pbNewBuffer, AccessUnit->AUData, AccessUnit->AUSize);
+		VERIFY_HR(InputSampleBuffer->Unlock(), "Failed to unlock video decoder input sample buffer", ERRCODE_INTERNAL_COULD_NOT_UNLOCK_INPUT_BUFFER);
+		VERIFY_HR(InputSampleBuffer->SetCurrentLength((DWORD) AccessUnit->AUSize), "Failed to set video decoder input sample buffer length", ERRCODE_INTERNAL_COULD_NOT_SET_BUFFER_CURRENT_LENGTH);
+		// Set sample attributes
+		llSampleTime = AccessUnit->PTS.GetAsHNS();
+		VERIFY_HR(InputSample->SetSampleTime(llSampleTime), "Failed to set video decoder input sample presentation time", ERRCODE_INTERNAL_COULD_NOT_SET_INPUT_SAMPLE_TIME);
+		llSampleTime = AccessUnit->Duration.GetAsHNS();
+		VERIFY_HR(InputSample->SetSampleDuration(llSampleTime), "Failed to set video decode intput sample duration", ERRCODE_INTERNAL_COULD_NOT_SET_INPUT_SAMPLE_DURATION);
+		llSampleTime = AccessUnit->DTS.GetAsHNS();
+		VERIFY_HR(InputSample->SetUINT64(MFSampleExtension_DecodeTimestamp, llSampleTime), "Failed to set video decoder input sample decode time", ERRCODE_INTERNAL_COULD_NOT_SET_INPUT_SAMPLE_DECODETIME);
+		VERIFY_HR(InputSample->SetUINT32(MFSampleExtension_CleanPoint, AccessUnit->bIsSyncSample ? 1 : 0), "Failed to set video decoder input sample clean point", ERRCODE_INTERNAL_COULD_NOT_SET_INPUT_SAMPLE_KEYFRAME);
+	}
+
+	while(!TerminateThreadSignal.IsSignaled())
+	{
+		// Tick output buffer pool to handle state progression of samples in flight
+		Renderer->TickOutputBufferPool();
+		// Tick platform code in case something needs doing
+		PlatformTick();
+
+		// Check if to be flushed now. This may have resulted in AcquireOutputBuffer() to return with no buffer!
+		if (FlushDecoderSignal.IsSignaled())
+		{
+			break;
+		}
+
+		if (!CurrentDecoderOutputBuffer.IsValid())
+		{
+			SCOPE_CYCLE_COUNTER(STAT_ElectraPlayer_VideoH265Decode);
+			CSV_SCOPED_TIMING_STAT(ElectraPlayer, VideoH265Decode);
+
+			if (!CreateDecoderOutputBuffer())
+			{
+				return false;
+			}
+		}
+
+		CurrentDecoderOutputBuffer->PrepareForProcess();
+		DWORD	dwStatus = 0;
+		res = DecoderTransform->ProcessOutput(0, 1, &CurrentDecoderOutputBuffer->mOutputBuffer, &dwStatus);
+		CurrentDecoderOutputBuffer->UnprepareAfterProcess();
+
+		if (res == MF_E_TRANSFORM_NEED_MORE_INPUT)
+		{
+			if (bFlush)
+			{
+				// This means we have received all pending output and are done now.
+				// Issue a flush for completeness sake.
+				SCOPE_CYCLE_COUNTER(STAT_ElectraPlayer_VideoH265Decode);
+				CSV_SCOPED_TIMING_STAT(ElectraPlayer, VideoH265Decode);
+				VERIFY_HR(DecoderTransform->ProcessMessage(MFT_MESSAGE_COMMAND_FLUSH, 0), "Failed to issue video decoder flush command", ERRCODE_INTERNAL_COULD_NOT_SET_DECODER_FLUSHCOMMAND);
+				// And start over.
+				VERIFY_HR(DecoderTransform->ProcessMessage(MFT_MESSAGE_NOTIFY_START_OF_STREAM, 0), "Failed to start video decoder", ERRCODE_INTERNAL_COULD_NOT_SET_DECODER_START);
+				CurrentDecoderOutputBuffer.Reset();
+				NumFramesInDecoder = 0;
+				return true;
+			}
+			else if (InputSample.IsValid())
+			{
+				VERIFY_HR(DecoderTransform->ProcessInput(0, InputSample.GetReference(), 0), "Failed to process video decoder input", ERRCODE_INTERNAL_COULD_NOT_PROCESS_INPUT);
+				// Used this sample. Have no further input data for now, but continue processing to produce output if possible.
+				InputSample = nullptr;
+				++NumFramesInDecoder;
+			}
+			else
+			{
+				// Need more input but have none right now.
+				return true;
+			}
+		}
+		else if (res == MF_E_TRANSFORM_STREAM_CHANGE)
+		{
+			SCOPE_CYCLE_COUNTER(STAT_ElectraPlayer_VideoH265Decode);
+			CSV_SCOPED_TIMING_STAT(ElectraPlayer, VideoH265Decode);
+			// Update output type.
+			if (!DecoderSetOutputType())
+			{
+				return false;
+			}
+		}
+		else if (SUCCEEDED(res))
+		{
+			TRefCountPtr<IMFSample> DecodedOutputSample = CurrentDecoderOutputBuffer->DetachOutputSample();
+			//LogMessage(IInfoLog::ELevel::Info, StringHelpers::SPrintf(TEXT("OUT: %p"), DecodedOutputSample.Get()));
+			CurrentDecoderOutputBuffer.Reset();
+
+			if (DecodedOutputSample)
+			{
+				// HW accelerated decoding needs to get a buffer from the renderer only now after decoding.
+				// SW decoding required it earlier for the call to ProcessOutput().
+				if (bIsHardwareAccelerated)
+				{
+					if (!AcquireOutputBuffer())
+					{
+						return false;
+					}
+				}
+
+				// Check if to be flushed now. This may have resulted in AcquireOutputBuffer() to return with no buffer!
+				if (FlushDecoderSignal.IsSignaled())
+				{
+					break;
+				}
+
+				if (!ConvertDecodedImage(DecodedOutputSample))
+				{
+					return false;
+				}
+			}
+			--NumFramesInDecoder;
+		}
+		else
+		{
+			// Error!
+			VERIFY_HR(res, "Failed to process video decoder output", ERRCODE_INTERNAL_COULD_NOT_PROCESS_OUTPUT);
+			return false;
+		}
+	}
+	return true;
+}
+
+
+bool FVideoDecoderH265::DecodeDummy(FAccessUnit* AccessUnit)
+{
+	check(AccessUnit->Duration.IsValid());
+
+	FParamDict OutputBufferSampleProperties;
+
+	OutputBufferSampleProperties.Set("pts", FVariantValue(AccessUnit->PTS));
+	OutputBufferSampleProperties.Set("duration", FVariantValue(AccessUnit->Duration));
+	OutputBufferSampleProperties.Set("is_dummy", FVariantValue(true));
+
+	if (!AcquireOutputBuffer())
+	{
+		return false;
+	}
+
+	if (CurrentRenderOutputBuffer != nullptr)
+	{
+//TODO: THIS BEING A DUMMY - THE DICT IN ANY DECODER OUTPUT BUFFER (is there any?) WILL BE OUTDATED THIS WAY!!
+		Renderer->ReturnBuffer(CurrentRenderOutputBuffer, true, OutputBufferSampleProperties);
+		CurrentRenderOutputBuffer = nullptr;
+	}
+	return true;
+}
+
+
+void FVideoDecoderH265::DecoderCreate()
+{
+	if (InternalDecoderCreate())
+	{
+		// Configure the decoder with our default values.
+		Configure();
+		if (!bError)
+		{
+			StartStreaming();
+			bHaveDecoder = true;
+		}
+		else
+		{
+			InternalDecoderDestroy();
+		}
+	}
+	else
+	{
+		bError = true;
+		PostError(S_OK, "Failed to create H.265 decoder transform", ERRCODE_INTERNAL_COULD_NOT_CREATE_DECODER);
+	}
+}
+
+
+//-----------------------------------------------------------------------------
+/**
+ * Creates a decoder and checks if it will use hardware or software decoding.
+ */
+void FVideoDecoderH265::TestHardwareDecoding()
+{
+	// No-op.
+}
+
+
+//-----------------------------------------------------------------------------
+/**
+ * Application has entered foreground.
+ */
+void FVideoDecoderH265::HandleApplicationHasEnteredForeground()
+{
+	ApplicationRunningSignal.Signal();
+}
+
+
+//-----------------------------------------------------------------------------
+/**
+ * Application goes into background.
+ */
+void FVideoDecoderH265::HandleApplicationWillEnterBackground()
+{
+	ApplicationSuspendConfirmedSignal.Reset();
+	ApplicationRunningSignal.Reset();
+	if (!ApplicationSuspendConfirmedSignal.WaitTimeoutAndReset(1000 * 500))
+	{
+		UE_LOG(LogElectraPlayer, Log, TEXT("FVideoDecoderH265(%p): Did not suspend in time"), this);
+	}
+}
+
+
+//-----------------------------------------------------------------------------
+/**
+ * H265 video decoder main threaded decode loop
+ */
+void FVideoDecoderH265::WorkerThread()
+{
+	LLM_SCOPE(ELLMTag::ElectraPlayer);
+
+	ApplicationRunningSignal.Signal();
+	ApplicationSuspendConfirmedSignal.Reset();
+	ApplicationSuspendedDelegate = FCoreDelegates::ApplicationWillEnterBackgroundDelegate.AddRaw(this, &FVideoDecoderH265::HandleApplicationWillEnterBackground);
+	ApplicationResumeDelegate = FCoreDelegates::ApplicationHasEnteredForegroundDelegate.AddRaw(this, &FVideoDecoderH265::HandleApplicationHasEnteredForeground);
+
+	CurrentRenderOutputBuffer     = nullptr;
+	CurrentAccessUnit   		  = nullptr;
+	bHaveDecoder				  = false;
+	NumFramesInDecoder  		  = 0;
+	bError  					  = false;
+
+	// Set the configured maximum resolution as the current resolution which is used to create the decoder transform.
+	CurrentSampleInfo.SetResolution(FStreamCodecInformation::FResolution(Align(Config.MaxFrameWidth, 16), Align(Config.MaxFrameHeight, 16)));
+	NewSampleInfo = CurrentSampleInfo;
+	SetupBufferAcquisitionProperties();
+
+	// Create the decoded image pool.
+	if (!CreateDecodedImagePool())
+	{
+		bError = true;
+	}
+
+	bool bDone = false;
+	bool bInDummyDecodeMode = false;
+	bool bGotLastSequenceAU = false;
+
+	while(!TerminateThreadSignal.IsSignaled())
+	{
+		// If in background, wait until we get activated again.
+		if (!ApplicationRunningSignal.IsSignaled())
+		{
+			UE_LOG(LogElectraPlayer, Log, TEXT("FVideoDecoderH265(%p): OnSuspending"), this);
+			ApplicationSuspendConfirmedSignal.Signal();
+			while(!ApplicationRunningSignal.WaitTimeout(100 * 1000) && !TerminateThreadSignal.IsSignaled())
+			{
+			}
+			UE_LOG(LogElectraPlayer, Log, TEXT("FVideoDecoderH265(%p): OnResuming"), this);
+			continue;
+		}
+
+		if (!bDrainForCodecChange)
+		{
+			// Notify optional buffer listener that we will now be needing an AU for our input buffer.
+			if (!bError && InputBufferListener && AccessUnits.Num() == 0)
+			{
+				SCOPE_CYCLE_COUNTER(STAT_ElectraPlayer_VideoH265Decode);
+				CSV_SCOPED_TIMING_STAT(ElectraPlayer, VideoH265Decode);
+
+				FAccessUnitBufferInfo	sin;
+				IAccessUnitBufferListener::FBufferStats	stats;
+				AccessUnits.GetStats(sin);
+				stats.NumAUsAvailable  = sin.NumCurrentAccessUnits;
+				stats.NumBytesInBuffer = sin.CurrentMemInUse;
+				stats.MaxBytesOfBuffer = sin.MaxDataSize;
+				stats.bEODSignaled     = sin.bEndOfData;
+				stats.bEODReached      = sin.bEndOfData && sin.NumCurrentAccessUnits == 0;
+				ListenerMutex.Lock();
+				if (InputBufferListener)
+				{
+					InputBufferListener->DecoderInputNeeded(stats);
+				}
+				ListenerMutex.Unlock();
+			}
+			// Wait for data.
+			bool bHaveData = AccessUnits.WaitForData(1000 * 5);
+			// When there is data, even and especially after a previous EOD, we are no longer done and idling.
+			if (bHaveData && bDone)
+			{
+				bDone = false;
+			}
+
+			// Tick output buffer pool to handle state progression of samples in flight
+			Renderer->TickOutputBufferPool();
+			// Tick platform code in case something needs doing
+			PlatformTick();
+
+			if (bHaveData && Renderer->CanReceiveOutputFrames(1)) // note: the check for full output buffers ignores the NumFramesInDecoder value as we don't seem to receive new data if we don't feed new input
+			{
+				CurrentAccessUnit = nullptr;
+				bool bOk = AccessUnits.Pop(CurrentAccessUnit);
+				MEDIA_UNUSED_VAR(bOk);
+				check(bOk);
+
+				if (!bError)
+				{
+					if (!CurrentAccessUnit->bIsDummyData)
+					{
+						{
+						SCOPE_CYCLE_COUNTER(STAT_ElectraPlayer_VideoH265Decode);
+						CSV_SCOPED_TIMING_STAT(ElectraPlayer, VideoH265Decode);
+						PrepareAU(CurrentAccessUnit);
+						}
+
+						// Decode
+						if (!bError)
+						{
+							// Update the buffer acquisition after a change in streams?
+							bool bFormatChangedJustNow = CurrentSampleInfo.IsDifferentFromOtherVideo(NewSampleInfo) || bGotLastSequenceAU;
+
+							if (bInDummyDecodeMode)
+							{
+								bFormatChangedJustNow = true;
+								bInDummyDecodeMode = false;
+							}
+
+							if (!bHaveDecoder)
+							{
+								SCOPE_CYCLE_COUNTER(STAT_ElectraPlayer_VideoH265Decode);
+								CSV_SCOPED_TIMING_STAT(ElectraPlayer, VideoH265Decode);
+								DecoderCreate();
+							}
+
+							if (bFormatChangedJustNow)
+							{
+								if (!Decode(nullptr))
+								{
+									// FIXME: Is this error fatal or safe to ignore?
+								}
+								CurrentSampleInfo = NewSampleInfo;
+								SetupBufferAcquisitionProperties();
+							}
+							bGotLastSequenceAU = CurrentAccessUnit->bIsLastInPeriod;
+							if (!Decode(CurrentAccessUnit))
+							{
+								bError = true;
+							}
+						}
+					}
+					else
+					{
+						if (!bInDummyDecodeMode)
+						{
+							bInDummyDecodeMode = true;
+							Decode(nullptr);
+						}
+
+						if (!DecodeDummy(CurrentAccessUnit))
+						{
+							bError = true;
+						}
+					}
+				}
+
+				// If decoding was aborted and the AU not handed off, delete it now.
+				FAccessUnit::Release(CurrentAccessUnit);
+				CurrentAccessUnit = nullptr;
+			}
+			else
+			{
+				if (!bHaveData)
+				{
+					// No data. Is the buffer at EOD?
+					if (AccessUnits.IsEndOfData())
+					{
+						NotifyReadyBufferListener(true);
+						// Are we done yet?
+						if (!bDone && !bError)
+						{
+							bError = !Decode(nullptr);
+							NumFramesInDecoder = 0;
+						}
+						bDone = true;
+						FMediaRunnable::SleepMilliseconds(10);
+					}
+				}
+				else
+				{
+					// We could not decode new data as the output buffers are full... notify listener & wait a bit...
+					NotifyReadyBufferListener(false);
+					FMediaRunnable::SleepMilliseconds(10);
+				}
+			}
+		}
+		else
+		{
+			bError = !Decode(nullptr);
+			break;
+		}
+
+		// Flush?
+		if (FlushDecoderSignal.IsSignaled())
+		{
+			SCOPE_CYCLE_COUNTER(STAT_ElectraPlayer_VideoH265Decode);
+			CSV_SCOPED_TIMING_STAT(ElectraPlayer, VideoH265Decode);
+
+			// Destroy and re-create.
+			InternalDecoderDestroy();
+
+			ReturnUnusedFrame();
+			AccessUnits.Flush();
+
+			FlushDecoderSignal.Reset();
+			DecoderFlushedSignal.Signal();
+			NumFramesInDecoder = 0;
+			CurrentSampleInfo.SetResolution(FStreamCodecInformation::FResolution(Align(Config.MaxFrameWidth, 16), Align(Config.MaxFrameHeight, 16)));
+			NewSampleInfo = CurrentSampleInfo;
+
+			// Reset done state.
+			bDone = false;
+		}
+	}
+
+	InternalDecoderDestroy();
+	ReturnUnusedFrame();
+	DestroyDecodedImagePool();
+	AccessUnits.Flush();
+	AccessUnits.CapacitySet(0);
+
+	if (ApplicationSuspendedDelegate.IsValid())
+	{
+		FCoreDelegates::ApplicationWillEnterBackgroundDelegate.Remove(ApplicationSuspendedDelegate);
+		ApplicationSuspendedDelegate.Reset();
+	}
+	if (ApplicationResumeDelegate.IsValid())
+	{
+		FCoreDelegates::ApplicationHasEnteredForegroundDelegate.Remove(ApplicationResumeDelegate);
+		ApplicationResumeDelegate.Reset();
+	}
+
+	if (bDrainForCodecChange)
+	{
+		// Notify the player that we have finished draining.
+		PlayerSessionServices->SendMessageToPlayer(FDecoderMessage::Create(FDecoderMessage::EReason::DrainingFinished, this, EStreamType::Video, FStreamCodecInformation::ECodec::H265));
+		// We need to wait to get terminated. Also check if flushing is requested and acknowledge if it is.
+		while(!TerminateThreadSignal.IsSignaled())
+		{
+			if (FlushDecoderSignal.WaitTimeoutAndReset(1000 * 10))
+			{
+				DecoderFlushedSignal.Signal();
+			}
+		}
+	}
+}
+
+} // namespace Electra
+
+#endif // ELECTRA_ENABLE_MFDECODER
+
+
+#endif
diff --git a/Engine/Plugins/Media/ElectraPlayer/Source/ElectraPlayerRuntime/Private/Runtime/Decoder/DX/VideoDecoderH265_DX.h b/Engine/Plugins/Media/ElectraPlayer/Source/ElectraPlayerRuntime/Private/Runtime/Decoder/DX/VideoDecoderH265_DX.h
new file mode 100644
--- /dev/null
+++ b/Engine/Plugins/Media/ElectraPlayer/Source/ElectraPlayerRuntime/Private/Runtime/Decoder/DX/VideoDecoderH265_DX.h
@@ -0,0 +1,173 @@
+// Copyright Epic Games, Inc. All Rights Reserved.
+
+#pragma once
+
+namespace Electra {
+	/**
+	 * H265 video decoder class implementation.
+	**/
+	class FVideoDecoderH265 : public IVideoDecoderH265, public FMediaThread
+	{
+	public:
+		static bool Startup(const FParamDict& Options);
+		static void Shutdown();
+
+		static bool GetStreamDecodeCapability(FStreamDecodeCapability& OutResult, const FStreamDecodeCapability& InStreamParameter);
+
+		FVideoDecoderH265();
+		virtual ~FVideoDecoderH265();
+
+		void TestHardwareDecoding();
+
+		virtual void SetPlayerSessionServices(IPlayerSessionServices* SessionServices) override;
+
+		virtual void Open(const FInstanceConfiguration& InConfig) override;
+		virtual void Close() override;
+		virtual void DrainForCodecChange() override;
+
+		virtual void SetMaximumDecodeCapability(int32 MaxTier, int32 MaxWidth, int32 MaxHeight, int32 MaxProfile, int32 MaxProfileLevel, const FParamDict& AdditionalOptions) override;
+
+		virtual void SetAUInputBufferListener(IAccessUnitBufferListener* InListener) override;
+
+		virtual void SetReadyBufferListener(IDecoderOutputBufferListener* InListener) override;
+
+		virtual void SetRenderer(TSharedPtr<IMediaRenderer, ESPMode::ThreadSafe> InRenderer) override;
+
+		virtual void SetResourceDelegate(const TSharedPtr<IVideoDecoderResourceDelegate, ESPMode::ThreadSafe>& ResourceDelegate) override;
+
+		virtual IAccessUnitBufferInterface::EAUpushResult AUdataPushAU(FAccessUnit* InAccessUnit) override;
+		virtual void AUdataPushEOD() override;
+		virtual void AUdataFlushEverything() override;
+
+	protected:
+		struct FDecoderOutputBuffer
+		{
+			FDecoderOutputBuffer()
+			{
+				FMemory::Memzero(mOutputStreamInfo);
+				FMemory::Memzero(mOutputBuffer);
+			}
+			~FDecoderOutputBuffer()
+			{
+				if (mOutputBuffer.pSample)
+				{
+					mOutputBuffer.pSample->Release();
+				}
+			}
+			TRefCountPtr<IMFSample> DetachOutputSample()
+			{
+				TRefCountPtr<IMFSample> pOutputSample;
+				if (mOutputBuffer.pSample)
+				{
+					pOutputSample = TRefCountPtr<IMFSample>(mOutputBuffer.pSample, false);
+					mOutputBuffer.pSample = nullptr;
+				}
+				return pOutputSample;
+			}
+			void PrepareForProcess()
+			{
+				mOutputBuffer.dwStatus = 0;
+				mOutputBuffer.dwStreamID = 0;
+				mOutputBuffer.pEvents = nullptr;
+			}
+			void UnprepareAfterProcess()
+			{
+				if (mOutputBuffer.pEvents)
+				{
+					// https://docs.microsoft.com/en-us/windows/desktop/api/mftransform/nf-mftransform-imftransform-processoutput
+					// The caller is responsible for releasing any events that the MFT allocates.
+					mOutputBuffer.pEvents->Release();
+					mOutputBuffer.pEvents = nullptr;
+				}
+			}
+			MFT_OUTPUT_STREAM_INFO	mOutputStreamInfo;
+			MFT_OUTPUT_DATA_BUFFER	mOutputBuffer;
+		};
+
+		void InternalDecoderDestroy();
+		void DecoderCreate();
+		bool DecoderSetInputType();
+		bool DecoderSetOutputType();
+		bool DecoderVerifyStatus();
+		void StartThread();
+		void StopThread();
+		void WorkerThread();
+
+		bool CreateDecodedImagePool();
+		void DestroyDecodedImagePool();
+
+		void NotifyReadyBufferListener(bool bHaveOutput);
+
+		void SetupBufferAcquisitionProperties();
+
+		bool AcquireOutputBuffer();
+		bool ConvertDecodedImage(const TRefCountPtr<IMFSample>& DecodedSample);
+
+		void PrepareAU(FAccessUnit* InAccessUnit);
+
+		bool Decode(FAccessUnit* InAccessUnit);
+		bool DecodeDummy(FAccessUnit* InAccessUnit);
+
+		void ReturnUnusedFrame();
+
+		void PostError(int32 ApiReturnValue, const FString& Message, uint16 Code, UEMediaError Error = UEMEDIA_ERROR_OK);
+		void LogMessage(IInfoLog::ELevel Level, const FString& Message);
+		bool ReconfigureForSwDecoding(FString Reason);
+		bool Configure();
+		bool StartStreaming();
+
+		void HandleApplicationHasEnteredForeground();
+		void HandleApplicationWillEnterBackground();
+
+		// Per platform specialization
+		virtual bool InternalDecoderCreate() = 0;
+		virtual bool CreateDecoderOutputBuffer() = 0;
+		virtual void PreInitDecodeOutputForSW(const FIntPoint& Dim) = 0;
+		virtual bool SetupDecodeOutputData(const FIntPoint& ImageDim, const TRefCountPtr<IMFSample>& DecodedOutputSample, FParamDict* OutputBufferSampleProperties) = 0;
+		virtual void PlatformTick() {}
+
+		FInstanceConfiguration								Config;
+
+		FDelegateHandle										ApplicationSuspendedDelegate;
+		FDelegateHandle										ApplicationResumeDelegate;
+		FMediaEvent											ApplicationRunningSignal;
+		FMediaEvent											ApplicationSuspendConfirmedSignal;
+
+		FMediaEvent											TerminateThreadSignal;
+		FMediaEvent											FlushDecoderSignal;
+		FMediaEvent											DecoderFlushedSignal;
+		bool												bThreadStarted;
+		bool												bDrainForCodecChange;
+
+		IPlayerSessionServices*								PlayerSessionServices;
+
+		TSharedPtr<IMediaRenderer, ESPMode::ThreadSafe>		Renderer;
+
+		TWeakPtr<IVideoDecoderResourceDelegate, ESPMode::ThreadSafe> ResourceDelegate;
+
+		FAccessUnitBuffer									AccessUnits;
+		FStreamCodecInformation								NewSampleInfo;
+		FStreamCodecInformation								CurrentSampleInfo;
+		FAccessUnit*										CurrentAccessUnit;
+
+		FMediaCriticalSection								ListenerMutex;
+		IAccessUnitBufferListener*							InputBufferListener;
+		IDecoderOutputBufferListener*						ReadyBufferListener;
+
+		TRefCountPtr<IMFTransform>							DecoderTransform;
+		TRefCountPtr<IMFMediaType>							CurrentOutputMediaType;
+		MFT_OUTPUT_STREAM_INFO								DecoderOutputStreamInfo;
+		bool												bIsHardwareAccelerated;
+		int32												NumFramesInDecoder;
+		bool												bError;
+
+		TUniquePtr<FDecoderOutputBuffer>					CurrentDecoderOutputBuffer;
+		IMediaRenderer::IBuffer*							CurrentRenderOutputBuffer;
+		FParamDict											BufferAcquireOptions;
+		bool												bHaveDecoder;
+		int32												MaxDecodeBufferSize;
+
+		FIntPoint											MaxDecodeDim;
+	};
+
+} // namespace
diff --git a/Engine/Plugins/Media/ElectraPlayer/Source/ElectraPlayerRuntime/Private/Runtime/Decoder/DecoderBase.h b/Engine/Plugins/Media/ElectraPlayer/Source/ElectraPlayerRuntime/Private/Runtime/Decoder/DecoderBase.h
--- a/Engine/Plugins/Media/ElectraPlayer/Source/ElectraPlayerRuntime/Private/Runtime/Decoder/DecoderBase.h
+++ b/Engine/Plugins/Media/ElectraPlayer/Source/ElectraPlayerRuntime/Private/Runtime/Decoder/DecoderBase.h
@@ -3,6 +3,7 @@
 #pragma once
 
 #include "PlayerCore.h"
+#include "Player/PlayerSessionServices.h"
 #include "Renderer/RendererBase.h"
 
 
@@ -45,5 +46,67 @@ namespace Electra
 	};
 
 
+
+
+	class FDecoderMessage : public IPlayerMessage
+	{
+	public:
+		enum class EReason
+		{
+			DrainingFinished
+		};
+
+		static TSharedPtrTS<IPlayerMessage> Create(EReason InReason, IDecoderBase* InDecoderInstance, EStreamType InStreamType, FStreamCodecInformation::ECodec InCodec)
+		{
+			TSharedPtrTS<FDecoderMessage> p(new FDecoderMessage(InReason, InDecoderInstance, InStreamType, InCodec));
+			return p;
+		}
+
+		static const FString& Type()
+		{
+			static FString TypeName("Decoder");
+			return TypeName;
+		}
+
+		virtual const FString& GetType() const
+		{
+			return Type();
+		}
+
+		EReason GetReason() const
+		{
+			return Reason;
+		}
+
+		IDecoderBase* GetDecoderInstance() const
+		{
+			return DecoderInstance;
+		}
+		
+		EStreamType GetStreamType() const
+		{
+			return StreamType;
+		}
+		
+		FStreamCodecInformation::ECodec GetCodec() const
+		{
+			return Codec;
+		}
+
+	private:
+		FDecoderMessage(EReason InReason, IDecoderBase* InDecoderInstance, EStreamType InStreamType, FStreamCodecInformation::ECodec InCodec)
+			: Reason(InReason)
+			, DecoderInstance(InDecoderInstance)
+			, StreamType(InStreamType)
+			, Codec(InCodec)
+		{
+		}
+		EReason							Reason;
+		IDecoderBase*					DecoderInstance;
+		EStreamType						StreamType;
+		FStreamCodecInformation::ECodec	Codec;
+	};
+
+
 } // namespace Electra
 
diff --git a/Engine/Plugins/Media/ElectraPlayer/Source/ElectraPlayerRuntime/Private/Runtime/Decoder/Dummy/VideoDecoderH265_Dummy.cpp b/Engine/Plugins/Media/ElectraPlayer/Source/ElectraPlayerRuntime/Private/Runtime/Decoder/Dummy/VideoDecoderH265_Dummy.cpp
new file mode 100644
--- /dev/null
+++ b/Engine/Plugins/Media/ElectraPlayer/Source/ElectraPlayerRuntime/Private/Runtime/Decoder/Dummy/VideoDecoderH265_Dummy.cpp
@@ -0,0 +1,75 @@
+// Copyright Epic Games, Inc. All Rights Reserved.
+
+#include "PlayerCore.h"
+
+#if ELECTRA_PLATFORM_HAS_H265_DECODER == 0
+
+#include "Decoder/VideoDecoderH265.h"
+
+namespace Electra
+{
+	class FVideoDecoderH265 : public IVideoDecoderH265
+	{
+	public:
+		static bool Startup(const FParamDict& Options)
+		{ return true; }
+		static void Shutdown()
+		{ }
+		static bool GetStreamDecodeCapability(FStreamDecodeCapability& OutResult, const FStreamDecodeCapability& InStreamParameter)
+		{ return false; }
+		static FVideoDecoderH265* Create()
+		{ return nullptr; }
+		virtual ~FVideoDecoderH265() = default;
+		virtual void SetPlayerSessionServices(IPlayerSessionServices* SessionServices) override
+		{ }
+		virtual void Open(const FInstanceConfiguration& InConfig) override
+		{ }
+		virtual void Close() override
+		{ }
+		virtual void DrainForCodecChange() override
+		{ }
+		virtual void SetMaximumDecodeCapability(int32 MaxTier, int32 MaxWidth, int32 MaxHeight, int32 MaxProfile, int32 MaxProfileLevel, const FParamDict& AdditionalOptions) override
+		{ }
+		virtual void SetResourceDelegate(const TSharedPtr<IVideoDecoderResourceDelegate, ESPMode::ThreadSafe>& ResourceDelegate) override
+		{ }
+		virtual void SetRenderer(TSharedPtr<IMediaRenderer, ESPMode::ThreadSafe> InRenderer) override
+		{ }
+		virtual EAUpushResult AUdataPushAU(FAccessUnit* AccessUnit) override
+		{ return EAUpushResult::Error; }
+		virtual void AUdataPushEOD() override
+		{ }
+		virtual void AUdataFlushEverything() override
+		{ }
+		virtual void SetAUInputBufferListener(IAccessUnitBufferListener* Listener) override
+		{ }
+		virtual void SetReadyBufferListener(IDecoderOutputBufferListener* Listener) override
+		{ }
+	#if PLATFORM_ANDROID
+		virtual void Android_UpdateSurface(const TSharedPtr<IOptionPointerValueContainer>& Surface) override
+		{ }
+	#endif
+	};
+
+	IVideoDecoderH265::FInstanceConfiguration::FInstanceConfiguration()
+	{ }
+
+	IVideoDecoderH265* IVideoDecoderH265::Create()
+	{
+		return FVideoDecoderH265::Create();
+	}
+	bool IVideoDecoderH265::Startup(const FParamDict& Options)
+	{ 
+		return FVideoDecoderH265::Startup(Options);
+	}
+	void IVideoDecoderH265::Shutdown()
+	{ 
+		FVideoDecoderH265::Shutdown();
+	}
+	bool IVideoDecoderH265::GetStreamDecodeCapability(FStreamDecodeCapability& OutResult, const FStreamDecodeCapability& InStreamParameter)
+	{ 
+		return FVideoDecoderH265::GetStreamDecodeCapability(OutResult, InStreamParameter);
+	}
+
+
+}
+#endif
diff --git a/Engine/Plugins/Media/ElectraPlayer/Source/ElectraPlayerRuntime/Private/Runtime/Decoder/VideoDecoderBase.h b/Engine/Plugins/Media/ElectraPlayer/Source/ElectraPlayerRuntime/Private/Runtime/Decoder/VideoDecoderBase.h
new file mode 100644
--- /dev/null
+++ b/Engine/Plugins/Media/ElectraPlayer/Source/ElectraPlayerRuntime/Private/Runtime/Decoder/VideoDecoderBase.h
@@ -0,0 +1,73 @@
+// Copyright Epic Games, Inc. All Rights Reserved.
+
+#pragma once
+
+#include "PlayerCore.h"
+#include "StreamAccessUnitBuffer.h"
+#include "Decoder/DecoderBase.h"
+#include "VideoDecoderResourceDelegate.h"
+
+class IOptionPointerValueContainer;
+
+namespace Electra
+{
+	class IPlayerSessionServices;
+
+
+	/**
+	 * Video decoder.
+	**/
+	class IVideoDecoderBase : public IDecoderBase, public IAccessUnitBufferInterface, public IDecoderAUBufferDiags, public IDecoderReadyBufferDiags
+	{
+	public:
+		virtual ~IVideoDecoderBase() = default;
+
+		virtual void SetPlayerSessionServices(IPlayerSessionServices* SessionServices) = 0;
+
+		virtual void Close() = 0;
+
+		virtual void SetResourceDelegate(const TSharedPtr<IVideoDecoderResourceDelegate, ESPMode::ThreadSafe>& ResourceDelegate) = 0;
+
+
+		//! Drains the decoder of all enqueued input and ends it, after which the decoder must send an FDecoderMessage to the player
+		//! to signal completion.
+		virtual void DrainForCodecChange() = 0;
+
+
+		//-------------------------------------------------------------------------
+		// Methods from IDecoderBase
+		//
+		virtual void SetRenderer(TSharedPtr<IMediaRenderer, ESPMode::ThreadSafe> InRenderer) = 0;
+
+		//-------------------------------------------------------------------------
+		// Methods from IAccessUnitBufferInterface
+		//
+		//! Attempts to push an access unit to the decoder. Ownership of the access unit is transferred if the push is successful.
+		virtual EAUpushResult AUdataPushAU(FAccessUnit* AccessUnit) = 0;
+		//! Notifies the decoder that there will be no further access units.
+		virtual void AUdataPushEOD() = 0;
+		//! Instructs the decoder to flush all pending input and all already decoded output.
+		virtual void AUdataFlushEverything() = 0;
+
+		//-------------------------------------------------------------------------
+		// Methods from IDecoderAUBufferDiags
+		//
+		//! Registers an AU input buffer listener.
+		virtual void SetAUInputBufferListener(IAccessUnitBufferListener* Listener) = 0;
+
+		//-------------------------------------------------------------------------
+		// Methods from IDecoderReadyBufferDiags
+		//
+		virtual void SetReadyBufferListener(IDecoderOutputBufferListener* Listener) = 0;
+
+		//-------------------------------------------------------------------------
+		// Platform specifics
+		//
+#if PLATFORM_ANDROID
+		virtual void Android_UpdateSurface(const TSharedPtr<IOptionPointerValueContainer>& Surface) = 0;
+#endif
+	};
+
+
+} // namespace Electra
+
diff --git a/Engine/Plugins/Media/ElectraPlayer/Source/ElectraPlayerRuntime/Private/Runtime/Decoder/VideoDecoderH264.h b/Engine/Plugins/Media/ElectraPlayer/Source/ElectraPlayerRuntime/Private/Runtime/Decoder/VideoDecoderH264.h
--- a/Engine/Plugins/Media/ElectraPlayer/Source/ElectraPlayerRuntime/Private/Runtime/Decoder/VideoDecoderH264.h
+++ b/Engine/Plugins/Media/ElectraPlayer/Source/ElectraPlayerRuntime/Private/Runtime/Decoder/VideoDecoderH264.h
@@ -4,7 +4,7 @@
 
 #include "PlayerCore.h"
 #include "StreamAccessUnitBuffer.h"
-#include "Decoder/DecoderBase.h"
+#include "Decoder/VideoDecoderBase.h"
 #include "VideoDecoderResourceDelegate.h"
 
 class IOptionPointerValueContainer;
@@ -17,7 +17,7 @@ namespace Electra
 	/**
 	 * H264 video decoder.
 	**/
-	class IVideoDecoderH264 : public IDecoderBase, public IAccessUnitBufferInterface, public IDecoderAUBufferDiags, public IDecoderReadyBufferDiags
+	class IVideoDecoderH264 : public IVideoDecoderBase
 	{
 	public:
 		struct FSystemConfiguration
@@ -79,6 +79,7 @@ namespace Electra
 
 		virtual void Open(const FInstanceConfiguration& InConfig) = 0;
 		virtual void Close() = 0;
+		virtual void DrainForCodecChange() = 0;
 
 
 		virtual void SetMaximumDecodeCapability(int32 MaxWidth, int32 MaxHeight, int32 MaxProfile, int32 MaxProfileLevel, const FParamDict& AdditionalOptions) = 0;
diff --git a/Engine/Plugins/Media/ElectraPlayer/Source/ElectraPlayerRuntime/Private/Runtime/Decoder/VideoDecoderH265.h b/Engine/Plugins/Media/ElectraPlayer/Source/ElectraPlayerRuntime/Private/Runtime/Decoder/VideoDecoderH265.h
new file mode 100644
--- /dev/null
+++ b/Engine/Plugins/Media/ElectraPlayer/Source/ElectraPlayerRuntime/Private/Runtime/Decoder/VideoDecoderH265.h
@@ -0,0 +1,119 @@
+// Copyright Epic Games, Inc. All Rights Reserved.
+
+#pragma once
+
+#include "PlayerCore.h"
+#include "StreamAccessUnitBuffer.h"
+#include "Decoder/VideoDecoderBase.h"
+#include "VideoDecoderResourceDelegate.h"
+
+class IOptionPointerValueContainer;
+
+namespace Electra
+{
+	class IPlayerSessionServices;
+
+
+	/**
+	 * H265 video decoder.
+	**/
+	class IVideoDecoderH265 : public IVideoDecoderBase
+	{
+	public:
+		struct FInstanceConfiguration
+		{
+			FInstanceConfiguration();
+			int32								MaxFrameWidth = 0;				//!< Maximum width of any image to be decoded
+			int32								MaxFrameHeight = 0;				//!< Maximum height of any image to be decoded
+			int32								Tier = 0;						//!< Maximum tier level
+			int32								Profile = 0;					//!< Profile
+			int32								Level = 0;						//!< Level
+			uint32								MaxDecodedFrames = 0;			//!< Maximum number of decoded frames
+			FParamDict							AdditionalOptions;
+		};
+
+		static bool Startup(const FParamDict& Options);
+		static void Shutdown();
+
+		struct FStreamDecodeCapability
+		{
+			enum class ESupported
+			{
+				NotSupported,
+				SoftwareOnly,
+				HardwareOnly,
+				HardAndSoftware
+			};
+			ESupported	DecoderSupportType = ESupported::NotSupported;
+			int32		Width = 0;
+			int32		Height = 0;
+			int32		Tier = 0;
+			int32		Profile = 0;
+			uint32		CompatibilityFlags = 0;
+			int32		Level = 0;
+			uint64		ConstraintFlags = 0;
+			double		FPS = 0.0;
+			FParamDict	AdditionalOptions;
+		};
+		/*
+			Queries decoder support/capability for a stream with given properties.
+			Can be called after Startup() but should be called only shortly before playback to ensure all relevant
+			libraries are initialized.
+		*/
+		static bool GetStreamDecodeCapability(FStreamDecodeCapability& OutResult, const FStreamDecodeCapability& InStreamParameter);
+
+		static IVideoDecoderH265* Create();
+
+		virtual ~IVideoDecoderH265() = default;
+
+		virtual void SetPlayerSessionServices(IPlayerSessionServices* SessionServices) = 0;
+
+		virtual void Open(const FInstanceConfiguration& InConfig) = 0;
+		virtual void Close() = 0;
+		virtual void DrainForCodecChange() = 0;
+
+
+		virtual void SetMaximumDecodeCapability(int32 MaxTier, int32 MaxWidth, int32 MaxHeight, int32 MaxProfile, int32 MaxProfileLevel, const FParamDict& AdditionalOptions) = 0;
+
+		virtual void SetResourceDelegate(const TSharedPtr<IVideoDecoderResourceDelegate, ESPMode::ThreadSafe>& ResourceDelegate) = 0;
+
+		//-------------------------------------------------------------------------
+		// Methods from IDecoderBase
+		//
+		virtual void SetRenderer(TSharedPtr<IMediaRenderer, ESPMode::ThreadSafe> InRenderer) = 0;
+
+
+		//-------------------------------------------------------------------------
+		// Methods from IAccessUnitBufferInterface
+		//
+		//! Attempts to push an access unit to the decoder. Ownership of the access unit is transferred if the push is successful.
+		virtual EAUpushResult AUdataPushAU(FAccessUnit* AccessUnit) = 0;
+		//! Notifies the decoder that there will be no further access units.
+		virtual void AUdataPushEOD() = 0;
+		//! Instructs the decoder to flush all pending input and all already decoded output.
+		virtual void AUdataFlushEverything() = 0;
+
+
+		//-------------------------------------------------------------------------
+		// Methods from IDecoderAUBufferDiags
+		//
+		//! Registers an AU input buffer listener.
+		virtual void SetAUInputBufferListener(IAccessUnitBufferListener* Listener) = 0;
+
+
+		//-------------------------------------------------------------------------
+		// Methods from IDecoderReadyBufferDiags
+		//
+		virtual void SetReadyBufferListener(IDecoderOutputBufferListener* Listener) = 0;
+
+		//-------------------------------------------------------------------------
+		// Platform specifics
+		//
+#if PLATFORM_ANDROID
+		virtual void Android_UpdateSurface(const TSharedPtr<IOptionPointerValueContainer>& Surface) = 0;
+#endif
+	};
+
+
+} // namespace Electra
+
diff --git a/Engine/Plugins/Media/ElectraPlayer/Source/ElectraPlayerRuntime/Private/Runtime/Decoder/Windows/PlatformVideoDecoderH265.h b/Engine/Plugins/Media/ElectraPlayer/Source/ElectraPlayerRuntime/Private/Runtime/Decoder/Windows/PlatformVideoDecoderH265.h
new file mode 100644
--- /dev/null
+++ b/Engine/Plugins/Media/ElectraPlayer/Source/ElectraPlayerRuntime/Private/Runtime/Decoder/Windows/PlatformVideoDecoderH265.h
@@ -0,0 +1,19 @@
+// Copyright Epic Games, Inc. All Rights Reserved.
+
+#pragma once
+
+#include "Decoder/VideoDecoderH265.h"
+
+namespace Electra
+{
+	/**
+	 * H265 video decoder.
+	**/
+	class FPlatformVideoDecoderH265
+	{
+	public:
+		static bool GetPlatformStreamDecodeCapability(IVideoDecoderH265::FStreamDecodeCapability& OutResult, const IVideoDecoderH265::FStreamDecodeCapability& InStreamParameter);
+	};
+
+
+} // namespace Electra
diff --git a/Engine/Plugins/Media/ElectraPlayer/Source/ElectraPlayerRuntime/Private/Runtime/Decoder/Windows/PlatformVideoDecoderH265_Windows.cpp b/Engine/Plugins/Media/ElectraPlayer/Source/ElectraPlayerRuntime/Private/Runtime/Decoder/Windows/PlatformVideoDecoderH265_Windows.cpp
new file mode 100644
--- /dev/null
+++ b/Engine/Plugins/Media/ElectraPlayer/Source/ElectraPlayerRuntime/Private/Runtime/Decoder/Windows/PlatformVideoDecoderH265_Windows.cpp
@@ -0,0 +1,103 @@
+// Copyright Epic Games, Inc. All Rights Reserved.
+
+#include "Decoder/Windows/PlatformVideoDecoderH265.h"
+
+#if ELECTRA_PLATFORM_HAS_H265_DECODER
+#include "Windows/AllowWindowsPlatformTypes.h"
+#include "Windows/WindowsHWrapper.h"
+#include "HAL/LowLevelMemTracker.h"
+
+THIRD_PARTY_INCLUDES_START
+#include "mftransform.h"
+#include "mfapi.h"
+#include "mferror.h"
+#include "mfidl.h"
+THIRD_PARTY_INCLUDES_END
+
+#include "Windows/HideWindowsPlatformTypes.h"
+#endif
+
+namespace Electra
+{
+	#if ELECTRA_PLATFORM_HAS_H265_DECODER
+	namespace
+	{
+		static bool HaveMFTHEVCDecoder()
+		{
+			static int32 HaveHEVCDecoder = -1;
+			if (HaveHEVCDecoder < 0)
+			{
+				MFT_REGISTER_TYPE_INFO		MediaInputInfo { MFMediaType_Video , MFVideoFormat_HEVC };
+				TRefCountPtr<IMFAttributes>	Attributes;
+				TRefCountPtr<IMFTransform>	Decoder;
+				IMFActivate**				ActivateObjects = nullptr;
+				UINT32						NumActivateObjects = 0;
+				HRESULT						res;
+
+				HaveHEVCDecoder = 0;
+				if (MFTEnumEx(MFT_CATEGORY_VIDEO_DECODER, MFT_ENUM_FLAG_SYNCMFT, &MediaInputInfo, nullptr, &ActivateObjects, &NumActivateObjects) == S_OK)
+				{
+					if (NumActivateObjects)
+					{
+						IMFTransform* NewDecoder = nullptr;
+						res = ActivateObjects[0]->ActivateObject(IID_PPV_ARGS(&NewDecoder));
+						if (res == S_OK && NewDecoder)
+						{
+							HaveHEVCDecoder = 1;
+						}
+						if (NewDecoder)
+						{
+							NewDecoder->Release();
+						}
+						for(UINT32 i=0; i<NumActivateObjects; ++i)
+						{
+							ActivateObjects[i]->Release();
+						}
+						CoTaskMemFree(ActivateObjects);
+					}
+				}
+			}
+			return HaveHEVCDecoder == 1;
+		}
+	}
+	#endif
+
+
+	bool FPlatformVideoDecoderH265::GetPlatformStreamDecodeCapability(IVideoDecoderH265::FStreamDecodeCapability& OutResult, const IVideoDecoderH265::FStreamDecodeCapability& InStreamParameter)
+	{
+		#if ELECTRA_PLATFORM_HAS_H265_DECODER
+			bool bHaveMFTHEVCDecoder = HaveMFTHEVCDecoder();
+			if (bHaveMFTHEVCDecoder)
+			{
+				// Set the output same as the input. Modifications are made below.
+				OutResult = InStreamParameter;
+				OutResult.DecoderSupportType = IVideoDecoderH265::FStreamDecodeCapability::ESupported::NotSupported;
+
+				// Global query?
+				if (InStreamParameter.Width <= 0 && InStreamParameter.Height <= 0)
+				{
+					OutResult.Width = 3840;
+					OutResult.Height = 2160;
+					OutResult.Profile = 1;
+					OutResult.Level = 153;
+					OutResult.FPS = 60.0;
+					OutResult.DecoderSupportType = IVideoDecoderH265::FStreamDecodeCapability::ESupported::HardwareOnly;
+				}
+				else
+				{
+					// TBD - What limits should apply here, exactly?
+					if (InStreamParameter.FPS > 0.0 && InStreamParameter.FPS <= 60.0 &&
+						InStreamParameter.Level <= 153)		// Level 5.1
+					{
+						int32 maxMB = (3840 / 16) * (2160 / 16) * 60;
+						int32 numMB = (int32)(((InStreamParameter.Width + 15) / 16) * ((InStreamParameter.Height + 15) / 16) * (InStreamParameter.FPS > 0.0 ? InStreamParameter.FPS : 30.0));
+						OutResult.DecoderSupportType = numMB <= maxMB ? IVideoDecoderH265::FStreamDecodeCapability::ESupported::HardwareOnly : IVideoDecoderH265::FStreamDecodeCapability::ESupported::NotSupported;
+					}
+				}
+				return true;
+			}
+		#endif
+		return false;
+	}
+
+} // namespace Electra
diff --git a/Engine/Plugins/Media/ElectraPlayer/Source/ElectraPlayerRuntime/Private/Runtime/Decoder/Windows/VideoDecoderH264_PC.cpp b/Engine/Plugins/Media/ElectraPlayer/Source/ElectraPlayerRuntime/Private/Runtime/Decoder/Windows/VideoDecoderH264_PC.cpp
--- a/Engine/Plugins/Media/ElectraPlayer/Source/ElectraPlayerRuntime/Private/Runtime/Decoder/Windows/VideoDecoderH264_PC.cpp
+++ b/Engine/Plugins/Media/ElectraPlayer/Source/ElectraPlayerRuntime/Private/Runtime/Decoder/Windows/VideoDecoderH264_PC.cpp
@@ -31,6 +31,7 @@ THIRD_PARTY_INCLUDES_END
 #include "Windows/HideWindowsPlatformTypes.h"
 
 #include "Decoder/DX/VideoDecoderH264_DX.h"
+#include "Decoder/DX/MediaFoundationGUIDs.h"
 
 
 #define ELECTRA_USE_IMF2DBUFFER	1		// set to 1 to use IMF2DBuffer interface rather then plain media buffer when retrieving HW decoded data for CPU use
@@ -38,10 +39,6 @@ THIRD_PARTY_INCLUDES_END
 
 const D3DFORMAT DX9_NV12_FORMAT = (D3DFORMAT)MAKEFOURCC('N', 'V', '1', '2');
 
-#if (WINVER < _WIN32_WINNT_WIN8)
-const GUID MF_SA_D3D11_AWARE = { 0x206b4fc8, 0xfcf9, 0x4c51, { 0xaf, 0xe3, 0x97, 0x64, 0x36, 0x9e, 0x33, 0xa0 } };
-#endif
-
 /*********************************************************************************************************************/
 /*********************************************************************************************************************/
 /*********************************************************************************************************************/
@@ -211,8 +208,9 @@ void FElectraPlayerVideoDecoderOutputPC::PreInitForDecode(FIntPoint OutputDim, c
 		}
 		else // Software decode into CPU buffer
 		{
+			DWORD AllocSize = ((OutputDim.X + 15) & ~15) * ((OutputDim.Y + 15) & ~15) * sizeof(uint8*) * 2;
 			// SW decode results are just delivered in a simple CPU-side buffer. Create the decoder side version of this...
-			if (FAILED(Result = MFCreateMemoryBuffer(Dim.X * Dim.Y * sizeof(uint8), MediaBuffer.GetInitReference())))
+			if (FAILED(Result = MFCreateMemoryBuffer(AllocSize, MediaBuffer.GetInitReference())))
 			{
 				UE_LOG(LogElectraPlayer, Error, TEXT("MFCreateMemoryBuffer() failed with 0x%X %s"), Result, *GetComErrorDescription(Result));
 				PostError(Result, "Failed to create software decode output buffer", ERRCODE_INTERNAL_COULD_NOT_CREATE_OUTPUTBUFFER, UEMEDIA_ERROR_OK);
@@ -253,6 +251,9 @@ void FElectraPlayerVideoDecoderOutputPC::ProcessDecodeOutput(FIntPoint OutputDim
 
 	if (OutputType == EOutputType::SoftwareWin7) // Win7 & DX12 (SW)
 	{
+		// This needs to be a multiple of 16
+		SampleDim.Y = ((OutputDim.Y + 15) & ~15) * 3 / 2;
+
 		// Retrieve frame data and store it in Buffer for rendering later
 		TRefCountPtr<IMFMediaBuffer> MediaBuffer;
 		if (MFSample->GetBufferByIndex(0, MediaBuffer.GetInitReference()) != S_OK)
@@ -464,6 +465,11 @@ namespace Electra
 class FVideoDecoderH264_PC : public FVideoDecoderH264
 {
 public:
+	enum
+	{
+		MaxHWDecodeW = 1920,
+		MaxHWDecodeH = 1088
+	};
 private:
 	virtual bool InternalDecoderCreate() override;
 	virtual bool CreateDecoderOutputBuffer() override;
@@ -473,8 +479,6 @@ private:
 	bool CopyTexture(const TRefCountPtr<IMFSample>& DecodedSample, Electra::FParamDict* ParamDict, FIntPoint OutputDim);
 };
 
-static const int32 MaxHWDecodeW = 1920;
-static const int32 MaxHWDecodeH = 1088;
 
 /***************************************************************************************************************************************************/
 /***************************************************************************************************************************************************/
@@ -512,8 +516,8 @@ bool FVideoDecoderH264::GetStreamDecodeCapability(FStreamDecodeCapability& OutRe
 			// Check for HW support by creating and configuring a decoder.
 			FVideoDecoderH264* TestDecoder = new FVideoDecoderH264_PC;
 			// Set width and height in the decoder configuration. These are used in creating the decoder.
-			TestDecoder->Config.MaxFrameWidth = InStreamParameter.Width ? InStreamParameter.Width : MaxHWDecodeW;
-			TestDecoder->Config.MaxFrameHeight = InStreamParameter.Height ? InStreamParameter.Height : MaxHWDecodeH;
+			TestDecoder->Config.MaxFrameWidth = InStreamParameter.Width ? InStreamParameter.Width : FVideoDecoderH264_PC::MaxHWDecodeW;
+			TestDecoder->Config.MaxFrameHeight = InStreamParameter.Height ? InStreamParameter.Height : FVideoDecoderH264_PC::MaxHWDecodeH;
 			TestDecoder->TestHardwareDecoding();
 			bIsHWSupported = TestDecoder->bIsHardwareAccelerated;
 			delete TestDecoder;
