ElectraPlayer: Added support for HEVC decoder where available

diff --git a/Engine/Plugins/Media/ElectraPlayer/Source/ElectraPlayerRuntime/Private/Runtime/Decoder/Windows/VideoDecoderH265_PC.cpp b/Engine/Plugins/Media/ElectraPlayer/Source/ElectraPlayerRuntime/Private/Runtime/Decoder/Windows/VideoDecoderH265_PC.cpp
new file mode 100644
--- /dev/null
+++ b/Engine/Plugins/Media/ElectraPlayer/Source/ElectraPlayerRuntime/Private/Runtime/Decoder/Windows/VideoDecoderH265_PC.cpp
@@ -0,0 +1,729 @@
+// Copyright Epic Games, Inc. All Rights Reserved.
+
+#include "PlayerCore.h"
+#include "ElectraPlayerPrivate.h"
+#include "ElectraPlayerPrivate_Platform.h"
+#include "StreamAccessUnitBuffer.h"
+#include "Decoder/VideoDecoderH265.h"
+#include "Renderer/RendererBase.h"
+#include "Renderer/RendererVideo.h"
+#include "Player/PlayerSessionServices.h"
+#include "Utilities/StringHelpers.h"
+
+#if ELECTRA_PLATFORM_HAS_H265_DECODER
+
+#include "Decoder/DX/DecoderErrors_DX.h"
+#include "MediaVideoDecoderOutputPC.h"
+
+/*********************************************************************************************************************/
+/*********************************************************************************************************************/
+/*********************************************************************************************************************/
+
+#include "Windows/AllowWindowsPlatformTypes.h"
+#include "Windows/WindowsHWrapper.h"
+#include "HAL/LowLevelMemTracker.h"
+
+THIRD_PARTY_INCLUDES_START
+#include "mftransform.h"
+#include "mfapi.h"
+#include "mferror.h"
+#include "mfidl.h"
+THIRD_PARTY_INCLUDES_END
+
+#include "Decoder/DX/MediaFoundationGUIDs.h"
+
+#include "Windows/HideWindowsPlatformTypes.h"
+
+#include "Decoder/DX/VideoDecoderH265_DX.h"
+#include "Decoder/Windows/PlatformVideoDecoderH265.h"
+
+
+#define ELECTRA_USE_IMF2DBUFFER	1		// set to 1 to use IMF2DBuffer interface rather then plain media buffer when retrieving HW decoded data for CPU use
+
+/*********************************************************************************************************************/
+/*********************************************************************************************************************/
+/*********************************************************************************************************************/
+
+class FElectraPlayerVideoDecoderOutputH265PC : public FVideoDecoderOutputPC
+{
+public:
+	FElectraPlayerVideoDecoderOutputH265PC()
+		: OutputType(EOutputType::Unknown)
+		, Stride(0)
+		, SampleDim(0,0)
+	{
+	}
+
+	// Hardware decode to buffer (DX12)
+	void InitializeWithBuffer(const void* InBuffer, uint32 InSize, uint32 InStride, FIntPoint Dim, Electra::FParamDict* InParamDict);
+
+	// Hardware decode to shared DX11 texture (Win8+)
+	void InitializeWithSharedTexture(const TRefCountPtr<ID3D11Device>& InD3D11Device, const TRefCountPtr<ID3D11DeviceContext> InDeviceContext, const TRefCountPtr<IMFSample>& MFSample, const FIntPoint& OutputDim, Electra::FParamDict* InParamDict);
+
+	// Software decode (into texture if DX11 device specified - available only Win8+)
+	void PreInitForDecode(FIntPoint OutputDim, const TFunction<void(int32 /*ApiReturnValue*/, const FString& /*Message*/, uint16 /*Code*/, UEMediaError /*Error*/)>& PostError);
+	void ProcessDecodeOutput(FIntPoint OutputDim, Electra::FParamDict* InParamDict);
+
+	void SetOwner(const TSharedPtr<IDecoderOutputOwner, ESPMode::ThreadSafe>& InOwningRenderer) override
+	{
+		OwningRenderer = InOwningRenderer;
+	}
+
+	void ShutdownPoolable() override;
+
+	virtual EOutputType GetOutputType() const override
+	{
+		return OutputType;
+	}
+
+	virtual TRefCountPtr<IMFSample> GetMFSample() const override
+	{
+		check(OutputType == EOutputType::SoftwareWin8Plus || OutputType == EOutputType::SoftwareWin7);
+		return MFSample;
+	}
+
+	virtual const TArray<uint8>& GetBuffer() const override
+	{
+		check(OutputType == EOutputType::SoftwareWin8Plus || OutputType == EOutputType::SoftwareWin7 || OutputType == EOutputType::HardwareDX9_DX12);
+		return Buffer;
+	}
+
+	virtual uint32 GetStride() const override
+	{
+		return Stride;
+	}
+
+	virtual TRefCountPtr<ID3D11Texture2D> GetTexture() const override
+	{
+		return Texture;
+	}
+
+	virtual TRefCountPtr<ID3D11Device> GetDevice() const override
+	{
+		check(OutputType == EOutputType::SoftwareWin8Plus || OutputType == EOutputType::HardwareWin8Plus);
+		return D3D11Device;
+	}
+
+	virtual FIntPoint GetDim() const override
+	{
+		return SampleDim;
+	}
+
+private:
+	// Decoder output type
+	EOutputType OutputType;
+
+	// Output texture (with device that created it) for SW decoder output Win8+
+	TRefCountPtr<ID3D11Texture2D> Texture;
+	TRefCountPtr<ID3D11Texture2D> SharedTexture;
+	TRefCountPtr<ID3D11Device> D3D11Device;
+
+	// CPU-side buffer to 
+	TArray<uint8> Buffer;
+	uint32 Stride;
+
+	// WMF sample (owned by this class if SW decoder is used)
+	TRefCountPtr<IMFSample> MFSample;
+
+	// Dimension of any internally allocated buffer - stored explicitly to cover various special cases for DX
+	FIntPoint SampleDim;
+
+	// We hold a weak reference to the video renderer. During destruction the video renderer could be destroyed while samples are still out there..
+	TWeakPtr<IDecoderOutputOwner, ESPMode::ThreadSafe> OwningRenderer;
+};
+
+
+void FElectraPlayerVideoDecoderOutputH265PC::PreInitForDecode(FIntPoint OutputDim, const TFunction<void(int32 /*ApiReturnValue*/, const FString& /*Message*/, uint16 /*Code*/, UEMediaError /*Error*/ )>& PostError)
+{
+	FIntPoint Dim;
+	Dim.X = OutputDim.X;
+	Dim.Y = OutputDim.Y * 3 / 2;	// adjust height to encompass Y and UV planes
+
+	bool bNeedNew = !MFSample.IsValid() || SampleDim != Dim;
+
+	if (bNeedNew)
+	{
+		HRESULT Result;
+		MFSample = nullptr;
+		Texture = nullptr;
+		TRefCountPtr<IMFMediaBuffer> MediaBuffer;
+
+		SampleDim = Dim;
+
+		// Windows 8+ and a valid DX11 rendering device?
+		// (Rendering device will be null if the rendering device is not DX11 (DX12, Vulkan...))
+		if (Electra::IsWindows8Plus() && Electra::FDXDeviceInfo::s_DXDeviceInfo->RenderingDx11Device)
+		{
+			// Software decode into DX11 texture
+
+			// We SW decode right into a DX11 texture (via the media buffer created from it below)
+			/*
+				As we cannot use RHI at this point, we cannot create a texture RHI can readily use,
+				but can just generate a plain D3D one. We cannot easily use a delegate to ask UE to do
+				it for use as we want to keep the "how" of the decoder hidden behind the abstract interfaces
+				so we can also encapsulate non-UE-aware systems easily.
+
+				So: we will later need to copy this once. (TODO: see if the render team can provide a way RHI could use
+				a native texture like this directly)
+			*/
+			D3D11_TEXTURE2D_DESC TextureDesc;
+			TextureDesc.Width = Dim.X;
+			TextureDesc.Height = Dim.Y;
+			TextureDesc.MipLevels = 1;
+			TextureDesc.ArraySize = 1;
+			TextureDesc.Format = DXGI_FORMAT_R8_UINT;
+			TextureDesc.SampleDesc.Count = 1;
+			TextureDesc.SampleDesc.Quality = 0;
+			TextureDesc.Usage = D3D11_USAGE_DEFAULT;
+			TextureDesc.BindFlags = 0;
+			TextureDesc.CPUAccessFlags = 0;
+			TextureDesc.MiscFlags = 0;
+
+			if (FAILED(Result = Electra::FDXDeviceInfo::s_DXDeviceInfo->RenderingDx11Device->CreateTexture2D(&TextureDesc, nullptr, Texture.GetInitReference())))
+			{
+				UE_LOG(LogElectraPlayer, Error, TEXT("CreateTexture2D() failed with 0x%X %s"), Result, *GetComErrorDescription(Result));
+				PostError(Result, "Failed to create software decode output texture", ERRCODE_INTERNAL_COULD_NOT_CREATE_OUTPUTBUFFER, UEMEDIA_ERROR_OK);
+				return;
+			}
+			if (FAILED(Result = MFCreateDXGISurfaceBuffer(__uuidof(ID3D11Texture2D), Texture, 0, false, MediaBuffer.GetInitReference())))
+			{
+				UE_LOG(LogElectraPlayer, Error, TEXT("MFCreateDXGISurfaceBuffer() failed with 0x%X %s"), Result, *GetComErrorDescription(Result));
+				PostError(Result, "Failed to create software decode output texture", ERRCODE_INTERNAL_COULD_NOT_CREATE_OUTPUTBUFFER, UEMEDIA_ERROR_OK);
+				return;
+			}
+
+			OutputType = EOutputType::SoftwareWin8Plus;
+		}
+		else // Software decode into CPU buffer
+		{
+			// SW decode results are just delivered in a simple CPU-side buffer. Create the decoder side version of this...
+			if (FAILED(Result = MFCreateMemoryBuffer(Dim.X * Dim.Y * sizeof(uint8), MediaBuffer.GetInitReference())))
+			{
+				UE_LOG(LogElectraPlayer, Error, TEXT("MFCreateMemoryBuffer() failed with 0x%X %s"), Result, *GetComErrorDescription(Result));
+				PostError(Result, "Failed to create software decode output buffer", ERRCODE_INTERNAL_COULD_NOT_CREATE_OUTPUTBUFFER, UEMEDIA_ERROR_OK);
+				return;
+			}
+			OutputType = EOutputType::SoftwareWin7;
+		}
+
+		if (SUCCEEDED(Result = MFCreateSample(MFSample.GetInitReference())))
+		{
+			if (FAILED(Result = MFSample->AddBuffer(MediaBuffer)))
+			{
+				UE_LOG(LogElectraPlayer, Error, TEXT("IMFSample::AddBuffer() failed with 0x%X %s"), Result, *GetComErrorDescription(Result));
+				PostError(Result, "Failed to add output buffer to software decoder sample", ERRCODE_INTERNAL_COULD_NOT_ADD_OUTPUT_BUFFER_TO_SAMPLE, UEMEDIA_ERROR_OK);
+				return;
+			}
+		}
+		else
+		{
+			UE_LOG(LogElectraPlayer, Error, TEXT("MFCreateSample() failed with 0x%X %s"), Result, *GetComErrorDescription(Result));
+			PostError(Result, "Failed to create software decode output sample", ERRCODE_INTERNAL_COULD_NOT_CREATE_OUTPUT_SAMPLE, UEMEDIA_ERROR_OK);
+			return;
+		}
+	}
+}
+
+
+void FElectraPlayerVideoDecoderOutputH265PC::ProcessDecodeOutput(FIntPoint OutputDim, Electra::FParamDict* InParamDict)
+{
+	check(OutputType == EOutputType::SoftwareWin7 || OutputType == EOutputType::SoftwareWin8Plus);
+
+	FVideoDecoderOutput::Initialize(InParamDict);
+
+	SampleDim.X = OutputDim.X;
+	SampleDim.Y = OutputDim.Y * 3 / 2;
+
+	Stride = SampleDim.X * sizeof(uint8);
+
+	if (OutputType == EOutputType::SoftwareWin7) // Win7 & DX12 (SW)
+	{
+		// Retrieve frame data and store it in Buffer for rendering later
+		TRefCountPtr<IMFMediaBuffer> MediaBuffer;
+		if (MFSample->GetBufferByIndex(0, MediaBuffer.GetInitReference()) != S_OK)
+		{
+			return;
+		}
+		DWORD BufferSize = 0;
+		if (MediaBuffer->GetCurrentLength(&BufferSize) != S_OK)
+		{
+			return;
+		}
+		check(SampleDim.X * SampleDim.Y <= (int32)BufferSize);
+
+		// MediaBuffer reports incorrect buffer sizes (too large) for some resolutions: we use our computed values!
+		BufferSize = SampleDim.X * SampleDim.Y;
+
+		uint8* Data = nullptr;
+		Buffer.Reset(BufferSize);
+		if (MediaBuffer->Lock(&Data, NULL, NULL) == S_OK)
+		{
+			Buffer.Append(Data, BufferSize);
+			MediaBuffer->Unlock();
+		}
+
+		// The output IMFSample needs to be released (and recreated for the next frame) for unknown reason.
+		// If not destroyed decoder throws an unknown error later on.
+		MFSample = nullptr;
+	}
+	else
+	{
+		// SW decode into texture need Win8+ (but no additional processing)
+		check(Electra::IsWindows8Plus());
+	}
+}
+
+
+void FElectraPlayerVideoDecoderOutputH265PC::InitializeWithBuffer(const void* InBuffer, uint32 InSize, uint32 InStride, FIntPoint Dim, Electra::FParamDict* InParamDict)
+{
+	FVideoDecoderOutput::Initialize(InParamDict);
+
+	OutputType = EOutputType::HardwareDX9_DX12;
+
+	Buffer.Reset(InSize);
+	Buffer.Append((uint8*)InBuffer, InSize);
+
+	SampleDim = Dim;
+	Stride = InStride;
+}
+
+
+void FElectraPlayerVideoDecoderOutputH265PC::InitializeWithSharedTexture(const TRefCountPtr<ID3D11Device>& InD3D11Device, const TRefCountPtr<ID3D11DeviceContext> InDeviceContext, const TRefCountPtr<IMFSample>& InMFSample, const FIntPoint& OutputDim, Electra::FParamDict* InParamDict)
+{
+	FVideoDecoderOutput::Initialize(InParamDict);
+
+	OutputType = EOutputType::HardwareWin8Plus;
+
+	bool bNeedsNew = !Texture.IsValid() || (SampleDim.X != OutputDim.X || SampleDim.Y != OutputDim.Y);
+
+	if (bNeedsNew)
+	{
+		SampleDim = OutputDim;
+
+		// Make a texture we pass on as output
+		D3D11_TEXTURE2D_DESC TextureDesc;
+		TextureDesc.Width = SampleDim.X;
+		TextureDesc.Height = SampleDim.Y;
+		TextureDesc.MipLevels = 1;
+		TextureDesc.ArraySize = 1;
+		TextureDesc.Format = DXGI_FORMAT_NV12;
+		TextureDesc.SampleDesc.Count = 1;
+		TextureDesc.SampleDesc.Quality = 0;
+		TextureDesc.Usage = D3D11_USAGE_DEFAULT;
+		TextureDesc.BindFlags = 0;
+		TextureDesc.CPUAccessFlags = 0;
+		TextureDesc.MiscFlags = D3D11_RESOURCE_MISC_SHARED_KEYEDMUTEX;
+
+		if (InD3D11Device->CreateTexture2D(&TextureDesc, nullptr, Texture.GetInitReference()) != S_OK)
+		{
+			return;
+		}
+
+		D3D11Device = InD3D11Device;
+	}
+
+	// If we got a texture, copy the data from the decoder into it...
+	if (Texture)
+	{
+		// Get output texture from decoder...
+		TRefCountPtr<IMFMediaBuffer> MediaBuffer;
+		if (InMFSample->GetBufferByIndex(0, MediaBuffer.GetInitReference()) != S_OK)
+		{
+			return;
+		}
+		TRefCountPtr<IMFDXGIBuffer> DXGIBuffer;
+		if (MediaBuffer->QueryInterface(__uuidof(IMFDXGIBuffer), (void**)DXGIBuffer.GetInitReference()) != S_OK)
+		{
+			return;
+		}
+		TRefCountPtr<ID3D11Texture2D> DecoderTexture;
+		if (DXGIBuffer->GetResource(IID_PPV_ARGS(DecoderTexture.GetInitReference())) != S_OK)
+		{
+			return;
+		}
+		uint32 ViewIndex = 0;
+		if (DXGIBuffer->GetSubresourceIndex(&ViewIndex) != S_OK)
+		{
+			return;
+		}
+
+		// Source data may be larger than desired output, but crop area will be aligned such that we can always copy from 0,0
+		D3D11_BOX SrcBox;
+		SrcBox.left = 0;
+		SrcBox.top = 0;
+		SrcBox.front = 0;
+		SrcBox.right = OutputDim.X;
+		SrcBox.bottom = OutputDim.Y;
+		SrcBox.back = 1;
+
+		TRefCountPtr<IDXGIKeyedMutex> KeyedMutex;
+		HRESULT res = Texture->QueryInterface(_uuidof(IDXGIKeyedMutex), (void**)&KeyedMutex);
+
+		if (KeyedMutex)
+		{
+			// No wait on acquire since sample is new and key is 0.
+			res = KeyedMutex->AcquireSync(0, 0);
+			if (res == S_OK)
+			{
+				// Copy texture using the decoder DX11 device... (and apply any cropping - see above note)
+				InDeviceContext->CopySubresourceRegion(Texture, 0, 0, 0, 0, DecoderTexture, ViewIndex, &SrcBox);
+
+				// Mark texture as updated with key of 1
+				// Sample will be read in Convert() method of texture sample
+				KeyedMutex->ReleaseSync(1);
+			}
+		}
+
+		// Make sure texture is updated before giving access to the sample in the rendering thread.
+		InDeviceContext->Flush();
+	}
+}
+
+void FElectraPlayerVideoDecoderOutputH265PC::ShutdownPoolable()
+{
+	TSharedPtr<IDecoderOutputOwner, ESPMode::ThreadSafe> lockedVideoRenderer = OwningRenderer.Pin();
+	if (lockedVideoRenderer.IsValid())
+	{
+		lockedVideoRenderer->SampleReleasedToPool(GetDuration());
+	}
+
+	if (OutputType != EOutputType::HardwareWin8Plus)
+	{
+		return;
+	}
+
+	// Correctly release the keyed mutex when the sample is returned to the pool
+	TRefCountPtr<IDXGIResource> OtherResource(nullptr);
+	if (Texture)
+	{
+		Texture->QueryInterface(__uuidof(IDXGIResource), (void**)&OtherResource);
+	}
+
+	if (OtherResource)
+	{
+		HANDLE SharedHandle = nullptr;
+		if (OtherResource->GetSharedHandle(&SharedHandle) == S_OK)
+		{
+			TRefCountPtr<ID3D11Resource> SharedResource;
+			D3D11Device->OpenSharedResource(SharedHandle, __uuidof(ID3D11Texture2D), (void**)&SharedResource);
+			if (SharedResource)
+			{
+				TRefCountPtr<IDXGIKeyedMutex> KeyedMutex;
+				OtherResource->QueryInterface(_uuidof(IDXGIKeyedMutex), (void**)&KeyedMutex);
+
+				if (KeyedMutex)
+				{
+					// Reset keyed mutex
+					if (KeyedMutex->AcquireSync(1, 0) == S_OK)
+					{
+						// Texture was never read
+						KeyedMutex->ReleaseSync(0);
+					}
+					else if (KeyedMutex->AcquireSync(2, 0) == S_OK)
+					{
+						// Texture was read at least once
+						KeyedMutex->ReleaseSync(0);
+					}
+				}
+			}
+		}
+	}
+
+}
+
+// -----------------------------------------------------------------------------------------------------------------------------
+
+
+namespace Electra
+{
+
+/*********************************************************************************************************************/
+/*********************************************************************************************************************/
+/*********************************************************************************************************************/
+
+class FVideoDecoderH265_PC : public FVideoDecoderH265
+{
+public:
+private:
+	virtual bool InternalDecoderCreate() override;
+	virtual bool CreateDecoderOutputBuffer() override;
+	virtual void PreInitDecodeOutputForSW(const FIntPoint& Dim);
+	virtual bool SetupDecodeOutputData(const FIntPoint& ImageDim, const TRefCountPtr<IMFSample>& DecodedOutputSample, FParamDict* OutputBufferSampleProperties) override;
+
+	bool CopyTexture(const TRefCountPtr<IMFSample>& DecodedSample, Electra::FParamDict* ParamDict, FIntPoint OutputDim);
+};
+
+/***************************************************************************************************************************************************/
+/***************************************************************************************************************************************************/
+/***************************************************************************************************************************************************/
+
+IVideoDecoderH265* IVideoDecoderH265::Create()
+{
+	return new FVideoDecoderH265_PC();
+}
+
+//-----------------------------------------------------------------------------
+/**
+ *	Queries decoder support/capability for a stream with given properties.
+ *	Can be called after Startup() but should be called only shortly before playback to ensure all relevant
+ *	libraries are initialized.
+ */
+bool FVideoDecoderH265::GetStreamDecodeCapability(FStreamDecodeCapability& OutResult, const FStreamDecodeCapability& InStreamParameter)
+{
+	return FPlatformVideoDecoderH265::GetPlatformStreamDecodeCapability(OutResult, InStreamParameter);
+}
+
+//-----------------------------------------------------------------------------
+/**
+ * Create a decoder instance.
+ *
+ * @return true if successful, false on error
+ */
+bool FVideoDecoderH265_PC::InternalDecoderCreate()
+{
+	MFT_REGISTER_TYPE_INFO		MediaInputInfo { MFMediaType_Video , MFVideoFormat_HEVC };
+	TRefCountPtr<IMFAttributes>	Attributes;
+	TRefCountPtr<IMFTransform>	Decoder;
+	IMFActivate**				ActivateObjects = nullptr;
+	UINT32						NumActivateObjects = 0;
+	HRESULT						res;
+
+	if (Electra::IsWindows8Plus())
+	{
+		// Check if there is any reason for a "device lost" - if not we know all is stil well; otherwise we bail without creating a decoder
+		if (Electra::FDXDeviceInfo::s_DXDeviceInfo->DxDevice && (res = Electra::FDXDeviceInfo::s_DXDeviceInfo->DxDevice->GetDeviceRemovedReason()) != S_OK)
+		{
+			PostError(res, "Device lost detected.", ERRCODE_INTERNAL_COULD_NOT_SET_OUTPUT_DESIRED_MEDIA_TYPE);
+			return false;
+		}
+	}
+
+	VERIFY_HR(MFTEnumEx(MFT_CATEGORY_VIDEO_DECODER, MFT_ENUM_FLAG_SYNCMFT, &MediaInputInfo, nullptr, &ActivateObjects, &NumActivateObjects), "MFTEnumEx failed", ERRCODE_INTERNAL_COULD_NOT_CREATE_DECODER);
+	if (NumActivateObjects == 0)
+	{
+		PostError(S_OK, "MFTEnumEx returned zero activation objects", ERRCODE_INTERNAL_COULD_NOT_CREATE_DECODER);
+		return false;
+	}
+
+	IMFTransform* NewDecoder = nullptr;
+	res = ActivateObjects[0]->ActivateObject(IID_PPV_ARGS(&NewDecoder));
+	if (res == S_OK)
+	{
+		*Decoder.GetInitReference() = NewDecoder;
+	}
+	for(UINT32 i=0; i<NumActivateObjects; ++i)
+	{
+		ActivateObjects[i]->Release();
+	}
+	CoTaskMemFree(ActivateObjects);
+	ActivateObjects = nullptr;
+	if (res != S_OK)
+	{
+		PostError(res, "HEVC decoder transform activation failed", ERRCODE_INTERNAL_COULD_NOT_CREATE_DECODER);
+		return false;
+	}
+
+	VERIFY_HR(Decoder->GetAttributes(Attributes.GetInitReference()), "Failed to get video decoder transform attributes", ERRCODE_INTERNAL_COULD_GET_TRANSFORM_ATTRIBUTES);
+
+	// Check if the transform is D3D aware
+	if (Electra::IsWindows8Plus())
+	{
+		if (!Electra::FDXDeviceInfo::s_DXDeviceInfo->DxDeviceManager)
+		{
+			// Win8+ with no DX11. No need to check for any DX11 awareness...
+			return true;
+		}
+
+		uint32 IsDX11Aware = 0;
+		res = Attributes->GetUINT32(MF_SA_D3D11_AWARE, &IsDX11Aware);
+		if (FAILED(res))
+		{
+			PostError(res, TEXT("Failed to get MF_SA_D3D11_AWARE"), ERRCODE_INTERNAL_COULD_NOT_CREATE_DECODER);
+			return false;
+		}
+		else if (IsDX11Aware == 0)
+		{
+			PostError(res, TEXT("Not MF_SA_D3D11_AWARE"), ERRCODE_INTERNAL_COULD_NOT_CREATE_DECODER);
+			return false;
+		}
+		else if (FAILED(res = Decoder->ProcessMessage(MFT_MESSAGE_SET_D3D_MANAGER, reinterpret_cast<ULONG_PTR>(Electra::FDXDeviceInfo::s_DXDeviceInfo->DxDeviceManager.GetReference()))))
+		{
+			PostError(res, FString::Printf(TEXT("Failed to set MFT_MESSAGE_SET_D3D_MANAGER: 0x%X %s"), res, *GetComErrorDescription(res)), ERRCODE_INTERNAL_COULD_NOT_CREATE_DECODER);
+			return false;
+		}
+	}
+	else // Windows 7
+	{
+		return false;
+	}
+
+	// Prepare for the maximum resolution
+	VERIFY_HR(Attributes->SetUINT32(MF_MT_DECODER_USE_MAX_RESOLUTION, 1), "Failed to allocate for maximum resolution", ERRCODE_INTERNAL_COULD_NOT_CREATE_DECODER);
+
+	if (FAILED(res = Attributes->SetUINT32(CODECAPI_AVLowLatencyMode, 1)))
+	{
+		// Not an error. If it doesn't work it just doesn't.
+	}
+
+	// Create successful, take on the decoder.
+	DecoderTransform = Decoder;
+
+	return true;
+}
+
+
+bool FVideoDecoderH265_PC::CreateDecoderOutputBuffer()
+{
+	HRESULT									res;
+	TUniquePtr<FDecoderOutputBuffer>		NewDecoderOutputBuffer(new FDecoderOutputBuffer());
+
+	VERIFY_HR(DecoderTransform->GetOutputStreamInfo(0, &NewDecoderOutputBuffer->mOutputStreamInfo), "Failed to get video decoder output stream information", ERRCODE_INTERNAL_COULD_NOT_GET_OUTPUT_STREAM_INFO);
+	// Do we need to provide the sample output buffer or does the decoder create it for us?
+	if ((NewDecoderOutputBuffer->mOutputStreamInfo.dwFlags & (MFT_OUTPUT_STREAM_PROVIDES_SAMPLES | MFT_OUTPUT_STREAM_CAN_PROVIDE_SAMPLES)) == 0)
+	{
+		return false;
+	}
+	CurrentDecoderOutputBuffer = MoveTemp(NewDecoderOutputBuffer);
+	return true;
+}
+
+
+void FVideoDecoderH265_PC::PreInitDecodeOutputForSW(const FIntPoint& Dim)
+{
+	TSharedPtr<FElectraPlayerVideoDecoderOutputH265PC, ESPMode::ThreadSafe> DecoderOutput = CurrentRenderOutputBuffer->GetBufferProperties().GetValue("texture").GetSharedPointer<FElectraPlayerVideoDecoderOutputH265PC>();
+	check(DecoderOutput);
+	DecoderOutput->PreInitForDecode(Dim, [this](int32 ApiReturnValue, const FString& Message, uint16 Code, UEMediaError Error) { PostError(ApiReturnValue, Message, Code, Error); });
+}
+
+
+bool FVideoDecoderH265_PC::SetupDecodeOutputData(const FIntPoint & ImageDim, const TRefCountPtr<IMFSample>& DecodedOutputSample, FParamDict* OutputBufferSampleProperties)
+{
+	bool bCopyResult = CopyTexture(DecodedOutputSample, OutputBufferSampleProperties, ImageDim);
+	if (!bCopyResult)
+	{
+		// Failed for some reason. Let's best not render this then, but put in the buffers nonetheless to maintain the normal data flow.
+		OutputBufferSampleProperties->Set("is_dummy", FVariantValue(true));
+	}
+	return true;
+}
+
+
+bool FVideoDecoderH265_PC::CopyTexture(const TRefCountPtr<IMFSample>& Sample, Electra::FParamDict* ParamDict, FIntPoint OutputDim)
+{
+#if !UE_SERVER
+	TSharedPtr<FElectraPlayerVideoDecoderOutputH265PC, ESPMode::ThreadSafe> DecoderOutput = CurrentRenderOutputBuffer->GetBufferProperties().GetValue("texture").GetSharedPointer<FElectraPlayerVideoDecoderOutputH265PC>();
+	check(DecoderOutput);
+
+	DWORD BuffersNum = 0;
+	CHECK_HR(Sample->GetBufferCount(&BuffersNum));
+
+	if (BuffersNum != 1)
+	{
+		return false;
+	}
+
+	LONGLONG SampleTime = 0;
+	CHECK_HR(Sample->GetSampleTime(&SampleTime));
+	LONGLONG SampleDuration = 0;
+	CHECK_HR(Sample->GetSampleDuration(&SampleDuration));
+
+	TRefCountPtr<IMFMediaBuffer> Buffer;
+	CHECK_HR(Sample->GetBufferByIndex(0, Buffer.GetInitReference()));
+
+	if (FDXDeviceInfo::s_DXDeviceInfo->DxVersion == FDXDeviceInfo::ED3DVersion::Version11Win8)
+	{
+		check(FDXDeviceInfo::s_DXDeviceInfo->DxDevice);
+
+		DecoderOutput->InitializeWithSharedTexture(
+			FDXDeviceInfo::s_DXDeviceInfo->DxDevice,
+			FDXDeviceInfo::s_DXDeviceInfo->DxDeviceContext,
+			Sample,
+			OutputDim,
+			ParamDict);
+
+		if (!DecoderOutput->GetTexture())
+		{
+			UE_LOG(LogElectraPlayer, Error, TEXT("FVideoDecoderH265: InD3D11Device->CreateTexture2D() failed!"));
+			return false;
+		}
+	}
+	else if (FDXDeviceInfo::s_DXDeviceInfo->DxVersion == FDXDeviceInfo::ED3DVersion::Version12Win10)
+	{
+		// Gather some info from the DX11 texture in the buffer...
+		TRefCountPtr<IMFDXGIBuffer> DXGIBuffer;
+		CHECK_HR(Buffer->QueryInterface(__uuidof(IMFDXGIBuffer), (void**)DXGIBuffer.GetInitReference()));
+
+		TRefCountPtr<ID3D11Texture2D> Texture2D;
+		CHECK_HR(DXGIBuffer->GetResource(IID_PPV_ARGS(Texture2D.GetInitReference())));
+		D3D11_TEXTURE2D_DESC TextureDesc;
+		Texture2D->GetDesc(&TextureDesc);
+		if (TextureDesc.Format != DXGI_FORMAT_NV12)
+		{
+			LogMessage(IInfoLog::ELevel::Error, "FVideoDecoderH265::CopyTexture(): Decoded texture is not in NV12 format");
+			return false;
+		}
+
+		// Retrieve frame data and store it in Buffer for rendering later
+#if ELECTRA_USE_IMF2DBUFFER
+		TRefCountPtr<IMF2DBuffer> Buffer2D;
+		CHECK_HR(Buffer->QueryInterface(__uuidof(IMF2DBuffer), (void**)Buffer2D.GetInitReference()));
+
+		uint8* Data = nullptr;
+		LONG Pitch;
+		CHECK_HR(Buffer2D->Lock2D(&Data, &Pitch));
+
+		DecoderOutput->InitializeWithBuffer(Data, Pitch * (TextureDesc.Height * 3 / 2),
+			Pitch,															// Buffer stride
+			FIntPoint(TextureDesc.Width, TextureDesc.Height * 3 / 2),		// Buffer dimensions
+			ParamDict);
+
+		CHECK_HR(Buffer2D->Unlock2D());
+#else
+		DWORD BufferSize = 0;
+		CHECK_HR(Buffer->GetCurrentLength(&BufferSize));
+		// Not really trusting the size here. Seen crash reports where the returned buffer size is larger than what is contained,
+		// as if the texture was wider (a multiple of 256 actually). The crash then happens on the memcpy() trying to _read_ more
+		// data than was actually mapped to the buffer (access violation on exactly the start of the next page (4096 byte aligned)).
+		// https://docs.microsoft.com/en-us/windows/win32/api/dxgiformat/ne-dxgiformat-dxgi_format stipulates the format to be aligned
+		// to 'rowPitch' (whatever that is) and 'SysMemPitch' (https://docs.microsoft.com/en-us/windows/win32/api/d3d11/ns-d3d11-d3d11_subresource_data)
+		// which as far as I can tell is a property set when the texture was created. Which we have NO control over since the video decoder creates this
+		// texture for us. There does not seem to be any API to query this pitch at all.
+		// NV12 format as described here https://docs.microsoft.com/en-us/windows/win32/medfound/recommended-8-bit-yuv-formats-for-video-rendering
+		// does not suggest there to be any padding/alignment/pitch.
+		// What makes the crashes so weird is that they do not happen all the time right on the first frame. It takes a bit before wrong values are returned.
+		// For instance, for a 1152*656 texture 1259520 (1280*656) is returned when really it should be only 1133568,
+		// or for a 864*480 texture instead of 622080 bytes we are told there'd be 737280 bytes (1024*480).
+		// Since we are not using any "actual pitch" in the Initialize() call but the actual width of the texture we must only copy as much data anyway.
+		// Even if there was padding in the texture that we would be copying over into our buffer it would not be skipped on rendering later on account
+		// of us specifying the width for pitch.
+		// It seems prudent at this point not to trust the returned buffer size but to calculate it from the actual dimensions ourselves.
+		// NOTE: this value is probably dependent on the GPU and driver implementing the decoder and may therefore not be a general issue!
+		BufferSize = TextureDesc.Width * (TextureDesc.Height * 3 / 2);
+
+		uint8* Data = nullptr;
+		CHECK_HR(Buffer->Lock(&Data, NULL, NULL));
+
+		DecoderOutput->InitializeWithBuffer(Data, BufferSize,
+			TextureDesc.Width,												// Buffer stride
+			FIntPoint(TextureDesc.Width, TextureDesc.Height * 3 / 2),		// Buffer dimensions
+			ParamDict);
+
+		CHECK_HR(Buffer->Unlock());
+#endif
+
+		return true;
+	}
+	else
+	{
+		LogMessage(IInfoLog::ELevel::Error, "FVideoDecoderH265::CopyTexture(): Unhandled D3D version");
+		return false;
+	}
+	return true;
+#else
+	return false;
+#endif
+}
+
+} // namespace Electra
+
+#endif
diff --git a/Engine/Plugins/Media/ElectraPlayer/Source/ElectraPlayerRuntime/Private/Runtime/Demuxer/ParserISO14496-12.cpp b/Engine/Plugins/Media/ElectraPlayer/Source/ElectraPlayerRuntime/Private/Runtime/Demuxer/ParserISO14496-12.cpp
--- a/Engine/Plugins/Media/ElectraPlayer/Source/ElectraPlayerRuntime/Private/Runtime/Demuxer/ParserISO14496-12.cpp
+++ b/Engine/Plugins/Media/ElectraPlayer/Source/ElectraPlayerRuntime/Private/Runtime/Demuxer/ParserISO14496-12.cpp
@@ -598,6 +598,7 @@ private:
 		static const IParserISO14496_12::FBoxType kBox_btrt = MAKE_BOX_ATOM('b', 't', 'r', 't');
 		static const IParserISO14496_12::FBoxType kBox_pasp = MAKE_BOX_ATOM('p', 'a', 's', 'p');
 		static const IParserISO14496_12::FBoxType kBox_avcC = MAKE_BOX_ATOM('a', 'v', 'c', 'C');
+		static const IParserISO14496_12::FBoxType kBox_hvcC = MAKE_BOX_ATOM('h', 'v', 'c', 'C');
 		static const IParserISO14496_12::FBoxType kBox_esds = MAKE_BOX_ATOM('e', 's', 'd', 's');
 		static const IParserISO14496_12::FBoxType kBox_dac3 = MAKE_BOX_ATOM('d', 'a', 'c', '3');
 		static const IParserISO14496_12::FBoxType kBox_dec3 = MAKE_BOX_ATOM('d', 'e', 'c', '3');
@@ -658,6 +659,9 @@ private:
 		 * Sample types
 		 */
 		static const IParserISO14496_12::FBoxType kSample_avc1 = MAKE_BOX_ATOM('a', 'v', 'c', '1');
+		static const IParserISO14496_12::FBoxType kSample_avc3 = MAKE_BOX_ATOM('a', 'v', 'c', '3');
+		static const IParserISO14496_12::FBoxType kSample_hvc1 = MAKE_BOX_ATOM('h', 'v', 'c', '1');
+		static const IParserISO14496_12::FBoxType kSample_hev1 = MAKE_BOX_ATOM('h', 'e', 'v', '1');
 		static const IParserISO14496_12::FBoxType kSample_mp4a = MAKE_BOX_ATOM('m', 'p', '4', 'a');
 		static const IParserISO14496_12::FBoxType kSample_stpp = MAKE_BOX_ATOM('s', 't', 'p', 'p');
 		static const IParserISO14496_12::FBoxType kSample_sbtt = MAKE_BOX_ATOM('s', 'b', 't', 't');
@@ -1979,6 +1983,55 @@ private:
 		MPEG::FAVCDecoderConfigurationRecord		ConfigurationRecord;
 	};
 
+	/**
+	 * HEVC Decoder Configuration box (ISO/IEC 14496-15:2014 - 8.3.3.1.2)
+	 */
+	class FMP4BoxHVCC : public FMP4BoxBasic
+	{
+	public:
+		FMP4BoxHVCC(IParserISO14496_12::FBoxType InBoxType, int64 InBoxSize, int64 InStartOffset, int64 InDataOffset, bool bInIsLeafBox)
+			: FMP4BoxBasic(InBoxType, InBoxSize, InStartOffset, InDataOffset, bInIsLeafBox)
+		{
+		}
+
+		virtual ~FMP4BoxHVCC()
+		{
+		}
+
+		const MPEG::FHEVCDecoderConfigurationRecord& GetDecoderConfigurationRecord() const
+		{
+			return ConfigurationRecord;
+		}
+
+	private:
+		FMP4BoxHVCC() = delete;
+		FMP4BoxHVCC(const FMP4BoxHVCC&) = delete;
+
+	protected:
+		virtual UEMediaError ReadAndParseAttributes(FMP4ParseInfo* ParseInfo) override
+		{
+			UEMediaError Error = UEMEDIA_ERROR_OK;
+			uint32 BytesRemaining = BoxSize - (ParseInfo->Reader()->GetCurrentReadOffset() - StartOffset);
+			if (BytesRemaining)
+			{
+				void* TempBuffer = FMemory::Malloc(BytesRemaining);
+				if (!TempBuffer)
+				{
+					return UEMEDIA_ERROR_OOM;
+				}
+				Error = ParseInfo->Reader()->ReadBytes(TempBuffer, BytesRemaining);
+				if (Error == UEMEDIA_ERROR_OK)
+				{
+					ConfigurationRecord.SetRawData(TempBuffer, BytesRemaining);
+				}
+				FMemory::Free(TempBuffer);
+			}
+			return Error;
+		}
+
+		MPEG::FHEVCDecoderConfigurationRecord ConfigurationRecord;
+	};
+
 
 	/**
 	 * MPEG-4 Sample Description Box containing an ISO/IEC 14496-1 8.6.5 ES_Descriptor (ISO/IEC 14496-14:2013 - 5.6)
@@ -4364,6 +4417,9 @@ private:
 						case FMP4Box::kBox_avcC:
 							NextBox = new FMP4BoxAVCC(BoxType, BoxSize, BoxStartOffset, BoxDataOffset, true);
 							break;
+						case FMP4Box::kBox_hvcC:
+							NextBox = new FMP4BoxHVCC(BoxType, BoxSize, BoxStartOffset, BoxDataOffset, true);
+							break;
 						case FMP4Box::kBox_btrt:
 							NextBox = new FMP4BoxBTRT(BoxType, BoxSize, BoxStartOffset, BoxDataOffset, true);
 							break;
@@ -4811,6 +4867,7 @@ private:
 			int64									CompositionTimeEditOffset = 0;
 
 			MPEG::FAVCDecoderConfigurationRecord	CodecSpecificDataAVC;
+			MPEG::FHEVCDecoderConfigurationRecord	CodecSpecificDataHEVC;
 			MPEG::FESDescriptor						CodecSpecificDataMP4A;
 			FBitrateInfo							BitrateInfo;
 		};
@@ -4882,6 +4939,7 @@ private:
 
 
 		UEMediaError ParseAVC1SampleType(FTrack* Track, const FMP4Box* SampleBox);
+		UEMediaError ParseHVC1SampleType(FTrack* Track, const FMP4Box* SampleBox);
 		UEMediaError ParseMP4ASampleType(FTrack* Track, const FMP4Box* SampleBox);
 
 		FMP4ParseInfo* ParsedData;
@@ -4973,6 +5031,10 @@ private:
 			{
 				return CodecSpecificDataAVC.GetCodecSpecificData();
 			}
+			case FStreamCodecInformation::ECodec::H265:
+			{
+				return CodecSpecificDataHEVC.GetCodecSpecificData();
+			}
 			case FStreamCodecInformation::ECodec::AAC:
 			{
 				return CodecSpecificDataMP4A.GetCodecSpecificData();
@@ -4994,6 +5056,10 @@ private:
 			{
 				return CodecSpecificDataAVC.GetRawData();
 			}
+			case FStreamCodecInformation::ECodec::H265:
+			{
+				return CodecSpecificDataHEVC.GetRawData();
+			}
 			case FStreamCodecInformation::ECodec::AAC:
 			{
 				return CodecSpecificDataMP4A.GetRawData();
@@ -5816,6 +5882,85 @@ private:
 		return UEMEDIA_ERROR_FORMAT_ERROR;
 	}
 
+	UEMediaError FParserISO14496_12::ParseHVC1SampleType(FTrack* Track, const FMP4Box* SampleBox)
+	{
+		const FMP4BoxVisualSampleEntry* VisualSampleEntry = static_cast<const FMP4BoxVisualSampleEntry*>(SampleBox);
+		check(VisualSampleEntry->GetNumberOfChildren() > 0);
+		if (VisualSampleEntry->GetNumberOfChildren() > 0)
+		{
+			// There may be several entries. Usually there is an 'hvcC' (required) and optional boxes like 'pasp', 'btrt', 'clap', etc.
+			bool bGotVideoFormat = false;
+			for(int32 j = 0, jMax = VisualSampleEntry->GetNumberOfChildren(); j < jMax; ++j)
+			{
+				const FMP4Box* SampleEntry = VisualSampleEntry->GetChildBox(j);
+				switch(SampleEntry->GetType())
+				{
+					case FMP4Box::kBox_hvcC:
+					{
+						const FMP4BoxHVCC* HVCCBox = static_cast<const FMP4BoxHVCC*>(SampleEntry);
+						Track->CodecSpecificDataHEVC = HVCCBox->GetDecoderConfigurationRecord();
+						bool bOk = Track->CodecSpecificDataHEVC.Parse();
+						check(bOk);
+						if (bOk)
+						{
+							Track->CodecInformation.SetStreamType(EStreamType::Video);
+							Track->CodecInformation.SetCodec(FStreamCodecInformation::ECodec::H265);
+							Track->CodecInformation.SetStreamLanguageCode(Track->GetLanguage());
+							if (Track->CodecSpecificDataHEVC.GetNumberOfSPS() == 0)
+							{
+								return UEMEDIA_ERROR_FORMAT_ERROR;
+							}
+							const MPEG::FISO23008_2_seq_parameter_set_data& sps = Track->CodecSpecificDataHEVC.GetParsedSPS(0);
+							int32 CropL, CropR, CropT, CropB;
+							sps.GetCrop(CropL, CropR, CropT, CropB);
+							Track->CodecInformation.SetResolution(FStreamCodecInformation::FResolution(sps.GetWidth() - CropL - CropR, sps.GetHeight() - CropT - CropB));
+							Track->CodecInformation.SetCrop(FStreamCodecInformation::FCrop(CropL, CropT, CropR, CropB));
+							FStreamCodecInformation::FAspectRatio ar;
+							sps.GetAspect(ar.Width, ar.Height);
+							Track->CodecInformation.SetAspectRatio(ar);
+							Track->CodecInformation.SetFrameRate(sps.GetTiming());
+							Track->CodecInformation.SetProfileSpace(sps.general_profile_space);
+							Track->CodecInformation.SetProfileTier(sps.general_tier_flag);
+							Track->CodecInformation.SetProfile(sps.general_profile_idc);
+							Track->CodecInformation.SetProfileLevel(sps.general_level_idc);
+							Track->CodecInformation.SetProfileConstraints(sps.GetConstraintFlags());
+							Track->CodecInformation.SetProfileCompatibilityFlags(sps.general_profile_compatibility_flag);
+							Track->CodecInformation.SetCodecSpecifierRFC6381(sps.GetRFC6381(TEXT("hvc1")));
+							bGotVideoFormat = true;
+						}
+						else
+						{
+							return UEMEDIA_ERROR_FORMAT_ERROR;
+						}
+						break;
+					}
+					case FMP4Box::kBox_btrt:
+					{
+						const FMP4BoxBTRT* BTRTBox = static_cast<const FMP4BoxBTRT*>(SampleEntry);
+						Track->BitrateInfo.BufferSizeDB = BTRTBox->GetBufferSizeDB();
+						Track->BitrateInfo.MaxBitrate = BTRTBox->GetMaxBitrate();
+						Track->BitrateInfo.AvgBitrate = BTRTBox->GetAverageBitrate();
+						break;
+					}
+					case FMP4Box::kBox_pasp:
+					{
+						break;
+					}
+					default:
+					{
+						break;
+					}
+				}
+			}
+			// Check that we got the required information
+			if (bGotVideoFormat)
+			{
+				return UEMEDIA_ERROR_OK;
+			}
+		}
+		return UEMEDIA_ERROR_FORMAT_ERROR;
+	}
+
 	UEMediaError FParserISO14496_12::ParseMP4ASampleType(FTrack* Track, const FMP4Box* SampleBox)
 	{
 		const FMP4BoxAudioSampleEntry* AudioSampleEntry = static_cast<const FMP4BoxAudioSampleEntry*>(Track->STSDBox->GetChildBox(0));
@@ -6111,18 +6256,25 @@ private:
 								{
 									return UEMEDIA_ERROR_FORMAT_ERROR;
 								}
-								// Expecting avc1 for video right now.
-								if (FRMABox->GetDataFormat() != FMP4Box::kSample_avc1)
+								// Expecting encryption to be 'cenc'
+								if (SCHMBox->GetSchemeType() != FMP4Box::kEncryptionScheme_cenc)
 								{
 									return UEMEDIA_ERROR_FORMAT_ERROR;
 								}
-								// Likewise fo encryption it needs to be 'cenc'
-								if (SCHMBox->GetSchemeType() != FMP4Box::kEncryptionScheme_cenc)
+								// Expecting avc1/avc3/hvc1/hev1 for video right now.
+								UEMediaError Error = UEMEDIA_ERROR_FORMAT_ERROR;
+								if (FRMABox->GetDataFormat() == FMP4Box::kSample_avc1 || FRMABox->GetDataFormat() == FMP4Box::kSample_avc3)
+								{
+									Error = ParseAVC1SampleType(Track.Get(), STSDFirstChildBox);
+								}
+								else if ( FRMABox->GetDataFormat() == FMP4Box::kSample_hvc1 || FRMABox->GetDataFormat() == FMP4Box::kSample_hev1)
+								{
+									Error = ParseHVC1SampleType(Track.Get(), STSDFirstChildBox);
+								}
+								else
 								{
 									return UEMEDIA_ERROR_FORMAT_ERROR;
 								}
-								// Parse the sample
-								UEMediaError Error = ParseAVC1SampleType(Track.Get(), STSDFirstChildBox);
 								/*if (Error == UEMEDIA_ERROR_NOT_SUPPORTED)
 								{
 									bIsSupported = false;
@@ -6170,6 +6322,7 @@ private:
 								break;
 							}
 							case FMP4Box::kSample_avc1:
+							case FMP4Box::kSample_avc3:
 							{
 								UEMediaError Error = ParseAVC1SampleType(Track.Get(), STSDFirstChildBox);
 								/*
@@ -6183,6 +6336,21 @@ private:
 								}
 								break;
 							}
+							case FMP4Box::kSample_hvc1:
+							case FMP4Box::kSample_hev1:
+							{
+								UEMediaError Error = ParseHVC1SampleType(Track.Get(), STSDFirstChildBox);
+								/*
+								if (Error == UEMEDIA_ERROR_NOT_SUPPORTED)
+								{
+									bIsSupported = false;
+								}
+								else*/ if (Error != UEMEDIA_ERROR_OK)
+								{
+									return Error;
+								}
+								break;
+							}
 							case FMP4Box::kSample_mp4a:
 							{
 								UEMediaError Error = ParseMP4ASampleType(Track.Get(), STSDFirstChildBox);
diff --git a/Engine/Plugins/Media/ElectraPlayer/Source/ElectraPlayerRuntime/Private/Runtime/Player/AdaptiveStreamingPlayer.cpp b/Engine/Plugins/Media/ElectraPlayer/Source/ElectraPlayerRuntime/Private/Runtime/Player/AdaptiveStreamingPlayer.cpp
--- a/Engine/Plugins/Media/ElectraPlayer/Source/ElectraPlayerRuntime/Private/Runtime/Player/AdaptiveStreamingPlayer.cpp
+++ b/Engine/Plugins/Media/ElectraPlayer/Source/ElectraPlayerRuntime/Private/Runtime/Player/AdaptiveStreamingPlayer.cpp
@@ -819,103 +819,6 @@ void FAdaptiveStreamingPlayer::DestroyRenderers()
 
 
 
-//-----------------------------------------------------------------------------
-/**
- * Create the necessary AV decoders.
- *
- * @return
- */
-int32 FAdaptiveStreamingPlayer::CreateInitialDecoder(EStreamType type)
-{
-	if (type == EStreamType::Video)
-	{
-		if (VideoDecoder.Decoder == nullptr)
-		{
-			// Get the largest stream resolution of the currently selected video adaptation set.
-			// This is only an initial selection as there could be other adaptation sets in upcoming periods
-			// that have a larger resolution that is still within the allowed limits.
-			FStreamCodecInformation HighestStream;
-			FindMatchingStreamInfo(HighestStream, 0, 0);
-
-			// Create H.264 video decoder
-			VideoDecoder.Decoder = IVideoDecoderH264::Create();
-			VideoDecoder.Decoder->SetPlayerSessionServices(this);
-			VideoDecoder.Parent = this;
-			IVideoDecoderH264::FInstanceConfiguration h264Cfg = PlayerConfig.DecoderCfg264;
-			h264Cfg.ProfileIdc = HighestStream.GetProfile();
-			h264Cfg.LevelIdc = HighestStream.GetProfileLevel();
-			h264Cfg.MaxFrameWidth = HighestStream.GetResolution().Width;
-			h264Cfg.MaxFrameHeight = HighestStream.GetResolution().Height;
-			h264Cfg.AdditionalOptions = HighestStream.GetExtras();
-
-			// Add in any player options that are for decoder use
-			TArray<FString> DecoderOptionKeys;
-			PlayerOptions.GetKeysStartingWith("videoDecoder", DecoderOptionKeys);
-			for (const FString & Key : DecoderOptionKeys)
-			{
-				h264Cfg.AdditionalOptions.Set(Key, PlayerOptions.GetValue(Key));
-			}
-
-			// Attach video decoder buffer monitor.
-			VideoDecoder.Decoder->SetAUInputBufferListener(&VideoDecoder);
-			VideoDecoder.Decoder->SetReadyBufferListener(&VideoDecoder);
-			// Have the video decoder send its output to the video renderer
-			VideoDecoder.Decoder->SetRenderer(VideoRender.Renderer);
-			// Hand it (may be nullptr) a delegate for platform for resource queries
-			VideoDecoder.Decoder->SetResourceDelegate(VideoDecoderResourceDelegate.Pin());
-			// Open the decoder after having set all listeners.
-			VideoDecoder.Decoder->Open(h264Cfg);
-
-			// Now we get the currently limited stream resolution and let the decoder now what we will be using
-			// at most right now. This allows the decoder to be created with a smaller memory footprint at first.
-			UpdateStreamResolutionLimit();
-		}
-		return 0;
-	}
-	else if (type == EStreamType::Audio)
-	{
-		if (AudioDecoder.Decoder == nullptr)
-		{
-			// Create an AAC audio decoder
-			AudioDecoder.Decoder = IAudioDecoderAAC::Create();
-			AudioDecoder.Decoder->SetPlayerSessionServices(this);
-			AudioDecoder.Parent = this;
-			// Attach buffer monitors.
-			AudioDecoder.Decoder->SetAUInputBufferListener(&AudioDecoder);
-			AudioDecoder.Decoder->SetReadyBufferListener(&AudioDecoder);
-			// Have to audio decoder send its output to the audio renderer
-			AudioDecoder.Decoder->SetRenderer(AudioRender.Renderer);
-			// Open the decoder after having set all listeners.
-			AudioDecoder.Decoder->Open(PlayerConfig.DecoderCfgAAC);
-		}
-		return 0;
-	}
-	return -1;
-}
-
-
-//-----------------------------------------------------------------------------
-/**
- * Destroys the decoders.
- */
-void FAdaptiveStreamingPlayer::DestroyDecoders()
-{
-/*
-NOTE: We do not clear out the renderers from the decoder. On their way down the decoders should still be able
-      to access the renderer without harm and dispatch their last remaining data.
-
-	if (mVideoDecoder.Decoder)
-		mVideoDecoder.Decoder->SetRenderer(nullptr);
-	if (mAudioDecoder.Decoder)
-		mAudioDecoder.Decoder->SetRenderer(nullptr);
-*/
-	AudioDecoder.Close();
-	VideoDecoder.Close();
-}
-
-
-
-
 
 
 void FAdaptiveStreamingPlayer::PostLog(Facility::EFacility FromFacility, IInfoLog::ELevel inLogLevel, const FString &Message)
@@ -1019,31 +922,58 @@ bool FAdaptiveStreamingPlayer::CanDecodeStream(const FStreamCodecInformation& In
 {
 	if (InStreamCodecInfo.IsVideoCodec())
 	{
-		const AdaptiveStreamingPlayerConfig::FConfiguration::FVideoDecoderLimit* DecoderLimit = &PlayerConfig.H264LimitUpto30fps;
 		double Rate = InStreamCodecInfo.GetFrameRate().IsValid() ? InStreamCodecInfo.GetFrameRate().GetAsDouble() : 30.0;
-		if (Rate > 31.0)
-		{
-			DecoderLimit = &PlayerConfig.H264LimitAbove30fps;
-		}
-		// Check against user configured resolution limit
-		if (DecoderLimit->MaxResolution.Height && InStreamCodecInfo.GetResolution().Height > DecoderLimit->MaxResolution.Height)
-		{
-			return false;
-		}
-
-		// Check against video decoder capabilities.
-		IVideoDecoderH264::FStreamDecodeCapability StreamParam, Capability;
-		StreamParam.Width = InStreamCodecInfo.GetResolution().Width;
-		StreamParam.Height = InStreamCodecInfo.GetResolution().Height;
-		StreamParam.Profile = InStreamCodecInfo.GetProfile();
-		StreamParam.Level = InStreamCodecInfo.GetProfileLevel();
-		StreamParam.FPS = Rate;
-		if (IVideoDecoderH264::GetStreamDecodeCapability(Capability, StreamParam))
+		if (InStreamCodecInfo.GetCodec() == FStreamCodecInformation::ECodec::H264)
 		{
-			if (Capability.DecoderSupportType == IVideoDecoderH264::FStreamDecodeCapability::ESupported::NotSupported)
+			const AdaptiveStreamingPlayerConfig::FConfiguration::FVideoDecoderLimit* DecoderLimit = &PlayerConfig.H264LimitUpto30fps;
+			if (Rate > 31.0)
+			{
+				DecoderLimit = &PlayerConfig.H264LimitAbove30fps;
+			}
+			// Check against user configured resolution limit
+			if (DecoderLimit->MaxResolution.Height && InStreamCodecInfo.GetResolution().Height > DecoderLimit->MaxResolution.Height)
 			{
 				return false;
 			}
+
+			// Check against video decoder capabilities.
+			IVideoDecoderH264::FStreamDecodeCapability StreamParam, Capability;
+			StreamParam.Width = InStreamCodecInfo.GetResolution().Width;
+			StreamParam.Height = InStreamCodecInfo.GetResolution().Height;
+			StreamParam.Profile = InStreamCodecInfo.GetProfile();
+			StreamParam.Level = InStreamCodecInfo.GetProfileLevel();
+			StreamParam.FPS = Rate;
+			if (IVideoDecoderH264::GetStreamDecodeCapability(Capability, StreamParam))
+			{
+				if (Capability.DecoderSupportType == IVideoDecoderH264::FStreamDecodeCapability::ESupported::NotSupported)
+				{
+					return false;
+				}
+			}
+		}
+		else if (InStreamCodecInfo.GetCodec() == FStreamCodecInformation::ECodec::H265)
+		{
+			#if ELECTRA_PLATFORM_HAS_H265_DECODER
+				// Check against video decoder capabilities.
+				IVideoDecoderH265::FStreamDecodeCapability StreamParam, Capability;
+				StreamParam.Width = InStreamCodecInfo.GetResolution().Width;
+				StreamParam.Height = InStreamCodecInfo.GetResolution().Height;
+				StreamParam.Tier = InStreamCodecInfo.GetProfileTier();
+				StreamParam.CompatibilityFlags = InStreamCodecInfo.GetProfileCompatibilityFlags();
+				StreamParam.Profile = InStreamCodecInfo.GetProfile();
+				StreamParam.Level = InStreamCodecInfo.GetProfileLevel();
+				StreamParam.ConstraintFlags = InStreamCodecInfo.GetProfileConstraints();
+				StreamParam.FPS = Rate;
+				if (IVideoDecoderH265::GetStreamDecodeCapability(Capability, StreamParam))
+				{
+					if (Capability.DecoderSupportType == IVideoDecoderH265::FStreamDecodeCapability::ESupported::NotSupported)
+					{
+						return false;
+					}
+				}
+			#else
+				return false;
+			#endif
 		}
 	}
 
@@ -1581,48 +1511,61 @@ void FAdaptiveStreamingPlayer::HandleMetadataChanges()
 		PlaybackState.SetDuration(Manifest->GetDuration());
 	}
 
+	FTimeValue Current = GetCurrentPlayTime();
+
 	// See if we may be reaching the next period.
-	if (UpcomingPeriods.Num())
+	if (UpcomingPeriods.Num() && Current.IsValid())
 	{
-		FTimeValue Current = GetCurrentPlayTime();
-		if (Current.IsValid())
+		for(int32 i=0; i<UpcomingPeriods.Num(); )
 		{
-			for(int32 i=0; i<UpcomingPeriods.Num(); )
+			if (Current < UpcomingPeriods[i].TimeRange.Start)
 			{
-				if (Current < UpcomingPeriods[i].TimeRange.Start)
-				{
-					break;
-				}
-				else
-				{
-					// Did we already pass the end of the period?
-					// This could be happening with really short periods and in this case we do not want to
-					// handle those. Since we passed them already there is no point.
-					if (!UpcomingPeriods[i].TimeRange.End.IsValid() || Current < UpcomingPeriods[i].TimeRange.End)
+				break;
+			}
+			else
+			{
+				// Did we already pass the end of the period?
+				// This could be happening with really short periods and in this case we do not want to
+				// handle those. Since we passed them already there is no point.
+				if (!UpcomingPeriods[i].TimeRange.End.IsValid() || Current < UpcomingPeriods[i].TimeRange.End)
+				{
+					// Handle it.
+					TArray<FTrackMetadata> MetadataVideo;
+					TArray<FTrackMetadata> MetadataAudio;
+					//TArray<FTrackMetadata> MetadataSubtitle;
+					UpcomingPeriods[i].Period->GetMetaData(MetadataVideo, EStreamType::Video);
+					UpcomingPeriods[i].Period->GetMetaData(MetadataAudio, EStreamType::Audio);
+					//UpcomingPeriods[i].Period->GetMetaData(MetadataSubtitle, EStreamType::Subtitle);
+
+					bool bMetadataChanged = PlaybackState.SetTrackMetadata(MetadataVideo, MetadataAudio);
+					PlaybackState.SetHaveMetadata(true);
+					if (bMetadataChanged)
 					{
-						// Handle it.
-						TArray<FTrackMetadata> MetadataVideo;
-						TArray<FTrackMetadata> MetadataAudio;
-						//TArray<FTrackMetadata> MetadataSubtitle;
-						UpcomingPeriods[i].Period->GetMetaData(MetadataVideo, EStreamType::Video);
-						UpcomingPeriods[i].Period->GetMetaData(MetadataAudio, EStreamType::Audio);
-						//UpcomingPeriods[i].Period->GetMetaData(MetadataSubtitle, EStreamType::Subtitle);
-
-						bool bMetadataChanged = PlaybackState.SetTrackMetadata(MetadataVideo, MetadataAudio);
-						PlaybackState.SetHaveMetadata(true);
-						if (bMetadataChanged)
-						{
-							DispatchEvent(FMetricEvent::ReportTracksChanged());
-						}
+						DispatchEvent(FMetricEvent::ReportTracksChanged());
 					}
+				}
 
-					// Inform AEMS handler of the period transition.
-					AEMSEventHandler->PlaybackPeriodTransition(UpcomingPeriods[i].ID, UpcomingPeriods[i].TimeRange);
+				// Inform AEMS handler of the period transition.
+				AEMSEventHandler->PlaybackPeriodTransition(UpcomingPeriods[i].ID, UpcomingPeriods[i].TimeRange);
 
-					// Remove it and continue with the next.
-					UpcomingPeriods.RemoveAt(i);
-					++i;
-				}
+				// Remove it and continue with the next.
+				UpcomingPeriods.RemoveAt(i);
+				++i;
+			}
+		}
+	}
+
+	// Dump periods that are no longer active because we passed their end with some threshold.
+	if (Current.IsValid())
+	{
+		FTimeValue PastTime = Current - FTimeValue().SetFromMilliseconds(1000);
+		for(int32 i=0; i<ActivePeriods.Num(); ++i)
+		{
+			if ((ActivePeriods[i].TimeRange.End.IsValid() && PastTime > ActivePeriods[i].TimeRange.End) ||
+				(!ActivePeriods[i].TimeRange.End.IsValid() && (i+1) < ActivePeriods.Num() && PastTime >= ActivePeriods[i+1].TimeRange.Start))
+			{
+				ActivePeriods.RemoveAt(i);
+				--i;
 			}
 		}
 	}
@@ -1710,6 +1653,15 @@ void FAdaptiveStreamingPlayer::HandleSessionMessage(TSharedPtrTS<IPlayerMessage>
 		stats.FailureReason  = pMsg->GetResult().IsError() ? pMsg->GetResult().GetPrintable() : FString(); //pMsg->GetConnectionInfo().StatusInfo.ErrorDetail.GetMessage();
 		DispatchEvent(FMetricEvent::ReportLicenseKey(stats));
 	}
+	// Decoder message?
+	else if (SessionMessage->GetType() == FDecoderMessage::Type())
+	{
+		FDecoderMessage* pMsg = static_cast<FDecoderMessage*>(SessionMessage.Get());
+		if (pMsg->GetStreamType() == EStreamType::Video)
+		{
+			VideoDecoder.bDrainingForCodecChangeDone = true;
+		}
+	}
 	else
 	{
 		checkNoEntry();
@@ -1836,31 +1788,11 @@ void FAdaptiveStreamingPlayer::HandleNewBufferedData()
 
 				// Send pre-roll begin event.
 				DispatchEvent(FMetricEvent::ReportPrerollStart());
-
-				// Create the first decoders here.
-				if (bHaveVideoReader.GetWithDefault(false))
-				{
-					CreateInitialDecoder(EStreamType::Video);
-				}
-				if (bHaveAudioReader.GetWithDefault(false))
-				{
-					CreateInitialDecoder(EStreamType::Audio);
-				}
 			}
 		}
 	}
 }
 
-//-----------------------------------------------------------------------------
-/**
- * Check if the decoders need to change.
- */
-void FAdaptiveStreamingPlayer::HandleDecoderChanges()
-{
-	// Nothing to do for now.
-}
-
-
 //-----------------------------------------------------------------------------
 /**
  * Returns the current internal playback time.
@@ -2915,9 +2847,9 @@ void FAdaptiveStreamingPlayer::AddUpcomingPeriod(TSharedPtrTS<IManifest::IPlayPe
 	{
 		FString PeriodID = Period->GetUniqueIdentifier();
 		// Only add if it is not already in the list.
-		if (!UpcomingPeriods.ContainsByPredicate([PeriodID](const FUpcomingPeriod& e){ return e.ID.Equals(PeriodID); }))
+		if (!UpcomingPeriods.ContainsByPredicate([PeriodID](const FPeriodInformation& e){ return e.ID.Equals(PeriodID); }))
 		{
-			FUpcomingPeriod Next;
+			FPeriodInformation Next;
 			Next.ID = MoveTemp(PeriodID);
 			Next.TimeRange = Period->GetTimeRange();
 			Next.Period = MoveTemp(Period);
@@ -2927,6 +2859,7 @@ void FAdaptiveStreamingPlayer::AddUpcomingPeriod(TSharedPtrTS<IManifest::IPlayPe
 			// Tell the AEMS handler that a new period will be coming up. It needs this information to cut overlapping events.
 			AEMSEventHandler->NewUpcomingPeriod(Next.ID, Next.TimeRange);
 
+			ActivePeriods.Add(Next);
 			UpcomingPeriods.Emplace(MoveTemp(Next));
 		}
 	}
@@ -2988,292 +2921,14 @@ void FAdaptiveStreamingPlayer::UpdateDataAvailabilityState(Metrics::FDataAvailab
 }
 
 
-
-//-----------------------------------------------------------------------------
-/**
- * Sends an available AU to a decoder.
- * If the current buffer level is below the underrun threshold an underrun
- * message is sent to the worker thread.
- *
- * @param Type
- * @param FromMultistreamBuffer
- * @param Decoder
- */
-void FAdaptiveStreamingPlayer::FeedDecoder(EStreamType Type, FMultiTrackAccessUnitBuffer& FromMultistreamBuffer, IAccessUnitBufferInterface* Decoder)
-{
-	// Lock the AU buffer for the duration of this function to ensure this can never clash with a Flush() call
-	// since we are checking size, eod state and subsequently popping an AU, for which the buffer must stay consistent inbetween!
-	// Also to ensure the active buffer doesn't get changed from one track to another.
-	FMultiTrackAccessUnitBuffer::FScopedLock lock(FromMultistreamBuffer);
-
-	FBufferStats* pStats = nullptr;
-	Metrics::FDataAvailabilityChange* pAvailability = nullptr;
-	switch(Type)
-	{
-		case EStreamType::Video:
-			pStats = &VideoBufferStats;
-			pAvailability = &DataAvailabilityStateVid;
-			break;
-		case EStreamType::Audio:
-			pStats = &AudioBufferStats;
-			pAvailability = &DataAvailabilityStateAud;
-			break;
-		case EStreamType::Subtitle:
-			pStats = &TextBufferStats;
-			pAvailability = &DataAvailabilityStateTxt;
-			break;
-		default:
-			checkNoEntry();
-			return;
-	}
-
-	// Is the buffer (the Type of elementary stream actually) active/selected?
-	if (!FromMultistreamBuffer.IsDeselected())
-	{
-		// Check for buffer underrun.
-		if (!bRebufferPending && CurrentState == EPlayerState::eState_Playing && StreamState == EStreamState::eStream_Running && PipelineState == EPipelineState::ePipeline_Running)
-		{
-			bool bEODSet = FromMultistreamBuffer.IsEODFlagSet();
-			if (!bEODSet && FromMultistreamBuffer.Num() == 0)
-			{
-				// Buffer underrun.
-				bRebufferPending = true;
-				FTimeValue LastKnownPTS = FromMultistreamBuffer.GetLastPoppedPTS();
-				// Only set the 'rebuffer at' time if we have a valid last known PTS. If we don't
-				// then maybe this is a cascade failure from a previous rebuffer attempt for which
-				// we then try that time once more.
-				if (LastKnownPTS.IsValid())
-				{
-					RebufferDetectedAtPlayPos = LastKnownPTS;
-				}
-				WorkerThread.SendMessage(FWorkerThread::FMessage::EType::BufferUnderrun);
-			}
-		}
-
-		FAccessUnit* AccessUnit = nullptr;
-		bool bOk;
-		bOk = FromMultistreamBuffer.Pop(AccessUnit);
-		if (AccessUnit)
-		{
-			if (Decoder)
-			{
-			// The decoder has asked to be fed a new AU so it better be able to accept it.
-				/*IAccessUnitBufferInterface::EAUpushResult auPushRes =*/ Decoder->AUdataPushAU(AccessUnit);
-			}
-
-			// The decoder will have added a ref count if it took the AU. If it didnt' for whatever reason
-			// we still release it to get rid of it and not cause a memory leak.
-			FAccessUnit::Release(AccessUnit);
-			AccessUnit = nullptr;
-
-			if (pAvailability)
-			{
-				UpdateDataAvailabilityState(*pAvailability, Metrics::FDataAvailabilityChange::EAvailability::DataAvailable);
-			}
-		}
-	}
-	// An AU is not tagged as being "the last" one. Instead the EOD is handled separately and must be dealt with
-	// by the decoders accordingly.
-	if (FromMultistreamBuffer.IsEODFlagSet() && FromMultistreamBuffer.Num() == 0)
-	{
-		if (pStats && !pStats->DecoderInputBuffer.bEODSignaled)
-		{
-			if (Decoder)
-			{
-				Decoder->AUdataPushEOD();
-			}
-		}
-		if (pAvailability)
-		{
-			UpdateDataAvailabilityState(*pAvailability, Metrics::FDataAvailabilityChange::EAvailability::DataNotAvailable);
-		}
-	}
-}
-
-void FAdaptiveStreamingPlayer::VideoDecoderInputNeeded(const IAccessUnitBufferListener::FBufferStats &currentInputBufferStats)
-{
-	DiagnosticsCriticalSection.Lock();
-	VideoBufferStats.DecoderInputBuffer = currentInputBufferStats;
-	DiagnosticsCriticalSection.Unlock();
-	EDecoderState decState = DecoderState;
-	if (decState == EDecoderState::eDecoder_Running)
-	{
-		FeedDecoder(EStreamType::Video, MultiStreamBufferVid, VideoDecoder.Decoder);
-	}
-}
-
-void FAdaptiveStreamingPlayer::VideoDecoderOutputReady(const IDecoderOutputBufferListener::FDecodeReadyStats &currentReadyStats)
-{
-	DiagnosticsCriticalSection.Lock();
-	VideoBufferStats.DecoderOutputBuffer = currentReadyStats;
-	DiagnosticsCriticalSection.Unlock();
-}
-
-
-void FAdaptiveStreamingPlayer::AudioDecoderInputNeeded(const IAccessUnitBufferListener::FBufferStats &currentInputBufferStats)
-{
-	DiagnosticsCriticalSection.Lock();
-	AudioBufferStats.DecoderInputBuffer = currentInputBufferStats;
-	DiagnosticsCriticalSection.Unlock();
-	EDecoderState decState = DecoderState;
-	if (decState == EDecoderState::eDecoder_Running)
-	{
-		FeedDecoder(EStreamType::Audio, MultiStreamBufferAud, AudioDecoder.Decoder);
-	}
-}
-
-void FAdaptiveStreamingPlayer::AudioDecoderOutputReady(const IDecoderOutputBufferListener::FDecodeReadyStats &currentReadyStats)
-{
-	DiagnosticsCriticalSection.Lock();
-	AudioBufferStats.DecoderOutputBuffer = currentReadyStats;
-	DiagnosticsCriticalSection.Unlock();
-}
-
-
-
-
-bool FAdaptiveStreamingPlayer::FindMatchingStreamInfo(FStreamCodecInformation& OutStreamInfo, int32 MaxWidth, int32 MaxHeight)
-{
-	if (!CurrentPlayPeriodVideo.IsValid())
-	{
-		return false;
-	}
-	TArray<FStreamCodecInformation> VideoCodecInfos;
-
-	TSharedPtrTS<ITimelineMediaAsset> Asset = CurrentPlayPeriodVideo->GetMediaAsset();
-	if (Asset.IsValid())
-	{
-		if (Asset->GetNumberOfAdaptationSets(EStreamType::Video) > 0)
-		{
-			// What if this is more than one?
-			TSharedPtrTS<IPlaybackAssetAdaptationSet> VideoSet = Asset->GetAdaptationSetByTypeAndIndex(EStreamType::Video, 0);
-			check(VideoSet.IsValid());
-			if (VideoSet.IsValid())
-			{
-				for(int32 i = 0, iMax = VideoSet->GetNumberOfRepresentations(); i < iMax; ++i)
-				{
-					VideoCodecInfos.Push(VideoSet->GetRepresentationByIndex(i)->GetCodecInformation());
-				}
-				check(VideoCodecInfos.Num());
-				if (VideoCodecInfos.Num())
-				{
-					if (MaxWidth == 0 && MaxHeight == 0)
-					{
-						FStreamCodecInformation Best = VideoCodecInfos[0];
-						for(int32 i=1; i<VideoCodecInfos.Num(); ++i)
-						{
-							const FStreamCodecInformation::FResolution& Res = VideoCodecInfos[i].GetResolution();
-							if (Res.Width > Best.GetResolution().Width)
-							{
-								Best.SetResolution(FStreamCodecInformation::FResolution(Res.Width, Best.GetResolution().Height));
-							}
-							if (Res.Height > Best.GetResolution().Height)
-							{
-								Best.SetResolution(FStreamCodecInformation::FResolution(Best.GetResolution().Width, Res.Height));
-							}
-							// Note: the final RFC 6381 codec string will be bogus since we do not re-create it here.
-							if (VideoCodecInfos[i].GetProfile() > Best.GetProfile())
-							{
-								Best.SetProfile(VideoCodecInfos[i].GetProfile());
-							}
-							if (VideoCodecInfos[i].GetProfileLevel() > Best.GetProfileLevel())
-							{
-								Best.SetProfileLevel(VideoCodecInfos[i].GetProfileLevel());
-							}
-							if (VideoCodecInfos[i].GetExtras().GetValue("b_frames").SafeGetInt64(0))
-							{
-								Best.GetExtras().Set("b_frames", FVariantValue((int64) 1));
-							}
-						}
-						OutStreamInfo = Best;
-					}
-					else
-					{
-						if (MaxWidth == 0)
-						{
-							MaxWidth = 32768;
-						}
-						if (MaxHeight == 0)
-						{
-							MaxHeight = 32768;
-						}
-						FStreamCodecInformation		Best;
-						bool bFirst = true;
-						for(int32 i=0; i<VideoCodecInfos.Num(); ++i)
-						{
-							const FStreamCodecInformation::FResolution& Res = VideoCodecInfos[i].GetResolution();
-							if (Res.ExceedsLimit(MaxWidth, MaxHeight))
-							{
-								continue;
-							}
-							if (bFirst)
-							{
-								bFirst = false;
-								Best = VideoCodecInfos[i];
-							}
-							if (Res.Width > Best.GetResolution().Width)
-							{
-								Best.SetResolution(FStreamCodecInformation::FResolution(Res.Width, Best.GetResolution().Height));
-							}
-							if (Res.Height > Best.GetResolution().Height)
-							{
-								Best.SetResolution(FStreamCodecInformation::FResolution(Best.GetResolution().Width, Res.Height));
-							}
-							// Note: the final RFC 6381 codec string will be bogus since we do not re-create it here.
-							if (VideoCodecInfos[i].GetProfile() > Best.GetProfile())
-							{
-								Best.SetProfile(VideoCodecInfos[i].GetProfile());
-							}
-							if (VideoCodecInfos[i].GetProfileLevel() > Best.GetProfileLevel())
-							{
-								Best.SetProfileLevel(VideoCodecInfos[i].GetProfileLevel());
-							}
-							if (VideoCodecInfos[i].GetExtras().GetValue("b_frames").SafeGetInt64(0))
-							{
-								Best.GetExtras().Set("b_frames", FVariantValue((int64) 1));
-							}
-						}
-						// Found none? (resolution limit set too low)
-						if (bFirst)
-						{
-							// Find smallest by bandwidth
-							Best = VideoCodecInfos[0];
-							int32 BestBW = VideoSet->GetRepresentationByIndex(0)->GetBitrate();
-							for(int32 i=1; i<VideoCodecInfos.Num(); ++i)
-							{
-								if (VideoSet->GetRepresentationByIndex(i)->GetBitrate() < BestBW)
-								{
-									Best = VideoCodecInfos[i];
-									BestBW = VideoSet->GetRepresentationByIndex(i)->GetBitrate();
-								}
-							}
-						}
-						OutStreamInfo = Best;
-					}
-					return true;
-				}
-			}
-		}
-	}
-	return false;
-}
-
-
 //-----------------------------------------------------------------------------
 /**
  * Updates the ABR and video decoder with maximum stream resolution limits.
  */
 void FAdaptiveStreamingPlayer::UpdateStreamResolutionLimit()
 {
-	FStreamCodecInformation StreamInfo;
-	if (FindMatchingStreamInfo(StreamInfo, VideoResolutionLimitWidth, VideoResolutionLimitHeight))
-	{
-		if (VideoDecoder.Decoder)
-		{
-			VideoDecoder.Decoder->SetMaximumDecodeCapability(StreamInfo.GetResolution().Width, StreamInfo.GetResolution().Height, StreamInfo.GetProfile(), StreamInfo.GetProfileLevel(), StreamInfo.GetExtras());
-		}
-	}
 	StreamSelector->SetMaxVideoResolution(VideoResolutionLimitWidth, VideoResolutionLimitHeight);
+	VideoDecoder.bApplyNewLimits = true;
 }
 
 
@@ -3570,6 +3225,7 @@ void FAdaptiveStreamingPlayer::InternalStop(bool bHoldCurrentFrame)
 	CurrentPlayPeriodVideo.Reset();
 	CurrentPlayPeriodAudio.Reset();
 	CurrentPlayPeriodText.Reset();
+	ActivePeriods.Empty();
 	UpcomingPeriods.Empty();
 
 	PlaybackRate = 0.0;
diff --git a/Engine/Plugins/Media/ElectraPlayer/Source/ElectraPlayerRuntime/Private/Runtime/Player/AdaptiveStreamingPlayerDecoders.cpp b/Engine/Plugins/Media/ElectraPlayer/Source/ElectraPlayerRuntime/Private/Runtime/Player/AdaptiveStreamingPlayerDecoders.cpp
new file mode 100644
--- /dev/null
+++ b/Engine/Plugins/Media/ElectraPlayer/Source/ElectraPlayerRuntime/Private/Runtime/Player/AdaptiveStreamingPlayerDecoders.cpp
@@ -0,0 +1,565 @@
+// Copyright Epic Games, Inc. All Rights Reserved.
+
+#include "PlayerCore.h"
+#include "Player/AdaptiveStreamingPlayerInternal.h"
+#include "Player/AdaptivePlayerOptionKeynames.h"
+#include "Utilities/Utilities.h"
+#include "ParameterDictionary.h"
+
+
+namespace Electra
+{
+
+
+bool FAdaptiveStreamingPlayer::FindMatchingStreamInfo(FStreamCodecInformation& OutStreamInfo, const FTimeValue& AtTime, int32 MaxWidth, int32 MaxHeight)
+{
+	TArray<FStreamCodecInformation> VideoCodecInfos;
+	TSharedPtrTS<ITimelineMediaAsset> Asset;
+
+	if (ActivePeriods.Num() == 0)
+	{
+		return false;
+	}
+
+	// Locate the period for the specified time.
+	for(int32 i=0; i<ActivePeriods.Num(); ++i)
+	{
+		if (AtTime >= ActivePeriods[i].TimeRange.Start && AtTime < (ActivePeriods[i].TimeRange.End.IsValid() ? ActivePeriods[i].TimeRange.End : FTimeValue::GetPositiveInfinity()))
+		{
+			Asset = ActivePeriods[i].Period;
+		}
+	}
+	if (Asset.IsValid())
+	{
+		if (Asset->GetNumberOfAdaptationSets(EStreamType::Video) > 0)
+		{
+			// What if this is more than one?
+			TSharedPtrTS<IPlaybackAssetAdaptationSet> VideoSet = Asset->GetAdaptationSetByTypeAndIndex(EStreamType::Video, 0);
+			check(VideoSet.IsValid());
+			if (VideoSet.IsValid())
+			{
+				for(int32 i = 0, iMax = VideoSet->GetNumberOfRepresentations(); i < iMax; ++i)
+				{
+					VideoCodecInfos.Push(VideoSet->GetRepresentationByIndex(i)->GetCodecInformation());
+				}
+				check(VideoCodecInfos.Num());
+				if (VideoCodecInfos.Num())
+				{
+					if (MaxWidth == 0 && MaxHeight == 0)
+					{
+						FStreamCodecInformation Best = VideoCodecInfos[0];
+						for(int32 i=1; i<VideoCodecInfos.Num(); ++i)
+						{
+							const FStreamCodecInformation::FResolution& Res = VideoCodecInfos[i].GetResolution();
+							if (Res.Width > Best.GetResolution().Width)
+							{
+								Best.SetResolution(FStreamCodecInformation::FResolution(Res.Width, Best.GetResolution().Height));
+							}
+							if (Res.Height > Best.GetResolution().Height)
+							{
+								Best.SetResolution(FStreamCodecInformation::FResolution(Best.GetResolution().Width, Res.Height));
+							}
+							// Note: the final RFC 6381 codec string will be bogus since we do not re-create it here.
+							if (VideoCodecInfos[i].GetProfile() > Best.GetProfile())
+							{
+								Best.SetProfile(VideoCodecInfos[i].GetProfile());
+							}
+							if (VideoCodecInfos[i].GetProfileLevel() > Best.GetProfileLevel())
+							{
+								Best.SetProfileLevel(VideoCodecInfos[i].GetProfileLevel());
+							}
+							if (VideoCodecInfos[i].GetExtras().GetValue("b_frames").SafeGetInt64(0))
+							{
+								Best.GetExtras().Set("b_frames", FVariantValue((int64) 1));
+							}
+						}
+						OutStreamInfo = Best;
+					}
+					else
+					{
+						if (MaxWidth == 0)
+						{
+							MaxWidth = 32768;
+						}
+						if (MaxHeight == 0)
+						{
+							MaxHeight = 32768;
+						}
+						FStreamCodecInformation		Best;
+						bool bFirst = true;
+						for(int32 i=0; i<VideoCodecInfos.Num(); ++i)
+						{
+							const FStreamCodecInformation::FResolution& Res = VideoCodecInfos[i].GetResolution();
+							if (Res.ExceedsLimit(MaxWidth, MaxHeight))
+							{
+								continue;
+							}
+							if (bFirst)
+							{
+								bFirst = false;
+								Best = VideoCodecInfos[i];
+							}
+							if (Res.Width > Best.GetResolution().Width)
+							{
+								Best.SetResolution(FStreamCodecInformation::FResolution(Res.Width, Best.GetResolution().Height));
+							}
+							if (Res.Height > Best.GetResolution().Height)
+							{
+								Best.SetResolution(FStreamCodecInformation::FResolution(Best.GetResolution().Width, Res.Height));
+							}
+							// Note: the final RFC 6381 codec string will be bogus since we do not re-create it here.
+							if (VideoCodecInfos[i].GetProfile() > Best.GetProfile())
+							{
+								Best.SetProfile(VideoCodecInfos[i].GetProfile());
+							}
+							if (VideoCodecInfos[i].GetProfileLevel() > Best.GetProfileLevel())
+							{
+								Best.SetProfileLevel(VideoCodecInfos[i].GetProfileLevel());
+							}
+							if (VideoCodecInfos[i].GetExtras().GetValue("b_frames").SafeGetInt64(0))
+							{
+								Best.GetExtras().Set("b_frames", FVariantValue((int64) 1));
+							}
+						}
+						// Found none? (resolution limit set too low)
+						if (bFirst)
+						{
+							// Find smallest by bandwidth
+							Best = VideoCodecInfos[0];
+							int32 BestBW = VideoSet->GetRepresentationByIndex(0)->GetBitrate();
+							for(int32 i=1; i<VideoCodecInfos.Num(); ++i)
+							{
+								if (VideoSet->GetRepresentationByIndex(i)->GetBitrate() < BestBW)
+								{
+									Best = VideoCodecInfos[i];
+									BestBW = VideoSet->GetRepresentationByIndex(i)->GetBitrate();
+								}
+							}
+						}
+						OutStreamInfo = Best;
+					}
+					return true;
+				}
+			}
+		}
+	}
+	return false;
+}
+
+//-----------------------------------------------------------------------------
+/**
+ * Creates a decoder for the specified stream type based on the first access
+ * unit's format.
+ *
+ * @return
+ */
+int32 FAdaptiveStreamingPlayer::CreateDecoder(EStreamType type)
+{
+	if (type == EStreamType::Video)
+	{
+		if (VideoDecoder.Decoder == nullptr)
+		{
+			FAccessUnit* AccessUnit = nullptr;
+			bool bOk = MultiStreamBufferVid.PeekAndAddRef(AccessUnit);
+			if (AccessUnit)
+			{
+				FTimeValue DecodeTime = AccessUnit->PTS;
+				VideoDecoder.CurrentCodecInfo.Clear();
+				if (AccessUnit->AUCodecData.IsValid())
+				{
+					VideoDecoder.CurrentCodecInfo = AccessUnit->AUCodecData->ParsedInfo;
+				}
+				FAccessUnit::Release(AccessUnit);
+
+				// Get the largest stream resolution of the currently selected video adaptation set.
+				// This is only an initial selection as there could be other adaptation sets in upcoming periods
+				// that have a larger resolution that is still within the allowed limits.
+				FStreamCodecInformation HighestStream;
+				if (!FindMatchingStreamInfo(HighestStream, DecodeTime, 0, 0))
+				{
+					FErrorDetail err;
+					err.SetFacility(Facility::EFacility::Player);
+					err.SetMessage("Stream information not found when creating video decoder");
+					err.SetCode(INTERR_NO_STREAM_INFORMATION);
+					PostError(err);
+					return -1;
+				}
+
+				if (VideoDecoder.CurrentCodecInfo.GetCodec() == FStreamCodecInformation::ECodec::H264)
+				{
+					// Create H.264 video decoder
+					IVideoDecoderH264* DecoderH264 = IVideoDecoderH264::Create();
+					if (DecoderH264)
+					{
+						VideoDecoder.Decoder = DecoderH264;
+						VideoDecoder.Parent = this;
+						DecoderH264->SetPlayerSessionServices(this);
+						IVideoDecoderH264::FInstanceConfiguration h264Cfg = PlayerConfig.DecoderCfg264;
+						h264Cfg.ProfileIdc = HighestStream.GetProfile();
+						h264Cfg.LevelIdc = HighestStream.GetProfileLevel();
+						h264Cfg.MaxFrameWidth = HighestStream.GetResolution().Width;
+						h264Cfg.MaxFrameHeight = HighestStream.GetResolution().Height;
+						h264Cfg.AdditionalOptions = HighestStream.GetExtras();
+
+						// Add in any player options that are for decoder use
+						TArray<FString> DecoderOptionKeys;
+						PlayerOptions.GetKeysStartingWith("videoDecoder", DecoderOptionKeys);
+						for (const FString & Key : DecoderOptionKeys)
+						{
+							h264Cfg.AdditionalOptions.Set(Key, PlayerOptions.GetValue(Key));
+						}
+
+						// Attach video decoder buffer monitor.
+						DecoderH264->SetAUInputBufferListener(&VideoDecoder);
+						DecoderH264->SetReadyBufferListener(&VideoDecoder);
+						// Have the video decoder send its output to the video renderer
+						DecoderH264->SetRenderer(VideoRender.Renderer);
+						// Hand it (may be nullptr) a delegate for platform for resource queries
+						DecoderH264->SetResourceDelegate(VideoDecoderResourceDelegate.Pin());
+						// Open the decoder after having set all listeners.
+						DecoderH264->Open(h264Cfg);
+					}
+				}
+#if ELECTRA_PLATFORM_HAS_H265_DECODER
+				else if (VideoDecoder.CurrentCodecInfo.GetCodec() == FStreamCodecInformation::ECodec::H265)
+				{
+					// Create H.265 video decoder
+					IVideoDecoderH265* DecoderH265 = IVideoDecoderH265::Create();
+					if (DecoderH265)
+					{
+						VideoDecoder.Decoder = DecoderH265;
+						VideoDecoder.Parent = this;
+						DecoderH265->SetPlayerSessionServices(this);
+						IVideoDecoderH265::FInstanceConfiguration h265Cfg;// = PlayerConfig.DecoderCfg265;
+						h265Cfg.Tier = HighestStream.GetProfileTier();
+						h265Cfg.Profile = HighestStream.GetProfile();
+						h265Cfg.Level = HighestStream.GetProfileLevel();
+						h265Cfg.MaxFrameWidth = HighestStream.GetResolution().Width;
+						h265Cfg.MaxFrameHeight = HighestStream.GetResolution().Height;
+						h265Cfg.AdditionalOptions = HighestStream.GetExtras();
+
+						// Add in any player options that are for decoder use
+						TArray<FString> DecoderOptionKeys;
+						PlayerOptions.GetKeysStartingWith("videoDecoder", DecoderOptionKeys);
+						for (const FString & Key : DecoderOptionKeys)
+						{
+							h265Cfg.AdditionalOptions.Set(Key, PlayerOptions.GetValue(Key));
+						}
+
+						// Attach video decoder buffer monitor.
+						DecoderH265->SetAUInputBufferListener(&VideoDecoder);
+						DecoderH265->SetReadyBufferListener(&VideoDecoder);
+						// Have the video decoder send its output to the video renderer
+						DecoderH265->SetRenderer(VideoRender.Renderer);
+						// Hand it (may be nullptr) a delegate for platform for resource queries
+						DecoderH265->SetResourceDelegate(VideoDecoderResourceDelegate.Pin());
+						// Open the decoder after having set all listeners.
+						DecoderH265->Open(h265Cfg);
+					}
+				}
+#endif
+
+				if (VideoDecoder.Decoder)
+				{
+					// Now we get the currently limited stream resolution and let the decoder now what we will be using
+					// at most right now. This allows the decoder to be created with a smaller memory footprint at first.
+					UpdateStreamResolutionLimit();
+				}
+				else
+				{
+					FErrorDetail err;
+					err.SetFacility(Facility::EFacility::Player);
+					err.SetMessage("Unsupported video codec");
+					err.SetCode(INTERR_UNSUPPORTED_CODEC);
+					PostError(err);
+					return -1;
+				}
+			}
+		}
+		return 0;
+	}
+	else if (type == EStreamType::Audio)
+	{
+		if (AudioDecoder.Decoder == nullptr)
+		{
+			FAccessUnit* AccessUnit = nullptr;
+			bool bOk = MultiStreamBufferAud.PeekAndAddRef(AccessUnit);
+			if (AccessUnit)
+			{
+				AudioDecoder.CurrentCodecInfo.Clear();
+				if (AccessUnit->AUCodecData.IsValid())
+				{
+					AudioDecoder.CurrentCodecInfo = AccessUnit->AUCodecData->ParsedInfo;
+				}
+				FAccessUnit::Release(AccessUnit);
+
+				if (AudioDecoder.CurrentCodecInfo.GetCodec() == FStreamCodecInformation::ECodec::AAC)
+				{
+					// Create an AAC audio decoder
+					AudioDecoder.Decoder = IAudioDecoderAAC::Create();
+					AudioDecoder.Decoder->SetPlayerSessionServices(this);
+					AudioDecoder.Parent = this;
+					// Attach buffer monitors.
+					AudioDecoder.Decoder->SetAUInputBufferListener(&AudioDecoder);
+					AudioDecoder.Decoder->SetReadyBufferListener(&AudioDecoder);
+					// Have to audio decoder send its output to the audio renderer
+					AudioDecoder.Decoder->SetRenderer(AudioRender.Renderer);
+					// Open the decoder after having set all listeners.
+					AudioDecoder.Decoder->Open(PlayerConfig.DecoderCfgAAC);
+				}
+
+				if (!AudioDecoder.Decoder)
+				{
+					FErrorDetail err;
+					err.SetFacility(Facility::EFacility::Player);
+					err.SetMessage("Unsupported audio codec");
+					err.SetCode(INTERR_UNSUPPORTED_CODEC);
+					PostError(err);
+					return -1;
+				}
+			}
+		}
+		return 0;
+	}
+	return -1;
+}
+
+
+//-----------------------------------------------------------------------------
+/**
+ * Destroys the decoders.
+ */
+void FAdaptiveStreamingPlayer::DestroyDecoders()
+{
+/*
+NOTE: We do not clear out the renderers from the decoder. On their way down the decoders should still be able
+      to access the renderer without harm and dispatch their last remaining data.
+
+	if (mVideoDecoder.Decoder)
+		mVideoDecoder.Decoder->SetRenderer(nullptr);
+	if (mAudioDecoder.Decoder)
+		mAudioDecoder.Decoder->SetRenderer(nullptr);
+*/
+	AudioDecoder.Close();
+	VideoDecoder.Close();
+}
+
+
+//-----------------------------------------------------------------------------
+/**
+ * Check if the decoders need to change.
+ */
+void FAdaptiveStreamingPlayer::HandleDecoderChanges()
+{
+	if (VideoDecoder.bDrainingForCodecChange && VideoDecoder.bDrainingForCodecChangeDone)
+	{
+		VideoDecoder.Close();
+		VideoDecoder.bDrainingForCodecChange = false;
+		VideoDecoder.bDrainingForCodecChangeDone = false;
+	}
+	CreateDecoder(EStreamType::Video);
+	CreateDecoder(EStreamType::Audio);
+}
+
+
+void FAdaptiveStreamingPlayer::VideoDecoderInputNeeded(const IAccessUnitBufferListener::FBufferStats &currentInputBufferStats)
+{
+	DiagnosticsCriticalSection.Lock();
+	VideoBufferStats.DecoderInputBuffer = currentInputBufferStats;
+	DiagnosticsCriticalSection.Unlock();
+	if (!VideoDecoder.bDrainingForCodecChange)
+	{
+		EDecoderState decState = DecoderState;
+		if (decState == EDecoderState::eDecoder_Running)
+		{
+			FeedDecoder(EStreamType::Video, MultiStreamBufferVid, VideoDecoder.Decoder);
+		}
+	}
+}
+
+void FAdaptiveStreamingPlayer::VideoDecoderOutputReady(const IDecoderOutputBufferListener::FDecodeReadyStats &currentReadyStats)
+{
+	DiagnosticsCriticalSection.Lock();
+	VideoBufferStats.DecoderOutputBuffer = currentReadyStats;
+	DiagnosticsCriticalSection.Unlock();
+}
+
+
+void FAdaptiveStreamingPlayer::AudioDecoderInputNeeded(const IAccessUnitBufferListener::FBufferStats &currentInputBufferStats)
+{
+	DiagnosticsCriticalSection.Lock();
+	AudioBufferStats.DecoderInputBuffer = currentInputBufferStats;
+	DiagnosticsCriticalSection.Unlock();
+	EDecoderState decState = DecoderState;
+	if (decState == EDecoderState::eDecoder_Running)
+	{
+		FeedDecoder(EStreamType::Audio, MultiStreamBufferAud, AudioDecoder.Decoder);
+	}
+}
+
+void FAdaptiveStreamingPlayer::AudioDecoderOutputReady(const IDecoderOutputBufferListener::FDecodeReadyStats &currentReadyStats)
+{
+	DiagnosticsCriticalSection.Lock();
+	AudioBufferStats.DecoderOutputBuffer = currentReadyStats;
+	DiagnosticsCriticalSection.Unlock();
+}
+
+
+
+//-----------------------------------------------------------------------------
+/**
+ * Sends an available AU to a decoder.
+ * If the current buffer level is below the underrun threshold an underrun
+ * message is sent to the worker thread.
+ *
+ * @param Type
+ * @param FromMultistreamBuffer
+ * @param Decoder
+ */
+void FAdaptiveStreamingPlayer::FeedDecoder(EStreamType Type, FMultiTrackAccessUnitBuffer& FromMultistreamBuffer, IAccessUnitBufferInterface* Decoder)
+{
+	// Lock the AU buffer for the duration of this function to ensure this can never clash with a Flush() call
+	// since we are checking size, eod state and subsequently popping an AU, for which the buffer must stay consistent inbetween!
+	// Also to ensure the active buffer doesn't get changed from one track to another.
+	FMultiTrackAccessUnitBuffer::FScopedLock lock(FromMultistreamBuffer);
+
+	FBufferStats* pStats = nullptr;
+	Metrics::FDataAvailabilityChange* pAvailability = nullptr;
+	FStreamCodecInformation* CurrentCodecInfo = nullptr;
+	bool bCodecChangeDetected = false;
+
+	switch(Type)
+	{
+		case EStreamType::Video:
+			pStats = &VideoBufferStats;
+			pAvailability = &DataAvailabilityStateVid;
+			CurrentCodecInfo = &VideoDecoder.CurrentCodecInfo;
+			break;
+		case EStreamType::Audio:
+			pStats = &AudioBufferStats;
+			pAvailability = &DataAvailabilityStateAud;
+			CurrentCodecInfo = &AudioDecoder.CurrentCodecInfo;
+			break;
+		case EStreamType::Subtitle:
+			pStats = &TextBufferStats;
+			pAvailability = &DataAvailabilityStateTxt;
+			break;
+		default:
+			checkNoEntry();
+			return;
+	}
+
+	// Is the buffer (the Type of elementary stream actually) active/selected?
+	if (!FromMultistreamBuffer.IsDeselected())
+	{
+		// Check for buffer underrun.
+		if (!bRebufferPending && CurrentState == EPlayerState::eState_Playing && StreamState == EStreamState::eStream_Running && PipelineState == EPipelineState::ePipeline_Running)
+		{
+			bool bEODSet = FromMultistreamBuffer.IsEODFlagSet();
+			if (!bEODSet && FromMultistreamBuffer.Num() == 0)
+			{
+				// Buffer underrun.
+				bRebufferPending = true;
+				FTimeValue LastKnownPTS = FromMultistreamBuffer.GetLastPoppedPTS();
+				// Only set the 'rebuffer at' time if we have a valid last known PTS. If we don't
+				// then maybe this is a cascade failure from a previous rebuffer attempt for which
+				// we then try that time once more.
+				if (LastKnownPTS.IsValid())
+				{
+					RebufferDetectedAtPlayPos = LastKnownPTS;
+				}
+				WorkerThread.SendMessage(FWorkerThread::FMessage::EType::BufferUnderrun);
+			}
+		}
+
+		FAccessUnit* AccessUnit = nullptr;
+		FromMultistreamBuffer.PeekAndAddRef(AccessUnit);
+		if (AccessUnit)
+		{
+			// Change in codec?
+			if (CurrentCodecInfo && AccessUnit->AUCodecData.IsValid() && AccessUnit->AUCodecData->ParsedInfo.GetCodec() != CurrentCodecInfo->GetCodec())
+			{
+				bCodecChangeDetected = true;
+				if (Decoder)
+				{
+					// Check type of stream. We can currently change the video codec only.
+					if (Type == EStreamType::Video)
+					{
+						if (VideoDecoder.Decoder && !VideoDecoder.bDrainingForCodecChange)
+						{
+							VideoDecoder.bDrainingForCodecChange = true;
+							VideoDecoder.Decoder->DrainForCodecChange();
+						}
+					}
+					else
+					{
+						FErrorDetail err;
+						err.SetFacility(Facility::EFacility::Player);
+						err.SetMessage("Codec change not supported");
+						err.SetCode(INTERR_CODEC_CHANGE_NOT_SUPPORTED);
+						PostError(err);
+					}
+				}
+			}
+			else
+			{
+				if (Type == EStreamType::Video && VideoDecoder.bApplyNewLimits)
+				{
+					FStreamCodecInformation StreamInfo;
+					if (FindMatchingStreamInfo(StreamInfo, AccessUnit->PTS, VideoResolutionLimitWidth, VideoResolutionLimitHeight))
+					{
+						if (VideoDecoder.Decoder)
+						{
+							if (VideoDecoder.CurrentCodecInfo.GetCodec() == FStreamCodecInformation::ECodec::H264 && StreamInfo.GetCodec() == FStreamCodecInformation::ECodec::H264)
+							{
+								static_cast<IVideoDecoderH264*>(VideoDecoder.Decoder)->SetMaximumDecodeCapability(StreamInfo.GetResolution().Width, StreamInfo.GetResolution().Height, StreamInfo.GetProfile(), StreamInfo.GetProfileLevel(), StreamInfo.GetExtras());
+							}
+						}
+					}
+					VideoDecoder.bApplyNewLimits = false;
+				}
+
+				FromMultistreamBuffer.Pop(AccessUnit);
+				if (Decoder)
+				{
+				// The decoder has asked to be fed a new AU so it better be able to accept it.
+					/*IAccessUnitBufferInterface::EAUpushResult auPushRes =*/ Decoder->AUdataPushAU(AccessUnit);
+				}
+
+				// The decoder will have added a ref count if it took the AU. If it didnt' for whatever reason
+				// we still release it to get rid of it and not cause a memory leak.
+				FAccessUnit::Release(AccessUnit);
+				AccessUnit = nullptr;
+
+				if (pAvailability)
+				{
+					UpdateDataAvailabilityState(*pAvailability, Metrics::FDataAvailabilityChange::EAvailability::DataAvailable);
+				}
+			}
+			// Release ref count from peeking.
+			FAccessUnit::Release(AccessUnit);
+		}
+	}
+	// An AU is not tagged as being "the last" one. Instead the EOD is handled separately and must be dealt with
+	// by the decoders accordingly.
+	if (!bCodecChangeDetected && FromMultistreamBuffer.IsEODFlagSet() && FromMultistreamBuffer.Num() == 0)
+	{
+		if (pStats && !pStats->DecoderInputBuffer.bEODSignaled)
+		{
+			if (Decoder)
+			{
+				Decoder->AUdataPushEOD();
+			}
+		}
+		if (pAvailability)
+		{
+			UpdateDataAvailabilityState(*pAvailability, Metrics::FDataAvailabilityChange::EAvailability::DataNotAvailable);
+		}
+	}
+}
+
+
+
+} // namespace Electra
+
+
diff --git a/Engine/Plugins/Media/ElectraPlayer/Source/ElectraPlayerRuntime/Private/Runtime/Player/AdaptiveStreamingPlayerInternal.h b/Engine/Plugins/Media/ElectraPlayer/Source/ElectraPlayerRuntime/Private/Runtime/Player/AdaptiveStreamingPlayerInternal.h
--- a/Engine/Plugins/Media/ElectraPlayer/Source/ElectraPlayerRuntime/Private/Runtime/Player/AdaptiveStreamingPlayerInternal.h
+++ b/Engine/Plugins/Media/ElectraPlayer/Source/ElectraPlayerRuntime/Private/Runtime/Player/AdaptiveStreamingPlayerInternal.h
@@ -26,6 +26,9 @@
 #define INTERR_UNSUPPORTED_FORMAT				2
 #define INTERR_COULD_NOT_LOCATE_START_SEGMENT	3
 #define INTERR_COULD_NOT_LOCATE_START_PERIOD	4
+#define INTERR_UNSUPPORTED_CODEC				5
+#define INTERR_CODEC_CHANGE_NOT_SUPPORTED		6
+#define INTERR_NO_STREAM_INFORMATION			7
 #define INTERR_FRAGMENT_NOT_AVAILABLE			0x101
 #define INTERR_FRAGMENT_READER_REQUEST			0x102
 #define INTERR_CREATE_FRAGMENT_READER			0x103
@@ -629,8 +632,8 @@ private:
 	void StopWorkerThread();
 	void WorkerThreadFN();
 
-	FMediaThread													WorkerThread;
-	TMediaMessageQueueDynamicNoTimeout<TSharedPtrTS<FMetricEvent>>		EventQueue;
+	FMediaThread																WorkerThread;
+	TMediaMessageQueueDynamicNoTimeout<TSharedPtrTS<FMetricEvent>>				EventQueue;
 	static TWeakPtr<FAdaptiveStreamingPlayerEventHandler, ESPMode::ThreadSafe>	SingletonSelf;
 	static FCriticalSection														SingletonLock;
 };
@@ -854,11 +857,6 @@ private:
 
 	struct FVideoDecoder : public IAccessUnitBufferListener, public IDecoderOutputBufferListener
 		{
-		FVideoDecoder()
-			: Parent(nullptr)
-			, Decoder(nullptr)
-		{
-		}
 		void Close()
 		{
 			if (Decoder)
@@ -892,8 +890,12 @@ private:
 			}
 		}
 
-		FAdaptiveStreamingPlayer*		Parent;
-		IVideoDecoderH264*				Decoder;
+		FStreamCodecInformation			CurrentCodecInfo;
+		FAdaptiveStreamingPlayer*		Parent = nullptr;
+		IVideoDecoderBase*				Decoder = nullptr;
+		bool							bDrainingForCodecChange = false;
+		bool							bDrainingForCodecChangeDone = false;
+		bool							bApplyNewLimits = false;
 		};
 
 	struct FAudioDecoder : public IAccessUnitBufferListener, public IDecoderOutputBufferListener
@@ -936,6 +938,7 @@ private:
 			}
 		}
 
+		FStreamCodecInformation			CurrentCodecInfo;
 		FAdaptiveStreamingPlayer*		Parent;
 		IAudioDecoderAAC*				Decoder;
 		};
@@ -1327,7 +1330,7 @@ private:
 		FPlayStartPosition							StartoverPosition;
 	};
 
-	struct FUpcomingPeriod
+	struct FPeriodInformation
 	{
 		TSharedPtrTS<ITimelineMediaAsset>		Period;
 		FString									ID;
@@ -1356,9 +1359,9 @@ private:
 	int32 CreateRenderers();
 	void DestroyRenderers();
 
-	int32 CreateInitialDecoder(EStreamType type);
+	int32 CreateDecoder(EStreamType type);
 	void DestroyDecoders();
-	bool FindMatchingStreamInfo(FStreamCodecInformation& OutStreamInfo, int32 MaxWidth, int32 MaxHeight);
+	bool FindMatchingStreamInfo(FStreamCodecInformation& OutStreamInfo, const FTimeValue& AtTime, int32 MaxWidth, int32 MaxHeight);
 	void UpdateStreamResolutionLimit();
 	void AddUpcomingPeriod(TSharedPtrTS<IManifest::IPlayPeriod> InUpcomingPeriod);
 
@@ -1504,7 +1507,8 @@ private:
 	FLoopParam															CurrentLoopParam;
 	TMediaQueueDynamicNoLock<FPlayerLoopState>							NextLoopStates;
 
-	TArray<FUpcomingPeriod>												UpcomingPeriods;
+	TArray<FPeriodInformation>											ActivePeriods;
+	TArray<FPeriodInformation>											UpcomingPeriods;
 	TSharedPtrTS<IManifest::IPlayPeriod>								InitialPlayPeriod;
 	TSharedPtrTS<IManifest::IPlayPeriod>								CurrentPlayPeriodVideo;
 	TSharedPtrTS<IManifest::IPlayPeriod>								CurrentPlayPeriodAudio;
diff --git a/Engine/Plugins/Media/ElectraPlayer/Source/ElectraPlayerRuntime/Private/Runtime/Player/AdaptiveStreamingPlayerInternalConfig.h b/Engine/Plugins/Media/ElectraPlayer/Source/ElectraPlayerRuntime/Private/Runtime/Player/AdaptiveStreamingPlayerInternalConfig.h
--- a/Engine/Plugins/Media/ElectraPlayer/Source/ElectraPlayerRuntime/Private/Runtime/Player/AdaptiveStreamingPlayerInternalConfig.h
+++ b/Engine/Plugins/Media/ElectraPlayer/Source/ElectraPlayerRuntime/Private/Runtime/Player/AdaptiveStreamingPlayerInternalConfig.h
@@ -3,6 +3,7 @@
 #pragma once
 
 #include "Decoder/VideoDecoderH264.h"
+#include "Decoder/VideoDecoderH265.h"
 #include "Decoder/AudioDecoderAAC.h"
 #include "StreamAccessUnitBuffer.h"
 #include "Player/AdaptiveStreamingPlayerABR.h"
diff --git a/Engine/Plugins/Media/ElectraPlayer/Source/ElectraPlayerRuntime/Private/Runtime/Player/DASH/ManifestDASH.cpp b/Engine/Plugins/Media/ElectraPlayer/Source/ElectraPlayerRuntime/Private/Runtime/Player/DASH/ManifestDASH.cpp
--- a/Engine/Plugins/Media/ElectraPlayer/Source/ElectraPlayerRuntime/Private/Runtime/Player/DASH/ManifestDASH.cpp
+++ b/Engine/Plugins/Media/ElectraPlayer/Source/ElectraPlayerRuntime/Private/Runtime/Player/DASH/ManifestDASH.cpp
@@ -359,16 +359,9 @@ IManifest::FResult FManifestDASH::FindPlayPeriod(TSharedPtrTS<IPlayPeriod>& OutP
 				{
 					case ESearchType::Closest:
 					{
-						// If there is a following period whose start is closer to the current position than the start of this period and the position
-						// is within a short duration to the next we jump ahead to the next.
-						if (DiffToNextPeriod < DiffToStart && DiffToNextPeriod.GetAsMilliseconds() < 1000)
-						{
-							SelectedPeriod = Periods[nPeriod + 1];
-						}
-						else
-						{
-							SelectedPeriod = Periods[nPeriod];
-						}
+						// There is no actual choice. We have to use the period the time falls into. Why would we want
+						// to snap to a different period that won't contain the segments for the time we're looking for.
+						SelectedPeriod = Periods[nPeriod];
 						break;
 					}
 					case ESearchType::Before:
@@ -1974,6 +1967,7 @@ FManifestDASHInternal::FRepresentation::ESearchResult FManifestDASHInternal::FRe
 	int32 PreviousN = 0;
 	int64 CurrentOffset = 0;
 	int64 PreviousOffset = 0;
+	bool bFound = false;
 	for(CurrentN=0; CurrentN<EndNumber; ++CurrentN)
 	{
 		const IParserISO14496_12::ISegmentIndex::FEntry& SegmentInfo = Sidx->GetEntry(CurrentN);
@@ -2009,6 +2003,8 @@ FManifestDASHInternal::FRepresentation::ESearchResult FManifestDASHInternal::FRe
 		// Does the segment start on or after the time we're looking for?
 		if (CurrentT >= MediaLocalSearchTime)
 		{
+			bFound = true;
+
 			// Yes, so we have now found the segment of interest. It is either this one or the previous one.
 			if (InSearchOptions.SearchType == IManifest::ESearchType::Closest)
 			{
@@ -2081,8 +2077,21 @@ FManifestDASHInternal::FRepresentation::ESearchResult FManifestDASHInternal::FRe
 		CurrentOffset += SegmentInfo.Size;
 	}
 
+	// If the search time falls into the last segment we will not have found it above.
+	if (!bFound && CurrentT >= MediaLocalSearchTime && CurrentN == EndNumber)
+	{
+		if (InSearchOptions.SearchType == IManifest::ESearchType::Closest || InSearchOptions.SearchType == IManifest::ESearchType::Same || InSearchOptions.SearchType == IManifest::ESearchType::Before || InSearchOptions.SearchType == IManifest::ESearchType::StrictlyBefore)
+		{
+			--CurrentN;
+			CurrentD = PreviousD;
+			CurrentT = PreviousT;
+			CurrentOffset -= Sidx->GetEntry(CurrentN).Size;
+			bFound = true;
+		}
+	}
+
 	// Did we find it?
-	if (CurrentN < EndNumber && CurrentT < MediaLocalPeriodEnd)
+	if (bFound && CurrentT < MediaLocalPeriodEnd)
 	{
 		OutSegmentInfo.Time = CurrentT;
 		OutSegmentInfo.PTO = PTO;
diff --git a/Engine/Plugins/Media/ElectraPlayer/Source/ElectraPlayerRuntime/Private/Runtime/Player/DASH/StreamReaderFMP4DASH.cpp b/Engine/Plugins/Media/ElectraPlayer/Source/ElectraPlayerRuntime/Private/Runtime/Player/DASH/StreamReaderFMP4DASH.cpp
--- a/Engine/Plugins/Media/ElectraPlayer/Source/ElectraPlayerRuntime/Private/Runtime/Player/DASH/StreamReaderFMP4DASH.cpp
+++ b/Engine/Plugins/Media/ElectraPlayer/Source/ElectraPlayerRuntime/Private/Runtime/Player/DASH/StreamReaderFMP4DASH.cpp
@@ -800,12 +800,14 @@ void FStreamReaderFMP4DASH::FStreamHandler::HandleRequest()
 
 								int64 MediaLocalFirstAUTime = Request->Segment.MediaLocalFirstAUTime;
 								int64 MediaLocalLastAUTime = Request->Segment.MediaLocalLastAUTime;
+								int64 PTO = Request->Segment.PTO;
 
 								if (TrackTimescale != Request->Segment.Timescale)
 								{
 									// Need to rescale the AU times from the MPD timescale to the media timescale.
 									MediaLocalFirstAUTime = FTimeFraction(Request->Segment.MediaLocalFirstAUTime, Request->Segment.Timescale).GetAsTimebase(TrackTimescale);
 									MediaLocalLastAUTime = FTimeFraction(Request->Segment.MediaLocalLastAUTime, Request->Segment.Timescale).GetAsTimebase(TrackTimescale);
+									PTO = FTimeFraction(Request->Segment.PTO, Request->Segment.Timescale).GetAsTimebase(TrackTimescale);
 
 									if (!Request->bWarnedAboutTimescale)
 									{
@@ -875,9 +877,9 @@ void FStreamReaderFMP4DASH::FStreamHandler::HandleRequest()
 									AccessUnit->Duration = Duration;
 
 									// Offset the AU's DTS and PTS to the time mapping of the segment.
-									AccessUnit->DTS.SetFromND(AUDTS, TrackTimescale);
+									AccessUnit->DTS.SetFromND(AUDTS - PTO, TrackTimescale);
 									AccessUnit->DTS += TimeOffset;
-									AccessUnit->PTS.SetFromND(AUPTS, TrackTimescale);
+									AccessUnit->PTS.SetFromND(AUPTS - PTO, TrackTimescale);
 									AccessUnit->PTS += TimeOffset;
 
 									ElectraCDM::FMediaCDMSampleInfo SampleEncryptionInfo;
@@ -957,6 +959,7 @@ void FStreamReaderFMP4DASH::FStreamHandler::HandleRequest()
 									// Check if the AU is outside the time range we are allowed to read.
 									// The last one (the one that is already outside the range, actually) is tagged as such and sent into the buffer.
 									// The respective decoder has to handle this flag if necessary and/or drop the AU.
+									// We need to send at least one AU down so the FMultiTrackAccessUnitBuffer does not stay empty for this period!
 									if (AccessUnit)
 									{
 										// Already sent the last one?
diff --git a/Engine/Plugins/Media/ElectraPlayer/Source/ElectraPlayerRuntime/Private/Runtime/PlayerCodecConfiguration.h b/Engine/Plugins/Media/ElectraPlayer/Source/ElectraPlayerRuntime/Private/Runtime/PlayerCodecConfiguration.h
new file mode 100644
--- /dev/null
+++ b/Engine/Plugins/Media/ElectraPlayer/Source/ElectraPlayerRuntime/Private/Runtime/PlayerCodecConfiguration.h
@@ -0,0 +1,8 @@
+// Copyright Epic Games, Inc. All Rights Reserved.
+
+#pragma once
+
+#ifndef ELECTRA_PLATFORM_HAS_H265_DECODER
+#define ELECTRA_PLATFORM_HAS_H265_DECODER 0
+#endif
+
diff --git a/Engine/Plugins/Media/ElectraPlayer/Source/ElectraPlayerRuntime/Private/Runtime/PlayerCore.h b/Engine/Plugins/Media/ElectraPlayer/Source/ElectraPlayerRuntime/Private/Runtime/PlayerCore.h
--- a/Engine/Plugins/Media/ElectraPlayer/Source/ElectraPlayerRuntime/Private/Runtime/PlayerCore.h
+++ b/Engine/Plugins/Media/ElectraPlayer/Source/ElectraPlayerRuntime/Private/Runtime/PlayerCore.h
@@ -15,4 +15,4 @@
 #include "Core/MediaCriticalSection.h"
 #include "Core/MediaMessageQueue.h"
 
-
+#include "PlayerCodecConfiguration.h"
diff --git a/Engine/Plugins/Media/ElectraPlayer/Source/ElectraPlayerRuntime/Private/Runtime/PlayerFacility.h b/Engine/Plugins/Media/ElectraPlayer/Source/ElectraPlayerRuntime/Private/Runtime/PlayerFacility.h
--- a/Engine/Plugins/Media/ElectraPlayer/Source/ElectraPlayerRuntime/Private/Runtime/PlayerFacility.h
+++ b/Engine/Plugins/Media/ElectraPlayer/Source/ElectraPlayerRuntime/Private/Runtime/PlayerFacility.h
@@ -40,6 +40,7 @@ namespace Electra
 			DASHFMP4Reader,
 			EntityCache,
 			DRM,
+			H265Decoder,
 			LastEntry = 255
 		};
 
@@ -101,6 +102,8 @@ namespace Electra
 					return TEXT("Entity cache");
 				case EFacility::DRM:
 					return TEXT("DRM");
+				case EFacility::H265Decoder:
+					return TEXT("H.265 decoder");
 			}
 			return TEXT("???");
 		}
diff --git a/Engine/Plugins/Media/ElectraPlayer/Source/ElectraPlayerRuntime/Private/Runtime/StreamAccessUnitBuffer.cpp b/Engine/Plugins/Media/ElectraPlayer/Source/ElectraPlayerRuntime/Private/Runtime/StreamAccessUnitBuffer.cpp
--- a/Engine/Plugins/Media/ElectraPlayer/Source/ElectraPlayerRuntime/Private/Runtime/StreamAccessUnitBuffer.cpp
+++ b/Engine/Plugins/Media/ElectraPlayer/Source/ElectraPlayerRuntime/Private/Runtime/StreamAccessUnitBuffer.cpp
@@ -127,10 +127,17 @@ namespace Electra
 	void FMultiTrackAccessUnitBuffer::Flush()
 	{
 		FMediaCriticalSection::ScopedLock lock(AccessLock);
-		// Flush all existing buffers but keep them in the track map.
-		for(auto& It : TrackBuffers)
+		if (bIsParallelTrackMode)
 		{
-			It.Value->Flush();
+			// Flush all existing buffers but keep them in the track map.
+			for(auto& It : TrackBuffers)
+			{
+				It.Value->Flush();
+			}
+		}
+		else
+		{
+			TrackBuffers.Empty();
 		}
 		Clear();
 	}
@@ -448,6 +455,21 @@ namespace Electra
 		}
 	}
 
+	bool FMultiTrackAccessUnitBuffer::PeekAndAddRef(FAccessUnit*& OutAU)
+	{
+		FScopedLock lock(*this);
+		ActivateBuffer(false);
+		TSharedPtrTS<FAccessUnitBuffer> Buf = FMultiTrackAccessUnitBuffer::GetSelectedTrackBuffer();
+		if (Buf->Num())
+		{
+			return Buf->Peek(OutAU);
+		}
+		else
+		{
+			OutAU = nullptr;
+			return false;
+		}
+	}
 
 	bool FMultiTrackAccessUnitBuffer::Pop(FAccessUnit*& OutAU)
 	{
diff --git a/Engine/Plugins/Media/ElectraPlayer/Source/ElectraPlayerRuntime/Private/Runtime/StreamAccessUnitBuffer.h b/Engine/Plugins/Media/ElectraPlayer/Source/ElectraPlayerRuntime/Private/Runtime/StreamAccessUnitBuffer.h
--- a/Engine/Plugins/Media/ElectraPlayer/Source/ElectraPlayerRuntime/Private/Runtime/StreamAccessUnitBuffer.h
+++ b/Engine/Plugins/Media/ElectraPlayer/Source/ElectraPlayerRuntime/Private/Runtime/StreamAccessUnitBuffer.h
@@ -619,6 +619,7 @@ namespace Electra
 			const FMultiTrackAccessUnitBuffer& LockedSelf;
 		};
 
+		bool PeekAndAddRef(FAccessUnit*& OutAU);
 		bool Pop(FAccessUnit*& OutAU);
 		void PopDiscardUntil(FTimeValue UntilTime);
 		bool IsEODFlagSet();
