ElectraPlayer: Added support for HEVC decoder where available

diff --git a/Engine/Plugins/Media/ElectraPlayer/Source/ElectraPlayerRuntime/Private/Runtime/StreamTypes.cpp b/Engine/Plugins/Media/ElectraPlayer/Source/ElectraPlayerRuntime/Private/Runtime/StreamTypes.cpp
--- a/Engine/Plugins/Media/ElectraPlayer/Source/ElectraPlayerRuntime/Private/Runtime/StreamTypes.cpp
+++ b/Engine/Plugins/Media/ElectraPlayer/Source/ElectraPlayerRuntime/Private/Runtime/StreamTypes.cpp
@@ -3,6 +3,7 @@
 #include "PlayerCore.h"
 #include "StreamTypes.h"
 #include "Utilities/StringHelpers.h"
+#include "Utilities/Utilities.h"
 
 namespace Electra
 {
@@ -95,11 +96,79 @@ namespace Electra
 		}
 		else if (CodecOTI.StartsWith("hvc") || CodecOTI.StartsWith("hev"))
 		{
+			FString oti = CodecOTI;
+			FString Temp;
 			// hvc1 and hev1 (inband VPS/SPS/PPS) are recognized.
 			StreamType = EStreamType::Video;
 			CodecSpecifier = CodecOTI;
 			Codec = ECodec::H265;
-			// Presently not supported.
+
+			int32 DotPos;
+			if (oti.FindChar(TCHAR('.'), DotPos))
+			{
+				int32 general_profile_space = 0;
+				int32 general_tier_flag = 0;
+				int32 general_profile_idc = 0;
+				int32 general_level_idc = 0;
+				uint32 general_profile_compatibility_flag = 0;
+				uint64 contraint_flags = 0;
+				oti.RightChopInline(DotPos + 1);
+				// optional general_profile_space
+				if (oti[0] == TCHAR('A') || oti[0] == TCHAR('B') || oti[0] == TCHAR('C'))
+				{
+					general_profile_space = oti[0] - TCHAR('A') + 1;
+					oti.RightChopInline(1);
+				}
+				else if (oti[0] == TCHAR('a') || oti[0] == TCHAR('b') || oti[0] == TCHAR('c'))
+				{
+					general_profile_space = oti[0] - TCHAR('a') + 1;
+					oti.RightChopInline(1);
+				}
+				// general_profile_idc
+				if (oti.FindChar(TCHAR('.'), DotPos))
+				{
+					Temp = oti.Left(DotPos);
+					oti.RightChopInline(DotPos + 1);
+					LexFromString(general_profile_idc, *Temp);
+				}
+				// general_profile_compatibility_flags
+				if (oti.FindChar(TCHAR('.'), DotPos))
+				{
+					Temp = oti.Left(DotPos);
+					oti.RightChopInline(DotPos + 1);
+					LexFromString(general_profile_compatibility_flag, *Temp);
+				}
+				// general_tier_flag
+				if (oti[0] != TCHAR('L') && oti[0] != TCHAR('H') && oti[0] != TCHAR('l') && oti[0] != TCHAR('h'))
+				{
+					return false;
+				}
+				else if (oti[0] == TCHAR('H') || oti[0] == TCHAR('h'))
+				{
+					general_tier_flag = 1;
+				}
+				oti.RightChopInline(1);
+				// constraint_flags
+				FString ConstraintFlags;
+				if (oti.FindChar(TCHAR('.'), DotPos))
+				{
+					ConstraintFlags = oti.Mid(DotPos + 1);
+					oti.LeftInline(DotPos);
+					ConstraintFlags.ReplaceInline(TEXT("."), TEXT(""));
+					ConstraintFlags += TEXT("000000000000");
+					ConstraintFlags.LeftInline(12);
+					LexFromStringHexU64(contraint_flags, *ConstraintFlags);
+				}
+				// general_level_idc
+				LexFromString(general_level_idc, *oti);
+				SetProfileSpace(general_profile_space);
+				SetProfileCompatibilityFlags(Utils::BitReverse32(general_profile_compatibility_flag));
+				SetProfileTier(general_tier_flag);
+				SetProfile(general_profile_idc);
+				SetProfileLevel(general_level_idc);
+				SetProfileConstraints(contraint_flags);
+				return true;
+			}
 			return false;
 		}
 		else if (CodecOTI.StartsWith("mp4a"))
@@ -122,7 +191,7 @@ namespace Electra
 				int32 ProfileValue = 0;
 				LexFromString(ProfileValue, *Profile);
 				SetProfile(ProfileValue);
-				// AAC-LC, AAC-HE, AAC-HEv2
+				// AAC-LC, AAC-HE (SBR), AAC-HEv2 (PS)
 				if (!(ProfileValue == 2 || ProfileValue == 5 || ProfileValue == 29))
 				{
 					return false;
diff --git a/Engine/Plugins/Media/ElectraPlayer/Source/ElectraPlayerRuntime/Private/Runtime/StreamTypes.h b/Engine/Plugins/Media/ElectraPlayer/Source/ElectraPlayerRuntime/Private/Runtime/StreamTypes.h
--- a/Engine/Plugins/Media/ElectraPlayer/Source/ElectraPlayerRuntime/Private/Runtime/StreamTypes.h
+++ b/Engine/Plugins/Media/ElectraPlayer/Source/ElectraPlayerRuntime/Private/Runtime/StreamTypes.h
@@ -275,6 +275,16 @@ namespace Electra
 			FrameRate = InFrameRate;
 		}
 
+		void SetProfileSpace(int32 InProfileSpace)
+		{
+			ProfileLevel.ProfileSpace = InProfileSpace;
+		}
+
+		int32 GetProfileSpace() const
+		{
+			return ProfileLevel.ProfileSpace;
+		}
+
 		void SetProfile(int32 InProfile)
 		{
 			ProfileLevel.Profile = InProfile;
@@ -295,12 +305,22 @@ namespace Electra
 			return ProfileLevel.Level;
 		}
 
-		void SetProfileConstraints(int32 InConstraints)
+		void SetProfileCompatibilityFlags(uint32 InCompatibilityFlags)
+		{
+			ProfileLevel.CompatibilityFlags = InCompatibilityFlags;
+		}
+
+		uint32 GetProfileCompatibilityFlags() const
+		{
+			return ProfileLevel.CompatibilityFlags;
+		}
+
+		void SetProfileConstraints(uint64 InConstraints)
 		{
 			ProfileLevel.Constraints = InConstraints;
 		}
 
-		int32 GetProfileConstraints() const
+		uint64 GetProfileConstraints() const
 		{
 			return ProfileLevel.Constraints;
 		}
@@ -459,19 +479,23 @@ namespace Electra
 			}
 			void Clear()
 			{
+				ProfileSpace = 0;
 				Profile = 0;
 				Level = 0;
+				CompatibilityFlags = 0;
 				Constraints = 0;
 				Tier = 0;
 			}
+			int32		ProfileSpace;
 			int32		Profile;
 			int32		Level;
-			int32		Constraints;
+			uint32		CompatibilityFlags;
+			uint64		Constraints;
 			int32		Tier;
 
 			bool operator == (const FProfileLevel& rhs) const
 			{
-				return Profile == rhs.Profile && Level == rhs.Level && Constraints == rhs.Constraints && Tier == rhs.Tier;
+				return ProfileSpace == rhs.ProfileSpace && Profile == rhs.Profile && Level == rhs.Level && Constraints == rhs.Constraints && Tier == rhs.Tier;
 			}
 
 			bool operator != (const FProfileLevel& rhs) const
diff --git a/Engine/Plugins/Media/ElectraPlayer/Source/ElectraPlayerRuntime/Private/Runtime/Utilities/Utilities.h b/Engine/Plugins/Media/ElectraPlayer/Source/ElectraPlayerRuntime/Private/Runtime/Utilities/Utilities.h
--- a/Engine/Plugins/Media/ElectraPlayer/Source/ElectraPlayerRuntime/Private/Runtime/Utilities/Utilities.h
+++ b/Engine/Plugins/Media/ElectraPlayer/Source/ElectraPlayerRuntime/Private/Runtime/Utilities/Utilities.h
@@ -31,6 +31,17 @@ namespace Electra
 			return a > b ? a : b;
 		}
 
+		inline uint32 BitReverse32(uint32 InValue)
+		{
+			uint32 rev = 0;
+			for(int32 i=0; i<32; ++i)
+			{
+				rev = (rev << 1) | (InValue & 1);
+				InValue >>= 1;
+			}
+			return rev;
+		}
+
 
 	} // namespace Utils
 
diff --git a/Engine/Plugins/Media/ElectraPlayer/Source/ElectraPlayerRuntime/Private/Runtime/Utilities/UtilsMPEGVideo.cpp b/Engine/Plugins/Media/ElectraPlayer/Source/ElectraPlayerRuntime/Private/Runtime/Utilities/UtilsMPEGVideo.cpp
--- a/Engine/Plugins/Media/ElectraPlayer/Source/ElectraPlayerRuntime/Private/Runtime/Utilities/UtilsMPEGVideo.cpp
+++ b/Engine/Plugins/Media/ElectraPlayer/Source/ElectraPlayerRuntime/Private/Runtime/Utilities/UtilsMPEGVideo.cpp
@@ -1,6 +1,7 @@
 // Copyright Epic Games, Inc. All Rights Reserved.
 
 #include "PlayerCore.h"
+#include "Utilities/Utilities.h"
 #include "Utilities/UtilsMPEGVideo.h"
 #include "BitDataStream.h"
 
@@ -9,20 +10,6 @@ namespace Electra
 {
 	namespace MPEG
 	{
-
-		FAVCDecoderConfigurationRecord::FAVCDecoderConfigurationRecord()
-		{
-			ConfigurationVersion = 0;
-			AVCProfileIndication = 0;
-			ProfileCompatibility = 0;
-			AVCLevelIndication = 0;
-			NALUnitLength = 0;
-			ChromaFormat = 0;
-			BitDepthLumaMinus8 = 0;
-			BitDepthChromaMinus8 = 0;
-			bHaveAdditionalProfileIndication = false;
-		}
-
 		void FAVCDecoderConfigurationRecord::SetRawData(const void* Data, int64 Size)
 		{
 			RawData.Empty();
@@ -247,6 +234,173 @@ namespace Electra
 
 
 
+		void FHEVCDecoderConfigurationRecord::SetRawData(const void* Data, int64 Size)
+		{
+			RawData.Empty();
+			if (Size)
+			{
+				RawData.Reserve((uint32)Size);
+				RawData.SetNumUninitialized((uint32)Size);
+				FMemory::Memcpy(RawData.GetData(), Data, Size);
+			}
+		}
+
+		const TArray<uint8>& FHEVCDecoderConfigurationRecord::GetRawData() const
+		{
+			return RawData;
+		}
+
+		const TArray<uint8>& FHEVCDecoderConfigurationRecord::GetCodecSpecificData() const
+		{
+			return CodecSpecificData;
+		}
+
+		int32 FHEVCDecoderConfigurationRecord::GetNumberOfSPS() const
+		{
+			return ParsedSPSs.Num();
+		}
+		
+		const FISO23008_2_seq_parameter_set_data& FHEVCDecoderConfigurationRecord::GetParsedSPS(int32 SpsIndex) const
+		{
+			return ParsedSPSs[SpsIndex];
+		}
+
+		void FHEVCDecoderConfigurationRecord::Reset()
+		{
+			RawData.Empty();
+			CodecSpecificData.Empty();
+			Arrays.Empty();
+			ConfigurationVersion = 0;
+			GeneralProfileSpace = 0;
+			GeneralTierFlag = 0;
+			GeneralProfileIDC = 0;
+			GeneralProfileCompatibilityFlags = 0;
+			GeneralConstraintIndicatorFlags = 0;
+			GeneralLevelIDC = 0;
+			MinSpatialSegmentationIDC = 0;
+			ParallelismType = 0;
+			ChromaFormat = 0;
+			BitDepthLumaMinus8 = 0;
+			BitDepthChromaMinus8 = 0;
+			AverageFrameRate = 0;
+			ConstantFrameRate = 0;
+			NumTemporalLayers = 0;
+			TemporalIdNested = 0;
+			NALUnitLengthMinus1 = 0;
+			ParsedSPSs.Empty();
+		}
+
+		bool FHEVCDecoderConfigurationRecord::Parse()
+		{
+			CodecSpecificData.Empty();
+
+			if (RawData.Num() >= 23)
+			{
+				FBitDataStream BitReader(RawData.GetData(), RawData.Num());
+
+				ConfigurationVersion = (uint8)BitReader.GetBits(8);
+				if (ConfigurationVersion != 1)
+				{
+					return false;
+				}
+				GeneralProfileSpace = (uint8)BitReader.GetBits(2);
+				GeneralTierFlag = (uint8)BitReader.GetBits(1);
+				GeneralProfileIDC = (uint8)BitReader.GetBits(5);
+				GeneralProfileCompatibilityFlags  = BitReader.GetBits(32);
+				for(int32 i=0; i<48; ++i)
+				{
+					GeneralConstraintIndicatorFlags = (GeneralConstraintIndicatorFlags << 1U) | uint64(BitReader.GetBits(1));
+				}
+				GeneralLevelIDC = (uint8)BitReader.GetBits(8);
+				if (BitReader.GetBits(4) != 15)
+				{
+					return false;
+				}
+				MinSpatialSegmentationIDC = (uint16)BitReader.GetBits(12);
+				if (BitReader.GetBits(6) != 63)
+				{
+					return false;
+				}
+				ParallelismType = (uint8)BitReader.GetBits(2);
+				if (BitReader.GetBits(6) != 63)
+				{
+					return false;
+				}
+				ChromaFormat = (uint8)BitReader.GetBits(2);
+				if (BitReader.GetBits(5) != 31)
+				{
+					return false;
+				}
+				BitDepthLumaMinus8 = (uint8)BitReader.GetBits(3);
+				if (BitReader.GetBits(5) != 31)
+				{
+					return false;
+				}
+				BitDepthChromaMinus8 = (uint8)BitReader.GetBits(3);
+				AverageFrameRate = (uint16)BitReader.GetBits(16);
+				ConstantFrameRate = (uint8)BitReader.GetBits(2);
+				NumTemporalLayers = (uint8)BitReader.GetBits(3);
+				TemporalIdNested = (uint8)BitReader.GetBits(1);
+				NALUnitLengthMinus1 = (uint8)BitReader.GetBits(2);
+				uint32 numArrays = BitReader.GetBits(8);
+				Arrays.Reserve(numArrays);
+				for(uint32 i=0; i<numArrays; ++i)
+				{
+					FArray &a = Arrays.AddDefaulted_GetRef();
+					a.Completeness = (uint8)BitReader.GetBits(1);
+					if (BitReader.GetBits(1) != 0)
+					{
+						return false;
+					}
+					a.NALUnitType = (uint8)BitReader.GetBits(6);
+					uint32 numNALUs = BitReader.GetBits(16);
+					a.NALUs.Reserve(numNALUs);
+					for(uint32 j=0; j<numNALUs; ++j)
+					{
+						uint16 naluLen = (uint16)BitReader.GetBits(16);
+						TArray<uint8> &n = a.NALUs.AddDefaulted_GetRef();
+						n.Reserve(naluLen);
+						for(uint32 k=0; k<naluLen; ++k)
+						{
+							n.Emplace((uint8)BitReader.GetBits(8));
+						}
+					}
+				}
+				if (BitReader.GetRemainingBits() != 0)
+				{
+					return false;
+				}
+
+				// CSD
+				CodecSpecificData.Empty();
+				for(int32 i=0; i<Arrays.Num(); ++i)
+				{
+					const FArray& a = Arrays[i];
+					// SPS nut?
+					if (a.NALUnitType == 33)
+					{
+						FISO23008_2_seq_parameter_set_data sps;
+						if (ParseH265SPS(sps, a.NALUs[0].GetData(), a.NALUs[0].Num()))
+						{
+							ParsedSPSs.Emplace(MoveTemp(sps));
+						}
+					}
+
+					for(int32 j=0; j<a.NALUs.Num(); ++j)
+					{
+						CodecSpecificData.Add(0);
+						CodecSpecificData.Add(0);
+						CodecSpecificData.Add(0);
+						CodecSpecificData.Add(1);
+						CodecSpecificData.Append(a.NALUs[j]);
+					}
+				}
+				return true;
+			}
+			return false;
+		}
+
+
 
 		static int32 FindStartCode(const uint8* InData, SIZE_T InDataSize, int32& NALUnitLength)
 		{
@@ -535,6 +689,356 @@ namespace Electra
 		}
 
 
+		//! Parses a H.265 (ISO/IEC 23008-2) SPS NALU.
+		bool ParseH265SPS(FISO23008_2_seq_parameter_set_data& OutSPS, const void* Data, int32 Size)
+		{
+			struct FSyntaxElement
+			{
+				static uint32 ue_v(FBitDataStream& BitStream)
+				{
+					int32 lz = -1;
+					for(uint32 b = 0; b == 0; ++lz)
+					{
+						b = BitStream.GetBits(1);
+					}
+					if (lz)
+					{
+						return ((1 << lz) | BitStream.GetBits(lz)) - 1;
+					}
+					return 0;
+				}
+				static int32 se_v(FBitDataStream& BitStream)
+				{
+					uint32 c = ue_v(BitStream);
+					return c & 1 ? int32((c + 1) >> 1) : -int32((c + 1) >> 1);
+				}
+			};
+
+			// SPS is usually an EBSP so we need to strip it down.
+			FScopedDataPtr<uint8> RBSP(FMemory::Malloc(Size));
+			int32 RBSPsize = EBSPtoRBSP(RBSP, static_cast<const uint8*>(Data), Size);
+			FBitDataStream BitReader(RBSP, RBSPsize);
+
+			FMemory::Memzero(OutSPS);
+
+			if (BitReader.GetBits(1) != 0)	// forbidden_zero_bit
+				return false;
+			if (BitReader.GetBits(6) != 33)	// sps_nut ?
+				return false;
+			BitReader.SkipBits(6);			// nuh_layer_id
+			BitReader.SkipBits(3);			// nuh_temporal_id_plus1
+
+			OutSPS.sps_video_parameter_set_id = BitReader.GetBits(4);
+			OutSPS.sps_max_sub_layers_minus1 = BitReader.GetBits(3);
+			OutSPS.sps_temporal_id_nesting_flag = BitReader.GetBits(1);
+
+			// profile_tier_level( sps_max_sub_layers_minus1 )
+			OutSPS.general_profile_space = BitReader.GetBits(2);
+			OutSPS.general_tier_flag = BitReader.GetBits(1);
+			OutSPS.general_profile_idc = BitReader.GetBits(5);
+			for(int32 i=0; i<32; ++i)
+			{
+				OutSPS.general_profile_compatibility_flag = (OutSPS.general_profile_compatibility_flag << 1) | BitReader.GetBits(1);
+			}
+			OutSPS.general_progressive_source_flag = BitReader.GetBits(1);
+			OutSPS.general_interlaced_source_flag = BitReader.GetBits(1);
+			OutSPS.general_non_packed_constraint_flag = BitReader.GetBits(1);
+			OutSPS.general_frame_only_constraint_flag = BitReader.GetBits(1);
+			// We skip over the next 44 bits that vary in layout depending on general_profile_idc and general_profile_compatibility_flag.
+			// For the moment none of the flags interest us, so we do not parse them for simplicity.
+			OutSPS.general_reserved_zero_44bits = ((uint64)BitReader.GetBits(32) << 32) | BitReader.GetBits(12);
+
+			OutSPS.general_level_idc = BitReader.GetBits(8);
+			for(uint32 i=0; i<OutSPS.sps_max_sub_layers_minus1; ++i)
+			{
+				OutSPS.sub_layer_profile_present_flag[i] = BitReader.GetBits(1);
+				OutSPS.sub_layer_level_present_flag[i] = BitReader.GetBits(1);
+			}
+			if (OutSPS.sps_max_sub_layers_minus1 > 0)
+			{
+				for(uint32 i = OutSPS.sps_max_sub_layers_minus1; i<8; ++i)
+				{
+					BitReader.SkipBits(2);		// reserved_zero_2bits[ i ]
+				}
+			}
+			for(uint32 i=0; i<OutSPS.sps_max_sub_layers_minus1; ++i)
+			{
+				// Similarly, we skip all of these flags that are of no interest to us right now.
+				if (OutSPS.sub_layer_profile_present_flag[i])
+				{
+					BitReader.SkipBits(2);		// sub_layer_profile_space[ i ]
+					BitReader.SkipBits(1);		// sub_layer_tier_flag[ i ]
+					BitReader.SkipBits(5);		// sub_layer_profile_idc[ i ]
+					BitReader.SkipBits(32);		// for( j = 0; j < 32; j++ ) sub_layer_profile_compatibility_flag[ i ][ j ] u(1)
+					BitReader.SkipBits(1);		// sub_layer_progressive_source_flag[ i ]
+					BitReader.SkipBits(1);		// sub_layer_interlaced_source_flag[ i ]
+					BitReader.SkipBits(1);		// sub_layer_non_packed_constraint_flag[ i ]
+					BitReader.SkipBits(1);		// sub_layer_frame_only_constraint_flag[ i ]
+					BitReader.SkipBits(44);		// sub_layer_reserved_zero_44bits[ i ]
+				}
+				if (OutSPS.sub_layer_level_present_flag[i])
+				{
+					BitReader.SkipBits(8);		// sub_layer_level_idc[ i ]
+				}
+			}
+
+			OutSPS.sps_seq_parameter_set_id = FSyntaxElement::ue_v(BitReader);
+			OutSPS.chroma_format_idc = FSyntaxElement::ue_v(BitReader);
+			if (OutSPS.chroma_format_idc == 3)
+			{
+				OutSPS.separate_colour_plane_flag = BitReader.GetBits(1);			// separate_colour_plane_flag
+			}
+			OutSPS.pic_width_in_luma_samples = FSyntaxElement::ue_v(BitReader);
+			OutSPS.pic_height_in_luma_samples = FSyntaxElement::ue_v(BitReader);
+
+			OutSPS.conformance_window_flag = BitReader.GetBits(1);
+			if (OutSPS.conformance_window_flag)
+			{
+				OutSPS.conf_win_left_offset = FSyntaxElement::ue_v(BitReader);
+				OutSPS.conf_win_right_offset = FSyntaxElement::ue_v(BitReader);
+				OutSPS.conf_win_top_offset = FSyntaxElement::ue_v(BitReader);
+				OutSPS.conf_win_bottom_offset = FSyntaxElement::ue_v(BitReader);
+			}
+			OutSPS.bit_depth_luma_minus8 = FSyntaxElement::ue_v(BitReader);
+			OutSPS.bit_depth_chroma_minus8 = FSyntaxElement::ue_v(BitReader);
+			OutSPS.log2_max_pic_order_cnt_lsb_minus4 = FSyntaxElement::ue_v(BitReader);
+			OutSPS.sps_sub_layer_ordering_info_present_flag = BitReader.GetBits(1);
+			for(uint32 i=(OutSPS.sps_sub_layer_ordering_info_present_flag ? 0 : OutSPS.sps_max_sub_layers_minus1); i<=OutSPS.sps_max_sub_layers_minus1; ++i)
+			{
+				OutSPS.sps_max_dec_pic_buffering_minus1[i] = FSyntaxElement::ue_v(BitReader);
+				OutSPS.sps_max_num_reorder_pics[i] = FSyntaxElement::ue_v(BitReader);
+				OutSPS.sps_max_latency_increase_plus1[i] = FSyntaxElement::ue_v(BitReader);
+			}
+			OutSPS.log2_min_luma_coding_block_size_minus3 = FSyntaxElement::ue_v(BitReader);
+			OutSPS.log2_diff_max_min_luma_coding_block_size = FSyntaxElement::ue_v(BitReader);
+			OutSPS.log2_min_transform_block_size_minus2 = FSyntaxElement::ue_v(BitReader);
+			OutSPS.log2_diff_max_min_transform_block_size = FSyntaxElement::ue_v(BitReader);
+			OutSPS.max_transform_hierarchy_depth_inter = FSyntaxElement::ue_v(BitReader);
+			OutSPS.max_transform_hierarchy_depth_intra = FSyntaxElement::ue_v(BitReader);
+			OutSPS.scaling_list_enabled_flag = BitReader.GetBits(1);
+			if (OutSPS.scaling_list_enabled_flag)
+			{
+				// Parse the scaling list into temp variables only. We do not need them right now.
+				OutSPS.sps_scaling_list_data_present_flag = BitReader.GetBits(1);
+				if (OutSPS.sps_scaling_list_data_present_flag)
+				{
+					// scaling_list_data( )
+					for(uint32 sizeId=0; sizeId<4; ++sizeId)
+					{
+						for(int32 matrixId=0; matrixId<((sizeId==3)?2:6); ++matrixId)
+						{
+							uint8 scaling_list_pred_mode_flag = BitReader.GetBits(1);	// [ sizeId ][ matrixId ] u(1)
+							if (!scaling_list_pred_mode_flag/*[sizeId][matrixId]*/)
+							{
+								uint32 scaling_list_pred_matrix_id_delta/*[sizeId][matrixId]*/ = FSyntaxElement::ue_v(BitReader);
+							}
+							else
+							{
+								//nextCoef = 8;
+								uint32 coefNum = Utils::Min((uint32)64U, (uint32)(1 << (4 + (sizeId << 1))));
+								if (sizeId > 1)
+								{
+									uint32 scaling_list_dc_coef_minus8/*[sizeId ? 2][matrixId]*/ = FSyntaxElement::se_v(BitReader);
+									//nextCoef = scaling_list_dc_coef_minus8[sizeId ? 2][matrixId] + 8;
+								}
+								for(uint32 i=0; i<coefNum; i++)
+								{
+									uint32 scaling_list_delta_coef = FSyntaxElement::se_v(BitReader);
+									//nextCoef = ( nextCoef + scaling_list_delta_coef + 256 ) % 256
+									//ScalingList[ sizeId ][ matrixId ][ i ] = nextCoef
+								}
+							}
+						}
+					}
+				}
+			}
+
+			OutSPS.amp_enabled_flag = BitReader.GetBits(1);
+			OutSPS.sample_adaptive_offset_enabled_flag = BitReader.GetBits(1);
+			OutSPS.pcm_enabled_flag = BitReader.GetBits(1);
+			if (OutSPS.pcm_enabled_flag)
+			{
+				OutSPS.pcm_sample_bit_depth_luma_minus1 = BitReader.GetBits(4);
+				OutSPS.pcm_sample_bit_depth_chroma_minus1 = BitReader.GetBits(4);
+				OutSPS.log2_min_pcm_luma_coding_block_size_minus3 = FSyntaxElement::ue_v(BitReader);
+				OutSPS.log2_diff_max_min_pcm_luma_coding_block_size = FSyntaxElement::ue_v(BitReader);
+				OutSPS.pcm_loop_filter_disabled_flag = BitReader.GetBits(1);
+			}
+
+			OutSPS.num_short_term_ref_pic_sets = FSyntaxElement::ue_v(BitReader);
+			// For now there must not be any as it it not exactly clear how to generate NumDeltaPocs[] that are needed
+			check(OutSPS.num_short_term_ref_pic_sets == 0);
+			if (OutSPS.num_short_term_ref_pic_sets != 0)
+			{
+				return false;
+			}
+		#if 0
+			for(uint32 i=0; i<OutSPS.num_short_term_ref_pic_sets; ++i)
+			{
+				//short_term_ref_pic_set(i);
+				struct short_term_ref_pic_set
+				{
+					static void parse(uint32 stRpsIdx, uint32 num_short_term_ref_pic_sets, FBitDataStream& BitReader)
+					{
+						uint8 inter_ref_pic_set_prediction_flag = 0;
+						if (stRpsIdx != 0)
+						{
+							inter_ref_pic_set_prediction_flag = BitReader.GetBits(1);
+						}
+						if (inter_ref_pic_set_prediction_flag)
+						{
+							uint32 delta_idx_minus1 = 0;
+							if (stRpsIdx == num_short_term_ref_pic_sets)
+							{
+								delta_idx_minus1 = FSyntaxElement::ue_v(BitReader);
+							}
+							BitReader.SkipBits(1);	// delta_rps_sign
+							uint32 abs_delta_rps_minus1 = FSyntaxElement::ue_v(BitReader);
+							uint32 RefRpsIdx = stRpsIdx - (delta_idx_minus1 + 1);
+							uint32 NumDeltaPocs[64] = {0};
+							for(uint32 j=0; j<=NumDeltaPocs[RefRpsIdx]; ++j)
+							{
+								uint8 used_by_curr_pic_flag/*[j]*/ = BitReader.GetBits(1);
+								if (!used_by_curr_pic_flag/*[j]*/)
+								{
+									BitReader.SkipBits(1);		// use_delta_flag[ j ]
+								}
+							}
+						}
+						else
+						{
+							uint32 num_negative_pics = FSyntaxElement::ue_v(BitReader);
+							uint32 num_positive_pics = FSyntaxElement::ue_v(BitReader);
+							for(uint32 j=0; j<num_negative_pics; ++j)
+							{
+								uint32 delta_poc_s0_minus1/*[j]*/ = FSyntaxElement::ue_v(BitReader);
+								BitReader.SkipBits(1);				// used_by_curr_pic_s0_flag[ j ]
+							}
+							for(uint32 j=0; j<num_positive_pics; ++j)
+							{
+								uint32 delta_poc_s1_minus1/*[j]*/ = FSyntaxElement::ue_v(BitReader);
+								BitReader.SkipBits(1);				// used_by_curr_pic_s1_flag[ j ]
+							}
+						}
+					}
+				};
+				short_term_ref_pic_set::parse(i, OutSPS.num_short_term_ref_pic_sets, BitReader);
+			}
+		#endif
+			
+			uint8 long_term_ref_pics_present_flag = BitReader.GetBits(1);
+			if (long_term_ref_pics_present_flag)
+			{
+				uint32 num_long_term_ref_pics_sps = FSyntaxElement::ue_v(BitReader);
+				for(uint32 i=0; i<num_long_term_ref_pics_sps; ++i)
+				{
+					uint32 lt_ref_pic_poc_lsb_sps = FSyntaxElement::ue_v(BitReader);	//[ i ] u(v)
+					BitReader.SkipBits(1);							//	used_by_curr_pic_lt_sps_flag[ i ]
+				}
+			}
+			
+			OutSPS.sps_temporal_mvp_enabled_flag = BitReader.GetBits(1);
+			OutSPS.strong_intra_smoothing_enabled_flag = BitReader.GetBits(1);
+			OutSPS.vui_parameters_present_flag = BitReader.GetBits(1);
+			if (OutSPS.vui_parameters_present_flag)
+			{
+				// vui_parameters()
+				OutSPS.aspect_ratio_info_present_flag = BitReader.GetBits(1);
+				if (OutSPS.aspect_ratio_info_present_flag)
+				{
+					OutSPS.aspect_ratio_idc = BitReader.GetBits(8);
+					if (OutSPS.aspect_ratio_idc == 255)
+					{
+						OutSPS.sar_width = BitReader.GetBits(16);
+						OutSPS.sar_height = BitReader.GetBits(16);
+					}
+				}
+				OutSPS.overscan_info_present_flag = BitReader.GetBits(1);
+				if (OutSPS.overscan_info_present_flag)
+				{
+					OutSPS.overscan_appropriate_flag = BitReader.GetBits(1);
+				}
+				OutSPS.video_signal_type_present_flag = BitReader.GetBits(1);
+				if (OutSPS.video_signal_type_present_flag)
+				{
+					OutSPS.video_format = BitReader.GetBits(3);
+					OutSPS.video_full_range_flag = BitReader.GetBits(1);
+					OutSPS.colour_description_present_flag = BitReader.GetBits(1);
+					if (OutSPS.colour_description_present_flag)
+					{
+						OutSPS.colour_primaries = BitReader.GetBits(8);
+						OutSPS.transfer_characteristics = BitReader.GetBits(8);
+						OutSPS.matrix_coeffs = BitReader.GetBits(8);
+					}
+				}
+				OutSPS.chroma_loc_info_present_flag = BitReader.GetBits(1);
+				if (OutSPS.chroma_loc_info_present_flag)
+				{
+					OutSPS.chroma_sample_loc_type_top_field = FSyntaxElement::ue_v(BitReader);
+					OutSPS.chroma_sample_loc_type_bottom_field = FSyntaxElement::ue_v(BitReader);
+				}
+				OutSPS.neutral_chroma_indication_flag = BitReader.GetBits(1);
+				OutSPS.field_seq_flag = BitReader.GetBits(1);
+				OutSPS.frame_field_info_present_flag = BitReader.GetBits(1);
+				OutSPS.default_display_window_flag = BitReader.GetBits(1);
+				if (OutSPS.default_display_window_flag)
+				{
+					OutSPS.def_disp_win_left_offset = FSyntaxElement::ue_v(BitReader);
+					OutSPS.def_disp_win_right_offset = FSyntaxElement::ue_v(BitReader);
+					OutSPS.def_disp_win_top_offset = FSyntaxElement::ue_v(BitReader);
+					OutSPS.def_disp_win_bottom_offset = FSyntaxElement::ue_v(BitReader);
+				}
+				OutSPS.vui_timing_info_present_flag = BitReader.GetBits(1);
+				if (OutSPS.vui_timing_info_present_flag)
+				{
+					OutSPS.vui_num_units_in_tick = BitReader.GetBits(32);
+					OutSPS.vui_time_scale = BitReader.GetBits(32);
+					OutSPS.vui_poc_proportional_to_timing_flag = BitReader.GetBits(1);
+					if (OutSPS.vui_poc_proportional_to_timing_flag)
+					{
+						OutSPS.vui_num_ticks_poc_diff_one_minus1 = FSyntaxElement::ue_v(BitReader);
+// None of the following is of interest to us now, so we stop parsing at this point.
+					}
+				}
+			return true;
+			}
+		return false;
+		}
+
+
+		uint64 FISO23008_2_seq_parameter_set_data::GetConstraintFlags() const
+		{
+			uint64 ConstraintFlags = ((uint64)general_progressive_source_flag << 3) | ((uint64)general_interlaced_source_flag << 2) | ((uint64)general_non_packed_constraint_flag << 1) | ((uint64)general_frame_only_constraint_flag << 44) | general_reserved_zero_44bits;
+			return ConstraintFlags;
+		}
+
+		FString FISO23008_2_seq_parameter_set_data::GetRFC6381(const TCHAR* SampleTypePrefix) const
+		{
+			// As per ISO/IEC 14496-15:2014 Annex E.3
+			uint32 cf = Utils::BitReverse32(general_profile_compatibility_flag);
+			uint64 ConstraintFlags = GetConstraintFlags();
+			FString cfs;
+			bool bNonZero = false;
+			for(int32 i=5; i>=0; --i, ConstraintFlags >>= 8)
+			{
+				if ((ConstraintFlags & 255) == 0 && !bNonZero)
+				{
+					continue;
+				}
+				bNonZero = true;
+				cfs = FString::Printf(TEXT(".%02X"), (uint8)(ConstraintFlags & 255)) + cfs;
+			}
+
+			if (general_profile_space == 0)
+			{
+				return FString::Printf(TEXT("%s.%d.%X%c%d%s"), SampleTypePrefix, general_profile_idc, cf, general_tier_flag ? TCHAR('H') : TCHAR('L'), general_level_idc, *cfs);
+			}
+			else
+			{
+				return FString::Printf(TEXT("%s.%c%d.%X%c%d%s"), SampleTypePrefix, TCHAR('A')+general_profile_space-1, general_profile_idc, cf, general_tier_flag ? TCHAR('H') : TCHAR('L'), general_level_idc, *cfs);
+			}
+		}
+
 
 	} // namespace MPEG
 } // namespace Electra
diff --git a/Engine/Plugins/Media/ElectraPlayer/Source/ElectraPlayerRuntime/Private/Runtime/Utilities/UtilsMPEGVideo.h b/Engine/Plugins/Media/ElectraPlayer/Source/ElectraPlayerRuntime/Private/Runtime/Utilities/UtilsMPEGVideo.h
--- a/Engine/Plugins/Media/ElectraPlayer/Source/ElectraPlayerRuntime/Private/Runtime/Utilities/UtilsMPEGVideo.h
+++ b/Engine/Plugins/Media/ElectraPlayer/Source/ElectraPlayerRuntime/Private/Runtime/Utilities/UtilsMPEGVideo.h
@@ -20,61 +20,61 @@ namespace Electra
 		 */
 		struct FISO14496_10_seq_parameter_set_data
 		{
-			uint8	profile_idc;
-			uint8	constraint_set0_flag;
-			uint8	constraint_set1_flag;
-			uint8	constraint_set2_flag;
-			uint8	constraint_set3_flag;
-			uint8	constraint_set4_flag;
-			uint8	constraint_set5_flag;
-			uint8	level_idc;
-			uint32	seq_parameter_set_id;
-			uint32	chroma_format_idc;
-			uint8	separate_colour_plane_flag;
-			uint32	bit_depth_luma_minus8;
-			uint32	bit_depth_chroma_minus8;
-			uint8	qpprime_y_zero_transform_bypass_flag;
-			uint8	seq_scaling_matrix_present_flag;
-			uint32	log2_max_frame_num_minus4;
-			uint32	pic_order_cnt_type;
-			uint32	log2_max_pic_order_cnt_lsb_minus4;
-			uint32	delta_pic_order_always_zero_flag;
-			int32	offset_for_non_ref_pic;
-			int32	offset_for_top_to_bottom_field;
-			uint32	num_ref_frames_in_pic_order_cnt_cycle;
-			uint32	max_num_ref_frames;
-			uint8	gaps_in_frame_num_value_allowed_flag;
-			uint32	pic_width_in_mbs_minus1;
-			uint32	pic_height_in_map_units_minus1;
-			uint8	frame_mbs_only_flag;
-			uint8	mb_adaptive_frame_field_flag;
-			uint8	direct_8x8_inference_flag;
-			uint8	frame_cropping_flag;
-			uint32	frame_crop_left_offset;
-			uint32	frame_crop_right_offset;
-			uint32	frame_crop_top_offset;
-			uint32	frame_crop_bottom_offset;
-			uint8	vui_parameters_present_flag;
-			uint8	aspect_ratio_info_present_flag;
-			uint8	aspect_ratio_idc;
-			uint16	sar_width;
-			uint16	sar_height;
-			uint8	overscan_info_present_flag;
-			uint8	overscan_appropriate_flag;
-			uint8	video_signal_type_present_flag;
-			uint8	video_format;
-			uint8	video_full_range_flag;
-			uint8	colour_description_present_flag;
-			uint8	colour_primaries;
-			uint8	transfer_characteristics;
-			uint8	matrix_coefficients;
-			uint8	chroma_loc_info_present_flag;
-			uint32	chroma_sample_loc_type_top_field;
-			uint32	chroma_sample_loc_type_bottom_field;
-			uint8	timing_info_present_flag;
-			uint32	num_units_in_tick;
-			uint32	time_scale;
-			uint8	fixed_frame_rate_flag;
+			uint8	profile_idc = 0;
+			uint8	constraint_set0_flag = 0;
+			uint8	constraint_set1_flag = 0;
+			uint8	constraint_set2_flag = 0;
+			uint8	constraint_set3_flag = 0;
+			uint8	constraint_set4_flag = 0;
+			uint8	constraint_set5_flag = 0;
+			uint8	level_idc = 0;
+			uint32	seq_parameter_set_id = 0;
+			uint32	chroma_format_idc = 0;
+			uint8	separate_colour_plane_flag = 0;
+			uint32	bit_depth_luma_minus8 = 0;
+			uint32	bit_depth_chroma_minus8 = 0;
+			uint8	qpprime_y_zero_transform_bypass_flag = 0;
+			uint8	seq_scaling_matrix_present_flag = 0;
+			uint32	log2_max_frame_num_minus4 = 0;
+			uint32	pic_order_cnt_type = 0;
+			uint32	log2_max_pic_order_cnt_lsb_minus4 = 0;
+			uint32	delta_pic_order_always_zero_flag = 0;
+			int32	offset_for_non_ref_pic = 0;
+			int32	offset_for_top_to_bottom_field = 0;
+			uint32	num_ref_frames_in_pic_order_cnt_cycle = 0;
+			uint32	max_num_ref_frames = 0;
+			uint8	gaps_in_frame_num_value_allowed_flag = 0;
+			uint32	pic_width_in_mbs_minus1 = 0;
+			uint32	pic_height_in_map_units_minus1 = 0;
+			uint8	frame_mbs_only_flag = 0;
+			uint8	mb_adaptive_frame_field_flag = 0;
+			uint8	direct_8x8_inference_flag = 0;
+			uint8	frame_cropping_flag = 0;
+			uint32	frame_crop_left_offset = 0;
+			uint32	frame_crop_right_offset = 0;
+			uint32	frame_crop_top_offset = 0;
+			uint32	frame_crop_bottom_offset = 0;
+			uint8	vui_parameters_present_flag = 0;
+			uint8	aspect_ratio_info_present_flag = 0;
+			uint8	aspect_ratio_idc = 0;
+			uint16	sar_width = 0;
+			uint16	sar_height = 0;
+			uint8	overscan_info_present_flag = 0;
+			uint8	overscan_appropriate_flag = 0;
+			uint8	video_signal_type_present_flag = 0;
+			uint8	video_format = 0;
+			uint8	video_full_range_flag = 0;
+			uint8	colour_description_present_flag = 0;
+			uint8	colour_primaries = 0;
+			uint8	transfer_characteristics = 0;
+			uint8	matrix_coefficients = 0;
+			uint8	chroma_loc_info_present_flag = 0;
+			uint32	chroma_sample_loc_type_top_field = 0;
+			uint32	chroma_sample_loc_type_bottom_field = 0;
+			uint8	timing_info_present_flag = 0;
+			uint32	num_units_in_tick = 0;
+			uint32	time_scale = 0;
+			uint8	fixed_frame_rate_flag = 0;
 
 			int32 GetWidth() const
 			{
@@ -153,7 +153,6 @@ namespace Electra
 		class FAVCDecoderConfigurationRecord
 		{
 		public:
-			FAVCDecoderConfigurationRecord();
 			void SetRawData(const void* Data, int64 Size);
 			const TArray<uint8>& GetRawData() const;
 
@@ -174,20 +173,244 @@ namespace Electra
 			TArray<TArray<uint8>>											SequenceParameterSets;
 			TArray<TArray<uint8>>											PictureParameterSets;
 			TArray<TArray<uint8>>											SequenceParameterSetsExt;
-			uint8															ConfigurationVersion;
-			uint8															AVCProfileIndication;
-			uint8															ProfileCompatibility;
-			uint8															AVCLevelIndication;
-			uint8															NALUnitLength;
-			uint8															ChromaFormat;
-			uint8															BitDepthLumaMinus8;
-			uint8															BitDepthChromaMinus8;
-			bool															bHaveAdditionalProfileIndication;
+			uint8															ConfigurationVersion = 0;
+			uint8															AVCProfileIndication = 0;
+			uint8															ProfileCompatibility = 0;
+			uint8															AVCLevelIndication = 0;
+			uint8															NALUnitLength = 0;
+			uint8															ChromaFormat = 0;
+			uint8															BitDepthLumaMinus8 = 0;
+			uint8															BitDepthChromaMinus8 = 0;
+			bool															bHaveAdditionalProfileIndication = false;
 
 			TArray<FISO14496_10_seq_parameter_set_data>						ParsedSPSs;
 		};
 
 
+		struct FISO23008_2_seq_parameter_set_data
+		{
+			uint8		sps_video_parameter_set_id = 0;
+			uint8		sps_max_sub_layers_minus1 = 0;
+			uint8		sps_temporal_id_nesting_flag = 0;
+			uint8		general_profile_space = 0;
+			uint8		general_tier_flag = 0;
+			uint8		general_profile_idc = 0;
+			uint32		general_profile_compatibility_flag = 0;
+			uint8		general_progressive_source_flag = 0;
+			uint8		general_interlaced_source_flag = 0;
+			uint8		general_non_packed_constraint_flag = 0;
+			uint8		general_frame_only_constraint_flag = 0;
+			uint64		general_reserved_zero_44bits = 0;
+			uint8		general_level_idc = 0;
+			uint8		sub_layer_profile_present_flag[8] = {0};
+			uint8		sub_layer_level_present_flag[8] = {0};
+			uint32		sps_seq_parameter_set_id = 0;
+			uint32		chroma_format_idc = 0;
+			uint8		separate_colour_plane_flag = 0;
+			uint32		pic_width_in_luma_samples = 0;
+			uint32		pic_height_in_luma_samples = 0;
+			uint8		conformance_window_flag = 0;
+			uint32		conf_win_left_offset = 0;
+			uint32		conf_win_right_offset = 0;
+			uint32		conf_win_top_offset = 0;
+			uint32		conf_win_bottom_offset = 0;
+			uint32		bit_depth_luma_minus8 = 0;
+			uint32		bit_depth_chroma_minus8 = 0;
+			uint32		log2_max_pic_order_cnt_lsb_minus4 = 0;
+			uint8		sps_sub_layer_ordering_info_present_flag = 0;
+			uint32		sps_max_dec_pic_buffering_minus1[8] = {0};
+			uint32		sps_max_num_reorder_pics[8] = {0};
+			uint32		sps_max_latency_increase_plus1[8] = {0};
+			uint32		log2_min_luma_coding_block_size_minus3 = 0;
+			uint32		log2_diff_max_min_luma_coding_block_size = 0;
+			uint32		log2_min_transform_block_size_minus2 = 0;
+			uint32		log2_diff_max_min_transform_block_size = 0;
+			uint32		max_transform_hierarchy_depth_inter = 0;
+			uint32		max_transform_hierarchy_depth_intra = 0;
+			uint8		scaling_list_enabled_flag = 0;
+			uint8		sps_scaling_list_data_present_flag = 0;
+			uint8		amp_enabled_flag = 0;
+			uint8		sample_adaptive_offset_enabled_flag = 0;
+			uint8		pcm_enabled_flag = 0;
+			uint8		pcm_sample_bit_depth_luma_minus1 = 0;
+			uint8		pcm_sample_bit_depth_chroma_minus1 = 0;
+			uint32		log2_min_pcm_luma_coding_block_size_minus3 = 0;
+			uint32		log2_diff_max_min_pcm_luma_coding_block_size = 0;
+			uint8		pcm_loop_filter_disabled_flag = 0;
+			uint32		num_short_term_ref_pic_sets = 0;
+			uint8		sps_temporal_mvp_enabled_flag = 0;
+			uint8		strong_intra_smoothing_enabled_flag = 0;
+			uint8		vui_parameters_present_flag = 0;
+			uint8		aspect_ratio_info_present_flag = 0;
+			uint8		aspect_ratio_idc = 0;
+			uint16		sar_width = 0;
+			uint16		sar_height = 0;
+			uint8		overscan_info_present_flag = 0;
+			uint8		overscan_appropriate_flag = 0;
+			uint8		video_signal_type_present_flag = 0;
+			uint8		video_format = 0;
+			uint8		video_full_range_flag = 0;
+			uint8		colour_description_present_flag = 0;
+			uint8		colour_primaries = 0;
+			uint8		transfer_characteristics = 0;
+			uint8		matrix_coeffs = 0;
+			uint8		chroma_loc_info_present_flag = 0;
+			uint32		chroma_sample_loc_type_top_field = 0;
+			uint32		chroma_sample_loc_type_bottom_field = 0;
+			uint8		neutral_chroma_indication_flag = 0;
+			uint8		field_seq_flag = 0;
+			uint8		frame_field_info_present_flag = 0;
+			uint8		default_display_window_flag = 0;
+			uint32		def_disp_win_left_offset = 0;
+			uint32		def_disp_win_right_offset = 0;
+			uint32		def_disp_win_top_offset = 0;
+			uint32		def_disp_win_bottom_offset = 0;
+			uint8		vui_timing_info_present_flag = 0;
+			uint32		vui_num_units_in_tick = 0;
+			uint32		vui_time_scale = 0;
+			uint8		vui_poc_proportional_to_timing_flag = 0;
+			uint32		vui_num_ticks_poc_diff_one_minus1 = 0;
+
+
+			int32 GetWidth() const
+			{
+				return (int32) pic_width_in_luma_samples;
+			}
+
+			int32 GetHeight() const
+			{
+				return (int32) pic_height_in_luma_samples;
+			}
+
+			void GetCrop(int32& Left, int32& Right, int32& Top, int32& Bottom) const
+			{
+				if (conformance_window_flag)
+				{
+					uint32 SubWidthC = 0;
+					uint32 SubHeightC = 0;
+					switch(chroma_format_idc)
+					{
+						default:
+						case 0:
+							SubWidthC = 1; 
+							SubHeightC = 1;	
+							break;
+						case 1:	
+							SubWidthC = 2; 
+							SubHeightC = 2;	
+							break;
+						case 2:	
+							SubWidthC = 2; 
+							SubHeightC = 1;	
+							break;
+						case 3:	
+							SubWidthC = 1; 
+							SubHeightC = 1;	
+							break;
+					}
+					Left = conf_win_left_offset * SubWidthC;
+					Right = conf_win_right_offset * SubWidthC;
+					Top	= conf_win_top_offset * SubHeightC;
+					Bottom = conf_win_bottom_offset * SubHeightC;
+				}
+				else
+				{
+					Left = Right = Top = Bottom = 0;
+				}
+			}
+
+			void GetAspect(int32& SarW, int32& SarH) const
+			{
+				if (vui_parameters_present_flag && aspect_ratio_info_present_flag)
+				{
+					switch (aspect_ratio_idc)
+					{
+						default:	SarW = SarH = 0;		break;
+						case 0:		SarW = SarH = 0;		break;
+						case 1:		SarW = SarH = 1;		break;
+						case 2:		SarW = 12; SarH = 11;	break;
+						case 3:		SarW = 10; SarH = 11;	break;
+						case 4:		SarW = 16; SarH = 11;	break;
+						case 5:		SarW = 40; SarH = 33;	break;
+						case 6:		SarW = 24; SarH = 11;	break;
+						case 7:		SarW = 20; SarH = 11;	break;
+						case 8:		SarW = 32; SarH = 11;	break;
+						case 9:		SarW = 80; SarH = 33;	break;
+						case 10:	SarW = 18; SarH = 11;	break;
+						case 11:	SarW = 15; SarH = 11;	break;
+						case 12:	SarW = 64; SarH = 33;	break;
+						case 13:	SarW = 160; SarH = 99;	break;
+						case 14:	SarW = 4; SarH = 3;		break;
+						case 15:	SarW = 3; SarH = 2;		break;
+						case 16:	SarW = 2; SarH = 1;		break;
+						case 255:	SarW = sar_width; SarH = sar_height;	break;
+					}
+				}
+				else
+				{
+					SarW = SarH = 1;
+				}
+			}
+			FTimeFraction GetTiming() const
+			{
+				if (vui_parameters_present_flag && vui_timing_info_present_flag && vui_time_scale)
+				{
+					return FTimeFraction(vui_num_units_in_tick, vui_time_scale);
+				}
+				return FTimeFraction();
+			}
+
+			uint64 GetConstraintFlags() const;
+			FString GetRFC6381(const TCHAR* SampleTypePrefix) const;
+
+		};
+
+		class FHEVCDecoderConfigurationRecord
+		{
+		public:
+			void SetRawData(const void* Data, int64 Size);
+			const TArray<uint8>& GetRawData() const;
+
+			void Reset();
+
+			bool Parse();
+
+			const TArray<uint8>& GetCodecSpecificData() const;
+
+			int32 GetNumberOfSPS() const;
+			const FISO23008_2_seq_parameter_set_data& GetParsedSPS(int32 SpsIndex) const;
+
+		private:
+			struct FArray
+			{
+				uint8					Completeness = 0;
+				uint8					NALUnitType = 0;
+				TArray<TArray<uint8>>	NALUs;
+			};
+			TArray<uint8>								RawData;
+			TArray<uint8>								CodecSpecificData;
+			TArray<FArray>								Arrays;
+			uint8										ConfigurationVersion = 0;
+			uint8										GeneralProfileSpace = 0;
+			uint8										GeneralTierFlag = 0;
+			uint8										GeneralProfileIDC = 0;
+			uint32										GeneralProfileCompatibilityFlags = 0;
+			uint64										GeneralConstraintIndicatorFlags = 0;
+			uint8										GeneralLevelIDC = 0;
+			uint16										MinSpatialSegmentationIDC = 0;
+			uint8										ParallelismType = 0;
+			uint8										ChromaFormat = 0;
+			uint8										BitDepthLumaMinus8 = 0;
+			uint8										BitDepthChromaMinus8 = 0;
+			uint16										AverageFrameRate = 0;
+			uint8										ConstantFrameRate = 0;
+			uint8										NumTemporalLayers = 0;
+			uint8										TemporalIdNested = 0;
+			uint8										NALUnitLengthMinus1 = 0;
+			TArray<FISO23008_2_seq_parameter_set_data>	ParsedSPSs;
+		};
+
+
 
 		//! Parses a H.264 (ISO/IEC 14496-10) bitstream for NALUs.
 		struct FNaluInfo
@@ -203,6 +426,8 @@ namespace Electra
 		//! Parses a H.264 (ISO/IEC 14496-10) SPS NALU.
 		bool ParseH264SPS(FISO14496_10_seq_parameter_set_data& OutSPS, const void* Data, int32 Size);
 
+		//! Parses a H.265 (ISO/IEC 23008-2) SPS NALU.
+		bool ParseH265SPS(FISO23008_2_seq_parameter_set_data& OutSPS, const void* Data, int32 Size);
 
 	} // namespace MPEG
 } // namespace Electra
