ElectraPlayer: Reducing number of worker threads per instance

#rnx

diff --git a/Engine/Plugins/Media/ElectraPlayer/Source/ElectraPlayerRuntime/Private/Runtime/Player/AdaptiveStreamingPlayer.cpp b/Engine/Plugins/Media/ElectraPlayer/Source/ElectraPlayerRuntime/Private/Runtime/Player/AdaptiveStreamingPlayer.cpp
--- a/Engine/Plugins/Media/ElectraPlayer/Source/ElectraPlayerRuntime/Private/Runtime/Player/AdaptiveStreamingPlayer.cpp
+++ b/Engine/Plugins/Media/ElectraPlayer/Source/ElectraPlayerRuntime/Private/Runtime/Player/AdaptiveStreamingPlayer.cpp
@@ -125,7 +125,6 @@ FAdaptiveStreamingPlayer::FAdaptiveStreamingPlayer(const IAdaptiveStreamingPlaye
 
 	TMediaInterlockedExchangePointer(PointerToLatestPlayer, this);
 
-	// Create the worker thread.
 	StartWorkerThread();
 }
 
@@ -224,6 +223,9 @@ void FAdaptiveStreamingPlayer::Initialize(const FParamDict& Options)
 	StreamSelector = IAdaptiveStreamSelector::Create(this, PlayerConfig.StreamSelectorConfig);
 	AddMetricsReceiver(StreamSelector.Get());
 
+	// Create renderers.
+	CreateRenderers();
+
 	// Set up video decoder resolution limits. As the media playlists are parsed the video streams will be
 	// compared against these limits and those that exceed the limit will not be considered for playback.
 
@@ -356,14 +358,16 @@ void FAdaptiveStreamingPlayer::DispatchEventAndWait(TSharedPtrTS<FMetricEvent> E
  */
 void FAdaptiveStreamingPlayer::StartWorkerThread()
 {
-	if (!WorkerThread.bStarted)
+	// Get us an event dispatcher and add ourselves to the shared worker thread.
+	if (!EventDispatcher.IsValid())
+	{
+		EventDispatcher = FAdaptiveStreamingPlayerEventHandler::Create();
+	}
+	if (!SharedWorkerThread.IsValid())
 	{
-		WorkerThread.MediaThread.ThreadSetPriority(PlayerConfig.WorkerThread.Priority);
-		WorkerThread.MediaThread.ThreadSetCoreAffinity(PlayerConfig.WorkerThread.CoreAffinity);
-		WorkerThread.MediaThread.ThreadSetStackSize(PlayerConfig.WorkerThread.StackSize);
-		WorkerThread.MediaThread.ThreadSetName("ElectraPlayer::Worker");
-		WorkerThread.MediaThread.ThreadStart(Electra::MakeDelegate(this, &FAdaptiveStreamingPlayer::WorkerThreadFN));
-		WorkerThread.bStarted = true;
+		SharedWorkerThread = FAdaptiveStreamingPlayerWorkerThread::Create();
+		WorkerThread.SetSharedWorkerThread(SharedWorkerThread);
+		SharedWorkerThread->AddPlayerInstance(this);
 	}
 }
 
@@ -374,12 +378,16 @@ void FAdaptiveStreamingPlayer::StartWorkerThread()
  */
 void FAdaptiveStreamingPlayer::StopWorkerThread()
 {
-	if (WorkerThread.bStarted)
+	if (SharedWorkerThread.IsValid())
+	{
+		WorkerThread.SetSharedWorkerThread(nullptr);
+		SharedWorkerThread->RemovePlayerInstance(this);
+		SharedWorkerThread.Reset();
+	}
+	if (EventDispatcher.IsValid())
 	{
-		WorkerThread.SendMessage(FWorkerThread::FMessage::EType::Quit);
-		WorkerThread.MediaThread.ThreadWaitDone();
-		WorkerThread.MediaThread.ThreadReset();
-		WorkerThread.bStarted = false;
+		EventDispatcher->DispatchEventAndWait(TSharedPtrTS<FMetricEvent>());
+		EventDispatcher.Reset();
 	}
 }
 
@@ -1007,7 +1015,7 @@ void FAdaptiveStreamingPlayer::AUDeallocate(IAccessUnitMemoryProvider::EDataType
  */
 void FAdaptiveStreamingPlayer::OnFragmentOpen(TSharedPtrTS<IStreamSegment> pRequest)
 {
-	WorkerThread.SendMessage(FWorkerThread::FMessage::EType::FragmentOpen, pRequest);
+	WorkerThread.Enqueue(FWorkerThreadMessages::FMessage::EType::FragmentOpen, pRequest);
 }
 
 
@@ -1096,374 +1104,360 @@ void FAdaptiveStreamingPlayer::OnFragmentReachedEOS(EStreamType InStreamType, TS
  */
 void FAdaptiveStreamingPlayer::OnFragmentClose(TSharedPtrTS<IStreamSegment> pRequest)
 {
-	WorkerThread.SendMessage(FWorkerThread::FMessage::EType::FragmentClose, pRequest);
+	WorkerThread.Enqueue(FWorkerThreadMessages::FMessage::EType::FragmentClose, pRequest);
 }
 
 
 
-//-----------------------------------------------------------------------------
-/**
- * Worker thread.
- */
-void FAdaptiveStreamingPlayer::WorkerThreadFN()
-{
-	LLM_SCOPE(ELLMTag::ElectraPlayer);
-
-	bool	bQuit = false;
-
-	// Get the event dispatcher.
-	EventDispatcher = FAdaptiveStreamingPlayerEventHandler::Create();
-
 
-	// Create renderers and decoders.
-	// FIXME: At some later point in time we should do this on-demand only. We do not know the codecs
-	//        here yet and the codecs could also change dynamically as streams get switched.
-	//        For now this is ok and also prevents spikes from loading codec DLL/PRX later.
-	/*int32 CreateResult =*/ CreateRenderers();
-
-	while(!bQuit)
+void FAdaptiveStreamingPlayer::InternalHandleOnce()
+{
+	if (!bIsClosing)
 	{
-		FWorkerThread::FMessage msg;
-		if (WorkerThread.WorkMessages.ReceiveMessage(msg, 1000 * 20))
+		// Update the play position we return to the interested caller.
+		if (!PlaybackState.GetHasEnded())
 		{
-			SCOPE_CYCLE_COUNTER(STAT_ElectraPlayer_AdaptiveWorker);
-			CSV_SCOPED_TIMING_STAT(ElectraPlayer, AdaptiveStreamingPlayer_Worker);
-
-			// While closed ignore all but the quit message.
-			if (bIsClosing && msg.Type != FWorkerThread::FMessage::EType::Quit)
+			if (RenderClock->IsRunning())
 			{
-				if (msg.Data.MediaEvent.Event)
+				FTimeValue playPos = RenderClock->GetInterpolatedRenderTime(IMediaRenderClock::ERendererType::Video);
+				PlaybackState.SetPlayPosition(playPos);
+				// When we reach the time at which the next loop occurred pop it off the queue and update the playback state loop state.
+				while(NextLoopStates.Num())
 				{
-					msg.Data.MediaEvent.Event->Signal();
+					FTimeValue loopAtTime = NextLoopStates.FrontRef().LoopBasetime;
+					if (playPos >= loopAtTime)
+					{
+						FPlayerLoopState newLoopState = NextLoopStates.Pop();
+						PlaybackState.SetLoopState(newLoopState);
+						DispatchEvent(FMetricEvent::ReportJumpInPlayPosition(playPos, playPos, Metrics::ETimeJumpReason::Looping));
+					}
+					else
+					{
+						break;
+					}
 				}
-				continue;
 			}
+		}
+		else
+		{
+			// When playback has ended we lock the position to the end of the timeline.
+			// This is only in case the application checks for the play position to reach the end of the timeline
+			// ie. calling GetPlayPosition() to compare against the end of GetTimelineRange()
+			// instead of using the dedicated HasEnded() method.
+			FTimeRange TimelineRange;
+			PlaybackState.GetTimelineRange(TimelineRange);
+			PlaybackState.SetPlayPosition(TimelineRange.End);
+		}
+
+		// Handle state changes to match the user request.
+		HandlePlayStateChanges();
+
+		// Update diag buffers
+		UpdateDiagnostics();
+		// Check for end of stream.
+		CheckForStreamEnd();
+		// Check the error queue for new arrivals.
+		CheckForErrors();
+		// Handle pending media segment requests.
+		HandlePendingMediaSegmentRequests();
+		// Handle changes in metadata, like timeline changes or track availability.
+		HandleMetadataChanges();
+		// Handle AEMS events
+		HandleAEMSEvents();
+		// Handle buffer level changes
+		HandleNewBufferedData();
+		// Handle new output data (finish prerolling)
+		HandleNewOutputData();
+		// Handle data buffers from deselected tracks to align with selected ones.
+		HandleDeselectedBuffers();
+		// Handle decoder changes
+		HandleDecoderChanges();
+		// Handle entity cache expirations.
+		if (EntityCache.IsValid())
+		{
+			EntityCache->HandleEntityExpiration();
+		}
+		// Handle completed DRM requests.
+		if (DrmManager.IsValid())
+		{
+			DrmManager->Tick();
+		}
+	}
+}
+
+bool FAdaptiveStreamingPlayer::InternalHandleThreadMessages()
+{
+	bool bGotMessage = false;
+	FWorkerThreadMessages::FMessage msg;
+	while(WorkerThread.WorkMessages.Dequeue(msg))
+	{
+		// While closed ignore all messages.
+		if (bIsClosing)
+		{
+			if (msg.Data.MediaEvent.Event)
+			{
+				msg.Data.MediaEvent.Event->Signal();
+			}
+			continue;
+		}
 
-			// What is it?
-			switch(msg.Type)
+		bGotMessage = true;
+		// What is it?
+		switch(msg.Type)
+		{
+			case FWorkerThreadMessages::FMessage::EType::LoadManifest:
+			{
+				InternalLoadManifest(msg.Data.ManifestToLoad.URL, msg.Data.ManifestToLoad.MimeType);
+				break;
+			}
+			case FWorkerThreadMessages::FMessage::EType::Pause:
+			{
+				bShouldBePaused = true;
+				bShouldBePlaying = false;
+				break;
+			}
+			case FWorkerThreadMessages::FMessage::EType::Resume:
+			{
+				bShouldBePaused = false;
+				bShouldBePlaying = true;
+				break;
+			}
+			case FWorkerThreadMessages::FMessage::EType::Seek:
 			{
-				case FWorkerThread::FMessage::EType::Quit:
+				StartAtTime = msg.Data.StartPlay.Position;
+
+				if (bIsPlayStart)
 				{
-					bQuit = true;
-					break;
+					CurrentState = EPlayerState::eState_Buffering;
+					bool bOk = InternalStartAt(StartAtTime);
 				}
-				case FWorkerThread::FMessage::EType::LoadManifest:
+				else
 				{
-					InternalLoadManifest(msg.Data.ManifestToLoad.URL, msg.Data.ManifestToLoad.MimeType);
-					break;
+					// Stop everything.
+					InternalStop(PlayerConfig.bHoldLastFrameDuringSeek);
+					if (Manifest.IsValid())
+					{
+						Manifest->UpdateDynamicRefetchCounter();
+					}
+					CurrentState = EPlayerState::eState_Seeking;
+					bool bOk = InternalStartAt(StartAtTime);
 				}
-				case FWorkerThread::FMessage::EType::Pause:
+				break;
+			}
+			case FWorkerThreadMessages::FMessage::EType::Loop:
+			{
+				InternalSetLoop(msg.Data.Looping.Loop);
+				if (msg.Data.Looping.Signal)
 				{
-					bShouldBePaused = true;
-					bShouldBePlaying = false;
-					break;
+					msg.Data.Looping.Signal->Signal();
 				}
-				case FWorkerThread::FMessage::EType::Resume:
+				break;
+			}
+			case FWorkerThreadMessages::FMessage::EType::Close:
+			{
+				if (!bIsClosing)
 				{
-					bShouldBePaused = false;
-					bShouldBePlaying = true;
-					break;
+					bIsClosing = true;
+					InternalStop(false);
+					InternalClose();
+					// This is the end of the rope anyway so let's terminate the event dispatcher.
+					// This ensures that the client will not get out of the Stop() call without having
+					// received the stop message.
+					DispatchEventAndWait(FMetricEvent::ReportPlaybackStopped());
+					EventDispatcher.Reset();
 				}
-				case FWorkerThread::FMessage::EType::Seek:
+				if (msg.Data.MediaEvent.Event)
 				{
-					StartAtTime = msg.Data.StartPlay.Position;
-
-					if (bIsPlayStart)
-					{
-						CurrentState = EPlayerState::eState_Buffering;
-						bool bOk = InternalStartAt(StartAtTime);
-					}
-					else
-					{
-						// Stop everything.
-						InternalStop(PlayerConfig.bHoldLastFrameDuringSeek);
-						if (Manifest.IsValid())
-						{
-							Manifest->UpdateDynamicRefetchCounter();
-						}
-						CurrentState = EPlayerState::eState_Seeking;
-						bool bOk = InternalStartAt(StartAtTime);
-					}
-					break;
+					msg.Data.MediaEvent.Event->Signal();
 				}
-				case FWorkerThread::FMessage::EType::Loop:
-				{
-					InternalSetLoop(msg.Data.Looping.Loop);
-					if (msg.Data.Looping.Signal)
-					{
-						msg.Data.Looping.Signal->Signal();
-					}
-					break;
+				break;
 				}
-				case FWorkerThread::FMessage::EType::Close:
-				{
-					if (!bIsClosing)
-					{
-						bIsClosing = true;
-						InternalStop(false);
-						InternalClose();
-						// This is the end of the rope anyway so let's terminate the event dispatcher.
-						// This ensures that the client will not get out of the Stop() call without having
-						// received the stop message.
-						DispatchEventAndWait(FMetricEvent::ReportPlaybackStopped());
-						EventDispatcher.Reset();
-					}
-					if (msg.Data.MediaEvent.Event)
-					{
-						msg.Data.MediaEvent.Event->Signal();
-					}
-					break;
-					}
 
-				case FWorkerThread::FMessage::EType::ChangeBitrate:
+			case FWorkerThreadMessages::FMessage::EType::ChangeBitrate:
+			{
+				BitrateCeiling = msg.Data.Bitrate.Value;
+				StreamSelector->SetBandwidthCeiling(BitrateCeiling);
+				break;
+			}
+
+			case FWorkerThreadMessages::FMessage::EType::LimitResolution:
+			{
+				VideoResolutionLimitWidth  = msg.Data.Resolution.Width;
+				VideoResolutionLimitHeight = msg.Data.Resolution.Height;
+				UpdateStreamResolutionLimit();
+				break;
+			}
+
+			case FWorkerThreadMessages::FMessage::EType::InitialStreamAttributes:
+			{
+				// Map the language in case it is not yet.
+				if (msg.Data.InitialStreamAttribute.InitialSelection.Language_ISO639.IsSet())
 				{
-					BitrateCeiling = msg.Data.Bitrate.Value;
-					StreamSelector->SetBandwidthCeiling(BitrateCeiling);
-					break;
+					msg.Data.InitialStreamAttribute.InitialSelection.Language_ISO639 = ISO639::MapTo639_1(msg.Data.InitialStreamAttribute.InitialSelection.Language_ISO639.GetValue());
 				}
-
-				case FWorkerThread::FMessage::EType::LimitResolution:
+				switch(msg.Data.InitialStreamAttribute.StreamType)
 				{
-					VideoResolutionLimitWidth  = msg.Data.Resolution.Width;
-					VideoResolutionLimitHeight = msg.Data.Resolution.Height;
-					UpdateStreamResolutionLimit();
-					break;
+					case EStreamType::Video:
+						StreamSelectionAttributesVid = msg.Data.InitialStreamAttribute.InitialSelection;
+						break;
+					case EStreamType::Audio:
+						StreamSelectionAttributesAud = msg.Data.InitialStreamAttribute.InitialSelection;
+						break;
+					case EStreamType::Subtitle:
+						StreamSelectionAttributesTxt = msg.Data.InitialStreamAttribute.InitialSelection;
+						break;
+					default:
+						break;
 				}
+				break;
+			}
 
-				case FWorkerThread::FMessage::EType::InitialStreamAttributes:
+			case FWorkerThreadMessages::FMessage::EType::SelectTrackByAttributes:
+			{
+				if (!bIsClosing)
 				{
 					// Map the language in case it is not yet.
-					if (msg.Data.InitialStreamAttribute.InitialSelection.Language_ISO639.IsSet())
+					if (msg.Data.TrackSelection.TrackAttributes.Language_ISO639.IsSet())
 					{
-						msg.Data.InitialStreamAttribute.InitialSelection.Language_ISO639 = ISO639::MapTo639_1(msg.Data.InitialStreamAttribute.InitialSelection.Language_ISO639.GetValue());
+						msg.Data.TrackSelection.TrackAttributes.Language_ISO639 = ISO639::MapTo639_1(msg.Data.TrackSelection.TrackAttributes.Language_ISO639.GetValue());
 					}
-					switch(msg.Data.InitialStreamAttribute.StreamType)
+					switch(msg.Data.TrackSelection.StreamType)
 					{
 						case EStreamType::Video:
-							StreamSelectionAttributesVid = msg.Data.InitialStreamAttribute.InitialSelection;
+							PendingTrackSelectionVid = MakeSharedTS<FStreamSelectionAttributes>(msg.Data.TrackSelection.TrackAttributes);
+							MultiStreamBufferVid.Activate();
 							break;
 						case EStreamType::Audio:
-							StreamSelectionAttributesAud = msg.Data.InitialStreamAttribute.InitialSelection;
+							PendingTrackSelectionAud = MakeSharedTS<FStreamSelectionAttributes>(msg.Data.TrackSelection.TrackAttributes);
+							MultiStreamBufferAud.Activate();
 							break;
 						case EStreamType::Subtitle:
-							StreamSelectionAttributesTxt = msg.Data.InitialStreamAttribute.InitialSelection;
-							break;
-						default:
+							PendingTrackSelectionTxt = MakeSharedTS<FStreamSelectionAttributes>(msg.Data.TrackSelection.TrackAttributes);
+							MultiStreamBufferTxt.Activate();
 							break;
 					}
-					break;
 				}
+				break;
+			}
 
-				case FWorkerThread::FMessage::EType::SelectTrackByAttributes:
-				{
-					if (!bIsClosing)
-					{
-						// Map the language in case it is not yet.
-						if (msg.Data.TrackSelection.TrackAttributes.Language_ISO639.IsSet())
-						{
-							msg.Data.TrackSelection.TrackAttributes.Language_ISO639 = ISO639::MapTo639_1(msg.Data.TrackSelection.TrackAttributes.Language_ISO639.GetValue());
-						}
-						switch(msg.Data.TrackSelection.StreamType)
-						{
-							case EStreamType::Video:
-								PendingTrackSelectionVid = MakeSharedTS<FStreamSelectionAttributes>(msg.Data.TrackSelection.TrackAttributes);
-								MultiStreamBufferVid.Activate();
-								break;
-							case EStreamType::Audio:
-								PendingTrackSelectionAud = MakeSharedTS<FStreamSelectionAttributes>(msg.Data.TrackSelection.TrackAttributes);
-								MultiStreamBufferAud.Activate();
-								break;
-							case EStreamType::Subtitle:
-								PendingTrackSelectionTxt = MakeSharedTS<FStreamSelectionAttributes>(msg.Data.TrackSelection.TrackAttributes);
-								MultiStreamBufferTxt.Activate();
-								break;
-						}
-					}
-					break;
-				}
+			case FWorkerThreadMessages::FMessage::EType::SelectTrackByMetadata:
+			{
+				// Currently not used. May be used again later.
+				break;
+			}
 
-				case FWorkerThread::FMessage::EType::SelectTrackByMetadata:
+			case FWorkerThreadMessages::FMessage::EType::DeselectTrack:
+			{
+				switch(msg.Data.TrackSelection.StreamType)
 				{
-					// Currently not used. May be used again later.
-					break;
+					case EStreamType::Video:
+						MultiStreamBufferVid.Deselect();
+						break;
+					case EStreamType::Audio:
+						MultiStreamBufferAud.Deselect();
+						break;
+					case EStreamType::Subtitle:
+						MultiStreamBufferTxt.Deselect();
+						break;
 				}
+				break;
+			}
 
-				case FWorkerThread::FMessage::EType::DeselectTrack:
+			case FWorkerThreadMessages::FMessage::EType::FragmentOpen:
+			{
+				TSharedPtrTS<IStreamSegment> pRequest = msg.Data.StreamReader.Request;
+				EStreamType reqType = msg.Data.StreamReader.Request->GetType();
+				// Check that the request is for this current playback sequence and not an outdated one.
+				if (pRequest.IsValid() && pRequest->GetPlaybackSequenceID() == CurrentPlaybackSequenceID[StreamTypeToArrayIndex(reqType)])
 				{
-					switch(msg.Data.TrackSelection.StreamType)
-					{
-						case EStreamType::Video:
-							MultiStreamBufferVid.Deselect();
-							break;
-						case EStreamType::Audio:
-							MultiStreamBufferAud.Deselect();
-							break;
-						case EStreamType::Subtitle:
-							MultiStreamBufferTxt.Deselect();
-							break;
-					}
-					break;
-				}
+					DispatchBufferUtilizationEvent(msg.Data.StreamReader.Request->GetType() == EStreamType::Video ? EStreamType::Video : EStreamType::Audio);
 
-				case FWorkerThread::FMessage::EType::FragmentOpen:
-				{
-					TSharedPtrTS<IStreamSegment> pRequest = msg.Data.StreamReader.Request;
-					EStreamType reqType = msg.Data.StreamReader.Request->GetType();
-					// Check that the request is for this current playback sequence and not an outdated one.
-					if (pRequest.IsValid() && pRequest->GetPlaybackSequenceID() == CurrentPlaybackSequenceID[StreamTypeToArrayIndex(reqType)])
+					// Video bitrate change?
+					if (msg.Data.StreamReader.Request->GetType() == EStreamType::Video)
 					{
-						DispatchBufferUtilizationEvent(msg.Data.StreamReader.Request->GetType() == EStreamType::Video ? EStreamType::Video : EStreamType::Audio);
-
-						// Video bitrate change?
-						if (msg.Data.StreamReader.Request->GetType() == EStreamType::Video)
+						int32 SegmentBitrate = pRequest->GetBitrate();
+						int32 SegmentQualityLevel = pRequest->GetQualityIndex();
+						if (SegmentBitrate != CurrentVideoStreamBitrate.Bitrate)
 						{
-							int32 SegmentBitrate = pRequest->GetBitrate();
-							int32 SegmentQualityLevel = pRequest->GetQualityIndex();
-							if (SegmentBitrate != CurrentVideoStreamBitrate.Bitrate)
-							{
-								bool bDrastic = CurrentVideoStreamBitrate.Bitrate && SegmentQualityLevel < CurrentVideoStreamBitrate.QualityLevel-1;
-								DispatchEvent(FMetricEvent::ReportVideoQualityChange(SegmentBitrate, CurrentVideoStreamBitrate.Bitrate, bDrastic));
-								CurrentVideoStreamBitrate.Bitrate      = SegmentBitrate;
-								CurrentVideoStreamBitrate.QualityLevel = SegmentQualityLevel;
-							}
+							bool bDrastic = CurrentVideoStreamBitrate.Bitrate && SegmentQualityLevel < CurrentVideoStreamBitrate.QualityLevel-1;
+							DispatchEvent(FMetricEvent::ReportVideoQualityChange(SegmentBitrate, CurrentVideoStreamBitrate.Bitrate, bDrastic));
+							CurrentVideoStreamBitrate.Bitrate      = SegmentBitrate;
+							CurrentVideoStreamBitrate.QualityLevel = SegmentQualityLevel;
 						}
 					}
-					break;
 				}
+				break;
+			}
 
-				case FWorkerThread::FMessage::EType::FragmentClose:
+			case FWorkerThreadMessages::FMessage::EType::FragmentClose:
+			{
+				TSharedPtrTS<IStreamSegment> pRequest = msg.Data.StreamReader.Request;
+				EStreamType reqType = msg.Data.StreamReader.Request->GetType();
+				// Check that the request is for this current playback sequence and not an outdated one.
+				if (pRequest.IsValid() && pRequest->GetPlaybackSequenceID() == CurrentPlaybackSequenceID[StreamTypeToArrayIndex(reqType)])
 				{
-					TSharedPtrTS<IStreamSegment> pRequest = msg.Data.StreamReader.Request;
-					EStreamType reqType = msg.Data.StreamReader.Request->GetType();
-					// Check that the request is for this current playback sequence and not an outdated one.
-					if (pRequest.IsValid() && pRequest->GetPlaybackSequenceID() == CurrentPlaybackSequenceID[StreamTypeToArrayIndex(reqType)])
-					{
-						// Dispatch download event
-						DispatchSegmentDownloadedEvent(pRequest);
-
-						// Dispatch buffer utilization
-						DispatchBufferUtilizationEvent(reqType == EStreamType::Video ? EStreamType::Video : EStreamType::Audio);
+					// Dispatch download event
+					DispatchSegmentDownloadedEvent(pRequest);
 
-						// Dispatch average throughput, bandwidth and latency event for video segments only.
-						if (reqType == EStreamType::Video)
-						{
-							DispatchEvent(FMetricEvent::ReportBandwidth(StreamSelector->GetAverageBandwidth(), StreamSelector->GetAverageThroughput(), StreamSelector->GetAverageLatency()));
-						}
+					// Dispatch buffer utilization
+					DispatchBufferUtilizationEvent(reqType == EStreamType::Video ? EStreamType::Video : EStreamType::Audio);
 
-						// Add to the list of completed requests for which we need to find the next or retry segment to fetch.
-						FPendingSegmentRequest NextReq;
-						NextReq.Request = pRequest;
-						NextPendingSegmentRequests.Enqueue(MoveTemp(NextReq));
+					// Dispatch average throughput, bandwidth and latency event for video segments only.
+					if (reqType == EStreamType::Video)
+					{
+						DispatchEvent(FMetricEvent::ReportBandwidth(StreamSelector->GetAverageBandwidth(), StreamSelector->GetAverageThroughput(), StreamSelector->GetAverageLatency()));
 					}
-					break;
-				}
-
-				case FWorkerThread::FMessage::EType::BufferUnderrun:
-				{
-					InternalRebuffer();
-					break;
-				}
-
-				case FWorkerThread::FMessage::EType::PlayerSession:
-				{
-					HandleSessionMessage(msg.Data.Session.PlayerMessage);
-					break;
-				}
 
-				default:
-				{
-					checkNoEntry();
-					break;
+					// Add to the list of completed requests for which we need to find the next or retry segment to fetch.
+					FPendingSegmentRequest NextReq;
+					NextReq.Request = pRequest;
+					NextPendingSegmentRequests.Enqueue(MoveTemp(NextReq));
 				}
+				break;
 			}
-		}
 
-		if (!bIsClosing)
-		{
-			SCOPE_CYCLE_COUNTER(STAT_ElectraPlayer_AdaptiveWorker);
-			CSV_SCOPED_TIMING_STAT(ElectraPlayer, AdaptiveStreamingPlayer_Worker);
+			case FWorkerThreadMessages::FMessage::EType::BufferUnderrun:
+			{
+				InternalRebuffer();
+				break;
+			}
 
-			// Update the play position we return to the interested caller.
-			if (!PlaybackState.GetHasEnded())
+			case FWorkerThreadMessages::FMessage::EType::PlayerSession:
 			{
-				if (RenderClock->IsRunning())
-				{
-					FTimeValue playPos = RenderClock->GetInterpolatedRenderTime(IMediaRenderClock::ERendererType::Video);
-					PlaybackState.SetPlayPosition(playPos);
-					// When we reach the time at which the next loop occurred pop it off the queue and update the playback state loop state.
-					while(NextLoopStates.Num())
-					{
-						FTimeValue loopAtTime = NextLoopStates.FrontRef().LoopBasetime;
-						if (playPos >= loopAtTime)
-						{
-							FPlayerLoopState newLoopState = NextLoopStates.Pop();
-							PlaybackState.SetLoopState(newLoopState);
-							DispatchEvent(FMetricEvent::ReportJumpInPlayPosition(playPos, playPos, Metrics::ETimeJumpReason::Looping));
-						}
-						else
-						{
-							break;
-						}
-					}
-				}
+				HandleSessionMessage(msg.Data.Session.PlayerMessage);
+				break;
 			}
-			else
+
+			default:
 			{
-				// When playback has ended we lock the position to the end of the timeline.
-				// This is only in case the application checks for the play position to reach the end of the timeline
-				// ie. calling GetPlayPosition() to compare against the end of GetTimelineRange()
-				// instead of using the dedicated HasEnded() method.
-				FTimeRange TimelineRange;
-				PlaybackState.GetTimelineRange(TimelineRange);
-				PlaybackState.SetPlayPosition(TimelineRange.End);
-			}
-
-			// Handle state changes to match the user request.
-			HandlePlayStateChanges();
-
-			// Update diag buffers
-			UpdateDiagnostics();
-			// Check for end of stream.
-			CheckForStreamEnd();
-			// Check the error queue for new arrivals.
-			CheckForErrors();
-			// Handle pending media segment requests.
-			HandlePendingMediaSegmentRequests();
-			// Handle changes in metadata, like timeline changes or track availability.
-			HandleMetadataChanges();
-			// Handle AEMS events
-			HandleAEMSEvents();
-			// Handle buffer level changes
-			HandleNewBufferedData();
-			// Handle new output data (finish prerolling)
-			HandleNewOutputData();
-			// Handle data buffers from deselected tracks to align with selected ones.
-			HandleDeselectedBuffers();
-			// Handle decoder changes
-			HandleDecoderChanges();
-			// Handle entity cache expirations.
-			if (EntityCache.IsValid())
-			{
-				EntityCache->HandleEntityExpiration();
-			}
-			// Handle completed DRM requests.
-			if (DrmManager.IsValid())
-			{
-				DrmManager->Tick();
+				checkNoEntry();
+				break;
 			}
 		}
+			
+		InternalHandleOnce();
 	}
+	return bGotMessage;
+}
 
-	if (EventDispatcher.IsValid())
-	{
-		EventDispatcher->DispatchEventAndWait(TSharedPtrTS<FMetricEvent>());
-		EventDispatcher.Reset();
-	}
+//-----------------------------------------------------------------------------
+/**
+ * Handling function called by the shared worker thread.
+ */
+void FAdaptiveStreamingPlayer::HandleOnce()
+{
+	SCOPE_CYCLE_COUNTER(STAT_ElectraPlayer_AdaptiveWorker);
+	CSV_SCOPED_TIMING_STAT(ElectraPlayer, AdaptiveStreamingPlayer_Worker);
 
-	while(WorkerThread.WorkMessages.HaveMessage())
+	if (!InternalHandleThreadMessages())
 	{
-		WorkerThread.WorkMessages.ReceiveMessage();
+		InternalHandleOnce();
 	}
+	// Handle the components that are bound to add new worker messages.
+	// This should not be done while processing the messages that were enqueued earlier,
+	// especially not since handling those will result in handling all state checks
+	// in InternalHandleOnce() again!
+	InternalHandleManifestReader();
 }
 
 
@@ -2892,11 +2886,7 @@ void FAdaptiveStreamingPlayer::CheckForErrors()
 				DispatchEvent(FMetricEvent::ReportError(LastErrorDetail));
 			}
 			// In error state we do not need any periodic manifest refetches any more.
-			if (ManifestReader.IsValid())
-			{
-				ManifestReader->Close();
-				ManifestReader.Reset();
-			}
+			InternalCloseManifestReader();
 		}
 	}
 }
@@ -3260,11 +3250,7 @@ void FAdaptiveStreamingPlayer::InternalClose()
 {
 	// No longer need the manifest reader/updater
 	InternalCancelLoadManifest();
-	if (ManifestReader.IsValid())
-	{
-		ManifestReader->Close();
-		ManifestReader.Reset();
-	}
+	InternalCloseManifestReader();
 
 	DestroyDecoders();
 	DestroyRenderers();
@@ -3473,13 +3459,72 @@ void FAdaptiveStreamingPlayer::DebugHandle(void* pPlayer, void (*debugDrawPrintf
 //---------------------------------------------------------------------------------------------------------------------------------
 //---------------------------------------------------------------------------------------------------------------------------------
 
-TWeakPtr<FAdaptiveStreamingPlayerEventHandler, ESPMode::ThreadSafe>	FAdaptiveStreamingPlayerEventHandler::SingletonSelf;
-FCriticalSection													FAdaptiveStreamingPlayerEventHandler::SingletonLock;
 
-TSharedPtr<FAdaptiveStreamingPlayerEventHandler, ESPMode::ThreadSafe> FAdaptiveStreamingPlayerEventHandler::Create()
+TWeakPtrTS<FAdaptiveStreamingPlayerWorkerThread>	FAdaptiveStreamingPlayerWorkerThread::SingletonSelf;
+FCriticalSection									FAdaptiveStreamingPlayerWorkerThread::SingletonLock;
+
+TSharedPtrTS<FAdaptiveStreamingPlayerWorkerThread> FAdaptiveStreamingPlayerWorkerThread::Create()
+{
+	FScopeLock lock(&SingletonLock);
+	TSharedPtrTS<FAdaptiveStreamingPlayerWorkerThread> Self = SingletonSelf.Pin();
+	if (!Self.IsValid())
+	{
+		FAdaptiveStreamingPlayerWorkerThread* Handler = new FAdaptiveStreamingPlayerWorkerThread;
+		Handler->StartWorkerThread();
+		Self = MakeShareable(Handler);
+		SingletonSelf = Self;
+	}
+	return Self;
+}
+
+FAdaptiveStreamingPlayerWorkerThread::~FAdaptiveStreamingPlayerWorkerThread()
+{
+	StopWorkerThread();
+}
+
+void FAdaptiveStreamingPlayerWorkerThread::StartWorkerThread()
+{
+	WorkerThread.ThreadSetName("ElectraPlayer::SharedWorker");
+	WorkerThread.ThreadStart(Electra::MakeDelegate(this, &FAdaptiveStreamingPlayerWorkerThread::WorkerThreadFN));
+}
+
+void FAdaptiveStreamingPlayerWorkerThread::StopWorkerThread()
+{
+	bTerminate = true;
+	HaveWorkSignal.Signal();
+	WorkerThread.ThreadWaitDone();
+	WorkerThread.ThreadReset();
+}
+
+void FAdaptiveStreamingPlayerWorkerThread::WorkerThreadFN()
+{
+	LLM_SCOPE(ELLMTag::ElectraPlayer);
+	while(!bTerminate)
+	{
+		HaveWorkSignal.WaitTimeoutAndReset(1000 * 20);
+
+		TArray<FAdaptiveStreamingPlayer*>	ActiveInstances(PlayerInstances);
+		SingletonLock.Lock();
+		for(int32 i=0; i<ActiveInstances.Num(); ++i)
+		{
+			ActiveInstances[i]->HandleOnce();
+		}
+		SingletonLock.Unlock();
+	}
+}
+
+
+//---------------------------------------------------------------------------------------------------------------------------------
+//---------------------------------------------------------------------------------------------------------------------------------
+//---------------------------------------------------------------------------------------------------------------------------------
+
+TWeakPtrTS<FAdaptiveStreamingPlayerEventHandler>	FAdaptiveStreamingPlayerEventHandler::SingletonSelf;
+FCriticalSection									FAdaptiveStreamingPlayerEventHandler::SingletonLock;
+
+TSharedPtrTS<FAdaptiveStreamingPlayerEventHandler> FAdaptiveStreamingPlayerEventHandler::Create()
 {
 	FScopeLock lock(&SingletonLock);
-	TSharedPtr<FAdaptiveStreamingPlayerEventHandler, ESPMode::ThreadSafe> Self = SingletonSelf.Pin();
+	TSharedPtrTS<FAdaptiveStreamingPlayerEventHandler> Self = SingletonSelf.Pin();
 	if (!Self.IsValid())
 	{
 		FAdaptiveStreamingPlayerEventHandler* Handler = new FAdaptiveStreamingPlayerEventHandler;
@@ -3537,7 +3582,7 @@ void FAdaptiveStreamingPlayerEventHandler::WorkerThreadFN()
 			CSV_SCOPED_TIMING_STAT(ElectraPlayer, AdaptiveStreamingPlayer_Event);
 
 			// Get the player that sends the event. If it no longer exists ignore the event.
-			TSharedPtr<FAdaptiveStreamingPlayer, ESPMode::ThreadSafe>	Player = pEvt->Player.Pin();
+			TSharedPtrTS<FAdaptiveStreamingPlayer>	Player = pEvt->Player.Pin();
 			if (Player.IsValid())
 			{
 				// Call listeners under lock. They are not allowed to add or remove themselves or other listeners.
diff --git a/Engine/Plugins/Media/ElectraPlayer/Source/ElectraPlayerRuntime/Private/Runtime/Player/AdaptiveStreamingPlayerDecoders.cpp b/Engine/Plugins/Media/ElectraPlayer/Source/ElectraPlayerRuntime/Private/Runtime/Player/AdaptiveStreamingPlayerDecoders.cpp
--- a/Engine/Plugins/Media/ElectraPlayer/Source/ElectraPlayerRuntime/Private/Runtime/Player/AdaptiveStreamingPlayerDecoders.cpp
+++ b/Engine/Plugins/Media/ElectraPlayer/Source/ElectraPlayerRuntime/Private/Runtime/Player/AdaptiveStreamingPlayerDecoders.cpp
@@ -468,7 +468,7 @@ void FAdaptiveStreamingPlayer::FeedDecoder(EStreamType Type, FMultiTrackAccessUn
 				{
 					RebufferDetectedAtPlayPos = LastKnownPTS;
 				}
-				WorkerThread.SendMessage(FWorkerThread::FMessage::EType::BufferUnderrun);
+				WorkerThread.Enqueue(FWorkerThreadMessages::FMessage::EType::BufferUnderrun);
 			}
 		}
 
diff --git a/Engine/Plugins/Media/ElectraPlayer/Source/ElectraPlayerRuntime/Private/Runtime/Player/AdaptiveStreamingPlayerInternal.h b/Engine/Plugins/Media/ElectraPlayer/Source/ElectraPlayerRuntime/Private/Runtime/Player/AdaptiveStreamingPlayerInternal.h
--- a/Engine/Plugins/Media/ElectraPlayer/Source/ElectraPlayerRuntime/Private/Runtime/Player/AdaptiveStreamingPlayerInternal.h
+++ b/Engine/Plugins/Media/ElectraPlayer/Source/ElectraPlayerRuntime/Private/Runtime/Player/AdaptiveStreamingPlayerInternal.h
@@ -618,10 +618,47 @@ struct FMetricEvent
 
 
 
+class FAdaptiveStreamingPlayerWorkerThread
+{
+public:
+	static TSharedPtrTS<FAdaptiveStreamingPlayerWorkerThread> Create();
+
+	void AddPlayerInstance(FAdaptiveStreamingPlayer* Instance)
+	{
+		FScopeLock lock(&SingletonLock);
+		PlayerInstances.Emplace(Instance);
+	}
+	void RemovePlayerInstance(FAdaptiveStreamingPlayer* Instance)
+	{
+		FScopeLock lock(&SingletonLock);
+		PlayerInstances.Remove(Instance);
+	}
+
+	void TriggerWork()
+	{
+		HaveWorkSignal.Signal();
+	}
+
+	virtual ~FAdaptiveStreamingPlayerWorkerThread();
+private:
+	void StartWorkerThread();
+	void StopWorkerThread();
+	void WorkerThreadFN();
+
+	FMediaThread						WorkerThread;
+	TArray<FAdaptiveStreamingPlayer*>	PlayerInstances;
+	FMediaEvent							HaveWorkSignal;
+	bool								bTerminate = false;
+
+	static TWeakPtrTS<FAdaptiveStreamingPlayerWorkerThread>	SingletonSelf;
+	static FCriticalSection									SingletonLock;
+};
+
+
 class FAdaptiveStreamingPlayerEventHandler
 {
 public:
-	static TSharedPtr<FAdaptiveStreamingPlayerEventHandler, ESPMode::ThreadSafe> Create();
+	static TSharedPtrTS<FAdaptiveStreamingPlayerEventHandler> Create();
 
 	void DispatchEvent(TSharedPtrTS<FMetricEvent> InEvent);
 	void DispatchEventAndWait(TSharedPtrTS<FMetricEvent> InEvent);
@@ -632,10 +669,10 @@ private:
 	void StopWorkerThread();
 	void WorkerThreadFN();
 
-	FMediaThread																WorkerThread;
-	TMediaMessageQueueDynamicNoTimeout<TSharedPtrTS<FMetricEvent>>				EventQueue;
-	static TWeakPtr<FAdaptiveStreamingPlayerEventHandler, ESPMode::ThreadSafe>	SingletonSelf;
-	static FCriticalSection														SingletonLock;
+	FMediaThread													WorkerThread;
+	TMediaMessageQueueDynamicNoTimeout<TSharedPtrTS<FMetricEvent>>	EventQueue;
+	static TWeakPtrTS<FAdaptiveStreamingPlayerEventHandler>			SingletonSelf;
+	static FCriticalSection											SingletonLock;
 };
 
 
@@ -657,6 +694,8 @@ public:
 	virtual void AddMetricsReceiver(IAdaptiveStreamingPlayerMetrics* InMetricsReceiver) override;
 	virtual void RemoveMetricsReceiver(IAdaptiveStreamingPlayerMetrics* InMetricsReceiver) override;
 
+	void HandleOnce();
+
 	// Metrics receiver accessor for event dispatcher thread.
 	void LockMetricsReceivers()
 	{
@@ -944,12 +983,16 @@ private:
 		};
 
 
-	class FWorkerThread
+	class FWorkerThreadMessages
 	{
 	public:
-		FWorkerThread()
-			: bStarted(false)
+		~FWorkerThreadMessages()
+		{
+			check(WorkMessages.IsEmpty());
+		}
+		void SetSharedWorkerThread(TSharedPtrTS<FAdaptiveStreamingPlayerWorkerThread> InSharedWorkerThread)
 		{
+			SharedWorkerThread = InSharedWorkerThread;
 		}
 
 		struct FMessage
@@ -958,7 +1001,6 @@ private:
 			{
 				// Commands
 				LoadManifest,
-				Quit,
 				Pause,
 				Resume,
 				Seek,
@@ -1058,25 +1100,25 @@ private:
 			FData   				Data;
 		};
 
-		void SendMessage(FMessage::EType type)
+		void Enqueue(FMessage::EType type)
 		{
 			FMessage Msg;
 			Msg.Type = type;
-			WorkMessages.SendMessage(Msg);
+			TriggerSharedWorkerThread(MoveTemp(Msg));
 		}
-		void SendMessage(FMessage::EType type, TSharedPtrTS<IStreamSegment> pRequest)
+		void Enqueue(FMessage::EType type, TSharedPtrTS<IStreamSegment> pRequest)
 		{
 			FMessage Msg;
 			Msg.Type = type;
 			Msg.Data.StreamReader.Request = pRequest;
-			WorkMessages.SendMessage(Msg);
+			TriggerSharedWorkerThread(MoveTemp(Msg));
 		}
-		void SendMessage(FMessage::EType type, FMediaEvent* pEventSignal)
+		void Enqueue(FMessage::EType type, FMediaEvent* pEventSignal)
 		{
 			FMessage Msg;
 			Msg.Type = type;
 			Msg.Data.MediaEvent.Event = pEventSignal;
-			WorkMessages.SendMessage(Msg);
+			TriggerSharedWorkerThread(MoveTemp(Msg));
 		}
 		void SendLoadManifestMessage(const FString& URL, const FString& MimeType)
 		{
@@ -1084,54 +1126,54 @@ private:
 			Msg.Type = FMessage::EType::LoadManifest;
 			Msg.Data.ManifestToLoad.URL = URL;
 			Msg.Data.ManifestToLoad.MimeType = MimeType;
-			WorkMessages.SendMessage(Msg);
+			TriggerSharedWorkerThread(MoveTemp(Msg));
 		}
 		void SendPauseMessage()
 		{
 			FMessage Msg;
 			Msg.Type = FMessage::EType::Pause;
-			WorkMessages.SendMessage(Msg);
+			TriggerSharedWorkerThread(MoveTemp(Msg));
 		}
 		void SendResumeMessage()
 		{
 			FMessage Msg;
 			Msg.Type = FMessage::EType::Resume;
-			WorkMessages.SendMessage(Msg);
+			TriggerSharedWorkerThread(MoveTemp(Msg));
 		}
 		void SendSeekMessage(const FSeekParam& NewPosition)
 		{
 			FMessage Msg;
 			Msg.Type = FMessage::EType::Seek;
 			Msg.Data.StartPlay.Position = NewPosition;
-			WorkMessages.SendMessage(Msg);
+			TriggerSharedWorkerThread(MoveTemp(Msg));
 		}
 		void SendLoopMessage(const FLoopParam& InLoopParams)
 		{
 			FMessage Msg;
 			Msg.Type = FMessage::EType::Loop;
 			Msg.Data.Looping.Loop = InLoopParams;
-			WorkMessages.SendMessage(Msg);
+			TriggerSharedWorkerThread(MoveTemp(Msg));
 		}
 		void SendCloseMessage(FMediaEvent* pEventSignal)
 		{
 			FMessage Msg;
 			Msg.Type = FMessage::EType::Close;
 			Msg.Data.MediaEvent.Event = pEventSignal;
-			WorkMessages.SendMessage(Msg);
+			TriggerSharedWorkerThread(MoveTemp(Msg));
 		}
 		void SendBitrateMessage(EStreamType type, int32 value, int32 which)
 		{
 			FMessage Msg;
 			Msg.Type = FMessage::EType::ChangeBitrate;
 			Msg.Data.Bitrate.Value = value;
-			WorkMessages.SendMessage(Msg);
+			TriggerSharedWorkerThread(MoveTemp(Msg));
 		}
 		void SendPlayerSessionMessage(const TSharedPtrTS<IPlayerMessage>& message)
 		{
 			FMessage Msg;
 			Msg.Type = FMessage::EType::PlayerSession;
 			Msg.Data.Session.PlayerMessage = message;
-			WorkMessages.SendMessage(Msg);
+			TriggerSharedWorkerThread(MoveTemp(Msg));
 		}
 		void SendResolutionMessage(int32 Width, int32 Height)
 		{
@@ -1139,7 +1181,7 @@ private:
 			Msg.Type = FMessage::EType::LimitResolution;
 			Msg.Data.Resolution.Width = Width;
 			Msg.Data.Resolution.Height = Height;
-			WorkMessages.SendMessage(Msg);
+			TriggerSharedWorkerThread(MoveTemp(Msg));
 		}
 		void SendInitialStreamAttributeMessage(EStreamType StreamType, const FStreamSelectionAttributes& InitialSelection)
 		{
@@ -1147,7 +1189,7 @@ private:
 			Msg.Type = FMessage::EType::InitialStreamAttributes;
 			Msg.Data.InitialStreamAttribute.StreamType = StreamType;
 			Msg.Data.InitialStreamAttribute.InitialSelection = InitialSelection;
-			WorkMessages.SendMessage(Msg);
+			TriggerSharedWorkerThread(MoveTemp(Msg));
 		}
 		void SendTrackSelectByMetadataMessage(EStreamType StreamType, const FTrackMetadata& TrackMetadata)
 		{
@@ -1155,7 +1197,7 @@ private:
 			Msg.Type = FMessage::EType::SelectTrackByMetadata;
 			Msg.Data.TrackSelection.StreamType = StreamType;
 			Msg.Data.TrackSelection.TrackMetadata = TrackMetadata;
-			WorkMessages.SendMessage(Msg);
+			TriggerSharedWorkerThread(MoveTemp(Msg));
 		}
 		void SendTrackSelectByAttributeMessage(EStreamType StreamType, const FStreamSelectionAttributes& TrackAttributes)
 		{
@@ -1163,19 +1205,26 @@ private:
 			Msg.Type = FMessage::EType::SelectTrackByAttributes;
 			Msg.Data.TrackSelection.StreamType = StreamType;
 			Msg.Data.TrackSelection.TrackAttributes = TrackAttributes;
-			WorkMessages.SendMessage(Msg);
+			TriggerSharedWorkerThread(MoveTemp(Msg));
 		}
 		void SendTrackDeselectMessage(EStreamType StreamType)
 		{
 			FMessage Msg;
 			Msg.Type = FMessage::EType::DeselectTrack;
 			Msg.Data.TrackSelection.StreamType = StreamType;
-			WorkMessages.SendMessage(Msg);
+			TriggerSharedWorkerThread(MoveTemp(Msg));
 		}
 
-		FMediaThread									MediaThread;
-		TMediaMessageQueueDynamicWithTimeout<FMessage>	WorkMessages;
-		bool											bStarted;
+		void TriggerSharedWorkerThread(FMessage Mesg)
+		{
+			WorkMessages.Enqueue(MoveTemp(Mesg));
+			if (SharedWorkerThread.IsValid())
+			{
+				SharedWorkerThread->TriggerWork();
+			}
+		}
+		TSharedPtrTS<FAdaptiveStreamingPlayerWorkerThread>	SharedWorkerThread;
+		TQueue<FMessage, EQueueMode::Spsc>					WorkMessages;
 	};
 
 	struct FPendingStartRequest
@@ -1340,6 +1389,7 @@ private:
 
 	void InternalLoadManifest(const FString& URL, const FString& MimeType);
 	void InternalCancelLoadManifest();
+	void InternalCloseManifestReader();
 	bool SelectManifest();
 	void UpdateManifest();
 	void OnManifestGetMimeTypeComplete(TSharedPtrTS<FHTTPResourceRequest> InRequest);
@@ -1350,9 +1400,6 @@ private:
 	void AudioDecoderInputNeeded(const IAccessUnitBufferListener::FBufferStats& currentInputBufferStats);
 	void AudioDecoderOutputReady(const IDecoderOutputBufferListener::FDecodeReadyStats& currentReadyStats);
 
-
-
-	void WorkerThreadFN();
 	void StartWorkerThread();
 	void StopWorkerThread();
 
@@ -1375,6 +1422,9 @@ private:
 	void OnFragmentReachedEOS(EStreamType InStreamType, TSharedPtr<const FBufferSourceInfo, ESPMode::ThreadSafe> InStreamSourceInfo) override;
 	void OnFragmentClose(TSharedPtrTS<IStreamSegment> pRequest) override;
 
+	void InternalHandleOnce();
+	bool InternalHandleThreadMessages();
+	void InternalHandleManifestReader();
 	void InternalHandlePendingStartRequest(const FTimeValue& CurrentTime);
 	void InternalHandlePendingFirstSegmentRequest(const FTimeValue& CurrentTime);
 	void InternalHandleCompletedSegmentRequests(const FTimeValue& CurrentTime);
@@ -1428,8 +1478,9 @@ private:
 	TWeakPtr<IAdaptiveStreamingPlayerResourceProvider, ESPMode::ThreadSafe> StaticResourceProvider;
 	TWeakPtr<IVideoDecoderResourceDelegate, ESPMode::ThreadSafe>		VideoDecoderResourceDelegate;
 
-	FWorkerThread															WorkerThread;
-	TSharedPtr<FAdaptiveStreamingPlayerEventHandler, ESPMode::ThreadSafe>	EventDispatcher;
+	TSharedPtrTS<FAdaptiveStreamingPlayerEventHandler>					EventDispatcher;
+	TSharedPtrTS<FAdaptiveStreamingPlayerWorkerThread>					SharedWorkerThread;
+	FWorkerThreadMessages												WorkerThread;
 
 	FParamDict															PlayerOptions;
 	FPlaybackState														PlaybackState;
diff --git a/Engine/Plugins/Media/ElectraPlayer/Source/ElectraPlayerRuntime/Private/Runtime/Player/AdaptiveStreamingPlayerManifest.cpp b/Engine/Plugins/Media/ElectraPlayer/Source/ElectraPlayerRuntime/Private/Runtime/Player/AdaptiveStreamingPlayerManifest.cpp
--- a/Engine/Plugins/Media/ElectraPlayer/Source/ElectraPlayerRuntime/Private/Runtime/Player/AdaptiveStreamingPlayerManifest.cpp
+++ b/Engine/Plugins/Media/ElectraPlayer/Source/ElectraPlayerRuntime/Private/Runtime/Player/AdaptiveStreamingPlayerManifest.cpp
@@ -113,6 +113,16 @@ void FAdaptiveStreamingPlayer::InternalCancelLoadManifest()
 }
 
 
+void FAdaptiveStreamingPlayer::InternalCloseManifestReader()
+{
+	if (ManifestReader.IsValid())
+	{
+		ManifestReader->Close();
+		ManifestReader.Reset();
+	}
+}
+
+
 //-----------------------------------------------------------------------------
 /**
  * Starts asynchronous loading and parsing of a manifest.
@@ -198,6 +208,16 @@ void FAdaptiveStreamingPlayer::InternalLoadManifest(const FString& URL, const FS
 }
 
 
+void FAdaptiveStreamingPlayer::InternalHandleManifestReader()
+{
+	if (ManifestReader.IsValid())
+	{
+		ManifestReader->HandleOnce();
+	}
+}
+
+
+
 //-----------------------------------------------------------------------------
 /**
  * Selects the internal presentation for playback after having selected/disabled candidate streams via AccessManifest().
@@ -206,7 +226,7 @@ void FAdaptiveStreamingPlayer::InternalLoadManifest(const FString& URL, const FS
  */
 bool FAdaptiveStreamingPlayer::SelectManifest()
 {
-	if (ManifestReader)
+	if (ManifestReader.IsValid())
 	{
 		check(Manifest == nullptr);
 		if (ManifestType != EMediaFormatType::Unknown)
@@ -237,6 +257,12 @@ bool FAdaptiveStreamingPlayer::SelectManifest()
 			PlayerConfig.SeekBufferMinTimeAvailBeforePlayback    = Utils::Min(minBufTimeMPD, PlayerConfig.SeekBufferMinTimeAvailBeforePlayback);
 			PlayerConfig.RebufferMinTimeAvailBeforePlayback 	 = Utils::Min(minBufTimeMPD, PlayerConfig.RebufferMinTimeAvailBeforePlayback);
 
+			// For an mp4 stream we can now get rid of the manifest reader. It is no longer needed and we don't need to have it linger.
+			if (ManifestType == EMediaFormatType::ISOBMFF)
+			{
+				InternalCloseManifestReader();
+			}
+
 			return true;
 		}
 		else
diff --git a/Engine/Plugins/Media/ElectraPlayer/Source/ElectraPlayerRuntime/Private/Runtime/Player/DASH/PlaylistReaderDASH.cpp b/Engine/Plugins/Media/ElectraPlayer/Source/ElectraPlayerRuntime/Private/Runtime/Player/DASH/PlaylistReaderDASH.cpp
--- a/Engine/Plugins/Media/ElectraPlayer/Source/ElectraPlayerRuntime/Private/Runtime/Player/DASH/PlaylistReaderDASH.cpp
+++ b/Engine/Plugins/Media/ElectraPlayer/Source/ElectraPlayerRuntime/Private/Runtime/Player/DASH/PlaylistReaderDASH.cpp
@@ -32,6 +32,8 @@
 DECLARE_CYCLE_STAT(TEXT("FPlaylistReaderDASH_WorkerThread"), STAT_ElectraPlayer_DASH_PlaylistWorker, STATGROUP_ElectraPlayer);
 
 
+#define DASH_PLAYLISTREADER_USE_DEDICATED_WORKER_THREAD		0
+
 //#define DO_NOT_PERFORM_CONDITIONAL_GET
 
 namespace Electra
@@ -40,7 +42,10 @@ namespace Electra
 /**
  * This class is responsible for downloading a DASH MPD and parsing it.
  */
-class FPlaylistReaderDASH : public TSharedFromThis<FPlaylistReaderDASH, ESPMode::ThreadSafe>, public IPlaylistReaderDASH, public IAdaptiveStreamingPlayerAEMSReceiver, public FMediaThread
+class FPlaylistReaderDASH : public TSharedFromThis<FPlaylistReaderDASH, ESPMode::ThreadSafe>, public IPlaylistReaderDASH, public IAdaptiveStreamingPlayerAEMSReceiver
+#if DASH_PLAYLISTREADER_USE_DEDICATED_WORKER_THREAD
+						  , public FMediaThread
+#endif
 {
 public:
 	FPlaylistReaderDASH();
@@ -48,6 +53,7 @@ public:
 
 	virtual ~FPlaylistReaderDASH();
 	virtual void Close() override;
+	virtual void HandleOnce() override;
 	virtual const FString& GetPlaylistType() const override
 	{
 		static FString Type("dash");
@@ -78,6 +84,11 @@ private:
 	void StartWorkerThread();
 	void StopWorkerThread();
 	void WorkerThread();
+
+	void InternalSetup();
+	void InternalCleanup();
+	void InternalHandleOnce();
+
 	void ExecutePendingRequests(const FTimeValue& TimeNow);
 	void HandleCompletedRequests(const FTimeValue& TimeNow);
 	void HandleStaticRequestCompletions(const FTimeValue& TimeNow);
@@ -175,7 +186,9 @@ TSharedPtrTS<IPlaylistReader> IPlaylistReaderDASH::Create(IPlayerSessionServices
 
 
 FPlaylistReaderDASH::FPlaylistReaderDASH()
+#if DASH_PLAYLISTREADER_USE_DEDICATED_WORKER_THREAD
 	: FMediaThread("ElectraPlayer::DASH MPD")
+#endif
 {
 }
 
@@ -204,11 +217,45 @@ void FPlaylistReaderDASH::Close()
 	StopWorkerThread();
 }
 
+void FPlaylistReaderDASH::HandleOnce()
+{
+#if !DASH_PLAYLISTREADER_USE_DEDICATED_WORKER_THREAD
+	InternalHandleOnce();
+#endif
+}
+
+void FPlaylistReaderDASH::InternalHandleOnce()
+{
+	if (bIsWorkerThreadStarted)
+	{
+		SCOPE_CYCLE_COUNTER(STAT_ElectraPlayer_DASH_PlaylistWorker);
+		CSV_SCOPED_TIMING_STAT(ElectraPlayer, DASH_PlaylistWorker);
+		FTimeValue Now = PlayerSessionServices->GetSynchronizedUTCTime()->GetTime();
+
+		HandleCompletedRequests(Now);
+
+		ExecutePendingRequests(Now);
+
+		// Check if the MPD must be updated.
+		CheckForMPDUpdate();
+
+		// Time sync
+		if (NextTimeSyncTime.IsValid() && Now >= NextTimeSyncTime)
+		{
+			TriggerTimeSynchronization();
+		}
+	}
+}
+
 void FPlaylistReaderDASH::StartWorkerThread()
 {
 	check(!bIsWorkerThreadStarted);
+#if DASH_PLAYLISTREADER_USE_DEDICATED_WORKER_THREAD
 	bTerminateWorkerThread = false;
 	ThreadStart(Electra::MakeDelegate(this, &FPlaylistReaderDASH::WorkerThread));
+#else
+	InternalSetup();
+#endif
 	bIsWorkerThreadStarted = true;
 }
 
@@ -216,10 +263,14 @@ void FPlaylistReaderDASH::StopWorkerThread()
 {
 	if (bIsWorkerThreadStarted)
 	{
+#if DASH_PLAYLISTREADER_USE_DEDICATED_WORKER_THREAD
 		bTerminateWorkerThread = true;
 		WorkerThreadSignal.Release();
 		ThreadWaitDone();
 		ThreadReset();
+#else
+		InternalCleanup();
+#endif
 		bIsWorkerThreadStarted = false;
 	}
 }
@@ -720,11 +771,8 @@ void FPlaylistReaderDASH::TriggerTimeSynchronization()
 }
 
 
-
-void FPlaylistReaderDASH::WorkerThread()
+void FPlaylistReaderDASH::InternalSetup()
 {
-	LLM_SCOPE(ELLMTag::ElectraPlayer);
-
 	bInbandEventByStreamType[0] = false;
 	bInbandEventByStreamType[1] = false;
 	bInbandEventByStreamType[2] = false;
@@ -748,34 +796,10 @@ void FPlaylistReaderDASH::WorkerThread()
 	Fragment = UrlParser.GetFragment();
 	PlaylistLoadRequest->CompleteCallback.BindThreadSafeSP(AsShared(), &FPlaylistReaderDASH::ManifestDownloadCompleted);
 	EnqueueResourceRequest(MoveTemp(PlaylistLoadRequest));
+}
 
-	while(!bTerminateWorkerThread)
-	{
-		WorkerThreadSignal.Obtain(1000 * 100);
-		if (bTerminateWorkerThread)
-		{
-			break;
-		}
-		{
-			SCOPE_CYCLE_COUNTER(STAT_ElectraPlayer_DASH_PlaylistWorker);
-			CSV_SCOPED_TIMING_STAT(ElectraPlayer, DASH_PlaylistWorker);
-			FTimeValue Now = PlayerSessionServices->GetSynchronizedUTCTime()->GetTime();
-
-			HandleCompletedRequests(Now);
-
-			ExecutePendingRequests(Now);
-
-			// Check if the MPD must be updated.
-			CheckForMPDUpdate();
-
-			// Time sync
-			if (NextTimeSyncTime.IsValid() && Now >= NextTimeSyncTime)
-			{
-				TriggerTimeSynchronization();
-			}
-		}
-	}
-
+void FPlaylistReaderDASH::InternalCleanup()
+{
 	// Unregister event callbacks
 	PlayerSessionServices->GetAEMSEventHandler()->RemoveAEMSReceiver(AsShared(), DASH::Schemes::ManifestEvents::Scheme_urn_mpeg_dash_event_ttfn_2016, TEXT(""), IAdaptiveStreamingPlayerAEMSReceiver::EDispatchMode::OnStart);
 	PlayerSessionServices->GetAEMSEventHandler()->RemoveAEMSReceiver(AsShared(), DASH::Schemes::ManifestEvents::Scheme_urn_mpeg_dash_event_callback_2015, TEXT("1"), IAdaptiveStreamingPlayerAEMSReceiver::EDispatchMode::OnStart);
@@ -797,6 +821,27 @@ void FPlaylistReaderDASH::WorkerThread()
 	Builder.Reset();
 }
 
+void FPlaylistReaderDASH::WorkerThread()
+{
+#if DASH_PLAYLISTREADER_USE_DEDICATED_WORKER_THREAD
+	LLM_SCOPE(ELLMTag::ElectraPlayer);
+
+	InternalSetup();
+
+	while(!bTerminateWorkerThread)
+	{
+		WorkerThreadSignal.Obtain(1000 * 100);
+		if (bTerminateWorkerThread)
+		{
+			break;
+		}
+		InternalHandleOnce();
+	}
+
+	InternalCleanup();
+#endif
+}
+
 
 void FPlaylistReaderDASH::ExecutePendingRequests(const FTimeValue& TimeNow)
 {
diff --git a/Engine/Plugins/Media/ElectraPlayer/Source/ElectraPlayerRuntime/Private/Runtime/Player/DASH/StreamReaderFMP4DASH.cpp b/Engine/Plugins/Media/ElectraPlayer/Source/ElectraPlayerRuntime/Private/Runtime/Player/DASH/StreamReaderFMP4DASH.cpp
--- a/Engine/Plugins/Media/ElectraPlayer/Source/ElectraPlayerRuntime/Private/Runtime/Player/DASH/StreamReaderFMP4DASH.cpp
+++ b/Engine/Plugins/Media/ElectraPlayer/Source/ElectraPlayerRuntime/Private/Runtime/Player/DASH/StreamReaderFMP4DASH.cpp
@@ -133,13 +133,13 @@ UEMediaError FStreamReaderFMP4DASH::Create(IPlayerSessionServices* InPlayerSessi
 		StreamHandlers[i].PlayerSessionService = PlayerSessionService;
 		StreamHandlers[i].Parameters		   = InCreateParam;
 		StreamHandlers[i].bTerminate		   = false;
+		StreamHandlers[i].bWasStarted		   = false;
 		StreamHandlers[i].bRequestCanceled     = false;
 		StreamHandlers[i].bSilentCancellation  = false;
 		StreamHandlers[i].bHasErrored   	   = false;
 		StreamHandlers[i].IsIdleSignal.Signal();
 
 		StreamHandlers[i].ThreadSetName(i==0?"ElectraPlayer::fmp4 Video":"ElectraPlayer::fmp4 Audio");
-		StreamHandlers[i].ThreadStart(Electra::MakeDelegate(&StreamHandlers[i], &FStreamHandler::WorkerThread));
 	}
 	return UEMEDIA_ERROR_OK;
 }
@@ -159,8 +159,11 @@ void FStreamReaderFMP4DASH::Close()
 		// Wait until they finished.
 		for(int32 i=0; i<FMEDIA_STATIC_ARRAY_COUNT(StreamHandlers); ++i)
 		{
-			StreamHandlers[i].ThreadWaitDone();
-			StreamHandlers[i].ThreadReset();
+			if (StreamHandlers[i].bWasStarted)
+			{
+				StreamHandlers[i].ThreadWaitDone();
+				StreamHandlers[i].ThreadReset();
+			}
 		}
 	}
 }
@@ -218,6 +221,12 @@ IStreamReader::EAddResult FStreamReaderFMP4DASH::AddRequest(uint32 CurrentPlayba
 		}
 
 		Request->SetPlaybackSequenceID(CurrentPlaybackSequenceID);
+		if (!Handler->bWasStarted)
+		{
+			Handler->ThreadStart(Electra::MakeDelegate(Handler, &FStreamHandler::WorkerThread));
+			Handler->bWasStarted = true;
+		}
+
 		Handler->bRequestCanceled = false;
 		Handler->bSilentCancellation = false;
 		Handler->CurrentRequest = Request;
diff --git a/Engine/Plugins/Media/ElectraPlayer/Source/ElectraPlayerRuntime/Private/Runtime/Player/DASH/StreamReaderFMP4DASH.h b/Engine/Plugins/Media/ElectraPlayer/Source/ElectraPlayerRuntime/Private/Runtime/Player/DASH/StreamReaderFMP4DASH.h
--- a/Engine/Plugins/Media/ElectraPlayer/Source/ElectraPlayerRuntime/Private/Runtime/Player/DASH/StreamReaderFMP4DASH.h
+++ b/Engine/Plugins/Media/ElectraPlayer/Source/ElectraPlayerRuntime/Private/Runtime/Player/DASH/StreamReaderFMP4DASH.h
@@ -126,6 +126,7 @@ private:
 		FMediaSemaphore											WorkSignal;
 		FMediaEvent												IsIdleSignal;
 		volatile bool											bTerminate = false;
+		volatile bool											bWasStarted = false;
 		volatile bool											bRequestCanceled = false;
 		volatile bool											bSilentCancellation = false;
 		volatile bool											bHasErrored = false;
diff --git a/Engine/Plugins/Media/ElectraPlayer/Source/ElectraPlayerRuntime/Private/Runtime/Player/HLS/PlaylistReaderHLS.cpp b/Engine/Plugins/Media/ElectraPlayer/Source/ElectraPlayerRuntime/Private/Runtime/Player/HLS/PlaylistReaderHLS.cpp
--- a/Engine/Plugins/Media/ElectraPlayer/Source/ElectraPlayerRuntime/Private/Runtime/Player/HLS/PlaylistReaderHLS.cpp
+++ b/Engine/Plugins/Media/ElectraPlayer/Source/ElectraPlayerRuntime/Private/Runtime/Player/HLS/PlaylistReaderHLS.cpp
@@ -26,6 +26,7 @@
 
 DECLARE_CYCLE_STAT(TEXT("FPlaylistReaderHLS_WorkerThread"), STAT_ElectraPlayer_HLS_PlaylistWorker, STATGROUP_ElectraPlayer);
 
+#define HLS_PLAYLISTREADER_USE_DEDICATED_WORKER_THREAD		0
 
 namespace Electra
 {
@@ -46,7 +47,10 @@ const FString IPlaylistReaderHLS::OptionKeyUpdatePlaylistLoadNoDataTimeout  (TEX
 /**
  * This class is responsible for downloading HLS playlists and parsing them.
  */
-class FPlaylistReaderHLS : public IPlaylistReaderHLS, public FMediaThread
+class FPlaylistReaderHLS : public IPlaylistReaderHLS
+#if HLS_PLAYLISTREADER_USE_DEDICATED_WORKER_THREAD
+						 , public FMediaThread
+#endif
 {
 public:
 	FPlaylistReaderHLS();
@@ -54,6 +58,7 @@ public:
 
 	virtual ~FPlaylistReaderHLS();
 	virtual void Close() override;
+	virtual void HandleOnce() override;
 
 	/**
 	 * Returns the type of playlist format.
@@ -225,6 +230,11 @@ private:
 	void StartWorkerThread();
 	void StopWorkerThread();
 	void WorkerThread();
+
+	void InternalSetup();
+	void InternalCleanup();
+	void InternalHandleOnce();
+
 	void HandleEnqueuedPlaylistDownloads(const FTimeValue& TimeNow);
 	void HandleCompletedPlaylistDownloads(const FTimeValue& TimeNow);
 	void HandleStaticRequestCompletions(const FTimeValue& TimeNow);
@@ -244,11 +254,11 @@ private:
 	FErrorDetail ParsePlaylist(FPlaylistRequestPtr FromRequest);
 
 
-	IPlayerSessionServices*									PlayerSessionServices;
+	IPlayerSessionServices*									PlayerSessionServices = nullptr;
 	FString													MasterPlaylistURL;
 	TSharedPtr<FMediaSemaphore, ESPMode::ThreadSafe>		WorkerThreadSignal;
-	bool													bIsWorkerThreadStarted;
-	volatile bool											bTerminateWorkerThread;
+	bool													bIsWorkerThreadStarted = false;
+	volatile bool											bTerminateWorkerThread = false;
 
 	FMediaCriticalSection									PendingPlaylistRequestsLock;
 	TArray<FPlaylistRequestPtr>								PendingPlaylistRequests;
@@ -258,7 +268,7 @@ private:
 	TDoubleLinkedList<FPlaylistLoadRequestHLS>				InitiallyRequiredPlaylistLoadRequests;				//!< Playlists that need to be fetched and parsed before we can report metadata ready to the player.
 	TSharedPtrTS<IElectraHttpManager::FProgressListener>	ProgressListener;
 
-	IManifestBuilderHLS* 									Builder;
+	IManifestBuilderHLS* 									Builder = nullptr;
 	TSharedPtrTS<FManifestHLSInternal>						Manifest;
 	FErrorDetail											LastErrorDetail;
 
@@ -286,11 +296,9 @@ TSharedPtrTS<IPlaylistReader> IPlaylistReaderHLS::Create(IPlayerSessionServices*
 
 
 FPlaylistReaderHLS::FPlaylistReaderHLS()
+#if HLS_PLAYLISTREADER_USE_DEDICATED_WORKER_THREAD
 	: FMediaThread("ElectraPlayer::HLS Playlist")
-	, PlayerSessionServices(nullptr)
-	, bIsWorkerThreadStarted(false)
-	, bTerminateWorkerThread(false)
-	, Builder(nullptr)
+#endif
 {
 	WorkerThreadSignal = MakeShared<FMediaSemaphore, ESPMode::ThreadSafe>();
 }
@@ -322,15 +330,23 @@ void FPlaylistReaderHLS::Close()
 	StopWorkerThread();
 }
 
+void FPlaylistReaderHLS::HandleOnce()
+{
+#if !HLS_PLAYLISTREADER_USE_DEDICATED_WORKER_THREAD
+	InternalHandleOnce();
+#endif
+}
+
+
 void FPlaylistReaderHLS::StartWorkerThread()
 {
 	check(!bIsWorkerThreadStarted);
-
+#if HLS_PLAYLISTREADER_USE_DEDICATED_WORKER_THREAD
 	bTerminateWorkerThread = false;
-	//ThreadSetPriority( ... );
-	//ThreadSetCoreAffinity( ... );
-	//ThreadSetStackSize( ... );
 	ThreadStart(Electra::MakeDelegate(this, &FPlaylistReaderHLS::WorkerThread));
+#else
+	InternalSetup();
+#endif
 	bIsWorkerThreadStarted = true;
 }
 
@@ -338,10 +354,14 @@ void FPlaylistReaderHLS::StopWorkerThread()
 {
 	if (bIsWorkerThreadStarted)
 	{
+#if HLS_PLAYLISTREADER_USE_DEDICATED_WORKER_THREAD
 		bTerminateWorkerThread = true;
 		WorkerThreadSignal->Release();
 		ThreadWaitDone();
 		ThreadReset();
+#else
+		InternalCleanup();
+#endif
 		bIsWorkerThreadStarted = false;
 	}
 }
@@ -455,10 +475,9 @@ void FPlaylistReaderHLS::CheckForPlaylistUpdate(const FTimeValue& TimeNow)
 	}
 }
 
-void FPlaylistReaderHLS::WorkerThread()
-{
-	LLM_SCOPE(ELLMTag::ElectraPlayer);
 
+void FPlaylistReaderHLS::InternalSetup()
+{
 	Builder = IManifestBuilderHLS::Create(PlayerSessionServices);
 
 	// Setup the playlist load request for the master playlist.
@@ -467,38 +486,10 @@ void FPlaylistReaderHLS::WorkerThread()
 	PlaylistLoadRequest.RequestedAtTime  = PlayerSessionServices->GetSynchronizedUTCTime()->GetTime();
 	PlaylistLoadRequest.InternalUniqueID = 0;
 	EnqueueLoadPlaylist(PlaylistLoadRequest, true);
+}
 
-	while(!bTerminateWorkerThread)
-	{
-		WorkerThreadSignal->Obtain(1000 * 100);
-		if (bTerminateWorkerThread)
-		{
-			break;
-		}
-
-		{
-			SCOPE_CYCLE_COUNTER(STAT_ElectraPlayer_HLS_PlaylistWorker);
-			CSV_SCOPED_TIMING_STAT(ElectraPlayer, HLS_PlaylistWorker);
-
-			FTimeValue Now = PlayerSessionServices->GetSynchronizedUTCTime()->GetTime();
-
-			// Check for completed static resource requests.
-			// NOTE: This needs to be done first in order to move the completed ones being moved onto the completed list to be
-			//       handled immediately in the following HandleCompletedPlaylistDownloads() call!!
-			HandleStaticRequestCompletions(Now);
-
-			// Check completed playlist downloads
-			HandleCompletedPlaylistDownloads(Now);
-
-			// Execute enqueued playlist download requests.
-			HandleEnqueuedPlaylistDownloads(Now);
-
-			// Check which playlists need to be reloaded.
-			CheckForPlaylistUpdate(Now);
-		}
-	}
-
-
+void FPlaylistReaderHLS::InternalCleanup()
+{
 	// No playlists are required any more.
 	InitiallyRequiredPlaylistLoadRequests.Empty();
 
@@ -544,6 +535,48 @@ void FPlaylistReaderHLS::WorkerThread()
 	Builder = nullptr;
 }
 
+void FPlaylistReaderHLS::InternalHandleOnce()
+{
+	SCOPE_CYCLE_COUNTER(STAT_ElectraPlayer_HLS_PlaylistWorker);
+	CSV_SCOPED_TIMING_STAT(ElectraPlayer, HLS_PlaylistWorker);
+
+	FTimeValue Now = PlayerSessionServices->GetSynchronizedUTCTime()->GetTime();
+
+	// Check for completed static resource requests.
+	// NOTE: This needs to be done first in order to move the completed ones being moved onto the completed list to be
+	//       handled immediately in the following HandleCompletedPlaylistDownloads() call!!
+	HandleStaticRequestCompletions(Now);
+
+	// Check completed playlist downloads
+	HandleCompletedPlaylistDownloads(Now);
+
+	// Execute enqueued playlist download requests.
+	HandleEnqueuedPlaylistDownloads(Now);
+
+	// Check which playlists need to be reloaded.
+	CheckForPlaylistUpdate(Now);
+}
+
+void FPlaylistReaderHLS::WorkerThread()
+{
+	LLM_SCOPE(ELLMTag::ElectraPlayer);
+
+	InternalSetup();
+
+	while(!bTerminateWorkerThread)
+	{
+		WorkerThreadSignal->Obtain(1000 * 100);
+		if (bTerminateWorkerThread)
+		{
+			break;
+		}
+
+		InternalHandleOnce();
+	}
+
+	InternalCleanup();
+}
+
 
 void FPlaylistReaderHLS::FPlaylistRequest::ExecuteStatic(IAdaptiveStreamingPlayerResourceProvider* InStaticResourceProvider, TWeakPtr<FMediaSemaphore, ESPMode::ThreadSafe> DoneSignal)
 {
diff --git a/Engine/Plugins/Media/ElectraPlayer/Source/ElectraPlayerRuntime/Private/Runtime/Player/HLS/StreamReaderHLSfmp4.cpp b/Engine/Plugins/Media/ElectraPlayer/Source/ElectraPlayerRuntime/Private/Runtime/Player/HLS/StreamReaderHLSfmp4.cpp
--- a/Engine/Plugins/Media/ElectraPlayer/Source/ElectraPlayerRuntime/Private/Runtime/Player/HLS/StreamReaderHLSfmp4.cpp
+++ b/Engine/Plugins/Media/ElectraPlayer/Source/ElectraPlayerRuntime/Private/Runtime/Player/HLS/StreamReaderHLSfmp4.cpp
@@ -140,13 +140,13 @@ UEMediaError FStreamReaderHLSfmp4::Create(IPlayerSessionServices* InPlayerSessio
 	{
 		StreamHandlers[i].PlayerSessionService = PlayerSessionService;
 		StreamHandlers[i].Parameters		   = InCreateParam;
+		StreamHandlers[i].bWasStarted		   = false;
 		StreamHandlers[i].bTerminate		   = false;
 		StreamHandlers[i].bRequestCanceled     = false;
 		StreamHandlers[i].bSilentCancellation  = false;
 		StreamHandlers[i].bHasErrored   	   = false;
 
 		StreamHandlers[i].ThreadSetName(i==0?"ElectraPlayer::fmp4 Video":"ElectraPlayer::fmp4 Audio");
-		StreamHandlers[i].ThreadStart(Electra::MakeDelegate(&StreamHandlers[i], &FStreamHandler::WorkerThread));
 	}
 	return UEMEDIA_ERROR_OK;
 }
@@ -166,8 +166,11 @@ void FStreamReaderHLSfmp4::Close()
 		// Wait until they finished.
 		for(int32 i=0; i<FMEDIA_STATIC_ARRAY_COUNT(StreamHandlers); ++i)
 		{
-			StreamHandlers[i].ThreadWaitDone();
-			StreamHandlers[i].ThreadReset();
+			if (StreamHandlers[i].bWasStarted)
+			{
+				StreamHandlers[i].ThreadWaitDone();
+				StreamHandlers[i].ThreadReset();
+			}
 		}
 	}
 }
@@ -260,6 +263,11 @@ IStreamReader::EAddResult FStreamReaderHLSfmp4::AddRequest(uint32 CurrentPlaybac
 		// Only add the request if this is not an EOD segment.
 		if (!Request2->bIsEOSSegment)
 		{
+			if (!Handler2->bWasStarted)
+			{
+				Handler2->ThreadStart(Electra::MakeDelegate(Handler2, &FStreamHandler::WorkerThread));
+				Handler2->bWasStarted = true;
+			}
 			Handler2->bRequestCanceled = false;
 			Handler2->bSilentCancellation = false;
 			Handler2->CurrentRequest = Request2;
@@ -270,6 +278,11 @@ IStreamReader::EAddResult FStreamReaderHLSfmp4::AddRequest(uint32 CurrentPlaybac
 	// Only add the request if this is not an EOD segment.
 	if (!Request->bIsEOSSegment)
 	{
+		if (!Handler->bWasStarted)
+		{
+			Handler->ThreadStart(Electra::MakeDelegate(Handler, &FStreamHandler::WorkerThread));
+			Handler->bWasStarted = true;
+		}
 		Handler->bRequestCanceled = false;
 		Handler->bSilentCancellation = false;
 		Handler->CurrentRequest = Request;
@@ -307,14 +320,6 @@ uint32 FStreamReaderHLSfmp4::FStreamHandler::UniqueDownloadID = 1;
 
 FStreamReaderHLSfmp4::FStreamHandler::FStreamHandler()
 {
-	PlayerSessionService = nullptr;
-	bTerminate  		 = false;
-	bRequestCanceled	 = false;
-	bSilentCancellation  = false;
-	bAbortedByABR   	 = false;
-	bHasErrored 		 = false;
-	NumMOOFBoxesFound    = 0;
-	ProgressReportCount  = 0;
 }
 
 FStreamReaderHLSfmp4::FStreamHandler::~FStreamHandler()
diff --git a/Engine/Plugins/Media/ElectraPlayer/Source/ElectraPlayerRuntime/Private/Runtime/Player/HLS/StreamReaderHLSfmp4.h b/Engine/Plugins/Media/ElectraPlayer/Source/ElectraPlayerRuntime/Private/Runtime/Player/HLS/StreamReaderHLSfmp4.h
--- a/Engine/Plugins/Media/ElectraPlayer/Source/ElectraPlayerRuntime/Private/Runtime/Player/HLS/StreamReaderHLSfmp4.h
+++ b/Engine/Plugins/Media/ElectraPlayer/Source/ElectraPlayerRuntime/Private/Runtime/Player/HLS/StreamReaderHLSfmp4.h
@@ -196,27 +196,28 @@ private:
 		IStreamReader::CreateParam								Parameters;
 		TSharedPtrTS<FStreamSegmentRequestHLSfmp4>				CurrentRequest;
 		FMediaSemaphore											WorkSignal;
-		volatile bool											bTerminate;
-		volatile bool											bRequestCanceled;
-		volatile bool											bSilentCancellation;
-		volatile bool											bHasErrored;
-		bool													bAbortedByABR;
-		bool													bAllowEarlyEmitting;
-		bool													bFillRemainingDuration;
-
-		IPlayerSessionServices*									PlayerSessionService;
+		volatile bool											bWasStarted = false;
+		volatile bool											bTerminate = false;
+		volatile bool											bRequestCanceled = false;
+		volatile bool											bSilentCancellation = false;
+		volatile bool											bHasErrored = false;
+		bool													bAbortedByABR = false;
+		bool													bAllowEarlyEmitting = false;
+		bool													bFillRemainingDuration = false;
+
+		IPlayerSessionServices*									PlayerSessionService = nullptr;
 		FReadBuffer												ReadBuffer;
 		TSharedPtr<ElectraCDM::IStreamDecrypterAES128, ESPMode::ThreadSafe>	Decrypter;
 		FMediaEvent												DownloadCompleteSignal;
 		TSharedPtrTS<IParserISO14496_12>						MP4Parser;
-		int32													NumMOOFBoxesFound;
+		int32													NumMOOFBoxesFound = 0;
 
 		TMediaQueueDynamicNoLock<FAccessUnit *>					AccessUnitFIFO;
 		FTimeValue 												DurationSuccessfullyRead;
 		FTimeValue 												DurationSuccessfullyDelivered;
 
 		FMediaCriticalSection									MetricUpdateLock;
-		int32													ProgressReportCount;
+		int32													ProgressReportCount = 0;
 		TSharedPtrTS<IAdaptiveStreamSelector>					StreamSelector;
 
 
diff --git a/Engine/Plugins/Media/ElectraPlayer/Source/ElectraPlayerRuntime/Private/Runtime/Player/PlaylistReader.h b/Engine/Plugins/Media/ElectraPlayer/Source/ElectraPlayerRuntime/Private/Runtime/Player/PlaylistReader.h
--- a/Engine/Plugins/Media/ElectraPlayer/Source/ElectraPlayerRuntime/Private/Runtime/Player/PlaylistReader.h
+++ b/Engine/Plugins/Media/ElectraPlayer/Source/ElectraPlayerRuntime/Private/Runtime/Player/PlaylistReader.h
@@ -26,6 +26,11 @@ namespace Electra
 		 */
 		virtual void Close() = 0;
 
+		/**
+		 * Called every so often by the player's worker thread to handle this class.
+		 */
+		virtual void HandleOnce() = 0;
+
 		/**
 		 * Returns the type of the playlist (eg. "hls", "dash", etc.)
 		 */
diff --git a/Engine/Plugins/Media/ElectraPlayer/Source/ElectraPlayerRuntime/Private/Runtime/Player/mp4/PlaylistReaderMP4.cpp b/Engine/Plugins/Media/ElectraPlayer/Source/ElectraPlayerRuntime/Private/Runtime/Player/mp4/PlaylistReaderMP4.cpp
--- a/Engine/Plugins/Media/ElectraPlayer/Source/ElectraPlayerRuntime/Private/Runtime/Player/mp4/PlaylistReaderMP4.cpp
+++ b/Engine/Plugins/Media/ElectraPlayer/Source/ElectraPlayerRuntime/Private/Runtime/Player/mp4/PlaylistReaderMP4.cpp
@@ -36,6 +36,7 @@ public:
 	virtual ~FPlaylistReaderMP4();
 
 	virtual void Close() override;
+	virtual void HandleOnce() override;
 
 	/**
 	 * Returns the type of playlist format.
@@ -188,6 +189,12 @@ void FPlaylistReaderMP4::Close()
 	StopWorkerThread();
 }
 
+void FPlaylistReaderMP4::HandleOnce()
+{
+	// No-op. This class is using a dedicated thread to read data from the stream
+	// which can stall at any moment and thus not lend itself to a tickable instance.
+}
+
 void FPlaylistReaderMP4::StartWorkerThread()
 {
 	check(!bIsWorkerThreadStarted);
