ElectraPlayer: Fixing an incorrect return value on Android when the video decoder returns an input configuration buffer back as an output.

#rnx

diff --git a/Engine/Plugins/Media/ElectraPlayer/Source/ElectraPlayerRuntime/Private/Runtime/Decoder/Android/VideoDecoderH264_Android.cpp b/Engine/Plugins/Media/ElectraPlayer/Source/ElectraPlayerRuntime/Private/Runtime/Decoder/Android/VideoDecoderH264_Android.cpp
--- a/Engine/Plugins/Media/ElectraPlayer/Source/ElectraPlayerRuntime/Private/Runtime/Decoder/Android/VideoDecoderH264_Android.cpp
+++ b/Engine/Plugins/Media/ElectraPlayer/Source/ElectraPlayerRuntime/Private/Runtime/Decoder/Android/VideoDecoderH264_Android.cpp
@@ -1166,13 +1166,11 @@ FVideoDecoderH264::EOutputResult FVideoDecoderH264::GetOutput()
 				result = DecoderInstance->ReleaseOutputBuffer(OutputBufferInfo.BufferIndex, OutputBufferInfo.ValidCount, false, -1);
 				return EOutputResult::EOS;
 			}
-			// We have not seen a config buffer coming back from a decoder, which makes sense. But we do not put it past some
-			// implementation to do that, so for good measure let's ignore it.
+			// If a configuration buffer is returned to us we ignore it. This serves no purpose.
 			if (OutputBufferInfo.bIsConfig)
 			{
-				LogMessage(Electra::IInfoLog::Warning, "Got CSD buffer back?");
 				result = DecoderInstance->ReleaseOutputBuffer(OutputBufferInfo.BufferIndex, OutputBufferInfo.ValidCount, false, -1);
-				return EOutputResult::EOS;
+				return EOutputResult::Ok;
 			}
 
 			IAndroidJavaH264VideoDecoder::FOutputFormatInfo OutputFormatInfo;
