ElectraPlayer: Added a missing critical section to the HLS playlist reader that could potentially cause an issue when using threaded playlist download mode.

#rnx

diff --git a/Engine/Plugins/Media/ElectraPlayer/Source/ElectraPlayerRuntime/Private/Runtime/Player/HLS/PlaylistReaderHLS.cpp b/Engine/Plugins/Media/ElectraPlayer/Source/ElectraPlayerRuntime/Private/Runtime/Player/HLS/PlaylistReaderHLS.cpp
--- a/Engine/Plugins/Media/ElectraPlayer/Source/ElectraPlayerRuntime/Private/Runtime/Player/HLS/PlaylistReaderHLS.cpp
+++ b/Engine/Plugins/Media/ElectraPlayer/Source/ElectraPlayerRuntime/Private/Runtime/Player/HLS/PlaylistReaderHLS.cpp
@@ -261,6 +261,7 @@ private:
 	volatile bool											bTerminateWorkerThread = false;
 
 	FMediaCriticalSection									PendingPlaylistRequestsLock;
+	FMediaCriticalSection									EnqueuedPlaylistRequestsLock;
 	TArray<FPlaylistRequestPtr>								PendingPlaylistRequests;
 	TArray<FPlaylistRequestPtr>								StaticResourcePlaylistRequests;
 	TArray<FPlaylistRequestPtr>								EnqueuedPlaylistRequests;
@@ -419,7 +420,9 @@ void FPlaylistReaderHLS::EnqueueLoadPlaylist(const FPlaylistLoadRequestHLS& InPl
 {
 	FPlaylistRequestPtr Request = MakeShared<FPlaylistRequest, ESPMode::ThreadSafe>(InPlaylistLoadRequest);
 	Request->SetIsMasterPlaylist(bIsMasterPlaylist);
+	EnqueuedPlaylistRequestsLock.Lock();
 	EnqueuedPlaylistRequests.Push(Request);
+	EnqueuedPlaylistRequestsLock.Unlock();
 	WorkerThreadSignal->Release();
 }
 
@@ -427,7 +430,9 @@ void FPlaylistReaderHLS::EnqueueRetryLoadPlaylist(FPlaylistRequestPtr RequestToR
 {
 	RequestToRetry->SetExecuteAtUTC(AtUTCTime);
 	RequestToRetry->SetRetryInfo(RetryInfo);
+	EnqueuedPlaylistRequestsLock.Lock();
 	EnqueuedPlaylistRequests.Push(RequestToRetry);
+	EnqueuedPlaylistRequestsLock.Unlock();
 	WorkerThreadSignal->Release();
 }
 
@@ -520,10 +525,9 @@ void FPlaylistReaderHLS::InternalCleanup()
 	ActivePendingRequests.Empty();
 
 	// Delete any unprocessed enqueued requests
-	while(EnqueuedPlaylistRequests.Num())
-	{
-		EnqueuedPlaylistRequests.Pop();
-	}
+	EnqueuedPlaylistRequestsLock.Lock();
+	EnqueuedPlaylistRequests.Empty();
+	EnqueuedPlaylistRequestsLock.Unlock();
 
 	// Delete any unprocessed completed downloads
 	while(!CompletedPlaylistRequests.IsEmpty())
@@ -732,6 +736,7 @@ void FPlaylistReaderHLS::HandleStaticRequestCompletions(const FTimeValue& TimeNo
 
 void FPlaylistReaderHLS::HandleEnqueuedPlaylistDownloads(const FTimeValue& TimeNow)
 {
+	EnqueuedPlaylistRequestsLock.Lock();
 	for(int32 i=0; i<EnqueuedPlaylistRequests.Num(); ++i)
 	{
 		FPlaylistRequestPtr Request = EnqueuedPlaylistRequests[i];
@@ -749,13 +754,16 @@ void FPlaylistReaderHLS::HandleEnqueuedPlaylistDownloads(const FTimeValue& TimeN
 			}
 			else
 			{
+				EnqueuedPlaylistRequestsLock.Unlock();
 				PendingPlaylistRequestsLock.Lock();
 				PendingPlaylistRequests.Push(Request);
 				Request->Execute(ProgressListener, PlayerSessionServices);
 				PendingPlaylistRequestsLock.Unlock();
+				EnqueuedPlaylistRequestsLock.Lock();
 			}
 		}
 	}
+	EnqueuedPlaylistRequestsLock.Unlock();
 }
 
 void FPlaylistReaderHLS::HandleCompletedPlaylistDownloads(const FTimeValue& TimeNow)
