ElectraPlayer: Added central handler for application background/foreground processing

#rnx

diff --git a/Engine/Plugins/Media/ElectraPlayer/Source/ElectraPlayerRuntime/Private/Runtime/Decoder/DX/AudioDecoderAAC_DX.cpp b/Engine/Plugins/Media/ElectraPlayer/Source/ElectraPlayerRuntime/Private/Runtime/Decoder/DX/AudioDecoderAAC_DX.cpp
--- a/Engine/Plugins/Media/ElectraPlayer/Source/ElectraPlayerRuntime/Private/Runtime/Decoder/DX/AudioDecoderAAC_DX.cpp
+++ b/Engine/Plugins/Media/ElectraPlayer/Source/ElectraPlayerRuntime/Private/Runtime/Decoder/DX/AudioDecoderAAC_DX.cpp
@@ -3,6 +3,7 @@
 #ifdef ELECTRA_ENABLE_MFDECODER
 
 #include "PlayerCore.h"
+#include "PlayerRuntimeGlobal.h"
 #include "ElectraPlayerPrivate.h"
 #include "ElectraPlayerPrivate_Platform.h"
 
@@ -23,7 +24,6 @@
 #include "Windows/AllowWindowsPlatformTypes.h"
 #include "Windows/WindowsHWrapper.h"
 #include "HAL/LowLevelMemTracker.h"
-#include "Misc/CoreDelegates.h"
 
 THIRD_PARTY_INCLUDES_START
 #include "mftransform.h"
@@ -159,8 +159,6 @@ namespace Electra
 
 		FAccessUnitBuffer														AccessUnits;
 
-		FDelegateHandle															ApplicationSuspendedDelegate;
-		FDelegateHandle															ApplicationResumeDelegate;
 		FMediaEvent																ApplicationRunningSignal;
 		FMediaEvent																ApplicationSuspendConfirmedSignal;
 
@@ -901,8 +899,11 @@ namespace Electra
 
 		ApplicationRunningSignal.Signal();
 		ApplicationSuspendConfirmedSignal.Reset();
-		ApplicationSuspendedDelegate = FCoreDelegates::ApplicationWillEnterBackgroundDelegate.AddRaw(this, &FAudioDecoderAAC::HandleApplicationWillEnterBackground);
-		ApplicationResumeDelegate = FCoreDelegates::ApplicationHasEnteredForegroundDelegate.AddRaw(this, &FAudioDecoderAAC::HandleApplicationHasEnteredForeground);
+
+		TSharedPtrTS<FFGBGNotificationHandlers> FGBGHandlers = MakeSharedTS<FFGBGNotificationHandlers>();
+		FGBGHandlers->WillEnterBackground = [this]() { HandleApplicationWillEnterBackground(); };
+		FGBGHandlers->HasEnteredForeground = [this]() { HandleApplicationHasEnteredForeground(); };
+		AddBGFGNotificationHandler(FGBGHandlers);
 
 		bool bError = false;
 		bool bInDummyDecodeMode = false;
@@ -1181,16 +1182,7 @@ namespace Electra
 		CurrentPTS.SetToInvalid();
 		NextExpectedPTS.SetToInvalid();
 
-		if (ApplicationSuspendedDelegate.IsValid())
-		{
-			FCoreDelegates::ApplicationWillEnterBackgroundDelegate.Remove(ApplicationSuspendedDelegate);
-			ApplicationSuspendedDelegate.Reset();
-		}
-		if (ApplicationResumeDelegate.IsValid())
-		{
-			FCoreDelegates::ApplicationHasEnteredForegroundDelegate.Remove(ApplicationResumeDelegate);
-			ApplicationResumeDelegate.Reset();
-		}
+		RemoveBGFGNotificationHandler(FGBGHandlers);
 	}
 
 
diff --git a/Engine/Plugins/Media/ElectraPlayer/Source/ElectraPlayerRuntime/Private/Runtime/Decoder/DX/VideoDecoderH264_DX.cpp b/Engine/Plugins/Media/ElectraPlayer/Source/ElectraPlayerRuntime/Private/Runtime/Decoder/DX/VideoDecoderH264_DX.cpp
--- a/Engine/Plugins/Media/ElectraPlayer/Source/ElectraPlayerRuntime/Private/Runtime/Decoder/DX/VideoDecoderH264_DX.cpp
+++ b/Engine/Plugins/Media/ElectraPlayer/Source/ElectraPlayerRuntime/Private/Runtime/Decoder/DX/VideoDecoderH264_DX.cpp
@@ -3,8 +3,8 @@
 #ifdef ELECTRA_ENABLE_MFDECODER
 
 #include "PlayerCore.h"
+#include "PlayerRuntimeGlobal.h"
 #include "ElectraPlayerPrivate.h"
-
 #include "ElectraPlayerPrivate_Platform.h"
 
 #include "StreamAccessUnitBuffer.h"
@@ -21,7 +21,6 @@
 #include "Windows/AllowWindowsPlatformTypes.h"
 #include "Windows/WindowsHWrapper.h"
 #include "HAL/LowLevelMemTracker.h"
-#include "Misc/CoreDelegates.h"
 
 THIRD_PARTY_INCLUDES_START
 #include "mftransform.h"
@@ -1316,8 +1315,11 @@ void FVideoDecoderH264::WorkerThread()
 
 	ApplicationRunningSignal.Signal();
 	ApplicationSuspendConfirmedSignal.Reset();
-	ApplicationSuspendedDelegate = FCoreDelegates::ApplicationWillEnterBackgroundDelegate.AddRaw(this, &FVideoDecoderH264::HandleApplicationWillEnterBackground);
-	ApplicationResumeDelegate = FCoreDelegates::ApplicationHasEnteredForegroundDelegate.AddRaw(this, &FVideoDecoderH264::HandleApplicationHasEnteredForeground);
+
+	TSharedPtrTS<FFGBGNotificationHandlers> FGBGHandlers = MakeSharedTS<FFGBGNotificationHandlers>();
+	FGBGHandlers->WillEnterBackground = [this]() { HandleApplicationWillEnterBackground(); };
+	FGBGHandlers->HasEnteredForeground = [this]() { HandleApplicationHasEnteredForeground(); };
+	AddBGFGNotificationHandler(FGBGHandlers);
 
 	CurrentRenderOutputBuffer     = nullptr;
 	CurrentAccessUnit   		  = nullptr;
@@ -1535,16 +1537,7 @@ void FVideoDecoderH264::WorkerThread()
 	AccessUnits.Flush();
 	AccessUnits.CapacitySet(0);
 
-	if (ApplicationSuspendedDelegate.IsValid())
-	{
-		FCoreDelegates::ApplicationWillEnterBackgroundDelegate.Remove(ApplicationSuspendedDelegate);
-		ApplicationSuspendedDelegate.Reset();
-	}
-	if (ApplicationResumeDelegate.IsValid())
-	{
-		FCoreDelegates::ApplicationHasEnteredForegroundDelegate.Remove(ApplicationResumeDelegate);
-		ApplicationResumeDelegate.Reset();
-	}
+	RemoveBGFGNotificationHandler(FGBGHandlers);
 
 	if (bDrainForCodecChange)
 	{
diff --git a/Engine/Plugins/Media/ElectraPlayer/Source/ElectraPlayerRuntime/Private/Runtime/Decoder/DX/VideoDecoderH264_DX.h b/Engine/Plugins/Media/ElectraPlayer/Source/ElectraPlayerRuntime/Private/Runtime/Decoder/DX/VideoDecoderH264_DX.h
--- a/Engine/Plugins/Media/ElectraPlayer/Source/ElectraPlayerRuntime/Private/Runtime/Decoder/DX/VideoDecoderH264_DX.h
+++ b/Engine/Plugins/Media/ElectraPlayer/Source/ElectraPlayerRuntime/Private/Runtime/Decoder/DX/VideoDecoderH264_DX.h
@@ -130,8 +130,6 @@ namespace Electra {
 
 		FInstanceConfiguration								Config;
 
-		FDelegateHandle										ApplicationSuspendedDelegate;
-		FDelegateHandle										ApplicationResumeDelegate;
 		FMediaEvent											ApplicationRunningSignal;
 		FMediaEvent											ApplicationSuspendConfirmedSignal;
 
diff --git a/Engine/Plugins/Media/ElectraPlayer/Source/ElectraPlayerRuntime/Private/Runtime/Decoder/DX/VideoDecoderH265_DX.cpp b/Engine/Plugins/Media/ElectraPlayer/Source/ElectraPlayerRuntime/Private/Runtime/Decoder/DX/VideoDecoderH265_DX.cpp
--- a/Engine/Plugins/Media/ElectraPlayer/Source/ElectraPlayerRuntime/Private/Runtime/Decoder/DX/VideoDecoderH265_DX.cpp
+++ b/Engine/Plugins/Media/ElectraPlayer/Source/ElectraPlayerRuntime/Private/Runtime/Decoder/DX/VideoDecoderH265_DX.cpp
@@ -3,8 +3,8 @@
 #ifdef ELECTRA_ENABLE_MFDECODER
 
 #include "PlayerCore.h"
+#include "PlayerRuntimeGlobal.h"
 #include "ElectraPlayerPrivate.h"
-
 #include "ElectraPlayerPrivate_Platform.h"
 
 #include "StreamAccessUnitBuffer.h"
@@ -24,7 +24,6 @@
 #include "Windows/AllowWindowsPlatformTypes.h"
 #include "Windows/WindowsHWrapper.h"
 #include "HAL/LowLevelMemTracker.h"
-#include "Misc/CoreDelegates.h"
 
 THIRD_PARTY_INCLUDES_START
 #include "mftransform.h"
@@ -1102,8 +1101,11 @@ void FVideoDecoderH265::WorkerThread()
 
 	ApplicationRunningSignal.Signal();
 	ApplicationSuspendConfirmedSignal.Reset();
-	ApplicationSuspendedDelegate = FCoreDelegates::ApplicationWillEnterBackgroundDelegate.AddRaw(this, &FVideoDecoderH265::HandleApplicationWillEnterBackground);
-	ApplicationResumeDelegate = FCoreDelegates::ApplicationHasEnteredForegroundDelegate.AddRaw(this, &FVideoDecoderH265::HandleApplicationHasEnteredForeground);
+
+	TSharedPtrTS<FFGBGNotificationHandlers> FGBGHandlers = MakeSharedTS<FFGBGNotificationHandlers>();
+	FGBGHandlers->WillEnterBackground = [this]() { HandleApplicationWillEnterBackground(); };
+	FGBGHandlers->HasEnteredForeground = [this]() { HandleApplicationHasEnteredForeground(); };
+	AddBGFGNotificationHandler(FGBGHandlers);
 
 	CurrentRenderOutputBuffer     = nullptr;
 	CurrentAccessUnit   		  = nullptr;
@@ -1308,16 +1310,7 @@ void FVideoDecoderH265::WorkerThread()
 	AccessUnits.Flush();
 	AccessUnits.CapacitySet(0);
 
-	if (ApplicationSuspendedDelegate.IsValid())
-	{
-		FCoreDelegates::ApplicationWillEnterBackgroundDelegate.Remove(ApplicationSuspendedDelegate);
-		ApplicationSuspendedDelegate.Reset();
-	}
-	if (ApplicationResumeDelegate.IsValid())
-	{
-		FCoreDelegates::ApplicationHasEnteredForegroundDelegate.Remove(ApplicationResumeDelegate);
-		ApplicationResumeDelegate.Reset();
-	}
+	RemoveBGFGNotificationHandler(FGBGHandlers);
 
 	if (bDrainForCodecChange)
 	{
diff --git a/Engine/Plugins/Media/ElectraPlayer/Source/ElectraPlayerRuntime/Private/Runtime/Decoder/DX/VideoDecoderH265_DX.h b/Engine/Plugins/Media/ElectraPlayer/Source/ElectraPlayerRuntime/Private/Runtime/Decoder/DX/VideoDecoderH265_DX.h
--- a/Engine/Plugins/Media/ElectraPlayer/Source/ElectraPlayerRuntime/Private/Runtime/Decoder/DX/VideoDecoderH265_DX.h
+++ b/Engine/Plugins/Media/ElectraPlayer/Source/ElectraPlayerRuntime/Private/Runtime/Decoder/DX/VideoDecoderH265_DX.h
@@ -128,8 +128,6 @@ namespace Electra {
 
 		FInstanceConfiguration								Config;
 
-		FDelegateHandle										ApplicationSuspendedDelegate;
-		FDelegateHandle										ApplicationResumeDelegate;
 		FMediaEvent											ApplicationRunningSignal;
 		FMediaEvent											ApplicationSuspendConfirmedSignal;
 
diff --git a/Engine/Plugins/Media/ElectraPlayer/Source/ElectraPlayerRuntime/Private/Runtime/PlayerRuntimeGlobal.cpp b/Engine/Plugins/Media/ElectraPlayer/Source/ElectraPlayerRuntime/Private/Runtime/PlayerRuntimeGlobal.cpp
--- a/Engine/Plugins/Media/ElectraPlayer/Source/ElectraPlayerRuntime/Private/Runtime/PlayerRuntimeGlobal.cpp
+++ b/Engine/Plugins/Media/ElectraPlayer/Source/ElectraPlayerRuntime/Private/Runtime/PlayerRuntimeGlobal.cpp
@@ -2,10 +2,51 @@
 
 #include "PlayerCore.h"
 #include "PlayerRuntimeGlobal.h"
+#include "Misc/CoreDelegates.h"
 
 namespace Electra
 {
-	TMap<FString, bool>	EnabledAnalyticsEvents;
+	namespace Global
+	{
+		TMap<FString, bool>								EnabledAnalyticsEvents;
+		FDelegateHandle									ApplicationSuspendedDelegate;
+		FDelegateHandle									ApplicationResumeDelegate;
+		FCriticalSection								ApplicationBGFGLock;
+		TArray<TWeakPtrTS<FFGBGNotificationHandlers>>	ApplicationBGFGHandlers;
+	}
+	using namespace Global;
+
+
+	static void HandleApplicationWillEnterBackground()
+	{
+		ApplicationBGFGLock.Lock();
+		TArray<TWeakPtrTS<FFGBGNotificationHandlers>>	CurrentHandlers(ApplicationBGFGHandlers);
+		ApplicationBGFGLock.Unlock();
+		for(auto &Handler : CurrentHandlers)
+		{
+			TSharedPtrTS<FFGBGNotificationHandlers> Hdlr = Handler.Pin();
+			if (Hdlr.IsValid())
+			{
+				Hdlr->WillEnterBackground();
+			}
+		}
+	}
+
+	static void HandleApplicationHasEnteredForeground()
+	{
+		ApplicationBGFGLock.Lock();
+		TArray<TWeakPtrTS<FFGBGNotificationHandlers>>	CurrentHandlers(ApplicationBGFGHandlers);
+		ApplicationBGFGLock.Unlock();
+		for(auto &Handler : CurrentHandlers)
+		{
+			TSharedPtrTS<FFGBGNotificationHandlers> Hdlr = Handler.Pin();
+			if (Hdlr.IsValid())
+			{
+				Hdlr->HasEnteredForeground();
+			}
+		}
+	}
+
 
 	//-----------------------------------------------------------------------------
 	/**
@@ -17,6 +58,15 @@ namespace Electra
 	 */
 	bool Startup(const Configuration& InConfiguration)
 	{
+		if (!ApplicationSuspendedDelegate.IsValid())
+		{
+			ApplicationSuspendedDelegate = FCoreDelegates::ApplicationWillEnterBackgroundDelegate.AddStatic(&HandleApplicationWillEnterBackground);
+		}
+		if (!ApplicationResumeDelegate.IsValid())
+		{
+			ApplicationResumeDelegate = FCoreDelegates::ApplicationHasEnteredForegroundDelegate.AddStatic(&HandleApplicationHasEnteredForeground);
+		}
+
 		EnabledAnalyticsEvents = InConfiguration.EnabledAnalyticsEvents;
 		return(true);
 	}
@@ -28,8 +78,31 @@ namespace Electra
 	 */
 	void Shutdown(void)
 	{
+		if (ApplicationSuspendedDelegate.IsValid())
+		{
+			FCoreDelegates::ApplicationWillEnterBackgroundDelegate.Remove(ApplicationSuspendedDelegate);
+			ApplicationSuspendedDelegate.Reset();
+		}
+		if (ApplicationResumeDelegate.IsValid())
+		{
+			FCoreDelegates::ApplicationHasEnteredForegroundDelegate.Remove(ApplicationResumeDelegate);
+			ApplicationResumeDelegate.Reset();
+		}
+	}
+
+	void AddBGFGNotificationHandler(TSharedPtrTS<FFGBGNotificationHandlers> InHandlers)
+	{
+		FScopeLock lock(&ApplicationBGFGLock);
+		ApplicationBGFGHandlers.Add(InHandlers);
 	}
 
+	void RemoveBGFGNotificationHandler(TSharedPtrTS<FFGBGNotificationHandlers> InHandlers)
+	{
+		FScopeLock lock(&ApplicationBGFGLock);
+		ApplicationBGFGHandlers.Remove(InHandlers);
+	}
+
+
 	/**
 	 * Check if the provided analytics event is enabled
 	 *
diff --git a/Engine/Plugins/Media/ElectraPlayer/Source/ElectraPlayerRuntime/Private/Runtime/PlayerRuntimeGlobal.h b/Engine/Plugins/Media/ElectraPlayer/Source/ElectraPlayerRuntime/Private/Runtime/PlayerRuntimeGlobal.h
--- a/Engine/Plugins/Media/ElectraPlayer/Source/ElectraPlayerRuntime/Private/Runtime/PlayerRuntimeGlobal.h
+++ b/Engine/Plugins/Media/ElectraPlayer/Source/ElectraPlayerRuntime/Private/Runtime/PlayerRuntimeGlobal.h
@@ -3,6 +3,7 @@
 #pragma once
 
 #include "PlayerCore.h"
+#include "Templates/Function.h"
 
 namespace Electra
 {
@@ -15,7 +16,7 @@ namespace Electra
 	bool Startup(const Configuration& Configuration);
 
 	//! Shuts down core services.
-	void Shutdown(void);
+	void Shutdown();
 
 	//! Waits until all player instances have terminated, which may happen asynchronously.
 	void WaitForAllPlayersToHaveTerminated();
@@ -27,5 +28,14 @@ namespace Electra
 	//! Check if an analytics event is enabled
 	bool IsAnalyticsEventEnabled(const FString& AnalyticsEventName);
 
+
+	//! Background / foreground handling
+	struct FFGBGNotificationHandlers
+	{
+		TFunction<void()> WillEnterBackground;
+		TFunction<void()> HasEnteredForeground;
+	};
+	void AddBGFGNotificationHandler(TSharedPtrTS<FFGBGNotificationHandlers> InHandlers);
+	void RemoveBGFGNotificationHandler(TSharedPtrTS<FFGBGNotificationHandlers> InHandlers);
 };
 
