ElectraPlayer: Added DASH AdaptationSet switching and selection of AdaptationSet based on user priorities

#rnx

diff --git a/Engine/Plugins/Media/ElectraPlayer/Source/ElectraPlayerPlugin/Private/ElectraPlayerPlugin.cpp b/Engine/Plugins/Media/ElectraPlayer/Source/ElectraPlayerPlugin/Private/ElectraPlayerPlugin.cpp
--- a/Engine/Plugins/Media/ElectraPlayer/Source/ElectraPlayerPlugin/Private/ElectraPlayerPlugin.cpp
+++ b/Engine/Plugins/Media/ElectraPlayer/Source/ElectraPlayerPlugin/Private/ElectraPlayerPlugin.cpp
@@ -433,6 +433,15 @@ bool FElectraPlayerPlugin::Open(const FString& Url, const IMediaOptions* Options
 
 	// Set up options to initialize the internal player with.
 	Electra::FParamDict PlayerOptions;
+	const TCHAR * const OptionsByString[] = { TEXT("excluded_codecs_video"), TEXT("excluded_codecs_audio"), TEXT("excluded_codecs_subtitles"), TEXT("preferred_codecs_video"), TEXT("preferred_codecs_audio"),TEXT("preferred_codecs_subtitles") };
+	for(auto &StringOption : OptionsByString)
+	{
+		FString Value = Options->GetMediaOption(StringOption, FString());
+		if (Value.Len())
+		{
+			PlayerOptions.Set(StringOption, Electra::FVariantValue(Value));
+		}
+	}
 
 	// Check for one-time initialization options that can't be changed during playback.
 	int64 InitialStreamBitrate = Options->GetMediaOption(TEXT("ElectraInitialBitrate"), (int64)-1);
diff --git a/Engine/Plugins/Media/ElectraPlayer/Source/ElectraPlayerRuntime/Private/Runtime/Decoder/Windows/PlatformVideoDecoderH265_Windows.cpp b/Engine/Plugins/Media/ElectraPlayer/Source/ElectraPlayerRuntime/Private/Runtime/Decoder/Windows/PlatformVideoDecoderH265_Windows.cpp
--- a/Engine/Plugins/Media/ElectraPlayer/Source/ElectraPlayerRuntime/Private/Runtime/Decoder/Windows/PlatformVideoDecoderH265_Windows.cpp
+++ b/Engine/Plugins/Media/ElectraPlayer/Source/ElectraPlayerRuntime/Private/Runtime/Decoder/Windows/PlatformVideoDecoderH265_Windows.cpp
@@ -85,6 +85,12 @@ namespace Electra
 				}
 				else
 				{
+					// For now we handle only 8 bit encodings. If Main-10 profile, or only Main-10 profile compatibility is signaled we refuse this stream.
+					if (InStreamParameter.Profile > 1 || (InStreamParameter.CompatibilityFlags & 0x60000000) == 0x20000000)
+					{
+						return true;
+					}
+
 					// TBD - What limits should apply here, exactly?
 					if (InStreamParameter.FPS > 0.0 && InStreamParameter.FPS <= 60.0 &&
 						InStreamParameter.Level <= 153)		// Level 5.1
diff --git a/Engine/Plugins/Media/ElectraPlayer/Source/ElectraPlayerRuntime/Private/Runtime/Demuxer/ParserISO14496-12.cpp b/Engine/Plugins/Media/ElectraPlayer/Source/ElectraPlayerRuntime/Private/Runtime/Demuxer/ParserISO14496-12.cpp
--- a/Engine/Plugins/Media/ElectraPlayer/Source/ElectraPlayerRuntime/Private/Runtime/Demuxer/ParserISO14496-12.cpp
+++ b/Engine/Plugins/Media/ElectraPlayer/Source/ElectraPlayerRuntime/Private/Runtime/Demuxer/ParserISO14496-12.cpp
@@ -597,6 +597,7 @@ private:
 		static const IParserISO14496_12::FBoxType kBox_stsd = MAKE_BOX_ATOM('s', 't', 's', 'd');
 		static const IParserISO14496_12::FBoxType kBox_btrt = MAKE_BOX_ATOM('b', 't', 'r', 't');
 		static const IParserISO14496_12::FBoxType kBox_pasp = MAKE_BOX_ATOM('p', 'a', 's', 'p');
+		static const IParserISO14496_12::FBoxType kBox_fiel = MAKE_BOX_ATOM('f', 'i', 'e', 'l');
 		static const IParserISO14496_12::FBoxType kBox_avcC = MAKE_BOX_ATOM('a', 'v', 'c', 'C');
 		static const IParserISO14496_12::FBoxType kBox_hvcC = MAKE_BOX_ATOM('h', 'v', 'c', 'C');
 		static const IParserISO14496_12::FBoxType kBox_esds = MAKE_BOX_ATOM('e', 's', 'd', 's');
@@ -4492,6 +4493,7 @@ private:
 						// Boxes we ignore for now.
 						case FMP4Box::kBox_tref:
 						case FMP4Box::kBox_colr:
+						case FMP4Box::kBox_fiel:
 						case FMP4Box::kBox_trep:
 						case FMP4Box::kBox_sdtp:
 						case FMP4Box::kBox_gmhd:
diff --git a/Engine/Plugins/Media/ElectraPlayer/Source/ElectraPlayerRuntime/Private/Runtime/Player/AdaptiveStreamingPlayer.cpp b/Engine/Plugins/Media/ElectraPlayer/Source/ElectraPlayerRuntime/Private/Runtime/Player/AdaptiveStreamingPlayer.cpp
--- a/Engine/Plugins/Media/ElectraPlayer/Source/ElectraPlayerRuntime/Private/Runtime/Player/AdaptiveStreamingPlayer.cpp
+++ b/Engine/Plugins/Media/ElectraPlayer/Source/ElectraPlayerRuntime/Private/Runtime/Player/AdaptiveStreamingPlayer.cpp
@@ -230,15 +230,15 @@ void FAdaptiveStreamingPlayer::Initialize(const FParamDict& Options)
 	// compared against these limits and those that exceed the limit will not be considered for playback.
 
 	// Maximum allowed vertical resolution specified?
-	if (PlayerOptions.HaveKey("max_resoY"))
+	if (PlayerOptions.HaveKey(TEXT("max_resoY")))
 	{
-		PlayerConfig.H264LimitUpto30fps.MaxResolution.Height = (int32) PlayerOptions.GetValue("max_resoY").GetInt64();
-		PlayerConfig.H264LimitAbove30fps.MaxResolution.Height = (int32) PlayerOptions.GetValue("max_resoY").GetInt64();
+		PlayerConfig.H264LimitUpto30fps.MaxResolution.Height = (int32) PlayerOptions.GetValue(TEXT("max_resoY")).GetInt64();
+		PlayerConfig.H264LimitAbove30fps.MaxResolution.Height = (int32) PlayerOptions.GetValue(TEXT("max_resoY")).GetInt64();
 	}
 	// A limit in vertical resolution for streams with more than 30fps?
-	if (PlayerOptions.HaveKey("max_resoY_above_30fps"))
+	if (PlayerOptions.HaveKey(TEXT("max_resoY_above_30fps")))
 	{
-		PlayerConfig.H264LimitAbove30fps.MaxResolution.Height = (int32) PlayerOptions.GetValue("max_resoY_above_30fps").GetInt64();
+		PlayerConfig.H264LimitAbove30fps.MaxResolution.Height = (int32) PlayerOptions.GetValue(TEXT("max_resoY_above_30fps")).GetInt64();
 	}
 	// Note: We could add additional limits here if need be.
 	//       Eventually these need to be differentiated based on codec as well.
@@ -263,9 +263,9 @@ void FAdaptiveStreamingPlayer::Initialize(const FParamDict& Options)
 		}
 
 		// If this is software decoding only and there is limit for this in place on Windows, apply it.
-		if (PlayerOptions.HaveKey("max_resoY_windows_software") && Capability.DecoderSupportType == IVideoDecoderH264::FStreamDecodeCapability::ESupported::SoftwareOnly)
+		if (PlayerOptions.HaveKey(TEXT("max_resoY_windows_software")) && Capability.DecoderSupportType == IVideoDecoderH264::FStreamDecodeCapability::ESupported::SoftwareOnly)
 		{
-			int32 MaxWinSWHeight = (int32) PlayerOptions.GetValue("max_resoY_windows_software").GetInt64();
+			int32 MaxWinSWHeight = (int32) PlayerOptions.GetValue(TEXT("max_resoY_windows_software")).GetInt64();
 			PlayerConfig.H264LimitUpto30fps.MaxResolution.Height = Utils::Min(PlayerConfig.H264LimitUpto30fps.MaxResolution.Height, MaxWinSWHeight);
 			PlayerConfig.H264LimitAbove30fps.MaxResolution.Height = Utils::Min(PlayerConfig.H264LimitAbove30fps.MaxResolution.Height, MaxWinSWHeight);
 		}
@@ -279,6 +279,40 @@ void FAdaptiveStreamingPlayer::Initialize(const FParamDict& Options)
 			PlayerConfig.H264LimitAbove30fps.MaxTierProfileLevel.Level = 10;
 		}
 	}
+
+	// Check for codecs that are not to be used as per the user's choice, even if the device supports them.
+	auto GetExcludedCodecPrefixes = [](TArray<FString>& OutList, const FParamDict& InOptions, const FString& InKey) -> void
+	{
+		if (InOptions.HaveKey(InKey))
+		{
+			const TCHAR* const CommaDelimiter = TEXT(",");
+			TArray<FString> Prefixes;
+			FString Value = InOptions.GetValue(InKey).GetFString();
+			Value.TrimQuotesInline();
+			Value.ParseIntoArray(Prefixes, CommaDelimiter, true);
+			for(auto &Prefix : Prefixes)
+			{
+				OutList.Emplace(Prefix.TrimStartAndEnd().TrimQuotes().TrimStartAndEnd());
+			}
+		}
+	};
+	GetExcludedCodecPrefixes(ExcludedVideoDecoderPrefixes,    PlayerOptions, TEXT("excluded_codecs_video"));
+	GetExcludedCodecPrefixes(ExcludedAudioDecoderPrefixes,    PlayerOptions, TEXT("excluded_codecs_audio"));
+	GetExcludedCodecPrefixes(ExcludedSubtitleDecoderPrefixes, PlayerOptions, TEXT("excluded_codecs_subtitles"));
+
+	auto GetCodecSelectionPriorities = [this](FCodecSelectionPriorities& OutPriorities, const FParamDict& InOptions, const FString& InKey, const TCHAR* const InType) -> void
+	{
+		if (InOptions.HaveKey(InKey))
+		{
+			if (!OutPriorities.Initialize(InOptions.GetValue(InKey).GetFString()))
+			{
+				PostLog(Facility::EFacility::Player, IInfoLog::ELevel::Warning, FString::Printf(TEXT("Failed to parse %s codec selection priority configuration"), InType));
+			}
+		}
+	};
+	GetCodecSelectionPriorities(CodecPrioritiesVideo,     PlayerOptions, TEXT("preferred_codecs_video"),     TEXT("video"));
+	GetCodecSelectionPriorities(CodecPrioritiesAudio,     PlayerOptions, TEXT("preferred_codecs_audio"),     TEXT("audio"));
+	GetCodecSelectionPriorities(CodecPrioritiesSubtitles, PlayerOptions, TEXT("preferred_codecs_subtitles"), TEXT("subtitle"));
 }
 
 
@@ -829,12 +863,11 @@ void FAdaptiveStreamingPlayer::DestroyRenderers()
 
 
 
-void FAdaptiveStreamingPlayer::PostLog(Facility::EFacility FromFacility, IInfoLog::ELevel inLogLevel, const FString &Message)
+void FAdaptiveStreamingPlayer::PostLog(Facility::EFacility FromFacility, IInfoLog::ELevel InLogLevel, const FString &Message)
 {
 	int64 millisNow = SynchronizedUTCTime->GetTime().GetAsMilliseconds();
-	//FString msg = FString::Printf("[%s] %s: %s", IInfoLog::GetLevelName(inLogLevel), Facility::GetName(FromFacility), Message.c_str());
 	FString msg = FString::Printf(TEXT("%s: %s"), Facility::GetName(FromFacility), *Message);
-	DispatchEvent(FMetricEvent::ReportLogMessage(inLogLevel, msg, millisNow));
+	DispatchEvent(FMetricEvent::ReportLogMessage(InLogLevel, msg, millisNow));
 }
 
 void FAdaptiveStreamingPlayer::PostError(const FErrorDetail& InError)
@@ -926,10 +959,53 @@ IPlayerStreamFilter* FAdaptiveStreamingPlayer::GetStreamFilter()
 	return this;
 }
 
+const FCodecSelectionPriorities& FAdaptiveStreamingPlayer::GetCodecSelectionPriorities(EStreamType ForStream)
+{
+	switch(ForStream)
+	{
+		case EStreamType::Video:
+		{
+			return CodecPrioritiesVideo;
+		}
+		case EStreamType::Audio:
+		{
+			return CodecPrioritiesAudio;
+		}
+		case EStreamType::Subtitle:
+		{
+			return CodecPrioritiesSubtitles;
+		}
+		default:
+		{
+			static FCodecSelectionPriorities Dummy;
+			return Dummy;
+		}
+	}
+}
+
 bool FAdaptiveStreamingPlayer::CanDecodeStream(const FStreamCodecInformation& InStreamCodecInfo) const
 {
+	auto IsCodecPrefixExcluded = [](const FString& InCodec, const TArray<FString>& InExcludedPrefixes) -> bool
+	{
+		for(auto &Prefix : InExcludedPrefixes)
+		{
+			if (InCodec.StartsWith(Prefix, ESearchCase::IgnoreCase))
+			{
+				return true;
+			}
+		}
+		return false;
+	};
+
+
 	if (InStreamCodecInfo.IsVideoCodec())
 	{
+		// Check against user specified codec prefix exclusions
+		if (IsCodecPrefixExcluded(InStreamCodecInfo.GetCodecSpecifierRFC6381(), ExcludedVideoDecoderPrefixes))
+		{
+			return false;
+		}
+
 		double Rate = InStreamCodecInfo.GetFrameRate().IsValid() ? InStreamCodecInfo.GetFrameRate().GetAsDouble() : 30.0;
 		if (InStreamCodecInfo.GetCodec() == FStreamCodecInformation::ECodec::H264)
 		{
@@ -984,6 +1060,22 @@ bool FAdaptiveStreamingPlayer::CanDecodeStream(const FStreamCodecInformation& In
 			#endif
 		}
 	}
+	else if (InStreamCodecInfo.IsAudioCodec())
+	{
+		// Check against user specified codec prefix exclusions
+		if (IsCodecPrefixExcluded(InStreamCodecInfo.GetCodecSpecifierRFC6381(), ExcludedAudioDecoderPrefixes))
+		{
+			return false;
+		}
+	}
+	else if (InStreamCodecInfo.IsSubtitleCodec())
+	{
+		// Check against user specified codec prefix exclusions
+		if (IsCodecPrefixExcluded(InStreamCodecInfo.GetCodecSpecifierRFC6381(), ExcludedSubtitleDecoderPrefixes))
+		{
+			return false;
+		}
+	}
 
 	return true;
 }
@@ -2467,7 +2559,7 @@ void FAdaptiveStreamingPlayer::HandlePendingMediaSegmentRequests()
 				break;
 			}
 			case EStreamType::Video:
-				{
+			{
 				SegmentPlayPeriod = CurrentPlayPeriodVideo;
 				break;
 			}
@@ -3102,8 +3194,14 @@ bool FAdaptiveStreamingPlayer::InternalStartAt(const FSeekParam& NewPosition)
 void FAdaptiveStreamingPlayer::InternalSetPlaybackEnded()
 {
 	PostrollVars.Clear();
+
+	FTimeRange TimelineRange;
+	PlaybackState.GetTimelineRange(TimelineRange);
+	PlaybackState.SetPlayPosition(TimelineRange.End);
+
 	PlaybackState.SetHasEnded(true);
 	StreamState = EStreamState::eStream_ReachedEnd;
+
 	// Go to pause mode now.
 	InternalPause();
 	DispatchEvent(FMetricEvent::ReportPlaybackEnded());
diff --git a/Engine/Plugins/Media/ElectraPlayer/Source/ElectraPlayerRuntime/Private/Runtime/Player/AdaptiveStreamingPlayerInternal.h b/Engine/Plugins/Media/ElectraPlayer/Source/ElectraPlayerRuntime/Private/Runtime/Player/AdaptiveStreamingPlayerInternal.h
--- a/Engine/Plugins/Media/ElectraPlayer/Source/ElectraPlayerRuntime/Private/Runtime/Player/AdaptiveStreamingPlayerInternal.h
+++ b/Engine/Plugins/Media/ElectraPlayer/Source/ElectraPlayerRuntime/Private/Runtime/Player/AdaptiveStreamingPlayerInternal.h
@@ -771,6 +771,7 @@ private:
 	virtual TSharedPtrTS<IAdaptiveStreamSelector> GetStreamSelector() override;
 	virtual void GetStreamBufferStats(FAccessUnitBufferInfo& OutBufferStats, EStreamType ForStream) override;
 	virtual IPlayerStreamFilter* GetStreamFilter() override;
+	virtual const FCodecSelectionPriorities& GetCodecSelectionPriorities(EStreamType ForStream) override;
 	virtual	TSharedPtrTS<IPlaylistReader> GetManifestReader() override;
 	virtual TSharedPtrTS<IPlayerEntityCache> GetEntityCache() override;
 	virtual IAdaptiveStreamingPlayerAEMSHandler* GetAEMSEventHandler() override;
@@ -1528,6 +1529,13 @@ private:
 	TSharedPtrTS<FStreamSelectionAttributes>							PendingTrackSelectionAud;
 	TSharedPtrTS<FStreamSelectionAttributes>							PendingTrackSelectionTxt;
 
+	TArray<FString>														ExcludedVideoDecoderPrefixes;
+	TArray<FString>														ExcludedAudioDecoderPrefixes;
+	TArray<FString>														ExcludedSubtitleDecoderPrefixes;
+	FCodecSelectionPriorities											CodecPrioritiesVideo;
+	FCodecSelectionPriorities											CodecPrioritiesAudio;
+	FCodecSelectionPriorities											CodecPrioritiesSubtitles;
+
 	EPlayerState														CurrentState;
 	EPipelineState														PipelineState;
 	EDecoderState														DecoderState;
diff --git a/Engine/Plugins/Media/ElectraPlayer/Source/ElectraPlayerRuntime/Private/Runtime/Player/DASH/ManifestBuilderDASH.cpp b/Engine/Plugins/Media/ElectraPlayer/Source/ElectraPlayerRuntime/Private/Runtime/Player/DASH/ManifestBuilderDASH.cpp
--- a/Engine/Plugins/Media/ElectraPlayer/Source/ElectraPlayerRuntime/Private/Runtime/Player/DASH/ManifestBuilderDASH.cpp
+++ b/Engine/Plugins/Media/ElectraPlayer/Source/ElectraPlayerRuntime/Private/Runtime/Player/DASH/ManifestBuilderDASH.cpp
@@ -109,6 +109,15 @@ namespace DASHAttributeHelpers
 			{
 				return &Attributes[i];
 			}
+			// If not found in the desired namespace and it doesn't exist in the default namespace let's see if there is a match
+			// in a some other namespace. This sounds counter-intuitive but namespaces can be chosen arbitrarily.
+			// Since there are usually no attribute conflicts that warrant a namespace let's see if the attribute exists in some
+			// unexpected namespace.
+			NameWithNS = FString::Printf(TEXT(":%s"), InAttribute);
+			if (Attributes[i].GetName().EndsWith(NameWithNS, ESearchCase::CaseSensitive))
+			{
+				return &Attributes[i];
+			}
 		}
 		return nullptr;
 	}
@@ -1234,7 +1243,8 @@ void FManifestDASHInternal::PreparePeriodAdaptationSets(TSharedPtrTS<FPeriod> Pe
 
 		Period->AdaptationSets.Empty();
 		const TArray<TSharedPtrTS<FDashMPD_AdaptationSetType>>& MPDAdaptationSets = MPDPeriod->GetAdaptationSets();
-		for(int32 nAdapt=0, nAdaptMax=MPDAdaptationSets.Num(); nAdapt<nAdaptMax; ++nAdapt)
+		int32 nAdapt=0, nAdaptMax=MPDAdaptationSets.Num();
+		for(; nAdapt<nAdaptMax; ++nAdapt)
 		{
 			const TSharedPtrTS<FDashMPD_AdaptationSetType>& MPDAdaptationSet = MPDAdaptationSets[nAdapt];
 			TSharedPtrTS<FAdaptationSet> AdaptationSet = MakeSharedTS<FAdaptationSet>();
@@ -1248,6 +1258,7 @@ void FManifestDASHInternal::PreparePeriodAdaptationSets(TSharedPtrTS<FPeriod> Pe
 			AdaptationSet->AdaptationSet = MPDAdaptationSet;
 			AdaptationSet->PAR = MPDAdaptationSet->GetPAR();
 			AdaptationSet->Language = ISO639::RFC5646To639_1(MPDAdaptationSet->GetLanguage());
+			AdaptationSet->SelectionPriority = (int32) MPDAdaptationSet->GetSelectionPriority();
 
 			// Content components are not supported.
 			if (MPDAdaptationSet->GetContentComponents().Num())
@@ -1590,6 +1601,104 @@ void FManifestDASHInternal::PreparePeriodAdaptationSets(TSharedPtrTS<FPeriod> Pe
 				Period->AdaptationSets.Emplace(AdaptationSet);
 			}
 		}
+
+		// Go over the AdaptationSets that are now remaining and check if they are set to switch between.
+		for(auto &AdaptationSet : Period->GetAdaptationSets())
+		{
+			TArray<TSharedPtrTS<FDashMPD_DescriptorType>> SwitchedSets;
+			TSharedPtrTS<FDashMPD_AdaptationSetType> MPDAdaptationSet = AdaptationSet->AdaptationSet.Pin();
+			SwitchedSets = MPDAdaptationSet->GetSupplementalProperties().FilterByPredicate([](const TSharedPtrTS<FDashMPD_DescriptorType>& d)
+				{ return d->GetSchemeIdUri().Equals(TEXT("urn:mpeg:dash:adaptation-set-switching:2016")); });
+			if (SwitchedSets.Num())
+			{
+				const TCHAR* const CommaDelimiter = TEXT(",");
+				for(auto &SwitchDesc : SwitchedSets)
+				{
+					// Get the IDs of the switched-to adaptation sets from the comma separated @value of the descriptor.
+					TArray<FString> SwitchedToIDs;
+					SwitchDesc->GetValue().ParseIntoArray(SwitchedToIDs, CommaDelimiter, true);
+					for(auto &SwitchID : SwitchedToIDs)
+					{
+						// The ID is the ID of the AdaptationSet in the MPD, not the one of the FAdaptationSet!
+						// Locate the FAdaptationSet that wraps the MPD's AdaptationSet with the given ID.
+						// This may be NULL when that AdaptationSet was not usable or flat out does not exist.
+						TSharedPtrTS<FAdaptationSet> SwitchedAS = Period->GetAdaptationSetByMPDID(SwitchID.TrimStartAndEnd());
+						if (SwitchedAS.IsValid())
+						{
+							// Check that this is not the same adaptation set we are currently handling!
+							if (SwitchedAS == AdaptationSet)
+							{
+								LogMessage(PlayerSessionServices, IInfoLog::ELevel::Warning, FString::Printf(TEXT("AdaptationSet references self in adaptation-set-switching property!")));
+								continue;
+							}
+							// Cross reference the switch-to and switching-from sets.
+							AdaptationSet->SwitchToSetIDs.AddUnique(SwitchedAS->GetUniqueIdentifier());
+							SwitchedAS->SwitchedFromSetIDs.AddUnique(AdaptationSet->GetUniqueIdentifier());
+						}
+					}
+				}
+			}
+		}
+
+		// Build special switching AdaptationSets that aggregate all representations from the referenced sets.
+		TArray<TSharedPtrTS<FAdaptationSet>> SwitchingAdaptationSets;
+		for(auto &AdaptationSet : Period->GetAdaptationSets())
+		{
+			TArray<FString> SwitchGroupIDs;
+			struct FSwitchGroupBuilder
+			{
+				static void AddSet(TArray<FString>& SwitchGroupIDs, const TSharedPtrTS<const FPeriod>& P, const TSharedPtrTS<FAdaptationSet>& AS)
+				{
+					if (AS.IsValid() && !SwitchGroupIDs.Contains(AS->GetUniqueIdentifier()))
+					{
+						SwitchGroupIDs.Add(AS->GetUniqueIdentifier());
+						for(auto &N : AS->GetSwitchToSetIDs())
+						{
+							AddSet(SwitchGroupIDs, P, P->GetAdaptationSetByUniqueID(N));
+						}
+						for(auto &N : AS->GetSwitchedFromSetIDs())
+						{
+							AddSet(SwitchGroupIDs, P, P->GetAdaptationSetByUniqueID(N));
+						}
+					}
+				}
+			};
+			// If this set is referencing others or is itself being referenced
+			if (!AdaptationSet->bIsInSwitchGroup && (AdaptationSet->GetSwitchToSetIDs().Num() || AdaptationSet->GetSwitchedFromSetIDs().Num()))
+			{
+				FSwitchGroupBuilder::AddSet(SwitchGroupIDs, Period, AdaptationSet);
+
+				TSharedPtrTS<FAdaptationSet> SwitchSet = MakeSharedTS<FAdaptationSet>();
+				SwitchSet->UniqueSequentialSetIndex = nAdapt++;
+				SwitchSet->bIsSwitchGroup = true;
+				SwitchSet->SwitchToSetIDs = MoveTemp(SwitchGroupIDs);
+				SwitchSet->bIsUsable = true;
+				SwitchSet->bIsEnabled = true;
+				for(int32 i=0; i<SwitchSet->SwitchToSetIDs.Num(); ++i)
+				{
+					TSharedPtrTS<FAdaptationSet> SwitchedAS = Period->GetAdaptationSetByUniqueID(SwitchSet->SwitchToSetIDs[i]);
+					SwitchedAS->bIsInSwitchGroup = true;
+					if (i == 0)
+					{
+						SwitchSet->Roles = SwitchedAS->Roles;
+						SwitchSet->Accessibilities = SwitchedAS->Accessibilities;
+						SwitchSet->PAR = SwitchedAS->PAR;
+						SwitchSet->Language = SwitchedAS->Language;
+					}
+					if (SwitchedAS->MaxBandwidth > SwitchSet->MaxBandwidth)
+					{
+						SwitchSet->MaxBandwidth = SwitchedAS->MaxBandwidth;
+						SwitchSet->Codec = SwitchedAS->Codec;
+					}
+
+					SwitchSet->Representations.Append(SwitchedAS->Representations);
+				}
+				// Note: We do not aggregate encryption information here. These get accessed through the original adaptation sets.
+				SwitchingAdaptationSets.Emplace(MoveTemp(SwitchSet));
+			}
+		}
+		Period->AdaptationSets.Append(MoveTemp(SwitchingAdaptationSets));
+
 		Period->SetHasBeenPrepared(true);
 	}
 }
diff --git a/Engine/Plugins/Media/ElectraPlayer/Source/ElectraPlayerRuntime/Private/Runtime/Player/DASH/ManifestDASH.cpp b/Engine/Plugins/Media/ElectraPlayer/Source/ElectraPlayerRuntime/Private/Runtime/Player/DASH/ManifestDASH.cpp
--- a/Engine/Plugins/Media/ElectraPlayer/Source/ElectraPlayerRuntime/Private/Runtime/Player/DASH/ManifestDASH.cpp
+++ b/Engine/Plugins/Media/ElectraPlayer/Source/ElectraPlayerRuntime/Private/Runtime/Player/DASH/ManifestDASH.cpp
@@ -131,9 +131,19 @@ public:
 private:
 	ELECTRA_IMPL_DEFAULT_ERROR_METHODS(DASHManifest);
 
+	struct FPrioritizedSelection
+	{
+		FStreamCodecInformation CodecInfo;
+		int32 Index = -1;
+		int32 Priority = -1;
+		int32 Bitrate = 0;
+	};
+
 	TSharedPtrTS<FManifestDASHInternal> GetCurrentManifest() const;
 
 	TSharedPtrTS<IPlaybackAssetRepresentation> GetRepresentationFromAdaptationByMaxBandwidth(TWeakPtrTS<IPlaybackAssetAdaptationSet> InAdaptationSet, int32 NotExceedingBandwidth);
+	TSharedPtrTS<IPlaybackAssetRepresentation> GetRepresentationFromAdaptationByPriorityAndMaxBandwidth(TWeakPtrTS<IPlaybackAssetAdaptationSet> InAdaptationSet, int32 NotExceedingBandwidth, EStreamType StreamType);
+
 	enum class ENextSegType
 	{
 		SamePeriodNext,
@@ -147,6 +157,8 @@ private:
 
 	void SetupCommonSegmentRequestInfos(TSharedPtrTS<FStreamSegmentRequestFMP4DASH>& InOutSegmentRequest);
 
+	void PrioritizeSelection(TArray<FPrioritizedSelection>& Selection, EStreamType StreamType, bool bAdaptationSetLevel, bool bSortByBitrateDescending);
+
 	TSharedPtrTS<FManifestDASHInternal::FAdaptationSet> SelectAdaptationSetByAttributes(TSharedPtrTS<FBufferSourceInfo>& OutBufferSourceInfo, TSharedPtrTS<FManifestDASHInternal::FPeriod> Period, EStreamType StreamType, const FStreamSelectionAttributes& Attributes);
 
 	IPlayerSessionServices* PlayerSessionServices = nullptr;
@@ -474,8 +486,57 @@ IManifest::IPlayPeriod::EReadyState FDASHPlayPeriod::GetReadyState()
 	return ReadyState;
 }
 
+void FDASHPlayPeriod::PrioritizeSelection(TArray<FPrioritizedSelection>& Selection, EStreamType StreamType, bool bAdaptationSetLevel, bool bSortByBitrateDescending)
+{
+	const FCodecSelectionPriorities& SelectionPriorities = PlayerSessionServices->GetCodecSelectionPriorities(StreamType);
+	for(auto &Candidate : Selection)
+	{
+		int32 NewPriority = -1;
+		if (bAdaptationSetLevel)
+		{
+			NewPriority = SelectionPriorities.GetClassPriority(Candidate.CodecInfo.GetCodecSpecifierRFC6381());
+		}
+		else
+		{
+			NewPriority = SelectionPriorities.GetStreamPriority(Candidate.CodecInfo.GetCodecSpecifierRFC6381());
+		}
+		if (NewPriority >= 0)
+		{
+			Candidate.Priority = NewPriority;
+		}
+	}
+	// Sort forst by descending bitrate?
+	if (bSortByBitrateDescending)
+	{
+		Selection.StableSort([](const FPrioritizedSelection& a, const FPrioritizedSelection& b)
+		{
+			return a.Bitrate > b.Bitrate;
+		});
+	}
+	// Sort by descending priority.
+	Selection.StableSort([](const FPrioritizedSelection& a, const FPrioritizedSelection& b)
+	{
+		return a.Priority > b.Priority;
+	});
+}
+
+
 TSharedPtrTS<FManifestDASHInternal::FAdaptationSet> FDASHPlayPeriod::SelectAdaptationSetByAttributes(TSharedPtrTS<FBufferSourceInfo>& OutBufferSourceInfo, TSharedPtrTS<FManifestDASHInternal::FPeriod> Period, EStreamType StreamType, const FStreamSelectionAttributes& Attributes)
 {
+	TArray<FPrioritizedSelection> Selection;
+
+	auto AddAdaptationSetToSelection = [](TArray<FPrioritizedSelection>& InOutSelection, const TSharedPtrTS<FManifestDASHInternal::FAdaptationSet>& AS, int32 Index) -> void
+	{
+		if (AS->GetIsUsable() && !AS->GetIsInSwitchGroup())
+		{
+			FPrioritizedSelection Candidate;
+			Candidate.CodecInfo = AS->GetCodec();
+			Candidate.Index = Index;
+			Candidate.Priority = AS->GetSelectionPriority();
+			InOutSelection.Emplace(MoveTemp(Candidate));
+		}
+	};
+
 	TSharedPtrTS<FManifestDASHInternal::FAdaptationSet> AS;
 	int32 NumAdaptationSets = Period->GetNumberOfAdaptationSets(StreamType);
 	if (NumAdaptationSets > 0)
@@ -484,7 +545,20 @@ TSharedPtrTS<FManifestDASHInternal::FAdaptationSet> FDASHPlayPeriod::SelectAdapt
 		// For video just select the first one for now.
 		if (StreamType == EStreamType::Video)
 		{
-			AS = StaticCastSharedPtr<FManifestDASHInternal::FAdaptationSet>(Period->GetAdaptationSetByTypeAndIndex(StreamType, 0));
+			// Create a list of candidate AdaptationSets
+			for(int32 i=0; i<NumAdaptationSets; ++i)
+			{
+				TSharedPtrTS<FManifestDASHInternal::FAdaptationSet> V = StaticCastSharedPtr<FManifestDASHInternal::FAdaptationSet>(Period->GetAdaptationSetByTypeAndIndex(StreamType, i));
+				AddAdaptationSetToSelection(Selection, V, i);
+			}
+			// Prioritize the candidates based on user configuration.
+			PrioritizeSelection(Selection, StreamType, true, false);
+			if (Selection.Num())
+			{
+				// Take the highest prioritized set.
+				SelectedTypeIndex = Selection[0].Index;
+				AS = StaticCastSharedPtr<FManifestDASHInternal::FAdaptationSet>(Period->GetAdaptationSetByTypeAndIndex(StreamType, SelectedTypeIndex));
+			}
 		}
 		else if (StreamType == EStreamType::Audio)
 		{
@@ -492,17 +566,23 @@ TSharedPtrTS<FManifestDASHInternal::FAdaptationSet> FDASHPlayPeriod::SelectAdapt
 			// For now we ignore the track kind.
 			if (Attributes.Language_ISO639.IsSet())
 			{
+				// Create a list of candidate AdaptationSets
 				for(int32 i=0; i<NumAdaptationSets; ++i)
 				{
-					TSharedPtrTS<FManifestDASHInternal::FAdaptationSet> A;
-					A = StaticCastSharedPtr<FManifestDASHInternal::FAdaptationSet>(Period->GetAdaptationSetByTypeAndIndex(StreamType, i));
+					TSharedPtrTS<FManifestDASHInternal::FAdaptationSet> A = StaticCastSharedPtr<FManifestDASHInternal::FAdaptationSet>(Period->GetAdaptationSetByTypeAndIndex(StreamType, i));
 					if (A->GetLanguage().Equals(Attributes.Language_ISO639.GetValue()))
 					{
-						SelectedTypeIndex = i;
-						AS = MoveTemp(A);
-						break;
+						AddAdaptationSetToSelection(Selection, A, i);
 					}
 				}
+				// Prioritize the candidates based on user configuration.
+				PrioritizeSelection(Selection, StreamType, true, false);
+				if (Selection.Num())
+				{
+					// Take the highest prioritized set.
+					SelectedTypeIndex = Selection[0].Index;
+					AS = StaticCastSharedPtr<FManifestDASHInternal::FAdaptationSet>(Period->GetAdaptationSetByTypeAndIndex(StreamType, SelectedTypeIndex));
+				}
 			}
 			// Matching language not found. Is there an explicit index given?
 			if (!AS.IsValid())
@@ -516,8 +596,19 @@ TSharedPtrTS<FManifestDASHInternal::FAdaptationSet> FDASHPlayPeriod::SelectAdapt
 			// Still nothing? Use the first one.
 			if (!AS.IsValid())
 			{
-				AS = StaticCastSharedPtr<FManifestDASHInternal::FAdaptationSet>(Period->GetAdaptationSetByTypeAndIndex(StreamType, 0));
-				SelectedTypeIndex = 0;
+				Selection.Empty();
+				for(int32 i=0; i<NumAdaptationSets; ++i)
+				{
+					TSharedPtrTS<FManifestDASHInternal::FAdaptationSet> A = StaticCastSharedPtr<FManifestDASHInternal::FAdaptationSet>(Period->GetAdaptationSetByTypeAndIndex(StreamType, i));
+					AddAdaptationSetToSelection(Selection, A, i);
+				}
+				PrioritizeSelection(Selection, StreamType, true, false);
+				if (Selection.Num())
+				{
+					// Take the highest prioritized set.
+					SelectedTypeIndex = Selection[0].Index;
+					AS = StaticCastSharedPtr<FManifestDASHInternal::FAdaptationSet>(Period->GetAdaptationSetByTypeAndIndex(StreamType, SelectedTypeIndex));
+				}
 			}
 		}
 		if (AS.IsValid())
@@ -576,8 +667,21 @@ void FDASHPlayPeriod::Load()
 			ActiveVideoAdaptationSetID = VideoAS->GetUniqueIdentifier();
 
 			// Add encryption schemes, if any.
-			const FManifestDASHInternal::FAdaptationSet* Adapt = static_cast<const FManifestDASHInternal::FAdaptationSet*>(VideoAS.Get());
-			ContentProtections.Append(Adapt->GetPossibleContentProtections());
+			if (VideoAS->GetIsSwitchGroup())
+			{
+				for(auto &SwitchedID : VideoAS->GetSwitchToSetIDs())
+				{
+					TSharedPtrTS<FManifestDASHInternal::FAdaptationSet> SwitchedAS = Period->GetAdaptationSetByUniqueID(SwitchedID);
+					if (SwitchedAS.IsValid())
+					{
+						ContentProtections.Append(SwitchedAS->GetPossibleContentProtections());
+					}
+				}
+			}
+			else
+			{
+				ContentProtections.Append(VideoAS->GetPossibleContentProtections());
+			}
 		}
 
 		TSharedPtrTS<FManifestDASHInternal::FAdaptationSet> AudioAS = SelectAdaptationSetByAttributes(SourceBufferInfoAudio, Period, EStreamType::Audio, AudioStreamPreferences);
@@ -586,8 +690,21 @@ void FDASHPlayPeriod::Load()
 			ActiveAudioAdaptationSetID = AudioAS->GetUniqueIdentifier();
 
 			// Add encryption schemes, if any.
-			const FManifestDASHInternal::FAdaptationSet* Adapt = static_cast<const FManifestDASHInternal::FAdaptationSet*>(AudioAS.Get());
-			ContentProtections.Append(Adapt->GetPossibleContentProtections());
+			if (AudioAS->GetIsSwitchGroup())
+			{
+				for(auto &SwitchedID : AudioAS->GetSwitchToSetIDs())
+				{
+					TSharedPtrTS<FManifestDASHInternal::FAdaptationSet> SwitchedAS = Period->GetAdaptationSetByUniqueID(SwitchedID);
+					if (SwitchedAS.IsValid())
+					{
+						ContentProtections.Append(SwitchedAS->GetPossibleContentProtections());
+					}
+				}
+			}
+			else
+			{
+				ContentProtections.Append(AudioAS->GetPossibleContentProtections());
+			}
 		}
 
 		// Prepare the DRM system for decryption.
@@ -627,7 +744,7 @@ void FDASHPlayPeriod::PrepareForPlay()
 		TSharedPtrTS<FManifestDASHInternal::FAdaptationSet> AudioAS = Period->GetAdaptationSetByUniqueID(ActiveAudioAdaptationSetID);
 		if (AudioAS.IsValid())
 		{
-			TSharedPtrTS<IPlaybackAssetRepresentation> AudioRepr = GetRepresentationFromAdaptationByMaxBandwidth(AudioAS, 256 * 1000);
+			TSharedPtrTS<IPlaybackAssetRepresentation> AudioRepr = GetRepresentationFromAdaptationByPriorityAndMaxBandwidth(AudioAS, 256 * 1000, EStreamType::Audio);
 			ActiveAudioRepresentationID = AudioRepr->GetUniqueIdentifier();
 		}
 
@@ -1338,6 +1455,53 @@ TSharedPtrTS<IPlaybackAssetRepresentation> FDASHPlayPeriod::GetRepresentationFro
 	return BestRepr;
 }
 
+TSharedPtrTS<IPlaybackAssetRepresentation> FDASHPlayPeriod::GetRepresentationFromAdaptationByPriorityAndMaxBandwidth(TWeakPtrTS<IPlaybackAssetAdaptationSet> InAdaptationSet, int32 NotExceedingBandwidth, EStreamType StreamType)
+{
+	TSharedPtrTS<IPlaybackAssetAdaptationSet> AS = InAdaptationSet.Pin();
+	TSharedPtrTS<IPlaybackAssetRepresentation> BestRepr;
+	if (AS.IsValid())
+	{
+		auto AddRepresentationToSelection = [](TArray<FPrioritizedSelection>& InOutSelection, const TSharedPtrTS<FManifestDASHInternal::FRepresentation>& R, int32 Index) -> void
+		{
+			if (R->CanBePlayed())
+			{
+				FPrioritizedSelection Candidate;
+				Candidate.CodecInfo = R->GetCodecInformation();
+				Candidate.Index = Index;
+				Candidate.Priority = R->GetSelectionPriority();
+				Candidate.Bitrate = R->GetBitrate();
+				InOutSelection.Emplace(MoveTemp(Candidate));
+			}
+		};
+		TArray<FPrioritizedSelection> Selection;
+
+		for(int32 i=0; i<AS->GetNumberOfRepresentations(); ++i)
+		{
+			TSharedPtrTS<FManifestDASHInternal::FRepresentation> Repr = StaticCastSharedPtr<FManifestDASHInternal::FRepresentation>(AS->GetRepresentationByIndex(i));
+			AddRepresentationToSelection(Selection, Repr, i);
+		}
+		PrioritizeSelection(Selection, StreamType, false, true);
+		if (Selection.Num())
+		{
+			for(int32 i=0; i<Selection.Num(); ++i)
+			{
+				TSharedPtrTS<FManifestDASHInternal::FRepresentation> Repr = StaticCastSharedPtr<FManifestDASHInternal::FRepresentation>(AS->GetRepresentationByIndex(Selection[i].Index));
+				if (Repr->GetBitrate() <= NotExceedingBandwidth)
+				{
+					BestRepr = Repr;
+					break;
+				}
+			}
+			if (!BestRepr.IsValid())
+			{
+				TSharedPtrTS<FManifestDASHInternal::FRepresentation> Repr = StaticCastSharedPtr<FManifestDASHInternal::FRepresentation>(AS->GetRepresentationByIndex(Selection.Last().Index));
+				BestRepr = Repr;
+			}
+		}
+	}
+	return BestRepr;
+}
+
 
 /***************************************************************************************************************************************************/
 /***************************************************************************************************************************************************/
diff --git a/Engine/Plugins/Media/ElectraPlayer/Source/ElectraPlayerRuntime/Private/Runtime/Player/DASH/PlaylistReaderDASH_Internal.h b/Engine/Plugins/Media/ElectraPlayer/Source/ElectraPlayerRuntime/Private/Runtime/Player/DASH/PlaylistReaderDASH_Internal.h
--- a/Engine/Plugins/Media/ElectraPlayer/Source/ElectraPlayerRuntime/Private/Runtime/Player/DASH/PlaylistReaderDASH_Internal.h
+++ b/Engine/Plugins/Media/ElectraPlayer/Source/ElectraPlayerRuntime/Private/Runtime/Player/DASH/PlaylistReaderDASH_Internal.h
@@ -227,6 +227,7 @@ public:
 
 		void GetSegmentInformation(TArray<IManifest::IPlayPeriod::FSegmentInformation>& OutSegmentInformation, FTimeValue& OutAverageSegmentDuration, TSharedPtrTS<const IStreamSegment> CurrentSegment, const FTimeValue& LookAheadTime, const TSharedPtrTS<IPlaybackAssetAdaptationSet>& AdaptationSet);
 
+		int32 GetSelectionPriority() const { return SelectionPriority; }
 
 
 		//----------------------------------------------
@@ -272,6 +273,7 @@ public:
 		FString ID;
 		int32 Bitrate = MAX_int32;
 		int32 QualityIndex = -1;
+		int32 SelectionPriority = 1;
 		bool bIsUsable = false;
 		bool bIsEnabled = true;
 		bool bWarnedAboutTimelineStartGap = false;
@@ -301,7 +303,10 @@ public:
 		const FTimeFraction& GetPAR() const											{ return PAR; }
 		//const FString& GetLanguage() const											{ return Language; }
 		int32 GetMaxBandwidth() const												{ return MaxBandwidth; }
+		int32 GetSelectionPriority() const											{ return SelectionPriority; }
 		bool GetIsUsable() const													{ return bIsUsable; }
+		bool GetIsInSwitchGroup() const												{ return bIsInSwitchGroup; }
+		bool GetIsSwitchGroup() const												{ return bIsSwitchGroup; }
 
 		TSharedPtrTS<FRepresentation> GetRepresentationByUniqueID(const FString& UniqueIdentifier) const
 		{
@@ -467,6 +472,9 @@ public:
 		const FString& GetCommonEncryptionScheme() const { return CommonEncryptionScheme; }
 		const FString& GetDefaultKID() const { return DefaultKID; }
 
+		const TArray<FString>& GetSwitchToSetIDs() const { return SwitchToSetIDs; }
+		const TArray<FString>& GetSwitchedFromSetIDs() const { return SwitchedFromSetIDs; }
+
 		//----------------------------------------------
 		// Methods from IPlaybackAssetAdaptationSet
 		//
@@ -508,6 +516,7 @@ public:
 		FTimeFraction PAR;
 		FString Language;
 		int32 MaxBandwidth = 0;
+		int32 SelectionPriority = 1;
 		int32 UniqueSequentialSetIndex = 0;
 		bool bIsUsable = false;
 		bool bIsEnabled = true;
@@ -515,6 +524,11 @@ public:
 		TArray<FContentProtection> PossibleContentProtections;
 		FString CommonEncryptionScheme;
 		FString DefaultKID;
+		// Switching related
+		TArray<FString> SwitchToSetIDs;
+		TArray<FString> SwitchedFromSetIDs;
+		bool bIsInSwitchGroup = false;
+		bool bIsSwitchGroup = false;
 	};
 
 	class FPeriod : public ITimelineMediaAsset
@@ -559,6 +573,20 @@ public:
 			return nullptr;
 		}
 
+		TSharedPtrTS<FAdaptationSet> GetAdaptationSetByMPDID(const FString& MPDID) const
+		{
+			for(int32 i=0; i<AdaptationSets.Num(); ++i)
+			{
+				TSharedPtrTS<FDashMPD_AdaptationSetType> MPDAdaptationSet = AdaptationSets[i]->AdaptationSet.Pin();
+				if (MPDAdaptationSet.IsValid() && MPDAdaptationSet->GetID_AsStr().Equals(MPDID))
+				{
+					return AdaptationSets[i];
+				}
+			}
+			return nullptr;
+		}
+
+
 		void EndPresentationAt(const FTimeValue& EndsAt)
 		{
 			FTimeValue NewDur = EndsAt - Start;
@@ -641,7 +669,7 @@ public:
 			{
 				TSharedPtrTS<IPlaybackAssetAdaptationSet> AdaptationSet = GetAdaptationSetByTypeAndIndex(OfStreamType, i);
 				const FAdaptationSet* Adapt = static_cast<const FAdaptationSet*>(AdaptationSet.Get());
-				if (Adapt && Adapt->GetIsUsable())
+				if (Adapt && Adapt->GetIsUsable() && !Adapt->GetIsInSwitchGroup())
 				{
 					FTrackMetadata tm;
 					Adapt->GetMetaData(tm, OfStreamType);
diff --git a/Engine/Plugins/Media/ElectraPlayer/Source/ElectraPlayerRuntime/Private/Runtime/Player/PlayerSessionServices.h b/Engine/Plugins/Media/ElectraPlayer/Source/ElectraPlayerRuntime/Private/Runtime/Player/PlayerSessionServices.h
--- a/Engine/Plugins/Media/ElectraPlayer/Source/ElectraPlayerRuntime/Private/Runtime/Player/PlayerSessionServices.h
+++ b/Engine/Plugins/Media/ElectraPlayer/Source/ElectraPlayerRuntime/Private/Runtime/Player/PlayerSessionServices.h
@@ -103,6 +103,11 @@ namespace Electra
 		 */
 		virtual IPlayerStreamFilter* GetStreamFilter() = 0;
 
+		/**
+		 * Returns user configured codec selection priorities.
+		 */
+		virtual const FCodecSelectionPriorities& GetCodecSelectionPriorities(EStreamType ForStream) = 0;
+
 		/**
 		 * Returns the entity cache of this player.
 		 */
diff --git a/Engine/Plugins/Media/ElectraPlayer/Source/ElectraPlayerRuntime/Private/Runtime/StreamTypes.cpp b/Engine/Plugins/Media/ElectraPlayer/Source/ElectraPlayerRuntime/Private/Runtime/StreamTypes.cpp
--- a/Engine/Plugins/Media/ElectraPlayer/Source/ElectraPlayerRuntime/Private/Runtime/StreamTypes.cpp
+++ b/Engine/Plugins/Media/ElectraPlayer/Source/ElectraPlayerRuntime/Private/Runtime/StreamTypes.cpp
@@ -223,4 +223,178 @@ namespace Electra
 			return false;
 		}
 	}
+
+
+
+
+
+	bool FCodecSelectionPriorities::Initialize(const FString& ConfigurationString)
+	{
+		ClassPriorities.Empty();
+		if (ConfigurationString.Len() && !ParseInternal(ConfigurationString))
+		{
+			ClassPriorities.Empty();
+			return false;
+		}
+		return true;
+	}
+	bool FCodecSelectionPriorities::ParseInternal(const FString& ConfigurationString)
+	{
+		auto SkipWhiteSpaces = [](StringHelpers::FStringIterator& it) -> void
+		{
+			while(it && TChar<TCHAR>::IsWhitespace(*it))
+			{
+				++it;
+			}
+		};
+
+		auto ParsePriority = [](int32& OutPrio, StringHelpers::FStringIterator& it, bool bInClass) -> bool
+		{
+			int64 Prio = 0;
+			bool bEmpty = true;
+			while(it && TChar<TCHAR>::IsDigit(*it))
+			{
+				bEmpty = false;
+				Prio *= 10;
+				Prio += *it - TCHAR('0');
+				++it;
+			}
+			while(it && TChar<TCHAR>::IsWhitespace(*it))
+			{
+				++it;
+			}
+			// Did we end the priority properly?
+			if (!bEmpty && (!it || *it == TCHAR(',') || (bInClass && *it == TCHAR('{'))))
+			{
+				OutPrio = Prio;
+				return true;
+			}
+			// Unexpected next character. Fail!
+			return false;
+		};
+
+		const TCHAR* const CommaDelimiter = TEXT(",");
+		StringHelpers::FStringIterator it(ConfigurationString);
+		while(it)
+		{
+			FClassPriority ClassPriority;
+			SkipWhiteSpaces(it);
+			while(it && *it != TCHAR('=') && *it != TCHAR('{') && *it != TCHAR(','))
+			{
+				ClassPriority.Prefix += *it++;
+			}
+			if (ClassPriority.Prefix.Len() == 0)
+			{
+				return false;
+			}
+
+			// Is the next char assigning a priority?
+			if (it && *it == TCHAR('='))
+			{
+				// Get the class priority
+				++it;
+				if (!ParsePriority(ClassPriority.Priority, it, true))
+				{
+					return false;
+				}
+			}
+			// If no priority then there must now be a group for stream specific priorities.
+			else if (!it || *it != TCHAR('{'))
+			{
+				return false;
+			}
+			// Do stream specific priorities follow?
+			if (it && *it == TCHAR('{'))
+			{
+				int32 GroupStart = it.GetIndex();
+				// Look for the end of the group.
+				while(it && *it != TCHAR('}'))
+				{
+					++it;
+				}
+				if (!it || *it != TCHAR('}'))
+				{
+					return false;
+				}
+				++it;
+				FString Group = ConfigurationString.Mid(GroupStart+1, it.GetIndex()-GroupStart-2);
+				TArray<FString> StreamPriorities;
+				Group.ParseIntoArray(StreamPriorities, CommaDelimiter, true);
+				if (StreamPriorities.Num() == 0)
+				{
+					return false;
+				}
+				for(auto &sp : StreamPriorities)
+				{
+					FStreamPriority StreamPriority;
+					StringHelpers::FStringIterator spIt(sp);
+					while(spIt)
+					{
+						SkipWhiteSpaces(spIt);
+						while(spIt && *spIt != TCHAR('=') && *spIt != TCHAR('{') && *spIt != TCHAR(','))
+						{
+							StreamPriority.Prefix += *spIt++;
+						}
+						if (StreamPriority.Prefix.Len() == 0)
+						{
+							return false;
+						}
+						if (spIt && *spIt == TCHAR('='))
+						{
+							++spIt;
+							if (!ParsePriority(StreamPriority.Priority, spIt, false))
+							{
+								return false;
+							}
+						}
+						else
+						{
+							return false;
+						}
+					}
+					ClassPriority.StreamPriorities.Emplace(MoveTemp(StreamPriority));
+				}
+			}
+			// Either there's a comma separating successive entries or we are done.
+			SkipWhiteSpaces(it);
+			if (it && *it != TCHAR(','))
+			{
+				return false;
+			}
+			++it;
+			ClassPriorities.Emplace(MoveTemp(ClassPriority));
+		}
+		return true;
+	}
+	
+	int32 FCodecSelectionPriorities::GetClassPriority(const FString& CodecSpecifierRFC6381) const
+	{
+		for(auto &CodecClass : ClassPriorities)
+		{
+			if (CodecSpecifierRFC6381.StartsWith(CodecClass.Prefix, ESearchCase::IgnoreCase))
+			{
+				return CodecClass.Priority;
+			}
+		}
+		return -1;
+	}
+	
+	int32 FCodecSelectionPriorities::GetStreamPriority(const FString& CodecSpecifierRFC6381) const
+	{
+		for(auto &CodecClass : ClassPriorities)
+		{
+			if (CodecSpecifierRFC6381.StartsWith(CodecClass.Prefix, ESearchCase::IgnoreCase))
+			{
+				for(auto &CodecStream : CodecClass.StreamPriorities)
+				{
+					if (CodecSpecifierRFC6381.StartsWith(CodecStream.Prefix, ESearchCase::IgnoreCase))
+					{
+						return CodecStream.Priority;
+					}
+				}
+			}
+		}
+		return -1;
+	}
+
 } // namespace Electra
diff --git a/Engine/Plugins/Media/ElectraPlayer/Source/ElectraPlayerRuntime/Private/Runtime/StreamTypes.h b/Engine/Plugins/Media/ElectraPlayer/Source/ElectraPlayerRuntime/Private/Runtime/StreamTypes.h
--- a/Engine/Plugins/Media/ElectraPlayer/Source/ElectraPlayerRuntime/Private/Runtime/StreamTypes.h
+++ b/Engine/Plugins/Media/ElectraPlayer/Source/ElectraPlayerRuntime/Private/Runtime/StreamTypes.h
@@ -652,6 +652,56 @@ namespace Electra
 	};
 
 
+	
+	class FCodecSelectionPriorities
+	{
+	public:
+		/**
+		 * Initializes this selector with a priority string in the following FORMAT:
+		 *   FORMAT = CLASSPRIO 0*[COMMA CLASSPRIO]
+		 *   COMMA = ,
+		 *   EQ = =
+		 *   PRIO = 1*DIGIT
+		 *   PREFIX = 1*VCHAR    ; except , = { }
+		 *   CLASS = PREFIX
+		 *   CODECPRIO = CLASS EQ PRIO
+		 *   CLASSWITHPRIO = CODECPRIO 0*[ { CODECPRIO 0*[ COMMA CODECPRIO ] } ]
+		 *   CLASSWITHOUTPRIO = CLASS 1*[ { CODECPRIO 0*[ COMMA CODECPRIO ] } ]
+		 *   CLASSPRIO = CLASSWITHPRIO / CLASSWITHOUTPRIO
+		 * 
+		 * Examples: hvc=2,hev=2,avc=1
+		 *           mp4a{mp4a.40.5=0,mp4a.40.2=1}
+		 * 
+		 * First codec priorities are given for an entire codec class (eg. "hvc").
+		 * Within each class, where it makes sense, individual streams can be prioritized.
+		 * Say within a class "mp4a" there are two AAC streams. One LC and one HE.
+		 * To use the LC over the HE stream the "mp4a" class gives more detailed codec
+		 * prefixes and their priorities like the above example.
+		 * 
+		 * If used with DASH streams the class priority can be thought of the priority
+		 * of an AdaptationSet and the stream priority of that of a Representation.
+		 * User defined priorities override the @selectionPriority attribute of a
+		 * DASH AdaptationSet or Representation.
+		 */
+		bool Initialize(const FString& ConfigurationString);
+		int32 GetClassPriority(const FString& CodecSpecifierRFC6381) const;
+		int32 GetStreamPriority(const FString& CodecSpecifierRFC6381) const;
+	private:
+		bool ParseInternal(const FString& ConfigurationString);
+		struct FStreamPriority
+		{
+			FString Prefix;
+			int32 Priority = -1;
+		};
+		struct FClassPriority
+		{
+			FString Prefix;
+			int32 Priority = -1;
+			TArray<FStreamPriority> StreamPriorities;
+		};
+		TArray<FClassPriority> ClassPriorities;
+	};
+
 
 
 	struct FPlayerLoopState
diff --git a/Engine/Plugins/Media/ElectraPlayer/Source/ElectraPlayerRuntime/Private/Runtime/Utilities/UtilsMPEGVideo.cpp b/Engine/Plugins/Media/ElectraPlayer/Source/ElectraPlayerRuntime/Private/Runtime/Utilities/UtilsMPEGVideo.cpp
--- a/Engine/Plugins/Media/ElectraPlayer/Source/ElectraPlayerRuntime/Private/Runtime/Utilities/UtilsMPEGVideo.cpp
+++ b/Engine/Plugins/Media/ElectraPlayer/Source/ElectraPlayerRuntime/Private/Runtime/Utilities/UtilsMPEGVideo.cpp
@@ -1008,7 +1008,7 @@ namespace Electra
 
 		uint64 FISO23008_2_seq_parameter_set_data::GetConstraintFlags() const
 		{
-			uint64 ConstraintFlags = ((uint64)general_progressive_source_flag << 3) | ((uint64)general_interlaced_source_flag << 2) | ((uint64)general_non_packed_constraint_flag << 1) | ((uint64)general_frame_only_constraint_flag << 44) | general_reserved_zero_44bits;
+			uint64 ConstraintFlags = ((uint64)((general_progressive_source_flag << 3) | (general_interlaced_source_flag << 2) | (general_non_packed_constraint_flag << 1) | general_frame_only_constraint_flag) << 44) | general_reserved_zero_44bits;
 			return ConstraintFlags;
 		}
 
@@ -1031,11 +1031,11 @@ namespace Electra
 
 			if (general_profile_space == 0)
 			{
-				return FString::Printf(TEXT("%s.%d.%X%c%d%s"), SampleTypePrefix, general_profile_idc, cf, general_tier_flag ? TCHAR('H') : TCHAR('L'), general_level_idc, *cfs);
+				return FString::Printf(TEXT("%s.%d.%X.%c%d%s"), SampleTypePrefix, general_profile_idc, cf, general_tier_flag ? TCHAR('H') : TCHAR('L'), general_level_idc, *cfs);
 			}
 			else
 			{
-				return FString::Printf(TEXT("%s.%c%d.%X%c%d%s"), SampleTypePrefix, TCHAR('A')+general_profile_space-1, general_profile_idc, cf, general_tier_flag ? TCHAR('H') : TCHAR('L'), general_level_idc, *cfs);
+				return FString::Printf(TEXT("%s.%c%d.%X.%c%d%s"), SampleTypePrefix, TCHAR('A')+general_profile_space-1, general_profile_idc, cf, general_tier_flag ? TCHAR('H') : TCHAR('L'), general_level_idc, *cfs);
 			}
 		}
 
