Fix a leak of FAccessUnit objects - a peeked AU was released, but not once it was actually popped

diff --git a/Engine/Plugins/Media/ElectraPlayer/Source/ElectraPlayerRuntime/Private/Runtime/Player/AdaptiveStreamingPlayerDecoders.cpp b/Engine/Plugins/Media/ElectraPlayer/Source/ElectraPlayerRuntime/Private/Runtime/Player/AdaptiveStreamingPlayerDecoders.cpp
--- a/Engine/Plugins/Media/ElectraPlayer/Source/ElectraPlayerRuntime/Private/Runtime/Player/AdaptiveStreamingPlayerDecoders.cpp
+++ b/Engine/Plugins/Media/ElectraPlayer/Source/ElectraPlayerRuntime/Private/Runtime/Player/AdaptiveStreamingPlayerDecoders.cpp
@@ -472,12 +472,12 @@ void FAdaptiveStreamingPlayer::FeedDecoder(EStreamType Type, FMultiTrackAccessUn
 			}
 		}
 
-		FAccessUnit* AccessUnit = nullptr;
-		FromMultistreamBuffer.PeekAndAddRef(AccessUnit);
-		if (AccessUnit)
+		FAccessUnit* PeekedAU = nullptr;
+		FromMultistreamBuffer.PeekAndAddRef(PeekedAU);
+		if (PeekedAU)
 		{
 			// Change in codec?
-			if (CurrentCodecInfo && AccessUnit->AUCodecData.IsValid() && AccessUnit->AUCodecData->ParsedInfo.GetCodec() != CurrentCodecInfo->GetCodec())
+			if (CurrentCodecInfo && PeekedAU->AUCodecData.IsValid() && PeekedAU->AUCodecData->ParsedInfo.GetCodec() != CurrentCodecInfo->GetCodec())
 			{
 				bCodecChangeDetected = true;
 				if (Decoder)
@@ -500,9 +500,17 @@ void FAdaptiveStreamingPlayer::FeedDecoder(EStreamType Type, FMultiTrackAccessUn
 						PostError(err);
 					}
 				}
+				FAccessUnit::Release(PeekedAU);
 			}
 			else
 			{
+				// release peeked and actually pop it
+				FAccessUnit* AccessUnit = nullptr;
+				FromMultistreamBuffer.Pop(AccessUnit);
+				check(PeekedAU == AccessUnit);
+				FAccessUnit::Release(PeekedAU);
+				PeekedAU = nullptr;
+
 				if (Type == EStreamType::Video && VideoDecoder.bApplyNewLimits)
 				{
 					FStreamCodecInformation StreamInfo;
@@ -519,25 +527,18 @@ void FAdaptiveStreamingPlayer::FeedDecoder(EStreamType Type, FMultiTrackAccessUn
 					VideoDecoder.bApplyNewLimits = false;
 				}
 
-				FromMultistreamBuffer.Pop(AccessUnit);
 				if (Decoder)
 				{
 				// The decoder has asked to be fed a new AU so it better be able to accept it.
 					/*IAccessUnitBufferInterface::EAUpushResult auPushRes =*/ Decoder->AUdataPushAU(AccessUnit);
 				}
 
-				// The decoder will have added a ref count if it took the AU. If it didnt' for whatever reason
-				// we still release it to get rid of it and not cause a memory leak.
-				FAccessUnit::Release(AccessUnit);
-				AccessUnit = nullptr;
-
 				if (pAvailability)
 				{
 					UpdateDataAvailabilityState(*pAvailability, Metrics::FDataAvailabilityChange::EAvailability::DataAvailable);
 				}
+				FAccessUnit::Release(AccessUnit);
 			}
-			// Release ref count from peeking.
-			FAccessUnit::Release(AccessUnit);
 		}
 	}
 	// An AU is not tagged as being "the last" one. Instead the EOD is handled separately and must be dealt with
