ElectraPlayer: Using output parameters from the AAC CSD instead of the values reported by the codec library which will be wrong for HE-AAC.

#rnx

diff --git a/Engine/Plugins/Media/ElectraPlayer/Source/ElectraPlayerRuntime/Private/Runtime/Utilities/UtilsMPEGAudio.h b/Engine/Plugins/Media/ElectraPlayer/Source/ElectraPlayerRuntime/Private/Runtime/Utilities/UtilsMPEGAudio.h
--- a/Engine/Plugins/Media/ElectraPlayer/Source/ElectraPlayerRuntime/Private/Runtime/Utilities/UtilsMPEGAudio.h
+++ b/Engine/Plugins/Media/ElectraPlayer/Source/ElectraPlayerRuntime/Private/Runtime/Utilities/UtilsMPEGAudio.h
@@ -19,6 +19,15 @@ namespace Electra
 			void Reset();
 			const TArray<uint8>& GetCodecSpecificData() const;
 
+			void SetRawData(const TArray<uint8>& InRawData)
+			{
+				RawData = InRawData;
+			}
+			const TArray<uint8>& GetRawData() const
+			{
+				return RawData;
+			}
+
 			int32		SBRSignal;
 			int32		PSSignal;
 			uint32		ChannelConfiguration;
@@ -30,6 +39,7 @@ namespace Electra
 			uint32		ExtAOT;
 		private:
 			TArray<uint8>	CodecSpecificData;
+			TArray<uint8>	RawData;
 		};
 
 
