a couple more fixes to ensure WindowsHWrapper is included on all Microsoft platforms.

#rnx

diff --git a/Engine/Plugins/Media/ElectraPlayer/Source/ElectraPlayerRuntime/Private/Runtime/Core/MediaTypes.h b/Engine/Plugins/Media/ElectraPlayer/Source/ElectraPlayerRuntime/Private/Runtime/Core/MediaTypes.h
--- a/Engine/Plugins/Media/ElectraPlayer/Source/ElectraPlayerRuntime/Private/Runtime/Core/MediaTypes.h
+++ b/Engine/Plugins/Media/ElectraPlayer/Source/ElectraPlayerRuntime/Private/Runtime/Core/MediaTypes.h
@@ -9,8 +9,8 @@
 #include "CoreTypes.h"
 #include "Templates/SharedPointer.h"
 
-#if PLATFORM_WINDOWS
-	#include "Windows/WindowsHWrapper.h"
+#if PLATFORM_MICROSOFT
+	#include "Microsoft/WindowsHWrapper.h"
 #endif
 
 typedef uint32  UEMediaError;
diff --git a/Engine/Source/Runtime/Core/Public/HAL/MallocStomp2.h b/Engine/Source/Runtime/Core/Public/HAL/MallocStomp2.h
--- a/Engine/Source/Runtime/Core/Public/HAL/MallocStomp2.h
+++ b/Engine/Source/Runtime/Core/Public/HAL/MallocStomp2.h
@@ -14,8 +14,8 @@
  * Run-time validation should be enabled using '-stomp2malloc' command line argument.
  */
 
-#if PLATFORM_WINDOWS
-	#include "Windows/WindowsHWrapper.h"
+#if PLATFORM_MICROSOFT
+	#include "Microsoft/WindowsHWrapper.h"
 #endif
 
 /**
