ElectraPlayer: Fixed a possible crash on shutdown due to the list of all instances being accesses outside the mutex lock.

#rnx

diff --git a/Engine/Plugins/Media/ElectraPlayer/Source/ElectraPlayerRuntime/Private/Runtime/Player/AdaptiveStreamingPlayer.cpp b/Engine/Plugins/Media/ElectraPlayer/Source/ElectraPlayerRuntime/Private/Runtime/Player/AdaptiveStreamingPlayer.cpp
--- a/Engine/Plugins/Media/ElectraPlayer/Source/ElectraPlayerRuntime/Private/Runtime/Player/AdaptiveStreamingPlayer.cpp
+++ b/Engine/Plugins/Media/ElectraPlayer/Source/ElectraPlayerRuntime/Private/Runtime/Player/AdaptiveStreamingPlayer.cpp
@@ -3601,11 +3601,10 @@ void FAdaptiveStreamingPlayerWorkerThread::WorkerThreadFN()
 	{
 		HaveWorkSignal.WaitTimeoutAndReset(1000 * 20);
 
-		TArray<FAdaptiveStreamingPlayer*>	ActiveInstances(PlayerInstances);
 		SingletonLock.Lock();
-		for(int32 i=0; i<ActiveInstances.Num(); ++i)
+		for(int32 i=0; i<PlayerInstances.Num(); ++i)
 		{
-			ActiveInstances[i]->HandleOnce();
+			PlayerInstances[i]->HandleOnce();
 		}
 		SingletonLock.Unlock();
 	}
