ElectraPlayer: Added support for subtitles

diff --git a/Engine/Plugins/Media/ElectraPlayer/ElectraPlayer.uplugin b/Engine/Plugins/Media/ElectraPlayer/ElectraPlayer.uplugin
--- a/Engine/Plugins/Media/ElectraPlayer/ElectraPlayer.uplugin
+++ b/Engine/Plugins/Media/ElectraPlayer/ElectraPlayer.uplugin
@@ -52,6 +52,10 @@
 		{
 			"Name": "ElectraCDM",
 			"Enabled": true
+		},
+		{
+			"Name": "ElectraSubtitles",
+			"Enabled": true
 		}
 	]
 }
\ No newline at end of file
diff --git a/Engine/Plugins/Media/ElectraPlayer/Source/ElectraPlayerPlugin/Private/ElectraPlayerPlugin.cpp b/Engine/Plugins/Media/ElectraPlayer/Source/ElectraPlayerPlugin/Private/ElectraPlayerPlugin.cpp
--- a/Engine/Plugins/Media/ElectraPlayer/Source/ElectraPlayerPlugin/Private/ElectraPlayerPlugin.cpp
+++ b/Engine/Plugins/Media/ElectraPlayer/Source/ElectraPlayerPlugin/Private/ElectraPlayerPlugin.cpp
@@ -10,7 +10,9 @@
 #include "IElectraPlayerRuntimeModule.h"
 #include "IElectraPlayerPluginModule.h"
 #include "IElectraMetadataSample.h"
+#include "IElectraSubtitleSample.h"
 
+#include "MediaSubtitleDecoderOutput.h"
 #include "MediaMetaDataDecoderOutput.h"
 
 using namespace Electra;
@@ -167,6 +169,47 @@ public:
 
 //-----------------------------------------------------------------------------
 
+class FElectraSubtitleSample : public IElectraSubtitleSample
+{
+public:
+	virtual FGuid GetGUID() const override
+	{ 
+		return IElectraSubtitleSample::GetSampleTypeGUID(); 
+	}
+
+	virtual FMediaTimeStamp GetTime() const override
+	{ 
+		FDecoderTimeStamp ts = Subtitle->GetTime(); 
+		return FMediaTimeStamp(ts.Time, ts.SequenceIndex);
+	}
+	
+	virtual FTimespan GetDuration() const override
+	{
+		return Subtitle->GetDuration();
+	}
+
+	virtual TOptional<FVector2D> GetPosition() const override
+	{
+		return TOptional<FVector2D>();
+	}
+	
+	virtual FText GetText() const override
+	{
+		FUTF8ToTCHAR cnv((const ANSICHAR*)Subtitle->GetData().GetData(), Subtitle->GetData().Num());
+		FString UTF8Text(cnv.Length(), cnv.Get());
+		return FText::FromString(UTF8Text);
+	}
+	
+	virtual EMediaOverlaySampleType GetType() const override
+	{
+		return EMediaOverlaySampleType::Subtitle;
+	}
+
+	ISubtitleDecoderOutputPtr Subtitle;
+};
+
+//-----------------------------------------------------------------------------
+
 
 Electra::FVariantValue FElectraPlayerPlugin::FPlayerAdapterDelegate::QueryOptions(EOptionType Type, const Electra::FVariantValue & Param)
 {
@@ -284,6 +327,19 @@ void FElectraPlayerPlugin::FPlayerAdapterDelegate::OnAudioFlush()
 }
 
 
+void FElectraPlayerPlugin::FPlayerAdapterDelegate::OnSubtitleFlush()
+{
+	TSharedPtr<FElectraPlayerPlugin, ESPMode::ThreadSafe> PinnedHost = Host.Pin();
+	if (PinnedHost.IsValid())
+	{
+		TRange<FTimespan> AllTime(FTimespan::MinValue(), FTimespan::MaxValue());
+		TSharedPtr<IMediaOverlaySample, ESPMode::ThreadSafe> FlushSample;
+		while (PinnedHost->GetSamples().FetchSubtitle(AllTime, FlushSample))
+		{ }
+	}
+}
+
+
 void FElectraPlayerPlugin::FPlayerAdapterDelegate::PresentVideoFrame(const FVideoDecoderOutputPtr & InVideoFrame)
 {
 	TSharedPtr<FElectraPlayerPlugin, ESPMode::ThreadSafe> PinnedHost = Host.Pin();
@@ -307,6 +363,17 @@ void FElectraPlayerPlugin::FPlayerAdapterDelegate::PresentAudioFrame(const IAudi
 	}
 }
 
+void FElectraPlayerPlugin::FPlayerAdapterDelegate::PresentSubtitleSample(const ISubtitleDecoderOutputPtr& InSubtitleSample)
+{
+	TSharedPtr<FElectraPlayerPlugin, ESPMode::ThreadSafe> PinnedHost = Host.Pin();
+	if (PinnedHost.IsValid())
+	{
+		TSharedRef<FElectraSubtitleSample, ESPMode::ThreadSafe> SubtitleSample = MakeShared<FElectraSubtitleSample, ESPMode::ThreadSafe>();
+		SubtitleSample->Subtitle = InSubtitleSample;
+		PinnedHost->MediaSamples->AddSubtitle(SubtitleSample);
+	}
+}
+
 void FElectraPlayerPlugin::FPlayerAdapterDelegate::PresentMetadataSample(const IMetaDataDecoderOutputPtr& InMetadataFrame)
 {
 	TSharedPtr<FElectraPlayerPlugin, ESPMode::ThreadSafe> PinnedHost = Host.Pin();
@@ -416,6 +483,7 @@ bool FElectraPlayerPlugin::Open(const FString& Url, const IMediaOptions* Options
 	{
 		LocalPlaystartOptions.TimeOffset = InPlayerOptions->SeekTime;
 		LocalPlaystartOptions.InitialAudioTrackAttributes.TrackIndexOverride = InPlayerOptions->Tracks.Audio;
+		LocalPlaystartOptions.InitialSubtitleTrackAttributes.TrackIndexOverride = InPlayerOptions->Tracks.Subtitle;
 	}
 	FString InitialAudioLanguage = Options->GetMediaOption(TEXT("InitialAudioLanguage"), FString());
 	if (InitialAudioLanguage.Len())
@@ -423,6 +491,12 @@ bool FElectraPlayerPlugin::Open(const FString& Url, const IMediaOptions* Options
 		LocalPlaystartOptions.InitialAudioTrackAttributes.Language_ISO639 = InitialAudioLanguage;
 		UE_LOG(LogElectraPlayerPlugin, Log, TEXT("[%p] IMediaPlayer::Open: Asking for initial audio language \"%s\""), this, *InitialAudioLanguage);
 	}
+	FString InitialSubtitleLanguage = Options->GetMediaOption(TEXT("InitialSubtitleLanguage"), FString());
+	if (InitialSubtitleLanguage.Len())
+	{
+		LocalPlaystartOptions.InitialSubtitleTrackAttributes.Language_ISO639 = InitialSubtitleLanguage;
+		UE_LOG(LogElectraPlayerPlugin, Log, TEXT("[%p] IMediaPlayer::Open: Asking for initial subtitle language \"%s\""), this, *InitialSubtitleLanguage);
+	}
 	bool bNoPreloading = Options->GetMediaOption(TEXT("ElectraNoPreloading"), (bool)false);
 	if (bNoPreloading)
 	{
diff --git a/Engine/Plugins/Media/ElectraPlayer/Source/ElectraPlayerPlugin/Private/ElectraPlayerPlugin.h b/Engine/Plugins/Media/ElectraPlayer/Source/ElectraPlayerPlugin/Private/ElectraPlayerPlugin.h
--- a/Engine/Plugins/Media/ElectraPlayer/Source/ElectraPlayerPlugin/Private/ElectraPlayerPlugin.h
+++ b/Engine/Plugins/Media/ElectraPlayer/Source/ElectraPlayerPlugin/Private/ElectraPlayerPlugin.h
@@ -162,8 +162,10 @@ private:
 		void SendMediaEvent(EPlayerEvent Event) override;
 		void OnVideoFlush() override;
 		void OnAudioFlush() override;
+		void OnSubtitleFlush() override;
 		void PresentVideoFrame(const FVideoDecoderOutputPtr& InVideoFrame) override;
 		void PresentAudioFrame(const IAudioDecoderOutputPtr& InAudioFrame) override;
+		void PresentSubtitleSample(const ISubtitleDecoderOutputPtr& InSubtitleSample) override;
 		void PresentMetadataSample(const IMetaDataDecoderOutputPtr& InMetadataSample) override;
 		bool CanReceiveVideoSamples(int32 NumFrames) override;
 		bool CanReceiveAudioSamples(int32 NumFrames) override;
diff --git a/Engine/Plugins/Media/ElectraPlayer/Source/ElectraPlayerRuntime/ElectraPlayerRuntime.Build.cs b/Engine/Plugins/Media/ElectraPlayer/Source/ElectraPlayerRuntime/ElectraPlayerRuntime.Build.cs
--- a/Engine/Plugins/Media/ElectraPlayer/Source/ElectraPlayerRuntime/ElectraPlayerRuntime.Build.cs
+++ b/Engine/Plugins/Media/ElectraPlayer/Source/ElectraPlayerRuntime/ElectraPlayerRuntime.Build.cs
@@ -23,6 +23,7 @@ namespace UnrealBuildTool.Rules
 					"Json",
 					"ElectraBase",
 					"ElectraCDM",
+					"ElectraSubtitles",
 					"XmlParser"
 				});
 			if (Target.bCompileAgainstEngine)
diff --git a/Engine/Plugins/Media/ElectraPlayer/Source/ElectraPlayerRuntime/Private/ElectraPlayer.cpp b/Engine/Plugins/Media/ElectraPlayer/Source/ElectraPlayerRuntime/Private/ElectraPlayer.cpp
--- a/Engine/Plugins/Media/ElectraPlayer/Source/ElectraPlayerRuntime/Private/ElectraPlayer.cpp
+++ b/Engine/Plugins/Media/ElectraPlayer/Source/ElectraPlayerRuntime/Private/ElectraPlayer.cpp
@@ -182,10 +182,13 @@ void FElectraPlayer::ClearToDefaultState()
 
 	NumTracksAudio = 0;
 	NumTracksVideo = 0;
+	NumTracksSubtitle = 0;
 	SelectedQuality = 0;
 	SelectedVideoTrackIndex = -1;
 	SelectedAudioTrackIndex = -1;
+	SelectedSubtitleTrackIndex = -1;
 	bAudioTrackIndexDirty = true;
+	bSubtitleTrackIndexDirty = true;
 	bInitialSeekPerformed = false;
 	LastPresentedFrameDimension = FIntPoint::ZeroValue;
 	DeferredPlayerEvents.Empty();
@@ -224,8 +227,9 @@ bool FElectraPlayer::OpenInternal(const FString& Url, const FParamDict& PlayerOp
 	// Create a new empty player structure. This contains the actual player instance, its associated renderers and sample queues.
 	TSharedPtr<FInternalPlayerImpl, ESPMode::ThreadSafe> NewPlayer = MakeShared<FInternalPlayerImpl, ESPMode::ThreadSafe>();
 
-	// Get a writable copy of the URL so we can override it for debugging.
+	// Get a writable copy of the URL so we can sanitize it if necessary.
 	MediaUrl = Url;
+	MediaUrl.TrimStartAndEndInline();
 	UE_LOG(LogElectraPlayer, Log, TEXT("[%p][%p] IMediaPlayer::Open(%s)"), this, NewPlayer.Get(), *SanitizeMessage(MediaUrl));
 
 	// Create the renderers so we can pass them to the internal player.
@@ -243,6 +247,12 @@ bool FElectraPlayer::OpenInternal(const FString& Url, const FParamDict& PlayerOp
 	NewPlayer->AdaptivePlayer->SetStaticResourceProviderCallback(StaticResourceProvider);
 	NewPlayer->AdaptivePlayer->SetVideoDecoderResourceDelegate(VideoDecoderResourceDelegate);
 
+	// Create the subtitle receiver and register it with the player.
+	MediaPlayerSubtitleReceiver = MakeSharedTS<FSubtitleEventReceiver>();
+	MediaPlayerSubtitleReceiver->GetSubtitleReceivedDelegate().BindRaw(this, &FElectraPlayer::OnSubtitleDecoded);
+	MediaPlayerSubtitleReceiver->GetSubtitleFlushDelegate().BindRaw(this, &FElectraPlayer::OnSubtitleFlush);
+	NewPlayer->AdaptivePlayer->AddSubtitleReceiver(MediaPlayerSubtitleReceiver);
+
 	// Create a new media player event receiver and register it to receive all non player internal events as soon as they are received.
 	MediaPlayerEventReceiver = MakeSharedTS<FAEMSEventReceiver>();
 	MediaPlayerEventReceiver->GetEventReceivedDelegate().BindRaw(this, &FElectraPlayer::OnMediaPlayerEventReceived);
@@ -345,6 +355,14 @@ void FElectraPlayer::CloseInternal(bool bKillAfterClose)
 			MediaPlayerEventReceiver.Reset();
 		}
 
+		if (MediaPlayerSubtitleReceiver.IsValid())
+		{
+			Player->AdaptivePlayer->RemoveSubtitleReceiver(MediaPlayerSubtitleReceiver);
+			MediaPlayerSubtitleReceiver->GetSubtitleReceivedDelegate().Unbind();
+			MediaPlayerSubtitleReceiver->GetSubtitleFlushDelegate().Unbind();
+			MediaPlayerSubtitleReceiver.Reset();
+		}
+
 		// Unregister ourselves as the provider for static resources.
 		Player->AdaptivePlayer->SetStaticResourceProviderCallback(nullptr);
 		// Also unregister us from receiving further metric callbacks.
@@ -584,6 +602,32 @@ void FElectraPlayer::OnAudioFlush()
 	}
 }
 
+void FElectraPlayer::OnSubtitleDecoded(ISubtitleDecoderOutputPtr DecoderOutput)
+{
+	TSharedPtr<FInternalPlayerImpl, ESPMode::ThreadSafe> Player = CurrentPlayer;
+	if (Player.IsValid())
+	{
+		if (State != EPlayerState::Closed)
+		{
+			PresentSubtitle(DecoderOutput);
+		}
+	}
+}
+
+void FElectraPlayer::OnSubtitleFlush()
+{
+	TSharedPtr<FInternalPlayerImpl, ESPMode::ThreadSafe> Player = CurrentPlayer;
+	if (Player.IsValid())
+	{
+		TSharedPtr<IElectraPlayerAdapterDelegate, ESPMode::ThreadSafe> PinnedAdapterDelegate = AdapterDelegate.Pin();
+		if (PinnedAdapterDelegate.IsValid())
+		{
+			PinnedAdapterDelegate->OnSubtitleFlush();
+		}
+	}
+}
+
+
 void FElectraPlayer::OnVideoRenderingStarted()
 {
 	TSharedPtr<FInternalPlayerImpl, ESPMode::ThreadSafe> Player = CurrentPlayer;
@@ -655,6 +699,17 @@ bool FElectraPlayer::PresentAudioFrame(const IAudioDecoderOutputPtr& DecoderOutp
 	return true;
 }
 
+
+bool FElectraPlayer::PresentSubtitle(const ISubtitleDecoderOutputPtr& DecoderOutput)
+{
+	TSharedPtr<IElectraPlayerAdapterDelegate, ESPMode::ThreadSafe> PinnedAdapterDelegate = AdapterDelegate.Pin();
+	if (PinnedAdapterDelegate.IsValid())
+	{
+		PinnedAdapterDelegate->PresentSubtitleSample(DecoderOutput);
+	}
+	return true;
+}
+
 //-----------------------------------------------------------------------------
 /**
  * Attempt to drop any old frames from the presentation queue
@@ -944,6 +999,10 @@ TSharedPtr<Electra::FTrackMetadata, ESPMode::ThreadSafe> FElectraPlayer::GetTrac
 		{
 			LockedPlayer->AdaptivePlayer->GetTrackMetadata(TrackMetaData, Electra::EStreamType::Audio);
 		}
+		else if (TrackType == EPlayerTrackType::Subtitle)
+		{
+			LockedPlayer->AdaptivePlayer->GetTrackMetadata(TrackMetaData, Electra::EStreamType::Subtitle);
+		}
 		if (TrackIndex >= 0 && TrackIndex < TrackMetaData.Num())
 		{
 			return MakeShared<Electra::FTrackMetadata, ESPMode::ThreadSafe>(TrackMetaData[TrackIndex]);
@@ -1003,6 +1062,10 @@ int32 FElectraPlayer::GetNumTracks(EPlayerTrackType TrackType) const
 	{
 		return NumTracksVideo;
 	}
+	else if (TrackType == EPlayerTrackType::Subtitle)
+	{
+		return NumTracksSubtitle;
+	}
 	// TODO: Implement missing track types.
 
 	return 0;
@@ -1012,7 +1075,8 @@ int32 FElectraPlayer::GetNumTrackFormats(EPlayerTrackType TrackType, int32 Track
 {
 	// Right now we only have a single format per track
 	if ((TrackType == EPlayerTrackType::Video && NumTracksVideo != 0) ||
-		(TrackType == EPlayerTrackType::Audio && NumTracksAudio != 0))
+		(TrackType == EPlayerTrackType::Audio && NumTracksAudio != 0) ||
+		(TrackType == EPlayerTrackType::Subtitle && NumTracksSubtitle != 0))
 	{
 		return 1;
 	}
@@ -1026,17 +1090,17 @@ int32 FElectraPlayer::GetSelectedTrack(EPlayerTrackType TrackType) const
 	{
 		return SelectedVideoTrackIndex;
 	}
+	/*
+		To reduce the overhead of this function we check for the track the underlying player has
+		actually selected only when we were told the tracks changed.
+
+		It is possible that the underlying player changes the track automatically as playback progresses.
+		For instance, when playing a DASH stream consisting of several periods the player needs to re-select
+		the audio stream when transitioning from one period into the next, which may change the index of
+		the selected track.
+	*/
 	else if (TrackType == EPlayerTrackType::Audio)
 	{
-		/*
-			To reduce the overhead of this function we check for the track the underlying player has
-			actually selected only when we were told the tracks changed.
-
-			It is possible that the underlying player changes the track automatically as playback progresses.
-			For instance, when playing a DASH stream consisting of several periods the player needs to re-select
-			the audio stream when transitioning from one period into the next, which may change the index of
-			the selected track.
-		*/
 		if (bAudioTrackIndexDirty)
 		{
 			if (NumTracksAudio == 0)
@@ -1068,6 +1132,39 @@ int32 FElectraPlayer::GetSelectedTrack(EPlayerTrackType TrackType) const
 		}
 		return SelectedAudioTrackIndex;
 	}
+	else if (TrackType == EPlayerTrackType::Subtitle)
+	{
+		if (bSubtitleTrackIndexDirty)
+		{
+			if (NumTracksSubtitle == 0)
+			{
+				SelectedSubtitleTrackIndex = -1;
+			}
+			else
+			{
+				TSharedPtr<FInternalPlayerImpl, ESPMode::ThreadSafe> LockedPlayer = CurrentPlayer;
+				if (LockedPlayer.IsValid() && LockedPlayer->AdaptivePlayer.IsValid())
+				{
+					if (LockedPlayer->AdaptivePlayer->IsTrackDeselected(Electra::EStreamType::Subtitle))
+					{
+						SelectedSubtitleTrackIndex = -1;
+						bSubtitleTrackIndexDirty = false;
+					}
+					else
+					{
+						Electra::FStreamSelectionAttributes Attributes;
+						LockedPlayer->AdaptivePlayer->GetSelectedTrackAttributes(Attributes, Electra::EStreamType::Subtitle);
+						if (Attributes.OverrideIndex.IsSet())
+						{
+							SelectedSubtitleTrackIndex = Attributes.OverrideIndex.GetValue();
+							bSubtitleTrackIndexDirty = false;
+						}
+					}
+				}
+			}
+		}
+		return SelectedSubtitleTrackIndex;
+	}
 	return -1;
 }
 
@@ -1078,12 +1175,33 @@ FText FElectraPlayer::GetTrackDisplayName(EPlayerTrackType TrackType, int32 Trac
 	{
 		if (TrackType == EPlayerTrackType::Video)
 		{
+			if (!Meta->Label.IsEmpty())
+			{
+				return FText::FromString(FString::Printf(TEXT("%s, ID=%s"), *Meta->Label, *Meta->ID));
+			}
 			return FText::FromString(FString::Printf(TEXT("Video Track ID %s"), *Meta->ID));
 		}
 		else if (TrackType == EPlayerTrackType::Audio)
 		{
+			if (!Meta->Label.IsEmpty())
+			{
+				return FText::FromString(FString::Printf(TEXT("%s, ID=%s"), *Meta->Label, *Meta->ID));
+			}
 			return FText::FromString(FString::Printf(TEXT("Audio Track ID %s"), *Meta->ID));
 		}
+		else if (TrackType == EPlayerTrackType::Subtitle)
+		{
+			FString Name;
+			if (!Meta->Label.IsEmpty())
+			{
+				Name = FString::Printf(TEXT("%s (%s), ID=%s"), *Meta->Label, *Meta->HighestBandwidthCodec.GetCodecSpecifierRFC6381(), *Meta->ID);
+			}
+			else
+			{
+				Name = FString::Printf(TEXT("Subtitle Track ID %s (%s)"), *Meta->ID, *Meta->HighestBandwidthCodec.GetCodecSpecifierRFC6381());
+			}
+			return FText::FromString(Name);
+		}
 	}
 	return FText();
 }
@@ -1188,6 +1306,57 @@ bool FElectraPlayer::SelectTrack(EPlayerTrackType TrackType, int32 TrackIndex)
 			return true;
 		}
 	}
+	else if (TrackType == EPlayerTrackType::Subtitle)
+	{
+		// Select a track or deselect?
+		if (TrackIndex >= 0)
+		{
+			// Check if the track index exists by checking the presence of the track metadata.
+			// If for some reason the index is not valid the selection will not be changed.
+			TSharedPtr<Electra::FTrackMetadata, ESPMode::ThreadSafe> Meta = GetTrackStreamMetadata(EPlayerTrackType::Subtitle, TrackIndex);
+			if (Meta.IsValid())
+			{
+				// Switch only when the track index has changed.
+				if (GetSelectedTrack(TrackType) != TrackIndex)
+				{
+					Electra::FStreamSelectionAttributes SubtitleAttributes;
+					SubtitleAttributes.OverrideIndex = TrackIndex;
+
+					PlaystartOptions.InitialSubtitleTrackAttributes.TrackIndexOverride = TrackIndex;
+					if (!Meta->Kind.IsEmpty())
+					{
+						SubtitleAttributes.Kind = Meta->Kind;
+						PlaystartOptions.InitialSubtitleTrackAttributes.Kind = Meta->Kind;
+					}
+					if (!Meta->Language.IsEmpty())
+					{
+						SubtitleAttributes.Language_ISO639 = Meta->Language;
+						PlaystartOptions.InitialSubtitleTrackAttributes.Language_ISO639 = Meta->Language;
+					}
+					TSharedPtr<FInternalPlayerImpl, ESPMode::ThreadSafe> LockedPlayer = CurrentPlayer;
+					if (LockedPlayer.IsValid() && LockedPlayer->AdaptivePlayer.IsValid())
+					{
+						LockedPlayer->AdaptivePlayer->SelectTrackByAttributes(Electra::EStreamType::Subtitle, SubtitleAttributes);
+					}
+
+					SelectedSubtitleTrackIndex = TrackIndex;
+				}
+				return true;
+			}
+		}
+		else
+		{
+			// Deselect subtitle track.
+			PlaystartOptions.InitialSubtitleTrackAttributes.TrackIndexOverride = -1;
+			SelectedSubtitleTrackIndex = -1;
+			TSharedPtr<FInternalPlayerImpl, ESPMode::ThreadSafe> LockedPlayer = CurrentPlayer;
+			if (LockedPlayer.IsValid() && LockedPlayer->AdaptivePlayer.IsValid())
+			{
+				LockedPlayer->AdaptivePlayer->DeselectTrack(Electra::EStreamType::Subtitle);
+			}
+			return true;
+		}
+	}
 	// The player is selecting tracks for now.
 	return false;
 }
@@ -1522,6 +1691,10 @@ void FElectraPlayer::HandlePlayerEventReceivedPlaylists()
 	CurrentPlayer->AdaptivePlayer->GetTrackMetadata(AudioStreamMetaData, Electra::EStreamType::Audio);
 	NumTracksAudio = AudioStreamMetaData.Num();
 
+	TArray<Electra::FTrackMetadata> SubtitleStreamMetaData;
+	CurrentPlayer->AdaptivePlayer->GetTrackMetadata(SubtitleStreamMetaData, Electra::EStreamType::Subtitle);
+	NumTracksSubtitle = SubtitleStreamMetaData.Num();
+
 	// Set the initial audio track selection attributes.
 	Electra::FStreamSelectionAttributes InitialAudioAttributes;
 	InitialAudioAttributes.Kind = PlaystartOptions.InitialAudioTrackAttributes.Kind;
@@ -1529,6 +1702,13 @@ void FElectraPlayer::HandlePlayerEventReceivedPlaylists()
 	InitialAudioAttributes.OverrideIndex = PlaystartOptions.InitialAudioTrackAttributes.TrackIndexOverride;
 	CurrentPlayer->AdaptivePlayer->SetInitialStreamAttributes(Electra::EStreamType::Audio, InitialAudioAttributes);
 
+	// Set the initial subtitle track selection attributes.
+	Electra::FStreamSelectionAttributes InitialSubtitleAttributes;
+	InitialSubtitleAttributes.Kind = PlaystartOptions.InitialSubtitleTrackAttributes.Kind;
+	InitialSubtitleAttributes.Language_ISO639 = PlaystartOptions.InitialSubtitleTrackAttributes.Language_ISO639;
+	InitialSubtitleAttributes.OverrideIndex = PlaystartOptions.InitialSubtitleTrackAttributes.TrackIndexOverride;
+	CurrentPlayer->AdaptivePlayer->SetInitialStreamAttributes(Electra::EStreamType::Subtitle, InitialSubtitleAttributes);
+
 	// Trigger preloading unless forbidden.
 	if (PlaystartOptions.bDoNotPreload == false)
 	{
@@ -1555,7 +1735,13 @@ void FElectraPlayer::HandlePlayerEventTracksChanged()
 	TArray<Electra::FTrackMetadata> AudioStreamMetaData;
 	CurrentPlayer->AdaptivePlayer->GetTrackMetadata(AudioStreamMetaData, Electra::EStreamType::Audio);
 	NumTracksAudio = AudioStreamMetaData.Num();
+
+	TArray<Electra::FTrackMetadata> SubtitleStreamMetaData;
+	CurrentPlayer->AdaptivePlayer->GetTrackMetadata(SubtitleStreamMetaData, Electra::EStreamType::Subtitle);
+	NumTracksSubtitle = SubtitleStreamMetaData.Num();
+
 	bAudioTrackIndexDirty = true;
+	bSubtitleTrackIndexDirty = true;
 
 	/*
 	Note: Do not send TracksChanged here since this would trigger a re-selecting of the default tracks.
diff --git a/Engine/Plugins/Media/ElectraPlayer/Source/ElectraPlayerRuntime/Private/ElectraPlayer.h b/Engine/Plugins/Media/ElectraPlayer/Source/ElectraPlayerRuntime/Private/ElectraPlayer.h
--- a/Engine/Plugins/Media/ElectraPlayer/Source/ElectraPlayerRuntime/Private/ElectraPlayer.h
+++ b/Engine/Plugins/Media/ElectraPlayer/Source/ElectraPlayerRuntime/Private/ElectraPlayer.h
@@ -19,7 +19,8 @@ class IAudioDecoderOutput;
 using IAudioDecoderOutputPtr = TSharedPtr<IAudioDecoderOutput, ESPMode::ThreadSafe>;
 class IMetaDataDecoderOutput;
 using IMetaDataDecoderOutputPtr = TSharedPtr<IMetaDataDecoderOutput, ESPMode::ThreadSafe>;
-
+class ISubtitleDecoderOutput;
+using ISubtitleDecoderOutputPtr = TSharedPtr<ISubtitleDecoderOutput, ESPMode::ThreadSafe>;
 
 namespace Electra
 {
@@ -52,6 +53,8 @@ public:
 	void OnVideoFlush();
 	void OnAudioDecoded(const IAudioDecoderOutputPtr& DecoderOutput);
 	void OnAudioFlush();
+	void OnSubtitleDecoded(ISubtitleDecoderOutputPtr DecoderOutput);
+	void OnSubtitleFlush();
 
 	void OnVideoRenderingStarted();
 	void OnVideoRenderingStopped();
@@ -135,6 +138,26 @@ private:
 		FOnMediaPlayerEventReceivedDelegate EventReceivedDelegate;
 	};
 
+	DECLARE_DELEGATE_OneParam(FOnMediaPlayerSubtitleReceivedDelegate, ISubtitleDecoderOutputPtr);
+	DECLARE_DELEGATE(FOnMediaPlayerSubtitleFlushDelegate);
+	class FSubtitleEventReceiver : public IAdaptiveStreamingPlayerSubtitleReceiver
+	{
+	public:
+		virtual ~FSubtitleEventReceiver() = default;
+		FOnMediaPlayerSubtitleReceivedDelegate& GetSubtitleReceivedDelegate()
+		{ return SubtitleReceivedDelegate; }
+		FOnMediaPlayerSubtitleFlushDelegate& GetSubtitleFlushDelegate()
+		{ return SubtitleFlushDelegate; }
+	private:
+		virtual void OnMediaPlayerSubtitleReceived(ISubtitleDecoderOutputPtr Subtitle) override
+		{ SubtitleReceivedDelegate.ExecuteIfBound(Subtitle); }
+		virtual void OnMediaPlayerFlushSubtitles() override
+		{ SubtitleFlushDelegate.ExecuteIfBound(); }		 
+		FOnMediaPlayerSubtitleReceivedDelegate SubtitleReceivedDelegate;
+		FOnMediaPlayerSubtitleFlushDelegate SubtitleFlushDelegate;
+	};
+
+
 
 	struct FPlayerMetricEventBase
 	{
@@ -254,6 +277,7 @@ private:
 
 	bool PresentVideoFrame(const FVideoDecoderOutputPtr& InVideoFrame);
 	bool PresentAudioFrame(const IAudioDecoderOutputPtr& DecoderOutput);
+	bool PresentSubtitle(const ISubtitleDecoderOutputPtr& DecoderOutput);
 
 	void PlatformNotifyOfOptionChange();
 
@@ -330,10 +354,13 @@ private:
 	// Contains number of audio tracks available to expose it later.
 	int32											NumTracksAudio;
 	int32											NumTracksVideo;
+	int32											NumTracksSubtitle;
 	int32											SelectedQuality;
 	int32											SelectedVideoTrackIndex;
 	mutable int32									SelectedAudioTrackIndex;
+	mutable int32									SelectedSubtitleTrackIndex;
 	mutable bool									bAudioTrackIndexDirty;
+	mutable bool									bSubtitleTrackIndexDirty;
 
 	bool											bInitialSeekPerformed;
 
@@ -364,6 +391,8 @@ private:
 	TQueue<IElectraPlayerAdapterDelegate::EPlayerEvent>	DeferredEvents;
 	TQueue<TSharedPtrTS<FPlayerMetricEventBase>>	DeferredPlayerEvents;
 	TSharedPtrTS<FAEMSEventReceiver>				MediaPlayerEventReceiver;
+	TSharedPtrTS<FSubtitleEventReceiver>			MediaPlayerSubtitleReceiver;
+
 
 	/** The URL of the currently opened media. */
 	FString											MediaUrl;
diff --git a/Engine/Plugins/Media/ElectraPlayer/Source/ElectraPlayerRuntime/Private/Runtime/Decoder/SubtitleDecoder.cpp b/Engine/Plugins/Media/ElectraPlayer/Source/ElectraPlayerRuntime/Private/Runtime/Decoder/SubtitleDecoder.cpp
new file mode 100644
--- /dev/null
+++ b/Engine/Plugins/Media/ElectraPlayer/Source/ElectraPlayerRuntime/Private/Runtime/Decoder/SubtitleDecoder.cpp
@@ -0,0 +1,360 @@
+// Copyright Epic Games, Inc. All Rights Reserved.
+
+#include "SubtitleDecoder.h"
+#include "IElectraSubtitleModule.h"
+
+
+namespace Electra
+{
+
+class FSubtitleDecoderPlugins
+{
+public:
+	static FSubtitleDecoderPlugins& Get()
+	{
+		static FSubtitleDecoderPlugins This;
+		This.CollectSupportedDecoderPlugins();
+		return This;
+	}
+
+	bool IsSupported(const FString& CodecName);
+	TSharedPtr<IElectraSubtitleDecoder, ESPMode::ThreadSafe> CreateDecoder(const FString& CodecName);
+
+private:
+	struct FCodecPlugin
+	{
+		IElectraSubtitleModularFeature* Plugin = nullptr;
+		int32 Priority = -1;
+	};
+
+	void CollectSupportedDecoderPlugins();
+
+	FCriticalSection Lock;
+	TMap<FString, FCodecPlugin> CodecPlugins;
+	bool bIsValidList = false;
+};
+
+
+void FSubtitleDecoderPlugins::CollectSupportedDecoderPlugins()
+{
+	FScopeLock lock(&Lock);
+	if (!bIsValidList)
+	{
+		// Get the list of all the registered modular features implementing a subtitle decoder.
+		TArray<IElectraSubtitleModularFeature*> PluginImplementations = IModularFeatures::Get().GetModularFeatureImplementations<IElectraSubtitleModularFeature>(IElectraSubtitlesModule::GetModularFeatureName());
+		for(auto &Plugin : PluginImplementations)
+		{
+			if (!Plugin)
+			{
+				continue;
+			}
+			// Get the names of supported codecs and add them to our list, replacing those of lower priority.
+			TArray<FString> CodecNames;
+			Plugin->GetSupportedFormats(CodecNames);
+			for(auto &CodecName : CodecNames)
+			{
+				int32 Prio = Plugin->GetPriorityForFormat(CodecName);
+				bool bAdd = true;
+				for(auto &cp : CodecPlugins)
+				{
+					if (cp.Key.Equals(CodecName) && Prio < cp.Value.Priority)
+					{
+						bAdd = false;
+						break;
+					}
+				}
+				if (bAdd)
+				{
+					CodecPlugins.Emplace(CodecName, FCodecPlugin( {Plugin, Prio} ));
+				}
+			}
+		}
+		bIsValidList = true;
+	}
+}
+
+bool FSubtitleDecoderPlugins::IsSupported(const FString& CodecName)
+{
+	CollectSupportedDecoderPlugins();
+	FScopeLock lock(&Lock);
+	return CodecPlugins.Contains(CodecName);
+}
+
+TSharedPtr<IElectraSubtitleDecoder, ESPMode::ThreadSafe> FSubtitleDecoderPlugins::CreateDecoder(const FString& CodecName)
+{
+	CollectSupportedDecoderPlugins();
+	FScopeLock lock(&Lock);
+	FCodecPlugin* Plugin = CodecPlugins.Find(CodecName);
+	if (Plugin && Plugin->Plugin)
+	{
+		return Plugin->Plugin->CreateDecoderForFormat(CodecName);
+	}
+	return nullptr;
+}
+
+
+/*********************************************************************************************************************/
+/*********************************************************************************************************************/
+/*********************************************************************************************************************/
+
+
+
+
+/**
+ * Generic subtitle decoder
+**/
+class FSubtitleDecoder : public ISubtitleDecoder
+{
+public:
+	static bool IsSupported(const FString& MimeType, const FString& Codec);
+
+	FSubtitleDecoder(TSharedPtr<IElectraSubtitleDecoder, ESPMode::ThreadSafe> InPluginDecoder);
+	virtual ~FSubtitleDecoder();
+
+	virtual void SetPlayerSessionServices(IPlayerSessionServices* SessionServices) override;
+
+	virtual void Open() override;
+	virtual void Close() override;
+	virtual void Start() override;
+	virtual void Stop() override;
+	virtual void UpdatePlaybackPosition(FTimeValue InAbsolutePosition, FTimeValue InLocalPosition) override;
+
+	virtual FTimeValue GetStreamedDeliveryTimeOffset() override;
+
+	virtual FDecodedSubtitleReceivedDelegate& GetDecodedSubtitleReceiveDelegate() override;
+	virtual FDecodedSubtitleFlushDelegate& GetDecodedSubtitleFlushDelegate() override;
+
+	virtual EAUpushResult AUdataPushAU(FAccessUnit* AccessUnit) override;
+	virtual void AUdataPushEOD() override;
+	virtual void AUdataFlushEverything() override;
+
+	void SetCSD(const TArray<uint8> InRawCSD)
+	{
+		RawCSD = InRawCSD;
+	}
+	void SetCodecInfo(const FStreamCodecInformation& InCodecInfo)
+	{
+		CodecInfo = InCodecInfo;
+	}
+
+	bool ResetForTrackChange(FAccessUnit* AccessUnit);
+
+
+private:
+	void OnDecodedSubtitleReceived(ISubtitleDecoderOutputPtr DecodedSubtitle);
+
+	TSharedPtr<IElectraSubtitleDecoder, ESPMode::ThreadSafe> PluginDecoder;
+	IPlayerSessionServices* PlayerSessionServices = nullptr;
+	TArray<uint8> RawCSD;
+	FStreamCodecInformation CodecInfo;
+	FTimeValue DeliveryOffset;
+	FDelegateHandle ReceiveDelegateHandle;
+	FDecodedSubtitleReceivedDelegate SubtitleReceiverDelegate;
+	FDecodedSubtitleFlushDelegate SubtitleFlushDelegate;
+	bool bIsStarted = false;
+};
+
+bool ISubtitleDecoder::IsSupported(const FString& MimeType, const FString& Codec)
+{
+	return FSubtitleDecoder::IsSupported(MimeType, Codec);
+}
+
+ISubtitleDecoder* ISubtitleDecoder::Create(FAccessUnit* AccessUnit)
+{
+	check(AccessUnit);
+	if (AccessUnit && AccessUnit->AUCodecData.IsValid())
+	{
+		TSharedPtr<IElectraSubtitleDecoder, ESPMode::ThreadSafe> PluginDecoder = FSubtitleDecoderPlugins::Get().CreateDecoder(AccessUnit->AUCodecData->ParsedInfo.GetCodecSpecifierRFC6381());
+		if (PluginDecoder.IsValid())
+		{
+			FSubtitleDecoder* Decoder = new FSubtitleDecoder(MoveTemp(PluginDecoder));
+			Decoder->SetCSD(AccessUnit->AUCodecData->RawCSD);
+			Decoder->SetCodecInfo(AccessUnit->AUCodecData->ParsedInfo);
+			return Decoder;
+		}
+	}
+	return nullptr;
+}
+
+
+bool FSubtitleDecoder::IsSupported(const FString& MimeType, const FString& Codec)
+{
+	if (MimeType.IsEmpty() && Codec.IsEmpty())
+	{
+		return false;
+	}
+	bool bMime = MimeType.Len() ? FSubtitleDecoderPlugins::Get().IsSupported(MimeType) : false;
+	bool bCodec = Codec.Len() ? FSubtitleDecoderPlugins::Get().IsSupported(Codec) : false;
+	return bMime || bCodec;
+}
+
+
+FSubtitleDecoder::FSubtitleDecoder(TSharedPtr<IElectraSubtitleDecoder, ESPMode::ThreadSafe> InPluginDecoder)
+	: PluginDecoder(InPluginDecoder)
+{
+	DeliveryOffset.SetToZero();
+}
+
+FSubtitleDecoder::~FSubtitleDecoder()
+{
+}
+
+void FSubtitleDecoder::SetPlayerSessionServices(IPlayerSessionServices* InSessionServices)
+{
+	PlayerSessionServices = InSessionServices;
+}
+
+void FSubtitleDecoder::Open()
+{
+	check(PlayerSessionServices);
+	check(PluginDecoder.IsValid());
+	if (PluginDecoder.IsValid())
+	{
+		FParamDict Addtl;
+		Addtl.Set(TEXT("width"), FVariantValue((int64) CodecInfo.GetResolution().Width));
+		Addtl.Set(TEXT("height"), FVariantValue((int64) CodecInfo.GetResolution().Height));
+		Addtl.Set(TEXT("offset_x"), FVariantValue((int64) CodecInfo.GetTranslation().GetX()));
+		Addtl.Set(TEXT("offset_y"), FVariantValue((int64) CodecInfo.GetTranslation().GetY()));
+		Addtl.Set(TEXT("timescale"), FVariantValue((int64) CodecInfo.GetFrameRate().GetDenominator()));
+		if (PluginDecoder->InitializeStreamWithCSD(RawCSD, Addtl))
+		{
+			DeliveryOffset = PluginDecoder->GetStreamedDeliveryTimeOffset();
+			ReceiveDelegateHandle = PluginDecoder->GetParsedSubtitleReceiveDelegate().AddRaw(this, &FSubtitleDecoder::OnDecodedSubtitleReceived);
+		}
+		else
+		{
+			PlayerSessionServices->PostError(FErrorDetail().SetFacility(Facility::EFacility::SubtitleDecoder).SetCode(1).SetError(UEMEDIA_ERROR_INTERNAL).SetMessage(FString::Printf(TEXT("Subtitle decoder plugin failed to initialize with CSD"))));
+		}
+	}
+	else
+	{
+		PlayerSessionServices->PostError(FErrorDetail().SetFacility(Facility::EFacility::SubtitleDecoder).SetCode(1).SetError(UEMEDIA_ERROR_INTERNAL).SetMessage(FString::Printf(TEXT("No suitable subtitle decoder plugin found"))));
+	}
+}
+
+void FSubtitleDecoder::Close()
+{
+	if (ReceiveDelegateHandle.IsValid())
+	{
+		check(PluginDecoder.IsValid());
+		PluginDecoder->GetParsedSubtitleReceiveDelegate().Remove(ReceiveDelegateHandle);
+		ReceiveDelegateHandle.Reset();
+	}
+	PluginDecoder.Reset();
+	RawCSD.Empty();
+	CodecInfo.Clear();
+}
+
+void FSubtitleDecoder::Start()
+{
+	if (PluginDecoder.IsValid())
+	{
+		PluginDecoder->Start();
+	}
+	bIsStarted = true;
+}
+
+void FSubtitleDecoder::Stop()
+{
+	if (PluginDecoder.IsValid())
+	{
+		PluginDecoder->Stop();
+	}
+	bIsStarted = false;
+}
+
+void FSubtitleDecoder::UpdatePlaybackPosition(FTimeValue InAbsolutePosition, FTimeValue InLocalPosition)
+{
+	if (PluginDecoder.IsValid())
+	{
+		PluginDecoder->UpdatePlaybackPosition(InAbsolutePosition, InLocalPosition);
+	}
+}
+
+
+FTimeValue FSubtitleDecoder::GetStreamedDeliveryTimeOffset()
+{
+	return DeliveryOffset;
+}
+
+ISubtitleDecoder::FDecodedSubtitleReceivedDelegate& FSubtitleDecoder::GetDecodedSubtitleReceiveDelegate()
+{
+	return SubtitleReceiverDelegate;
+}
+
+ISubtitleDecoder::FDecodedSubtitleFlushDelegate& FSubtitleDecoder::GetDecodedSubtitleFlushDelegate()
+{
+	return SubtitleFlushDelegate;
+}
+
+void FSubtitleDecoder::OnDecodedSubtitleReceived(ISubtitleDecoderOutputPtr DecodedSubtitle)
+{
+	SubtitleReceiverDelegate.ExecuteIfBound(DecodedSubtitle);
+}
+
+
+bool FSubtitleDecoder::ResetForTrackChange(FAccessUnit* AccessUnit)
+{
+	bool bWasStarted = bIsStarted;
+	Stop();
+	SubtitleFlushDelegate.ExecuteIfBound();
+	Close();
+	if (AccessUnit && AccessUnit->AUCodecData.IsValid())
+	{
+		PluginDecoder = FSubtitleDecoderPlugins::Get().CreateDecoder(AccessUnit->AUCodecData->ParsedInfo.GetCodecSpecifierRFC6381());
+		if (PluginDecoder.IsValid())
+		{
+			SetCSD(AccessUnit->AUCodecData->RawCSD);
+			SetCodecInfo(AccessUnit->AUCodecData->ParsedInfo);
+			Open();
+			if (bWasStarted)
+			{
+				Start();
+			}
+			return true;
+		}
+	}
+	return false;
+}
+
+
+IAccessUnitBufferInterface::EAUpushResult FSubtitleDecoder::AUdataPushAU(FAccessUnit* AccessUnit)
+{
+	if (PluginDecoder.IsValid())
+	{
+		if (!AccessUnit->bIsDummyData)
+		{
+			if (AccessUnit->bTrackChangeDiscontinuity)
+			{
+				if (!ResetForTrackChange(AccessUnit))
+				{
+					return IAccessUnitBufferInterface::EAUpushResult::Error;
+				}
+			}
+			FParamDict Addtl;
+			PluginDecoder->AddStreamedSubtitleData(TArray<uint8>((const uint8*)AccessUnit->AUData, (int32)AccessUnit->AUSize), AccessUnit->PTS, AccessUnit->Duration, Addtl);
+		}
+		return IAccessUnitBufferInterface::EAUpushResult::Ok;
+	}
+	return IAccessUnitBufferInterface::EAUpushResult::Error;
+}
+
+void FSubtitleDecoder::AUdataPushEOD()
+{
+	if (PluginDecoder.IsValid())
+	{
+		PluginDecoder->SignalStreamedSubtitleEOD();
+	}
+}
+
+void FSubtitleDecoder::AUdataFlushEverything()
+{
+	if (PluginDecoder.IsValid())
+	{
+		PluginDecoder->Flush();
+	}
+	SubtitleFlushDelegate.ExecuteIfBound();
+}
+
+} // namespace Electra
+
diff --git a/Engine/Plugins/Media/ElectraPlayer/Source/ElectraPlayerRuntime/Private/Runtime/Decoder/SubtitleDecoder.h b/Engine/Plugins/Media/ElectraPlayer/Source/ElectraPlayerRuntime/Private/Runtime/Decoder/SubtitleDecoder.h
new file mode 100644
--- /dev/null
+++ b/Engine/Plugins/Media/ElectraPlayer/Source/ElectraPlayerRuntime/Private/Runtime/Decoder/SubtitleDecoder.h
@@ -0,0 +1,67 @@
+// Copyright Epic Games, Inc. All Rights Reserved.
+
+#pragma once
+
+#include "PlayerCore.h"
+
+#include "StreamAccessUnitBuffer.h"
+#include "Decoder/DecoderBase.h"
+
+class ISubtitleDecoderOutput;
+using ISubtitleDecoderOutputPtr = TSharedPtr<ISubtitleDecoderOutput, ESPMode::ThreadSafe>;
+
+
+namespace Electra
+{
+	class IPlayerSessionServices;
+
+
+	/**
+	 * Generic subtitle decoder
+	**/
+	class ISubtitleDecoder : public IAccessUnitBufferInterface
+	{
+	public:
+		static bool IsSupported(const FString& MimeType, const FString& Codec);
+
+		static ISubtitleDecoder* Create(FAccessUnit* AccessUnit);
+
+		virtual ~ISubtitleDecoder() = default;
+
+		//-------------------------------------------------------------------------
+		// State methods
+		//
+		virtual void SetPlayerSessionServices(IPlayerSessionServices* SessionServices) = 0;
+		virtual void Open() = 0;
+		virtual void Close() = 0;
+		virtual void Start() = 0;
+		virtual void Stop() = 0;
+		virtual void UpdatePlaybackPosition(FTimeValue InAbsolutePosition, FTimeValue InLocalPosition) = 0;
+
+
+		//! Returns a time offset at which streamed subtitles should be delivered ahead of their actual timestamp (if they are available earlier)
+		virtual FTimeValue GetStreamedDeliveryTimeOffset() = 0;
+
+
+		//-------------------------------------------------------------------------
+		// Decoded subtitle receiver
+		//
+		DECLARE_DELEGATE_OneParam(FDecodedSubtitleReceivedDelegate, ISubtitleDecoderOutputPtr /* DecodedSubtitle */);
+		virtual FDecodedSubtitleReceivedDelegate& GetDecodedSubtitleReceiveDelegate() = 0;
+		DECLARE_DELEGATE(FDecodedSubtitleFlushDelegate);
+		virtual FDecodedSubtitleFlushDelegate& GetDecodedSubtitleFlushDelegate() = 0;
+
+		//-------------------------------------------------------------------------
+		// Methods from IAccessUnitBufferInterface
+		//
+		//! Attempts to push an access unit to the decoder. Ownership of the access unit is transferred if the push is successful.
+		virtual EAUpushResult AUdataPushAU(FAccessUnit* AccessUnit) = 0;
+		//! Notifies the decoder that there will be no further access units.
+		virtual void AUdataPushEOD() = 0;
+		//! Instructs the decoder to flush all pending input and all already decoded output.
+		virtual void AUdataFlushEverything() = 0;
+	};
+
+
+} // namespace Electra
+
diff --git a/Engine/Plugins/Media/ElectraPlayer/Source/ElectraPlayerRuntime/Private/Runtime/Demuxer/ParserISO14496-12.cpp b/Engine/Plugins/Media/ElectraPlayer/Source/ElectraPlayerRuntime/Private/Runtime/Demuxer/ParserISO14496-12.cpp
--- a/Engine/Plugins/Media/ElectraPlayer/Source/ElectraPlayerRuntime/Private/Runtime/Demuxer/ParserISO14496-12.cpp
+++ b/Engine/Plugins/Media/ElectraPlayer/Source/ElectraPlayerRuntime/Private/Runtime/Demuxer/ParserISO14496-12.cpp
@@ -12,6 +12,7 @@
 #include "Utilities/ISO639-Map.h"
 #include "InfoLog.h"
 #include "Player/PlayerSessionServices.h"
+#include "Decoder/SubtitleDecoder.h"
 
 #if UE_BUILD_DEBUG
 #define MEDIA_DEBUG_HAS_BOX_NAMES 1
@@ -54,6 +55,7 @@ namespace Electra
 		FMP4BoxReader(IParserISO14496_12::IReader* InDataReader, IParserISO14496_12::IBoxCallback* InBoxParseCallback)
 			: DataReader(InDataReader)
 			, BoxParseCallback(InBoxParseCallback)
+			, DataDuplicationBuffer(nullptr)
 		{
 			check(DataReader);
 		}
@@ -79,6 +81,10 @@ namespace Electra
 			int64 NumRead = DataReader->ReadData(&Temp, sizeof(T));
 			if (NumRead == sizeof(T))
 			{
+				if (DataDuplicationBuffer)
+				{
+					DataDuplicationBuffer->Append((const uint8*)&Temp, sizeof(T));
+				}
 				value = Utils::ValueFromBigEndian(Temp);
 				return UEMEDIA_ERROR_OK;
 			}
@@ -136,7 +142,22 @@ namespace Electra
 			{
 				return UEMEDIA_ERROR_OK;
 			}
-			int64 NumRead = DataReader->ReadData(Buffer, NumBytes);
+
+			int64 NumRead = 0;
+			if (!Buffer && DataDuplicationBuffer)
+			{
+				int32 BufferSizeNow = DataDuplicationBuffer->Num();
+				DataDuplicationBuffer->AddUninitialized((int32)NumBytes);
+				NumRead = DataReader->ReadData(DataDuplicationBuffer->GetData() + BufferSizeNow, NumBytes);
+			}
+			else
+			{
+				NumRead = DataReader->ReadData(Buffer, NumBytes);
+				if (NumRead == NumBytes && DataDuplicationBuffer)
+				{
+					DataDuplicationBuffer->Append((const uint8*)Buffer, (int32)NumBytes);
+				}
+			}
 			if (NumRead == NumBytes)
 			{
 				return UEMEDIA_ERROR_OK;
@@ -171,12 +192,36 @@ namespace Electra
 		}
 
 
+		void SetDataDuplicationBuffer(TArray<uint8>* InDataDuplicationBuffer)
+		{
+			DataDuplicationBuffer = InDataDuplicationBuffer;
+		}
+
+		class ScopedDataDuplication : private TMediaNoncopyable<ScopedDataDuplication>
+		{
+		public:
+			explicit ScopedDataDuplication(FMP4BoxReader* InBoxReader, TArray<uint8>& DuplicationBuffer)
+				: BoxReader(InBoxReader)
+			{
+				BoxReader->SetDataDuplicationBuffer(&DuplicationBuffer);
+			}
+			~ScopedDataDuplication()
+			{
+				BoxReader->SetDataDuplicationBuffer(nullptr);
+			}
+		private:
+			ScopedDataDuplication();
+			FMP4BoxReader* BoxReader;
+		};
+
 	private:
 		FMP4BoxReader() = delete;
 		FMP4BoxReader(const FMP4BoxReader&) = delete;
 
-		IParserISO14496_12::IReader* DataReader;
-		IParserISO14496_12::IBoxCallback* BoxParseCallback;
+		IParserISO14496_12::IReader* DataReader = nullptr;
+		IParserISO14496_12::IBoxCallback* BoxParseCallback = nullptr;
+
+		TArray<uint8>* DataDuplicationBuffer = nullptr;
 	};
 
 
@@ -666,6 +711,7 @@ private:
 		static const IParserISO14496_12::FBoxType kSample_mp4a = MAKE_BOX_ATOM('m', 'p', '4', 'a');
 		static const IParserISO14496_12::FBoxType kSample_stpp = MAKE_BOX_ATOM('s', 't', 'p', 'p');
 		static const IParserISO14496_12::FBoxType kSample_sbtt = MAKE_BOX_ATOM('s', 'b', 't', 't');
+		static const IParserISO14496_12::FBoxType kSample_tx3g = MAKE_BOX_ATOM('t', 'x', '3', 'g');
 		static const IParserISO14496_12::FBoxType kSample_enca = MAKE_BOX_ATOM('e', 'n', 'c', 'a');
 		static const IParserISO14496_12::FBoxType kSample_encv = MAKE_BOX_ATOM('e', 'n', 'c', 'v');
 
@@ -1107,6 +1153,27 @@ private:
 			return TrackID;
 		}
 
+		const double* GetMatrix() const
+		{
+			return Matrix;
+		}
+
+		double GetMatrixValue(int32 Index) const
+		{
+			return (Index >= 0 && Index < 9) ? Matrix[Index] : 0.0;
+		}
+
+		uint16 GetWidth() const
+		{
+			return Width >> 16;
+		}
+
+		uint16 GetHeight() const
+		{
+			return Height >> 16;
+		}
+
+
 	private:
 		FMP4BoxTKHD() = delete;
 		FMP4BoxTKHD(const FMP4BoxTKHD&) = delete;
@@ -1144,10 +1211,22 @@ private:
 			RETURN_IF_ERROR(ParseInfo->Reader()->Read(AlternateGroup));			// alternate_group
 			RETURN_IF_ERROR(ParseInfo->Reader()->Read(Value16));				// volume
 			RETURN_IF_ERROR(ParseInfo->Reader()->Read(Value16));				// reserved
-			// Skip matrix
-			for(int32 i = 0; i < 9; ++i)
+			// Matrix
+			for(int32 i=0; i<9; ++i)
 			{
+				/*
+					"All the values in a matrix are stored as 16.16 fixed‐point values, except for u, v and w, which are stored
+					as 2.30 fixed‐point values. The values in the matrix are stored in the order {a,b,u, c,d,v, x,y,w}."
+				*/
 				RETURN_IF_ERROR(ParseInfo->Reader()->Read(Value32));
+				if (i == 2 || i ==5 || i == 8)
+				{
+					Matrix[i] = (int32)Value32 / 1073741824.0;
+				}
+				else
+				{
+					Matrix[i] = (int32)Value32 / 65536.0;
+				}
 			}
 			RETURN_IF_ERROR(ParseInfo->Reader()->Read(Width));					// width (16.16 fixed point). If flags & 8 then indicates aspect ratio.
 			RETURN_IF_ERROR(ParseInfo->Reader()->Read(Height));					// height (16.16 fixed point). If flags & 8 then indicates aspect ratio.
@@ -1167,6 +1246,7 @@ private:
 		uint32		Height;
 		int16		Layer;
 		int16		AlternateGroup;
+		double		Matrix[9];
 	};
 
 
@@ -1935,6 +2015,233 @@ private:
 	};
 
 
+	/**
+	 * Plain Text sample entry (12.5.3 - Sample Entry)
+	 */
+	class FMP4BoxPlainTextSampleEntry : public FMP4BoxSampleEntry
+	{
+	public:
+		FMP4BoxPlainTextSampleEntry(IParserISO14496_12::FBoxType InBoxType, int64 InBoxSize, int64 InStartOffset, int64 InDataOffset, bool bInIsLeafBox, int32 InSTSDBoxVersion)
+			: FMP4BoxSampleEntry(InBoxType, InBoxSize, InStartOffset, InDataOffset, bInIsLeafBox)
+			, STSDBoxVersion(InSTSDBoxVersion)
+		{
+		}
+
+		virtual ~FMP4BoxPlainTextSampleEntry()
+		{
+		}
+
+	private:
+		FMP4BoxPlainTextSampleEntry() = delete;
+		FMP4BoxPlainTextSampleEntry(const FMP4BoxPlainTextSampleEntry&) = delete;
+
+	protected:
+		virtual UEMediaError ReadAndParseAttributes(FMP4ParseInfo* ParseInfo) override
+		{
+			UEMediaError Error = UEMEDIA_ERROR_OK;
+			RETURN_IF_ERROR(FMP4BoxSampleEntry::ReadAndParseAttributes(ParseInfo));
+
+			// There can now be additional boxes following
+			uint32 BytesRemaining = BoxSize - (ParseInfo->Reader()->GetCurrentReadOffset() - StartOffset);
+			if (BytesRemaining)
+			{
+				// Now that there are additional boxes we ourselves are no longer a leaf box.
+				bIsLeafBox = false;
+				// This will read all the boxes in here and add them as children.
+				RETURN_IF_ERROR(ParseInfo->ReadAndParseNextBox(this));
+			}
+
+			return Error;
+		}
+
+		int32		STSDBoxVersion;
+	};
+
+
+	/**
+	 * 3GPP / TX3G Text sample entry (ETSI TS 126 245 V11.0.0 - 5.16 Sample Description Format)
+	 */
+	class FMP4BoxTX3GSampleEntry : public FMP4BoxSampleEntry
+	{
+	public:
+		FMP4BoxTX3GSampleEntry(IParserISO14496_12::FBoxType InBoxType, int64 InBoxSize, int64 InStartOffset, int64 InDataOffset, bool bInIsLeafBox, int32 InSTSDBoxVersion)
+			: FMP4BoxSampleEntry(InBoxType, InBoxSize, InStartOffset, InDataOffset, bInIsLeafBox)
+			, STSDBoxVersion(InSTSDBoxVersion)
+		{
+		}
+
+		virtual ~FMP4BoxTX3GSampleEntry()
+		{
+		}
+
+		struct FBoxRecord
+		{
+			int16 Top = 0;
+			int16 Left = 0;
+			int16 Bottom = 0;
+			int16 Right = 0;
+		};
+
+		struct FStyleRecord
+		{
+			uint16 StartChar = 0;
+			uint16 EndChar = 0;
+			uint16 FontID = 0;
+			uint8 FaceStyleFlags = 0;
+			uint8 FontSize = 0;
+			uint8 TextColorRGBA[4] { 0 };
+		};
+
+		struct FFontRecord
+		{
+			uint16 FontID = 0;
+			FString FontName;
+		};
+
+		const TArray<uint8> GetRawBoxData() const
+		{
+			return BoxData;
+		}
+
+		int32 GetTranslationX() const
+		{
+			return TranslationX;
+		}
+
+		int32 GetTranslationY() const
+		{
+			return TranslationY;
+		}
+
+		uint16 GetWidth() const
+		{
+			return Width;
+		}
+
+		uint16 GetHeight() const
+		{
+			return Height;
+		}
+
+		uint32 GetTimescale() const
+		{
+			return Timescale;
+		}
+
+	private:
+		FMP4BoxTX3GSampleEntry() = delete;
+		FMP4BoxTX3GSampleEntry(const FMP4BoxTX3GSampleEntry&) = delete;
+
+	protected:
+		virtual UEMediaError ReadAndParseAttributes(FMP4ParseInfo* ParseInfo) override
+		{
+			UEMediaError Error = UEMEDIA_ERROR_OK;
+
+			// Duplicate all the data we are reading into a buffer that we need to provide to the subtitle plugin.
+			FMP4BoxReader::ScopedDataDuplication DuplicateBoxData(ParseInfo->Reader(), BoxData);
+
+			RETURN_IF_ERROR(FMP4BoxSampleEntry::ReadAndParseAttributes(ParseInfo));
+
+			const FMP4BoxTKHD* TKHD = ParseInfo->GetCurrentTrackBox() ? static_cast<const FMP4BoxTKHD*>(ParseInfo->GetCurrentTrackBox()->GetBoxPath(FMP4Box::kBox_tkhd)) : nullptr;
+			if (TKHD)
+			{
+				Width = TKHD->GetWidth();
+				Height = TKHD->GetHeight();
+				// Translation must not be using fractional digits and can thus be treated like integers.
+				TranslationX = (int32) TKHD->GetMatrixValue(6);
+				TranslationY = (int32) TKHD->GetMatrixValue(7);
+			}
+
+			const FMP4BoxMDHD* MDHD = ParseInfo->GetCurrentTrackBox() ? static_cast<const FMP4BoxMDHD*>(ParseInfo->GetCurrentTrackBox()->FindBox(FMP4Box::kBox_mdhd)) : nullptr;
+			if (MDHD)
+			{
+				Timescale = MDHD->GetTimescale();
+			}
+
+			RETURN_IF_ERROR(ParseInfo->Reader()->Read(DisplayFlags));
+			RETURN_IF_ERROR(ParseInfo->Reader()->Read(HorizontalJustification));
+			RETURN_IF_ERROR(ParseInfo->Reader()->Read(VerticalJustification));
+			RETURN_IF_ERROR(ParseInfo->Reader()->ReadBytes(BackgroundColorRGBA, 4));
+
+			// BoxRecord
+			RETURN_IF_ERROR(ParseInfo->Reader()->Read(BoxRecord.Top));
+			RETURN_IF_ERROR(ParseInfo->Reader()->Read(BoxRecord.Left));
+			RETURN_IF_ERROR(ParseInfo->Reader()->Read(BoxRecord.Bottom));
+			RETURN_IF_ERROR(ParseInfo->Reader()->Read(BoxRecord.Right));
+
+			// Style record
+			RETURN_IF_ERROR(ParseInfo->Reader()->Read(StyleRecord.StartChar));
+			RETURN_IF_ERROR(ParseInfo->Reader()->Read(StyleRecord.EndChar));
+			RETURN_IF_ERROR(ParseInfo->Reader()->Read(StyleRecord.FontID));
+			RETURN_IF_ERROR(ParseInfo->Reader()->Read(StyleRecord.FaceStyleFlags));
+			RETURN_IF_ERROR(ParseInfo->Reader()->Read(StyleRecord.FontSize));
+			RETURN_IF_ERROR(ParseInfo->Reader()->ReadBytes(StyleRecord.TextColorRGBA, 4));
+
+			uint32 BytesRemaining = BoxSize - (ParseInfo->Reader()->GetCurrentReadOffset() - StartOffset);
+			// The font table is stored in a box. Not sure if there can be more than one though.
+			// We skip over the box itself and parse it in place. Any additional boxes we ignore.
+			if (BytesRemaining < 10)
+			{
+				UE_LOG(LogElectraMP4Parser, Error, TEXT("ERROR: ElectraPlayer/FMP4BoxTX3GSampleEntry::ReadAndParseAttributes(): Not enough data for the 'ftab' box"));
+				return UEMEDIA_ERROR_FORMAT_ERROR;
+			}
+			RETURN_IF_ERROR(ParseInfo->Reader()->ReadBytes(nullptr, 8));
+
+			// FontTableBox
+			uint16 EntryCount;
+			RETURN_IF_ERROR(ParseInfo->Reader()->Read(EntryCount));
+			FontRecords.Reserve(EntryCount);
+			for(int32 nEntry=0; nEntry<EntryCount; ++nEntry)
+			{
+				FFontRecord& fr = FontRecords.AddDefaulted_GetRef();
+				RETURN_IF_ERROR(ParseInfo->Reader()->Read(fr.FontID));
+				uint8 FontNameLength;
+				RETURN_IF_ERROR(ParseInfo->Reader()->Read(FontNameLength));
+				if (FontNameLength)
+				{
+					uint8 FontName[256] {0};
+					RETURN_IF_ERROR(ParseInfo->Reader()->ReadBytes(FontName, FontNameLength));
+					if (FontNameLength >= 2 && ((FontName[0] == 0xff && FontName[1] == 0xfe) || (FontName[0] == 0xfe && FontName[1] == 0xff)))
+					{
+						UE_LOG(LogElectraMP4Parser, Error, TEXT("ERROR: ElectraPlayer/FMP4BoxTX3GSampleEntry::ReadAndParseAttributes(): Font name uses UTF16 which is not supported"));
+						return UEMEDIA_ERROR_FORMAT_ERROR;
+					}
+					FUTF8ToTCHAR cnv((const ANSICHAR*)FontName, FontNameLength);
+					fr.FontName = FString(cnv.Length(), cnv.Get());
+				}
+			}
+
+			// There could now be additional boxes following
+			BytesRemaining = BoxSize - (ParseInfo->Reader()->GetCurrentReadOffset() - StartOffset);
+			if (BytesRemaining)
+			{
+				// Now that there are additional boxes we ourselves are no longer a leaf box.
+				bIsLeafBox = false;
+				// This will read all the boxes in here and add them as children.
+				RETURN_IF_ERROR(ParseInfo->ReadAndParseNextBox(this));
+			}
+
+			return Error;
+		}
+
+		TArray<uint8>		BoxData;
+
+		int32				STSDBoxVersion;
+		uint32				Timescale = 0;
+		uint16				Width = 0;
+		uint16				Height = 0;
+		int32				TranslationX = 0;
+		int32				TranslationY = 0;
+		uint32				DisplayFlags = 0;
+		int8				HorizontalJustification = 0;
+		int8				VerticalJustification = 0;
+		uint8				BackgroundColorRGBA[4] {0};
+		FBoxRecord			BoxRecord;
+		FStyleRecord		StyleRecord;
+		TArray<FFontRecord>	FontRecords;
+	};
+
+
 	/**
 	 * AVC Decoder Configuration box (ISO/IEC 14496-15:2014 - 5.4.2.1.2)
 	 */
@@ -2489,6 +2796,7 @@ private:
 			{
 				ESTVideo,
 				ESTAudio,
+				ESTSubtitles,
 				ESTIgnored
 			};
 			EExpectedSampleType ExpectedSampleType = EExpectedSampleType::ESTIgnored;
@@ -2505,6 +2813,8 @@ private:
 						break;
 					case FMP4Box::kBox_nmhd:
 					case FMP4Box::kBox_sthd:
+						ExpectedSampleType = EExpectedSampleType::ESTSubtitles;
+						break;
 					case FMP4Box::kBox_gmhd:
 					default:
 						ExpectedSampleType = EExpectedSampleType::ESTIgnored;
@@ -2531,7 +2841,7 @@ private:
 			// Check if this is a supported sample format.
 			bIsSupportedFormat = ExpectedSampleType != EExpectedSampleType::ESTIgnored;
 
-			for(uint32 i = 0; i < EntryCount; ++i)
+			for(uint32 i=0; i<EntryCount; ++i)
 			{
 				IParserISO14496_12::FBoxType	ChildBoxType;
 				int64							ChildBoxSize;
@@ -2553,6 +2863,16 @@ private:
 					case EExpectedSampleType::ESTAudio:
 						NextBox = new FMP4BoxAudioSampleEntry(ChildBoxType, ChildBoxSize, BoxStartOffset, BoxDataOffset, true, Version);
 						break;
+					case EExpectedSampleType::ESTSubtitles:
+						if (ChildBoxType == FMP4Box::kSample_tx3g)
+						{
+							NextBox = new FMP4BoxTX3GSampleEntry(ChildBoxType, ChildBoxSize, BoxStartOffset, BoxDataOffset, true, Version);
+						}
+						else
+						{
+							NextBox = new FMP4BoxIgnored(ChildBoxType, ChildBoxSize, BoxStartOffset, BoxDataOffset, true);
+						}
+						break;
 						/*
 						// Subtitles, captions and subpictures (bitmap subtitles) are ignored for now.
 						case FMP4Box::kTrack_subt:
@@ -4871,6 +5191,7 @@ private:
 			MPEG::FAVCDecoderConfigurationRecord	CodecSpecificDataAVC;
 			MPEG::FHEVCDecoderConfigurationRecord	CodecSpecificDataHEVC;
 			MPEG::FESDescriptor						CodecSpecificDataMP4A;
+			TArray<uint8>							CodecSpecificDataRAW;
 			FBitrateInfo							BitrateInfo;
 		};
 
@@ -4943,6 +5264,7 @@ private:
 		UEMediaError ParseAVC1SampleType(FTrack* Track, const FMP4Box* SampleBox);
 		UEMediaError ParseHVC1SampleType(FTrack* Track, const FMP4Box* SampleBox);
 		UEMediaError ParseMP4ASampleType(FTrack* Track, const FMP4Box* SampleBox);
+		UEMediaError ParseTX3GSampleType(FTrack* Track, const FMP4Box* SampleBox);
 
 		FMP4ParseInfo* ParsedData;
 
@@ -5041,6 +5363,10 @@ private:
 			{
 				return CodecSpecificDataMP4A.GetCodecSpecificData();
 			}
+			case FStreamCodecInformation::ECodec::TX3G:
+			{
+				return CodecSpecificDataRAW;
+			}
 			default:
 			{
 				break;
@@ -5066,6 +5392,10 @@ private:
 			{
 				return CodecSpecificDataMP4A.GetRawData();
 			}
+			case FStreamCodecInformation::ECodec::TX3G:
+			{
+				return CodecSpecificDataRAW;
+			}
 			default:
 			{
 				break;
@@ -6051,6 +6381,26 @@ private:
 		return UEMEDIA_ERROR_FORMAT_ERROR;
 	}
 
+	UEMediaError FParserISO14496_12::ParseTX3GSampleType(FTrack* Track, const FMP4Box* SampleBox)
+	{
+		// Need to check if there is a subtitle decoder capable of decoding this.
+		if (ISubtitleDecoder::IsSupported(TEXT(""), TEXT("tx3g")))
+		{
+			const FMP4BoxTX3GSampleEntry* SubtitleSampleEntry = static_cast<const FMP4BoxTX3GSampleEntry*>(Track->STSDBox->GetChildBox(0));
+
+			Track->CodecSpecificDataRAW = SubtitleSampleEntry->GetRawBoxData();
+			Track->CodecInformation.SetStreamType(EStreamType::Subtitle);
+			Track->CodecInformation.SetCodec(FStreamCodecInformation::ECodec::TX3G);
+			Track->CodecInformation.SetStreamLanguageCode(Track->GetLanguage());
+			Track->CodecInformation.SetCodecSpecifierRFC6381(TEXT("tx3g"));
+			Track->CodecInformation.SetResolution(FStreamCodecInformation::FResolution(SubtitleSampleEntry->GetWidth(), SubtitleSampleEntry->GetHeight()));
+			Track->CodecInformation.SetTranslation(FStreamCodecInformation::FTranslation(SubtitleSampleEntry->GetTranslationX(), SubtitleSampleEntry->GetTranslationY()));
+			Track->CodecInformation.SetFrameRate(FTimeFraction(1, SubtitleSampleEntry->GetTimescale()));
+			return UEMEDIA_ERROR_OK;
+		}
+		return UEMEDIA_ERROR_NOT_SUPPORTED;
+	}
+
 	UEMediaError FParserISO14496_12::PrepareTracks(TSharedPtrTS<const IParserISO14496_12> OptionalMP4InitSegment)
 	{
 		delete ParsedTrackInfo;
@@ -6370,6 +6720,19 @@ private:
 							{
 								break;
 							}
+							case FMP4Box::kSample_tx3g:
+							{
+								UEMediaError Error = ParseTX3GSampleType(Track.Get(), STSDFirstChildBox);
+								if (Error == UEMEDIA_ERROR_NOT_SUPPORTED)
+								{
+									bIsSupported = false;
+								}
+								else if (Error != UEMEDIA_ERROR_OK)
+								{
+									return Error;
+								}
+								break;
+							}
 							default:
 							{
 								break;
diff --git a/Engine/Plugins/Media/ElectraPlayer/Source/ElectraPlayerRuntime/Private/Runtime/Player/AdaptiveStreamingPlayer.cpp b/Engine/Plugins/Media/ElectraPlayer/Source/ElectraPlayerRuntime/Private/Runtime/Player/AdaptiveStreamingPlayer.cpp
--- a/Engine/Plugins/Media/ElectraPlayer/Source/ElectraPlayerRuntime/Private/Runtime/Player/AdaptiveStreamingPlayer.cpp
+++ b/Engine/Plugins/Media/ElectraPlayer/Source/ElectraPlayerRuntime/Private/Runtime/Player/AdaptiveStreamingPlayer.cpp
@@ -101,6 +101,10 @@ FAdaptiveStreamingPlayer::FAdaptiveStreamingPlayer(const IAdaptiveStreamingPlaye
 	DataAvailabilityStateTxt.StreamType = EStreamType::Subtitle;
 	DataAvailabilityStateTxt.Availability = Metrics::FDataAvailabilityChange::EAvailability::DataNotAvailable;
 
+	// Subtitles need to be explicitly selected by the user.
+	SelectedStreamAttributesTxt.Deselect();
+
+
 	// Create the UTC time handler for this player session
 	SynchronizedUTCTime = ISynchronizedUTCTime::Create();
 	SynchronizedUTCTime->SetTime(MEDIAutcTime::Current());
@@ -359,6 +363,53 @@ void FAdaptiveStreamingPlayer::RemoveAEMSReceiver(TWeakPtrTS<IAdaptiveStreamingP
 }
 
 
+void FAdaptiveStreamingPlayer::AddSubtitleReceiver(TWeakPtrTS<IAdaptiveStreamingPlayerSubtitleReceiver> InReceiver)
+{
+	if (InReceiver.IsValid())
+	{
+		FScopeLock lock(&SubtitleReceiversCriticalSection);
+		SubtitleReceivers.Add(InReceiver);
+	}
+}
+
+void FAdaptiveStreamingPlayer::RemoveSubtitleReceiver(TWeakPtrTS<IAdaptiveStreamingPlayerSubtitleReceiver> InReceiver)
+{
+	if (InReceiver.IsValid())
+	{
+		FScopeLock lock(&SubtitleReceiversCriticalSection);
+		SubtitleReceivers.Remove(InReceiver);
+	}
+}
+
+void FAdaptiveStreamingPlayer::OnDecodedSubtitleReceived(ISubtitleDecoderOutputPtr Subtitle)
+{
+	FScopeLock lock(&SubtitleReceiversCriticalSection);
+	for(auto &Receiver : SubtitleReceivers)
+	{
+		TSharedPtrTS<IAdaptiveStreamingPlayerSubtitleReceiver> Recv = Receiver.Pin();
+		if (Recv.IsValid())
+		{
+			Recv->OnMediaPlayerSubtitleReceived(Subtitle);
+		}
+	}
+}
+
+void FAdaptiveStreamingPlayer::OnFlushSubtitleReceivers()
+{
+	FScopeLock lock(&SubtitleReceiversCriticalSection);
+	for(auto &Receiver : SubtitleReceivers)
+	{
+		TSharedPtrTS<IAdaptiveStreamingPlayerSubtitleReceiver> Recv = Receiver.Pin();
+		if (Recv.IsValid())
+		{
+			Recv->OnMediaPlayerFlushSubtitles();
+		}
+	}
+}
+
+
+
+
 //-----------------------------------------------------------------------------
 /**
  * Dispatches an event to the listeners.
@@ -648,14 +699,18 @@ void FAdaptiveStreamingPlayer::GetLoopState(FPlayerLoopState& OutLoopState) cons
  */
 void FAdaptiveStreamingPlayer::GetTrackMetadata(TArray<FTrackMetadata>& OutTrackMetadata, EStreamType StreamType) const
 {
-	TArray<FTrackMetadata> dummy;
+	TArray<FTrackMetadata> dummy1, dummy2;
 	if (StreamType == EStreamType::Video)
 	{
-		PlaybackState.GetTrackMetadata(OutTrackMetadata, dummy);
+		PlaybackState.GetTrackMetadata(OutTrackMetadata, dummy1, dummy2);
 	}
 	else if (StreamType == EStreamType::Audio)
 	{
-		PlaybackState.GetTrackMetadata(dummy, OutTrackMetadata);
+		PlaybackState.GetTrackMetadata(dummy1, OutTrackMetadata, dummy2);
+	}
+	else if (StreamType == EStreamType::Subtitle)
+	{
+		PlaybackState.GetTrackMetadata(dummy1, dummy2, OutTrackMetadata);
 	}
 	else
 	{
@@ -954,150 +1009,6 @@ void FAdaptiveStreamingPlayer::GetStreamBufferStats(FAccessUnitBufferInfo& OutBu
 	}
 }
 
-IPlayerStreamFilter* FAdaptiveStreamingPlayer::GetStreamFilter()
-{
-	return this;
-}
-
-const FCodecSelectionPriorities& FAdaptiveStreamingPlayer::GetCodecSelectionPriorities(EStreamType ForStream)
-{
-	switch(ForStream)
-	{
-		case EStreamType::Video:
-		{
-			return CodecPrioritiesVideo;
-		}
-		case EStreamType::Audio:
-		{
-			return CodecPrioritiesAudio;
-		}
-		case EStreamType::Subtitle:
-		{
-			return CodecPrioritiesSubtitles;
-		}
-		default:
-		{
-			static FCodecSelectionPriorities Dummy;
-			return Dummy;
-		}
-	}
-}
-
-bool FAdaptiveStreamingPlayer::CanDecodeStream(const FStreamCodecInformation& InStreamCodecInfo) const
-{
-	auto IsCodecPrefixExcluded = [](const FString& InCodec, const TArray<FString>& InExcludedPrefixes) -> bool
-	{
-		for(auto &Prefix : InExcludedPrefixes)
-		{
-			if (InCodec.StartsWith(Prefix, ESearchCase::IgnoreCase))
-			{
-				return true;
-			}
-		}
-		return false;
-	};
-
-
-	if (InStreamCodecInfo.IsVideoCodec())
-	{
-		// Check against user specified codec prefix exclusions
-		if (IsCodecPrefixExcluded(InStreamCodecInfo.GetCodecSpecifierRFC6381(), ExcludedVideoDecoderPrefixes))
-		{
-			return false;
-		}
-
-		double Rate = InStreamCodecInfo.GetFrameRate().IsValid() ? InStreamCodecInfo.GetFrameRate().GetAsDouble() : 30.0;
-		if (InStreamCodecInfo.GetCodec() == FStreamCodecInformation::ECodec::H264)
-		{
-			const AdaptiveStreamingPlayerConfig::FConfiguration::FVideoDecoderLimit* DecoderLimit = &PlayerConfig.H264LimitUpto30fps;
-			if (Rate > 31.0)
-			{
-				DecoderLimit = &PlayerConfig.H264LimitAbove30fps;
-			}
-			// Check against user configured resolution limit
-			if (DecoderLimit->MaxResolution.Height && InStreamCodecInfo.GetResolution().Height > DecoderLimit->MaxResolution.Height)
-			{
-				return false;
-			}
-
-			// Check against video decoder capabilities.
-			IVideoDecoderH264::FStreamDecodeCapability StreamParam, Capability;
-			StreamParam.Width = InStreamCodecInfo.GetResolution().Width;
-			StreamParam.Height = InStreamCodecInfo.GetResolution().Height;
-			StreamParam.Profile = InStreamCodecInfo.GetProfile();
-			StreamParam.Level = InStreamCodecInfo.GetProfileLevel();
-			StreamParam.FPS = Rate;
-			if (IVideoDecoderH264::GetStreamDecodeCapability(Capability, StreamParam))
-			{
-				if (Capability.DecoderSupportType == IVideoDecoderH264::FStreamDecodeCapability::ESupported::NotSupported)
-				{
-					return false;
-				}
-			}
-		}
-		else if (InStreamCodecInfo.GetCodec() == FStreamCodecInformation::ECodec::H265)
-		{
-			#if ELECTRA_PLATFORM_HAS_H265_DECODER
-				// Check against video decoder capabilities.
-				IVideoDecoderH265::FStreamDecodeCapability StreamParam, Capability;
-				StreamParam.Width = InStreamCodecInfo.GetResolution().Width;
-				StreamParam.Height = InStreamCodecInfo.GetResolution().Height;
-				StreamParam.Tier = InStreamCodecInfo.GetProfileTier();
-				StreamParam.CompatibilityFlags = InStreamCodecInfo.GetProfileCompatibilityFlags();
-				StreamParam.Profile = InStreamCodecInfo.GetProfile();
-				StreamParam.Level = InStreamCodecInfo.GetProfileLevel();
-				StreamParam.ConstraintFlags = InStreamCodecInfo.GetProfileConstraints();
-				StreamParam.FPS = Rate;
-				if (IVideoDecoderH265::GetStreamDecodeCapability(Capability, StreamParam))
-				{
-					if (Capability.DecoderSupportType == IVideoDecoderH265::FStreamDecodeCapability::ESupported::NotSupported)
-					{
-						return false;
-					}
-				}
-			#else
-				return false;
-			#endif
-		}
-	}
-	else if (InStreamCodecInfo.IsAudioCodec())
-	{
-		// Check against user specified codec prefix exclusions
-		if (IsCodecPrefixExcluded(InStreamCodecInfo.GetCodecSpecifierRFC6381(), ExcludedAudioDecoderPrefixes))
-		{
-			return false;
-		}
-	}
-	else if (InStreamCodecInfo.IsSubtitleCodec())
-	{
-		// Check against user specified codec prefix exclusions
-		if (IsCodecPrefixExcluded(InStreamCodecInfo.GetCodecSpecifierRFC6381(), ExcludedSubtitleDecoderPrefixes))
-		{
-			return false;
-		}
-	}
-
-	return true;
-}
-
-
-
-
-// AU memory
-void* FAdaptiveStreamingPlayer::AUAllocate(IAccessUnitMemoryProvider::EDataType type, SIZE_T NumBytes, SIZE_T Alignment)
-{
-	if (Alignment)
-	{
-		return FMemory::Malloc(NumBytes, Alignment);
-	}
-	return FMemory::Malloc(NumBytes);
-}
-
-void FAdaptiveStreamingPlayer::AUDeallocate(IAccessUnitMemoryProvider::EDataType type, void* Address)
-{
-	FMemory::Free(Address);
-}
-
 
 //-----------------------------------------------------------------------------
 /**
@@ -1256,6 +1167,8 @@ void FAdaptiveStreamingPlayer::InternalHandleOnce()
 		HandleMetadataChanges();
 		// Handle AEMS events
 		HandleAEMSEvents();
+		// Handle subtitle decoder
+		HandleSubtitleDecoder();
 		// Handle buffer level changes
 		HandleNewBufferedData();
 		// Handle new output data (finish prerolling)
@@ -1387,6 +1300,11 @@ bool FAdaptiveStreamingPlayer::InternalHandleThreadMessages()
 				{
 					msg.Data.InitialStreamAttribute.InitialSelection.Language_ISO639 = ISO639::MapTo639_1(msg.Data.InitialStreamAttribute.InitialSelection.Language_ISO639.GetValue());
 				}
+				bool bInitiallyDisabled = msg.Data.InitialStreamAttribute.InitialSelection.OverrideIndex.IsSet() && msg.Data.InitialStreamAttribute.InitialSelection.OverrideIndex.GetValue() < 0;
+				if (bInitiallyDisabled)
+				{
+					msg.Data.InitialStreamAttribute.InitialSelection.ClearOverrideIndex();
+				}
 				switch(msg.Data.InitialStreamAttribute.StreamType)
 				{
 					case EStreamType::Video:
@@ -1401,6 +1319,10 @@ bool FAdaptiveStreamingPlayer::InternalHandleThreadMessages()
 					default:
 						break;
 				}
+				if (bInitiallyDisabled)
+				{
+					InternalDeselectStream(msg.Data.InitialStreamAttribute.StreamType);
+				}
 				break;
 			}
 
@@ -1427,6 +1349,8 @@ bool FAdaptiveStreamingPlayer::InternalHandleThreadMessages()
 							PendingTrackSelectionTxt = MakeSharedTS<FStreamSelectionAttributes>(msg.Data.TrackSelection.TrackAttributes);
 							MultiStreamBufferTxt.Activate();
 							break;
+						default:
+							break;
 					}
 				}
 				break;
@@ -1440,18 +1364,7 @@ bool FAdaptiveStreamingPlayer::InternalHandleThreadMessages()
 
 			case FWorkerThreadMessages::FMessage::EType::DeselectTrack:
 			{
-				switch(msg.Data.TrackSelection.StreamType)
-				{
-					case EStreamType::Video:
-						MultiStreamBufferVid.Deselect();
-						break;
-					case EStreamType::Audio:
-						MultiStreamBufferAud.Deselect();
-						break;
-					case EStreamType::Subtitle:
-						MultiStreamBufferTxt.Deselect();
-						break;
-				}
+				InternalDeselectStream(msg.Data.TrackSelection.StreamType);
 				break;
 			}
 
@@ -1618,12 +1531,12 @@ void FAdaptiveStreamingPlayer::HandleMetadataChanges()
 					// Handle it.
 					TArray<FTrackMetadata> MetadataVideo;
 					TArray<FTrackMetadata> MetadataAudio;
-					//TArray<FTrackMetadata> MetadataSubtitle;
+					TArray<FTrackMetadata> MetadataSubtitle;
 					UpcomingPeriods[i].Period->GetMetaData(MetadataVideo, EStreamType::Video);
 					UpcomingPeriods[i].Period->GetMetaData(MetadataAudio, EStreamType::Audio);
-					//UpcomingPeriods[i].Period->GetMetaData(MetadataSubtitle, EStreamType::Subtitle);
+					UpcomingPeriods[i].Period->GetMetaData(MetadataSubtitle, EStreamType::Subtitle);
 
-					bool bMetadataChanged = PlaybackState.SetTrackMetadata(MetadataVideo, MetadataAudio);
+					bool bMetadataChanged = PlaybackState.SetTrackMetadata(MetadataVideo, MetadataAudio, MetadataSubtitle);
 					PlaybackState.SetHaveMetadata(true);
 					if (bMetadataChanged)
 					{
@@ -1850,7 +1763,8 @@ void FAdaptiveStreamingPlayer::HandleNewBufferedData()
 			if (ManifestType == EMediaFormatType::ISOBMFF)
 			{
 				if (VideoBufferStats.StreamBuffer.bLastPushWasBlocked ||
-					AudioBufferStats.StreamBuffer.bLastPushWasBlocked)
+					AudioBufferStats.StreamBuffer.bLastPushWasBlocked ||
+					TextBufferStats.StreamBuffer.bLastPushWasBlocked)
 				{
 					bHaveEnoughVideo = bHaveEnoughAudio = true;
 				}
@@ -1889,18 +1803,43 @@ void FAdaptiveStreamingPlayer::HandleNewBufferedData()
  */
 FTimeValue FAdaptiveStreamingPlayer::GetCurrentPlayTime()
 {
-	FTimeValue Current = PlaybackState.GetPlayPosition();
+	FTimeValue Current;
 	FTimeValue v;
-	v = MultiStreamBufferVid.GetLastPoppedPTS();
-	if (v.IsValid() && v < Current)
+	v = MultiStreamBufferVid.GetLastPoppedDTS();
+	if (v.IsValid() && !Current.IsValid())
 	{
 		Current = v;
 	}
-	v = MultiStreamBufferAud.GetLastPoppedPTS();
-	if (v.IsValid() && v < Current)
+	v = MultiStreamBufferAud.GetLastPoppedDTS();
+	if (v.IsValid() && !Current.IsValid())
 	{
 		Current = v;
 	}
+	if (!Current.IsValid())
+	{
+		Current = PlaybackState.GetPlayPosition();
+	}
+
+	// Check for clock drift
+	// FIXME: This is until such time when the application gives us feedback at which time it is currently
+	//        rendering video or audio. We have no control here over how the application handles the decoded
+	//        output or how deep its pipeline is. Our own timestamps for decoding are ahead of the time
+	//        the output is rendered, but we do not know by how much.
+	if (Current.IsValid() && RenderClock->IsRunning())
+	{
+		FTimeValue clockPos = RenderClock->GetInterpolatedRenderTime(IMediaRenderClock::ERendererType::Video);
+		if (clockPos.IsValid())
+		{
+			static const FTimeValue ClockDriftThreshold(FTimeValue::MillisecondsToHNS(200));
+			if (clockPos > Current ||
+				clockPos + ClockDriftThreshold < Current)
+			{
+				RenderClock->SetCurrentTime(IMediaRenderClock::ERendererType::Video, Current);
+				RenderClock->SetCurrentTime(IMediaRenderClock::ERendererType::Audio, Current);
+				PlaybackState.SetPlayPosition(Current);
+			}
+		}
+	}
 	return Current;
 }
 
@@ -2372,6 +2311,12 @@ void FAdaptiveStreamingPlayer::InternalHandlePendingFirstSegmentRequest(const FT
 			{
 				FTimeValue CurrentPlayPos = PlaybackState.GetPlayPosition();
 				FTimeValue NewPlayPos = PendingFirstSegmentRequest->GetFirstPTS();
+				// If the PTS is offset by an edit list it could be negative. We do not want to
+				// use negative values so we set it to zero instead. This is not a problem.
+				if (NewPlayPos < FTimeValue::GetZero())
+				{
+					NewPlayPos.SetToZero();
+				}
 				PlaybackState.SetPlayPosition(NewPlayPos);
 				PlaybackState.SetLoopState(CurrentLoopState);
 
@@ -2393,11 +2338,11 @@ void FAdaptiveStreamingPlayer::InternalHandlePendingFirstSegmentRequest(const FT
 					case EStreamType::Subtitle:	bHaveTextReader.Set(true);	break;
 					default:												break;
 				}
-				TArray<IStreamSegment::FDependentStreams> DependentStreams;
+				TArray<TSharedPtrTS<IStreamSegment>> DependentStreams;
 				PendingFirstSegmentRequest->GetDependentStreams(DependentStreams);
 				for(int32 i=0; i<DependentStreams.Num(); ++i)
 				{
-					switch(DependentStreams[i].StreamType)
+					switch(DependentStreams[i]->GetType())
 					{
 						case EStreamType::Video:	bHaveVideoReader.Set(true);	break;
 						case EStreamType::Audio:	bHaveAudioReader.Set(true);	break;
@@ -2445,15 +2390,13 @@ void FAdaptiveStreamingPlayer::InternalHandlePendingFirstSegmentRequest(const FT
 				CompletedSegmentRequests.Add(AlreadyEndedStreams[i]->GetType(), AlreadyEndedStreams[i]);
 			}
 
-			// Kick off the request. Use the sequence ID from video. All sequence IDs have been set equal for start requests so it does not matter.
-			IStreamReader::EAddResult ReqResult = StreamReaderHandler->AddRequest(CurrentPlaybackSequenceID[0], PendingFirstSegmentRequest);
-			if (ReqResult != IStreamReader::EAddResult::Added)
+			// Get the requested streams from the initial request (which may be a wrapper for several individual streams)
+			// and add them to the request queue.
+			TArray<TSharedPtrTS<IStreamSegment>> RequestedStreams;
+			PendingFirstSegmentRequest->GetRequestedStreams(RequestedStreams);
+			for(auto &RequestedStream : RequestedStreams)
 			{
-				FErrorDetail err;
-				err.SetFacility(Facility::EFacility::Player);
-				err.SetMessage("Failed to add stream segment request to reader");
-				err.SetCode(INTERR_FRAGMENT_READER_REQUEST);
-				PostError(err);
+				ReadyWaitingSegmentRequests.Enqueue(RequestedStream);
 			}
 		}
 		PendingFirstSegmentRequest.Reset();
@@ -2463,7 +2406,6 @@ void FAdaptiveStreamingPlayer::InternalHandlePendingFirstSegmentRequest(const FT
 
 void FAdaptiveStreamingPlayer::InternalHandleCompletedSegmentRequests(const FTimeValue& CurrentTime)
 {
-
 	// Check the streams that reached EOS. If all of them are done decide what to do, end playback or loop.
 	if (CompletedSegmentRequests.Num())
 	{
@@ -2768,15 +2710,7 @@ void FAdaptiveStreamingPlayer::HandlePendingMediaSegmentRequests()
 					}
 
 					NextSegment->SetExecutionDelay(ActionDelay);
-					IStreamReader::EAddResult ReqResult = StreamReaderHandler->AddRequest(CurrentPlaybackSequenceID[StreamTypeToArrayIndex(StreamType)], NextSegment);
-					if (ReqResult != IStreamReader::EAddResult::Added)
-					{
-						FErrorDetail err;
-						err.SetFacility(Facility::EFacility::Player);
-						err.SetMessage("Failed to add stream segment request to reader");
-						err.SetCode(INTERR_FRAGMENT_READER_REQUEST);
-						PostError(err);
-					}
+					ReadyWaitingSegmentRequests.Enqueue(NextSegment);
 					break;
 				}
 				case IManifest::FResult::EType::PastEOS:
@@ -2807,14 +2741,19 @@ void FAdaptiveStreamingPlayer::HandlePendingMediaSegmentRequests()
 							// Get the dependent stream types, if any.
 							// This is mainly for multiplexed streams like an .mp4 at this point.
 							// All stream types have reached EOS now.
-							TArray<IStreamSegment::FDependentStreams> DependentStreams;
-							FinishedReq.Request->GetDependentStreams(DependentStreams);
-							// Add the primary stream type to the list as well.
-							DependentStreams.AddDefaulted_GetRef().StreamType = FinishedReq.Request->GetType();
-							for(int32 i=0; i<DependentStreams.Num(); ++i)
+							// NOTE: If is possible that there is no actual request here. This happens when a stream
+							// gets enabled (like a subtitle stream) but is immediately at EOS and thus has no actual segment request!
+							if (FinishedReq.Request.IsValid())
 							{
-								CompletedSegmentRequests.Add(DependentStreams[i].StreamType, FinishedReq.Request);
+								TArray<TSharedPtrTS<IStreamSegment>> DependentStreams;
+								FinishedReq.Request->GetDependentStreams(DependentStreams);
+								for(int32 i=0; i<DependentStreams.Num(); ++i)
+								{
+									CompletedSegmentRequests.Add(DependentStreams[i]->GetType(), FinishedReq.Request);
+								}
 							}
+							// Add the primary stream type to the list as well.
+							CompletedSegmentRequests.Add(StreamType, FinishedReq.Request);
 							break;
 						}
 					}
@@ -2836,9 +2775,115 @@ void FAdaptiveStreamingPlayer::HandlePendingMediaSegmentRequests()
 		}
 	}
 
+
+	// Go through the queue of ready-to-fetch segments to see if they should be fetched now.
+	TQueue<TSharedPtrTS<IStreamSegment>> readyRequests;
+	Swap(readyRequests, ReadyWaitingSegmentRequests);
+	FTimeValue UTCNow = GetSynchronizedUTCTime()->GetTime();
+	while(!readyRequests.IsEmpty())
+	{
+		TSharedPtrTS<IStreamSegment> ReadyReq;
+		readyRequests.Dequeue(ReadyReq);
+
+		bool bExecuteNow = true;
+		bool bSkipOver = false;
+
+		// Does this request have an availability window?
+		FTimeValue When = ReadyReq->GetExecuteAtUTCTime();
+		if (When.IsValid() && When > UTCNow)
+		{
+			bExecuteNow = false;
+		}
+
+		// Subtitles need to be throttled based on play position.
+		// Due them being sparse and small they would essentially all get fetched right away.
+		if (ReadyReq->GetType() == EStreamType::Subtitle)
+		{
+			// Fetch up to 20 seconds ahead of time.
+			const FTimeValue AdvanceFetchThreshold(FTimeValue::MillisecondsToHNS(1000 * 20));
+			FTimeValue NextPTS = ReadyReq->GetFirstPTS();
+			FTimeValue CurrentPos = GetPlayPosition();
+			if (CurrentPos.IsValid() && NextPTS.IsValid() && CurrentPos + AdvanceFetchThreshold < NextPTS)
+			{
+				bExecuteNow = false;
+			}
+			// When subtitles are disabled we do not actually fetch the segment.
+			bSkipOver = SelectedStreamAttributesTxt.IsDeselected();
+		}
+
+		if (bExecuteNow)
+		{
+			if (!bSkipOver)
+			{
+				IStreamReader::EAddResult ReqResult = StreamReaderHandler->AddRequest(CurrentPlaybackSequenceID[StreamTypeToArrayIndex(ReadyReq->GetType())], ReadyReq);
+				if (ReqResult != IStreamReader::EAddResult::Added)
+				{
+					FErrorDetail err;
+					err.SetFacility(Facility::EFacility::Player);
+					err.SetMessage("Failed to add stream segment request to reader");
+					err.SetCode(INTERR_FRAGMENT_READER_REQUEST);
+					PostError(err);
+					break;
+				}
+			}
+			else
+			{
+				// Skip over this segment. Add to the list of completed requests for which we need to find the next segment to fetch.
+				FPendingSegmentRequest NextReq;
+				NextReq.Request = ReadyReq;
+				NextPendingSegmentRequests.Enqueue(MoveTemp(NextReq));
+			}
+		}
+		else
+		{
+			ReadyWaitingSegmentRequests.Enqueue(ReadyReq);
+		}
+	}
+
 	InternalHandleCompletedSegmentRequests(CurrentTime);
 }
 
+void FAdaptiveStreamingPlayer::CancelPendingMediaSegmentRequests(EStreamType StreamType)
+{
+	if (StreamReaderHandler)
+	{
+		StreamReaderHandler->CancelRequest(StreamType, true);
+	}
+	// Go through the queue of ready-to-fetch segments and remove those that are to be canceled.
+	TQueue<TSharedPtrTS<IStreamSegment>> readyRequests;
+	Swap(readyRequests, ReadyWaitingSegmentRequests);
+	while(!readyRequests.IsEmpty())
+	{
+		TSharedPtrTS<IStreamSegment> ReadyReq;
+		readyRequests.Dequeue(ReadyReq);
+		if (ReadyReq->GetType() != StreamType)
+		{
+			ReadyWaitingSegmentRequests.Enqueue(ReadyReq);
+		}
+	}
+}
+
+
+void FAdaptiveStreamingPlayer::InternalDeselectStream(EStreamType StreamType)
+{
+	switch(StreamType)
+	{
+		case EStreamType::Video:
+			MultiStreamBufferVid.Deselect();
+			break;
+		case EStreamType::Audio:
+			MultiStreamBufferAud.Deselect();
+			break;
+		case EStreamType::Subtitle:
+			MultiStreamBufferTxt.Deselect();
+			SelectedStreamAttributesTxt.Deselect();
+			SubtitleDecoder.Flush();
+			break;
+		default:
+			break;
+	}
+}
+
 
 void FAdaptiveStreamingPlayer::InternalHandleSegmentTrackChanges(const FTimeValue& CurrentTime)
 {
@@ -2856,9 +2901,23 @@ void FAdaptiveStreamingPlayer::InternalHandleSegmentTrackChanges(const FTimeValu
 			SelectedStreamAttributesAud.UpdateWith(BufferSourceInfo->Kind, BufferSourceInfo->Language, BufferSourceInfo->HardIndex);
 			PendingTrackSelectionAud.Reset();
 		}
-		// Ignore video and subtitle changes for now.
+
+		if (PendingTrackSelectionTxt.IsValid() && CurrentPlayPeriodText.IsValid())
+		{
+			IManifest::IPlayPeriod::ETrackChangeResult Result = CurrentPlayPeriodText->ChangeTrackStreamPreference(EStreamType::Subtitle, *PendingTrackSelectionTxt);
+			TSharedPtrTS<FBufferSourceInfo> BufferSourceInfo = CurrentPlayPeriodText->GetSelectedStreamBufferSourceInfo(EStreamType::Subtitle);
+			if (Result == IManifest::IPlayPeriod::ETrackChangeResult::Changed)
+			{
+				MultiStreamBufferTxt.SelectTrackWhenAvailable(BufferSourceInfo);
+				StreamSelectionAttributesTxt = *PendingTrackSelectionTxt;
+			}
+			SelectedStreamAttributesTxt.UpdateWith(BufferSourceInfo->Kind, BufferSourceInfo->Language, BufferSourceInfo->HardIndex);
+			SelectedStreamAttributesTxt.Select();
+			PendingTrackSelectionTxt.Reset();
+		}
+
+		// Ignore video changes for now.
 		PendingTrackSelectionVid.Reset();
-		PendingTrackSelectionTxt.Reset();
 	}
 	else if (ManifestType == EMediaFormatType::HLS)
 	{
@@ -2891,10 +2950,7 @@ void FAdaptiveStreamingPlayer::InternalHandleSegmentTrackChanges(const FTimeValu
 					// Make sure any finishing current requests will be ignored by increasing the sequence ID.
 					++CurrentPlaybackSequenceID[StreamTypeToArrayIndex(EStreamType::Audio)];
 					// Cancel any ongoing segment download now.
-					if (StreamReaderHandler)
-					{
-						StreamReaderHandler->CancelRequest(EStreamType::Audio, true);
-					}
+					CancelPendingMediaSegmentRequests(EStreamType::Audio);
 
 					StreamSelectionAttributesAud = *PendingTrackSelectionAud;
 
@@ -2908,9 +2964,33 @@ void FAdaptiveStreamingPlayer::InternalHandleSegmentTrackChanges(const FTimeValu
 			PendingTrackSelectionAud.Reset();
 		}
 
-		// Ignore video and subtitle changes for now.
+		if (PendingTrackSelectionTxt.IsValid() && CurrentPlayPeriodText.IsValid())
+		{
+			if (SelectedStreamAttributesTxt.IsDeselected() || !SelectedStreamAttributesTxt.IsCompatibleWith(*PendingTrackSelectionTxt))
+			{
+				IManifest::IPlayPeriod::ETrackChangeResult TrackChangeResult = CurrentPlayPeriodText->ChangeTrackStreamPreference(EStreamType::Subtitle, *PendingTrackSelectionTxt);
+				if (TrackChangeResult == IManifest::IPlayPeriod::ETrackChangeResult::NewPeriodNeeded)
+				{
+					// Make sure any finishing current requests will be ignored by increasing the sequence ID.
+					++CurrentPlaybackSequenceID[StreamTypeToArrayIndex(EStreamType::Subtitle)];
+					// Cancel any ongoing segment download now.
+					CancelPendingMediaSegmentRequests(EStreamType::Subtitle);
+
+					StreamSelectionAttributesTxt = *PendingTrackSelectionTxt;
+					SelectedStreamAttributesTxt.Select();
+
+					FPendingSegmentRequest NextReq;
+					NextReq.bStartOver = true;
+					NextReq.StartoverPosition.Time = PlaybackState.GetPlayPosition();
+					NextReq.StreamType = EStreamType::Subtitle;
+					NextPendingSegmentRequests.Enqueue(MoveTemp(NextReq));
+				}
+			}
+			PendingTrackSelectionTxt.Reset();
+		}
+
+		// Ignore video changes for now.
 		PendingTrackSelectionVid.Reset();
-		PendingTrackSelectionTxt.Reset();
 	}
 	else
 	{
@@ -3154,8 +3234,10 @@ bool FAdaptiveStreamingPlayer::InternalStartAt(const FSeekParam& NewPosition)
 	check(!InitialPlayPeriod.IsValid());
 	check(!PendingFirstSegmentRequest.IsValid());
 	check(NextPendingSegmentRequests.IsEmpty());
+	check(ReadyWaitingSegmentRequests.IsEmpty());
 	check(CompletedSegmentRequests.Num() == 0);
 	NextPendingSegmentRequests.Empty();
+	ReadyWaitingSegmentRequests.Empty();
 	PendingFirstSegmentRequest.Reset();
 	CompletedSegmentRequests.Empty();
 	InitialPlayPeriod.Reset();
@@ -3308,6 +3390,7 @@ void FAdaptiveStreamingPlayer::InternalStop(bool bHoldCurrentFrame)
 	PendingStartRequest.Reset();
 	PendingFirstSegmentRequest.Reset();
 	NextPendingSegmentRequests.Empty();
+	ReadyWaitingSegmentRequests.Empty();
 	CompletedSegmentRequests.Empty();
 	InitialPlayPeriod.Reset();
 	CurrentPlayPeriodVideo.Reset();
@@ -3334,6 +3417,7 @@ void FAdaptiveStreamingPlayer::InternalStop(bool bHoldCurrentFrame)
 	// Flush the decoders.
 	AudioDecoder.Flush();
 	VideoDecoder.Flush();
+	SubtitleDecoder.Flush();
 
 	// Flush the renderers again, this time to discard everything the decoders may have emitted while being flushed.
 	AudioRender.Flush();
@@ -3520,6 +3604,11 @@ void FAdaptiveStreamingPlayer::DebugPrint(void* pPlayer, void (*pPrintFN)(void*
 		pPrintFN(pPlayer, "Audio buffer  : %3u AUs; %8u/%8u bytes in; eod %d; %#7.4fs", (uint32)pD->NumCurrentAccessUnits, (uint32)pD->CurrentMemInUse, (uint32)pD->MaxDataSize, pD->bEndOfData, pD->PushedDuration.GetAsSeconds());
 		pPrintFN(pPlayer, "Audio decoder : %2u in decoder, %2u total, %s; EOD in %d; EOD out %d", (uint32)AudioBufferStats.DecoderOutputBuffer.NumElementsInDecoder, (uint32)AudioBufferStats.DecoderOutputBuffer.MaxDecodedElementsReady, AudioBufferStats.DecoderOutputBuffer.bOutputStalled ? "    stalled" : "not stalled", AudioBufferStats.DecoderInputBuffer.bEODReached, AudioBufferStats.DecoderOutputBuffer.bEODreached);
 	}
+	if (StreamReaderHandler && bHaveTextReader.GetWithDefault(false))
+	{
+		pD = &TextBufferStats.StreamBuffer;
+		pPrintFN(pPlayer, "Text buffer   : %3u AUs; %8u/%8u bytes in; eod %d; %#7.4fs", (uint32)pD->NumCurrentAccessUnits, (uint32)pD->CurrentMemInUse, (uint32)pD->MaxDataSize, pD->bEndOfData, pD->PushedDuration.GetAsSeconds());
+	}
 	DiagnosticsCriticalSection.Unlock();
 
 	pPrintFN(pPlayer, "Player state  : %s", GetPlayerStateName(CurrentState));
diff --git a/Engine/Plugins/Media/ElectraPlayer/Source/ElectraPlayerRuntime/Private/Runtime/Player/AdaptiveStreamingPlayer.h b/Engine/Plugins/Media/ElectraPlayer/Source/ElectraPlayerRuntime/Private/Runtime/Player/AdaptiveStreamingPlayer.h
--- a/Engine/Plugins/Media/ElectraPlayer/Source/ElectraPlayerRuntime/Private/Runtime/Player/AdaptiveStreamingPlayer.h
+++ b/Engine/Plugins/Media/ElectraPlayer/Source/ElectraPlayerRuntime/Private/Runtime/Player/AdaptiveStreamingPlayer.h
@@ -11,6 +11,7 @@
 #include "StreamTypes.h"
 #include "AdaptiveStreamingPlayerMetrics.h"
 #include "AdaptiveStreamingPlayerEvents.h"
+#include "AdaptiveStreamingPlayerSubtitles.h"
 #include "AdaptiveStreamingPlayerResourceRequest.h"
 
 class IOptionPointerValueContainer;
@@ -55,6 +56,11 @@ public:
 	virtual void AddAEMSReceiver(TWeakPtrTS<IAdaptiveStreamingPlayerAEMSReceiver> InReceiver, FString InForSchemeIdUri, FString InForValue, IAdaptiveStreamingPlayerAEMSReceiver::EDispatchMode InDispatchMode) = 0;
 	virtual void RemoveAEMSReceiver(TWeakPtrTS<IAdaptiveStreamingPlayerAEMSReceiver> InReceiver, FString InForSchemeIdUri, FString InForValue, IAdaptiveStreamingPlayerAEMSReceiver::EDispatchMode InDispatchMode) = 0;
 
+	//-------------------------------------------------------------------------
+	// Subtitle receiver
+	//
+	virtual void AddSubtitleReceiver(TWeakPtrTS<IAdaptiveStreamingPlayerSubtitleReceiver> InReceiver) = 0;
+	virtual void RemoveSubtitleReceiver(TWeakPtrTS<IAdaptiveStreamingPlayerSubtitleReceiver> InReceiver) = 0;
 
 	//-------------------------------------------------------------------------
 	// Initializes the player. Options may be passed in to affect behaviour.
@@ -184,7 +190,7 @@ public:
 	//! This selection will explicitly override the initial stream attributes set by SetInitialStreamAttributes() and be applied automatically for upcoming periods.
 	virtual void SelectTrackByAttributes(EStreamType StreamType, const FStreamSelectionAttributes& Attributes) = 0;
 
-	//! Deselect track. The stream will continue to stream to allow for immediate selection/activation but no data will be fed to the decoder.
+	//! Deselect track. The stream may continue to stream to allow for immediate selection/activation but no data will be fed to the decoder.
 	virtual void DeselectTrack(EStreamType StreamType) = 0;
 
 	//! Returns true if the track stream of the specified type has beed deselected through DeselectTrack().
diff --git a/Engine/Plugins/Media/ElectraPlayer/Source/ElectraPlayerRuntime/Private/Runtime/Player/AdaptiveStreamingPlayerDecoders.cpp b/Engine/Plugins/Media/ElectraPlayer/Source/ElectraPlayerRuntime/Private/Runtime/Player/AdaptiveStreamingPlayerDecoders.cpp
--- a/Engine/Plugins/Media/ElectraPlayer/Source/ElectraPlayerRuntime/Private/Runtime/Player/AdaptiveStreamingPlayerDecoders.cpp
+++ b/Engine/Plugins/Media/ElectraPlayer/Source/ElectraPlayerRuntime/Private/Runtime/Player/AdaptiveStreamingPlayerDecoders.cpp
@@ -3,6 +3,7 @@
 #include "PlayerCore.h"
 #include "Player/AdaptiveStreamingPlayerInternal.h"
 #include "Player/AdaptivePlayerOptionKeynames.h"
+#include "Decoder/SubtitleDecoder.h"
 #include "Utilities/Utilities.h"
 #include "ParameterDictionary.h"
 
@@ -10,6 +11,157 @@
 namespace Electra
 {
 
+// AU memory
+void* FAdaptiveStreamingPlayer::AUAllocate(IAccessUnitMemoryProvider::EDataType type, SIZE_T NumBytes, SIZE_T Alignment)
+{
+	if (Alignment)
+	{
+		return FMemory::Malloc(NumBytes, Alignment);
+	}
+	return FMemory::Malloc(NumBytes);
+}
+
+void FAdaptiveStreamingPlayer::AUDeallocate(IAccessUnitMemoryProvider::EDataType type, void* Address)
+{
+	FMemory::Free(Address);
+}
+
+IPlayerStreamFilter* FAdaptiveStreamingPlayer::GetStreamFilter()
+{
+	return this;
+}
+
+const FCodecSelectionPriorities& FAdaptiveStreamingPlayer::GetCodecSelectionPriorities(EStreamType ForStream)
+{
+	switch(ForStream)
+	{
+		case EStreamType::Video:
+		{
+			return CodecPrioritiesVideo;
+		}
+		case EStreamType::Audio:
+		{
+			return CodecPrioritiesAudio;
+		}
+		case EStreamType::Subtitle:
+		{
+			return CodecPrioritiesSubtitles;
+		}
+		default:
+		{
+			static FCodecSelectionPriorities Dummy;
+			return Dummy;
+		}
+	}
+}
+
+
+
+bool FAdaptiveStreamingPlayer::CanDecodeStream(const FStreamCodecInformation& InStreamCodecInfo) const
+{
+	auto IsCodecPrefixExcluded = [](const FString& InCodec, const TArray<FString>& InExcludedPrefixes) -> bool
+	{
+		for(auto &Prefix : InExcludedPrefixes)
+		{
+			if (InCodec.StartsWith(Prefix, ESearchCase::IgnoreCase))
+			{
+				return true;
+			}
+		}
+		return false;
+	};
+
+
+	if (InStreamCodecInfo.IsVideoCodec())
+	{
+		// Check against user specified codec prefix exclusions
+		if (IsCodecPrefixExcluded(InStreamCodecInfo.GetCodecSpecifierRFC6381(), ExcludedVideoDecoderPrefixes))
+		{
+			return false;
+		}
+
+		double Rate = InStreamCodecInfo.GetFrameRate().IsValid() ? InStreamCodecInfo.GetFrameRate().GetAsDouble() : 30.0;
+		if (InStreamCodecInfo.GetCodec() == FStreamCodecInformation::ECodec::H264)
+		{
+			const AdaptiveStreamingPlayerConfig::FConfiguration::FVideoDecoderLimit* DecoderLimit = &PlayerConfig.H264LimitUpto30fps;
+			if (Rate > 31.0)
+			{
+				DecoderLimit = &PlayerConfig.H264LimitAbove30fps;
+			}
+			// Check against user configured resolution limit
+			if (DecoderLimit->MaxResolution.Height && InStreamCodecInfo.GetResolution().Height > DecoderLimit->MaxResolution.Height)
+			{
+				return false;
+			}
+
+			// Check against video decoder capabilities.
+			IVideoDecoderH264::FStreamDecodeCapability StreamParam, Capability;
+			StreamParam.Width = InStreamCodecInfo.GetResolution().Width;
+			StreamParam.Height = InStreamCodecInfo.GetResolution().Height;
+			StreamParam.Profile = InStreamCodecInfo.GetProfile();
+			StreamParam.Level = InStreamCodecInfo.GetProfileLevel();
+			StreamParam.FPS = Rate;
+			if (IVideoDecoderH264::GetStreamDecodeCapability(Capability, StreamParam))
+			{
+				if (Capability.DecoderSupportType == IVideoDecoderH264::FStreamDecodeCapability::ESupported::NotSupported)
+				{
+					return false;
+				}
+			}
+		}
+		else if (InStreamCodecInfo.GetCodec() == FStreamCodecInformation::ECodec::H265)
+		{
+			#if ELECTRA_PLATFORM_HAS_H265_DECODER
+				// Check against video decoder capabilities.
+				IVideoDecoderH265::FStreamDecodeCapability StreamParam, Capability;
+				StreamParam.Width = InStreamCodecInfo.GetResolution().Width;
+				StreamParam.Height = InStreamCodecInfo.GetResolution().Height;
+				StreamParam.Tier = InStreamCodecInfo.GetProfileTier();
+				StreamParam.CompatibilityFlags = InStreamCodecInfo.GetProfileCompatibilityFlags();
+				StreamParam.Profile = InStreamCodecInfo.GetProfile();
+				StreamParam.Level = InStreamCodecInfo.GetProfileLevel();
+				StreamParam.ConstraintFlags = InStreamCodecInfo.GetProfileConstraints();
+				StreamParam.FPS = Rate;
+				if (IVideoDecoderH265::GetStreamDecodeCapability(Capability, StreamParam))
+				{
+					if (Capability.DecoderSupportType == IVideoDecoderH265::FStreamDecodeCapability::ESupported::NotSupported)
+					{
+						return false;
+					}
+				}
+			#else
+				return false;
+			#endif
+		}
+	}
+	else if (InStreamCodecInfo.IsAudioCodec())
+	{
+		// Check against user specified codec prefix exclusions
+		if (IsCodecPrefixExcluded(InStreamCodecInfo.GetCodecSpecifierRFC6381(), ExcludedAudioDecoderPrefixes))
+		{
+			return false;
+		}
+	}
+	else if (InStreamCodecInfo.IsSubtitleCodec())
+	{
+		// Check against user specified codec prefix exclusions
+		if (IsCodecPrefixExcluded(InStreamCodecInfo.GetCodecSpecifierRFC6381(), ExcludedSubtitleDecoderPrefixes))
+		{
+			return false;
+		}
+	}
+
+	return true;
+}
+
+
+bool FAdaptiveStreamingPlayer::CanDecodeSubtitle(const FString& MimeType, const FString& Codec) const
+{
+	return ISubtitleDecoder::IsSupported(MimeType, Codec);
+}
+
+
+
 
 bool FAdaptiveStreamingPlayer::FindMatchingStreamInfo(FStreamCodecInformation& OutStreamInfo, const FTimeValue& AtTime, int32 MaxWidth, int32 MaxHeight)
 {
@@ -321,6 +473,45 @@ int32 FAdaptiveStreamingPlayer::CreateDecoder(EStreamType type)
 		}
 		return 0;
 	}
+	else if (type == EStreamType::Subtitle)
+	{
+		if (SubtitleDecoder.Decoder == nullptr)
+		{
+			FAccessUnit* AccessUnit = nullptr;
+			bool bOk = MultiStreamBufferTxt.PeekAndAddRef(AccessUnit);
+			if (AccessUnit)
+			{
+				SubtitleDecoder.CurrentCodecInfo.Clear();
+				if (AccessUnit->AUCodecData.IsValid())
+				{
+					SubtitleDecoder.CurrentCodecInfo = AccessUnit->AUCodecData->ParsedInfo;
+				}
+				FAccessUnit::Release(AccessUnit);
+
+				// Create subtitle decoder for this type of AU.
+				SubtitleDecoder.Decoder = ISubtitleDecoder::Create(AccessUnit);
+				if (SubtitleDecoder.Decoder)
+				{
+					SubtitleDecoder.Decoder->SetPlayerSessionServices(this);
+					SubtitleDecoder.Decoder->GetDecodedSubtitleReceiveDelegate().BindRaw(this, &FAdaptiveStreamingPlayer::OnDecodedSubtitleReceived);
+					SubtitleDecoder.Decoder->GetDecodedSubtitleFlushDelegate().BindRaw(this, &FAdaptiveStreamingPlayer::OnFlushSubtitleReceivers);
+					SubtitleDecoder.Decoder->Open();
+					// Start the decoder, if decoding is currently running.
+					SubtitleDecoder.Start(DecoderState == EDecoderState::eDecoder_Running);
+				}
+				else
+				{
+					FErrorDetail err;
+					err.SetFacility(Facility::EFacility::Player);
+					err.SetMessage("Unsupported subtitle codec");
+					err.SetCode(INTERR_UNSUPPORTED_CODEC);
+					PostError(err);
+					return -1;
+				}
+			}
+		}
+		return 0;
+	}
 	return -1;
 }
 
@@ -342,6 +533,7 @@ NOTE: We do not clear out the renderers from the decoder. On their way down the
 */
 	AudioDecoder.Close();
 	VideoDecoder.Close();
+	SubtitleDecoder.Close();
 }
 
 
@@ -359,9 +551,45 @@ void FAdaptiveStreamingPlayer::HandleDecoderChanges()
 	}
 	CreateDecoder(EStreamType::Video);
 	CreateDecoder(EStreamType::Audio);
+	CreateDecoder(EStreamType::Subtitle);
 }
 
 
+void FAdaptiveStreamingPlayer::HandleSubtitleDecoder()
+{
+	// Feeds access units from the subtitle buffer to the subtitle decoder.
+	// Since this is very small and sparse data the decoder does not run in a dedicated thread like
+	// the video and audio decoders do.
+	if (SubtitleDecoder.Decoder && !MultiStreamBufferTxt.IsDeselected())
+	{
+		FAccessUnit* PeekedAU = nullptr;
+		MultiStreamBufferTxt.PeekAndAddRef(PeekedAU);
+		FTimeValue NextPTS = PeekedAU ? PeekedAU->PTS - SubtitleDecoder.Decoder->GetStreamedDeliveryTimeOffset() : FTimeValue::GetInvalid();
+		FAccessUnit::Release(PeekedAU);
+		FTimeValue CurrentPos = GetCurrentPlayTime();
+		if (NextPTS.IsValid() && CurrentPos >= NextPTS)
+		{
+			FeedDecoder(EStreamType::Subtitle, MultiStreamBufferTxt, SubtitleDecoder.Decoder, false);
+		}
+
+		if (CurrentPos.IsValid() && RenderClock->IsRunning())
+		{
+			FTimeValue LocalPos;
+			for(int32 i=0, iMax=ActivePeriods.Num(); i<iMax; ++i)
+			{
+				if (CurrentPos >= ActivePeriods[i].TimeRange.Start && ((i+1 == iMax) || CurrentPos < ActivePeriods[i+1].TimeRange.Start))
+				{
+					LocalPos = CurrentPos - ActivePeriods[i].TimeRange.Start;
+					break;
+				}
+			}
+			SubtitleDecoder.Decoder->UpdatePlaybackPosition(CurrentPos, LocalPos);
+		}
+	}
+}
+
+
+
 void FAdaptiveStreamingPlayer::VideoDecoderInputNeeded(const IAccessUnitBufferListener::FBufferStats &currentInputBufferStats)
 {
 	DiagnosticsCriticalSection.Lock();
@@ -372,7 +600,7 @@ void FAdaptiveStreamingPlayer::VideoDecoderInputNeeded(const IAccessUnitBufferLi
 		EDecoderState decState = DecoderState;
 		if (decState == EDecoderState::eDecoder_Running)
 		{
-			FeedDecoder(EStreamType::Video, MultiStreamBufferVid, VideoDecoder.Decoder);
+			FeedDecoder(EStreamType::Video, MultiStreamBufferVid, VideoDecoder.Decoder, true);
 		}
 	}
 }
@@ -393,7 +621,7 @@ void FAdaptiveStreamingPlayer::AudioDecoderInputNeeded(const IAccessUnitBufferLi
 	EDecoderState decState = DecoderState;
 	if (decState == EDecoderState::eDecoder_Running)
 	{
-		FeedDecoder(EStreamType::Audio, MultiStreamBufferAud, AudioDecoder.Decoder);
+		FeedDecoder(EStreamType::Audio, MultiStreamBufferAud, AudioDecoder.Decoder, true);
 	}
 }
 
@@ -415,8 +643,9 @@ void FAdaptiveStreamingPlayer::AudioDecoderOutputReady(const IDecoderOutputBuffe
  * @param Type
  * @param FromMultistreamBuffer
  * @param Decoder
+ * @param bHandleUnderrun
  */
-void FAdaptiveStreamingPlayer::FeedDecoder(EStreamType Type, FMultiTrackAccessUnitBuffer& FromMultistreamBuffer, IAccessUnitBufferInterface* Decoder)
+void FAdaptiveStreamingPlayer::FeedDecoder(EStreamType Type, FMultiTrackAccessUnitBuffer& FromMultistreamBuffer, IAccessUnitBufferInterface* Decoder, bool bHandleUnderrun)
 {
 	// Lock the AU buffer for the duration of this function to ensure this can never clash with a Flush() call
 	// since we are checking size, eod state and subsequently popping an AU, for which the buffer must stay consistent inbetween!
@@ -443,6 +672,7 @@ void FAdaptiveStreamingPlayer::FeedDecoder(EStreamType Type, FMultiTrackAccessUn
 		case EStreamType::Subtitle:
 			pStats = &TextBufferStats;
 			pAvailability = &DataAvailabilityStateTxt;
+			CurrentCodecInfo = &SubtitleDecoder.CurrentCodecInfo;
 			break;
 		default:
 			checkNoEntry();
@@ -453,7 +683,7 @@ void FAdaptiveStreamingPlayer::FeedDecoder(EStreamType Type, FMultiTrackAccessUn
 	if (!FromMultistreamBuffer.IsDeselected())
 	{
 		// Check for buffer underrun.
-		if (!bRebufferPending && CurrentState == EPlayerState::eState_Playing && StreamState == EStreamState::eStream_Running && PipelineState == EPipelineState::ePipeline_Running)
+		if (bHandleUnderrun && !bRebufferPending && CurrentState == EPlayerState::eState_Playing && StreamState == EStreamState::eStream_Running && PipelineState == EPipelineState::ePipeline_Running)
 		{
 			bool bEODSet = FromMultistreamBuffer.IsEODFlagSet();
 			if (!bEODSet && FromMultistreamBuffer.Num() == 0)
@@ -504,7 +734,7 @@ void FAdaptiveStreamingPlayer::FeedDecoder(EStreamType Type, FMultiTrackAccessUn
 			}
 			else
 			{
-				// release peeked and actually pop it
+				// Release peeked AU and actually pop it
 				FAccessUnit* AccessUnit = nullptr;
 				FromMultistreamBuffer.Pop(AccessUnit);
 				check(PeekedAU == AccessUnit);
diff --git a/Engine/Plugins/Media/ElectraPlayer/Source/ElectraPlayerRuntime/Private/Runtime/Player/AdaptiveStreamingPlayerEvents.cpp b/Engine/Plugins/Media/ElectraPlayer/Source/ElectraPlayerRuntime/Private/Runtime/Player/AdaptiveStreamingPlayerEvents.cpp
--- a/Engine/Plugins/Media/ElectraPlayer/Source/ElectraPlayerRuntime/Private/Runtime/Player/AdaptiveStreamingPlayerEvents.cpp
+++ b/Engine/Plugins/Media/ElectraPlayer/Source/ElectraPlayerRuntime/Private/Runtime/Player/AdaptiveStreamingPlayerEvents.cpp
@@ -62,6 +62,10 @@ void FAdaptiveStreamingPlayer::DispatchBufferUtilizationEvent(EStreamType Buffer
 	{
 		MultiStreamBufferAud.GetStats(bufStats);
 	}
+	else if (BufferType == EStreamType::Subtitle)
+	{
+		MultiStreamBufferTxt.GetStats(bufStats);
+	}
 
 	stats.MaxDurationInSeconds = bufStats.MaxDuration.GetAsSeconds();
 	stats.DurationInUse 	   = bufStats.PushedDuration.GetAsSeconds();
diff --git a/Engine/Plugins/Media/ElectraPlayer/Source/ElectraPlayerRuntime/Private/Runtime/Player/AdaptiveStreamingPlayerInternal.h b/Engine/Plugins/Media/ElectraPlayer/Source/ElectraPlayerRuntime/Private/Runtime/Player/AdaptiveStreamingPlayerInternal.h
--- a/Engine/Plugins/Media/ElectraPlayer/Source/ElectraPlayerRuntime/Private/Runtime/Player/AdaptiveStreamingPlayerInternal.h
+++ b/Engine/Plugins/Media/ElectraPlayer/Source/ElectraPlayerRuntime/Private/Runtime/Player/AdaptiveStreamingPlayerInternal.h
@@ -11,6 +11,7 @@
 #include "SynchronizedClock.h"
 
 #include "Demuxer/ParserISO14496-12.h"
+#include "Decoder/SubtitleDecoder.h"
 
 #include "Player/AdaptiveStreamingPlayerResourceRequest.h"
 #include "Player/PlayerStreamReader.h"
@@ -192,6 +193,7 @@ struct FPlaybackState
 	bool									bIsPaused;
 	TArray<FTrackMetadata>					VideoTracks;
 	TArray<FTrackMetadata>					AudioTracks;
+	TArray<FTrackMetadata>					SubtitleTracks;
 	TArray<FTimespan>						SeekablePositions;
 
 
@@ -323,7 +325,7 @@ struct FPlaybackState
 		bIsPlaying = bInIsPlaying;
 	}
 
-	bool SetTrackMetadata(const TArray<FTrackMetadata> &InVideoTracks, const TArray<FTrackMetadata>& InAudioTracks)
+	bool SetTrackMetadata(const TArray<FTrackMetadata> &InVideoTracks, const TArray<FTrackMetadata>& InAudioTracks, const TArray<FTrackMetadata>& InSubtitleTracks)
 	{
 		FMediaCriticalSection::ScopedLock lock(Lock);
 
@@ -343,16 +345,18 @@ struct FPlaybackState
 			return true;
 		};
 
-		bool bChanged = ChangedTrackMetadata(VideoTracks, InVideoTracks) || ChangedTrackMetadata(AudioTracks, InAudioTracks);
+		bool bChanged = ChangedTrackMetadata(VideoTracks, InVideoTracks) || ChangedTrackMetadata(AudioTracks, InAudioTracks) || ChangedTrackMetadata(SubtitleTracks, InSubtitleTracks);
 		VideoTracks = InVideoTracks;
 		AudioTracks = InAudioTracks;
+		SubtitleTracks = InSubtitleTracks;
 		return bChanged;
 	}
-	void GetTrackMetadata(TArray<FTrackMetadata> &OutVideoTracks, TArray<FTrackMetadata>& OutAudioTracks) const
+	void GetTrackMetadata(TArray<FTrackMetadata> &OutVideoTracks, TArray<FTrackMetadata>& OutAudioTracks, TArray<FTrackMetadata>& OutSubtitleTracks) const
 	{
 		FMediaCriticalSection::ScopedLock lock(Lock);
 		OutVideoTracks = VideoTracks;
 		OutAudioTracks = AudioTracks;
+		OutSubtitleTracks = SubtitleTracks;
 	}
 
 	void SetLoopState(const FPlayerLoopState& InLoopState)
@@ -714,6 +718,9 @@ public:
 	virtual void AddAEMSReceiver(TWeakPtrTS<IAdaptiveStreamingPlayerAEMSReceiver> InReceiver, FString InForSchemeIdUri, FString InForValue, IAdaptiveStreamingPlayerAEMSReceiver::EDispatchMode InDispatchMode) override;
 	virtual void RemoveAEMSReceiver(TWeakPtrTS<IAdaptiveStreamingPlayerAEMSReceiver> InReceiver, FString InForSchemeIdUri, FString InForValue, IAdaptiveStreamingPlayerAEMSReceiver::EDispatchMode InDispatchMode) override;
 
+	virtual void AddSubtitleReceiver(TWeakPtrTS<IAdaptiveStreamingPlayerSubtitleReceiver> InReceiver) override;
+	virtual void RemoveSubtitleReceiver(TWeakPtrTS<IAdaptiveStreamingPlayerSubtitleReceiver> InReceiver) override;
+
 	virtual void Initialize(const FParamDict& Options) override;
 
 	virtual void SetInitialStreamAttributes(EStreamType StreamType, const FStreamSelectionAttributes& InitialSelection) override;
@@ -780,6 +787,7 @@ private:
 
 	// Methods from IPlayerStreamFilter
 	virtual bool CanDecodeStream(const FStreamCodecInformation& InStreamCodecInfo) const override;
+	virtual bool CanDecodeSubtitle(const FString& MimeType, const FString& Codec) const override;
 
 
 	enum class EPlayerState
@@ -896,7 +904,7 @@ private:
 	};
 
 	struct FVideoDecoder : public IAccessUnitBufferListener, public IDecoderOutputBufferListener
-		{
+	{
 		void Close()
 		{
 			if (Decoder)
@@ -936,10 +944,10 @@ private:
 		bool							bDrainingForCodecChange = false;
 		bool							bDrainingForCodecChangeDone = false;
 		bool							bApplyNewLimits = false;
-		};
+	};
 
 	struct FAudioDecoder : public IAccessUnitBufferListener, public IDecoderOutputBufferListener
-		{
+	{
 		FAudioDecoder()
 			: Parent(nullptr)
 			, Decoder(nullptr)
@@ -981,8 +989,60 @@ private:
 		FStreamCodecInformation			CurrentCodecInfo;
 		FAdaptiveStreamingPlayer*		Parent;
 		IAudioDecoderAAC*				Decoder;
-		};
+	};
+
+	struct FSubtitleDecoder
+	{
+		void Close()
+		{
+			if (Decoder)
+			{
+				Stop();
+				Decoder->Close();
+				Decoder->GetDecodedSubtitleReceiveDelegate().Unbind();
+				Decoder->GetDecodedSubtitleFlushDelegate().Unbind();
+				delete Decoder;
+				Decoder = nullptr;
+			}
+		}
+		void Flush()
+		{
+			if (Decoder)
+			{
+				Decoder->AUdataFlushEverything();
+			}
+		}
+		
+		void Start(bool bInitialState=true)
+		{
+			if (bInitialState != bIsRunning)
+			{
+				if (Decoder)
+				{
+					Decoder->Start();
+				}
+				bIsRunning = bInitialState;
+			}
+		}
+		
+		void Stop()
+		{
+			if (bIsRunning)
+			{
+				if (Decoder)
+				{
+					Decoder->Stop();
+				}
+				bIsRunning = false;
+			}
+		}
 
+		FStreamCodecInformation			CurrentCodecInfo;
+		ISubtitleDecoder*				Decoder = nullptr;
+		bool							bIsRunning = false;
+		bool							bDrainingForCodecChange = false;
+		bool							bDrainingForCodecChangeDone = false;
+	};
 
 	class FWorkerThreadMessages
 	{
@@ -1388,6 +1448,19 @@ private:
 	};
 
 
+	struct FInternalStreamSelectionAttributes : public FStreamSelectionAttributes
+	{
+		bool IsDeselected() const
+		{ return !bIsSelected; }
+		void Select()
+		{ bIsSelected = true; }
+		void Deselect()
+		{ bIsSelected = false; }
+	private:
+		bool bIsSelected = true;
+	};
+
+
 	void InternalLoadManifest(const FString& URL, const FString& MimeType);
 	void InternalCancelLoadManifest();
 	void InternalCloseManifestReader();
@@ -1401,6 +1474,9 @@ private:
 	void AudioDecoderInputNeeded(const IAccessUnitBufferListener::FBufferStats& currentInputBufferStats);
 	void AudioDecoderOutputReady(const IDecoderOutputBufferListener::FDecodeReadyStats& currentReadyStats);
 
+	void OnDecodedSubtitleReceived(ISubtitleDecoderOutputPtr Subtitle);
+	void OnFlushSubtitleReceivers();
+
 	void StartWorkerThread();
 	void StopWorkerThread();
 
@@ -1437,8 +1513,10 @@ private:
 	void HandleSessionMessage(TSharedPtrTS<IPlayerMessage> SessionMessage);
 	void HandlePlayStateChanges();
 	void HandlePendingMediaSegmentRequests();
+	void CancelPendingMediaSegmentRequests(EStreamType StreamType);
 	void HandleDeselectedBuffers();
 	void HandleDecoderChanges();
+	void HandleSubtitleDecoder();
 	void HandleMetadataChanges();
 	void HandleAEMSEvents();
 
@@ -1446,7 +1524,7 @@ private:
 
 	void CheckForErrors();
 
-	void FeedDecoder(EStreamType Type, FMultiTrackAccessUnitBuffer& FromMultistreamBuffer, IAccessUnitBufferInterface* Decoder);
+	void FeedDecoder(EStreamType Type, FMultiTrackAccessUnitBuffer& FromMultistreamBuffer, IAccessUnitBufferInterface* Decoder, bool bHandleUnderrun);
 
 	bool InternalStartAt(const FSeekParam& NewPosition);
 	void InternalPause();
@@ -1456,6 +1534,7 @@ private:
 	void InternalClose();
 	void InternalSetLoop(const FLoopParam& LoopParam);
 	void InternalSetPlaybackEnded();
+	void InternalDeselectStream(EStreamType StreamType);
 
 	void DispatchEvent(TSharedPtrTS<FMetricEvent> Event);
 	void DispatchEventAndWait(TSharedPtrTS<FMetricEvent> Event);
@@ -1521,9 +1600,9 @@ private:
 	FStreamSelectionAttributes											StreamSelectionAttributesAud;
 	FStreamSelectionAttributes											StreamSelectionAttributesTxt;
 
-	FStreamSelectionAttributes											SelectedStreamAttributesVid;
-	FStreamSelectionAttributes											SelectedStreamAttributesAud;
-	FStreamSelectionAttributes											SelectedStreamAttributesTxt;
+	FInternalStreamSelectionAttributes									SelectedStreamAttributesVid;
+	FInternalStreamSelectionAttributes									SelectedStreamAttributesAud;
+	FInternalStreamSelectionAttributes									SelectedStreamAttributesTxt;
 
 	TSharedPtrTS<FStreamSelectionAttributes>							PendingTrackSelectionVid;
 	TSharedPtrTS<FStreamSelectionAttributes>							PendingTrackSelectionAud;
@@ -1575,6 +1654,7 @@ private:
 	TSharedPtrTS<FPendingStartRequest>									PendingStartRequest;
 	TSharedPtrTS<IStreamSegment>										PendingFirstSegmentRequest;
 	TQueue<FPendingSegmentRequest>										NextPendingSegmentRequests;
+	TQueue<TSharedPtrTS<IStreamSegment>>								ReadyWaitingSegmentRequests;
 	TMultiMap<EStreamType, TSharedPtrTS<IStreamSegment>>				CompletedSegmentRequests;
 	bool																bFirstSegmentRequestIsForLooping;
 
@@ -1584,6 +1664,10 @@ private:
 	FAudioRenderer														AudioRender;
 	FVideoDecoder														VideoDecoder;
 	FAudioDecoder														AudioDecoder;
+	FSubtitleDecoder													SubtitleDecoder;
+	FCriticalSection													SubtitleReceiversCriticalSection;
+	TArray<TWeakPtrTS<IAdaptiveStreamingPlayerSubtitleReceiver>>		SubtitleReceivers;
+
 
 	FMediaCriticalSection												MetricListenerCriticalSection;
 	TArray<IAdaptiveStreamingPlayerMetrics*, TInlineAllocator<4>>		MetricListeners;
diff --git a/Engine/Plugins/Media/ElectraPlayer/Source/ElectraPlayerRuntime/Private/Runtime/Player/AdaptiveStreamingPlayerInternalConfig.h b/Engine/Plugins/Media/ElectraPlayer/Source/ElectraPlayerRuntime/Private/Runtime/Player/AdaptiveStreamingPlayerInternalConfig.h
--- a/Engine/Plugins/Media/ElectraPlayer/Source/ElectraPlayerRuntime/Private/Runtime/Player/AdaptiveStreamingPlayerInternalConfig.h
+++ b/Engine/Plugins/Media/ElectraPlayer/Source/ElectraPlayerRuntime/Private/Runtime/Player/AdaptiveStreamingPlayerInternalConfig.h
@@ -32,6 +32,11 @@ namespace Electra
 				StreamBufferConfigAudio.MaxDataSize = 2 << 20;
 				StreamBufferConfigAudio.MaxDuration.SetFromSeconds(20.0);
 
+				// Subtitle streams tend to be sparse and each AU could potentially have
+				// a huge duration, so we need to ignore durations altogether.
+				StreamBufferConfigText.MaxDataSize = 8 << 20;
+				StreamBufferConfigText.MaxDuration.SetToPositiveInfinity();
+
 				bHoldLastFrameDuringSeek = true;
 
 				InitialBufferMinTimeAvailBeforePlayback = 5.0;
diff --git a/Engine/Plugins/Media/ElectraPlayer/Source/ElectraPlayerRuntime/Private/Runtime/Player/AdaptiveStreamingPlayerManifest.cpp b/Engine/Plugins/Media/ElectraPlayer/Source/ElectraPlayerRuntime/Private/Runtime/Player/AdaptiveStreamingPlayerManifest.cpp
--- a/Engine/Plugins/Media/ElectraPlayer/Source/ElectraPlayerRuntime/Private/Runtime/Player/AdaptiveStreamingPlayerManifest.cpp
+++ b/Engine/Plugins/Media/ElectraPlayer/Source/ElectraPlayerRuntime/Private/Runtime/Player/AdaptiveStreamingPlayerManifest.cpp
@@ -243,9 +243,11 @@ bool FAdaptiveStreamingPlayer::SelectManifest()
 
 			TArray<FTrackMetadata> VideoTrackMetadata;
 			TArray<FTrackMetadata> AudioTrackMetadata;
+			TArray<FTrackMetadata> SubtitleTrackMetadata;
 			NewPresentation->GetTrackMetadata(VideoTrackMetadata, EStreamType::Video);
 			NewPresentation->GetTrackMetadata(AudioTrackMetadata, EStreamType::Audio);
-			PlaybackState.SetTrackMetadata(VideoTrackMetadata, AudioTrackMetadata);
+			NewPresentation->GetTrackMetadata(SubtitleTrackMetadata, EStreamType::Subtitle);
+			PlaybackState.SetTrackMetadata(VideoTrackMetadata, AudioTrackMetadata, SubtitleTrackMetadata);
 			PlaybackState.SetHaveMetadata(true);
 
 			Manifest = NewPresentation;
diff --git a/Engine/Plugins/Media/ElectraPlayer/Source/ElectraPlayerRuntime/Private/Runtime/Player/AdaptiveStreamingPlayerSubtitles.h b/Engine/Plugins/Media/ElectraPlayer/Source/ElectraPlayerRuntime/Private/Runtime/Player/AdaptiveStreamingPlayerSubtitles.h
new file mode 100644
--- /dev/null
+++ b/Engine/Plugins/Media/ElectraPlayer/Source/ElectraPlayerRuntime/Private/Runtime/Player/AdaptiveStreamingPlayerSubtitles.h
@@ -0,0 +1,23 @@
+// Copyright Epic Games, Inc. All Rights Reserved.
+
+#pragma once
+
+#include "PlayerCore.h"
+#include "MediaSubtitleDecoderOutput.h"
+
+namespace Electra
+{
+
+class IAdaptiveStreamingPlayerSubtitleReceiver
+{
+public:
+	virtual ~IAdaptiveStreamingPlayerSubtitleReceiver() = default;
+
+	virtual void OnMediaPlayerSubtitleReceived(ISubtitleDecoderOutputPtr Subtitle) = 0;
+	virtual void OnMediaPlayerFlushSubtitles() = 0;
+};
+
+
+
+} // namespace Electra
+
diff --git a/Engine/Plugins/Media/ElectraPlayer/Source/ElectraPlayerRuntime/Private/Runtime/Player/DASH/ManifestBuilderDASH.cpp b/Engine/Plugins/Media/ElectraPlayer/Source/ElectraPlayerRuntime/Private/Runtime/Player/DASH/ManifestBuilderDASH.cpp
--- a/Engine/Plugins/Media/ElectraPlayer/Source/ElectraPlayerRuntime/Private/Runtime/Player/DASH/ManifestBuilderDASH.cpp
+++ b/Engine/Plugins/Media/ElectraPlayer/Source/ElectraPlayerRuntime/Private/Runtime/Player/DASH/ManifestBuilderDASH.cpp
@@ -94,6 +94,11 @@ namespace
 	};
 
     const TCHAR* const Scheme_urn_mpeg_dash_mp4protection_2011 = TEXT("urn:mpeg:dash:mp4protection:2011");
+
+	// Recognized subtitle mime types
+    const TCHAR* const SubtitleMimeType_SideloadedTTML = TEXT("application/ttml+xml");
+    const TCHAR* const SubtitleMimeType_SideloadedVTT = TEXT("text/vtt");
+    const TCHAR* const SubtitleMimeType_Streamed = TEXT("application/mp4");
 }
 
 
@@ -1241,6 +1246,8 @@ void FManifestDASHInternal::PreparePeriodAdaptationSets(TSharedPtrTS<FPeriod> Pe
 		SCOPE_CYCLE_COUNTER(STAT_ElectraPlayer_FManifestDASHInternal_Build);
 		CSV_SCOPED_TIMING_STAT(ElectraPlayer, FManifestDASHInternal_Build);
 
+		IPlayerStreamFilter* StreamFilter = PlayerSessionServices->GetStreamFilter();
+
 		Period->AdaptationSets.Empty();
 		const TArray<TSharedPtrTS<FDashMPD_AdaptationSetType>>& MPDAdaptationSets = MPDPeriod->GetAdaptationSets();
 		int32 nAdapt=0, nAdaptMax=MPDAdaptationSets.Num();
@@ -1403,14 +1410,8 @@ void FManifestDASHInternal::PreparePeriodAdaptationSets(TSharedPtrTS<FPeriod> Pe
 					// If not then there needs to be one on the AdaptationSet. However, this will specify the highest profile and level
 					// necessary to decode any and all representations and is thus potentially too restrictive.
 					MPDCodecs = MPDAdaptationSet->GetCodecs();
-					// Need to have a codec now.
-					if (MPDCodecs.Num() == 0)
-					{
-						LogMessage(PlayerSessionServices, IInfoLog::ELevel::Info, FString::Printf(TEXT("Neither @codecs found on Representation or AdaptationSet level, ignoring this Representation.")));
-						continue;
-					}
 				}
-				// Parse each of the codecs. There *should* be only one since we are not considering multiplexed streams (ContentComponent / SubRepresentation).
+				// There *should* be only one codec since we are not considering multiplexed streams (ContentComponent / SubRepresentation).
 				if (MPDCodecs.Num() > 1)
 				{
 					LogMessage(PlayerSessionServices, IInfoLog::ELevel::Info, FString::Printf(TEXT("More than one codec found for Representation, using only first codec.")));
@@ -1419,6 +1420,48 @@ void FManifestDASHInternal::PreparePeriodAdaptationSets(TSharedPtrTS<FPeriod> Pe
 
 				TSharedPtrTS<FRepresentation> Representation = MakeSharedTS<FRepresentation>();
 				Representation->Representation = MPDRepresentation;
+
+				// Before checking for supported codecs we need to check if this is probably a subtitle representation.
+				// These can come in several flavors and several TTML profiles for which checking just the codec is not
+				// likely to cover all cases.
+				FString MimeType = MPDRepresentation->GetMimeType();
+				if (MimeType.IsEmpty())
+				{
+					MimeType = MPDAdaptationSet->GetMimeType();
+				}
+				// Mime type is mandatory on representation or adaptation set level. So if there is none we may ignore this representation as it's an authoring error.
+				if (MimeType.IsEmpty())
+				{
+					LogMessage(PlayerSessionServices, IInfoLog::ELevel::Info, FString::Printf(TEXT("Could not mandatory Representation@mimeType or the enclosing AdaptationSet. Ignoring this Representation.")));
+					continue;
+				}
+				// Sadly we cannot rely on either @contentType or a Role to be set.
+				if (MimeType.Equals(SubtitleMimeType_Streamed))
+				{
+					// application/mp4 needs to have @codecs set.
+				}
+				else if (MimeType.Equals(SubtitleMimeType_SideloadedTTML))
+				{
+					if (MPDCodecs.Num() == 0)
+					{
+						MPDCodecs.Emplace(TEXT("stpp"));
+					}
+					Representation->bIsSideloadedSubtitle = true;
+				}
+				else if (MimeType.Equals(SubtitleMimeType_SideloadedVTT))
+				{
+					if (MPDCodecs.Num() == 0)
+					{
+						MPDCodecs.Emplace(TEXT("wvtt"));
+					}
+					Representation->bIsSideloadedSubtitle = true;
+				}
+
+				if (MPDCodecs.Num() == 0)
+				{
+					LogMessage(PlayerSessionServices, IInfoLog::ELevel::Info, FString::Printf(TEXT("Neither @codecs found on Representation or AdaptationSet level, ignoring this Representation.")));
+					continue;
+				}
 				if (!Representation->CodecInfo.ParseFromRFC6381(MPDCodecs[0]))
 				{
 					LogMessage(PlayerSessionServices, IInfoLog::ELevel::Info, FString::Printf(TEXT("Could not parse Representation@codecs \"%s\", possibly unsupported codec. Ignoring this Representation."), *MPDCodecs[0]));
@@ -1465,6 +1508,13 @@ void FManifestDASHInternal::PreparePeriodAdaptationSets(TSharedPtrTS<FPeriod> Pe
 						Representation->CodecInfo.SetAspectRatio(FStreamCodecInformation::FAspectRatio(MPDAdaptationSet->GetSAR().GetNumerator(), MPDAdaptationSet->GetSAR().GetDenominator()));
 					}
 
+					// Can be used?
+					if (StreamFilter && !StreamFilter->CanDecodeStream(Representation->CodecInfo))
+					{
+						LogMessage(PlayerSessionServices, IInfoLog::ELevel::Verbose, FString::Printf(TEXT("Video representation \"%s\" in Period \"%s\" rejected by application."), *MPDRepresentation->GetID(), *Period->GetID()));
+						continue;
+					}
+
 					// Update the video codec in the adaptation set with that of the highest bandwidth.
 					if (MPDRepresentation->GetBandwidth() > AdaptationSet->MaxBandwidth)
 					{
@@ -1532,6 +1582,13 @@ void FManifestDASHInternal::PreparePeriodAdaptationSets(TSharedPtrTS<FPeriod> Pe
 						}
 					}
 
+					// Can be used?
+					if (StreamFilter && !StreamFilter->CanDecodeStream(Representation->CodecInfo))
+					{
+						LogMessage(PlayerSessionServices, IInfoLog::ELevel::Verbose, FString::Printf(TEXT("Audio representation \"%s\" in Period \"%s\" rejected by application."), *MPDRepresentation->GetID(), *Period->GetID()));
+						continue;
+					}
+
 					// Update the audio codec in the adaptation set with that of the highest bandwidth.
 					if (MPDRepresentation->GetBandwidth() > AdaptationSet->MaxBandwidth)
 					{
@@ -1541,26 +1598,40 @@ void FManifestDASHInternal::PreparePeriodAdaptationSets(TSharedPtrTS<FPeriod> Pe
 				}
 				else if (Representation->CodecInfo.IsSubtitleCodec())
 				{
-					// ...
+					// There is a possibility that the MPD uses "video/mp4" instead of "application/mp4".
+					if (MimeType.Equals(TEXT("video/mp4")))
+					{
+						MimeType = SubtitleMimeType_Streamed;
+					}
+					// Override the mime type. This is needed for sideloaded subtitles.
+					Representation->CodecInfo.SetMimeType(MimeType);
 
-					if (MPDRepresentation->GetBandwidth() > AdaptationSet->MaxBandwidth)
+					// Can be used?
+					// For sideloaded subtitles we only need the mime type while for regular subtitles we need the codec.
+					FString CodecMimeType, CodecName;
+					if (Representation->IsSideloadedSubtitle())
+					{
+						CodecMimeType = Representation->GetCodecInformation().GetMimeType();
+					}
+					else
+					{
+						CodecName = Representation->GetCodecInformation().GetCodecSpecifierRFC6381();
+					}
+					if (StreamFilter && !StreamFilter->CanDecodeSubtitle(CodecMimeType, CodecName))
+					{
+						LogMessage(PlayerSessionServices, IInfoLog::ELevel::Verbose, FString::Printf(TEXT("Subtitle representation \"%s\" in Period \"%s\" cannot be decoded, ignoring."), *MPDRepresentation->GetID(), *Period->GetID()));
+						continue;
+					}
+
+					if (Representation->bIsSideloadedSubtitle || MPDRepresentation->GetBandwidth() > AdaptationSet->MaxBandwidth)
 					{
 						AdaptationSet->MaxBandwidth = (int32) MPDRepresentation->GetBandwidth();
+						AdaptationSet->Codec = Representation->CodecInfo;
 					}
-					AdaptationSet->Codec = Representation->CodecInfo;
 				}
 				else
 				{
-					// ... ?
-				}
-
-				// Let the player decide if this representation can actually be used.
-				bool bCanDecodeStream = true;
-				IPlayerStreamFilter* StreamFilter = PlayerSessionServices->GetStreamFilter();
-				if (StreamFilter && !StreamFilter->CanDecodeStream(Representation->CodecInfo))
-				{
-					// Skip it then.
-					LogMessage(PlayerSessionServices, IInfoLog::ELevel::Verbose, FString::Printf(TEXT("Representation \"%s\" in Period \"%s\" rejected by application."), *MPDRepresentation->GetID(), *Period->GetID()));
+					LogMessage(PlayerSessionServices, IInfoLog::ELevel::Verbose, FString::Printf(TEXT("Unknown type of representation \"%s\" in Period \"%s\", ignoring."), *MPDRepresentation->GetID(), *Period->GetID()));
 					continue;
 				}
 
diff --git a/Engine/Plugins/Media/ElectraPlayer/Source/ElectraPlayerRuntime/Private/Runtime/Player/DASH/ManifestDASH.cpp b/Engine/Plugins/Media/ElectraPlayer/Source/ElectraPlayerRuntime/Private/Runtime/Player/DASH/ManifestDASH.cpp
--- a/Engine/Plugins/Media/ElectraPlayer/Source/ElectraPlayerRuntime/Private/Runtime/Player/DASH/ManifestDASH.cpp
+++ b/Engine/Plugins/Media/ElectraPlayer/Source/ElectraPlayerRuntime/Private/Runtime/Player/DASH/ManifestDASH.cpp
@@ -177,6 +177,7 @@ private:
 
 	FString ActiveVideoRepresentationID;
 	FString ActiveAudioRepresentationID;
+	FString ActiveSubtitleRepresentationID;
 
 	TSharedPtrTS<FBufferSourceInfo> SourceBufferInfoVideo;
 	TSharedPtrTS<FBufferSourceInfo> SourceBufferInfoAudio;
@@ -479,6 +480,10 @@ void FDASHPlayPeriod::SetStreamPreferences(EStreamType ForStreamType, const FStr
 	{
 		AudioStreamPreferences = StreamAttributes;
 	}
+	else if (ForStreamType == EStreamType::Subtitle)
+	{
+		SubtitleStreamPreferences = StreamAttributes;
+	}
 }
 
 IManifest::IPlayPeriod::EReadyState FDASHPlayPeriod::GetReadyState()
@@ -560,7 +565,7 @@ TSharedPtrTS<FManifestDASHInternal::FAdaptationSet> FDASHPlayPeriod::SelectAdapt
 				AS = StaticCastSharedPtr<FManifestDASHInternal::FAdaptationSet>(Period->GetAdaptationSetByTypeAndIndex(StreamType, SelectedTypeIndex));
 			}
 		}
-		else if (StreamType == EStreamType::Audio)
+		else if (StreamType == EStreamType::Audio || StreamType == EStreamType::Subtitle)
 		{
 			// Check for a matching language.
 			// For now we ignore the track kind.
@@ -587,14 +592,14 @@ TSharedPtrTS<FManifestDASHInternal::FAdaptationSet> FDASHPlayPeriod::SelectAdapt
 			// Matching language not found. Is there an explicit index given?
 			if (!AS.IsValid())
 			{
-				if (Attributes.OverrideIndex.IsSet() && Attributes.OverrideIndex.GetValue() < NumAdaptationSets)
+				if (Attributes.OverrideIndex.IsSet() && Attributes.OverrideIndex.GetValue() >= 0 && Attributes.OverrideIndex.GetValue() < NumAdaptationSets)
 				{
 					SelectedTypeIndex = Attributes.OverrideIndex.GetValue();
 					AS = StaticCastSharedPtr<FManifestDASHInternal::FAdaptationSet>(Period->GetAdaptationSetByTypeAndIndex(StreamType, SelectedTypeIndex));
 				}
 			}
-			// Still nothing? Use the first one.
-			if (!AS.IsValid())
+			// Still nothing? Use the first one, except for subtitles that need to be explicitly selected.
+			if (!AS.IsValid() && StreamType != EStreamType::Subtitle)
 			{
 				Selection.Empty();
 				for(int32 i=0; i<NumAdaptationSets; ++i)
@@ -707,6 +712,29 @@ void FDASHPlayPeriod::Load()
 			}
 		}
 
+		TSharedPtrTS<FManifestDASHInternal::FAdaptationSet> SubtitleAS = SelectAdaptationSetByAttributes(SourceBufferInfoSubtitles, Period, EStreamType::Subtitle, SubtitleStreamPreferences);
+		if (SubtitleAS.IsValid())
+		{
+			ActiveSubtitleAdaptationSetID = SubtitleAS->GetUniqueIdentifier();
+
+			// Add encryption schemes, if any.
+			if (SubtitleAS->GetIsSwitchGroup())
+			{
+				for(auto &SwitchedID : SubtitleAS->GetSwitchToSetIDs())
+				{
+					TSharedPtrTS<FManifestDASHInternal::FAdaptationSet> SwitchedAS = Period->GetAdaptationSetByUniqueID(SwitchedID);
+					if (SwitchedAS.IsValid())
+					{
+						ContentProtections.Append(SwitchedAS->GetPossibleContentProtections());
+					}
+				}
+			}
+			else
+			{
+				ContentProtections.Append(SubtitleAS->GetPossibleContentProtections());
+			}
+		}
+
 		// Prepare the DRM system for decryption.
 		if (PrepareDRM(ContentProtections))
 		{
@@ -748,6 +776,13 @@ void FDASHPlayPeriod::PrepareForPlay()
 			ActiveAudioRepresentationID = AudioRepr->GetUniqueIdentifier();
 		}
 
+		TSharedPtrTS<FManifestDASHInternal::FAdaptationSet> SubtitleAS = Period->GetAdaptationSetByUniqueID(ActiveSubtitleAdaptationSetID);
+		if (SubtitleAS.IsValid())
+		{
+			TSharedPtrTS<IPlaybackAssetRepresentation> SubtitleRepr = GetRepresentationFromAdaptationByPriorityAndMaxBandwidth(SubtitleAS, 256 * 1000, EStreamType::Subtitle);
+			ActiveSubtitleRepresentationID = SubtitleRepr->GetUniqueIdentifier();
+		}
+
 		// Emit all <EventStream> events of the period to the AEMS event handler.
 		Manifest->SendEventsFromAllPeriodEventStreams(Period);
 
@@ -858,6 +893,10 @@ void FDASHPlayPeriod::SelectStream(const FString& AdaptationSetID, const FString
 	{
 		ActiveAudioRepresentationID = RepresentationID;
 	}
+	else if (AdaptationSetID == ActiveSubtitleAdaptationSetID)
+	{
+		ActiveSubtitleRepresentationID = RepresentationID;
+	}
 	else
 	{
 		LogMessage(PlayerSessionServices, IInfoLog::ELevel::Warning, FString::Printf(TEXT("ABR tried to activate a stream from an inactive AdaptationSet!")));
@@ -934,6 +973,10 @@ IManifest::FResult FDASHPlayPeriod::GetStartingSegment(TSharedPtrTS<IStreamSegme
 	{
 		ActiveSelection.Emplace(FSelectedStream({EStreamType::Audio, ActiveAudioRepresentationID, ActiveAudioAdaptationSetID}));
 	}
+	if (!ActiveSubtitleRepresentationID.IsEmpty())
+	{
+		ActiveSelection.Emplace(FSelectedStream({EStreamType::Subtitle, ActiveSubtitleRepresentationID, ActiveSubtitleAdaptationSetID}));
+	}
 
 	bool bDidAdjustStartTime = false;
 	bool bTryAgainLater = false;
@@ -1138,6 +1181,10 @@ IManifest::FResult FDASHPlayPeriod::GetNextOrRetrySegment(TSharedPtrTS<IStreamSe
 			ActiveAdaptationSetIDByType = ActiveAudioAdaptationSetID;
 			ActiveRepresentationIDByType = ActiveAudioRepresentationID;
 			break;
+		case EStreamType::Subtitle:
+			ActiveAdaptationSetIDByType = ActiveSubtitleAdaptationSetID;
+			ActiveRepresentationIDByType = ActiveSubtitleRepresentationID;
+			break;
 		default:
 			break;
 	}
diff --git a/Engine/Plugins/Media/ElectraPlayer/Source/ElectraPlayerRuntime/Private/Runtime/Player/DASH/PlaylistReaderDASH_Internal.h b/Engine/Plugins/Media/ElectraPlayer/Source/ElectraPlayerRuntime/Private/Runtime/Player/DASH/PlaylistReaderDASH_Internal.h
--- a/Engine/Plugins/Media/ElectraPlayer/Source/ElectraPlayerRuntime/Private/Runtime/Player/DASH/PlaylistReaderDASH_Internal.h
+++ b/Engine/Plugins/Media/ElectraPlayer/Source/ElectraPlayerRuntime/Private/Runtime/Player/DASH/PlaylistReaderDASH_Internal.h
@@ -229,6 +229,7 @@ public:
 
 		int32 GetSelectionPriority() const { return SelectionPriority; }
 
+		bool IsSideloadedSubtitle() const { return bIsSideloadedSubtitle; }
 
 		//----------------------------------------------
 		// Methods from IPlaybackAssetRepresentation
@@ -288,6 +289,8 @@ public:
 		int64 SegmentIndexRangeStart = 0;
 		int64 SegmentIndexRangeSize = 0;
 		TSharedPtrTS<FMPDLoadRequestDASH> PendingSegmentIndexLoadRequest;
+		//
+		bool bIsSideloadedSubtitle = false;
 	};
 
 	class FAdaptationSet : public IPlaybackAssetAdaptationSet
@@ -418,10 +421,14 @@ public:
 				{
 					InOutMetadata.Kind = TEXT("metadata");
 				}
-				// TODO: ID and language (complicated by the CEA services)
-				check(!"need to add ID and language handling");
-				InOutMetadata.ID = TEXT("CC1");
-				InOutMetadata.Language = TEXT("en");
+				// ID and language for CEA services
+				if (bIsCEAService)
+				{
+					// TODO
+					check(!"need to add ID and language handling");
+					InOutMetadata.ID = TEXT("CC1");
+					InOutMetadata.Language = TEXT("en");
+				}
 			}
 		}
 
diff --git a/Engine/Plugins/Media/ElectraPlayer/Source/ElectraPlayerRuntime/Private/Runtime/Player/DASH/StreamReaderFMP4DASH.cpp b/Engine/Plugins/Media/ElectraPlayer/Source/ElectraPlayerRuntime/Private/Runtime/Player/DASH/StreamReaderFMP4DASH.cpp
--- a/Engine/Plugins/Media/ElectraPlayer/Source/ElectraPlayerRuntime/Private/Runtime/Player/DASH/StreamReaderFMP4DASH.cpp
+++ b/Engine/Plugins/Media/ElectraPlayer/Source/ElectraPlayerRuntime/Private/Runtime/Player/DASH/StreamReaderFMP4DASH.cpp
@@ -16,6 +16,7 @@
 #include "Player/DASH/PlayerEventDASH_Internal.h"
 #include "Player/DRM/DRMManager.h"
 #include "Utilities/Utilities.h"
+#include "Async/Async.h"
 
 #define INTERNAL_ERROR_INIT_SEGMENT_DOWNLOAD_ERROR					1
 #define INTERNAL_ERROR_INIT_SEGMENT_PARSE_ERROR						2
@@ -49,19 +50,37 @@ void FStreamSegmentRequestFMP4DASH::SetExecutionDelay(const FTimeValue& Executio
 	DownloadDelayTime = ExecutionDelay;
 }
 
+FTimeValue FStreamSegmentRequestFMP4DASH::GetExecuteAtUTCTime() const
+{
+	FTimeValue When = ASAST;
+	if (DownloadDelayTime.IsValid())
+	{
+		When += DownloadDelayTime;
+	}
+	return When;
+}
+
 
 EStreamType FStreamSegmentRequestFMP4DASH::GetType() const
 {
 	return StreamType;
 }
 
-void FStreamSegmentRequestFMP4DASH::GetDependentStreams(TArray<FDependentStreams>& OutDependentStreams) const
+void FStreamSegmentRequestFMP4DASH::GetDependentStreams(TArray<TSharedPtrTS<IStreamSegment>>& OutDependentStreams) const
 {
 	OutDependentStreams.Empty();
-	for(int32 i=0; i<DependentStreams.Num(); ++i)
+	for(auto& Stream : DependentStreams)
 	{
-		FDependentStreams& depStr = OutDependentStreams.AddDefaulted_GetRef();
-		depStr.StreamType = DependentStreams[i]->GetType();
+		OutDependentStreams.Emplace(Stream);
+	}
+}
+
+void FStreamSegmentRequestFMP4DASH::GetRequestedStreams(TArray<TSharedPtrTS<IStreamSegment>>& OutRequestedStreams)
+{
+	OutRequestedStreams.Empty();
+	for(auto& Stream : DependentStreams)
+	{
+		OutRequestedStreams.Emplace(Stream);
 	}
 }
 
@@ -139,7 +158,15 @@ UEMediaError FStreamReaderFMP4DASH::Create(IPlayerSessionServices* InPlayerSessi
 		StreamHandlers[i].bHasErrored   	   = false;
 		StreamHandlers[i].IsIdleSignal.Signal();
 
-		StreamHandlers[i].ThreadSetName(i==0?"ElectraPlayer::fmp4 Video":"ElectraPlayer::fmp4 Audio");
+		// Subtitles get fetched running on the thread pool.
+		if (i == 2)
+		{
+			StreamHandlers[i].bRunOnThreadPool = true;
+		}
+
+		StreamHandlers[i].ThreadSetName(i==0 ? "ElectraPlayer::fmp4 Video" : 
+										i==1 ? "ElectraPlayer::fmp4 Audio" :
+											   "ElectraPlayer::fmp4 Subtitle");
 	}
 	return UEMEDIA_ERROR_OK;
 }
@@ -159,7 +186,11 @@ void FStreamReaderFMP4DASH::Close()
 		// Wait until they finished.
 		for(int32 i=0; i<FMEDIA_STATIC_ARRAY_COUNT(StreamHandlers); ++i)
 		{
-			if (StreamHandlers[i].bWasStarted)
+			if (StreamHandlers[i].bRunOnThreadPool)
+			{
+				StreamHandlers[i].IsIdleSignal.Wait();
+			}
+			else if (StreamHandlers[i].bWasStarted)
 			{
 				StreamHandlers[i].ThreadWaitDone();
 				StreamHandlers[i].ThreadReset();
@@ -174,27 +205,12 @@ IStreamReader::EAddResult FStreamReaderFMP4DASH::AddRequest(uint32 CurrentPlayba
 
 	if (Request->bIsInitialStartRequest)
 	{
-		for(int32 i=0; i<Request->DependentStreams.Num(); ++i)
-		{
-			EAddResult Result = AddRequest(CurrentPlaybackSequenceID, Request->DependentStreams[i]);
-			if (Result != IStreamReader::EAddResult::Added)
-			{
-				return Result;
-			}
-		}
+		PostError(PlayerSessionService, TEXT("Initial start request segments cannot be enqueued!"), 0);
+		return IStreamReader::EAddResult::Error;
 	}
 	else
 	{
-		// Video and audio only for now.
-		if (Request->GetType() != EStreamType::Video && Request->GetType() != EStreamType::Audio)
-		{
-			ErrorDetail.SetMessage(FString::Printf(TEXT("Request is not for video or audio")));
-			return IStreamReader::EAddResult::Error;
-		}
-	
 		// Get the handler for the main request.
-		// TODO: We could have a pool of handlers or create a handler on demand if necessary.
-		//       Alternatively think about adding a local queue but be wary about request cancellation.
 		FStreamHandler* Handler = nullptr;
 		switch(Request->GetType())
 		{
@@ -204,6 +220,9 @@ IStreamReader::EAddResult FStreamReaderFMP4DASH::AddRequest(uint32 CurrentPlayba
 			case EStreamType::Audio:
 				Handler = &StreamHandlers[1];
 				break;
+			case EStreamType::Subtitle:
+				Handler = &StreamHandlers[2];
+				break;
 			default:
 				break;
 		}
@@ -221,10 +240,21 @@ IStreamReader::EAddResult FStreamReaderFMP4DASH::AddRequest(uint32 CurrentPlayba
 		}
 
 		Request->SetPlaybackSequenceID(CurrentPlaybackSequenceID);
-		if (!Handler->bWasStarted)
+		if (Handler->bRunOnThreadPool)
 		{
-			Handler->ThreadStart(Electra::MakeDelegate(Handler, &FStreamHandler::WorkerThread));
-			Handler->bWasStarted = true;
+			Handler->IsIdleSignal.Reset();
+			Async(EAsyncExecution::ThreadPool, [Handler]()
+			{
+				Handler->RunInThreadPool();
+			});
+		}
+		else
+		{
+			if (!Handler->bWasStarted)
+			{
+				Handler->ThreadStart(Electra::MakeDelegate(Handler, &FStreamHandler::WorkerThread));
+				Handler->bWasStarted = true;
+			}
 		}
 
 		Handler->bRequestCanceled = false;
@@ -245,6 +275,10 @@ void FStreamReaderFMP4DASH::CancelRequest(EStreamType StreamType, bool bSilent)
 	{
 		StreamHandlers[1].Cancel(bSilent);
 	}
+	else if (StreamType == EStreamType::Subtitle)
+	{
+		StreamHandlers[2].Cancel(bSilent);
+	}
 }
 
 void FStreamReaderFMP4DASH::CancelRequests()
@@ -267,6 +301,7 @@ FStreamReaderFMP4DASH::FStreamHandler::FStreamHandler()
 FStreamReaderFMP4DASH::FStreamHandler::~FStreamHandler()
 {
 	// NOTE: The thread will have been terminated by the enclosing FStreamReaderFMP4DASH's Close() method!
+	//       Also, this may have run on the thread pool instead of a dedicated worker thread.
 }
 
 void FStreamReaderFMP4DASH::FStreamHandler::Cancel(bool bSilent)
@@ -311,6 +346,34 @@ void FStreamReaderFMP4DASH::FStreamHandler::WorkerThread()
 	StreamSelector.Reset();
 }
 
+void FStreamReaderFMP4DASH::FStreamHandler::RunInThreadPool()
+{
+	LLM_SCOPE(ELLMTag::ElectraPlayer);
+
+	StreamSelector = PlayerSessionService->GetStreamSelector();
+	check(StreamSelector.IsValid());
+	WorkSignal.Obtain();
+	if (!bTerminate)
+	{
+		if (CurrentRequest.IsValid())
+		{
+			if (!bRequestCanceled)
+			{
+				HandleRequest();
+			}
+			else
+			{
+				CurrentRequest.Reset();
+			}
+		}
+		bRequestCanceled = false;
+		bSilentCancellation = false;
+	}
+	StreamSelector.Reset();
+	IsIdleSignal.Signal();
+}
+
+
 
 void FStreamReaderFMP4DASH::FStreamHandler::LogMessage(IInfoLog::ELevel Level, const FString& Message)
 {
@@ -668,35 +731,6 @@ void FStreamReaderFMP4DASH::FStreamHandler::HandleRequest()
 		return;
 	}
 
-	// Does this request have an availability window?
-	if (Request->ASAST.IsValid())
-	{
-		FTimeValue Now = PlayerSessionService->GetSynchronizedUTCTime()->GetTime();
-		FTimeValue When = Request->ASAST;
-		if (Request->DownloadDelayTime.IsValid())
-		{
-			When += Request->DownloadDelayTime;
-		}
-		// Availability start time not reached yet?
-		while(Now < When)
-		{
-			if (HasReadBeenAborted())
-			{
-				CurrentRequest.Reset();
-				return;
-			}
-			const int64 SleepIntervalMS = 100;
-			FMediaRunnable::SleepMilliseconds((uint32) Utils::Max(Utils::Min((When - Now).GetAsMilliseconds(), SleepIntervalMS), (int64)0));
-			Now = PlayerSessionService->GetSynchronizedUTCTime()->GetTime();
-		}
-
-		// Already expired?
-		if (Now > Request->SAET)
-		{
-			// TODO:
-		}
-	}
-
 	// Clear internal work variables.
 	bHasErrored 		   = false;
 	bAbortedByABR   	   = false;
diff --git a/Engine/Plugins/Media/ElectraPlayer/Source/ElectraPlayerRuntime/Private/Runtime/Player/DASH/StreamReaderFMP4DASH.h b/Engine/Plugins/Media/ElectraPlayer/Source/ElectraPlayerRuntime/Private/Runtime/Player/DASH/StreamReaderFMP4DASH.h
--- a/Engine/Plugins/Media/ElectraPlayer/Source/ElectraPlayerRuntime/Private/Runtime/Player/DASH/StreamReaderFMP4DASH.h
+++ b/Engine/Plugins/Media/ElectraPlayer/Source/ElectraPlayerRuntime/Private/Runtime/Player/DASH/StreamReaderFMP4DASH.h
@@ -23,11 +23,12 @@ public:
 	virtual uint32 GetPlaybackSequenceID() const override;
 
 	virtual void SetExecutionDelay(const FTimeValue& ExecutionDelay) override;
+	virtual FTimeValue GetExecuteAtUTCTime() const override;
 
 	virtual EStreamType GetType() const override;
 
-	virtual void GetDependentStreams(TArray<FDependentStreams>& OutDependentStreams) const override;
-
+	virtual void GetDependentStreams(TArray<TSharedPtrTS<IStreamSegment>>& OutDependentStreams) const override;
+	virtual void GetRequestedStreams(TArray<TSharedPtrTS<IStreamSegment>>& OutRequestedStreams) override;
 	virtual void GetEndedStreams(TArray<TSharedPtrTS<IStreamSegment>>& OutAlreadyEndedStreams) override;
 
 	//! Returns the first PTS value as indicated by the media timeline. This should correspond to the actual absolute PTS of the sample.
@@ -125,6 +126,7 @@ private:
 		TSharedPtrTS<FStreamSegmentRequestFMP4DASH>				CurrentRequest;
 		FMediaSemaphore											WorkSignal;
 		FMediaEvent												IsIdleSignal;
+		bool													bRunOnThreadPool = false;
 		volatile bool											bTerminate = false;
 		volatile bool											bWasStarted = false;
 		volatile bool											bRequestCanceled = false;
@@ -156,6 +158,7 @@ private:
 		void Cancel(bool bSilent);
 		void SignalWork();
 		void WorkerThread();
+		void RunInThreadPool();
 		void HandleRequest();
 
 		FErrorDetail GetInitSegment(TSharedPtrTS<const IParserISO14496_12>& OutMP4InitSegment, const TSharedPtrTS<FStreamSegmentRequestFMP4DASH>& InRequest);
@@ -180,9 +183,7 @@ private:
 		virtual IParserISO14496_12::IBoxCallback::EParseContinuation OnEndOfBox(IParserISO14496_12::FBoxType Box, int64 BoxSizeInBytes, int64 FileDataOffset, int64 BoxDataOffset) override;
 	};
 
-	// Currently set to use 2 handlers, one for video and one for audio. This could become a pool of n if we need to stream
-	// multiple dependent segments, keeping a pool of available and active handlers to cycle between.
-	FStreamHandler						StreamHandlers[2];		// 0 = video (MEDIAstreamType_Video), 1 = audio (MEDIAstreamType_Audio)
+	FStreamHandler						StreamHandlers[3];		// 0 = video, 1 = audio, 2 = subtitle 
 	IPlayerSessionServices*				PlayerSessionService = nullptr;
 	bool								bIsStarted = false;
 	FErrorDetail						ErrorDetail;
diff --git a/Engine/Plugins/Media/ElectraPlayer/Source/ElectraPlayerRuntime/Private/Runtime/Player/HLS/StreamReaderHLSfmp4.cpp b/Engine/Plugins/Media/ElectraPlayer/Source/ElectraPlayerRuntime/Private/Runtime/Player/HLS/StreamReaderHLSfmp4.cpp
--- a/Engine/Plugins/Media/ElectraPlayer/Source/ElectraPlayerRuntime/Private/Runtime/Player/HLS/StreamReaderHLSfmp4.cpp
+++ b/Engine/Plugins/Media/ElectraPlayer/Source/ElectraPlayerRuntime/Private/Runtime/Player/HLS/StreamReaderHLSfmp4.cpp
@@ -57,21 +57,38 @@ void FStreamSegmentRequestHLSfmp4::SetExecutionDelay(const FTimeValue& Execution
 	// No-op for HLS.
 }
 
+FTimeValue FStreamSegmentRequestHLSfmp4::GetExecuteAtUTCTime() const
+{
+	// Right now.
+	return FTimeValue::GetInvalid();
+}
+
+
 EStreamType FStreamSegmentRequestHLSfmp4::GetType() const
 {
 	return StreamType;
 }
 
-void FStreamSegmentRequestHLSfmp4::GetDependentStreams(TArray<FDependentStreams>& OutDependentStreams) const
+void FStreamSegmentRequestHLSfmp4::GetDependentStreams(TArray<TSharedPtrTS<IStreamSegment>>& OutDependentStreams) const
 {
 	OutDependentStreams.Empty();
-	for(int32 i=0; i<DependentStreams.Num(); ++i)
+	for(auto& Stream : DependentStreams)
+	{
+		OutDependentStreams.Emplace(Stream);
+	}
+}
+
+void FStreamSegmentRequestHLSfmp4::GetRequestedStreams(TArray<TSharedPtrTS<IStreamSegment>>& OutRequestedStreams)
+{
+	OutRequestedStreams.Empty();
+	OutRequestedStreams.Push(SharedThis(this));
+	for(auto& Stream : DependentStreams)
 	{
-		FDependentStreams& depStr = OutDependentStreams.AddDefaulted_GetRef();
-		depStr.StreamType = DependentStreams[i]->GetType();
+		OutRequestedStreams.Emplace(Stream);
 	}
 }
 
+
 void FStreamSegmentRequestHLSfmp4::GetEndedStreams(TArray<TSharedPtrTS<IStreamSegment>>& OutAlreadyEndedStreams)
 {
 	OutAlreadyEndedStreams.Empty();
@@ -179,6 +196,9 @@ IStreamReader::EAddResult FStreamReaderHLSfmp4::AddRequest(uint32 CurrentPlaybac
 {
 	TSharedPtrTS<FStreamSegmentRequestHLSfmp4> Request = StaticCastSharedPtr<FStreamSegmentRequestHLSfmp4>(InRequest);
 
+	// Dependent streams will be added individually. Anything that is still set here must be ignored from this point forth.
+	Request->DependentStreams.Empty();
+
 	// Video and audio only for now.
 	if (Request->GetType() != EStreamType::Video && Request->GetType() != EStreamType::Audio)
 	{
@@ -186,22 +206,8 @@ IStreamReader::EAddResult FStreamReaderHLSfmp4::AddRequest(uint32 CurrentPlaybac
 		ErrorDetail.SetMessage(FString::Printf(TEXT("Request is not video or audio")));
 		return IStreamReader::EAddResult::Error;
 	}
-	// Only initial requests are allowed to have a dependent stream for now.
-	if (Request->DependentStreams.Num() && !Request->bIsInitialStartRequest)
-	{
-		check(!"no good");
-		ErrorDetail.SetMessage(FString::Printf(TEXT("Dependent streams only allowed for initial request")));
-		return IStreamReader::EAddResult::Error;
-	}
-	// Also, there may only be one dependent stream.
-	if (Request->DependentStreams.Num() > 1)
-	{
-		check(!"no good");
-		ErrorDetail.SetMessage(FString::Printf(TEXT("Only one dependent streams allowed")));
-		return IStreamReader::EAddResult::Error;
-	}
 
-	// Get the handler for the main request.
+	// Get the handler for the request.
 	FStreamHandler* Handler = nullptr;
 	switch(Request->GetType())
 	{
@@ -228,52 +234,6 @@ IStreamReader::EAddResult FStreamReaderHLSfmp4::AddRequest(uint32 CurrentPlaybac
 		return IStreamReader::EAddResult::TryAgainLater;
 	}
 
-// TODO: make this work with more than just one request (asserted above) if necessary
-	if (Request->DependentStreams.Num())
-	{
-		TSharedPtrTS<FStreamSegmentRequestHLSfmp4> Request2 = Request->DependentStreams[0];
-		FStreamHandler* Handler2 = nullptr;
-		switch(Request2->GetType())
-		{
-			case  EStreamType::Video:
-				Handler2 = &StreamHandlers[0];
-				break;
-			case  EStreamType::Audio:
-				Handler2 = &StreamHandlers[1];
-				break;
-			default:
-				check(!"Whoops");
-				break;
-		}
-		if (!Handler2)
-		{
-			check(!"no good");
-			ErrorDetail.SetMessage(FString::Printf(TEXT("No handler for stream type")));
-			return IStreamReader::EAddResult::Error;
-		}
-		// Is the handler busy?
-		if (Handler2->CurrentRequest.IsValid())
-		{
-			check(!"why is the handler busy??");
-			return IStreamReader::EAddResult::TryAgainLater;
-		}
-
-		Request->DependentStreams.Empty();
-		Request2->SetPlaybackSequenceID(CurrentPlaybackSequenceID);
-		// Only add the request if this is not an EOD segment.
-		if (!Request2->bIsEOSSegment)
-		{
-			if (!Handler2->bWasStarted)
-			{
-				Handler2->ThreadStart(Electra::MakeDelegate(Handler2, &FStreamHandler::WorkerThread));
-				Handler2->bWasStarted = true;
-			}
-			Handler2->bRequestCanceled = false;
-			Handler2->bSilentCancellation = false;
-			Handler2->CurrentRequest = Request2;
-			Handler2->SignalWork();
-		}
-	}
 	Request->SetPlaybackSequenceID(CurrentPlaybackSequenceID);
 	// Only add the request if this is not an EOD segment.
 	if (!Request->bIsEOSSegment)
diff --git a/Engine/Plugins/Media/ElectraPlayer/Source/ElectraPlayerRuntime/Private/Runtime/Player/HLS/StreamReaderHLSfmp4.h b/Engine/Plugins/Media/ElectraPlayer/Source/ElectraPlayerRuntime/Private/Runtime/Player/HLS/StreamReaderHLSfmp4.h
--- a/Engine/Plugins/Media/ElectraPlayer/Source/ElectraPlayerRuntime/Private/Runtime/Player/HLS/StreamReaderHLSfmp4.h
+++ b/Engine/Plugins/Media/ElectraPlayer/Source/ElectraPlayerRuntime/Private/Runtime/Player/HLS/StreamReaderHLSfmp4.h
@@ -23,10 +23,12 @@ public:
 	virtual uint32 GetPlaybackSequenceID() const override;
 
 	virtual void SetExecutionDelay(const FTimeValue& ExecutionDelay) override;
+	virtual FTimeValue GetExecuteAtUTCTime() const override;
 
 	virtual EStreamType GetType() const override;
 
-	virtual void GetDependentStreams(TArray<FDependentStreams>& OutDependentStreams) const override;
+	virtual void GetDependentStreams(TArray<TSharedPtrTS<IStreamSegment>>& OutDependentStreams) const override;
+	virtual void GetRequestedStreams(TArray<TSharedPtrTS<IStreamSegment>>& OutRequestedStreams) override;
 	virtual void GetEndedStreams(TArray<TSharedPtrTS<IStreamSegment>>& OutAlreadyEndedStreams) override;
 
 	//! Returns the first PTS value as indicated by the media timeline. This should correspond to the actual absolute PTS of the sample.
diff --git a/Engine/Plugins/Media/ElectraPlayer/Source/ElectraPlayerRuntime/Private/Runtime/Player/Manifest.h b/Engine/Plugins/Media/ElectraPlayer/Source/ElectraPlayerRuntime/Private/Runtime/Player/Manifest.h
--- a/Engine/Plugins/Media/ElectraPlayer/Source/ElectraPlayerRuntime/Private/Runtime/Player/Manifest.h
+++ b/Engine/Plugins/Media/ElectraPlayer/Source/ElectraPlayerRuntime/Private/Runtime/Player/Manifest.h
@@ -55,16 +55,20 @@ namespace Electra
 
 		virtual void SetExecutionDelay(const FTimeValue& ExecutionDelay) = 0;
 
+		virtual FTimeValue GetExecuteAtUTCTime() const = 0;
+
 		virtual EStreamType GetType() const = 0;
 
-		struct FDependentStreams
-		{
-			EStreamType		StreamType;
-		};
 		/**
-		 * Returns a list of dependent stream types.
+		 * Returns a list of dependent streams.
+		 */
+		virtual void GetDependentStreams(TArray<TSharedPtrTS<IStreamSegment>>& OutDependentStreams) const = 0;
+
+		/**
+		 * Returns a list of requested streams. This must be the initial request that has a list of
+		 * all the segments of the streams selected for playback.
 		 */
-		virtual void GetDependentStreams(TArray<FDependentStreams>& OutDependentStreams) const = 0;
+		virtual void GetRequestedStreams(TArray<TSharedPtrTS<IStreamSegment>>& OutRequestedStreams) = 0;
 
 		/**
 		 * Returns a list of streams, primary or dependent, that are already at EOS in this request.
diff --git a/Engine/Plugins/Media/ElectraPlayer/Source/ElectraPlayerRuntime/Private/Runtime/Player/PlayerStreamFilter.h b/Engine/Plugins/Media/ElectraPlayer/Source/ElectraPlayerRuntime/Private/Runtime/Player/PlayerStreamFilter.h
--- a/Engine/Plugins/Media/ElectraPlayer/Source/ElectraPlayerRuntime/Private/Runtime/Player/PlayerStreamFilter.h
+++ b/Engine/Plugins/Media/ElectraPlayer/Source/ElectraPlayerRuntime/Private/Runtime/Player/PlayerStreamFilter.h
@@ -15,6 +15,8 @@ namespace Electra
 		virtual ~IPlayerStreamFilter() = default;
 
 		virtual bool CanDecodeStream(const FStreamCodecInformation& InStreamCodecInfo) const = 0;
+
+		virtual bool CanDecodeSubtitle(const FString& MimeType, const FString& Codec) const = 0;
 	};
 
 
diff --git a/Engine/Plugins/Media/ElectraPlayer/Source/ElectraPlayerRuntime/Private/Runtime/Player/mp4/ManifestMP4.cpp b/Engine/Plugins/Media/ElectraPlayer/Source/ElectraPlayerRuntime/Private/Runtime/Player/mp4/ManifestMP4.cpp
--- a/Engine/Plugins/Media/ElectraPlayer/Source/ElectraPlayerRuntime/Private/Runtime/Player/mp4/ManifestMP4.cpp
+++ b/Engine/Plugins/Media/ElectraPlayer/Source/ElectraPlayerRuntime/Private/Runtime/Player/mp4/ManifestMP4.cpp
@@ -200,6 +200,10 @@ void FManifestMP4Internal::FPlayPeriodMP4::SetStreamPreferences(EStreamType ForS
 	{
 		AudioPreferences = StreamAttributes;
 	}
+	else if (ForStreamType == EStreamType::Subtitle)
+	{
+		SubtitlePreferences = StreamAttributes;
+	}
 }
 
 
@@ -342,7 +346,7 @@ TSharedPtrTS<FTrackMetadata> FManifestMP4Internal::FPlayPeriodMP4::SelectMetadat
 		TArray<FTrackMetadata> Metadata;
 		Asset->GetMetaData(Metadata, StreamType);
 		// Is there a fixed index to be used?
-		if (InAttributes.OverrideIndex.IsSet() && InAttributes.OverrideIndex.GetValue() < Metadata.Num())
+		if (InAttributes.OverrideIndex.IsSet() && InAttributes.OverrideIndex.GetValue() >= 0 && InAttributes.OverrideIndex.GetValue() < Metadata.Num())
 		{
 			// Use this.
 			return MakeSharedTS<FTrackMetadata>(Metadata[InAttributes.OverrideIndex.GetValue()]);
@@ -578,6 +582,9 @@ FErrorDetail FManifestMP4Internal::FTimelineAssetMP4::Build(IPlayerSessionServic
 					case EStreamType::Audio:
 						AudioAdaptationSets.Add(AdaptationSet);
 						break;
+					case EStreamType::Subtitle:
+						SubtitleAdaptationSets.Add(AdaptationSet);
+						break;
 					default:
 						break;
 				}
@@ -705,6 +712,7 @@ IManifest::FResult FManifestMP4Internal::FTimelineAssetMP4::GetStartingSegment(T
 					FTimeFraction TrackDuration = Track->GetDuration();
 					if (TrackDuration.IsValid() && FTimeValue().SetFromTimeFraction(TrackDuration) <= StartPosition.Time && (SearchType == ESearchType::After || SearchType == ESearchType::StrictlyAfter || SearchType == ESearchType::Closest))
 					{
+						firstTimestamp.SetFromTimeFraction(TrackDuration);
 						err = UEMEDIA_ERROR_END_OF_STREAM;
 					}
 					else
@@ -729,6 +737,7 @@ IManifest::FResult FManifestMP4Internal::FTimelineAssetMP4::GetStartingSegment(T
 					}
 					else
 					{
+						firstTimestamp.SetFromTimeFraction(Track->GetDuration());
 						err = UEMEDIA_ERROR_END_OF_STREAM;
 					}
 				}
@@ -754,7 +763,11 @@ IManifest::FResult FManifestMP4Internal::FTimelineAssetMP4::GetStartingSegment(T
 					// FIXME: this may need to add all additional tracks at some point if their individual IDs matter
 					if (AudioAdaptationSets.Num())
 					{
-						req->DependentStreams.AddDefaulted_GetRef().StreamType = EStreamType::Audio;
+						req->DependentStreamTypes.Add(EStreamType::Audio);
+					}
+					if (SubtitleAdaptationSets.Num())
+					{
+						req->DependentStreamTypes.Add(EStreamType::Subtitle);
 					}
 
 					LimitSegmentDownloadSize(OutSegment, nullptr);
@@ -769,7 +782,7 @@ IManifest::FResult FManifestMP4Internal::FTimelineAssetMP4::GetStartingSegment(T
 						TSharedPtrTS<FStreamSegmentRequestMP4> req(new FStreamSegmentRequestMP4);
 						OutSegment = req;
 						req->MediaAsset 			= SharedThis(this);
-						req->FirstPTS   			= FTimeValue::GetZero();
+						req->FirstPTS   			= firstTimestamp;
 						req->PrimaryStreamType  	= EStreamType::Video;
 						req->Bitrate				= Repr->GetBitrate();
 						req->bStartingOnMOOF		= false;
@@ -777,7 +790,6 @@ IManifest::FResult FManifestMP4Internal::FTimelineAssetMP4::GetStartingSegment(T
 						req->bIsFirstSegment		= false;
 						req->bIsLastSegment			= true;
 						req->bAllTracksAtEOS		= true;
-						//req->DependentStreams.PushBack().StreamType = EStreamType::Audio;
 						return IManifest::FResult(IManifest::FResult::EType::Found);
 					}
 				}
@@ -825,6 +837,7 @@ IManifest::FResult FManifestMP4Internal::FTimelineAssetMP4::GetStartingSegment(T
 					FTimeFraction TrackDuration = Track->GetDuration();
 					if (TrackDuration.IsValid() && FTimeValue().SetFromTimeFraction(TrackDuration) <= StartPosition.Time && (SearchType == ESearchType::After || SearchType == ESearchType::StrictlyAfter || SearchType == ESearchType::Closest))
 					{
+						firstTimestamp.SetFromTimeFraction(TrackDuration);
 						err = UEMEDIA_ERROR_END_OF_STREAM;
 					}
 					else
@@ -849,6 +862,7 @@ IManifest::FResult FManifestMP4Internal::FTimelineAssetMP4::GetStartingSegment(T
 					}
 					else
 					{
+						firstTimestamp.SetFromTimeFraction(Track->GetDuration());
 						err = UEMEDIA_ERROR_END_OF_STREAM;
 					}
 				}
@@ -875,7 +889,11 @@ IManifest::FResult FManifestMP4Internal::FTimelineAssetMP4::GetStartingSegment(T
 					// (if it exists) in case the video will loop back to a point where there is video.
 					if (VideoAdaptationSets.Num())
 					{
-						req->DependentStreams.AddDefaulted_GetRef().StreamType = EStreamType::Video;
+						req->DependentStreamTypes.Add(EStreamType::Video);
+					}
+					if (SubtitleAdaptationSets.Num())
+					{
+						req->DependentStreamTypes.Add(EStreamType::Subtitle);
 					}
 
 					// FIXME: there may be subtitle tracks here we need to add as dependent streams.
@@ -889,7 +907,7 @@ IManifest::FResult FManifestMP4Internal::FTimelineAssetMP4::GetStartingSegment(T
 					TSharedPtrTS<FStreamSegmentRequestMP4> req(new FStreamSegmentRequestMP4);
 					OutSegment = req;
 					req->MediaAsset 			= SharedThis(this);
-					req->FirstPTS   			= FTimeValue::GetZero();
+					req->FirstPTS   			= firstTimestamp;
 					req->PrimaryStreamType  	= EStreamType::Audio;
 					req->Bitrate				= Repr->GetBitrate();
 					req->bStartingOnMOOF		= false;
@@ -900,7 +918,7 @@ IManifest::FResult FManifestMP4Internal::FTimelineAssetMP4::GetStartingSegment(T
 					// But if there is a video track we add it as a dependent stream that is also at EOS.
 					if (VideoAdaptationSets.Num())
 					{
-						req->DependentStreams.AddDefaulted_GetRef().StreamType = EStreamType::Video;
+						req->DependentStreamTypes.Add(EStreamType::Video);
 					}
 					return IManifest::FResult(IManifest::FResult::EType::Found);
 				}
diff --git a/Engine/Plugins/Media/ElectraPlayer/Source/ElectraPlayerRuntime/Private/Runtime/Player/mp4/ManifestMP4.h b/Engine/Plugins/Media/ElectraPlayer/Source/ElectraPlayerRuntime/Private/Runtime/Player/mp4/ManifestMP4.h
--- a/Engine/Plugins/Media/ElectraPlayer/Source/ElectraPlayerRuntime/Private/Runtime/Player/mp4/ManifestMP4.h
+++ b/Engine/Plugins/Media/ElectraPlayer/Source/ElectraPlayerRuntime/Private/Runtime/Player/mp4/ManifestMP4.h
@@ -172,6 +172,8 @@ public:
 					return VideoAdaptationSets.Num();
 				case EStreamType::Audio:
 					return AudioAdaptationSets.Num();
+				case EStreamType::Subtitle:
+					return SubtitleAdaptationSets.Num();
 				default:
 					return 0;
 			}
@@ -184,6 +186,8 @@ public:
 					return AdaptationSetIndex < VideoAdaptationSets.Num() ? VideoAdaptationSets[AdaptationSetIndex] : TSharedPtrTS<IPlaybackAssetAdaptationSet>();
 				case EStreamType::Audio:
 					return AdaptationSetIndex < AudioAdaptationSets.Num() ? AudioAdaptationSets[AdaptationSetIndex] : TSharedPtrTS<IPlaybackAssetAdaptationSet>();
+				case EStreamType::Subtitle:
+					return AdaptationSetIndex < SubtitleAdaptationSets.Num() ? SubtitleAdaptationSets[AdaptationSetIndex] : TSharedPtrTS<IPlaybackAssetAdaptationSet>();
 				default:
 					return 0;
 			}
@@ -247,6 +251,7 @@ public:
 		TSharedPtrTS<IParserISO14496_12>				MoovBoxParser;
 		TArray<TSharedPtrTS<FAdaptationSetMP4>>			VideoAdaptationSets;
 		TArray<TSharedPtrTS<FAdaptationSetMP4>>			AudioAdaptationSets;
+		TArray<TSharedPtrTS<FAdaptationSetMP4>>			SubtitleAdaptationSets;
 	};
 
 
diff --git a/Engine/Plugins/Media/ElectraPlayer/Source/ElectraPlayerRuntime/Private/Runtime/Player/mp4/StreamReaderMP4.cpp b/Engine/Plugins/Media/ElectraPlayer/Source/ElectraPlayerRuntime/Private/Runtime/Player/mp4/StreamReaderMP4.cpp
--- a/Engine/Plugins/Media/ElectraPlayer/Source/ElectraPlayerRuntime/Private/Runtime/Player/mp4/StreamReaderMP4.cpp
+++ b/Engine/Plugins/Media/ElectraPlayer/Source/ElectraPlayerRuntime/Private/Runtime/Player/mp4/StreamReaderMP4.cpp
@@ -62,24 +62,44 @@ void FStreamSegmentRequestMP4::SetExecutionDelay(const FTimeValue& ExecutionDela
 	// No op for mp4.
 }
 
-void FStreamSegmentRequestMP4::GetDependentStreams(TArray<FDependentStreams>& OutDependentStreams) const
+FTimeValue FStreamSegmentRequestMP4::GetExecuteAtUTCTime() const
+{
+	// Right now.
+	return FTimeValue::GetInvalid();
+}
+
+void FStreamSegmentRequestMP4::GetDependentStreams(TArray<TSharedPtrTS<IStreamSegment>>& OutDependentStreams) const
 {
 	// Those are not "real" dependent streams in that they are multiplexed and do not need to be fetched from
 	// a different source. This merely indicates the types of non-primary streams we will be demuxing.
-	OutDependentStreams = DependentStreams;
+	for(int32 i=0; i<DependentStreamTypes.Num(); ++i)
+	{
+		FStreamSegmentRequestMP4* DepReq = new FStreamSegmentRequestMP4(*this);
+		DepReq->PrimaryStreamType = DependentStreamTypes[i];
+		DepReq->DependentStreamTypes.Empty();
+		TSharedPtrTS<IStreamSegment> p(DepReq);
+		OutDependentStreams.Push(p);
+	}
 }
 
+void FStreamSegmentRequestMP4::GetRequestedStreams(TArray<TSharedPtrTS<IStreamSegment>>& OutRequestedStreams)
+{
+	OutRequestedStreams.Empty();
+	OutRequestedStreams.Emplace(SharedThis(this));
+}
+
+
 void FStreamSegmentRequestMP4::GetEndedStreams(TArray<TSharedPtrTS<IStreamSegment>>& OutAlreadyEndedStreams)
 {
 	OutAlreadyEndedStreams.Empty();
 	if (bAllTracksAtEOS)
 	{
 		OutAlreadyEndedStreams.Push(SharedThis(this));
-		for(int32 i=0; i<DependentStreams.Num(); ++i)
+		for(int32 i=0; i<DependentStreamTypes.Num(); ++i)
 		{
 			FStreamSegmentRequestMP4* DepReq = new FStreamSegmentRequestMP4;
 			// Only need to set the stream type here.
-			DepReq->PrimaryStreamType = DependentStreams[i].StreamType;
+			DepReq->PrimaryStreamType = DependentStreamTypes[i];
 			TSharedPtrTS<IStreamSegment> p(DepReq);
 			OutAlreadyEndedStreams.Push(p);
 		}
diff --git a/Engine/Plugins/Media/ElectraPlayer/Source/ElectraPlayerRuntime/Private/Runtime/Player/mp4/StreamReaderMP4.h b/Engine/Plugins/Media/ElectraPlayer/Source/ElectraPlayerRuntime/Private/Runtime/Player/mp4/StreamReaderMP4.h
--- a/Engine/Plugins/Media/ElectraPlayer/Source/ElectraPlayerRuntime/Private/Runtime/Player/mp4/StreamReaderMP4.h
+++ b/Engine/Plugins/Media/ElectraPlayer/Source/ElectraPlayerRuntime/Private/Runtime/Player/mp4/StreamReaderMP4.h
@@ -32,9 +32,12 @@ public:
 	virtual uint32 GetPlaybackSequenceID() const override;
 
 	virtual void SetExecutionDelay(const FTimeValue& ExecutionDelay) override;
+	virtual FTimeValue GetExecuteAtUTCTime() const override;
+
 	virtual EStreamType GetType() const override;
 
-	virtual void GetDependentStreams(TArray<FDependentStreams>& OutDependentStreams) const override;
+	virtual void GetDependentStreams(TArray<TSharedPtrTS<IStreamSegment>>& OutDependentStreams) const override;
+	virtual void GetRequestedStreams(TArray<TSharedPtrTS<IStreamSegment>>& OutRequestedStreams) override;
 	virtual void GetEndedStreams(TArray<TSharedPtrTS<IStreamSegment>>& OutAlreadyEndedStreams) override;
 
 	//! Returns the first PTS value as indicated by the media timeline. This should correspond to the actual absolute PTS of the sample.
@@ -49,7 +52,7 @@ public:
 
 	TSharedPtrTS<ITimelineMediaAsset>					MediaAsset;					//!< The entire mp4 asset
 	TSharedPtrTS<IParserISO14496_12::ITrackIterator>	PrimaryTrackIterator;
-	TArray<FDependentStreams>							DependentStreams;
+	TArray<EStreamType>									DependentStreamTypes;
 
 	FTimeValue											FirstPTS;					//!< The PTS of the first sample
 	FTimeValue											SegmentDuration;
diff --git a/Engine/Plugins/Media/ElectraPlayer/Source/ElectraPlayerRuntime/Private/Runtime/PlayerFacility.h b/Engine/Plugins/Media/ElectraPlayer/Source/ElectraPlayerRuntime/Private/Runtime/PlayerFacility.h
--- a/Engine/Plugins/Media/ElectraPlayer/Source/ElectraPlayerRuntime/Private/Runtime/PlayerFacility.h
+++ b/Engine/Plugins/Media/ElectraPlayer/Source/ElectraPlayerRuntime/Private/Runtime/PlayerFacility.h
@@ -41,6 +41,7 @@ namespace Electra
 			EntityCache,
 			DRM,
 			H265Decoder,
+			SubtitleDecoder,
 			LastEntry = 255
 		};
 
@@ -104,6 +105,8 @@ namespace Electra
 					return TEXT("DRM");
 				case EFacility::H265Decoder:
 					return TEXT("H.265 decoder");
+				case EFacility::SubtitleDecoder:
+					return TEXT("Subtitle decoder");
 			}
 			return TEXT("???");
 		}
diff --git a/Engine/Plugins/Media/ElectraPlayer/Source/ElectraPlayerRuntime/Private/Runtime/StreamAccessUnitBuffer.cpp b/Engine/Plugins/Media/ElectraPlayer/Source/ElectraPlayerRuntime/Private/Runtime/StreamAccessUnitBuffer.cpp
--- a/Engine/Plugins/Media/ElectraPlayer/Source/ElectraPlayerRuntime/Private/Runtime/StreamAccessUnitBuffer.cpp
+++ b/Engine/Plugins/Media/ElectraPlayer/Source/ElectraPlayerRuntime/Private/Runtime/StreamAccessUnitBuffer.cpp
@@ -263,6 +263,11 @@ namespace Electra
 		return LastPoppedPTS;
 	}
 
+	FTimeValue FMultiTrackAccessUnitBuffer::GetLastPoppedDTS()
+	{
+		FMediaCriticalSection::ScopedLock lock(AccessLock);
+		return LastPoppedDTS;
+	}
 
 	TSharedPtrTS<FAccessUnitBuffer> FMultiTrackAccessUnitBuffer::GetSelectedTrackBuffer()
 	{
@@ -485,8 +490,11 @@ namespace Electra
 				// Did we just pop from a different buffer than last time?
 				if (LastPoppedBufferInfo != ActiveOutputBufferInfo)
 				{
-					// FIXME: We may need to discard everything that is not tagged as a sync sample to ensure proper stream switching.
-					OutAU->bTrackChangeDiscontinuity = true;
+					if (LastPoppedBufferInfo)
+					{
+						// FIXME: We may need to discard everything that is not tagged as a sync sample to ensure proper stream switching.
+						OutAU->bTrackChangeDiscontinuity = true;
+					}
 
 					// Remember from which buffer we popped.
 					LastPoppedBufferInfo = ActiveOutputBufferInfo;
diff --git a/Engine/Plugins/Media/ElectraPlayer/Source/ElectraPlayerRuntime/Private/Runtime/StreamAccessUnitBuffer.h b/Engine/Plugins/Media/ElectraPlayer/Source/ElectraPlayerRuntime/Private/Runtime/StreamAccessUnitBuffer.h
--- a/Engine/Plugins/Media/ElectraPlayer/Source/ElectraPlayerRuntime/Private/Runtime/StreamAccessUnitBuffer.h
+++ b/Engine/Plugins/Media/ElectraPlayer/Source/ElectraPlayerRuntime/Private/Runtime/StreamAccessUnitBuffer.h
@@ -597,6 +597,7 @@ namespace Electra
 		void Flush();
 		void GetStats(FAccessUnitBufferInfo& OutStats);
 		bool IsDeselected();
+		FTimeValue GetLastPoppedDTS();
 		FTimeValue GetLastPoppedPTS();
 		TSharedPtrTS<FBufferSourceInfo> GetActiveOutputBufferInfo() const
 		{ return ActiveOutputBufferInfo; }
diff --git a/Engine/Plugins/Media/ElectraPlayer/Source/ElectraPlayerRuntime/Private/Runtime/StreamTypes.cpp b/Engine/Plugins/Media/ElectraPlayer/Source/ElectraPlayerRuntime/Private/Runtime/StreamTypes.cpp
--- a/Engine/Plugins/Media/ElectraPlayer/Source/ElectraPlayerRuntime/Private/Runtime/StreamTypes.cpp
+++ b/Engine/Plugins/Media/ElectraPlayer/Source/ElectraPlayerRuntime/Private/Runtime/StreamTypes.cpp
@@ -10,6 +10,10 @@ namespace Electra
 
 	FString FStreamCodecInformation::GetMimeType() const
 	{
+		if (!MimeType.IsEmpty())
+		{
+			return MimeType;
+		}
 		switch(GetCodec())
 		{
 			case ECodec::H264:
@@ -21,6 +25,9 @@ namespace Electra
 			case ECodec::EAC3:
 				return FString(TEXT("audio/mp4"));
 			case ECodec::WebVTT:
+			case ECodec::TTML:
+			case ECodec::TX3G:
+			case ECodec::OtherSubtitle:
 				return FString(TEXT("application/mp4"));
 			default:
 				return FString(TEXT("application/octet-stream"));
@@ -212,8 +219,21 @@ namespace Electra
 			StreamType = EStreamType::Subtitle;
 			CodecSpecifier = CodecOTI;
 			Codec = ECodec::WebVTT;
-			// Presently not supported.
-			return false;
+			return true;
+		}
+		else if (CodecOTI.Equals(TEXT("stpp")))		// this is indicating one of the many TTML variants, should be IMSC1, SMPTE-TT, EBU-TT
+		{
+			StreamType = EStreamType::Subtitle;
+			CodecSpecifier = CodecOTI;
+			Codec = ECodec::TTML;
+			return true;
+		}
+		else if (CodecOTI.Equals(TEXT("tx3g")))
+		{
+			StreamType = EStreamType::Subtitle;
+			CodecSpecifier = CodecOTI;
+			Codec = ECodec::TX3G;
+			return true;
 		}
 		else
 		{
diff --git a/Engine/Plugins/Media/ElectraPlayer/Source/ElectraPlayerRuntime/Private/Runtime/StreamTypes.h b/Engine/Plugins/Media/ElectraPlayer/Source/ElectraPlayerRuntime/Private/Runtime/StreamTypes.h
--- a/Engine/Plugins/Media/ElectraPlayer/Source/ElectraPlayerRuntime/Private/Runtime/StreamTypes.h
+++ b/Engine/Plugins/Media/ElectraPlayer/Source/ElectraPlayerRuntime/Private/Runtime/StreamTypes.h
@@ -54,6 +54,9 @@ namespace Electra
 			EAC3,
 			// --- Subtitle / Caption ---
 			WebVTT = 200,
+			TTML,
+			TX3G,
+			OtherSubtitle
 		};
 
 		EStreamType GetStreamType() const
@@ -61,6 +64,10 @@ namespace Electra
 			return StreamType;
 		}
 
+		void SetMimeType(const FString& InMimeType)
+		{
+			MimeType = InMimeType;
+		}
 		FString GetMimeType() const;
 		FString GetMimeTypeWithCodec() const;
 		FString GetMimeTypeWithCodecAndFeatures() const;
@@ -111,6 +118,9 @@ namespace Electra
 			switch (GetCodec())
 			{
 				case ECodec::WebVTT:
+				case ECodec::TTML:
+				case ECodec::TX3G:
+				case ECodec::OtherSubtitle:
 					return true;
 				default:
 					return false;
@@ -184,6 +194,40 @@ namespace Electra
 		}
 
 
+		struct FTranslation
+		{
+			FTranslation(int32 x = 0, int32 y = 0)
+				: Left(x)
+				, Top(y)
+			{
+			}
+			void Clear()
+			{
+				Left = Top = 0;
+			}
+			int32 GetX() const
+			{
+				return Left;
+			}
+			int32 GetY() const
+			{
+				return Top;
+			}
+			int32 Left;
+			int32 Top;
+		};
+
+		const FTranslation& GetTranslation() const
+		{
+			return Translation;
+		}
+
+		void SetTranslation(const FTranslation& InTranslation)
+		{
+			Translation = InTranslation;
+		}
+
+
 		struct FAspectRatio
 		{
 			FAspectRatio(int32 w = 0, int32 h = 0)
@@ -419,9 +463,11 @@ namespace Electra
 		{
 			StreamType = EStreamType::Unsupported;
 			CodecSpecifier.Empty();
+			MimeType.Empty();
 			Codec = ECodec::Unknown;
 			Resolution.Clear();
 			Crop.Clear();
+			Translation.Clear();
 			AspectRatio.Clear();
 			FrameRate = FTimeFraction::GetInvalid();
 			ProfileLevel.Clear();
@@ -506,9 +552,11 @@ namespace Electra
 
 		EStreamType		StreamType;
 		FString			CodecSpecifier;				//!< Codec specifier as per RFC 6381
+		FString			MimeType;					//!< Explicitly set mime type if it cannot be inferred.
 		ECodec			Codec;
 		FResolution		Resolution;					//!< Resolution, if this is a video stream
 		FCrop			Crop;						//!< Cropping, if this is a video stream
+		FTranslation	Translation;				//!< Top-left corner offset for subtitles
 		FAspectRatio	AspectRatio;				//!< Aspect ratio, if this is a video stream
 		FTimeFraction	FrameRate;					//!< Frame rate, if this is a video stream
 		FProfileLevel	ProfileLevel;
@@ -599,7 +647,9 @@ namespace Electra
 		// of the tracks as they are found. If the index is invalid the selection rules for kind and language are applied.
 		TOptional<int32> OverrideIndex;
 
-		bool IsCompatibleWith(const FStreamSelectionAttributes& Other)
+		virtual ~FStreamSelectionAttributes() = default;
+
+		virtual bool IsCompatibleWith(const FStreamSelectionAttributes& Other)
 		{
 			FString Kind1 = Kind.IsSet() ? Kind.GetValue() : FString();
 			FString Kind2 = Other.Kind.IsSet() ? Other.Kind.GetValue() : FString();
@@ -614,13 +664,13 @@ namespace Electra
 			return false;
 		}
 
-		void Reset()
+		virtual void Reset()
 		{
 			Kind.Reset();
 			Language_ISO639.Reset();
 			OverrideIndex.Reset();
 		}
-		void UpdateWith(const FString& InKind, const FString& InLanguage, int32 InOverrideIndex)
+		virtual void UpdateWith(const FString& InKind, const FString& InLanguage, int32 InOverrideIndex)
 		{
 			Reset();
 			if (!InKind.IsEmpty())
@@ -636,7 +686,7 @@ namespace Electra
 				OverrideIndex = InOverrideIndex;
 			}
 		}
-		void UpdateIfOverrideSet(const FString& InKind, const FString& InLanguage)
+		virtual void UpdateIfOverrideSet(const FString& InKind, const FString& InLanguage)
 		{
 			if (OverrideIndex.IsSet())
 			{
@@ -645,7 +695,7 @@ namespace Electra
 				Language_ISO639 = InLanguage;
 			}
 		}
-		void ClearOverrideIndex()
+		virtual void ClearOverrideIndex()
 		{
 			OverrideIndex.Reset();
 		}
diff --git a/Engine/Plugins/Media/ElectraPlayer/Source/ElectraPlayerRuntime/Public/IElectraPlayerInterface.h b/Engine/Plugins/Media/ElectraPlayer/Source/ElectraPlayerRuntime/Public/IElectraPlayerInterface.h
--- a/Engine/Plugins/Media/ElectraPlayer/Source/ElectraPlayerRuntime/Public/IElectraPlayerInterface.h
+++ b/Engine/Plugins/Media/ElectraPlayer/Source/ElectraPlayerRuntime/Public/IElectraPlayerInterface.h
@@ -14,6 +14,8 @@ class IAudioDecoderOutput;
 using IAudioDecoderOutputPtr = TSharedPtr<IAudioDecoderOutput, ESPMode::ThreadSafe>;
 class IMetaDataDecoderOutput;
 using IMetaDataDecoderOutputPtr = TSharedPtr<IMetaDataDecoderOutput, ESPMode::ThreadSafe>;
+class ISubtitleDecoderOutput;
+using ISubtitleDecoderOutputPtr = TSharedPtr<ISubtitleDecoderOutput, ESPMode::ThreadSafe>;
 
 // ---------------------------------------------------------------------------------------------
 
@@ -66,8 +68,10 @@ public:
 	virtual void SendMediaEvent(EPlayerEvent Event) = 0;
 	virtual void OnVideoFlush() = 0;
 	virtual void OnAudioFlush() = 0;
+	virtual void OnSubtitleFlush() = 0;
 	virtual void PresentVideoFrame(const FVideoDecoderOutputPtr& InVideoFrame) = 0;
 	virtual void PresentAudioFrame(const IAudioDecoderOutputPtr& InAudioFrame) = 0;
+	virtual void PresentSubtitleSample(const ISubtitleDecoderOutputPtr& InSubtitleSample) = 0;
 	virtual void PresentMetadataSample(const IMetaDataDecoderOutputPtr& InMetadataSample) = 0;
 	virtual bool CanReceiveVideoSamples(int32 NumFrames) = 0;
 	virtual bool CanReceiveAudioSamples(int32 NumFrames) = 0;
@@ -114,6 +118,7 @@ public:
 	{
 		TOptional<FTimespan>		TimeOffset;
 		FStreamSelectionAttributes	InitialAudioTrackAttributes;
+		FStreamSelectionAttributes	InitialSubtitleTrackAttributes;
 		TOptional<int32>			MaxVerticalStreamResolution;
 		TOptional<int32>			MaxBandwidthForStreaming;
 		bool						bDoNotPreload = false;
diff --git a/Engine/Plugins/Media/ElectraUtil/Source/ElectraSamples/Public/IElectraSubtitleSample.h b/Engine/Plugins/Media/ElectraUtil/Source/ElectraSamples/Public/IElectraSubtitleSample.h
new file mode 100644
--- /dev/null
+++ b/Engine/Plugins/Media/ElectraUtil/Source/ElectraSamples/Public/IElectraSubtitleSample.h
@@ -0,0 +1,33 @@
+// Copyright Epic Games, Inc. All Rights Reserved.
+#pragma once
+
+#include "Templates/SharedPointer.h"
+#include "Modules/ModuleInterface.h"
+#include "Misc/Timespan.h"
+#include "IMediaOverlaySample.h"
+
+/**
+ * Extension of the media overlay sample to carry Electra
+ * specific values.
+ */
+class IElectraSubtitleSample : public IMediaOverlaySample
+{
+public:
+	static FGuid GetSampleTypeGUID()
+	{
+		static FGuid SampleTypeGUID(0xC94C9B6F, 0x07DF4B88, 0xACE0B08A, 0x9EF6AA39);
+		return SampleTypeGUID;
+	}
+
+	virtual ~IElectraSubtitleSample() = default;
+
+	virtual FTimespan GetDuration() const = 0;
+	virtual TOptional<FVector2D> GetPosition() const = 0;
+	virtual FText GetText() const = 0;
+	virtual FMediaTimeStamp GetTime() const = 0;
+	//virtual TOptional<FTimecode> GetTimecode() const { return TOptional<FTimecode>(); }
+	virtual EMediaOverlaySampleType GetType() const = 0;
+
+	virtual FGuid GetGUID() const = 0;
+
+};
diff --git a/Engine/Source/Runtime/Media/Public/IMediaOverlaySample.h b/Engine/Source/Runtime/Media/Public/IMediaOverlaySample.h
--- a/Engine/Source/Runtime/Media/Public/IMediaOverlaySample.h
+++ b/Engine/Source/Runtime/Media/Public/IMediaOverlaySample.h
@@ -80,6 +80,13 @@ public:
 	 */
 	virtual EMediaOverlaySampleType GetType() const = 0;
 
+	/**
+	 * Get the GUID identifying the derived type that may implement additional,
+	 * type specific methods. If the GUID matches a type known type it is safe
+	 * to static cast this class to the derived type.
+	 */
+	virtual FGuid GetGUID() const { return FGuid(); }
+
 public:
 
 	/** Virtual destructor. */
diff --git a/Engine/Source/Runtime/Media/Public/IMediaSamples.h b/Engine/Source/Runtime/Media/Public/IMediaSamples.h
--- a/Engine/Source/Runtime/Media/Public/IMediaSamples.h
+++ b/Engine/Source/Runtime/Media/Public/IMediaSamples.h
@@ -122,6 +122,7 @@ public:
 	virtual bool PeekVideoSampleTime(FMediaTimeStamp & TimeStamp) = 0;
 
 	virtual uint32 PurgeOutdatedVideoSamples(const FMediaTimeStamp & ReferenceTime, bool bReversed) { return 0; };
+	virtual uint32 PurgeOutdatedSubtitleSamples(const FMediaTimeStamp & ReferenceTime, bool bReversed) { return 0; };
 
 	virtual bool CanReceiveVideoSamples(uint32 Num) const { return true; }
 	virtual bool CanReceiveAudioSamples(uint32 Num) const { return true; }
diff --git a/Engine/Source/Runtime/MediaUtils/Private/MediaPlayerFacade.cpp b/Engine/Source/Runtime/MediaUtils/Private/MediaPlayerFacade.cpp
--- a/Engine/Source/Runtime/MediaUtils/Private/MediaPlayerFacade.cpp
+++ b/Engine/Source/Runtime/MediaUtils/Private/MediaPlayerFacade.cpp
@@ -15,6 +15,7 @@
 #include "IMediaSamples.h"
 #include "IMediaAudioSample.h"
 #include "IMediaTextureSample.h"
+#include "IMediaOverlaySample.h"
 #include "IMediaTracks.h"
 #include "IMediaView.h"
 #include "IMediaTicker.h"
@@ -69,12 +70,17 @@ DECLARE_DWORD_COUNTER_STAT(TEXT("MediaPlayerFacade NumAudioSamples"), STAT_Media
 DECLARE_DWORD_COUNTER_STAT(TEXT("MediaPlayerFacade NumPurgedVideoSamples"), STAT_MediaUtils_FacadeNumPurgedVideoSamples, STATGROUP_Media);
 DECLARE_DWORD_ACCUMULATOR_STAT(TEXT("MediaPlayerFacade TotalPurgedVideoSamples"), STAT_MediaUtils_FacadeTotalPurgedVideoSamples, STATGROUP_Media);
 
+/** Number of purged subtitle samples */
+DECLARE_DWORD_COUNTER_STAT(TEXT("MediaPlayerFacade NumPurgedSubtitleSamples"), STAT_MediaUtils_FacadeNumPurgedSubtitleSamples, STATGROUP_Media);
+DECLARE_DWORD_ACCUMULATOR_STAT(TEXT("MediaPlayerFacade TotalPurgedSubtitleSamples"), STAT_MediaUtils_FacadeTotalPurgedSubtitleSamples, STATGROUP_Media);
+
 /* Some constants
 *****************************************************************************/
 
 const double kMaxTimeSinceFrameStart = 0.300; // max seconds we allow between the start of the frame and the player facade timing computations (to catch suspended apps & debugging)
 const double kMaxTimeSinceAudioTimeSampling = 0.250; // max seconds we allow to have passed between the last audio timing sampling and the player facade timing computations (to catch suspended apps & debugging - some platforms do update audio at a farily low rate: hence the big tollerance)
-const double kOutdatedVideoSamplesTollerance = 0.050; // seconds video samples are allowed to be "too old" to stay in the player's output queue despite of calculations indicating they need to go
+const double kOutdatedVideoSamplesTolerance = 0.050; // seconds video samples are allowed to be "too old" to stay in the player's output queue despite of calculations indicating they need to go
+const double kOutdatedSubtitleSamplesTolerance = 0.050; // seconds subtitle samples are allowed to be "too old" to stay in the player's output queue despite of calculations indicating they need to go
 
 /* Local helpers
 *****************************************************************************/
@@ -2655,10 +2661,24 @@ void FMediaPlayerFacade::ProcessSubtitleSamples(IMediaSamples& Samples, TRange<F
 {
 	TSharedPtr<IMediaOverlaySample, ESPMode::ThreadSafe> Sample;
 
+	// There might be samples in the subtitle sample queue that have may lie before the specified time range.
+	// Since these have no way of being displayed they must be removed from the sample queue.
+	TSharedPtr<IMediaPlayer, ESPMode::ThreadSafe> CurrentPlayer = Player;
+	if (CurrentPlayer.IsValid() && CurrentPlayer->GetPlayerFeatureFlag(IMediaPlayer::EFeatureFlag::UsePlaybackTimingV2))
+	{
+		bool bReverse = CurrentRate < 0.0f;
+		uint32 NumPurged = CurrentPlayer->GetSamples().PurgeOutdatedSubtitleSamples(TimeRange.GetLowerBoundValue() + (bReverse ? kOutdatedVideoSamplesTolerance : -kOutdatedVideoSamplesTolerance), bReverse);
+		SET_DWORD_STAT(STAT_MediaUtils_FacadeNumPurgedSubtitleSamples, NumPurged);
+		INC_DWORD_STAT_BY(STAT_MediaUtils_FacadeTotalPurgedSubtitleSamples, NumPurged);
+	}
+
 	while (Samples.FetchSubtitle(TimeRange, Sample))
 	{
 		if (Sample.IsValid() && !SubtitleSampleSinks.Enqueue(Sample.ToSharedRef(), FMediaPlayerQueueDepths::MaxSubtitleSinkDepth))
 		{
+			FString Caption = Sample->GetText().ToString();
+			UE_LOG(LogMediaUtils, Log, TEXT("New caption @%.3f: %s"), Sample->GetTime().Time.GetTotalSeconds(), *Caption);
+
 #if MEDIAPLAYERFACADE_TRACE_SINKOVERFLOWS
 			UE_LOG(LogMediaUtils, VeryVerbose, TEXT("PlayerFacade %p: Subtitle sample sink overflow"), this);
 #endif
@@ -2716,8 +2736,7 @@ void FMediaPlayerFacade::ReceiveMediaEvent(EMediaEvent Event)
 			}
 
 			bool bReverse = (Rate < 0.0f);
-			uint32 NumPurged = CurrentPlayer->GetSamples().PurgeOutdatedVideoSamples(TimeRange.GetLowerBoundValue() + (bReverse ? kOutdatedVideoSamplesTollerance : -kOutdatedVideoSamplesTollerance), bReverse);
-
+			uint32 NumPurged = CurrentPlayer->GetSamples().PurgeOutdatedVideoSamples(TimeRange.GetLowerBoundValue() + (bReverse ? kOutdatedVideoSamplesTolerance : -kOutdatedVideoSamplesTolerance), bReverse);
 			SET_DWORD_STAT(STAT_MediaUtils_FacadeNumPurgedVideoSamples, NumPurged);
 			INC_DWORD_STAT_BY(STAT_MediaUtils_FacadeTotalPurgedVideoSamples, NumPurged);
 
diff --git a/Engine/Source/Runtime/MediaUtils/Private/MediaSamples.cpp b/Engine/Source/Runtime/MediaUtils/Private/MediaSamples.cpp
--- a/Engine/Source/Runtime/MediaUtils/Private/MediaSamples.cpp
+++ b/Engine/Source/Runtime/MediaUtils/Private/MediaSamples.cpp
@@ -139,6 +139,14 @@ uint32 FMediaSamples::PurgeOutdatedVideoSamples(const FMediaTimeStamp & Referenc
 	return VideoSampleQueue.PurgeOutdatedSamples(ReferenceTime, bReversed);
 }
 
+/**
+ * Remove any subtitle samples from the queue that have no chance of being displayed anymore
+ */
+uint32 FMediaSamples::PurgeOutdatedSubtitleSamples(const FMediaTimeStamp & ReferenceTime, bool bReversed)
+{
+	return SubtitleSampleQueue.PurgeOutdatedSamples(ReferenceTime, bReversed);
+}
+
 /**
  * Check if can receive more video samples
  */
diff --git a/Engine/Source/Runtime/MediaUtils/Public/MediaSamples.h b/Engine/Source/Runtime/MediaUtils/Public/MediaSamples.h
--- a/Engine/Source/Runtime/MediaUtils/Public/MediaSamples.h
+++ b/Engine/Source/Runtime/MediaUtils/Public/MediaSamples.h
@@ -165,6 +165,7 @@ public:
 	}
 
 	virtual uint32 PurgeOutdatedVideoSamples(const FMediaTimeStamp & ReferenceTime, bool bReversed) override;
+	virtual uint32 PurgeOutdatedSubtitleSamples(const FMediaTimeStamp & ReferenceTime, bool bReversed) override;
 
 	virtual bool CanReceiveVideoSamples(uint32 Num) const override;
 	virtual bool CanReceiveAudioSamples(uint32 Num) const override;
