ElectraPlayer: Changed internal state transition between playing and paused (SetRate(1.0) and SetRate(0.0)) to immediately reflect the internal player state "Playing" and "Paused" that would normally be updated asynchronously. This is to avoid issues with the calling code expecting to see the "Playing" state immediately after a SetRate(1.0)

diff --git a/Engine/Plugins/Media/ElectraPlayer/Source/ElectraPlayerRuntime/Private/ElectraPlayer.cpp b/Engine/Plugins/Media/ElectraPlayer/Source/ElectraPlayerRuntime/Private/ElectraPlayer.cpp
--- a/Engine/Plugins/Media/ElectraPlayer/Source/ElectraPlayerRuntime/Private/ElectraPlayer.cpp
+++ b/Engine/Plugins/Media/ElectraPlayer/Source/ElectraPlayerRuntime/Private/ElectraPlayer.cpp
@@ -26,6 +26,10 @@ DECLARE_CYCLE_STAT(TEXT("FElectraPlayer::TickInput"), STAT_ElectraPlayer_Electra
 
 //-----------------------------------------------------------------------------
 
+#define USE_INTERNAL_PLAYBACK_STATE 1
+
+//-----------------------------------------------------------------------------
+
 #if UE_BUILD_SHIPPING
 #define HIDE_URLS_FROM_LOG	1
 #else
@@ -134,8 +138,9 @@ FElectraPlayer::FElectraPlayer(const TSharedPtr<IElectraPlayerAdapterDelegate, E
 	ReportVideoStreamingErrorDelegate.AddRaw(this, &FElectraPlayer::ReportVideoStreamingError);
 	ReportSubtitlesMetricsDelegate.AddRaw(this, &FElectraPlayer::ReportSubtitlesMetrics);
 
-	State = EPlayerState::Closed;
-	Status = EPlayerStatus::None;
+#if USE_INTERNAL_PLAYBACK_STATE
+	PlayerState.bUseInternal = true;
+#endif
 	bAllowKillAfterCloseEvent = false;
 	bPlayerHasClosed = false;
 	bHasPendingError = false;
@@ -180,6 +185,7 @@ void FElectraPlayer::ClearToDefaultState()
 {
 	FScopeLock lock(&PlayerLock);
 
+	PlayerState.Reset();
 	NumTracksAudio = 0;
 	NumTracksVideo = 0;
 	NumTracksSubtitle = 0;
@@ -209,8 +215,6 @@ bool FElectraPlayer::OpenInternal(const FString& Url, const FParamDict& PlayerOp
 	// Clear out our work variables
 	ClearToDefaultState();
 	bAllowKillAfterCloseEvent = false;
-	State = EPlayerState::Closed;
-	Status = EPlayerStatus::None;
 	bPlayerHasClosed = false;
 	bHasPendingError = false;
 
@@ -311,7 +315,7 @@ void FElectraPlayer::CloseInternal(bool bKillAfterClose)
 	TSharedPtr<FInternalPlayerImpl, ESPMode::ThreadSafe>	Player = CurrentPlayer;
 
 	// For all intents and purposes the player can be considered closed here now already.
-	State = EPlayerState::Closed;
+	PlayerState.State = EPlayerState::Closed;
 	MediaUrl = FString();
 	PlaystartOptions.TimeOffset.Reset();
 	PlaystartOptions.InitialAudioTrackAttributes.Reset();
@@ -431,6 +435,44 @@ void FElectraPlayer::FInternalPlayerImpl::DoCloseAsync(TSharedPtr<FInternalPlaye
 #endif
 }
 
+
+
+float FElectraPlayer::FPlayerState::GetRate() const
+{
+	return bUseInternal && IntendedPlayRate.IsSet() ? IntendedPlayRate.GetValue() : CurrentPlayRate;
+}
+
+FElectraPlayer::EPlayerState FElectraPlayer::FPlayerState::GetState() const
+{
+	if (bUseInternal && IntendedPlayRate.IsSet() && (State == FElectraPlayer::EPlayerState::Playing || State == FElectraPlayer::EPlayerState::Paused || State == FElectraPlayer::EPlayerState::Stopped))
+	{
+		return IntendedPlayRate.GetValue() != 0.0f ? FElectraPlayer::EPlayerState::Playing : FElectraPlayer::EPlayerState::Paused;
+	}
+	return State;
+}
+
+FElectraPlayer::EPlayerStatus FElectraPlayer::FPlayerState::GetStatus() const
+{
+	return Status;
+}
+
+void FElectraPlayer::FPlayerState::SetIntendedPlayRate(float InIntendedRate)
+{
+	IntendedPlayRate = InIntendedRate;
+}
+
+void FElectraPlayer::FPlayerState::SetPlayRateFromPlayer(float InCurrentPlayerPlayRate)
+{
+	CurrentPlayRate = InCurrentPlayerPlayRate;
+	// If reverse playback is selected even though it is not supported, leave it set as such.
+	if (IntendedPlayRate.IsSet() && IntendedPlayRate.GetValue() >= 0.0f)
+	{
+		IntendedPlayRate.Reset();
+	}
+}
+
+
+
 //-----------------------------------------------------------------------------
 /**
 */
@@ -459,7 +501,7 @@ void FElectraPlayer::Tick(FTimespan DeltaTime, FTimespan Timecode)
 	// Handle the internal player, if we have one.
 	PlayerLock.Lock();
 	TSharedPtr<FInternalPlayerImpl, ESPMode::ThreadSafe> Player = CurrentPlayer;
-	if (!bPlayerHasClosed && Player.IsValid() && State != EPlayerState::Error)
+	if (!bPlayerHasClosed && Player.IsValid() && PlayerState.State != EPlayerState::Error)
 	{
 		if (Player->RendererVideo.IsValid())
 		{
@@ -502,17 +544,17 @@ void FElectraPlayer::Tick(FTimespan DeltaTime, FTimespan Timecode)
 		if (bHasPendingError)
 		{
 			bHasPendingError = false;
-			if (State == EPlayerState::Preparing)
+			if (PlayerState.State == EPlayerState::Preparing)
 			{
 				DeferredEvents.Enqueue(IElectraPlayerAdapterDelegate::EPlayerEvent::MediaOpenFailed);
 			}
-			else if (State == EPlayerState::Playing)
+			else if (PlayerState.State == EPlayerState::Playing)
 			{
 				DeferredEvents.Enqueue(IElectraPlayerAdapterDelegate::EPlayerEvent::MediaClosed);
 			}
 
 			CloseInternal(true);
-			State = EPlayerState::Error;
+			PlayerState.State = EPlayerState::Error;
 		}
 	}
 	else
@@ -542,7 +584,7 @@ void FElectraPlayer::OnVideoDecoded(const FVideoDecoderOutputPtr& DecoderOutput,
 	TSharedPtr<FInternalPlayerImpl, ESPMode::ThreadSafe> Player = CurrentPlayer;
 	if (Player.IsValid())
 	{
-		if (State != EPlayerState::Closed)
+		if (PlayerState.State != EPlayerState::Closed)
 		{
 			if (DecoderOutput.IsValid())
 			{
@@ -582,7 +624,7 @@ void FElectraPlayer::OnAudioDecoded(const IAudioDecoderOutputPtr& DecoderOutput)
 	TSharedPtr<FInternalPlayerImpl, ESPMode::ThreadSafe> Player = CurrentPlayer;
 	if (Player.IsValid())
 	{
-		if (State != EPlayerState::Closed)
+		if (PlayerState.State != EPlayerState::Closed)
 		{
 			PresentAudioFrame(DecoderOutput);
 		}
@@ -607,7 +649,7 @@ void FElectraPlayer::OnSubtitleDecoded(ISubtitleDecoderOutputPtr DecoderOutput)
 	TSharedPtr<FInternalPlayerImpl, ESPMode::ThreadSafe> Player = CurrentPlayer;
 	if (Player.IsValid())
 	{
-		if (State != EPlayerState::Closed)
+		if (PlayerState.State != EPlayerState::Closed)
 		{
 			PresentSubtitle(DecoderOutput);
 		}
@@ -751,17 +793,27 @@ bool FElectraPlayer::CanPresentAudioFrames(uint64 NumFrames)
 
 float FElectraPlayer::GetRate() const
 {
-	return CurrentPlayer.Get() && CurrentPlayer->AdaptivePlayer->IsPlaying() ? 1.0f : 0.0f;
+	return PlayerState.GetRate();
+/*
+	if (PlayerState.bUseInternal)
+	{
+		return 0.0f;
+	}
+	else
+	{
+		return CurrentPlayer.Get() && CurrentPlayer->AdaptivePlayer->IsPlaying() ? 1.0f : 0.0f;
+	}
+*/
 }
 
 FElectraPlayer::EPlayerState FElectraPlayer::GetState() const
 {
-	return State;
+	return PlayerState.GetState();
 }
 
 FElectraPlayer::EPlayerStatus FElectraPlayer::GetStatus() const
 {
-	return Status;
+	return PlayerState.GetStatus();
 }
 
 bool FElectraPlayer::IsLooping() const
@@ -851,9 +903,13 @@ FTimespan FElectraPlayer::GetSeekableDuration() const
 
 bool FElectraPlayer::SetRate(float Rate)
 {
+	//UE_LOG(LogElectraPlayer, Verbose, TEXT("[%p][%p] IMediaPlayer::SetRate(%.3f)"), this, CurrentPlayer.Get(), Rate);
 	if (CurrentPlayer.Get())
 	{
-		if (Rate == 0.0f)
+		// Set the intended rate, which *may* be set negative. This is not supported and we put the adaptive player into pause
+		// if this happens, but we keep the intended rate set nevertheless.
+		PlayerState.SetIntendedPlayRate(Rate);
+		if (Rate <= 0.0f)
 		{
 			CurrentPlayer->AdaptivePlayer->Pause();
 		}
@@ -974,13 +1030,23 @@ bool FElectraPlayer::Seek(const FTimespan& Time)
 {
 	if (CurrentPlayer.Get())
 	{
-		FTimespan Target;
-		CalculateTargetSeekTime(Target, Time);
-		Electra::IAdaptiveStreamingPlayer::FSeekParam seek;
-		seek.Time.SetFromHNS(Target.GetTicks());
-		bInitialSeekPerformed = true;
-		CurrentPlayer->AdaptivePlayer->SeekTo(seek);
-		return true;
+		FTimespan CurrentPos = GetTime();
+		double Distance = (Time - CurrentPos).GetTotalSeconds();
+		// If the new position is close to the current position do not perform a seek.
+		if (Distance < -1.0 || Distance > 1.0)
+		{
+			FTimespan Target;
+			CalculateTargetSeekTime(Target, Time);
+			Electra::IAdaptiveStreamingPlayer::FSeekParam seek;
+			seek.Time.SetFromHNS(Target.GetTicks());
+			bInitialSeekPerformed = true;
+			CurrentPlayer->AdaptivePlayer->SeekTo(seek);
+			return true;
+		}
+		else
+		{
+			UE_LOG(LogElectraPlayer, Verbose, TEXT("[%p][%p] IMediaPlayer::Seek(%.3f) ignored. Only %.3f seconds away from current position @%.3f"), this, CurrentPlayer.Get(), Time.GetTotalSeconds(), Distance, CurrentPos.GetTotalSeconds());
+		}
 	}
 	return false;
 }
@@ -1479,11 +1545,15 @@ void FElectraPlayer::HandleDeferredPlayerEvents()
 				break;
 			}
 			case FPlayerMetricEventBase::EType::ReceivedPlaylists:
+			{
 				HandlePlayerEventReceivedPlaylists();
 				break;
+			}
 			case FPlayerMetricEventBase::EType::TracksChanged:
+			{
 				HandlePlayerEventTracksChanged();
 				break;
+			}
 			case FPlayerMetricEventBase::EType::PlaylistDownload:
 			{
 				FPlayerMetricEvent_PlaylistDownload* Ev = static_cast<FPlayerMetricEvent_PlaylistDownload*>(Event.Get());
@@ -1539,23 +1609,35 @@ void FElectraPlayer::HandleDeferredPlayerEvents()
 				break;
 			}
 			case FPlayerMetricEventBase::EType::PrerollStart:
+			{
 				HandlePlayerEventPrerollStart();
 				break;
+			}
 			case FPlayerMetricEventBase::EType::PrerollEnd:
+			{
 				HandlePlayerEventPrerollEnd();
 				break;
+			}
 			case FPlayerMetricEventBase::EType::PlaybackStart:
+			{
 				HandlePlayerEventPlaybackStart();
 				break;
+			}
 			case FPlayerMetricEventBase::EType::PlaybackPaused:
+			{
 				HandlePlayerEventPlaybackPaused();
 				break;
+			}
 			case FPlayerMetricEventBase::EType::PlaybackResumed:
+			{
 				HandlePlayerEventPlaybackResumed();
 				break;
+			}
 			case FPlayerMetricEventBase::EType::PlaybackEnded:
+			{
 				HandlePlayerEventPlaybackEnded();
 				break;
+			}
 			case FPlayerMetricEventBase::EType::JumpInPlayPosition:
 			{
 				FPlayerMetricEvent_JumpInPlayPosition* Ev = static_cast<FPlayerMetricEvent_JumpInPlayPosition*>(Event.Get());
@@ -1563,8 +1645,10 @@ void FElectraPlayer::HandleDeferredPlayerEvents()
 				break;
 			}
 			case FPlayerMetricEventBase::EType::PlaybackStopped:
+			{
 				HandlePlayerEventPlaybackStopped();
 				break;
+			}
 			case FPlayerMetricEventBase::EType::Error:
 			{
 				FPlayerMetricEvent_Error* Ev = static_cast<FPlayerMetricEvent_Error*>(Event.Get());
@@ -1578,22 +1662,28 @@ void FElectraPlayer::HandleDeferredPlayerEvents()
 				break;
 			}
 			case FPlayerMetricEventBase::EType::DroppedVideoFrame:
+			{
 				HandlePlayerEventDroppedVideoFrame();
 				break;
+			}
 			case FPlayerMetricEventBase::EType::DroppedAudioFrame:
+			{
 				HandlePlayerEventDroppedAudioFrame();
 				break;
+			}
 			default:
+			{
 				break;
+			}
 		}
 	}
 }
 
 void FElectraPlayer::HandlePlayerEventOpenSource(const FString& URL)
 {
-	Status = Status | EPlayerStatus::Connecting;
+	PlayerState.Status = PlayerState.Status | EPlayerStatus::Connecting;
 
-	State = EPlayerState::Preparing;
+	PlayerState.State = EPlayerState::Preparing;
 	DeferredEvents.Enqueue(IElectraPlayerAdapterDelegate::EPlayerEvent::MediaConnecting);
 
 	UE_LOG(LogElectraPlayer, Log, TEXT("[%p][%p] Opening stream at \"%s\""), this, CurrentPlayer.Get(), *SanitizeMessage(URL));
@@ -1636,9 +1726,9 @@ void FElectraPlayer::HandlePlayerEventReceivedMasterPlaylist(const FString& Effe
 
 void FElectraPlayer::HandlePlayerEventReceivedPlaylists()
 {
-	Status = Status & ~EPlayerStatus::Connecting;
+	PlayerState.Status = PlayerState.Status & ~EPlayerStatus::Connecting;
 
-	//Status = Status | EPlayerStatus::Buffering;
+	//PlayerState.Status = PlayerState.Status | EPlayerStatus::Buffering;
 	//DeferredEvents.Enqueue(EPlayerEvent::MediaBuffering);
 
 	// Player starts in paused mode. We need a SetRate() to start playback...
@@ -1818,7 +1908,7 @@ void FElectraPlayer::HandlePlayerEventDataAvailabilityChange(const Electra::Metr
 
 void FElectraPlayer::HandlePlayerEventBufferingStart(Electra::Metrics::EBufferingReason BufferingReason)
 {
-	Status = Status | EPlayerStatus::Buffering;
+	PlayerState.Status = PlayerState.Status | EPlayerStatus::Buffering;
 	DeferredEvents.Enqueue(IElectraPlayerAdapterDelegate::EPlayerEvent::MediaBuffering);
 
 	// Update statistics
@@ -1855,7 +1945,7 @@ void FElectraPlayer::HandlePlayerEventBufferingEnd(Electra::Metrics::EBufferingR
 	// Note: While this event signals the end of buffering the player will now immediately transition into the pre-rolling
 	//       state from which a playback start is not quite possible yet and would incur a slight delay until it is.
 	//       To avoid this we keep the state as buffering until the pre-rolling phase has also completed.
-	//Status = Status & ~EPlayerStatus::Buffering;
+	//PlayerState.Status = PlayerState.Status & ~EPlayerStatus::Buffering;
 
 	// Update statistics
 	FScopeLock Lock(&StatisticsLock);
@@ -2054,7 +2144,7 @@ void FElectraPlayer::HandlePlayerEventPrerollEnd()
 {
 	// Note: See comments in ReportBufferingEnd()
 	//       Preroll follows at the end of buffering and we keep the buffering state until preroll has finished as well.
-	Status = Status & ~EPlayerStatus::Buffering;
+	PlayerState.Status = PlayerState.Status & ~EPlayerStatus::Buffering;
 
 	// Update statistics
 	FScopeLock Lock(&StatisticsLock);
@@ -2076,7 +2166,7 @@ void FElectraPlayer::HandlePlayerEventPrerollEnd()
 
 void FElectraPlayer::HandlePlayerEventPlaybackStart()
 {
-	Status = Status & ~EPlayerStatus::Buffering;
+	PlayerState.Status = PlayerState.Status & ~EPlayerStatus::Buffering;
 	MediaStateOnPlay();
 
 	// Update statistics
@@ -2565,42 +2655,46 @@ void FElectraPlayer::ReportSubtitlesMetrics(const FGuid& InPlayerGuid, const FSt
 
 void FElectraPlayer::MediaStateOnPreparingFinished()
 {
-	if (!ensure(State == EPlayerState::Preparing))
+	if (!ensure(PlayerState.State == EPlayerState::Preparing))
 	{
 		return;
 	}
 
 	CSV_EVENT(ElectraPlayer, TEXT("MediaStateOnPreparingFinished"));
 
-	State = EPlayerState::Stopped;
+	PlayerState.State = EPlayerState::Stopped;
 	DeferredEvents.Enqueue(IElectraPlayerAdapterDelegate::EPlayerEvent::MediaOpened);
 	DeferredEvents.Enqueue(IElectraPlayerAdapterDelegate::EPlayerEvent::TracksChanged);
 }
 
 bool FElectraPlayer::MediaStateOnPlay()
 {
-	if (State != EPlayerState::Stopped && State != EPlayerState::Paused)
+	if (PlayerState.State != EPlayerState::Stopped && PlayerState.State != EPlayerState::Paused)
 	{
 		return false;
 	}
 
 	CSV_EVENT(ElectraPlayer, TEXT("MediaStateOnPlay"));
 
-	State = EPlayerState::Playing;
+	PlayerState.State = EPlayerState::Playing;
+	PlayerState.SetPlayRateFromPlayer(1.0f);
+
 	DeferredEvents.Enqueue(IElectraPlayerAdapterDelegate::EPlayerEvent::PlaybackResumed);
 	return true;
 }
 
 bool FElectraPlayer::MediaStateOnPause()
 {
-	if (State != EPlayerState::Playing)
+	if (PlayerState.State != EPlayerState::Playing)
 	{
 		return false;
 	}
 
 	CSV_EVENT(ElectraPlayer, TEXT("MediaStateOnPause"));
 
-	State = EPlayerState::Paused;
+	PlayerState.State = EPlayerState::Paused;
+	PlayerState.SetPlayRateFromPlayer(0.0f);
+
 	DeferredEvents.Enqueue(IElectraPlayerAdapterDelegate::EPlayerEvent::PlaybackSuspended);
 	return true;
 }
@@ -2609,7 +2703,7 @@ void FElectraPlayer::MediaStateOnEndReached()
 {
 	CSV_EVENT(ElectraPlayer, TEXT("MediaStateOnEndReached"));
 
-	switch (State)
+	switch (PlayerState.State)
 	{
 	case EPlayerState::Preparing:
 	case EPlayerState::Playing:
@@ -2620,7 +2714,7 @@ void FElectraPlayer::MediaStateOnEndReached()
 		// NOP
 		break;
 	}
-	State = EPlayerState::Stopped;
+	PlayerState.State = EPlayerState::Stopped;
 }
 
 void FElectraPlayer::MediaStateOnSeekFinished()
diff --git a/Engine/Plugins/Media/ElectraPlayer/Source/ElectraPlayerRuntime/Private/ElectraPlayer.h b/Engine/Plugins/Media/ElectraPlayer/Source/ElectraPlayerRuntime/Private/ElectraPlayer.h
--- a/Engine/Plugins/Media/ElectraPlayer/Source/ElectraPlayerRuntime/Private/ElectraPlayer.h
+++ b/Engine/Plugins/Media/ElectraPlayer/Source/ElectraPlayerRuntime/Private/ElectraPlayer.h
@@ -367,6 +367,33 @@ private:
 	FIntPoint										LastPresentedFrameDimension;
 
 
+	struct FPlayerState
+	{
+		TOptional<float>		IntendedPlayRate;
+		float					CurrentPlayRate = 0.0f;
+
+		TAtomic<EPlayerState>	State;
+		TAtomic<EPlayerStatus>	Status;
+
+		bool					bUseInternal = false;
+
+		void Reset()
+		{
+			IntendedPlayRate.Reset();
+			CurrentPlayRate = 0.0f;
+			State = EPlayerState::Closed;
+			Status = EPlayerStatus::None;
+		}
+
+		float GetRate() const;
+		EPlayerState GetState() const;
+		EPlayerStatus GetStatus() const;
+
+		void SetIntendedPlayRate(float InIntendedRate);
+		void SetPlayRateFromPlayer(float InCurrentPlayerPlayRate);
+	};
+
+
 	/** Media player Guid */
 	FGuid											PlayerGuid;
 	/** Metric delegates */
@@ -378,9 +405,7 @@ private:
 	/** Option interface **/
 	FPlaystartOptions								PlaystartOptions;
 
-	/**  */
-	TAtomic<EPlayerState>							State;
-	TAtomic<EPlayerStatus>							Status;
+	FPlayerState									PlayerState;
 
 	TAtomic<bool>									bPlayerHasClosed;
 	TAtomic<bool>									bHasPendingError;
@@ -705,7 +730,6 @@ private:
 	void HandlePlayerEventLogMessage(IInfoLog::ELevel InLogLevel, const FString& InLogMessage, int64 InPlayerWallclockMilliseconds);
 	void HandlePlayerEventDroppedVideoFrame();
 	void HandlePlayerEventDroppedAudioFrame();
-
 };
 
 ENUM_CLASS_FLAGS(FElectraPlayer::EPlayerStatus);
