ElectraPlayer: SW decode mode to use the buffer output size from the decoder instead of calculated values

#rnx

diff --git a/Engine/Plugins/Media/ElectraPlayer/Source/ElectraPlayerRuntime/Private/Runtime/Decoder/Windows/VideoDecoderH264_PC.cpp b/Engine/Plugins/Media/ElectraPlayer/Source/ElectraPlayerRuntime/Private/Runtime/Decoder/Windows/VideoDecoderH264_PC.cpp
--- a/Engine/Plugins/Media/ElectraPlayer/Source/ElectraPlayerRuntime/Private/Runtime/Decoder/Windows/VideoDecoderH264_PC.cpp
+++ b/Engine/Plugins/Media/ElectraPlayer/Source/ElectraPlayerRuntime/Private/Runtime/Decoder/Windows/VideoDecoderH264_PC.cpp
@@ -73,6 +73,10 @@ public:
 	void InitializeWithSharedTexture(const TRefCountPtr<ID3D11Device>& InD3D11Device, const TRefCountPtr<ID3D11DeviceContext> InDeviceContext, const TRefCountPtr<IMFSample>& MFSample, const FIntPoint& OutputDim, Electra::FParamDict* InParamDict);
 
 	// Software decode (into texture if DX11 device specified - available only Win8+)
+	void SetSWDecodeTargetBufferSize(uint32 InTargetBufferSize)
+	{
+		TargetBufferAllocSize = InTargetBufferSize;
+	}
 	void PreInitForDecode(FIntPoint OutputDim, const TFunction<void(int32 /*ApiReturnValue*/, const FString& /*Message*/, uint16 /*Code*/, UEMediaError /*Error*/)>& PostError);
 	void ProcessDecodeOutput(FIntPoint OutputDim, Electra::FParamDict* InParamDict);
 
@@ -131,6 +135,7 @@ private:
 	TRefCountPtr<ID3D11Device> D3D11Device;
 
 	// CPU-side buffer to 
+	TOptional<uint32> TargetBufferAllocSize;
 	TArray<uint8> Buffer;
 	uint32 Stride;
 
@@ -208,7 +213,17 @@ void FElectraPlayerVideoDecoderOutputPC::PreInitForDecode(FIntPoint OutputDim, c
 		}
 		else // Software decode into CPU buffer
 		{
-			DWORD AllocSize = ((OutputDim.X + 15) & ~15) * ((OutputDim.Y + 15) & ~15) * sizeof(uint8*) * 2;
+			DWORD AllocSize;
+
+			if (TargetBufferAllocSize.IsSet())
+			{
+				AllocSize = TargetBufferAllocSize.GetValue();
+			}
+			else
+			{
+				AllocSize = ((OutputDim.X + 15) & ~15) * ((OutputDim.Y + 15) & ~15) * sizeof(uint8*) * 2;
+			}
+
 			// SW decode results are just delivered in a simple CPU-side buffer. Create the decoder side version of this...
 			if (FAILED(Result = MFCreateMemoryBuffer(AllocSize, MediaBuffer.GetInitReference())))
 			{
@@ -703,6 +718,18 @@ void FVideoDecoderH264_PC::PreInitDecodeOutputForSW(const FIntPoint& Dim)
 {
 	TSharedPtr<FElectraPlayerVideoDecoderOutputPC, ESPMode::ThreadSafe> DecoderOutput = CurrentRenderOutputBuffer->GetBufferProperties().GetValue("texture").GetSharedPointer<FElectraPlayerVideoDecoderOutputPC>();
 	check(DecoderOutput);
+
+	if (DecoderTransform.IsValid())
+	{
+		MFT_OUTPUT_STREAM_INFO	OutputStreamInfo;
+		HRESULT Result = DecoderTransform->GetOutputStreamInfo(0, &OutputStreamInfo);
+		// Do we need to provide the sample output buffer or does the decoder create it for us?
+		if (SUCCEEDED(Result) && (OutputStreamInfo.dwFlags & (MFT_OUTPUT_STREAM_PROVIDES_SAMPLES | MFT_OUTPUT_STREAM_CAN_PROVIDE_SAMPLES)) == 0)
+		{
+			DecoderOutput->SetSWDecodeTargetBufferSize(OutputStreamInfo.cbSize);
+		}
+	}
+
 	DecoderOutput->PreInitForDecode(Dim, [this](int32 ApiReturnValue, const FString& Message, uint16 Code, UEMediaError Error) { PostError(ApiReturnValue, Message, Code, Error); });
 }
 
