ElectraPlayer: Added support for WebVTT (wvtt) decoding. Enabled subtitle codec switching.

#rnx

diff --git a/Engine/Plugins/Media/ElectraPlayer/Source/ElectraPlayerRuntime/Private/Runtime/Demuxer/ParserISO14496-12.cpp b/Engine/Plugins/Media/ElectraPlayer/Source/ElectraPlayerRuntime/Private/Runtime/Demuxer/ParserISO14496-12.cpp
--- a/Engine/Plugins/Media/ElectraPlayer/Source/ElectraPlayerRuntime/Private/Runtime/Demuxer/ParserISO14496-12.cpp
+++ b/Engine/Plugins/Media/ElectraPlayer/Source/ElectraPlayerRuntime/Private/Runtime/Demuxer/ParserISO14496-12.cpp
@@ -712,6 +712,7 @@ private:
 		static const IParserISO14496_12::FBoxType kSample_stpp = MAKE_BOX_ATOM('s', 't', 'p', 'p');
 		static const IParserISO14496_12::FBoxType kSample_sbtt = MAKE_BOX_ATOM('s', 'b', 't', 't');
 		static const IParserISO14496_12::FBoxType kSample_tx3g = MAKE_BOX_ATOM('t', 'x', '3', 'g');
+		static const IParserISO14496_12::FBoxType kSample_wvtt = MAKE_BOX_ATOM('w', 'v', 't', 't');
 		static const IParserISO14496_12::FBoxType kSample_enca = MAKE_BOX_ATOM('e', 'n', 'c', 'a');
 		static const IParserISO14496_12::FBoxType kSample_encv = MAKE_BOX_ATOM('e', 'n', 'c', 'v');
 
@@ -2242,6 +2243,128 @@ private:
 	};
 
 
+	/**
+	 * WebVTT sample entry (ISO/IEC 14496-30 Section 7.5 Sample entry format)
+	 */
+	class FMP4BoxWVTTSampleEntry : public FMP4BoxSampleEntry
+	{
+	public:
+		FMP4BoxWVTTSampleEntry(IParserISO14496_12::FBoxType InBoxType, int64 InBoxSize, int64 InStartOffset, int64 InDataOffset, bool bInIsLeafBox, int32 InSTSDBoxVersion)
+			: FMP4BoxSampleEntry(InBoxType, InBoxSize, InStartOffset, InDataOffset, bInIsLeafBox)
+			, STSDBoxVersion(InSTSDBoxVersion)
+		{
+		}
+
+		virtual ~FMP4BoxWVTTSampleEntry()
+		{
+		}
+
+		const TArray<uint8> GetRawBoxData() const
+		{
+			return BoxData;
+		}
+
+		int32 GetTranslationX() const
+		{
+			return TranslationX;
+		}
+
+		int32 GetTranslationY() const
+		{
+			return TranslationY;
+		}
+
+		uint16 GetWidth() const
+		{
+			return Width;
+		}
+
+		uint16 GetHeight() const
+		{
+			return Height;
+		}
+
+		uint32 GetTimescale() const
+		{
+			return Timescale;
+		}
+
+	private:
+		FMP4BoxWVTTSampleEntry() = delete;
+		FMP4BoxWVTTSampleEntry(const FMP4BoxWVTTSampleEntry&) = delete;
+
+	protected:
+		virtual UEMediaError ReadAndParseAttributes(FMP4ParseInfo* ParseInfo) override
+		{
+			UEMediaError Error = UEMEDIA_ERROR_OK;
+
+			// Duplicate all the data we are reading into a buffer that we need to provide to the subtitle plugin.
+			FMP4BoxReader::ScopedDataDuplication DuplicateBoxData(ParseInfo->Reader(), BoxData);
+
+			RETURN_IF_ERROR(FMP4BoxSampleEntry::ReadAndParseAttributes(ParseInfo));
+
+			const FMP4BoxTKHD* TKHD = ParseInfo->GetCurrentTrackBox() ? static_cast<const FMP4BoxTKHD*>(ParseInfo->GetCurrentTrackBox()->GetBoxPath(FMP4Box::kBox_tkhd)) : nullptr;
+			if (TKHD)
+			{
+				Width = TKHD->GetWidth();
+				Height = TKHD->GetHeight();
+				// Translation must not be using fractional digits and can thus be treated like integers.
+				TranslationX = (int32) TKHD->GetMatrixValue(6);
+				TranslationY = (int32) TKHD->GetMatrixValue(7);
+			}
+
+			const FMP4BoxMDHD* MDHD = ParseInfo->GetCurrentTrackBox() ? static_cast<const FMP4BoxMDHD*>(ParseInfo->GetCurrentTrackBox()->FindBox(FMP4Box::kBox_mdhd)) : nullptr;
+			if (MDHD)
+			{
+				Timescale = MDHD->GetTimescale();
+			}
+
+			// There ought to be child boxes here now, at least a 'vttC' box.
+			int32 BytesRemaining = (int32) (BoxSize - (ParseInfo->Reader()->GetCurrentReadOffset() - StartOffset));
+			if (BytesRemaining > 0)
+			{
+				// Now that there are additional boxes we ourselves are no longer a leaf box.
+				bIsLeafBox = false;
+				// This will read all the boxes in here and add them as children.
+				for(; BytesRemaining>0; BytesRemaining = (int32) (BoxSize - (ParseInfo->Reader()->GetCurrentReadOffset() - StartOffset)))
+				{
+					IParserISO14496_12::FBoxType	ChildBoxType;
+					int64							ChildBoxSize;
+					uint8							ChildBoxUUID[16];
+
+					int64 BoxStartOffset = ParseInfo->Reader()->GetCurrentReadOffset();
+					RETURN_IF_ERROR(ParseInfo->ReadBoxTypeAndSize(ChildBoxType, ChildBoxSize, ChildBoxUUID));
+#if MEDIA_DEBUG_HAS_BOX_NAMES
+					char boxName[4];
+					*((uint32*)&boxName) = MEDIA_TO_BIG_ENDIAN(ChildBoxType);
+#endif
+					int64 BoxDataOffset = ParseInfo->Reader()->GetCurrentReadOffset();
+					FMP4Box* NextBox = new FMP4BoxIgnored(ChildBoxType, ChildBoxSize, BoxStartOffset, BoxDataOffset, true);
+
+					AddChildBox(NextBox);
+					RETURN_IF_ERROR(ParseInfo->ReadAndParseNextBox(NextBox));
+				}
+			}
+			// There needs to be the 'vttC' config box now.
+			static const IParserISO14496_12::FBoxType kWebVTT_vttC = MAKE_BOX_ATOM('v', 't', 't', 'C');
+			if (!FindBox(kWebVTT_vttC, 1))
+			{
+				UE_LOG(LogElectraMP4Parser, Error, TEXT("ERROR: ElectraPlayer/FMP4BoxWVTTSampleEntry::ReadAndParseAttributes(): Required 'vttC' box not found!"));
+				return UEMEDIA_ERROR_FORMAT_ERROR;
+			}
+			return Error;
+		}
+
+		TArray<uint8>		BoxData;
+		int32				STSDBoxVersion;
+		uint32				Timescale = 0;
+		uint16				Width = 0;
+		uint16				Height = 0;
+		int32				TranslationX = 0;
+		int32				TranslationY = 0;
+	};
+
+
 	/**
 	 * AVC Decoder Configuration box (ISO/IEC 14496-15:2014 - 5.4.2.1.2)
 	 */
@@ -2868,6 +2991,10 @@ private:
 						{
 							NextBox = new FMP4BoxTX3GSampleEntry(ChildBoxType, ChildBoxSize, BoxStartOffset, BoxDataOffset, true, Version);
 						}
+						else if (ChildBoxType == FMP4Box::kSample_wvtt)
+						{
+							NextBox = new FMP4BoxWVTTSampleEntry(ChildBoxType, ChildBoxSize, BoxStartOffset, BoxDataOffset, true, Version);
+						}
 						else
 						{
 							NextBox = new FMP4BoxIgnored(ChildBoxType, ChildBoxSize, BoxStartOffset, BoxDataOffset, true);
@@ -5265,6 +5392,7 @@ private:
 		UEMediaError ParseHVC1SampleType(FTrack* Track, const FMP4Box* SampleBox);
 		UEMediaError ParseMP4ASampleType(FTrack* Track, const FMP4Box* SampleBox);
 		UEMediaError ParseTX3GSampleType(FTrack* Track, const FMP4Box* SampleBox);
+		UEMediaError ParseWVTTSampleType(FTrack* Track, const FMP4Box* SampleBox);
 
 		FMP4ParseInfo* ParsedData;
 
@@ -5364,6 +5492,7 @@ private:
 				return CodecSpecificDataMP4A.GetCodecSpecificData();
 			}
 			case FStreamCodecInformation::ECodec::TX3G:
+			case FStreamCodecInformation::ECodec::WebVTT:
 			{
 				return CodecSpecificDataRAW;
 			}
@@ -5393,6 +5522,7 @@ private:
 				return CodecSpecificDataMP4A.GetRawData();
 			}
 			case FStreamCodecInformation::ECodec::TX3G:
+			case FStreamCodecInformation::ECodec::WebVTT:
 			{
 				return CodecSpecificDataRAW;
 			}
@@ -6401,6 +6531,26 @@ private:
 		return UEMEDIA_ERROR_NOT_SUPPORTED;
 	}
 
+	UEMediaError FParserISO14496_12::ParseWVTTSampleType(FTrack* Track, const FMP4Box* SampleBox)
+	{
+		// Need to check if there is a subtitle decoder capable of decoding this.
+		if (ISubtitleDecoder::IsSupported(TEXT(""), TEXT("wvtt")))
+		{
+			const FMP4BoxWVTTSampleEntry* SubtitleSampleEntry = static_cast<const FMP4BoxWVTTSampleEntry*>(Track->STSDBox->GetChildBox(0));
+
+			Track->CodecSpecificDataRAW = SubtitleSampleEntry->GetRawBoxData();
+			Track->CodecInformation.SetStreamType(EStreamType::Subtitle);
+			Track->CodecInformation.SetCodec(FStreamCodecInformation::ECodec::WebVTT);
+			Track->CodecInformation.SetStreamLanguageCode(Track->GetLanguage());
+			Track->CodecInformation.SetCodecSpecifierRFC6381(TEXT("wvtt"));
+			Track->CodecInformation.SetResolution(FStreamCodecInformation::FResolution(SubtitleSampleEntry->GetWidth(), SubtitleSampleEntry->GetHeight()));
+			Track->CodecInformation.SetTranslation(FStreamCodecInformation::FTranslation(SubtitleSampleEntry->GetTranslationX(), SubtitleSampleEntry->GetTranslationY()));
+			Track->CodecInformation.SetFrameRate(FTimeFraction(1, SubtitleSampleEntry->GetTimescale()));
+			return UEMEDIA_ERROR_OK;
+		}
+		return UEMEDIA_ERROR_NOT_SUPPORTED;
+	}
+
 	UEMediaError FParserISO14496_12::PrepareTracks(TSharedPtrTS<const IParserISO14496_12> OptionalMP4InitSegment)
 	{
 		delete ParsedTrackInfo;
@@ -6733,6 +6883,19 @@ private:
 								}
 								break;
 							}
+							case FMP4Box::kSample_wvtt:
+							{
+								UEMediaError Error = ParseWVTTSampleType(Track.Get(), STSDFirstChildBox);
+								if (Error == UEMEDIA_ERROR_NOT_SUPPORTED)
+								{
+									bIsSupported = false;
+								}
+								else if (Error != UEMEDIA_ERROR_OK)
+								{
+									return Error;
+								}
+								break;
+							}
 							default:
 							{
 								break;
diff --git a/Engine/Plugins/Media/ElectraPlayer/Source/ElectraPlayerRuntime/Private/Runtime/Player/AdaptiveStreamingPlayerDecoders.cpp b/Engine/Plugins/Media/ElectraPlayer/Source/ElectraPlayerRuntime/Private/Runtime/Player/AdaptiveStreamingPlayerDecoders.cpp
--- a/Engine/Plugins/Media/ElectraPlayer/Source/ElectraPlayerRuntime/Private/Runtime/Player/AdaptiveStreamingPlayerDecoders.cpp
+++ b/Engine/Plugins/Media/ElectraPlayer/Source/ElectraPlayerRuntime/Private/Runtime/Player/AdaptiveStreamingPlayerDecoders.cpp
@@ -549,6 +549,11 @@ void FAdaptiveStreamingPlayer::HandleDecoderChanges()
 		VideoDecoder.bDrainingForCodecChange = false;
 		VideoDecoder.bDrainingForCodecChangeDone = false;
 	}
+	if (SubtitleDecoder.bRequireCodecChange)
+	{
+		SubtitleDecoder.Close();
+		SubtitleDecoder.bRequireCodecChange = false;
+	}
 	CreateDecoder(EStreamType::Video);
 	CreateDecoder(EStreamType::Audio);
 	CreateDecoder(EStreamType::Subtitle);
@@ -712,7 +717,7 @@ void FAdaptiveStreamingPlayer::FeedDecoder(EStreamType Type, FMultiTrackAccessUn
 				bCodecChangeDetected = true;
 				if (Decoder)
 				{
-					// Check type of stream. We can currently change the video codec only.
+					// Check type of stream. We can currently change the video and subtitle codecs only.
 					if (Type == EStreamType::Video)
 					{
 						if (VideoDecoder.Decoder && !VideoDecoder.bDrainingForCodecChange)
@@ -721,6 +726,10 @@ void FAdaptiveStreamingPlayer::FeedDecoder(EStreamType Type, FMultiTrackAccessUn
 							VideoDecoder.Decoder->DrainForCodecChange();
 						}
 					}
+					else if (Type == EStreamType::Subtitle)
+					{
+						SubtitleDecoder.bRequireCodecChange = true;
+					}
 					else
 					{
 						FErrorDetail err;
diff --git a/Engine/Plugins/Media/ElectraPlayer/Source/ElectraPlayerRuntime/Private/Runtime/Player/AdaptiveStreamingPlayerInternal.h b/Engine/Plugins/Media/ElectraPlayer/Source/ElectraPlayerRuntime/Private/Runtime/Player/AdaptiveStreamingPlayerInternal.h
--- a/Engine/Plugins/Media/ElectraPlayer/Source/ElectraPlayerRuntime/Private/Runtime/Player/AdaptiveStreamingPlayerInternal.h
+++ b/Engine/Plugins/Media/ElectraPlayer/Source/ElectraPlayerRuntime/Private/Runtime/Player/AdaptiveStreamingPlayerInternal.h
@@ -1040,8 +1040,7 @@ private:
 		FStreamCodecInformation			CurrentCodecInfo;
 		ISubtitleDecoder*				Decoder = nullptr;
 		bool							bIsRunning = false;
-		bool							bDrainingForCodecChange = false;
-		bool							bDrainingForCodecChangeDone = false;
+		bool							bRequireCodecChange = false;
 	};
 
 	class FWorkerThreadMessages
diff --git a/Engine/Plugins/Media/ElectraPlayer/Source/ElectraPlayerRuntime/Private/Runtime/Player/DASH/ManifestDASH.cpp b/Engine/Plugins/Media/ElectraPlayer/Source/ElectraPlayerRuntime/Private/Runtime/Player/DASH/ManifestDASH.cpp
--- a/Engine/Plugins/Media/ElectraPlayer/Source/ElectraPlayerRuntime/Private/Runtime/Player/DASH/ManifestDASH.cpp
+++ b/Engine/Plugins/Media/ElectraPlayer/Source/ElectraPlayerRuntime/Private/Runtime/Player/DASH/ManifestDASH.cpp
@@ -590,7 +590,9 @@ TSharedPtrTS<FManifestDASHInternal::FAdaptationSet> FDASHPlayPeriod::SelectAdapt
 				}
 			}
 			// Matching language not found. Is there an explicit index given?
-			if (!AS.IsValid())
+			// Note: for now we use the explicit override even if a matching language was already found.
+			//       there's a reason the override is specified, like to enforce a specific codec for the same language.
+			//if (!AS.IsValid())
 			{
 				if (Attributes.OverrideIndex.IsSet() && Attributes.OverrideIndex.GetValue() >= 0 && Attributes.OverrideIndex.GetValue() < NumAdaptationSets)
 				{
diff --git a/Engine/Plugins/Media/ElectraPlayer/Source/ElectraPlayerRuntime/Private/Runtime/StreamTypes.h b/Engine/Plugins/Media/ElectraPlayer/Source/ElectraPlayerRuntime/Private/Runtime/StreamTypes.h
--- a/Engine/Plugins/Media/ElectraPlayer/Source/ElectraPlayerRuntime/Private/Runtime/StreamTypes.h
+++ b/Engine/Plugins/Media/ElectraPlayer/Source/ElectraPlayerRuntime/Private/Runtime/StreamTypes.h
@@ -651,6 +651,10 @@ namespace Electra
 
 		virtual bool IsCompatibleWith(const FStreamSelectionAttributes& Other)
 		{
+			if (OverrideIndex.IsSet() && Other.OverrideIndex.IsSet() && OverrideIndex.GetValue() >= 0 && Other.OverrideIndex.GetValue() >= 0 && OverrideIndex.GetValue() != Other.OverrideIndex.GetValue())
+			{
+				return false;
+			}
 			FString Kind1 = Kind.IsSet() ? Kind.GetValue() : FString();
 			FString Kind2 = Other.Kind.IsSet() ? Other.Kind.GetValue() : FString();
 
