ElectraPlayer: Added support for multichannel AAC audio

#rnx

diff --git a/Engine/Plugins/Media/ElectraPlayer/Source/ElectraPlayerRuntime/Private/Runtime/Decoder/Android/AudioDecoderAAC_Android.cpp b/Engine/Plugins/Media/ElectraPlayer/Source/ElectraPlayerRuntime/Private/Runtime/Decoder/Android/AudioDecoderAAC_Android.cpp
--- a/Engine/Plugins/Media/ElectraPlayer/Source/ElectraPlayerRuntime/Private/Runtime/Decoder/Android/AudioDecoderAAC_Android.cpp
+++ b/Engine/Plugins/Media/ElectraPlayer/Source/ElectraPlayerRuntime/Private/Runtime/Decoder/Android/AudioDecoderAAC_Android.cpp
@@ -9,6 +9,7 @@
 #include "Utilities/Utilities.h"
 #include "Utilities/UtilsMPEGAudio.h"
 #include "Utilities/StringHelpers.h"
+#include "Utilities/AudioChannelMapper.h"
 #include "DecoderErrors_Android.h"
 #include "HAL/LowLevelMemTracker.h"
 #include "ElectraPlayerPrivate.h"
@@ -22,6 +23,21 @@ DECLARE_CYCLE_STAT(TEXT("FAudioDecoderAAC::ConvertOutput()"), STAT_ElectraPlayer
 namespace Electra
 {
 
+namespace AACDecoderChannelOutputMappingAndroid
+{
+	#define CP FAudioChannelMapper::EChannelPosition
+	static const CP _1[] = { CP::C };
+	static const CP _2[] = { CP::L, CP::R };
+	static const CP _3[] = { CP::L, CP::R, CP::C };
+	static const CP _4[] = { CP::L, CP::R, CP::C, CP::Cs };
+	static const CP _5[] = { CP::L, CP::R, CP::C, CP::Ls, CP::Rs };
+	static const CP _6[] = { CP::L, CP::R, CP::C, CP::LFE, CP::Ls, CP::Rs };
+	static const CP _7[] = { CP::L, CP::R, CP::C, CP::LFE, CP::Ls, CP::Rs, CP::Cs };
+	static const CP _8[] = { CP::L, CP::R, CP::C, CP::LFE, CP::Ls, CP::Rs, CP::Lsr, CP::Rsr };
+	static const CP * const Order[] = { _1, _2, _3, _4, _5, _6, _7, _8 };
+	#undef CP
+};
+
 
 /**
  * AAC audio decoder class implementation.
@@ -32,6 +48,8 @@ public:
 	static bool Startup(const IAudioDecoderAAC::FSystemConfiguration& InConfig);
 	static void Shutdown();
 
+	static bool CanDecodeStream(const FStreamCodecInformation& InCodecInfo);
+
 	FAudioDecoderAAC();
 	virtual ~FAudioDecoderAAC();
 
@@ -103,10 +121,15 @@ private:
 	FParamDict																BufferAcquireOptions;
 	FParamDict																OutputBufferSampleProperties;
 
+	int32																	PCMBufferSize;
+	int16*																	PCMBuffer;
+	FAudioChannelMapper														ChannelMapper;
+	static const uint8														NumChannelsForConfig[16];
 public:
 	static FSystemConfiguration												SystemConfig;
 };
 
+const uint8								FAudioDecoderAAC::NumChannelsForConfig[16] = { 0, 1, 2, 3, 4, 5, 6, 7, 0, 0, 0, 7, 8, 0, 8, 0 };
 IAudioDecoderAAC::FSystemConfiguration	FAudioDecoderAAC::SystemConfig;
 
 
@@ -140,6 +163,11 @@ void IAudioDecoderAAC::Shutdown()
 	FAudioDecoderAAC::Shutdown();
 }
 
+bool IAudioDecoderAAC::CanDecodeStream(const FStreamCodecInformation& InCodecInfo)
+{
+	return FAudioDecoderAAC::CanDecodeStream(InCodecInfo);
+}
+
 IAudioDecoderAAC* IAudioDecoderAAC::Create()
 {
 	return new FAudioDecoderAAC;
@@ -174,6 +202,16 @@ void FAudioDecoderAAC::Shutdown()
 }
 
 
+//-----------------------------------------------------------------------------
+/**
+ * Determines if this stream can be decoded based on its channel configuration.
+ */
+bool FAudioDecoderAAC::CanDecodeStream(const FStreamCodecInformation& InCodecInfo)
+{
+	return InCodecInfo.GetChannelConfiguration() != 0;
+}
+
+
 //-----------------------------------------------------------------------------
 /**
  * Constructor
@@ -188,6 +226,8 @@ FAudioDecoderAAC::FAudioDecoderAAC()
 	, SessionServices(nullptr)
 	, bInDummyDecodeMode(false)
 	, CurrentOutputBuffer(nullptr)
+	, PCMBuffer(nullptr)
+	, PCMBufferSize(0)
 {
 }
 
@@ -226,11 +266,9 @@ bool FAudioDecoderAAC::CreateDecodedSamplePool()
 {
 	check(Renderer);
 	FParamDict poolOpts;
-	// Assume LC-AAC with at most 6 channels and 16 bit decoded PCM samples. This is ok for HE-AAC as well which can at most be stereo.
-	uint32 frameSize = sizeof(int16) * 6 * 1024;
-
-	poolOpts.Set("max_buffer_size",  FVariantValue((int64) frameSize));
-	poolOpts.Set("num_buffers",      FVariantValue((int64) 8));
+	uint32 frameSize = sizeof(int16) * 8 * 2048;
+	poolOpts.Set("max_buffer_size", FVariantValue((int64) frameSize));
+	poolOpts.Set("num_buffers",     FVariantValue((int64) 8));
 
 	UEMediaError Error = Renderer->CreateBufferPool(poolOpts);
 	check(Error == UEMEDIA_ERROR_OK);
@@ -648,18 +686,51 @@ bool FAudioDecoderAAC::Decode(FAccessUnit* InAccessUnit)
 						PTS.SetFromMicroseconds(OutputBufferInfo.PresentationTimestamp);
 						bool bEOS = OutputBufferInfo.bIsEOS;
 
-						int32 CurrentRenderOutputBufferSize = (int32)CurrentOutputBuffer->GetBufferProperties().GetValue("size").GetInt64();
-						void* CurrentRenderOutputBufferAddress = CurrentOutputBuffer->GetBufferProperties().GetValue("address").GetPointer();
-						check(CurrentRenderOutputBufferSize >= OutputByteCount);
-						if (CurrentRenderOutputBufferSize >= OutputByteCount)
+						check(PCMBufferSize >= OutputByteCount);
+						if (PCMBufferSize >= OutputByteCount)
 						{
-							result = DecoderInstance->GetOutputBufferAndRelease(CurrentRenderOutputBufferAddress, CurrentRenderOutputBufferSize, OutputBufferInfo);
+							void* CopyBuffer = (void*)PCMBuffer;
+							result = DecoderInstance->GetOutputBufferAndRelease(CopyBuffer, PCMBufferSize, OutputBufferInfo);
 							if (result == 0)
 							{
+								if (!ChannelMapper.IsInitialized())
+								{
+									if (NumberOfChannels)
+									{
+										TArray<FAudioChannelMapper::FSourceLayout> Layout;
+										const FAudioChannelMapper::EChannelPosition* ChannelPositions = AACDecoderChannelOutputMappingAndroid::Order[NumberOfChannels - 1];
+										for(uint32 i=0; i<NumberOfChannels; ++i)
+										{
+											FAudioChannelMapper::FSourceLayout lo;
+											lo.ChannelPosition = ChannelPositions[i];
+											Layout.Emplace(MoveTemp(lo));
+										}
+										if (ChannelMapper.Initialize(sizeof(int16), Layout))
+										{
+											// Ok.
+										}
+										else
+										{
+											PostError(0, "Failed to initialize audio channel mapper", ERRCODE_INTERNAL_ANDROID_UNSUPPORTED_NUMBER_OF_AUDIO_CHANNELS);
+											return false;
+										}
+									}
+									else
+									{
+										PostError(0, "Bad number (0) of decoded audio channels", ERRCODE_INTERNAL_ANDROID_UNSUPPORTED_NUMBER_OF_AUDIO_CHANNELS);
+										return false;
+									}
+								}
+
+								int32 CurrentRenderOutputBufferSize = (int32)CurrentOutputBuffer->GetBufferProperties().GetValue("size").GetInt64();
+								void* CurrentRenderOutputBufferAddress = CurrentOutputBuffer->GetBufferProperties().GetValue("address").GetPointer();
+								int32 NumDecodedSamplesPerChannel = OutputByteCount / (sizeof(int16) * NumberOfChannels);
+								ChannelMapper.MapChannels(CurrentRenderOutputBufferAddress, CurrentRenderOutputBufferSize, PCMBuffer, OutputByteCount, NumDecodedSamplesPerChannel);
+
 								OutputBufferSampleProperties.Clear();
-								OutputBufferSampleProperties.Set("num_channels",  FVariantValue((int64) NumberOfChannels));
+								OutputBufferSampleProperties.Set("num_channels",  FVariantValue((int64) ChannelMapper.GetNumTargetChannels()));
 								OutputBufferSampleProperties.Set("sample_rate",   FVariantValue((int64) SamplingRate));
-								OutputBufferSampleProperties.Set("byte_size",     FVariantValue((int64) OutputByteCount));
+								OutputBufferSampleProperties.Set("byte_size",     FVariantValue((int64) (ChannelMapper.GetNumTargetChannels() * sizeof(int16) * NumDecodedSamplesPerChannel)));
 								OutputBufferSampleProperties.Set("duration",      FVariantValue(Duration));
 								OutputBufferSampleProperties.Set("pts",           FVariantValue(PTS));
 								OutputBufferSampleProperties.Set("discontinuity", FVariantValue((bool)bHaveDiscontinuity));
@@ -741,7 +812,7 @@ bool FAudioDecoderAAC::Decode(FAccessUnit* InAccessUnit)
 				if (ConfigRecord.IsValid())
 				{
 					// Parameteric stereo results in stereo (2 channels).
-					NumChannels = ConfigRecord->PSSignal > 0 ? 2 : ConfigRecord->ChannelConfiguration;
+					NumChannels = ConfigRecord->PSSignal > 0 ? 2 : NumChannelsForConfig[ConfigRecord->ChannelConfiguration];
 					SampleRate = ConfigRecord->ExtSamplingFrequency ? ConfigRecord->ExtSamplingFrequency : ConfigRecord->SamplingRate;
 					SamplesPerBlock = ConfigRecord->SBRSignal > 0 ? 2048 : 1024;
 				}
@@ -801,6 +872,9 @@ void FAudioDecoderAAC::WorkerThread()
 	bInDummyDecodeMode = false;
 	bHaveDiscontinuity = false;
 
+	PCMBufferSize = sizeof(int16) * 8 * 2048;
+	PCMBuffer = (int16*)FMemory::Malloc(PCMBufferSize);
+
 	bError = !CreateDecodedSamplePool();
 	check(!bError);
 
@@ -845,6 +919,7 @@ void FAudioDecoderAAC::WorkerThread()
 				ReturnUnusedOutputBuffer();
 				ConfigRecord.Reset();
 				CurrentCodecData.Reset();
+				ChannelMapper.Reset();
 			}
 
 			// Parse the CSD into a configuration record.
@@ -863,6 +938,17 @@ void FAudioDecoderAAC::WorkerThread()
 						PostError(0, "Failed to parse AAC configuration record", ERRCODE_INTERNAL_ANDROID_FAILED_TO_PARSE_CSD);
 						bError = true;
 					}
+					else
+					{
+						// Channel configuration must be specified. Internal CPE/SCE elements are not supported.
+						if (ConfigRecord->ChannelConfiguration == 0)
+						{
+							ConfigRecord.Reset();
+							CurrentCodecData.Reset();
+							PostError(0, "Unsupported AAC channel configuration", ERRCODE_INTERNAL_ANDROID_UNSUPPORTED_NUMBER_OF_AUDIO_CHANNELS);
+							bError = true;
+						}
+					}
 				}
 			}
 
@@ -934,6 +1020,7 @@ void FAudioDecoderAAC::WorkerThread()
 	// Flush any remaining input data.
 	AccessUnits.Flush();
 	AccessUnits.CapacitySet(0);
+	FMemory::Free(PCMBuffer);
 	}
 
 
diff --git a/Engine/Plugins/Media/ElectraPlayer/Source/ElectraPlayerRuntime/Private/Runtime/Decoder/Android/AudioDecoderAAC_JavaWrapper_Android.cpp b/Engine/Plugins/Media/ElectraPlayer/Source/ElectraPlayerRuntime/Private/Runtime/Decoder/Android/AudioDecoderAAC_JavaWrapper_Android.cpp
--- a/Engine/Plugins/Media/ElectraPlayer/Source/ElectraPlayerRuntime/Private/Runtime/Decoder/Android/AudioDecoderAAC_JavaWrapper_Android.cpp
+++ b/Engine/Plugins/Media/ElectraPlayer/Source/ElectraPlayerRuntime/Private/Runtime/Decoder/Android/AudioDecoderAAC_JavaWrapper_Android.cpp
@@ -193,9 +193,10 @@ int32 FAndroidJavaAACAudioDecoder::InitializeDecoder(const MPEG::FAACDecoderConf
 	}
 
 	// Create one
+	const uint8 NumChannelsForConfig[16] = { 0, 1, 2, 3, 4, 5, 6, 7, 0, 0, 0, 7, 8, 0, 8, 0 };
 	JNIEnv*	JEnv = AndroidJavaEnv::GetJavaEnv();
 	jbyteArray  CSD			= MakeJavaByteArray(InParsedConfigurationRecord.GetCodecSpecificData().GetData(), InParsedConfigurationRecord.GetCodecSpecificData().Num());
-	int 		NumChannels = (int) InParsedConfigurationRecord.ChannelConfiguration;
+	int 		NumChannels = (int) InParsedConfigurationRecord.PSSignal == 1 ? 2 : NumChannelsForConfig[InParsedConfigurationRecord.ChannelConfiguration];
 	int			SampleRate  = (int)(InParsedConfigurationRecord.ExtSamplingFrequency ? InParsedConfigurationRecord.ExtSamplingFrequency : InParsedConfigurationRecord.SamplingRate);
 	int32 result = CallMethod<int>(CreateDecoderFN, NumChannels, SampleRate, CSD);
 	JEnv->DeleteLocalRef(CSD);
diff --git a/Engine/Plugins/Media/ElectraPlayer/Source/ElectraPlayerRuntime/Private/Runtime/Decoder/Android/DecoderErrors_Android.h b/Engine/Plugins/Media/ElectraPlayer/Source/ElectraPlayerRuntime/Private/Runtime/Decoder/Android/DecoderErrors_Android.h
--- a/Engine/Plugins/Media/ElectraPlayer/Source/ElectraPlayerRuntime/Private/Runtime/Decoder/Android/DecoderErrors_Android.h
+++ b/Engine/Plugins/Media/ElectraPlayer/Source/ElectraPlayerRuntime/Private/Runtime/Decoder/Android/DecoderErrors_Android.h
@@ -26,5 +26,7 @@
 #define ERRCODE_INTERNAL_ANDROID_RENDER_BUFFER_TOO_SMALL			41
 #define ERRCODE_INTERNAL_ANDROID_COULD_NOT_FIND_OUTPUT_INFO			42
 
+#define ERRCODE_INTERNAL_ANDROID_UNSUPPORTED_NUMBER_OF_AUDIO_CHANNELS	50
+
 #define ERRCODE_INTERNAL_ANDROID_INTERNAL							10000
 
diff --git a/Engine/Plugins/Media/ElectraPlayer/Source/ElectraPlayerRuntime/Private/Runtime/Decoder/Apple/AudioDecoderAAC_Apple.cpp b/Engine/Plugins/Media/ElectraPlayer/Source/ElectraPlayerRuntime/Private/Runtime/Decoder/Apple/AudioDecoderAAC_Apple.cpp
--- a/Engine/Plugins/Media/ElectraPlayer/Source/ElectraPlayerRuntime/Private/Runtime/Decoder/Apple/AudioDecoderAAC_Apple.cpp
+++ b/Engine/Plugins/Media/ElectraPlayer/Source/ElectraPlayerRuntime/Private/Runtime/Decoder/Apple/AudioDecoderAAC_Apple.cpp
@@ -13,6 +13,7 @@
 #include "Utilities/Utilities.h"
 #include "Utilities/UtilsMPEGAudio.h"
 #include "Utilities/StringHelpers.h"
+#include "Utilities/AudioChannelMapper.h"
 #include "DecoderErrors_Apple.h"
 #include "HAL/LowLevelMemTracker.h"
 #include "ElectraPlayerPrivate.h"
@@ -28,6 +29,21 @@ DECLARE_CYCLE_STAT(TEXT("FAudioDecoderAAC::ConvertOutput()"), STAT_ElectraPlayer
 namespace Electra
 {
 
+namespace AACDecoderChannelOutputMappingApple
+{
+	#define CP FAudioChannelMapper::EChannelPosition
+	static const CP _1[] = { CP::C };
+	static const CP _2[] = { CP::L, CP::R };
+	static const CP _3[] = { CP::C, CP::L, CP::R };
+	static const CP _4[] = { CP::C, CP::L, CP::R, CP::Cs };
+	static const CP _5[] = { CP::C, CP::L, CP::R, CP::Ls, CP::Rs };
+	static const CP _6[] = { CP::C, CP::L, CP::R, CP::Ls, CP::Rs, CP::LFE };
+	static const CP _7[] = { CP::C, CP::L, CP::R, CP::Ls, CP::Rs, CP::Cs, CP::LFE };
+	static const CP _8[] = { CP::C, CP::L, CP::R, CP::Ls, CP::Rs, CP::Lsr, CP::Rsr, CP::LFE };
+	static const CP * const Order[] = { _1, _2, _3, _4, _5, _6, _7, _8 };
+	#undef CP
+};
+
 
 /**
  * AAC audio decoder class implementation.
@@ -38,6 +54,8 @@ public:
 	static bool Startup(const IAudioDecoderAAC::FSystemConfiguration& InConfig);
 	static void Shutdown();
 
+	static bool CanDecodeStream(const FStreamCodecInformation& InCodecInfo);
+
 	FAudioDecoderAAC();
 	virtual ~FAudioDecoderAAC();
 
@@ -149,11 +167,13 @@ private:
 	FParamDict																OutputBufferSampleProperties;
 	int32																	PCMBufferSize;
 	int16*																	PCMBuffer;
-
+	FAudioChannelMapper														ChannelMapper;
+	static const uint8														NumChannelsForConfig[16];
 public:
 	static FSystemConfiguration												SystemConfig;
 };
 
+const uint8								FAudioDecoderAAC::NumChannelsForConfig[16] = { 0, 1, 2, 3, 4, 5, 6, 7, 0, 0, 0, 7, 8, 0, 8, 0 };
 IAudioDecoderAAC::FSystemConfiguration	FAudioDecoderAAC::SystemConfig;
 
 
@@ -187,6 +207,11 @@ void IAudioDecoderAAC::Shutdown()
 	FAudioDecoderAAC::Shutdown();
 }
 
+bool IAudioDecoderAAC::CanDecodeStream(const FStreamCodecInformation& InCodecInfo)
+{
+	return FAudioDecoderAAC::CanDecodeStream(InCodecInfo);
+}
+
 IAudioDecoderAAC* IAudioDecoderAAC::Create()
 {
 	return new FAudioDecoderAAC;
@@ -221,6 +246,20 @@ void FAudioDecoderAAC::Shutdown()
 }
 
 
+//-----------------------------------------------------------------------------
+/**
+ * Determines if this stream can be decoded based on its channel configuration.
+ */
+bool FAudioDecoderAAC::CanDecodeStream(const FStreamCodecInformation& InCodecInfo)
+{
+	return InCodecInfo.GetChannelConfiguration() == 1 ||		// Mono
+		   InCodecInfo.GetChannelConfiguration() == 2 ||		// Stereo
+		   InCodecInfo.GetChannelConfiguration() == 3 ||		// L/C/R
+		   InCodecInfo.GetChannelConfiguration() == 6 ||		// 5.1
+		   InCodecInfo.GetChannelConfiguration() == 12;			// 7.1
+}
+
+
 //-----------------------------------------------------------------------------
 /**
  * Constructor
@@ -277,11 +316,9 @@ bool FAudioDecoderAAC::CreateDecodedSamplePool()
 {
 	check(Renderer);
 	FParamDict poolOpts;
-	// Assume LC-AAC with at most 8 channels and 16 bit decoded PCM samples. This is ok for HE-AAC as well which can at most be stereo.
-	uint32 frameSize = sizeof(int16) * 8 * 1024;
-
-	poolOpts.Set("max_buffer_size",  FVariantValue((int64) frameSize));
-	poolOpts.Set("num_buffers",      FVariantValue((int64) 8));
+	uint32 frameSize = sizeof(int16) * 8 * 2048;
+	poolOpts.Set("max_buffer_size", FVariantValue((int64) frameSize));
+	poolOpts.Set("num_buffers",     FVariantValue((int64) 8));
 
 	UEMediaError Error = Renderer->CreateBufferPool(poolOpts);
 	check(Error == UEMEDIA_ERROR_OK);
@@ -540,7 +577,7 @@ bool FAudioDecoderAAC::InternalDecoderCreate()
 	DecoderInstance->InputDescr.mFormatFlags      = DecoderInstance->InputDescr.mFormatID == kAudioFormatMPEG4AAC ? kMPEG4Object_AAC_LC : kMPEG4Object_AAC_SBR;
 	DecoderInstance->InputDescr.mFramesPerPacket  = DecoderInstance->InputDescr.mFormatID == kAudioFormatMPEG4AAC ? 1024 : 2048;
     // Decoding parametric stereo (PS) implies the output to go from 1 to 2 channels.
-	DecoderInstance->InputDescr.mChannelsPerFrame = ConfigRecord->PSSignal > 0 ? 2 : ConfigRecord->ChannelConfiguration;
+	DecoderInstance->InputDescr.mChannelsPerFrame = ConfigRecord->PSSignal > 0 ? 2 : NumChannelsForConfig[ConfigRecord->ChannelConfiguration];
 
 	// Want LPCM output, use convenience method to set this up.
 	memset(&DecoderInstance->OutputDescr, 0, sizeof(AudioStreamBasicDescription));
@@ -693,9 +730,9 @@ bool FAudioDecoderAAC::Decode(FAccessUnit* InAccessUnit)
 				DecoderInstance->WorkData.Clear();
 				DecoderInstance->WorkData.InputData = Electra::AdvancePointer(InAccessUnit->AUData, nDataOffset);
 				DecoderInstance->WorkData.InputSize = InAccessUnit->AUSize - nDataOffset;
-				int32 CurrentRenderOutputBufferSize = (int32)CurrentOutputBuffer->GetBufferProperties().GetValue("size").GetInt64();
-				void* CurrentRenderOutputBufferAddress = CurrentOutputBuffer->GetBufferProperties().GetValue("address").GetPointer();
-				int32 CurrentRenderOutputBufferOffset = 0;
+
+				void* CurrentOutputBufferAddr = PCMBuffer;
+				int32 CurrentOutputBufferAvailSize = PCMBufferSize;
 				int32 NumSamplesProduced = 0;
 				uint32 NumberOfChannels   = DecoderInstance->OutputDescr.mChannelsPerFrame;
 				uint32 SamplingRate       = DecoderInstance->InputDescr.mSampleRate;
@@ -703,9 +740,9 @@ bool FAudioDecoderAAC::Decode(FAccessUnit* InAccessUnit)
 				{
 					AudioBuffer		OutputBuffer;
 					AudioBufferList OutputBufferList;
-					OutputBuffer.mNumberChannels    = DecoderInstance->OutputDescr.mChannelsPerFrame;
-					OutputBuffer.mDataByteSize	    = PCMBufferSize;
-					OutputBuffer.mData			    = PCMBuffer;
+					OutputBuffer.mNumberChannels    = NumberOfChannels;
+					OutputBuffer.mDataByteSize	    = CurrentOutputBufferAvailSize;
+					OutputBuffer.mData			    = CurrentOutputBufferAddr;
 					OutputBufferList.mNumberBuffers = 1;
 					OutputBufferList.mBuffers[0]    = OutputBuffer;
 
@@ -726,11 +763,11 @@ bool FAudioDecoderAAC::Decode(FAccessUnit* InAccessUnit)
 						NumSamplesProduced += InOutPacketSize;
 
 						int32 OutputByteCount = InOutPacketSize * NumberOfChannels * sizeof(int16);
-						check(OutputByteCount + CurrentRenderOutputBufferOffset <= CurrentRenderOutputBufferSize);
-						if (OutputByteCount + CurrentRenderOutputBufferOffset <= CurrentRenderOutputBufferSize)
+						check(OutputByteCount <= CurrentOutputBufferAvailSize);
+						if (OutputByteCount <= CurrentOutputBufferAvailSize)
 						{
-							FMemory::Memcpy(Electra::AdvancePointer(CurrentRenderOutputBufferAddress, CurrentRenderOutputBufferOffset), PCMBuffer, OutputByteCount);
-							CurrentRenderOutputBufferOffset += OutputByteCount;
+							CurrentOutputBufferAddr = Electra::AdvancePointer(CurrentOutputBufferAddr, OutputByteCount);
+							CurrentOutputBufferAvailSize -= OutputByteCount;
 						}
 						else
 						{
@@ -758,13 +795,50 @@ bool FAudioDecoderAAC::Decode(FAccessUnit* InAccessUnit)
 				{
 					SCOPE_CYCLE_COUNTER(STAT_ElectraPlayer_AudioAACConvertOutput);
 					CSV_SCOPED_TIMING_STAT(ElectraPlayer, AudioAACConvertOutput);
+
+					if (!ChannelMapper.IsInitialized())
+					{
+						if (NumberOfChannels)
+						{
+							TArray<FAudioChannelMapper::FSourceLayout> Layout;
+							const FAudioChannelMapper::EChannelPosition* ChannelPositions = AACDecoderChannelOutputMappingApple::Order[NumberOfChannels - 1];
+							for(uint32 i=0; i<NumberOfChannels; ++i)
+							{
+								FAudioChannelMapper::FSourceLayout lo;
+								lo.ChannelPosition = ChannelPositions[i];
+								Layout.Emplace(MoveTemp(lo));
+							}
+							if (ChannelMapper.Initialize(sizeof(int16), Layout))
+							{
+								// Ok.
+							}
+							else
+							{
+								PostError(0, "Failed to initialize audio channel mapper", ERRCODE_INTERNAL_APPLE_UNSUPPORTED_NUMBER_OF_AUDIO_CHANNELS);
+								return false;
+							}
+						}
+						else
+						{
+							PostError(0, "Bad number (0) of decoded audio channels", ERRCODE_INTERNAL_APPLE_UNSUPPORTED_NUMBER_OF_AUDIO_CHANNELS);
+							return false;
+						}
+					}
+
+
+					int32 CurrentRenderOutputBufferSize = (int32)CurrentOutputBuffer->GetBufferProperties().GetValue("size").GetInt64();
+					void* CurrentRenderOutputBufferAddress = CurrentOutputBuffer->GetBufferProperties().GetValue("address").GetPointer();
+					int32 NumDecodedBytes = PCMBufferSize - CurrentOutputBufferAvailSize;
+					int32 NumDecodedSamplesPerChannel = NumDecodedBytes / (sizeof(int16) * NumberOfChannels);
+					ChannelMapper.MapChannels(CurrentRenderOutputBufferAddress, CurrentRenderOutputBufferSize, PCMBuffer, NumDecodedBytes, NumDecodedSamplesPerChannel);
+
 					FTimeValue Duration;
 					Duration.SetFromND(NumSamplesProduced, SamplingRate);
 
 					OutputBufferSampleProperties.Clear();
-					OutputBufferSampleProperties.Set("num_channels",  FVariantValue((int64) NumberOfChannels));
+					OutputBufferSampleProperties.Set("num_channels",  FVariantValue((int64)ChannelMapper.GetNumTargetChannels()));
+					OutputBufferSampleProperties.Set("byte_size",     FVariantValue((int64)(ChannelMapper.GetNumTargetChannels() * NumDecodedSamplesPerChannel * sizeof(int16))));
 					OutputBufferSampleProperties.Set("sample_rate",   FVariantValue((int64) SamplingRate));
-					OutputBufferSampleProperties.Set("byte_size",     FVariantValue((int64) CurrentRenderOutputBufferOffset));
 					OutputBufferSampleProperties.Set("duration",      FVariantValue(Duration));
 					OutputBufferSampleProperties.Set("pts",           FVariantValue(CurrentPTS));
 					OutputBufferSampleProperties.Set("discontinuity", FVariantValue((bool)bHaveDiscontinuity));
@@ -798,7 +872,7 @@ bool FAudioDecoderAAC::Decode(FAccessUnit* InAccessUnit)
 				// With a valid configuration record we can use the actual values.
 				if (ConfigRecord.IsValid())
 				{
-					NumChannels = ConfigRecord->PSSignal > 0 ? 2 : ConfigRecord->ChannelConfiguration;
+					NumChannels = ConfigRecord->PSSignal > 0 ? 2 : NumChannelsForConfig[ConfigRecord->ChannelConfiguration];
 					SampleRate = ConfigRecord->ExtSamplingFrequency ? ConfigRecord->ExtSamplingFrequency : ConfigRecord->SamplingRate;
 					SamplesPerBlock = ConfigRecord->SBRSignal > 0 ? 2048 : 1024;
 				}
@@ -866,7 +940,7 @@ void FAudioDecoderAAC::WorkerThread()
 	FAccessUnit*	CurrentAccessUnit = nullptr;
 	bool			bError			  = false;
 
-	PCMBufferSize   = sizeof(int16) * 8 * 2048;			// max. of 8 channels decoding HE-AAC
+	PCMBufferSize   = sizeof(int16) * 8 * 2048 * 2;			// max. of 8 channels decoding HE-AAC with twice the room to assemble packets
 	PCMBuffer   	= (int16*)FMemory::Malloc(PCMBufferSize, 32);
 	CurrentOutputBuffer = nullptr;
 	bInDummyDecodeMode = false;
@@ -916,6 +990,7 @@ void FAudioDecoderAAC::WorkerThread()
 				ReturnUnusedOutputBuffer();
 				ConfigRecord.Reset();
 				CurrentCodecData.Reset();
+				ChannelMapper.Reset();
 			}
 
 			// Parse the CSD into a configuration record.
@@ -934,6 +1009,21 @@ void FAudioDecoderAAC::WorkerThread()
 						PostError(0, "Failed to parse AAC configuration record", ERRCODE_INTERNAL_APPLE_FAILED_TO_PARSE_CSD);
 						bError = true;
 					}
+					else
+					{
+						// Check for unsupported configurations
+						if (ConfigRecord->ChannelConfiguration != 1 &&		// Mono
+							ConfigRecord->ChannelConfiguration != 2 &&		// Stereo
+							ConfigRecord->ChannelConfiguration != 3 &&		// L+C+R
+							ConfigRecord->ChannelConfiguration != 6 &&		// 5.1
+							ConfigRecord->ChannelConfiguration != 12)		// 7.1
+						{
+							ConfigRecord.Reset();
+							CurrentCodecData.Reset();
+							PostError(0, "Unsupported AAC channel configuration", ERRCODE_INTERNAL_APPLE_UNSUPPORTED_NUMBER_OF_AUDIO_CHANNELS);
+							bError = true;
+						}
+					}
 				}
 			}
 
diff --git a/Engine/Plugins/Media/ElectraPlayer/Source/ElectraPlayerRuntime/Private/Runtime/Decoder/Apple/DecoderErrors_Apple.h b/Engine/Plugins/Media/ElectraPlayer/Source/ElectraPlayerRuntime/Private/Runtime/Decoder/Apple/DecoderErrors_Apple.h
--- a/Engine/Plugins/Media/ElectraPlayer/Source/ElectraPlayerRuntime/Private/Runtime/Decoder/Apple/DecoderErrors_Apple.h
+++ b/Engine/Plugins/Media/ElectraPlayer/Source/ElectraPlayerRuntime/Private/Runtime/Decoder/Apple/DecoderErrors_Apple.h
@@ -10,6 +10,7 @@
 #define ERRCODE_INTERNAL_APPLE_FAILED_TO_PARSE_CSD					5
 #define ERRCODE_INTERNAL_APPLE_FAILED_TO_DECODE_AUDIO				6
 #define ERRCODE_INTERNAL_APPLE_AUDIO_OUTPUT_BUFFER_TOO_SMALL		7
+#define ERRCODE_INTERNAL_APPLE_UNSUPPORTED_NUMBER_OF_AUDIO_CHANNELS	8
 
 #define ERRCODE_INTERNAL_APPLE_COULD_NOT_CREATE_VIDEO_DECODER		100
 #define ERRCODE_INTERNAL_APPLE_COULD_NOT_CREATE_IMAGE_POOL			101
diff --git a/Engine/Plugins/Media/ElectraPlayer/Source/ElectraPlayerRuntime/Private/Runtime/Decoder/AudioDecoderAAC.h b/Engine/Plugins/Media/ElectraPlayer/Source/ElectraPlayerRuntime/Private/Runtime/Decoder/AudioDecoderAAC.h
--- a/Engine/Plugins/Media/ElectraPlayer/Source/ElectraPlayerRuntime/Private/Runtime/Decoder/AudioDecoderAAC.h
+++ b/Engine/Plugins/Media/ElectraPlayer/Source/ElectraPlayerRuntime/Private/Runtime/Decoder/AudioDecoderAAC.h
@@ -42,6 +42,8 @@ namespace Electra
 		static bool Startup(const FSystemConfiguration& Config);
 		static void Shutdown();
 
+		static bool CanDecodeStream(const FStreamCodecInformation& InCodecInfo);
+
 		static IAudioDecoderAAC* Create();
 
 		virtual ~IAudioDecoderAAC() = default;
diff --git a/Engine/Plugins/Media/ElectraPlayer/Source/ElectraPlayerRuntime/Private/Runtime/Decoder/DX/AudioDecoderAAC_DX.cpp b/Engine/Plugins/Media/ElectraPlayer/Source/ElectraPlayerRuntime/Private/Runtime/Decoder/DX/AudioDecoderAAC_DX.cpp
--- a/Engine/Plugins/Media/ElectraPlayer/Source/ElectraPlayerRuntime/Private/Runtime/Decoder/DX/AudioDecoderAAC_DX.cpp
+++ b/Engine/Plugins/Media/ElectraPlayer/Source/ElectraPlayerRuntime/Private/Runtime/Decoder/DX/AudioDecoderAAC_DX.cpp
@@ -14,6 +14,7 @@
 #include "Utilities/UtilsMPEG.h"
 #include "Utilities/UtilsMPEGAudio.h"
 #include "Utilities/StringHelpers.h"
+#include "Utilities/AudioChannelMapper.h"
 
 #include "DecoderErrors_DX.h"
 
@@ -63,13 +64,11 @@ DECLARE_CYCLE_STAT(TEXT("FAudioDecoderAAC::ConvertOutput()"), STAT_ElectraPlayer
 
 
 #define AACDEC_PCM_SAMPLE_SIZE  2048
-#define AACDEC_MAX_CHANNELS		6
+#define AACDEC_MAX_CHANNELS		8
 
 
 namespace Electra
 {
-
-
 	/**
 	 * AAC audio decoder class implementation.
 	**/
@@ -79,6 +78,8 @@ namespace Electra
 		static bool Startup(const IAudioDecoderAAC::FSystemConfiguration& InConfig);
 		static void Shutdown();
 
+		static bool CanDecodeStream(const FStreamCodecInformation& InCodecInfo);
+
 		FAudioDecoderAAC();
 		virtual ~FAudioDecoderAAC();
 
@@ -192,10 +193,13 @@ namespace Electra
 		FParamDict																BufferAcquireOptions;
 		FParamDict																OutputBufferSampleProperties;
 
+		FAudioChannelMapper														ChannelMapper;
+		static const uint8														NumChannelsForConfig[16];
 	public:
 		static FSystemConfiguration												SystemConfig;
 	};
 
+	const uint8								FAudioDecoderAAC::NumChannelsForConfig[16] = { 0, 1, 2, 3, 4, 5, 6, 7, 0, 0, 0, 7, 8, 0, 8, 0 };
 	IAudioDecoderAAC::FSystemConfiguration	FAudioDecoderAAC::SystemConfig;
 
 
@@ -229,6 +233,11 @@ namespace Electra
 		FAudioDecoderAAC::Shutdown();
 	}
 
+	bool IAudioDecoderAAC::CanDecodeStream(const FStreamCodecInformation& InCodecInfo)
+	{
+		return FAudioDecoderAAC::CanDecodeStream(InCodecInfo);
+	}
+
 	IAudioDecoderAAC* IAudioDecoderAAC::Create()
 	{
 		return new FAudioDecoderAAC;
@@ -265,6 +274,15 @@ namespace Electra
 	}
 
 
+	//-----------------------------------------------------------------------------
+	/**
+	 * Determines if this stream can be decoded based on its channel configuration.
+	 */
+	bool FAudioDecoderAAC::CanDecodeStream(const FStreamCodecInformation& InCodecInfo)
+	{
+		return InCodecInfo.GetChannelConfiguration() <= 6;
+	}
+
 
 	//-----------------------------------------------------------------------------
 	/**
@@ -309,10 +327,7 @@ namespace Electra
 	{
 		check(Renderer);
 		FParamDict poolOpts;
-		// Assume LC-AAC with at most 6 channels and 16 bit decoded PCM samples. This is ok for HE-AAC as well which can at most be stereo.
-		uint32 frameSize = sizeof(int16) * AACDEC_MAX_CHANNELS * 1024;
-		static_assert(sizeof(int16) * AACDEC_MAX_CHANNELS * 1024 >= sizeof(int16) * 2 * 2048, "Sample pool size not correct!");
-
+		uint32 frameSize = sizeof(int16) * AACDEC_MAX_CHANNELS * AACDEC_PCM_SAMPLE_SIZE;
 		poolOpts.Set("max_buffer_size", FVariantValue((int64)frameSize));
 		//	poolOpts.Set("buffer_alignment", FVariantValue((int64) 32));
 		poolOpts.Set("num_buffers", FVariantValue((int64)8));
@@ -520,6 +535,7 @@ namespace Electra
 		TRefCountPtr<IMFMediaType>	MediaType;
 		HRESULT					res;
 
+		ChannelMapper.Reset();
 		if (!ConfigRecord.IsValid())
 		{
 			PostError(0, "No CSD to create audio decoder with", ERRCODE_INTERNAL_COULD_NOT_CREATE_DECODER);
@@ -536,7 +552,12 @@ namespace Electra
 		VERIFY_HR(MediaType->SetGUID(MF_MT_SUBTYPE, MEDIASUBTYPE_RAW_AAC1_Audio), "Failed to set input media audio type to RAW AAC", ERRCODE_INTERNAL_COULD_NOT_SET_INPUT_AUDIO_AAC_SUBTYPE);
 		VERIFY_HR(MediaType->SetUINT32(MF_MT_AAC_PAYLOAD_TYPE, PayloadType), "Failed to set input media audio payload type", ERRCODE_INTERNAL_COULD_NOT_SET_INPUT_AUDIO_AAC_SUBTYPE);
 		VERIFY_HR(MediaType->SetUINT32(MF_MT_AUDIO_SAMPLES_PER_SECOND, ConfigRecord->SamplingRate), FString::Printf(TEXT("Failed to set input audio sampling rate to %u"), ConfigRecord->SamplingRate), ERRCODE_INTERNAL_COULD_NOT_SET_INPUT_AUDIO_SAMPLING_RATE);
-		VERIFY_HR(MediaType->SetUINT32(MF_MT_AUDIO_NUM_CHANNELS, ConfigRecord->ChannelConfiguration), FString::Printf(TEXT("Failed to set input audio number of channels to %u"), ConfigRecord->ChannelConfiguration), ERRCODE_INTERNAL_COULD_NOT_SET_INPUT_AUDIO_CHANNEL_COUNT);
+		// ConfigRecord->ChannelConfiguration is in the range 0-15.
+		check(ConfigRecord->ChannelConfiguration < 16);
+		if (NumChannelsForConfig[ConfigRecord->ChannelConfiguration])
+		{
+			VERIFY_HR(MediaType->SetUINT32(MF_MT_AUDIO_NUM_CHANNELS, NumChannelsForConfig[ConfigRecord->ChannelConfiguration]), FString::Printf(TEXT("Failed to set input audio number of channels for configuration %u"), ConfigRecord->ChannelConfiguration), ERRCODE_INTERNAL_COULD_NOT_SET_INPUT_AUDIO_CHANNEL_COUNT);
+		}
 		VERIFY_HR(MediaType->SetBlob(MF_MT_USER_DATA, ConfigRecord->GetCodecSpecificData().GetData(), ConfigRecord->GetCodecSpecificData().Num()), "Failed to set input audio CSD", ERRCODE_INTERNAL_COULD_NOT_SET_INPUT_AUDIO_CSD);
 		// Set input media type with decoder
 		VERIFY_HR(Decoder->SetInputType(0, MediaType, 0), "Failed to set audio decoder input type", ERRCODE_INTERNAL_COULD_NOT_SET_INPUT_TYPE);
@@ -562,11 +583,10 @@ namespace Electra
 
 	bool FAudioDecoderAAC::SetDecoderOutputType()
 	{
-		TRefCountPtr<IMFMediaType> MediaType;
-		HRESULT res;
-
-		uint32 TypeIndex = 0;
-		while (SUCCEEDED(DecoderTransform->GetOutputAvailableType(0, TypeIndex++, MediaType.GetInitReference())))
+		TRefCountPtr<IMFMediaType>	MediaType;
+		HRESULT						res;
+		uint32						TypeIndex = 0;
+		while(SUCCEEDED(DecoderTransform->GetOutputAvailableType(0, TypeIndex++, MediaType.GetInitReference())))
 		{
 			GUID Subtype;
 			res = MediaType->GetGUID(MF_MT_SUBTYPE, &Subtype);
@@ -626,9 +646,9 @@ namespace Electra
 
 	bool FAudioDecoderAAC::Decode(int32& OutInputSizeConsumed, const void* InData, int32 InSize, int64 PTS, int64 Duration, bool bFlushOnly, bool bRender)
 	{
-		HRESULT				res;
+		HRESULT					res;
 		TRefCountPtr<IMFSample>	InputSample;
-		bool				bFlush = InData == nullptr;
+		bool					bFlush = InData == nullptr;
 
 		OutInputSizeConsumed = 0;
 		if (bFlush)
@@ -647,10 +667,10 @@ namespace Electra
 		{
 			// Create the input sample.
 			TRefCountPtr<IMFMediaBuffer>	InputSampleBuffer;
-			BYTE* pbNewBuffer = nullptr;
-			DWORD					dwMaxBufferSize = 0;
-			DWORD					dwSize = 0;
-			LONGLONG				llSampleTime = 0;
+			BYTE*							pbNewBuffer = nullptr;
+			DWORD							dwMaxBufferSize = 0;
+			DWORD							dwSize = 0;
+			LONGLONG						llSampleTime = 0;
 
 			SCOPE_CYCLE_COUNTER(STAT_ElectraPlayer_AudioAACDecode);
 			CSV_SCOPED_TIMING_STAT(ElectraPlayer, AudioAACDecode);
@@ -669,7 +689,7 @@ namespace Electra
 			VERIFY_HR(InputSample->SetSampleDuration(llSampleTime), "Failed to set audio decode intput sample duration", ERRCODE_INTERNAL_COULD_NOT_SET_INPUT_SAMPLE_DURATION);
 		}
 
-		while (!TerminateThreadEvent.IsSignaled())
+		while(!TerminateThreadEvent.IsSignaled())
 		{
 			if (FlushDecoderEvent.IsSignaled() && !bFlush)
 			{
@@ -686,7 +706,7 @@ namespace Electra
 				}
 			}
 
-			DWORD	dwStatus = 0;
+			DWORD dwStatus = 0;
 			CurrentDecoderOutputBuffer->PrepareForProcess();
 
 			{
@@ -745,13 +765,14 @@ namespace Electra
 				CurrentDecoderOutputBuffer.Reset();
 
 				TRefCountPtr<IMFMediaBuffer>	DecodedLinearOutputBuffer;
-				DWORD					dwBufferLen;
-				DWORD					dwMaxBufferLen;
-				BYTE* pDecompressedData = nullptr;
-				LONGLONG				llTimeStamp = 0;
-				WAVEFORMATEX			OutputWaveFormat;
-				WAVEFORMATEX* OutputWaveFormatPtr = nullptr;
-				UINT32					WaveFormatSize = 0;
+				DWORD							dwBufferLen;
+				DWORD							dwMaxBufferLen;
+				BYTE*							pDecompressedData = nullptr;
+				LONGLONG						llTimeStamp = 0;
+				WAVEFORMATEX					OutputWaveFormat;
+				WAVEFORMATEX*					OutputWaveFormatPtr = nullptr;
+				UINT32							WaveFormatSize = 0;
+				DWORD							ChannelMask = 0;
 
 				VERIFY_HR(DecodedOutputSample->GetSampleTime(&llTimeStamp), "Failed to get audio decoder output sample timestamp", ERRCODE_INTERNAL_COULD_NOT_GET_OUTPUT_SAMPLE_TIME);
 				VERIFY_HR(DecodedOutputSample->ConvertToContiguousBuffer(DecodedLinearOutputBuffer.GetInitReference()), "Failed to convert audio decoder output sample to contiguous buffer", ERRCODE_INTERNAL_COULD_NOT_MAKE_CONTIGUOUS_OUTPUT_BUFFER);
@@ -759,6 +780,10 @@ namespace Electra
 				VERIFY_HR(DecodedLinearOutputBuffer->Lock(&pDecompressedData, &dwMaxBufferLen, &dwBufferLen), "Failed to lock audio decoder output buffer", ERRCODE_INTERNAL_COULD_NOT_LOCK_OUTPUT_BUFFER)
 				VERIFY_HR(MFCreateWaveFormatExFromMFMediaType(CurrentOutputMediaType.GetReference(), &OutputWaveFormatPtr, &WaveFormatSize, MFWaveFormatExConvertFlag_Normal), "Failed to create audio decoder output buffer format info", ERRCODE_INTERNAL_COULD_NOT_CREATE_OUTPUT_BUFFER_FORMAT_INFO);
 				FMemory::Memcpy(&OutputWaveFormat, OutputWaveFormatPtr, sizeof(OutputWaveFormat));
+				if (OutputWaveFormatPtr->wFormatTag == WAVE_FORMAT_EXTENSIBLE)
+				{
+					ChannelMask = ((WAVEFORMATEXTENSIBLE*)OutputWaveFormatPtr)->dwChannelMask;
+				}
 				CoTaskMemFree(OutputWaveFormatPtr);
 				OutputWaveFormatPtr = nullptr;
 
@@ -770,8 +795,77 @@ namespace Electra
 				//int32 nSamplesProduced = dwBufferLen / (sizeof(float) * OutputWaveFormat.nChannels);
 				//const float *pPCMSamples = (const float *)pDecompressedData;
 
+				if (!ChannelMapper.IsInitialized())
+				{
+					bool bMapperOk = false;
+					if (ChannelMask == 0)
+					{
+						check(ConfigRecord.IsValid());
+						if (ConfigRecord->ChannelConfiguration == OutputWaveFormat.nChannels)
+						{
+							bMapperOk = ChannelMapper.Initialize(sizeof(int16), ConfigRecord->ChannelConfiguration);
+						}
+						else if (ConfigRecord->PSSignal == 1)
+						{
+							// Parametric stereo outputs as stereo
+							bMapperOk = ChannelMapper.Initialize(sizeof(int16), OutputWaveFormat.nChannels);
+						}
+						else
+						{
+							// For up to 6 channels the AAC channel configuration value lines up with the actual number
+							// of decoded channels albeit with a different layout. In absence of any better knowledge we
+							// just take the number of channels for the configuration value and hope the channels line up.
+							bMapperOk = ChannelMapper.Initialize(sizeof(int16), OutputWaveFormat.nChannels);
+						}
+					}
+					else
+					{
+						// The channel positions from MS are oddly named with "back left" and "side left" in that "back left" is
+						// what is traditionally "left surround" in 5.1 and "left surround side" in 7.1 with the additional 2 surround
+						// channels called "left surround rear".
+						// As long as the AAC decoder only decodes 6 channels and the two surrounds have the channel mask set as
+						// "SPEAKER_BACK_LEFT" and "SPEAKER_BACK_RIGHT" there's no ambiguity.
+						static FAudioChannelMapper::EChannelPosition ChannelPositions[] = {
+							FAudioChannelMapper::EChannelPosition::L,		// SPEAKER_FRONT_LEFT
+							FAudioChannelMapper::EChannelPosition::R,		// SPEAKER_FRONT_RIGHT
+							FAudioChannelMapper::EChannelPosition::C,		// SPEAKER_FRONT_CENTER
+							FAudioChannelMapper::EChannelPosition::LFE,		// SPEAKER_LOW_FREQUENCY
+							FAudioChannelMapper::EChannelPosition::Ls,		// SPEAKER_BACK_LEFT
+							FAudioChannelMapper::EChannelPosition::Rs,		// SPEAKER_BACK_RIGHT
+							FAudioChannelMapper::EChannelPosition::Lc,		// SPEAKER_FRONT_LEFT_OF_CENTER
+							FAudioChannelMapper::EChannelPosition::Rc,		// SPEAKER_FRONT_RIGHT_OF_CENTER
+							FAudioChannelMapper::EChannelPosition::Cs,		// SPEAKER_BACK_CENTER
+							FAudioChannelMapper::EChannelPosition::Lss,		// SPEAKER_SIDE_LEFT
+							FAudioChannelMapper::EChannelPosition::Rss,		// SPEAKER_SIDE_RIGHT
+							FAudioChannelMapper::EChannelPosition::Ts,		// SPEAKER_TOP_CENTER
+							FAudioChannelMapper::EChannelPosition::Lv,		// SPEAKER_TOP_FRONT_LEFT
+							FAudioChannelMapper::EChannelPosition::Cv,		// SPEAKER_TOP_FRONT_CENTER
+							FAudioChannelMapper::EChannelPosition::Rv,		// SPEAKER_TOP_FRONT_RIGHT
+							FAudioChannelMapper::EChannelPosition::Lvs,		// SPEAKER_TOP_BACK_LEFT
+							FAudioChannelMapper::EChannelPosition::Cvr,		// SPEAKER_TOP_BACK_CENTER
+							FAudioChannelMapper::EChannelPosition::Rvs		// SPEAKER_TOP_BACK_RIGHT
+						};
+						TArray<FAudioChannelMapper::FSourceLayout> Layout;
+						for(int32 i=0; i<UE_ARRAY_COUNT(ChannelPositions); ++i)
+						{
+							if ((ChannelMask & (1 << i)) != 0)
+							{
+								FAudioChannelMapper::FSourceLayout lo;
+								lo.ChannelPosition = ChannelPositions[i];
+								Layout.Emplace(MoveTemp(lo));
+							}
+						}
+						bMapperOk = ChannelMapper.Initialize(sizeof(int16), Layout);
+					}
+					if (!bMapperOk)
+					{
+						PostError(S_OK, "Failed to set up the output channel mapper", ERRCODE_INTERNAL_FAILED_TO_MAP_OUTPUT_CHANNELS, UEMEDIA_ERROR_INTERNAL);
+						return false;
+					}
+				}
+
 				// Get an sample block from the pool.
-				while (!TerminateThreadEvent.IsSignaled())
+				while(!TerminateThreadEvent.IsSignaled())
 				{
 					if (FlushDecoderEvent.IsSignaled() && !bFlush)
 					{
@@ -806,12 +900,12 @@ namespace Electra
 						if (dwBufferLen <= CurrentRenderOutputBufferSize)
 						{
 							void* CurrentRenderOutputBufferAddress = CurrentRenderOutputBuffer->GetBufferProperties().GetValue("address").GetPointer();
-							FMemory::Memcpy(CurrentRenderOutputBufferAddress, pPCMSamples, dwBufferLen);
+							ChannelMapper.MapChannels(CurrentRenderOutputBufferAddress, CurrentRenderOutputBufferSize, pPCMSamples, dwBufferLen, nSamplesProduced);
 
 							OutputBufferSampleProperties.Clear();
-							OutputBufferSampleProperties.Set("num_channels", FVariantValue((int64)OutputWaveFormat.nChannels));
+							OutputBufferSampleProperties.Set("num_channels", FVariantValue((int64)ChannelMapper.GetNumTargetChannels()));
 							OutputBufferSampleProperties.Set("sample_rate", FVariantValue((int64)OutputWaveFormat.nSamplesPerSec));
-							OutputBufferSampleProperties.Set("byte_size", FVariantValue((int64)dwBufferLen));
+							OutputBufferSampleProperties.Set("byte_size", FVariantValue((int64)(ChannelMapper.GetNumTargetChannels() * nSamplesProduced * sizeof(int16))));
 							OutputBufferSampleProperties.Set("duration", FVariantValue(OutputSampleDuration));
 							OutputBufferSampleProperties.Set("pts", FVariantValue(pts));
 							OutputBufferSampleProperties.Set("discontinuity", FVariantValue((bool)bHaveDiscontinuity));
@@ -917,7 +1011,7 @@ namespace Electra
 		bError = !CreateDecodedSamplePool();
 		check(!bError);
 
-		while (!TerminateThreadEvent.IsSignaled())
+		while(!TerminateThreadEvent.IsSignaled())
 		{
 			// If in background, wait until we get activated again.
 			if (!ApplicationRunningSignal.IsSignaled())
@@ -977,6 +1071,7 @@ namespace Electra
 					CurrentOutputMediaType = nullptr;
 					ConfigRecord = nullptr;
 					CurrentCodecData = nullptr;
+					ChannelMapper.Reset();
 				}
 
 				// Parse the CSD into a configuration record.
@@ -1029,7 +1124,7 @@ namespace Electra
 
 				// Loop until all data has been consumed.
 				bool bAllInputConsumed = false;
-				while (!bError && !bAllInputConsumed && !TerminateThreadEvent.IsSignaled() && !FlushDecoderEvent.IsSignaled())
+				while(!bError && !bAllInputConsumed && !TerminateThreadEvent.IsSignaled() && !FlushDecoderEvent.IsSignaled())
 				{
 					if (!pData->bIsDummyData)
 					{
@@ -1094,7 +1189,7 @@ namespace Electra
 							if (ConfigRecord.IsValid())
 							{
 								// Parameteric stereo results in stereo (2 channels).
-								NumChannels = ConfigRecord->PSSignal > 0 ? 2 : ConfigRecord->ChannelConfiguration;
+								NumChannels = ConfigRecord->PSSignal > 0 ? 2 : NumChannelsForConfig[ConfigRecord->ChannelConfiguration];
 								SampleRate = ConfigRecord->ExtSamplingFrequency ? ConfigRecord->ExtSamplingFrequency : ConfigRecord->SamplingRate;
 								SamplesPerBlock = ConfigRecord->SBRSignal > 0 ? 2048 : 1024;
 							}
diff --git a/Engine/Plugins/Media/ElectraPlayer/Source/ElectraPlayerRuntime/Private/Runtime/Decoder/DX/DecoderErrors_DX.h b/Engine/Plugins/Media/ElectraPlayer/Source/ElectraPlayerRuntime/Private/Runtime/Decoder/DX/DecoderErrors_DX.h
--- a/Engine/Plugins/Media/ElectraPlayer/Source/ElectraPlayerRuntime/Private/Runtime/Decoder/DX/DecoderErrors_DX.h
+++ b/Engine/Plugins/Media/ElectraPlayer/Source/ElectraPlayerRuntime/Private/Runtime/Decoder/DX/DecoderErrors_DX.h
@@ -12,7 +12,6 @@
 #define ERRCODE_INTERNAL_COULD_NOT_CREATE_SAMPLE_POOL			8
 #define ERRCODE_INTERNAL_FAILED_TO_PARSE_CSD					9
 
-
 #define ERRCODE_INTERNAL_COULD_GET_TRANSFORM_ATTRIBUTES					10
 #define ERRCODE_INTERNAL_UNSUPPORTED_D3D_TYPE							11
 #define ERRCODE_INTERNAL_FAILED_TO_SWITCH_TO_SOFTWARE_MODE				12
@@ -66,6 +65,7 @@
 #define ERRCODE_INTERNAL_COULD_NOT_SET_OUTPUT_TYPE						128
 #define ERRCODE_INTERNAL_COULD_NOT_SET_OUTPUT_TYPE_TO_PCM				128
 #define ERRCODE_INTERNAL_COULD_NOT_GET_INPUT_STREAM_FORMAT_INFO			129
+#define ERRCODE_INTERNAL_FAILED_TO_MAP_OUTPUT_CHANNELS					130
 #define ERRCODE_INTERNAL_COULD_NOT_GET_OUTPUT_STREAM_FORMAT_INFO		139
 #define ERRCODE_INTERNAL_COULD_NOT_MAKE_CONTIGUOUS_OUTPUT_BUFFER		140
 #define ERRCODE_INTERNAL_COULD_NOT_GET_OUTPUT_BUFFER_LENGTH				141
diff --git a/Engine/Plugins/Media/ElectraPlayer/Source/ElectraPlayerRuntime/Private/Runtime/Decoder/GStreamer/AudioDecoderAAC_GStreamer.cpp b/Engine/Plugins/Media/ElectraPlayer/Source/ElectraPlayerRuntime/Private/Runtime/Decoder/GStreamer/AudioDecoderAAC_GStreamer.cpp
--- a/Engine/Plugins/Media/ElectraPlayer/Source/ElectraPlayerRuntime/Private/Runtime/Decoder/GStreamer/AudioDecoderAAC_GStreamer.cpp
+++ b/Engine/Plugins/Media/ElectraPlayer/Source/ElectraPlayerRuntime/Private/Runtime/Decoder/GStreamer/AudioDecoderAAC_GStreamer.cpp
@@ -33,6 +33,11 @@ void IAudioDecoderAAC::Shutdown()
 {
 }
 
+bool IAudioDecoderAAC::CanDecodeStream(const FStreamCodecInformation& InCodecInfo)
+{
+	return false;
+}
+
 IAudioDecoderAAC* IAudioDecoderAAC::Create()
 {
 	return nullptr;
diff --git a/Engine/Plugins/Media/ElectraPlayer/Source/ElectraPlayerRuntime/Private/Runtime/Demuxer/ParserISO14496-12.cpp b/Engine/Plugins/Media/ElectraPlayer/Source/ElectraPlayerRuntime/Private/Runtime/Demuxer/ParserISO14496-12.cpp
--- a/Engine/Plugins/Media/ElectraPlayer/Source/ElectraPlayerRuntime/Private/Runtime/Demuxer/ParserISO14496-12.cpp
+++ b/Engine/Plugins/Media/ElectraPlayer/Source/ElectraPlayerRuntime/Private/Runtime/Demuxer/ParserISO14496-12.cpp
@@ -6290,6 +6290,7 @@ private:
 						{
 							Track->CodecInformation.SetStreamType(EStreamType::Video);
 							Track->CodecInformation.SetCodec(FStreamCodecInformation::ECodec::H264);
+							Track->CodecInformation.SetCodecSpecificData(Track->CodecSpecificDataRAW);
 							Track->CodecInformation.SetStreamLanguageCode(Track->GetLanguage());
 							if (Track->CodecSpecificDataAVC.GetNumberOfSPS() == 0)
 							{
@@ -6367,6 +6368,7 @@ private:
 						{
 							Track->CodecInformation.SetStreamType(EStreamType::Video);
 							Track->CodecInformation.SetCodec(FStreamCodecInformation::ECodec::H265);
+							Track->CodecInformation.SetCodecSpecificData(Track->CodecSpecificDataRAW);
 							Track->CodecInformation.SetStreamLanguageCode(Track->GetLanguage());
 							if (Track->CodecSpecificDataHEVC.GetNumberOfSPS() == 0)
 							{
@@ -6450,6 +6452,7 @@ private:
 							{
 								Track->CodecInformation.SetStreamType(EStreamType::Audio);
 								Track->CodecInformation.SetCodec(FStreamCodecInformation::ECodec::AAC);
+								Track->CodecInformation.SetCodecSpecificData(Track->CodecSpecificDataRAW);
 								Track->CodecInformation.SetStreamLanguageCode(Track->GetLanguage());
 
 								MPEG::FAACDecoderConfigurationRecord ConfigRecord;
@@ -6457,7 +6460,8 @@ private:
 								{
 									Track->CodecInformation.SetCodecSpecifierRFC6381(FString::Printf(TEXT("mp4a.40.%d"), ConfigRecord.ExtAOT ? ConfigRecord.ExtAOT : ConfigRecord.AOT));
 									Track->CodecInformation.SetSamplingRate(ConfigRecord.ExtSamplingFrequency ? ConfigRecord.ExtSamplingFrequency : ConfigRecord.SamplingRate);
-									Track->CodecInformation.SetNumberOfChannels((int32)ConfigRecord.ChannelConfiguration);
+									Track->CodecInformation.SetChannelConfiguration(ConfigRecord.ChannelConfiguration);
+									Track->CodecInformation.SetNumberOfChannels(MPEG::AACUtils::GetNumberOfChannelsFromChannelConfiguration(ConfigRecord.ChannelConfiguration));
 									// We assume that all platforms can decode PS (parametric stereo). As such we change the channel count from mono to stereo
 									// to convey the _decoded_ format, not the source format.
 									if (ConfigRecord.ChannelConfiguration == 1 && ConfigRecord.PSSignal > 0)
@@ -6521,6 +6525,7 @@ private:
 			Track->CodecSpecificDataRAW = SubtitleSampleEntry->GetRawBoxData();
 			Track->CodecInformation.SetStreamType(EStreamType::Subtitle);
 			Track->CodecInformation.SetCodec(FStreamCodecInformation::ECodec::TX3G);
+			Track->CodecInformation.SetCodecSpecificData(Track->CodecSpecificDataRAW);
 			Track->CodecInformation.SetStreamLanguageCode(Track->GetLanguage());
 			Track->CodecInformation.SetCodecSpecifierRFC6381(TEXT("tx3g"));
 			Track->CodecInformation.SetResolution(FStreamCodecInformation::FResolution(SubtitleSampleEntry->GetWidth(), SubtitleSampleEntry->GetHeight()));
@@ -6541,6 +6546,7 @@ private:
 			Track->CodecSpecificDataRAW = SubtitleSampleEntry->GetRawBoxData();
 			Track->CodecInformation.SetStreamType(EStreamType::Subtitle);
 			Track->CodecInformation.SetCodec(FStreamCodecInformation::ECodec::WebVTT);
+			Track->CodecInformation.SetCodecSpecificData(Track->CodecSpecificDataRAW);
 			Track->CodecInformation.SetStreamLanguageCode(Track->GetLanguage());
 			Track->CodecInformation.SetCodecSpecifierRFC6381(TEXT("wvtt"));
 			Track->CodecInformation.SetResolution(FStreamCodecInformation::FResolution(SubtitleSampleEntry->GetWidth(), SubtitleSampleEntry->GetHeight()));
diff --git a/Engine/Plugins/Media/ElectraPlayer/Source/ElectraPlayerRuntime/Private/Runtime/Player/AdaptiveStreamingPlayerDecoders.cpp b/Engine/Plugins/Media/ElectraPlayer/Source/ElectraPlayerRuntime/Private/Runtime/Player/AdaptiveStreamingPlayerDecoders.cpp
--- a/Engine/Plugins/Media/ElectraPlayer/Source/ElectraPlayerRuntime/Private/Runtime/Player/AdaptiveStreamingPlayerDecoders.cpp
+++ b/Engine/Plugins/Media/ElectraPlayer/Source/ElectraPlayerRuntime/Private/Runtime/Player/AdaptiveStreamingPlayerDecoders.cpp
@@ -141,6 +141,10 @@ bool FAdaptiveStreamingPlayer::CanDecodeStream(const FStreamCodecInformation& In
 		{
 			return false;
 		}
+		if (InStreamCodecInfo.GetCodec() == FStreamCodecInformation::ECodec::AAC)
+		{
+			return IAudioDecoderAAC::CanDecodeStream(InStreamCodecInfo);
+		}
 	}
 	else if (InStreamCodecInfo.IsSubtitleCodec())
 	{
diff --git a/Engine/Plugins/Media/ElectraPlayer/Source/ElectraPlayerRuntime/Private/Runtime/Player/DASH/ManifestBuilderDASH.cpp b/Engine/Plugins/Media/ElectraPlayer/Source/ElectraPlayerRuntime/Private/Runtime/Player/DASH/ManifestBuilderDASH.cpp
--- a/Engine/Plugins/Media/ElectraPlayer/Source/ElectraPlayerRuntime/Private/Runtime/Player/DASH/ManifestBuilderDASH.cpp
+++ b/Engine/Plugins/Media/ElectraPlayer/Source/ElectraPlayerRuntime/Private/Runtime/Player/DASH/ManifestBuilderDASH.cpp
@@ -25,6 +25,7 @@
 #include "Utilities/Utilities.h"
 #include "Utilities/TimeUtilities.h"
 #include "Utilities/ISO639-Map.h"
+#include "Utilities/UtilsMPEGAudio.h"
 
 #include "Player/DRM/DRMManager.h"
 
@@ -1544,20 +1545,14 @@ void FManifestDASHInternal::PreparePeriodAdaptationSets(TSharedPtrTS<FPeriod> Pe
 					AudioChannelConfigurations.Append(MPDAdaptationSet->GetEssentialProperties());
 					for(int32 nACC=0, nACCMax=AudioChannelConfigurations.Num(); nACC<nACCMax; ++nACC)
 					{
-						if (AudioChannelConfigurations[nACC]->GetSchemeIdUri().Equals(AudioChannelConfigurationLegacy))		// "urn:mpeg:dash:23003:3:audio_channel_configuration:2011"
+						if (AudioChannelConfigurations[nACC]->GetSchemeIdUri().Equals(AudioChannelConfigurationLegacy) ||	// "urn:mpeg:dash:23003:3:audio_channel_configuration:2011"
+							AudioChannelConfigurations[nACC]->GetSchemeIdUri().Equals(AudioChannelConfiguration))			// "urn:mpeg:mpegB:cicp:ChannelConfiguration"
 						{
 							// Value = channel config as per 23001-8:2013 table 8
-							int32 v = 0;
+							uint32 v = 0;
 							LexFromString(v, *AudioChannelConfigurations[nACC]->GetValue());
-							Representation->CodecInfo.SetNumberOfChannels(v);
-							break;
-						}
-						else if (AudioChannelConfigurations[nACC]->GetSchemeIdUri().Equals(AudioChannelConfiguration))		// "urn:mpeg:mpegB:cicp:ChannelConfiguration"
-						{
-							// Value = channel config as per 23001-8:2013 table 8
-							int32 v = 0;
-							LexFromString(v, *AudioChannelConfigurations[nACC]->GetValue());
-							Representation->CodecInfo.SetNumberOfChannels(v);
+							Representation->CodecInfo.SetChannelConfiguration(v);
+							Representation->CodecInfo.SetNumberOfChannels(MPEG::AACUtils::GetNumberOfChannelsFromChannelConfiguration(v));
 							break;
 						}
 						else if (AudioChannelConfigurations[nACC]->GetSchemeIdUri().Equals(AudioChannelConfigurationDolby))	// "tag:dolby.com,2014:dash:audio_channel_configuration:2011"
diff --git a/Engine/Plugins/Media/ElectraPlayer/Source/ElectraPlayerRuntime/Private/Runtime/Player/HLS/ManifestBuilderHLS.cpp b/Engine/Plugins/Media/ElectraPlayer/Source/ElectraPlayerRuntime/Private/Runtime/Player/HLS/ManifestBuilderHLS.cpp
--- a/Engine/Plugins/Media/ElectraPlayer/Source/ElectraPlayerRuntime/Private/Runtime/Player/HLS/ManifestBuilderHLS.cpp
+++ b/Engine/Plugins/Media/ElectraPlayer/Source/ElectraPlayerRuntime/Private/Runtime/Player/HLS/ManifestBuilderHLS.cpp
@@ -579,6 +579,9 @@ FErrorDetail FManifestBuilderHLS::SetupVariants(FManifestHLSInternal* Manifest,
 					si.SetStreamType(EStreamType::Audio);
 					si.SetCodec(FStreamCodecInformation::ECodec::AAC);
 					si.SetCodecSpecifierRFC6381(CodecList[nCodec]);
+					// For lack of knowledge pretend this is stereo.
+					si.SetChannelConfiguration(2);
+					si.SetNumberOfChannels(2);
 					bHasAudio = true;
 				}
 				else
@@ -649,6 +652,9 @@ FErrorDetail FManifestBuilderHLS::SetupVariants(FManifestHLSInternal* Manifest,
 				// For codec we assume standard LC-AAC
 				si.SetCodec(FStreamCodecInformation::ECodec::AAC);
 				si.SetCodecSpecifierRFC6381(TEXT("mp4a.40.2"));
+				// For lack of knowledge pretend this is stereo.
+				si.SetChannelConfiguration(2);
+				si.SetNumberOfChannels(2);
 				vs->StreamCodecInformationList.Push(si);
 				bHasAudio = true;
 			}
diff --git a/Engine/Plugins/Media/ElectraPlayer/Source/ElectraPlayerRuntime/Private/Runtime/StreamTypes.h b/Engine/Plugins/Media/ElectraPlayer/Source/ElectraPlayerRuntime/Private/Runtime/StreamTypes.h
--- a/Engine/Plugins/Media/ElectraPlayer/Source/ElectraPlayerRuntime/Private/Runtime/StreamTypes.h
+++ b/Engine/Plugins/Media/ElectraPlayer/Source/ElectraPlayerRuntime/Private/Runtime/StreamTypes.h
@@ -449,6 +449,16 @@ namespace Electra
 			return StreamLanguageCode;
 		}
 
+		void SetCodecSpecificData(const TArray<uint8>& InCSD)
+		{
+			CSD = InCSD;
+		}
+
+		const TArray<uint8>& GetCodecSpecificData() const
+		{
+			return CSD;
+		}
+
 		FParamDict& GetExtras()
 		{
 			return Extras;
@@ -479,6 +489,7 @@ namespace Electra
 			AudioAccessibility = 0;
 			NumberOfAudioObjects = 0;
 			Extras.Clear();
+			CSD.Empty();
 		}
 
 		bool IsDifferentFromOtherVideo(const FStreamCodecInformation& Other) const
@@ -568,6 +579,7 @@ namespace Electra
 		int32			AudioAccessibility;			//!< Format specific audio accessibility
 		int32			NumberOfAudioObjects;		//!< Format specific number of audio objects
 		FParamDict		Extras;						//!< Additional details/properties, depending on playlist.
+		TArray<uint8>	CSD;						//!< Codec specific data, if available.
 	};
 
 
diff --git a/Engine/Plugins/Media/ElectraPlayer/Source/ElectraPlayerRuntime/Private/Runtime/Utilities/AudioChannelMapper.cpp b/Engine/Plugins/Media/ElectraPlayer/Source/ElectraPlayerRuntime/Private/Runtime/Utilities/AudioChannelMapper.cpp
new file mode 100644
--- /dev/null
+++ b/Engine/Plugins/Media/ElectraPlayer/Source/ElectraPlayerRuntime/Private/Runtime/Utilities/AudioChannelMapper.cpp
@@ -0,0 +1,344 @@
+// Copyright Epic Games, Inc. All Rights Reserved.
+
+#include "AudioChannelMapper.h"
+#include "Core/MediaMacros.h"
+
+namespace Electra
+{
+
+	namespace DirectlyCopyable
+	{
+		// This defines the order in which FMediaAudioResampler expects the source channels.
+		// See comments in MediaAudioResampler.cpp regarding the row order.
+		#define CP FAudioChannelMapper::EChannelPosition
+		static const CP _1[] = { CP::C };
+		static const CP _2[] = { CP::L, CP::R };
+		static const CP _3[] = { CP::L, CP::R, CP::C };
+		static const CP _4[] = { CP::L, CP::R, CP::Ls, CP::Rs };
+		static const CP _5[] = { CP::L, CP::R, CP::C, CP::Ls, CP::Rs };
+		static const CP _6[] = { CP::L, CP::R, CP::C, CP::LFE, CP::Ls, CP::Rs };
+		static const CP _7[] = { CP::L, CP::R, CP::Lsr, CP::LFE, CP::Rsr, CP::Ls, CP::Rs };
+		static const CP _8[] = { CP::L, CP::R, CP::C, CP::LFE, CP::Ls, CP::Rs, CP::Lsr, CP::Rsr };
+		static const CP * const ResamplerOrderMap[] = { _1, _2, _3, _4, _5, _6, _7, _8 };
+		#undef CP
+	}
+
+	FAudioChannelMapper::FAudioChannelMapper()
+	{
+	}
+
+	FAudioChannelMapper::~FAudioChannelMapper()
+	{
+	}
+
+
+	bool FAudioChannelMapper::Initialize(int32 NumBytesPerSample, uint32 ChannelConfiguration)
+	{
+		Reset();
+		#define CP(ch) { EChannelPosition::ch }
+		switch(ChannelConfiguration)
+		{
+			default:
+			{
+				break;
+			}
+			case 1:
+			{
+				FSourceLayout Layout[] = { CP(C) };
+				return Initialize(NumBytesPerSample, Layout);
+			}
+			case 2:
+			case 8:	// handle unspecified channel1+channel2 configuration as L+R
+			{
+				FSourceLayout Layout[] = { CP(L), CP(R) };
+				return Initialize(NumBytesPerSample, Layout);
+			}
+			case 3:
+			{
+				FSourceLayout Layout[] = { CP(C), CP(L), CP(R) };
+				return Initialize(NumBytesPerSample, Layout);
+			}
+			case 4:
+			{
+				FSourceLayout Layout[] = { CP(C), CP(L), CP(R), CP(Cs) };
+				return Initialize(NumBytesPerSample, Layout);
+			}
+			case 5:
+			{
+				FSourceLayout Layout[] = { CP(C), CP(L), CP(R), CP(Ls), CP(Rs) };
+				return Initialize(NumBytesPerSample, Layout);
+			}
+			case 6:
+			{
+				FSourceLayout Layout[] = { CP(C), CP(L), CP(R), CP(Ls), CP(Rs), CP(LFE) };
+				return Initialize(NumBytesPerSample, Layout);
+			}
+			case 7:
+			{
+				FSourceLayout Layout[] = { CP(C), CP(Lc), CP(Rc), CP(L), CP(R), CP(Ls), CP(Rs), CP(LFE) };
+				return Initialize(NumBytesPerSample, Layout);
+			}
+			case 9:
+			{
+				FSourceLayout Layout[] = { CP(L), CP(R), CP(Cs) };
+				return Initialize(NumBytesPerSample, Layout);
+			}
+			case 10:
+			{
+				FSourceLayout Layout[] = { CP(L), CP(R), CP(Ls), CP(Rs) };
+				return Initialize(NumBytesPerSample, Layout);
+			}
+			case 11:
+			{
+				FSourceLayout Layout[] = { CP(C), CP(L), CP(R), CP(Ls), CP(Rs), CP(Cs), CP(LFE) };
+				return Initialize(NumBytesPerSample, Layout);
+			}
+			case 12:
+			{
+				FSourceLayout Layout[] = { CP(C), CP(L), CP(R), CP(Ls), CP(Rs), CP(Lsr), CP(Rsr), CP(LFE) };
+				return Initialize(NumBytesPerSample, Layout);
+			}
+			case 13:
+			{
+				FSourceLayout Layout[] = { CP(C), CP(Lc), CP(Rc), CP(L), CP(R), CP(Lss), CP(Rss), CP(Lsr), CP(Rsr), CP(Cs), CP(LFE), CP(LFE2),
+										   CP(Cv), CP(Lv), CP(Rv), CP(Lvss), CP(Rvss), CP(Ts), CP(Lvr), CP(Rvr), CP(Cvr), CP(Cb), CP(Lb), CP(Rb) };
+				return Initialize(NumBytesPerSample, Layout);
+			}
+			case 14:
+			{
+				FSourceLayout Layout[] = { CP(C), CP(L), CP(R), CP(Ls), CP(Rs), CP(LFE), CP(Lv), CP(Rv) };
+				return Initialize(NumBytesPerSample, Layout);
+			}
+		}
+		#undef CP
+		return false;
+	}
+
+	bool FAudioChannelMapper::Initialize(int32 NumBytesPerSample, TArrayView<const FSourceLayout> ChannelPositions)
+	{
+		check(NumBytesPerSample == 2 || NumBytesPerSample == 4);
+		Reset();
+
+		bCanCopyDirectly = CanBeCopiedDirectly(ChannelPositions);
+		if (bCanCopyDirectly)
+		{
+			// Even though we can copy the source over directly we still set up the target source layout
+			// for later use. GetNumTargetChannels() requires this.
+			for(int32 i=0; i<ChannelPositions.Num(); ++i)
+			{
+				FTargetSource& ts = TargetSources.AddDefaulted_GetRef();
+				ts.ChannelPosition = ChannelPositions[i].ChannelPosition;
+				ts.FirstOffset = NumBytesPerSample * i;
+				ts.Stride = NumBytesPerSample * ChannelPositions.Num();
+			}
+		}
+		else
+		{
+			int32 Stride = NumBytesPerSample * ChannelPositions.Num();
+			// Setup the mapping from the input channels as a mutable array.
+			TArray<FTargetSource> InputSources;
+			for(int32 i=0; i<ChannelPositions.Num(); ++i)
+			{
+				FTargetSource& ts = InputSources.AddDefaulted_GetRef();
+				ts.ChannelPosition = ChannelPositions[i].ChannelPosition;
+				ts.FirstOffset = NumBytesPerSample * i;
+				ts.Stride = Stride;
+			}
+			// Now remove all the inputs we do not support.
+			for(int32 i=0; i<InputSources.Num(); ++i)
+			{
+				if (InputSources[i].ChannelPosition != EChannelPosition::C &&
+					InputSources[i].ChannelPosition != EChannelPosition::L &&
+					InputSources[i].ChannelPosition != EChannelPosition::R &&
+					InputSources[i].ChannelPosition != EChannelPosition::LFE &&
+					InputSources[i].ChannelPosition != EChannelPosition::Ls &&
+					InputSources[i].ChannelPosition != EChannelPosition::Rs &&
+					InputSources[i].ChannelPosition != EChannelPosition::Lsr &&
+					InputSources[i].ChannelPosition != EChannelPosition::Rsr)
+				{
+					InputSources.RemoveAt(i);
+					--i;
+				}
+			}
+
+			// Do the current sources match one of the channel layouts used by the resampler?
+			int32 ResamplerChannelLayout = MatchesResamplerLayout(InputSources);
+			// At the very least the last 8 channel layout has to match.
+			check(ResamplerChannelLayout);
+			if (ResamplerChannelLayout == 0)
+			{
+				return false;
+			}
+
+			// Create a channel mapping matching the channel layout of the resampler.
+			// Channels that the resampler uses but are not present are filled with empty
+			// mappings that create silence for that channel.
+			CreateResamplerChannelMapping(ResamplerChannelLayout-1, InputSources);
+		}
+
+		BytesPerSample = NumBytesPerSample;
+		return true;
+	}
+
+	bool FAudioChannelMapper::IsInitialized() const
+	{
+		return BytesPerSample != 0;
+	}
+
+
+	void FAudioChannelMapper::Reset()
+	{
+		TargetSources.Empty();
+		BytesPerSample = 0;
+		bCanCopyDirectly = false;
+	}
+
+	int32 FAudioChannelMapper::GetNumTargetChannels() const
+	{
+		return TargetSources.Num();
+	}
+
+	void FAudioChannelMapper::MapChannels(void* OutputBuffer, uint32 OutputBufferSize, const void* InputBuffer, uint32 InputBufferSize, int32 NumSamplesPerChannel) const
+	{
+		// Must have been initialized!
+		check(IsInitialized());
+		if (IsInitialized())
+		{
+			if (bCanCopyDirectly)
+			{
+				uint32 nc = InputBufferSize < OutputBufferSize ? InputBufferSize : OutputBufferSize;
+				FMemory::Memcpy(OutputBuffer, InputBuffer, nc);
+			}
+			else
+			{
+				int32 OutStride = BytesPerSample * TargetSources.Num();
+				// Copy channels in output order
+				for(int32 Ch=0,ChMax=TargetSources.Num(); Ch<ChMax; ++Ch)
+				{
+					void* OutCh = AdvancePointer(OutputBuffer, Ch * BytesPerSample);
+					const void* InCh = AdvancePointer(InputBuffer, TargetSources[Ch].FirstOffset);
+					int32 InStride = TargetSources[Ch].Stride;
+					if (BytesPerSample == 2)
+					{
+						if (TargetSources[Ch].ChannelPosition == EChannelPosition::Disabled)
+						{
+							for(int32 i=0; i<NumSamplesPerChannel; ++i)
+							{
+								*reinterpret_cast<uint16*>(OutCh) = 0;
+								OutCh = AdvancePointer(OutCh, OutStride);
+							}
+						}
+						else
+						{
+							for(int32 i=0; i<NumSamplesPerChannel; ++i)
+							{
+								*reinterpret_cast<uint16*>(OutCh) = *reinterpret_cast<const uint16*>(InCh);
+								OutCh = AdvancePointer(OutCh, OutStride);
+								InCh = AdvancePointer(InCh, InStride);
+							}
+						}
+					}
+					else if (BytesPerSample == 4)
+					{
+						if (TargetSources[Ch].ChannelPosition == EChannelPosition::Disabled)
+						{
+							for(int32 i=0; i<NumSamplesPerChannel; ++i)
+							{
+								*reinterpret_cast<uint32*>(OutCh) = 0;
+								OutCh = AdvancePointer(OutCh, OutStride);
+							}
+						}
+						else
+						{
+							for(int32 i=0; i<NumSamplesPerChannel; ++i)
+							{
+								*reinterpret_cast<uint32*>(OutCh) = *reinterpret_cast<const uint32*>(InCh);
+								OutCh = AdvancePointer(OutCh, OutStride);
+								InCh = AdvancePointer(InCh, InStride);
+							}
+						}
+					}
+				}
+			}
+		}
+		else
+		{
+			FMemory::Memzero(OutputBuffer, OutputBufferSize);
+		}
+	}
+
+	bool FAudioChannelMapper::CanBeCopiedDirectly(TArrayView<const FSourceLayout> ChannelPositions)
+	{
+		if (ChannelPositions.Num() >= 1 && ChannelPositions.Num() <= 8)
+		{
+			const FAudioChannelMapper::EChannelPosition * const ResamplerOrder = DirectlyCopyable::ResamplerOrderMap[ChannelPositions.Num() - 1];
+			if (ResamplerOrder)
+			{
+				for(int32 i=0; i<ChannelPositions.Num(); ++i)
+				{
+					if (ResamplerOrder[i] != ChannelPositions[i].ChannelPosition)
+					{
+						return false;
+					}
+				}
+				return true;
+			}
+		}
+		return false;
+	}
+
+	int32 FAudioChannelMapper::MatchesResamplerLayout(TArray<FTargetSource>& InputSources)
+	{
+		if (InputSources.Num())
+		{
+			for(int32 i=InputSources.Num(); i<=UE_ARRAY_COUNT(DirectlyCopyable::ResamplerOrderMap); ++i)
+			{
+				TArrayView<const FAudioChannelMapper::EChannelPosition> ResamplerChannels(DirectlyCopyable::ResamplerOrderMap[i-1], i);
+				bool bAllChannelsPresent = true;
+				for(int32 j=0; j<InputSources.Num(); ++j)
+				{
+					if (!ResamplerChannels.Contains(InputSources[j].ChannelPosition))
+					{
+						bAllChannelsPresent = false;
+						break;
+					}
+				}
+				if (bAllChannelsPresent)
+				{
+					return i;
+				}
+			}
+		}
+		return 0;
+	}
+
+	void FAudioChannelMapper::CreateResamplerChannelMapping(int32 ResamplerChannelIndex, const TArray<FTargetSource>& InSources)
+	{
+		check(ResamplerChannelIndex >= 0 && ResamplerChannelIndex < UE_ARRAY_COUNT(DirectlyCopyable::ResamplerOrderMap));
+		if (ResamplerChannelIndex >= 0 && ResamplerChannelIndex < UE_ARRAY_COUNT(DirectlyCopyable::ResamplerOrderMap))
+		{
+			TArrayView<const FAudioChannelMapper::EChannelPosition> ResamplerChannels(DirectlyCopyable::ResamplerOrderMap[ResamplerChannelIndex], ResamplerChannelIndex+1);
+			for(auto &ResamplerChannel : ResamplerChannels)
+			{
+				FTargetSource ts;
+				int32 SourceIndex = InSources.IndexOfByPredicate([ResamplerChannel](const FTargetSource& In){return In.ChannelPosition == ResamplerChannel;});
+				if (SourceIndex != INDEX_NONE)
+				{
+					TargetSources.Emplace(InSources[SourceIndex]);
+				}
+				else
+				{
+					TargetSources.AddDefaulted();
+				}
+			}
+		}
+		else
+		{
+			// Add at least one empty entry to keep everything happy.
+			TargetSources.AddDefaulted();
+		}
+	}
+
+
+} // namespace Electra
+
diff --git a/Engine/Plugins/Media/ElectraPlayer/Source/ElectraPlayerRuntime/Private/Runtime/Utilities/AudioChannelMapper.h b/Engine/Plugins/Media/ElectraPlayer/Source/ElectraPlayerRuntime/Private/Runtime/Utilities/AudioChannelMapper.h
new file mode 100644
--- /dev/null
+++ b/Engine/Plugins/Media/ElectraPlayer/Source/ElectraPlayerRuntime/Private/Runtime/Utilities/AudioChannelMapper.h
@@ -0,0 +1,71 @@
+// Copyright Epic Games, Inc. All Rights Reserved.
+#pragma once
+
+#include "CoreMinimal.h"
+#include "Containers/ArrayView.h"
+
+namespace Electra
+{
+	/**
+	 * Class to map audio channels from some decoder-specific output order into the
+	 * order expected by FMediaAudioResampler.
+	 * 
+	 * Channels that are not supported by UE will be ignored and are silent.
+	 * No effort to mix them into other channels is made.
+	 * As fas as media playback is concerned we handle the most common channel layouts
+	 * only, like mono, stereo, 5.1 and 7.1.
+	 */
+	class FAudioChannelMapper
+	{
+	public:
+		// Channel position as per ISO/IEC 23001-8 Table 7.
+		enum class EChannelPosition
+		{
+			L, R, C, LFE, Ls, Rs, Lc, Rc,
+			Lsr, Rsr, Cs, Lsd, Rsd, Lss, Rss, Lw,
+			Rw, Lv, Rv, Cv, Lvr, Rvr, Cvr, Lvss,
+			Rvss, Ts, LFE2, Lb, Rb, Cb, Lvs, Rvs,
+			Disabled = 255
+		};
+
+		struct FSourceLayout
+		{
+			EChannelPosition ChannelPosition = EChannelPosition::Disabled;
+		};
+
+		FAudioChannelMapper();
+		~FAudioChannelMapper();
+
+		// Initializes with an ISO/IEC 23001-8 channel configuration value from Table 7 (compatible with ISO/IEC 14496-3:2009/Amd 4)
+		bool Initialize(int32 NumBytesPerSample, uint32 ChannelConfiguration);
+
+		// Initializes with ISO/IEC 23001-8 channel positions.
+		bool Initialize(int32 NumBytesPerSample, TArrayView<const FSourceLayout> ChannelPositions);
+
+		bool IsInitialized() const;
+
+		void Reset();
+
+		// Gets the number of output channels mapped to the UE internal layout.
+		int32 GetNumTargetChannels() const;
+
+		// Maps the supported input channels to the output channels in UE internal layout.
+		void MapChannels(void* OutputBuffer, uint32 OutputBufferSize, const void* InputBuffer, uint32 InputBufferSize, int32 NumSamplesPerChannel) const;
+
+	private:
+		struct FTargetSource
+		{
+			EChannelPosition ChannelPosition = EChannelPosition::Disabled;
+			int32 FirstOffset = 0;
+			int32 Stride = 0;
+		};
+
+		static bool CanBeCopiedDirectly(TArrayView<const FSourceLayout> ChannelPositions);
+		int32 MatchesResamplerLayout(TArray<FTargetSource>& InputSources);
+		void CreateResamplerChannelMapping(int32 ResamplerChannelIndex, const TArray<FTargetSource>& InSources);
+
+		TArray<FTargetSource> TargetSources;
+		int32 BytesPerSample = 0;
+		bool bCanCopyDirectly = false;
+	};
+} // namespace Electra
diff --git a/Engine/Plugins/Media/ElectraPlayer/Source/ElectraPlayerRuntime/Private/Runtime/Utilities/UtilsMPEGAudio.cpp b/Engine/Plugins/Media/ElectraPlayer/Source/ElectraPlayerRuntime/Private/Runtime/Utilities/UtilsMPEGAudio.cpp
--- a/Engine/Plugins/Media/ElectraPlayer/Source/ElectraPlayerRuntime/Private/Runtime/Utilities/UtilsMPEGAudio.cpp
+++ b/Engine/Plugins/Media/ElectraPlayer/Source/ElectraPlayerRuntime/Private/Runtime/Utilities/UtilsMPEGAudio.cpp
@@ -7,6 +7,15 @@ namespace Electra
 {
 	namespace MPEG
 	{
+		namespace AACUtils
+		{
+			int32 GetNumberOfChannelsFromChannelConfiguration(uint32 InChannelConfiguration)
+			{
+				static const uint8 NumChannelsForConfig[16] = { 0, 1, 2, 3, 4, 5, 6, 7, 0, 0, 0, 7, 8, 0, 8, 0 };
+				return InChannelConfiguration < 16 ? NumChannelsForConfig[InChannelConfiguration] : 0;
+			}
+		}
+
 
 		namespace AACParseHelper
 		{
diff --git a/Engine/Plugins/Media/ElectraPlayer/Source/ElectraPlayerRuntime/Private/Runtime/Utilities/UtilsMPEGAudio.h b/Engine/Plugins/Media/ElectraPlayer/Source/ElectraPlayerRuntime/Private/Runtime/Utilities/UtilsMPEGAudio.h
--- a/Engine/Plugins/Media/ElectraPlayer/Source/ElectraPlayerRuntime/Private/Runtime/Utilities/UtilsMPEGAudio.h
+++ b/Engine/Plugins/Media/ElectraPlayer/Source/ElectraPlayerRuntime/Private/Runtime/Utilities/UtilsMPEGAudio.h
@@ -43,6 +43,10 @@ namespace Electra
 		};
 
 
+		namespace AACUtils
+		{
+			int32 GetNumberOfChannelsFromChannelConfiguration(uint32 InChannelConfiguration);
+		}
 
 
 	} // namespace MPEG
