ElectraPlayer: Filtering aac audio streams from an mp4 that the platform cannot decode

#rnx

diff --git a/Engine/Plugins/Media/ElectraPlayer/Source/ElectraPlayerRuntime/Private/Runtime/Demuxer/ParserISO14496-12.cpp b/Engine/Plugins/Media/ElectraPlayer/Source/ElectraPlayerRuntime/Private/Runtime/Demuxer/ParserISO14496-12.cpp
--- a/Engine/Plugins/Media/ElectraPlayer/Source/ElectraPlayerRuntime/Private/Runtime/Demuxer/ParserISO14496-12.cpp
+++ b/Engine/Plugins/Media/ElectraPlayer/Source/ElectraPlayerRuntime/Private/Runtime/Demuxer/ParserISO14496-12.cpp
@@ -12,6 +12,7 @@
 #include "Utilities/ISO639-Map.h"
 #include "InfoLog.h"
 #include "Player/PlayerSessionServices.h"
+#include "Player/PlayerStreamFilter.h"
 #include "Decoder/SubtitleDecoder.h"
 
 #if UE_BUILD_DEBUG
@@ -5126,7 +5127,7 @@ private:
 		virtual ~FParserISO14496_12();
 		virtual UEMediaError ParseHeader(IReader* DataReader, IBoxCallback* BoxParseCallback, IPlayerSessionServices* PlayerSession, const IParserISO14496_12* OptionalInitSegment) override;
 
-		virtual UEMediaError PrepareTracks(TSharedPtrTS<const IParserISO14496_12>	OptionalMP4InitSegment) override;
+		virtual UEMediaError PrepareTracks(IPlayerSessionServices* PlayerSession, TSharedPtrTS<const IParserISO14496_12> OptionalMP4InitSegment) override;
 
 		virtual TMediaOptionalValue<FTimeFraction> GetMovieDuration() const override;
 		virtual int32 GetNumberOfTracks() const override;
@@ -5388,11 +5389,11 @@ private:
 		};
 
 
-		UEMediaError ParseAVC1SampleType(FTrack* Track, const FMP4Box* SampleBox);
-		UEMediaError ParseHVC1SampleType(FTrack* Track, const FMP4Box* SampleBox);
-		UEMediaError ParseMP4ASampleType(FTrack* Track, const FMP4Box* SampleBox);
-		UEMediaError ParseTX3GSampleType(FTrack* Track, const FMP4Box* SampleBox);
-		UEMediaError ParseWVTTSampleType(FTrack* Track, const FMP4Box* SampleBox);
+		UEMediaError ParseAVC1SampleType(IPlayerSessionServices* PlayerSession, FTrack* Track, const FMP4Box* SampleBox);
+		UEMediaError ParseHVC1SampleType(IPlayerSessionServices* PlayerSession, FTrack* Track, const FMP4Box* SampleBox);
+		UEMediaError ParseMP4ASampleType(IPlayerSessionServices* PlayerSession, FTrack* Track, const FMP4Box* SampleBox);
+		UEMediaError ParseTX3GSampleType(IPlayerSessionServices* PlayerSession, FTrack* Track, const FMP4Box* SampleBox);
+		UEMediaError ParseWVTTSampleType(IPlayerSessionServices* PlayerSession, FTrack* Track, const FMP4Box* SampleBox);
 
 		FMP4ParseInfo* ParsedData;
 
@@ -6267,7 +6268,7 @@ private:
 	}
 
 
-	UEMediaError FParserISO14496_12::ParseAVC1SampleType(FTrack* Track, const FMP4Box* SampleBox)
+	UEMediaError FParserISO14496_12::ParseAVC1SampleType(IPlayerSessionServices* PlayerSession, FTrack* Track, const FMP4Box* SampleBox)
 	{
 		const FMP4BoxVisualSampleEntry* VisualSampleEntry = static_cast<const FMP4BoxVisualSampleEntry*>(SampleBox);
 		check(VisualSampleEntry->GetNumberOfChildren() > 0);
@@ -6345,7 +6346,7 @@ private:
 		return UEMEDIA_ERROR_FORMAT_ERROR;
 	}
 
-	UEMediaError FParserISO14496_12::ParseHVC1SampleType(FTrack* Track, const FMP4Box* SampleBox)
+	UEMediaError FParserISO14496_12::ParseHVC1SampleType(IPlayerSessionServices* PlayerSession, FTrack* Track, const FMP4Box* SampleBox)
 	{
 		const FMP4BoxVisualSampleEntry* VisualSampleEntry = static_cast<const FMP4BoxVisualSampleEntry*>(SampleBox);
 		check(VisualSampleEntry->GetNumberOfChildren() > 0);
@@ -6425,7 +6426,7 @@ private:
 		return UEMEDIA_ERROR_FORMAT_ERROR;
 	}
 
-	UEMediaError FParserISO14496_12::ParseMP4ASampleType(FTrack* Track, const FMP4Box* SampleBox)
+	UEMediaError FParserISO14496_12::ParseMP4ASampleType(IPlayerSessionServices* PlayerSession, FTrack* Track, const FMP4Box* SampleBox)
 	{
 		const FMP4BoxAudioSampleEntry* AudioSampleEntry = static_cast<const FMP4BoxAudioSampleEntry*>(Track->STSDBox->GetChildBox(0));
 		check(AudioSampleEntry->GetNumberOfChildren() > 0);
@@ -6477,6 +6478,16 @@ private:
 								Track->BitrateInfo.AvgBitrate = Track->CodecSpecificDataMP4A.GetAvgBitrate();
 
 								bGotAudioFormat = true;
+
+								if (PlayerSession)
+								{
+									IPlayerStreamFilter* StreamFilter = PlayerSession->GetStreamFilter();
+									if (StreamFilter && !StreamFilter->CanDecodeStream(Track->CodecInformation))
+									{
+										bIsSupported = false;
+									}
+								}
+
 							}
 							else
 							{
@@ -6515,7 +6526,7 @@ private:
 		return UEMEDIA_ERROR_FORMAT_ERROR;
 	}
 
-	UEMediaError FParserISO14496_12::ParseTX3GSampleType(FTrack* Track, const FMP4Box* SampleBox)
+	UEMediaError FParserISO14496_12::ParseTX3GSampleType(IPlayerSessionServices* PlayerSession, FTrack* Track, const FMP4Box* SampleBox)
 	{
 		// Need to check if there is a subtitle decoder capable of decoding this.
 		if (ISubtitleDecoder::IsSupported(TEXT(""), TEXT("tx3g")))
@@ -6536,7 +6547,7 @@ private:
 		return UEMEDIA_ERROR_NOT_SUPPORTED;
 	}
 
-	UEMediaError FParserISO14496_12::ParseWVTTSampleType(FTrack* Track, const FMP4Box* SampleBox)
+	UEMediaError FParserISO14496_12::ParseWVTTSampleType(IPlayerSessionServices* PlayerSession, FTrack* Track, const FMP4Box* SampleBox)
 	{
 		// Need to check if there is a subtitle decoder capable of decoding this.
 		if (ISubtitleDecoder::IsSupported(TEXT(""), TEXT("wvtt")))
@@ -6557,7 +6568,7 @@ private:
 		return UEMEDIA_ERROR_NOT_SUPPORTED;
 	}
 
-	UEMediaError FParserISO14496_12::PrepareTracks(TSharedPtrTS<const IParserISO14496_12> OptionalMP4InitSegment)
+	UEMediaError FParserISO14496_12::PrepareTracks(IPlayerSessionServices* PlayerSession, TSharedPtrTS<const IParserISO14496_12> OptionalMP4InitSegment)
 	{
 		delete ParsedTrackInfo;
 		ParsedTrackInfo = nullptr;
@@ -6773,11 +6784,11 @@ private:
 								UEMediaError Error = UEMEDIA_ERROR_FORMAT_ERROR;
 								if (FRMABox->GetDataFormat() == FMP4Box::kSample_avc1 || FRMABox->GetDataFormat() == FMP4Box::kSample_avc3)
 								{
-									Error = ParseAVC1SampleType(Track.Get(), STSDFirstChildBox);
+									Error = ParseAVC1SampleType(PlayerSession, Track.Get(), STSDFirstChildBox);
 								}
 								else if ( FRMABox->GetDataFormat() == FMP4Box::kSample_hvc1 || FRMABox->GetDataFormat() == FMP4Box::kSample_hev1)
 								{
-									Error = ParseHVC1SampleType(Track.Get(), STSDFirstChildBox);
+									Error = ParseHVC1SampleType(PlayerSession, Track.Get(), STSDFirstChildBox);
 								}
 								else
 								{
@@ -6815,7 +6826,7 @@ private:
 									return UEMEDIA_ERROR_FORMAT_ERROR;
 								}
 								// Parse the sample
-								UEMediaError Error = ParseMP4ASampleType(Track.Get(), STSDFirstChildBox);
+								UEMediaError Error = ParseMP4ASampleType(PlayerSession, Track.Get(), STSDFirstChildBox);
 								if (Error == UEMEDIA_ERROR_NOT_SUPPORTED)
 								{
 									bIsSupported = false;
@@ -6832,7 +6843,7 @@ private:
 							case FMP4Box::kSample_avc1:
 							case FMP4Box::kSample_avc3:
 							{
-								UEMediaError Error = ParseAVC1SampleType(Track.Get(), STSDFirstChildBox);
+								UEMediaError Error = ParseAVC1SampleType(PlayerSession, Track.Get(), STSDFirstChildBox);
 								/*
 								if (Error == UEMEDIA_ERROR_NOT_SUPPORTED)
 								{
@@ -6847,7 +6858,7 @@ private:
 							case FMP4Box::kSample_hvc1:
 							case FMP4Box::kSample_hev1:
 							{
-								UEMediaError Error = ParseHVC1SampleType(Track.Get(), STSDFirstChildBox);
+								UEMediaError Error = ParseHVC1SampleType(PlayerSession, Track.Get(), STSDFirstChildBox);
 								/*
 								if (Error == UEMEDIA_ERROR_NOT_SUPPORTED)
 								{
@@ -6861,7 +6872,7 @@ private:
 							}
 							case FMP4Box::kSample_mp4a:
 							{
-								UEMediaError Error = ParseMP4ASampleType(Track.Get(), STSDFirstChildBox);
+								UEMediaError Error = ParseMP4ASampleType(PlayerSession, Track.Get(), STSDFirstChildBox);
 								if (Error == UEMEDIA_ERROR_NOT_SUPPORTED)
 								{
 									bIsSupported = false;
@@ -6878,7 +6889,7 @@ private:
 							}
 							case FMP4Box::kSample_tx3g:
 							{
-								UEMediaError Error = ParseTX3GSampleType(Track.Get(), STSDFirstChildBox);
+								UEMediaError Error = ParseTX3GSampleType(PlayerSession, Track.Get(), STSDFirstChildBox);
 								if (Error == UEMEDIA_ERROR_NOT_SUPPORTED)
 								{
 									bIsSupported = false;
@@ -6891,7 +6902,7 @@ private:
 							}
 							case FMP4Box::kSample_wvtt:
 							{
-								UEMediaError Error = ParseWVTTSampleType(Track.Get(), STSDFirstChildBox);
+								UEMediaError Error = ParseWVTTSampleType(PlayerSession, Track.Get(), STSDFirstChildBox);
 								if (Error == UEMEDIA_ERROR_NOT_SUPPORTED)
 								{
 									bIsSupported = false;
diff --git a/Engine/Plugins/Media/ElectraPlayer/Source/ElectraPlayerRuntime/Private/Runtime/Demuxer/ParserISO14496-12.h b/Engine/Plugins/Media/ElectraPlayer/Source/ElectraPlayerRuntime/Private/Runtime/Demuxer/ParserISO14496-12.h
--- a/Engine/Plugins/Media/ElectraPlayer/Source/ElectraPlayerRuntime/Private/Runtime/Demuxer/ParserISO14496-12.h
+++ b/Engine/Plugins/Media/ElectraPlayer/Source/ElectraPlayerRuntime/Private/Runtime/Demuxer/ParserISO14496-12.h
@@ -153,7 +153,7 @@ namespace Electra
 		/*******************************************************************************************************************/
 		class ITrack;
 
-		virtual UEMediaError PrepareTracks(TSharedPtrTS<const IParserISO14496_12> OptionalMP4InitSegment) = 0;
+		virtual UEMediaError PrepareTracks(IPlayerSessionServices* PlayerSession, TSharedPtrTS<const IParserISO14496_12> OptionalMP4InitSegment) = 0;
 
 		virtual TMediaOptionalValue<FTimeFraction> GetMovieDuration() const = 0;
 
diff --git a/Engine/Plugins/Media/ElectraPlayer/Source/ElectraPlayerRuntime/Private/Runtime/Player/DASH/ManifestDASH.cpp b/Engine/Plugins/Media/ElectraPlayer/Source/ElectraPlayerRuntime/Private/Runtime/Player/DASH/ManifestDASH.cpp
--- a/Engine/Plugins/Media/ElectraPlayer/Source/ElectraPlayerRuntime/Private/Runtime/Player/DASH/ManifestDASH.cpp
+++ b/Engine/Plugins/Media/ElectraPlayer/Source/ElectraPlayerRuntime/Private/Runtime/Player/DASH/ManifestDASH.cpp
@@ -2963,7 +2963,7 @@ void FManifestDASHInternal::FRepresentation::SegmentIndexDownloadComplete(TShare
 			UEMediaError parseError = Index->ParseHeader(&BoxReader, &BoxReader, LoadRequest->PlayerSessionServices, nullptr);
 			if (parseError == UEMEDIA_ERROR_OK || parseError == UEMEDIA_ERROR_END_OF_STREAM)
 			{
-				if (Index->PrepareTracks(nullptr) == UEMEDIA_ERROR_OK && Index->GetNumberOfSegmentIndices() > 0)
+				if (Index->PrepareTracks(LoadRequest->PlayerSessionServices, nullptr) == UEMEDIA_ERROR_OK && Index->GetNumberOfSegmentIndices() > 0)
 				{
 					SegmentIndex = MoveTemp(Index);
 					bOk = true;
diff --git a/Engine/Plugins/Media/ElectraPlayer/Source/ElectraPlayerRuntime/Private/Runtime/Player/DASH/StreamReaderFMP4DASH.cpp b/Engine/Plugins/Media/ElectraPlayer/Source/ElectraPlayerRuntime/Private/Runtime/Player/DASH/StreamReaderFMP4DASH.cpp
--- a/Engine/Plugins/Media/ElectraPlayer/Source/ElectraPlayerRuntime/Private/Runtime/Player/DASH/StreamReaderFMP4DASH.cpp
+++ b/Engine/Plugins/Media/ElectraPlayer/Source/ElectraPlayerRuntime/Private/Runtime/Player/DASH/StreamReaderFMP4DASH.cpp
@@ -527,7 +527,7 @@ FErrorDetail FStreamReaderFMP4DASH::FStreamHandler::GetInitSegment(TSharedPtrTS<
 	if (parseError == UEMEDIA_ERROR_OK || parseError == UEMEDIA_ERROR_END_OF_STREAM)
 	{
 		// Parse the tracks of the init segment. We do this mainly to get to the CSD we might need should we have to insert filler data later.
-		parseError = Init->PrepareTracks(TSharedPtrTS<const IParserISO14496_12>());
+		parseError = Init->PrepareTracks(PlayerSessionService, TSharedPtrTS<const IParserISO14496_12>());
 		if (parseError == UEMEDIA_ERROR_OK)
 		{
 			// Add this to the entity cache in case it needs to be retrieved again.
@@ -821,7 +821,7 @@ void FStreamReaderFMP4DASH::FStreamHandler::HandleRequest()
 						SCOPE_CYCLE_COUNTER(STAT_ElectraPlayer_DASH_StreamReader);
 						CSV_SCOPED_TIMING_STAT(ElectraPlayer, DASH_StreamReader);
 						Request->Segment.bSawLMSG = MP4Parser->HasBrand(IParserISO14496_12::BrandType_lmsg);
-						parseError = MP4Parser->PrepareTracks(MP4InitSegment);
+						parseError = MP4Parser->PrepareTracks(PlayerSessionService, MP4InitSegment);
 					}
 					if (parseError == UEMEDIA_ERROR_OK)
 					{
diff --git a/Engine/Plugins/Media/ElectraPlayer/Source/ElectraPlayerRuntime/Private/Runtime/Player/HLS/StreamReaderHLSfmp4.cpp b/Engine/Plugins/Media/ElectraPlayer/Source/ElectraPlayerRuntime/Private/Runtime/Player/HLS/StreamReaderHLSfmp4.cpp
--- a/Engine/Plugins/Media/ElectraPlayer/Source/ElectraPlayerRuntime/Private/Runtime/Player/HLS/StreamReaderHLSfmp4.cpp
+++ b/Engine/Plugins/Media/ElectraPlayer/Source/ElectraPlayerRuntime/Private/Runtime/Player/HLS/StreamReaderHLSfmp4.cpp
@@ -570,7 +570,7 @@ FStreamReaderHLSfmp4::FStreamHandler::EInitSegmentResult FStreamReaderHLSfmp4::F
 				if (parseError == UEMEDIA_ERROR_OK || parseError == UEMEDIA_ERROR_END_OF_STREAM)
 				{
 					// Parse the tracks of the init segment. We do this mainly to get to the CSD we might need should we have to insert filler data later.
-					parseError = InitSegmentParser->PrepareTracks(TSharedPtrTS<const IParserISO14496_12>());
+					parseError = InitSegmentParser->PrepareTracks(PlayerSessionService, TSharedPtrTS<const IParserISO14496_12>());
 					if (parseError == UEMEDIA_ERROR_OK)
 					{
 						Request->InitSegmentCache->AddInitSegment(InitSegmentParser, Request->InitSegmentInfo, FTimeValue::GetPositiveInfinity());
@@ -836,7 +836,7 @@ void FStreamReaderHLSfmp4::FStreamHandler::HandleRequest()
 						{
 							SCOPE_CYCLE_COUNTER(STAT_ElectraPlayer_HLS_StreamReader);
 							CSV_SCOPED_TIMING_STAT(ElectraPlayer, HLS_StreamReader);
-							parseError = MP4Parser->PrepareTracks(MP4InitSegment);
+							parseError = MP4Parser->PrepareTracks(PlayerSessionService, MP4InitSegment);
 						}
 						if (parseError == UEMEDIA_ERROR_OK)
 						{
diff --git a/Engine/Plugins/Media/ElectraPlayer/Source/ElectraPlayerRuntime/Private/Runtime/Player/mp4/PlaylistReaderMP4.cpp b/Engine/Plugins/Media/ElectraPlayer/Source/ElectraPlayerRuntime/Private/Runtime/Player/mp4/PlaylistReaderMP4.cpp
--- a/Engine/Plugins/Media/ElectraPlayer/Source/ElectraPlayerRuntime/Private/Runtime/Player/mp4/PlaylistReaderMP4.cpp
+++ b/Engine/Plugins/Media/ElectraPlayer/Source/ElectraPlayerRuntime/Private/Runtime/Player/mp4/PlaylistReaderMP4.cpp
@@ -367,7 +367,7 @@ void FPlaylistReaderMP4::WorkerThread()
 				}
 
 				// Prepare the tracks in the stream that are of a supported codec.
-				parseError = MP4Parser->PrepareTracks(TSharedPtrTS<const IParserISO14496_12>());
+				parseError = MP4Parser->PrepareTracks(PlayerSessionServices, TSharedPtrTS<const IParserISO14496_12>());
 				if (parseError == UEMEDIA_ERROR_OK)
 				{
 					Manifest = MakeSharedTS<FManifestMP4Internal>(PlayerSessionServices);
