ElectraPlayer: Fixed an issue decoding non-16 pixel aligned videos on DX11 in software decode mode

#rnx

diff --git a/Engine/Plugins/Media/ElectraPlayer/Source/ElectraPlayerRuntime/Private/Runtime/Decoder/Windows/VideoDecoderH264_PC.cpp b/Engine/Plugins/Media/ElectraPlayer/Source/ElectraPlayerRuntime/Private/Runtime/Decoder/Windows/VideoDecoderH264_PC.cpp
--- a/Engine/Plugins/Media/ElectraPlayer/Source/ElectraPlayerRuntime/Private/Runtime/Decoder/Windows/VideoDecoderH264_PC.cpp
+++ b/Engine/Plugins/Media/ElectraPlayer/Source/ElectraPlayerRuntime/Private/Runtime/Decoder/Windows/VideoDecoderH264_PC.cpp
@@ -153,8 +153,8 @@ private:
 void FElectraPlayerVideoDecoderOutputPC::PreInitForDecode(FIntPoint OutputDim, const TFunction<void(int32 /*ApiReturnValue*/, const FString& /*Message*/, uint16 /*Code*/, UEMediaError /*Error*/ )>& PostError)
 {
 	FIntPoint Dim;
-	Dim.X = OutputDim.X;
-	Dim.Y = OutputDim.Y * 3 / 2;	// adjust height to encompass Y and UV planes
+	Dim.X = (OutputDim.X + 15) & ~15;
+	Dim.Y = ((OutputDim.Y + 15) & ~15) * 3/2;	// adjust height to encompass Y and UV planes
 
 	bool bNeedNew = !MFSample.IsValid() || SampleDim != Dim;
 
@@ -221,7 +221,7 @@ void FElectraPlayerVideoDecoderOutputPC::PreInitForDecode(FIntPoint OutputDim, c
 			}
 			else
 			{
-				AllocSize = ((OutputDim.X + 15) & ~15) * ((OutputDim.Y + 15) & ~15) * sizeof(uint8*) * 2;
+				AllocSize = Dim.X * Dim.Y * sizeof(uint8*) * 2;
 			}
 
 			// SW decode results are just delivered in a simple CPU-side buffer. Create the decoder side version of this...
@@ -259,16 +259,14 @@ void FElectraPlayerVideoDecoderOutputPC::ProcessDecodeOutput(FIntPoint OutputDim
 
 	FVideoDecoderOutput::Initialize(InParamDict);
 
-	SampleDim.X = OutputDim.X;
-	SampleDim.Y = OutputDim.Y * 3 / 2;
+	// This needs to be a multiple of 16
+	SampleDim.X = (OutputDim.X + 15) & ~15;
+	SampleDim.Y = ((OutputDim.Y + 15) & ~15) * 3 / 2;
 
 	Stride = SampleDim.X * sizeof(uint8);
 
 	if (OutputType == EOutputType::SoftwareWin7) // Win7 & DX12 (SW)
 	{
-		// This needs to be a multiple of 16
-		SampleDim.Y = ((OutputDim.Y + 15) & ~15) * 3 / 2;
-
 		// Retrieve frame data and store it in Buffer for rendering later
 		TRefCountPtr<IMFMediaBuffer> MediaBuffer;
 		if (MFSample->GetBufferByIndex(0, MediaBuffer.GetInitReference()) != S_OK)
