ElectraPlayer: Added a missing mutex around access to a variable used by different threads.

#rnx

diff --git a/Engine/Plugins/Media/ElectraPlayer/Source/ElectraPlayerRuntime/Private/Runtime/Player/AdaptiveStreamingPlayer.cpp b/Engine/Plugins/Media/ElectraPlayer/Source/ElectraPlayerRuntime/Private/Runtime/Player/AdaptiveStreamingPlayer.cpp
--- a/Engine/Plugins/Media/ElectraPlayer/Source/ElectraPlayerRuntime/Private/Runtime/Player/AdaptiveStreamingPlayer.cpp
+++ b/Engine/Plugins/Media/ElectraPlayer/Source/ElectraPlayerRuntime/Private/Runtime/Player/AdaptiveStreamingPlayer.cpp
@@ -1502,6 +1502,8 @@ void FAdaptiveStreamingPlayer::HandlePlayStateChanges()
  */
 void FAdaptiveStreamingPlayer::HandleMetadataChanges()
 {
+	FScopeLock Lock(&PeriodCriticalSection);
+
 	// The timeline can change dynamically. Refresh it on occasion.
 	if (Manifest.IsValid() && ManifestType == EMediaFormatType::DASH)
 	{
@@ -3011,6 +3013,8 @@ void FAdaptiveStreamingPlayer::AddUpcomingPeriod(TSharedPtrTS<IManifest::IPlayPe
 	TSharedPtrTS<ITimelineMediaAsset> Period = InUpcomingPeriod->GetMediaAsset();
 	if (Period.IsValid())
 	{
+		FScopeLock Lock(&PeriodCriticalSection);
+
 		FString PeriodID = Period->GetUniqueIdentifier();
 		// Only add if it is not already in the list.
 		if (!UpcomingPeriods.ContainsByPredicate([PeriodID](const FPeriodInformation& e){ return e.ID.Equals(PeriodID); }))
@@ -3396,8 +3400,10 @@ void FAdaptiveStreamingPlayer::InternalStop(bool bHoldCurrentFrame)
 	CurrentPlayPeriodVideo.Reset();
 	CurrentPlayPeriodAudio.Reset();
 	CurrentPlayPeriodText.Reset();
+	PeriodCriticalSection.Lock();
 	ActivePeriods.Empty();
 	UpcomingPeriods.Empty();
+	PeriodCriticalSection.Unlock();
 
 	PlaybackRate = 0.0;
 	CurrentState = EPlayerState::eState_Paused;
diff --git a/Engine/Plugins/Media/ElectraPlayer/Source/ElectraPlayerRuntime/Private/Runtime/Player/AdaptiveStreamingPlayerDecoders.cpp b/Engine/Plugins/Media/ElectraPlayer/Source/ElectraPlayerRuntime/Private/Runtime/Player/AdaptiveStreamingPlayerDecoders.cpp
--- a/Engine/Plugins/Media/ElectraPlayer/Source/ElectraPlayerRuntime/Private/Runtime/Player/AdaptiveStreamingPlayerDecoders.cpp
+++ b/Engine/Plugins/Media/ElectraPlayer/Source/ElectraPlayerRuntime/Private/Runtime/Player/AdaptiveStreamingPlayerDecoders.cpp
@@ -169,6 +169,8 @@ bool FAdaptiveStreamingPlayer::CanDecodeSubtitle(const FString& MimeType, const
 
 bool FAdaptiveStreamingPlayer::FindMatchingStreamInfo(FStreamCodecInformation& OutStreamInfo, const FTimeValue& AtTime, int32 MaxWidth, int32 MaxHeight)
 {
+	FScopeLock Lock(&PeriodCriticalSection);
+
 	TArray<FStreamCodecInformation> VideoCodecInfos;
 	TSharedPtrTS<ITimelineMediaAsset> Asset;
 
@@ -583,6 +585,7 @@ void FAdaptiveStreamingPlayer::HandleSubtitleDecoder()
 
 		if (CurrentPos.IsValid() && RenderClock->IsRunning())
 		{
+			PeriodCriticalSection.Lock();
 			FTimeValue LocalPos;
 			for(int32 i=0, iMax=ActivePeriods.Num(); i<iMax; ++i)
 			{
@@ -592,6 +595,7 @@ void FAdaptiveStreamingPlayer::HandleSubtitleDecoder()
 					break;
 				}
 			}
+			PeriodCriticalSection.Unlock();
 			SubtitleDecoder.Decoder->UpdatePlaybackPosition(CurrentPos, LocalPos);
 		}
 	}
diff --git a/Engine/Plugins/Media/ElectraPlayer/Source/ElectraPlayerRuntime/Private/Runtime/Player/AdaptiveStreamingPlayerInternal.h b/Engine/Plugins/Media/ElectraPlayer/Source/ElectraPlayerRuntime/Private/Runtime/Player/AdaptiveStreamingPlayerInternal.h
--- a/Engine/Plugins/Media/ElectraPlayer/Source/ElectraPlayerRuntime/Private/Runtime/Player/AdaptiveStreamingPlayerInternal.h
+++ b/Engine/Plugins/Media/ElectraPlayer/Source/ElectraPlayerRuntime/Private/Runtime/Player/AdaptiveStreamingPlayerInternal.h
@@ -1644,6 +1644,7 @@ private:
 	FLoopParam															CurrentLoopParam;
 	TMediaQueueDynamicNoLock<FPlayerLoopState>							NextLoopStates;
 
+	FCriticalSection													PeriodCriticalSection;
 	TArray<FPeriodInformation>											ActivePeriods;
 	TArray<FPeriodInformation>											UpcomingPeriods;
 	TSharedPtrTS<IManifest::IPlayPeriod>								InitialPlayPeriod;
