ElectraPlayer: Checking an .mp4 stream for decodability throwing an error if it cannot be played.

#rnx

diff --git a/Engine/Plugins/Media/ElectraPlayer/Source/ElectraPlayerRuntime/Private/Runtime/Player/AdaptiveStreamingPlayerABR.cpp b/Engine/Plugins/Media/ElectraPlayer/Source/ElectraPlayerRuntime/Private/Runtime/Player/AdaptiveStreamingPlayerABR.cpp
--- a/Engine/Plugins/Media/ElectraPlayer/Source/ElectraPlayerRuntime/Private/Runtime/Player/AdaptiveStreamingPlayerABR.cpp
+++ b/Engine/Plugins/Media/ElectraPlayer/Source/ElectraPlayerRuntime/Private/Runtime/Player/AdaptiveStreamingPlayerABR.cpp
@@ -889,7 +889,7 @@ namespace Electra
 					// If we cannot switch to another stream we need to retry this one.
 					if (!bCanSwitchToAlternateStreams)
 					{
-						return ESegmentAction::Retry;
+						return Stats.bParseFailure ? ESegmentAction::Fail : ESegmentAction::Retry;
 					}
 
 					// Is this an init segment?
diff --git a/Engine/Plugins/Media/ElectraPlayer/Source/ElectraPlayerRuntime/Private/Runtime/Player/mp4/StreamReaderMP4.cpp b/Engine/Plugins/Media/ElectraPlayer/Source/ElectraPlayerRuntime/Private/Runtime/Player/mp4/StreamReaderMP4.cpp
--- a/Engine/Plugins/Media/ElectraPlayer/Source/ElectraPlayerRuntime/Private/Runtime/Player/mp4/StreamReaderMP4.cpp
+++ b/Engine/Plugins/Media/ElectraPlayer/Source/ElectraPlayerRuntime/Private/Runtime/Player/mp4/StreamReaderMP4.cpp
@@ -13,6 +13,8 @@
 #include "Player/AdaptiveStreamingPlayerABR.h"
 #include "Player/mp4/ManifestMP4.h"
 #include "Player/mp4/StreamReaderMP4.h"
+#include "Player/PlayerSessionServices.h"
+#include "Player/PlayerStreamFilter.h"
 
 
 DECLARE_CYCLE_STAT(TEXT("FStreamReaderMP4_HandleRequest"), STAT_ElectraPlayer_MP4_StreamReader, STATGROUP_ElectraPlayer);
@@ -274,6 +276,7 @@ void FStreamReaderMP4::HTTPUpdateStats(const FTimeValue& CurrentTime, const IEle
 void FStreamReaderMP4::HandleRequest()
 {
 	TSharedPtrTS<FStreamSegmentRequestMP4>	Request = CurrentRequest;
+	FString ParsingErrorMessage;
 
 	FManifestMP4Internal::FTimelineAssetMP4* TimelineAsset = static_cast<FManifestMP4Internal::FTimelineAssetMP4*>(Request->MediaAsset.Get());
 
@@ -285,6 +288,12 @@ void FStreamReaderMP4::HandleRequest()
 	FTimeValue LoopTimestampOffset = Request->PlayerLoopState.LoopBasetime;
 	TSharedPtrTS<const FPlayerLoopState> PlayerLoopState = MakeSharedTS<const FPlayerLoopState>(Request->PlayerLoopState);
 
+	// Check if a remote stream contains a track that is not supported or decodable on this platform.
+	// Local files are not checked since they are supposed to be valid. This is to guard against streams
+	// of unknown origin.
+	IPlayerStreamFilter* StreamFilter = PlayerSessionServices ? PlayerSessionServices->GetStreamFilter() : nullptr;
+	bool bIsRemotePlayback = TimelineAsset->GetMediaURL().StartsWith(TEXT("https:")) || TimelineAsset->GetMediaURL().StartsWith(TEXT("http:"));
+	bool bHasUnsupportTrack = false;
 
 	// Get the list of all the tracks that have been selected in the asset.
 	// This does not mean their data will be _used_ for playback, only that the track is usable by the player
@@ -317,6 +326,12 @@ void FStreamReaderMP4::HandleRequest()
 				for(int32 nRepr=0; nRepr<NumRepr; ++nRepr)
 				{
 					TSharedPtrTS<IPlaybackAssetRepresentation> Representation = AdaptationSet->GetRepresentationByIndex(nRepr);
+
+					if (bIsRemotePlayback && StreamFilter && !StreamFilter->CanDecodeStream(Representation->GetCodecInformation()))
+					{
+						bHasUnsupportTrack = true;
+					}
+
 					// Note: By definition the representations unique identifier is a string of the numeric track ID and can thus be parsed back into a number.
 					FString ReprID = Representation->GetUniqueIdentifier();
 					uint32 TrackId;
@@ -433,6 +448,15 @@ void FStreamReaderMP4::HandleRequest()
 		CSV_SCOPED_TIMING_STAT(ElectraPlayer, MP4_StreamReader);
 		AllTrackIterator = TimelineAsset->GetMoovBoxParser()->CreateAllTrackIteratorByFilePos(Request->FileStartOffset);
 	}
+
+	// If a track cannot be used on this platforms we set up an error now.
+	if (bHasUnsupportTrack)
+	{
+		ds.bParseFailure  = true;
+		ParsingErrorMessage = FString::Printf(TEXT("Segment contains unsupported or non-decodable tracks."));
+		bHasErrored = true;
+	}
+
 	while(!bDone && !HasErrored() && !HasBeenAborted() && !bTerminate)
 	{
 		auto UpdateSelectedTrack = [&SelectedTrackMap](const IParserISO14496_12::ITrackIterator* trkIt, TMap<uint32, FSelectedTrackData>& ActiveTrks) -> FSelectedTrackData&
@@ -523,7 +547,7 @@ void FStreamReaderMP4::HandleRequest()
 			else if (SampleFileOffset < ReadBuffer.GetCurrentPos())
 			{
 				ds.bParseFailure  = true;
-				Request->ConnectionInfo.StatusInfo.ErrorDetail.SetMessage(FString::Printf(TEXT("Segment parse error. Sample offset %lld for sample #%u in track %u is before the current read position at %lld"), (long long int)SampleFileOffset, SampleNumber, TrackIt->GetTrack()->GetID(), (long long int)ReadBuffer.GetCurrentPos()));
+				ParsingErrorMessage = FString::Printf(TEXT("Segment parse error. Sample offset %lld for sample #%u in track %u is before the current read position at %lld"), (long long int)SampleFileOffset, SampleNumber, TrackIt->GetTrack()->GetID(), (long long int)ReadBuffer.GetCurrentPos());
 				bHasErrored = true;
 				break;
 			}
@@ -654,7 +678,6 @@ void FStreamReaderMP4::HandleRequest()
 	// Remember the next largest timestamp from all tracks.
 	Request->NextLargestExpectedTimestamp = NextLargestExpectedTimestamp;
 
-
 	// Set downloaded and delivered duration from the primary track.
 	FSelectedTrackData& PrimaryTrack = ActiveTrackMap.FindOrAdd(PrimaryTrackID);
 	DurationSuccessfullyRead	  = PrimaryTrack.DurationSuccessfullyRead;
@@ -667,6 +690,10 @@ void FStreamReaderMP4::HandleRequest()
 	{
 		ds.FailureReason = Request->ConnectionInfo.StatusInfo.ErrorDetail.GetMessage();
 	}
+	if (ParsingErrorMessage.Len())
+	{
+		ds.FailureReason = ParsingErrorMessage;
+	}
 //		if (bAbortedByABR)
 //		{
 //			// If aborted set the reason as the download failure.
