ElectraPlayer: Fixed popping access units from a deselected multitrack buffer

#rnx

diff --git a/Engine/Plugins/Media/ElectraPlayer/Source/ElectraPlayerRuntime/Private/Runtime/StreamAccessUnitBuffer.h b/Engine/Plugins/Media/ElectraPlayer/Source/ElectraPlayerRuntime/Private/Runtime/StreamAccessUnitBuffer.h
--- a/Engine/Plugins/Media/ElectraPlayer/Source/ElectraPlayerRuntime/Private/Runtime/StreamAccessUnitBuffer.h
+++ b/Engine/Plugins/Media/ElectraPlayer/Source/ElectraPlayerRuntime/Private/Runtime/StreamAccessUnitBuffer.h
@@ -436,24 +436,29 @@ namespace Electra
 			while(true)
 			{
 				FAccessUnit* NextAU = nullptr;
-				if (Peek(NextAU))
+				FAccessUnit* PeekedAU = nullptr;
+				if (Peek(PeekedAU))
 				{
-					if (NextAU)
+					if (PeekedAU)
 					{
-						bool bDTS = NextValidDTS.IsValid() ? NextAU->DTS < NextValidDTS : true;
-						bool bPTS = NextValidPTS.IsValid() ? NextAU->PTS < NextValidPTS : true;
+						bool bDTS = NextValidDTS.IsValid() ? PeekedAU->DTS < NextValidDTS : true;
+						bool bPTS = NextValidPTS.IsValid() ? PeekedAU->PTS < NextValidPTS : true;
 						if (bDTS && bPTS)
 						{
-							OutPoppedDTS = NextAU->DTS;
-							OutPoppedPTS = NextAU->PTS;
+							OutPoppedDTS = PeekedAU->DTS;
+							OutPoppedPTS = PeekedAU->PTS;
 							Pop(NextAU);
 							FAccessUnit::Release(NextAU);
+							NextAU = nullptr;
 						}
 						else
 						{
-							FAccessUnit::Release(NextAU);
+							FAccessUnit::Release(PeekedAU);
+							PeekedAU = nullptr;
 							break;
 						}
+						FAccessUnit::Release(PeekedAU);
+						PeekedAU = nullptr;
 					}
 					else
 					{
