ElectraPlayer: reducing the time holding a mutex lock to just the period the relevant data member is accessed.

#rnx

diff --git a/Engine/Plugins/Media/ElectraPlayer/Source/ElectraPlayerRuntime/Private/Runtime/Player/AdaptiveStreamingPlayer.cpp b/Engine/Plugins/Media/ElectraPlayer/Source/ElectraPlayerRuntime/Private/Runtime/Player/AdaptiveStreamingPlayer.cpp
--- a/Engine/Plugins/Media/ElectraPlayer/Source/ElectraPlayerRuntime/Private/Runtime/Player/AdaptiveStreamingPlayer.cpp
+++ b/Engine/Plugins/Media/ElectraPlayer/Source/ElectraPlayerRuntime/Private/Runtime/Player/AdaptiveStreamingPlayer.cpp
@@ -1502,8 +1502,6 @@ void FAdaptiveStreamingPlayer::HandlePlayStateChanges()
  */
 void FAdaptiveStreamingPlayer::HandleMetadataChanges()
 {
-	FScopeLock Lock(&PeriodCriticalSection);
-
 	// The timeline can change dynamically. Refresh it on occasion.
 	if (Manifest.IsValid() && ManifestType == EMediaFormatType::DASH)
 	{
@@ -1560,6 +1558,7 @@ void FAdaptiveStreamingPlayer::HandleMetadataChanges()
 	if (Current.IsValid())
 	{
 		FTimeValue PastTime = Current - FTimeValue().SetFromMilliseconds(1000);
+		FScopeLock Lock(&ActivePeriodCriticalSection);
 		for(int32 i=0; i<ActivePeriods.Num(); ++i)
 		{
 			if ((ActivePeriods[i].TimeRange.End.IsValid() && PastTime > ActivePeriods[i].TimeRange.End) ||
@@ -3013,8 +3012,6 @@ void FAdaptiveStreamingPlayer::AddUpcomingPeriod(TSharedPtrTS<IManifest::IPlayPe
 	TSharedPtrTS<ITimelineMediaAsset> Period = InUpcomingPeriod->GetMediaAsset();
 	if (Period.IsValid())
 	{
-		FScopeLock Lock(&PeriodCriticalSection);
-
 		FString PeriodID = Period->GetUniqueIdentifier();
 		// Only add if it is not already in the list.
 		if (!UpcomingPeriods.ContainsByPredicate([PeriodID](const FPeriodInformation& e){ return e.ID.Equals(PeriodID); }))
@@ -3029,7 +3026,9 @@ void FAdaptiveStreamingPlayer::AddUpcomingPeriod(TSharedPtrTS<IManifest::IPlayPe
 			// Tell the AEMS handler that a new period will be coming up. It needs this information to cut overlapping events.
 			AEMSEventHandler->NewUpcomingPeriod(Next.ID, Next.TimeRange);
 
+			ActivePeriodCriticalSection.Lock();
 			ActivePeriods.Add(Next);
+			ActivePeriodCriticalSection.Unlock();
 			UpcomingPeriods.Emplace(MoveTemp(Next));
 		}
 	}
@@ -3400,10 +3399,10 @@ void FAdaptiveStreamingPlayer::InternalStop(bool bHoldCurrentFrame)
 	CurrentPlayPeriodVideo.Reset();
 	CurrentPlayPeriodAudio.Reset();
 	CurrentPlayPeriodText.Reset();
-	PeriodCriticalSection.Lock();
+	ActivePeriodCriticalSection.Lock();
 	ActivePeriods.Empty();
+	ActivePeriodCriticalSection.Unlock();
 	UpcomingPeriods.Empty();
-	PeriodCriticalSection.Unlock();
 
 	PlaybackRate = 0.0;
 	CurrentState = EPlayerState::eState_Paused;
diff --git a/Engine/Plugins/Media/ElectraPlayer/Source/ElectraPlayerRuntime/Private/Runtime/Player/AdaptiveStreamingPlayerDecoders.cpp b/Engine/Plugins/Media/ElectraPlayer/Source/ElectraPlayerRuntime/Private/Runtime/Player/AdaptiveStreamingPlayerDecoders.cpp
--- a/Engine/Plugins/Media/ElectraPlayer/Source/ElectraPlayerRuntime/Private/Runtime/Player/AdaptiveStreamingPlayerDecoders.cpp
+++ b/Engine/Plugins/Media/ElectraPlayer/Source/ElectraPlayerRuntime/Private/Runtime/Player/AdaptiveStreamingPlayerDecoders.cpp
@@ -169,17 +169,16 @@ bool FAdaptiveStreamingPlayer::CanDecodeSubtitle(const FString& MimeType, const
 
 bool FAdaptiveStreamingPlayer::FindMatchingStreamInfo(FStreamCodecInformation& OutStreamInfo, const FTimeValue& AtTime, int32 MaxWidth, int32 MaxHeight)
 {
-	FScopeLock Lock(&PeriodCriticalSection);
-
 	TArray<FStreamCodecInformation> VideoCodecInfos;
 	TSharedPtrTS<ITimelineMediaAsset> Asset;
 
+	// Locate the period for the specified time.
+	ActivePeriodCriticalSection.Lock();
 	if (ActivePeriods.Num() == 0)
 	{
+		ActivePeriodCriticalSection.Unlock();
 		return false;
 	}
-
-	// Locate the period for the specified time.
 	for(int32 i=0; i<ActivePeriods.Num(); ++i)
 	{
 		if (AtTime >= ActivePeriods[i].TimeRange.Start && AtTime < (ActivePeriods[i].TimeRange.End.IsValid() ? ActivePeriods[i].TimeRange.End : FTimeValue::GetPositiveInfinity()))
@@ -187,6 +186,7 @@ bool FAdaptiveStreamingPlayer::FindMatchingStreamInfo(FStreamCodecInformation& O
 			Asset = ActivePeriods[i].Period;
 		}
 	}
+	ActivePeriodCriticalSection.Unlock();
 	if (Asset.IsValid())
 	{
 		if (Asset->GetNumberOfAdaptationSets(EStreamType::Video) > 0)
@@ -585,7 +585,7 @@ void FAdaptiveStreamingPlayer::HandleSubtitleDecoder()
 
 		if (CurrentPos.IsValid() && RenderClock->IsRunning())
 		{
-			PeriodCriticalSection.Lock();
+			ActivePeriodCriticalSection.Lock();
 			FTimeValue LocalPos;
 			for(int32 i=0, iMax=ActivePeriods.Num(); i<iMax; ++i)
 			{
@@ -595,7 +595,7 @@ void FAdaptiveStreamingPlayer::HandleSubtitleDecoder()
 					break;
 				}
 			}
-			PeriodCriticalSection.Unlock();
+			ActivePeriodCriticalSection.Unlock();
 			SubtitleDecoder.Decoder->UpdatePlaybackPosition(CurrentPos, LocalPos);
 		}
 	}
diff --git a/Engine/Plugins/Media/ElectraPlayer/Source/ElectraPlayerRuntime/Private/Runtime/Player/AdaptiveStreamingPlayerInternal.h b/Engine/Plugins/Media/ElectraPlayer/Source/ElectraPlayerRuntime/Private/Runtime/Player/AdaptiveStreamingPlayerInternal.h
--- a/Engine/Plugins/Media/ElectraPlayer/Source/ElectraPlayerRuntime/Private/Runtime/Player/AdaptiveStreamingPlayerInternal.h
+++ b/Engine/Plugins/Media/ElectraPlayer/Source/ElectraPlayerRuntime/Private/Runtime/Player/AdaptiveStreamingPlayerInternal.h
@@ -1644,7 +1644,7 @@ private:
 	FLoopParam															CurrentLoopParam;
 	TMediaQueueDynamicNoLock<FPlayerLoopState>							NextLoopStates;
 
-	FCriticalSection													PeriodCriticalSection;
+	FCriticalSection													ActivePeriodCriticalSection;
 	TArray<FPeriodInformation>											ActivePeriods;
 	TArray<FPeriodInformation>											UpcomingPeriods;
 	TSharedPtrTS<IManifest::IPlayPeriod>								InitialPlayPeriod;
