ElectraPlayer: merged H.264 and H.265 decoder output structures into a common one

#rnx

diff --git a/Engine/Plugins/Media/ElectraPlayer/Source/ElectraPlayerRuntime/Private/Runtime/Decoder/Apple/ElectraVideoDecoder_Apple.cpp b/Engine/Plugins/Media/ElectraPlayer/Source/ElectraPlayerRuntime/Private/Runtime/Decoder/Apple/ElectraVideoDecoder_Apple.cpp
new file mode 100644
--- /dev/null
+++ b/Engine/Plugins/Media/ElectraPlayer/Source/ElectraPlayerRuntime/Private/Runtime/Decoder/Apple/ElectraVideoDecoder_Apple.cpp
@@ -0,0 +1,81 @@
+// Copyright Epic Games, Inc. All Rights Reserved.
+
+#include "ElectraVideoDecoder_Apple.h"
+
+#if PLATFORM_MAC || PLATFORM_IOS || PLATFORM_TVOS
+
+namespace Electra
+{
+
+FElectraPlayerVideoDecoderOutputApple::FElectraPlayerVideoDecoderOutputApple()
+	: ImageBufferRef(nullptr)
+{
+}
+
+FElectraPlayerVideoDecoderOutputApple::~FElectraPlayerVideoDecoderOutputApple()
+{
+	if (ImageBufferRef)
+	{
+		CFRelease(ImageBufferRef);
+	}
+}
+
+void FElectraPlayerVideoDecoderOutputApple::Initialize(CVImageBufferRef InImageBufferRef, FParamDict* InParamDict)
+{
+	FVideoDecoderOutputApple::Initialize(InParamDict);
+
+	if (ImageBufferRef)
+	{
+		CFRelease(ImageBufferRef);
+	}
+	ImageBufferRef = InImageBufferRef;
+	if (ImageBufferRef)
+	{
+		CFRetain(ImageBufferRef);
+
+		Stride = CVPixelBufferGetBytesPerRow(ImageBufferRef);
+	}
+}
+
+void FElectraPlayerVideoDecoderOutputApple::SetOwner(const TSharedPtr<IDecoderOutputOwner, ESPMode::ThreadSafe>& InOwningRenderer)
+{
+	OwningRenderer = InOwningRenderer;
+}
+
+void FElectraPlayerVideoDecoderOutputApple::ShutdownPoolable()
+{
+	// release image buffer (we currently realloc it every time anyway)
+	if (ImageBufferRef)
+	{
+		CFRelease(ImageBufferRef);
+		ImageBufferRef = nullptr;
+	}
+
+	TSharedPtr<IDecoderOutputOwner, ESPMode::ThreadSafe> lockedVideoRenderer = OwningRenderer.Pin();
+	if (lockedVideoRenderer.IsValid())
+	{
+		lockedVideoRenderer->SampleReleasedToPool(GetDuration());
+	}
+}
+
+uint32 FElectraPlayerVideoDecoderOutputApple::GetStride() const
+{
+	return Stride;
+}
+
+CVImageBufferRef FElectraPlayerVideoDecoderOutputApple::GetImageBuffer() const
+{
+	return ImageBufferRef;
+}
+
+} // namespace Electra
+
+// -----------------------------------------------------------------------------------------------------------------------------
+
+FVideoDecoderOutput* FElectraPlayerPlatformVideoDecoderOutputFactory::Create()
+{
+    return new Electra::FElectraPlayerVideoDecoderOutputApple();
+}
+
+
+#endif
diff --git a/Engine/Plugins/Media/ElectraPlayer/Source/ElectraPlayerRuntime/Private/Runtime/Decoder/Apple/ElectraVideoDecoder_Apple.h b/Engine/Plugins/Media/ElectraPlayer/Source/ElectraPlayerRuntime/Private/Runtime/Decoder/Apple/ElectraVideoDecoder_Apple.h
new file mode 100644
--- /dev/null
+++ b/Engine/Plugins/Media/ElectraPlayer/Source/ElectraPlayerRuntime/Private/Runtime/Decoder/Apple/ElectraVideoDecoder_Apple.h
@@ -0,0 +1,47 @@
+// Copyright Epic Games, Inc. All Rights Reserved.
+
+#pragma once
+
+#include "HAL/Platform.h"
+
+#if PLATFORM_MAC || PLATFORM_IOS || PLATFORM_TVOS
+
+#include "PlayerCore.h"
+#include "MediaVideoDecoderOutputApple.h"
+#include "Renderer/RendererVideo.h"
+
+
+namespace Electra
+{
+
+class FElectraPlayerVideoDecoderOutputApple : public FVideoDecoderOutputApple
+{
+public:
+	FElectraPlayerVideoDecoderOutputApple();
+
+	~FElectraPlayerVideoDecoderOutputApple();
+
+    void Initialize(CVImageBufferRef InImageBufferRef, FParamDict* InParamDict);
+
+	void SetOwner(const TSharedPtr<IDecoderOutputOwner, ESPMode::ThreadSafe>& InOwningRenderer) override;
+
+	void ShutdownPoolable() override;
+
+	uint32 GetStride() const override;
+
+	CVImageBufferRef GetImageBuffer() const override;
+
+private:
+	uint32 Stride;
+
+	CVImageBufferRef ImageBufferRef;
+
+	// We hold a weak reference to the video renderer. During destruction the video renderer could be destroyed while samples are still out there..
+	TWeakPtr<IDecoderOutputOwner, ESPMode::ThreadSafe> OwningRenderer;
+};
+
+
+} // namespace Electra
+
+#endif
+
diff --git a/Engine/Plugins/Media/ElectraPlayer/Source/ElectraPlayerRuntime/Private/Runtime/Decoder/Apple/VideoDecoderH264_Apple.cpp b/Engine/Plugins/Media/ElectraPlayer/Source/ElectraPlayerRuntime/Private/Runtime/Decoder/Apple/VideoDecoderH264_Apple.cpp
--- a/Engine/Plugins/Media/ElectraPlayer/Source/ElectraPlayerRuntime/Private/Runtime/Decoder/Apple/VideoDecoderH264_Apple.cpp
+++ b/Engine/Plugins/Media/ElectraPlayer/Source/ElectraPlayerRuntime/Private/Runtime/Decoder/Apple/VideoDecoderH264_Apple.cpp
@@ -18,6 +18,8 @@
 #include "MediaVideoDecoderOutputApple.h"
 #include "Renderer/RendererVideo.h"
 
+#include "ElectraVideoDecoder_Apple.h"
+
 #include <VideoToolbox/VideoToolbox.h>
 
 DECLARE_CYCLE_STAT(TEXT("FVideoDecoderH264::Decode()"), STAT_ElectraPlayer_VideoH264Decode, STATGROUP_ElectraPlayer);
@@ -322,83 +324,6 @@ IVideoDecoderH264* IVideoDecoderH264::Create()
 }
 
 
-/***************************************************************************************************************************************************/
-/***************************************************************************************************************************************************/
-/***************************************************************************************************************************************************/
-
-class FElectraPlayerVideoDecoderOutputApple : public FVideoDecoderOutputApple
-{
-public:
-	FElectraPlayerVideoDecoderOutputApple()
-		: ImageBufferRef(nullptr)
-	{
-	}
-
-	~FElectraPlayerVideoDecoderOutputApple()
-	{
-		if (ImageBufferRef)
-		{
-			CFRelease(ImageBufferRef);
-		}
-	}
-
-    void Initialize(CVImageBufferRef InImageBufferRef, FParamDict* InParamDict)
-	{
-		FVideoDecoderOutputApple::Initialize(InParamDict);
-
-		if (ImageBufferRef)
-		{
-			CFRelease(ImageBufferRef);
-		}
-		ImageBufferRef = InImageBufferRef;
-		if (ImageBufferRef)
-		{
-			CFRetain(ImageBufferRef);
-
-			Stride = CVPixelBufferGetBytesPerRow(ImageBufferRef);
-		}
-	}
-
-	void SetOwner(const TSharedPtr<IDecoderOutputOwner, ESPMode::ThreadSafe>& InOwningRenderer) override
-	{
-		OwningRenderer = InOwningRenderer;
-	}
-
-	void ShutdownPoolable() override
-	{
-		// release image buffer (we currently realloc it every time anyway)
-        if (ImageBufferRef)
-        {
-            CFRelease(ImageBufferRef);
-            ImageBufferRef = nullptr;
-        }
-
-		TSharedPtr<IDecoderOutputOwner, ESPMode::ThreadSafe> lockedVideoRenderer = OwningRenderer.Pin();
-		if (lockedVideoRenderer.IsValid())
-		{
-			lockedVideoRenderer->SampleReleasedToPool(GetDuration());
-		}
-	}
-
-	virtual uint32 GetStride() const override
-	{
-		return Stride;
-	}
-
-	virtual CVImageBufferRef GetImageBuffer() const override
-	{
-		return ImageBufferRef;
-	}
-
-private:
-	uint32 Stride;
-
-	CVImageBufferRef ImageBufferRef;
-
-	// We hold a weak reference to the video renderer. During destruction the video renderer could be destroyed while samples are still out there..
-	TWeakPtr<IDecoderOutputOwner, ESPMode::ThreadSafe> OwningRenderer;
-};
-
 /***************************************************************************************************************************************************/
 /***************************************************************************************************************************************************/
 /***************************************************************************************************************************************************/
@@ -1708,13 +1633,6 @@ void FVideoDecoderH264::ProcessOutput(bool bFlush)
 
 } // namespace Electra
 
-// -----------------------------------------------------------------------------------------------------------------------------
-
-FVideoDecoderOutput* FElectraPlayerPlatformVideoDecoderOutputFactory::Create()
-{
-    return new Electra::FElectraPlayerVideoDecoderOutputApple();
-}
-
 
 #endif
 
diff --git a/Engine/Plugins/Media/ElectraPlayer/Source/ElectraPlayerRuntime/Private/Runtime/Decoder/Apple/VideoDecoderH265_Apple.cpp b/Engine/Plugins/Media/ElectraPlayer/Source/ElectraPlayerRuntime/Private/Runtime/Decoder/Apple/VideoDecoderH265_Apple.cpp
--- a/Engine/Plugins/Media/ElectraPlayer/Source/ElectraPlayerRuntime/Private/Runtime/Decoder/Apple/VideoDecoderH265_Apple.cpp
+++ b/Engine/Plugins/Media/ElectraPlayer/Source/ElectraPlayerRuntime/Private/Runtime/Decoder/Apple/VideoDecoderH265_Apple.cpp
@@ -20,6 +20,8 @@
 
 #if ELECTRA_PLATFORM_HAS_H265_DECODER
 
+#include "ElectraVideoDecoder_Apple.h"
+
 #include <VideoToolbox/VideoToolbox.h>
 
 DECLARE_CYCLE_STAT(TEXT("FVideoDecoderH265::Decode()"), STAT_ElectraPlayer_VideoH265Decode, STATGROUP_ElectraPlayer);
@@ -311,82 +313,6 @@ IVideoDecoderH265* IVideoDecoderH265::Create()
 /***************************************************************************************************************************************************/
 /***************************************************************************************************************************************************/
 
-class FElectraPlayerVideoDecoderOutputH265Apple : public FVideoDecoderOutputApple
-{
-public:
-	FElectraPlayerVideoDecoderOutputH265Apple()
-		: ImageBufferRef(nullptr)
-	{
-	}
-
-	~FElectraPlayerVideoDecoderOutputH265Apple()
-	{
-		if (ImageBufferRef)
-		{
-			CFRelease(ImageBufferRef);
-		}
-	}
-
-    void Initialize(CVImageBufferRef InImageBufferRef, FParamDict* InParamDict)
-	{
-		FVideoDecoderOutputApple::Initialize(InParamDict);
-
-		if (ImageBufferRef)
-		{
-			CFRelease(ImageBufferRef);
-		}
-		ImageBufferRef = InImageBufferRef;
-		if (ImageBufferRef)
-		{
-			CFRetain(ImageBufferRef);
-
-			Stride = CVPixelBufferGetBytesPerRow(ImageBufferRef);
-		}
-	}
-
-	void SetOwner(const TSharedPtr<IDecoderOutputOwner, ESPMode::ThreadSafe>& InOwningRenderer) override
-	{
-		OwningRenderer = InOwningRenderer;
-	}
-
-	void ShutdownPoolable() override
-	{
-		// release image buffer (we currently realloc it every time anyway)
-        if (ImageBufferRef)
-        {
-            CFRelease(ImageBufferRef);
-            ImageBufferRef = nullptr;
-        }
-
-		TSharedPtr<IDecoderOutputOwner, ESPMode::ThreadSafe> lockedVideoRenderer = OwningRenderer.Pin();
-		if (lockedVideoRenderer.IsValid())
-		{
-			lockedVideoRenderer->SampleReleasedToPool(GetDuration());
-		}
-	}
-
-	virtual uint32 GetStride() const override
-	{
-		return Stride;
-	}
-
-	virtual CVImageBufferRef GetImageBuffer() const override
-	{
-		return ImageBufferRef;
-	}
-
-private:
-	uint32 Stride;
-
-	CVImageBufferRef ImageBufferRef;
-
-	// We hold a weak reference to the video renderer. During destruction the video renderer could be destroyed while samples are still out there..
-	TWeakPtr<IDecoderOutputOwner, ESPMode::ThreadSafe> OwningRenderer;
-};
-
-/***************************************************************************************************************************************************/
-/***************************************************************************************************************************************************/
-/***************************************************************************************************************************************************/
 
 //-----------------------------------------------------------------------------
 /**
@@ -1663,7 +1589,7 @@ void FVideoDecoderH265::ProcessOutput(bool bFlush)
 				OutputBufferSampleProperties->Set("pts",          FVariantValue(NextImage.SourceInfo->PTS));
 				OutputBufferSampleProperties->Set("duration",     FVariantValue(NextImage.SourceInfo->Duration));
 
-				TSharedPtr<FElectraPlayerVideoDecoderOutputH265Apple, ESPMode::ThreadSafe> DecoderOutput = RenderOutputBuffer->GetBufferProperties().GetValue("texture").GetSharedPointer<FElectraPlayerVideoDecoderOutputH265Apple>();
+				TSharedPtr<FElectraPlayerVideoDecoderOutputApple, ESPMode::ThreadSafe> DecoderOutput = RenderOutputBuffer->GetBufferProperties().GetValue("texture").GetSharedPointer<FElectraPlayerVideoDecoderOutputApple>();
 
 				DecoderOutput->Initialize(ImageBufferRef, OutputBufferSampleProperties);
                 if (ImageBufferRef)
diff --git a/Engine/Plugins/Media/ElectraPlayer/Source/ElectraPlayerRuntime/Private/Runtime/Decoder/Windows/ElectraVideoDecoder_PC.cpp b/Engine/Plugins/Media/ElectraPlayer/Source/ElectraPlayerRuntime/Private/Runtime/Decoder/Windows/ElectraVideoDecoder_PC.cpp
new file mode 100644
--- /dev/null
+++ b/Engine/Plugins/Media/ElectraPlayer/Source/ElectraPlayerRuntime/Private/Runtime/Decoder/Windows/ElectraVideoDecoder_PC.cpp
@@ -0,0 +1,377 @@
+// Copyright Epic Games, Inc. All Rights Reserved.
+
+#include "ElectraVideoDecoder_PC.h"
+#include "Decoder/DX/DecoderErrors_DX.h"
+#include "Renderer/RendererVideo.h"
+
+FElectraPlayerVideoDecoderOutputPC::FElectraPlayerVideoDecoderOutputPC()
+	: OutputType(EOutputType::Unknown)
+	, Stride(0)
+	, SampleDim(0,0)
+{
+}
+
+void FElectraPlayerVideoDecoderOutputPC::SetSWDecodeTargetBufferSize(uint32 InTargetBufferSize)
+{
+	TargetBufferAllocSize = InTargetBufferSize;
+}
+
+void FElectraPlayerVideoDecoderOutputPC::SetOwner(const TSharedPtr<IDecoderOutputOwner, ESPMode::ThreadSafe>& InOwningRenderer)
+{
+	OwningRenderer = InOwningRenderer;
+}
+
+FElectraPlayerVideoDecoderOutputPC::EOutputType FElectraPlayerVideoDecoderOutputPC::GetOutputType() const
+{
+	return OutputType;
+}
+
+TRefCountPtr<IMFSample> FElectraPlayerVideoDecoderOutputPC::GetMFSample() const
+{
+	check(OutputType == EOutputType::SoftwareWin8Plus || OutputType == EOutputType::SoftwareWin7);
+	return MFSample;
+}
+
+const TArray<uint8>& FElectraPlayerVideoDecoderOutputPC::GetBuffer() const
+{
+	check(OutputType == EOutputType::SoftwareWin8Plus || OutputType == EOutputType::SoftwareWin7 || OutputType == EOutputType::HardwareDX9_DX12);
+	return Buffer;
+}
+
+uint32 FElectraPlayerVideoDecoderOutputPC::GetStride() const
+{
+	return Stride;
+}
+
+TRefCountPtr<ID3D11Texture2D> FElectraPlayerVideoDecoderOutputPC::GetTexture() const
+{
+	return Texture;
+}
+
+TRefCountPtr<ID3D11Device> FElectraPlayerVideoDecoderOutputPC::GetDevice() const
+{
+	check(OutputType == EOutputType::SoftwareWin8Plus || OutputType == EOutputType::HardwareWin8Plus);
+	return D3D11Device;
+}
+
+FIntPoint FElectraPlayerVideoDecoderOutputPC::GetDim() const
+{
+	return SampleDim;
+}
+
+void FElectraPlayerVideoDecoderOutputPC::PreInitForDecode(FIntPoint OutputDim, const TFunction<void(int32 /*ApiReturnValue*/, const FString& /*Message*/, uint16 /*Code*/, UEMediaError /*Error*/ )>& PostError)
+{
+	FIntPoint Dim;
+	Dim.X = (OutputDim.X + 15) & ~15;
+	Dim.Y = ((OutputDim.Y + 15) & ~15) * 3/2;	// adjust height to encompass Y and UV planes
+
+	bool bNeedNew = !MFSample.IsValid() || SampleDim != Dim;
+
+	if (bNeedNew)
+	{
+		HRESULT Result;
+		MFSample = nullptr;
+		Texture = nullptr;
+		TRefCountPtr<IMFMediaBuffer> MediaBuffer;
+
+		SampleDim = Dim;
+
+		// Windows 8+ and a valid DX11 rendering device?
+		// (Rendering device will be null if the rendering device is not DX11 (DX12, Vulkan...))
+		if (Electra::IsWindows8Plus() && Electra::FDXDeviceInfo::s_DXDeviceInfo->RenderingDx11Device)
+		{
+			// Software decode into DX11 texture
+
+			// We SW decode right into a DX11 texture (via the media buffer created from it below)
+			/*
+				As we cannot use RHI at this point, we cannot create a texture RHI can readily use,
+				but can just generate a plain D3D one. We cannot easily use a delegate to ask UE to do
+				it for use as we want to keep the "how" of the decoder hidden behind the abstract interfaces
+				so we can also encapsulate non-UE-aware systems easily.
+
+				So: we will later need to copy this once. (TODO: see if the render team can provide a way RHI could use
+				a native texture like this directly)
+			*/
+			D3D11_TEXTURE2D_DESC TextureDesc;
+			TextureDesc.Width = Dim.X;
+			TextureDesc.Height = Dim.Y;
+			TextureDesc.MipLevels = 1;
+			TextureDesc.ArraySize = 1;
+			TextureDesc.Format = DXGI_FORMAT_R8_UINT;
+			TextureDesc.SampleDesc.Count = 1;
+			TextureDesc.SampleDesc.Quality = 0;
+			TextureDesc.Usage = D3D11_USAGE_DEFAULT;
+			TextureDesc.BindFlags = 0;
+			TextureDesc.CPUAccessFlags = 0;
+			TextureDesc.MiscFlags = 0;
+
+			if (FAILED(Result = Electra::FDXDeviceInfo::s_DXDeviceInfo->RenderingDx11Device->CreateTexture2D(&TextureDesc, nullptr, Texture.GetInitReference())))
+			{
+				UE_LOG(LogElectraPlayer, Error, TEXT("CreateTexture2D() failed with 0x%X %s"), Result, *GetComErrorDescription(Result));
+				PostError(Result, "Failed to create software decode output texture", ERRCODE_INTERNAL_COULD_NOT_CREATE_OUTPUTBUFFER, UEMEDIA_ERROR_OK);
+				return;
+			}
+			if (FAILED(Result = MFCreateDXGISurfaceBuffer(__uuidof(ID3D11Texture2D), Texture, 0, false, MediaBuffer.GetInitReference())))
+			{
+				UE_LOG(LogElectraPlayer, Error, TEXT("MFCreateDXGISurfaceBuffer() failed with 0x%X %s"), Result, *GetComErrorDescription(Result));
+				PostError(Result, "Failed to create software decode output texture", ERRCODE_INTERNAL_COULD_NOT_CREATE_OUTPUTBUFFER, UEMEDIA_ERROR_OK);
+				return;
+			}
+
+			OutputType = EOutputType::SoftwareWin8Plus;
+		}
+		else // Software decode into CPU buffer
+		{
+			DWORD AllocSize;
+
+			if (TargetBufferAllocSize.IsSet())
+			{
+				AllocSize = TargetBufferAllocSize.GetValue();
+			}
+			else
+			{
+				AllocSize = Dim.X * Dim.Y * sizeof(uint8);
+			}
+
+			// SW decode results are just delivered in a simple CPU-side buffer. Create the decoder side version of this...
+			if (FAILED(Result = MFCreateMemoryBuffer(AllocSize, MediaBuffer.GetInitReference())))
+			{
+				UE_LOG(LogElectraPlayer, Error, TEXT("MFCreateMemoryBuffer() failed with 0x%X %s"), Result, *GetComErrorDescription(Result));
+				PostError(Result, "Failed to create software decode output buffer", ERRCODE_INTERNAL_COULD_NOT_CREATE_OUTPUTBUFFER, UEMEDIA_ERROR_OK);
+				return;
+			}
+			OutputType = EOutputType::SoftwareWin7;
+		}
+
+		if (SUCCEEDED(Result = MFCreateSample(MFSample.GetInitReference())))
+		{
+			if (FAILED(Result = MFSample->AddBuffer(MediaBuffer)))
+			{
+				UE_LOG(LogElectraPlayer, Error, TEXT("IMFSample::AddBuffer() failed with 0x%X %s"), Result, *GetComErrorDescription(Result));
+				PostError(Result, "Failed to add output buffer to software decoder sample", ERRCODE_INTERNAL_COULD_NOT_ADD_OUTPUT_BUFFER_TO_SAMPLE, UEMEDIA_ERROR_OK);
+				return;
+			}
+		}
+		else
+		{
+			UE_LOG(LogElectraPlayer, Error, TEXT("MFCreateSample() failed with 0x%X %s"), Result, *GetComErrorDescription(Result));
+			PostError(Result, "Failed to create software decode output sample", ERRCODE_INTERNAL_COULD_NOT_CREATE_OUTPUT_SAMPLE, UEMEDIA_ERROR_OK);
+			return;
+		}
+	}
+}
+
+
+void FElectraPlayerVideoDecoderOutputPC::ProcessDecodeOutput(FIntPoint OutputDim, Electra::FParamDict* InParamDict)
+{
+	check(OutputType == EOutputType::SoftwareWin7 || OutputType == EOutputType::SoftwareWin8Plus);
+
+	FVideoDecoderOutput::Initialize(InParamDict);
+
+	// This needs to be a multiple of 16
+	SampleDim.X = (OutputDim.X + 15) & ~15;
+	SampleDim.Y = ((OutputDim.Y + 15) & ~15) * 3 / 2;
+
+	Stride = SampleDim.X * sizeof(uint8);
+
+	if (OutputType == EOutputType::SoftwareWin7) // Win7 & DX12 (SW)
+	{
+		// Retrieve frame data and store it in Buffer for rendering later
+		TRefCountPtr<IMFMediaBuffer> MediaBuffer;
+		if (MFSample->GetBufferByIndex(0, MediaBuffer.GetInitReference()) != S_OK)
+		{
+			return;
+		}
+		DWORD BufferSize = 0;
+		if (MediaBuffer->GetCurrentLength(&BufferSize) != S_OK)
+		{
+			return;
+		}
+		check(SampleDim.X * SampleDim.Y <= (int32)BufferSize);
+
+		// MediaBuffer reports incorrect buffer sizes (too large) for some resolutions: we use our computed values!
+		BufferSize = SampleDim.X * SampleDim.Y;
+
+		uint8* Data = nullptr;
+		Buffer.Reset(BufferSize);
+		if (MediaBuffer->Lock(&Data, NULL, NULL) == S_OK)
+		{
+			Buffer.Append(Data, BufferSize);
+			MediaBuffer->Unlock();
+		}
+
+		// The output IMFSample needs to be released (and recreated for the next frame) for unknown reason.
+		// If not destroyed decoder throws an unknown error later on.
+		MFSample = nullptr;
+	}
+	else
+	{
+		// SW decode into texture need Win8+ (but no additional processing)
+		check(Electra::IsWindows8Plus());
+	}
+}
+
+
+void FElectraPlayerVideoDecoderOutputPC::InitializeWithBuffer(const void* InBuffer, uint32 InSize, uint32 InStride, FIntPoint Dim, Electra::FParamDict* InParamDict)
+{
+	FVideoDecoderOutput::Initialize(InParamDict);
+
+	OutputType = EOutputType::HardwareDX9_DX12;
+
+	Buffer.Reset(InSize);
+	Buffer.Append((uint8*)InBuffer, InSize);
+
+	SampleDim = Dim;
+	Stride = InStride;
+}
+
+
+void FElectraPlayerVideoDecoderOutputPC::InitializeWithSharedTexture(const TRefCountPtr<ID3D11Device>& InD3D11Device, const TRefCountPtr<ID3D11DeviceContext> InDeviceContext, const TRefCountPtr<IMFSample>& InMFSample, const FIntPoint& OutputDim, Electra::FParamDict* InParamDict)
+{
+	FVideoDecoderOutput::Initialize(InParamDict);
+
+	OutputType = EOutputType::HardwareWin8Plus;
+
+	bool bNeedsNew = !Texture.IsValid() || (SampleDim.X != OutputDim.X || SampleDim.Y != OutputDim.Y);
+
+	if (bNeedsNew)
+	{
+		SampleDim = OutputDim;
+
+		// Make a texture we pass on as output
+		D3D11_TEXTURE2D_DESC TextureDesc;
+		TextureDesc.Width = SampleDim.X;
+		TextureDesc.Height = SampleDim.Y;
+		TextureDesc.MipLevels = 1;
+		TextureDesc.ArraySize = 1;
+		TextureDesc.Format = DXGI_FORMAT_NV12;
+		TextureDesc.SampleDesc.Count = 1;
+		TextureDesc.SampleDesc.Quality = 0;
+		TextureDesc.Usage = D3D11_USAGE_DEFAULT;
+		TextureDesc.BindFlags = 0;
+		TextureDesc.CPUAccessFlags = 0;
+		TextureDesc.MiscFlags = D3D11_RESOURCE_MISC_SHARED_KEYEDMUTEX;
+
+		if (InD3D11Device->CreateTexture2D(&TextureDesc, nullptr, Texture.GetInitReference()) != S_OK)
+		{
+			return;
+		}
+
+		D3D11Device = InD3D11Device;
+	}
+
+	// If we got a texture, copy the data from the decoder into it...
+	if (Texture)
+	{
+		// Get output texture from decoder...
+		TRefCountPtr<IMFMediaBuffer> MediaBuffer;
+		if (InMFSample->GetBufferByIndex(0, MediaBuffer.GetInitReference()) != S_OK)
+		{
+			return;
+		}
+		TRefCountPtr<IMFDXGIBuffer> DXGIBuffer;
+		if (MediaBuffer->QueryInterface(__uuidof(IMFDXGIBuffer), (void**)DXGIBuffer.GetInitReference()) != S_OK)
+		{
+			return;
+		}
+		TRefCountPtr<ID3D11Texture2D> DecoderTexture;
+		if (DXGIBuffer->GetResource(IID_PPV_ARGS(DecoderTexture.GetInitReference())) != S_OK)
+		{
+			return;
+		}
+		uint32 ViewIndex = 0;
+		if (DXGIBuffer->GetSubresourceIndex(&ViewIndex) != S_OK)
+		{
+			return;
+		}
+
+		// Source data may be larger than desired output, but crop area will be aligned such that we can always copy from 0,0
+		D3D11_BOX SrcBox;
+		SrcBox.left = 0;
+		SrcBox.top = 0;
+		SrcBox.front = 0;
+		SrcBox.right = OutputDim.X;
+		SrcBox.bottom = OutputDim.Y;
+		SrcBox.back = 1;
+
+		TRefCountPtr<IDXGIKeyedMutex> KeyedMutex;
+		HRESULT res = Texture->QueryInterface(_uuidof(IDXGIKeyedMutex), (void**)&KeyedMutex);
+
+		if (KeyedMutex)
+		{
+			// No wait on acquire since sample is new and key is 0.
+			res = KeyedMutex->AcquireSync(0, 0);
+			if (res == S_OK)
+			{
+				// Copy texture using the decoder DX11 device... (and apply any cropping - see above note)
+				InDeviceContext->CopySubresourceRegion(Texture, 0, 0, 0, 0, DecoderTexture, ViewIndex, &SrcBox);
+
+				// Mark texture as updated with key of 1
+				// Sample will be read in Convert() method of texture sample
+				KeyedMutex->ReleaseSync(1);
+			}
+		}
+
+		// Make sure texture is updated before giving access to the sample in the rendering thread.
+		InDeviceContext->Flush();
+	}
+}
+
+void FElectraPlayerVideoDecoderOutputPC::ShutdownPoolable()
+{
+	TSharedPtr<IDecoderOutputOwner, ESPMode::ThreadSafe> lockedVideoRenderer = OwningRenderer.Pin();
+	if (lockedVideoRenderer.IsValid())
+	{
+		lockedVideoRenderer->SampleReleasedToPool(GetDuration());
+	}
+
+	if (OutputType != EOutputType::HardwareWin8Plus)
+	{
+		return;
+	}
+
+	// Correctly release the keyed mutex when the sample is returned to the pool
+	TRefCountPtr<IDXGIResource> OtherResource(nullptr);
+	if (Texture)
+	{
+		Texture->QueryInterface(__uuidof(IDXGIResource), (void**)&OtherResource);
+	}
+
+	if (OtherResource)
+	{
+		HANDLE SharedHandle = nullptr;
+		if (OtherResource->GetSharedHandle(&SharedHandle) == S_OK)
+		{
+			TRefCountPtr<ID3D11Resource> SharedResource;
+			D3D11Device->OpenSharedResource(SharedHandle, __uuidof(ID3D11Texture2D), (void**)&SharedResource);
+			if (SharedResource)
+			{
+				TRefCountPtr<IDXGIKeyedMutex> KeyedMutex;
+				OtherResource->QueryInterface(_uuidof(IDXGIKeyedMutex), (void**)&KeyedMutex);
+
+				if (KeyedMutex)
+				{
+					// Reset keyed mutex
+					if (KeyedMutex->AcquireSync(1, 0) == S_OK)
+					{
+						// Texture was never read
+						KeyedMutex->ReleaseSync(0);
+					}
+					else if (KeyedMutex->AcquireSync(2, 0) == S_OK)
+					{
+						// Texture was read at least once
+						KeyedMutex->ReleaseSync(0);
+					}
+				}
+			}
+		}
+	}
+
+}
+
+// -----------------------------------------------------------------------------------------------------------------------------
+
+FVideoDecoderOutput* FElectraPlayerPlatformVideoDecoderOutputFactory::Create()
+{
+	return new FElectraPlayerVideoDecoderOutputPC();
+}
diff --git a/Engine/Plugins/Media/ElectraPlayer/Source/ElectraPlayerRuntime/Private/Runtime/Decoder/Windows/ElectraVideoDecoder_PC.h b/Engine/Plugins/Media/ElectraPlayer/Source/ElectraPlayerRuntime/Private/Runtime/Decoder/Windows/ElectraVideoDecoder_PC.h
new file mode 100644
--- /dev/null
+++ b/Engine/Plugins/Media/ElectraPlayer/Source/ElectraPlayerRuntime/Private/Runtime/Decoder/Windows/ElectraVideoDecoder_PC.h
@@ -0,0 +1,88 @@
+// Copyright Epic Games, Inc. All Rights Reserved.
+
+#pragma once
+
+#include "PlayerCore.h"
+#include "ElectraPlayerPrivate.h"
+#include "ElectraPlayerPrivate_Platform.h"
+#include "MediaVideoDecoderOutputPC.h"
+
+/*********************************************************************************************************************/
+/*********************************************************************************************************************/
+/*********************************************************************************************************************/
+
+#include "Windows/AllowWindowsPlatformTypes.h"
+#include "Windows/WindowsHWrapper.h"
+#include "HAL/LowLevelMemTracker.h"
+
+THIRD_PARTY_INCLUDES_START
+#include "mftransform.h"
+#include "mfapi.h"
+#include "mferror.h"
+#include "mfidl.h"
+THIRD_PARTY_INCLUDES_END
+
+#include "Windows/HideWindowsPlatformTypes.h"
+
+/*********************************************************************************************************************/
+/*********************************************************************************************************************/
+/*********************************************************************************************************************/
+
+
+class FElectraPlayerVideoDecoderOutputPC : public FVideoDecoderOutputPC
+{
+public:
+	FElectraPlayerVideoDecoderOutputPC();
+
+	// Hardware decode to buffer (Win7/DX12)
+	void InitializeWithBuffer(const void* InBuffer, uint32 InSize, uint32 InStride, FIntPoint Dim, Electra::FParamDict* InParamDict);
+
+	// Hardware decode to shared DX11 texture (Win8+)
+	void InitializeWithSharedTexture(const TRefCountPtr<ID3D11Device>& InD3D11Device, const TRefCountPtr<ID3D11DeviceContext> InDeviceContext, const TRefCountPtr<IMFSample>& MFSample, const FIntPoint& OutputDim, Electra::FParamDict* InParamDict);
+
+	// Software decode (into texture if DX11 device specified - available only Win8+)
+	void SetSWDecodeTargetBufferSize(uint32 InTargetBufferSize);
+	void PreInitForDecode(FIntPoint OutputDim, const TFunction<void(int32 /*ApiReturnValue*/, const FString& /*Message*/, uint16 /*Code*/, UEMediaError /*Error*/)>& PostError);
+	void ProcessDecodeOutput(FIntPoint OutputDim, Electra::FParamDict* InParamDict);
+
+	void SetOwner(const TSharedPtr<IDecoderOutputOwner, ESPMode::ThreadSafe>& InOwningRenderer) override;
+
+	void ShutdownPoolable() override;
+
+	virtual EOutputType GetOutputType() const override;
+
+	virtual TRefCountPtr<IMFSample> GetMFSample() const override;
+
+	virtual const TArray<uint8>& GetBuffer() const override;
+
+	virtual uint32 GetStride() const override;
+
+	virtual TRefCountPtr<ID3D11Texture2D> GetTexture() const override;
+
+	virtual TRefCountPtr<ID3D11Device> GetDevice() const override;
+
+	virtual FIntPoint GetDim() const override;
+
+private:
+	// Decoder output type
+	EOutputType OutputType;
+
+	// Output texture (with device that created it) for SW decoder output Win8+
+	TRefCountPtr<ID3D11Texture2D> Texture;
+	TRefCountPtr<ID3D11Texture2D> SharedTexture;
+	TRefCountPtr<ID3D11Device> D3D11Device;
+
+	// CPU-side buffer
+	TOptional<uint32> TargetBufferAllocSize;
+	TArray<uint8> Buffer;
+	uint32 Stride;
+
+	// WMF sample (owned by this class if SW decoder is used)
+	TRefCountPtr<IMFSample> MFSample;
+
+	// Dimension of any internally allocated buffer - stored explicitly to cover various special cases for DX
+	FIntPoint SampleDim;
+
+	// We hold a weak reference to the video renderer. During destruction the video renderer could be destroyed while samples are still out there..
+	TWeakPtr<IDecoderOutputOwner, ESPMode::ThreadSafe> OwningRenderer;
+};
diff --git a/Engine/Plugins/Media/ElectraPlayer/Source/ElectraPlayerRuntime/Private/Runtime/Decoder/Windows/VideoDecoderH264_PC.cpp b/Engine/Plugins/Media/ElectraPlayer/Source/ElectraPlayerRuntime/Private/Runtime/Decoder/Windows/VideoDecoderH264_PC.cpp
--- a/Engine/Plugins/Media/ElectraPlayer/Source/ElectraPlayerRuntime/Private/Runtime/Decoder/Windows/VideoDecoderH264_PC.cpp
+++ b/Engine/Plugins/Media/ElectraPlayer/Source/ElectraPlayerRuntime/Private/Runtime/Decoder/Windows/VideoDecoderH264_PC.cpp
@@ -17,26 +17,16 @@
 /*********************************************************************************************************************/
 /*********************************************************************************************************************/
 
-#include "Windows/AllowWindowsPlatformTypes.h"
-#include "Windows/WindowsHWrapper.h"
-#include "HAL/LowLevelMemTracker.h"
+#include "ElectraVideoDecoder_PC.h"
 
-THIRD_PARTY_INCLUDES_START
-#include "mftransform.h"
-#include "mfapi.h"
-#include "mferror.h"
-#include "mfidl.h"
-THIRD_PARTY_INCLUDES_END
-
-#include "Windows/HideWindowsPlatformTypes.h"
-
-#include "Decoder/DX/VideoDecoderH264_DX.h"
 #include "Decoder/DX/MediaFoundationGUIDs.h"
+#include "Decoder/DX/VideoDecoderH264_DX.h"
 
 
 #define ELECTRA_USE_IMF2DBUFFER	1		// set to 1 to use IMF2DBuffer interface rather then plain media buffer when retrieving HW decoded data for CPU use
 
 
+
 const D3DFORMAT DX9_NV12_FORMAT = (D3DFORMAT)MAKEFOURCC('N', 'V', '1', '2');
 
 /*********************************************************************************************************************/
@@ -52,420 +42,6 @@ static TAutoConsoleVariable<int32> CVarElectraPCUseSoftwareDecoding(
 	ECVF_Default);
 #endif
 
-/*********************************************************************************************************************/
-/*********************************************************************************************************************/
-/*********************************************************************************************************************/
-
-class FElectraPlayerVideoDecoderOutputPC : public FVideoDecoderOutputPC
-{
-public:
-	FElectraPlayerVideoDecoderOutputPC()
-		: OutputType(EOutputType::Unknown)
-		, Stride(0)
-		, SampleDim(0,0)
-	{
-	}
-
-	// Hardware decode to buffer (Win7/DX12)
-	void InitializeWithBuffer(const void* InBuffer, uint32 InSize, uint32 InStride, FIntPoint Dim, Electra::FParamDict* InParamDict);
-
-	// Hardware decode to shared DX11 texture (Win8+)
-	void InitializeWithSharedTexture(const TRefCountPtr<ID3D11Device>& InD3D11Device, const TRefCountPtr<ID3D11DeviceContext> InDeviceContext, const TRefCountPtr<IMFSample>& MFSample, const FIntPoint& OutputDim, Electra::FParamDict* InParamDict);
-
-	// Software decode (into texture if DX11 device specified - available only Win8+)
-	void SetSWDecodeTargetBufferSize(uint32 InTargetBufferSize)
-	{
-		TargetBufferAllocSize = InTargetBufferSize;
-	}
-	void PreInitForDecode(FIntPoint OutputDim, const TFunction<void(int32 /*ApiReturnValue*/, const FString& /*Message*/, uint16 /*Code*/, UEMediaError /*Error*/)>& PostError);
-	void ProcessDecodeOutput(FIntPoint OutputDim, Electra::FParamDict* InParamDict);
-
-	void SetOwner(const TSharedPtr<IDecoderOutputOwner, ESPMode::ThreadSafe>& InOwningRenderer) override
-	{
-		OwningRenderer = InOwningRenderer;
-	}
-
-	void ShutdownPoolable() override;
-
-	virtual EOutputType GetOutputType() const override
-	{
-		return OutputType;
-	}
-
-	virtual TRefCountPtr<IMFSample> GetMFSample() const override
-	{
-		check(OutputType == EOutputType::SoftwareWin8Plus || OutputType == EOutputType::SoftwareWin7);
-		return MFSample;
-	}
-
-	virtual const TArray<uint8>& GetBuffer() const override
-	{
-		check(OutputType == EOutputType::SoftwareWin8Plus || OutputType == EOutputType::SoftwareWin7 || OutputType == EOutputType::HardwareDX9_DX12);
-		return Buffer;
-	}
-
-	virtual uint32 GetStride() const override
-	{
-		return Stride;
-	}
-
-	virtual TRefCountPtr<ID3D11Texture2D> GetTexture() const override
-	{
-		return Texture;
-	}
-
-	virtual TRefCountPtr<ID3D11Device> GetDevice() const override
-	{
-		check(OutputType == EOutputType::SoftwareWin8Plus || OutputType == EOutputType::HardwareWin8Plus);
-		return D3D11Device;
-	}
-
-	virtual FIntPoint GetDim() const override
-	{
-		return SampleDim;
-	}
-
-private:
-	// Decoder output type
-	EOutputType OutputType;
-
-	// Output texture (with device that created it) for SW decoder output Win8+
-	TRefCountPtr<ID3D11Texture2D> Texture;
-	TRefCountPtr<ID3D11Texture2D> SharedTexture;
-	TRefCountPtr<ID3D11Device> D3D11Device;
-
-	// CPU-side buffer to 
-	TOptional<uint32> TargetBufferAllocSize;
-	TArray<uint8> Buffer;
-	uint32 Stride;
-
-	// WMF sample (owned by this class if SW decoder is used)
-	TRefCountPtr<IMFSample> MFSample;
-
-	// Dimension of any internally allocated buffer - stored explicitly to cover various special cases for DX
-	FIntPoint SampleDim;
-
-	// We hold a weak reference to the video renderer. During destruction the video renderer could be destroyed while samples are still out there..
-	TWeakPtr<IDecoderOutputOwner, ESPMode::ThreadSafe> OwningRenderer;
-};
-
-
-void FElectraPlayerVideoDecoderOutputPC::PreInitForDecode(FIntPoint OutputDim, const TFunction<void(int32 /*ApiReturnValue*/, const FString& /*Message*/, uint16 /*Code*/, UEMediaError /*Error*/ )>& PostError)
-{
-	FIntPoint Dim;
-	Dim.X = (OutputDim.X + 15) & ~15;
-	Dim.Y = ((OutputDim.Y + 15) & ~15) * 3/2;	// adjust height to encompass Y and UV planes
-
-	bool bNeedNew = !MFSample.IsValid() || SampleDim != Dim;
-
-	if (bNeedNew)
-	{
-		HRESULT Result;
-		MFSample = nullptr;
-		Texture = nullptr;
-		TRefCountPtr<IMFMediaBuffer> MediaBuffer;
-
-		SampleDim = Dim;
-
-		// Windows 8+ and a valid DX11 rendering device?
-		// (Rendering device will be null if the rendering device is not DX11 (DX12, Vulkan...))
-		if (Electra::IsWindows8Plus() && Electra::FDXDeviceInfo::s_DXDeviceInfo->RenderingDx11Device)
-		{
-			// Software decode into DX11 texture
-
-			// We SW decode right into a DX11 texture (via the media buffer created from it below)
-			/*
-				As we cannot use RHI at this point, we cannot create a texture RHI can readily use,
-				but can just generate a plain D3D one. We cannot easily use a delegate to ask UE to do
-				it for use as we want to keep the "how" of the decoder hidden behind the abstract interfaces
-				so we can also encapsulate non-UE-aware systems easily.
-
-				So: we will later need to copy this once. (TODO: see if the render team can provide a way RHI could use
-				a native texture like this directly)
-			*/
-			D3D11_TEXTURE2D_DESC TextureDesc;
-			TextureDesc.Width = Dim.X;
-			TextureDesc.Height = Dim.Y;
-			TextureDesc.MipLevels = 1;
-			TextureDesc.ArraySize = 1;
-			TextureDesc.Format = DXGI_FORMAT_R8_UINT;
-			TextureDesc.SampleDesc.Count = 1;
-			TextureDesc.SampleDesc.Quality = 0;
-			TextureDesc.Usage = D3D11_USAGE_DEFAULT;
-			TextureDesc.BindFlags = 0;
-			TextureDesc.CPUAccessFlags = 0;
-			TextureDesc.MiscFlags = 0;
-
-			if (FAILED(Result = Electra::FDXDeviceInfo::s_DXDeviceInfo->RenderingDx11Device->CreateTexture2D(&TextureDesc, nullptr, Texture.GetInitReference())))
-			{
-				UE_LOG(LogElectraPlayer, Error, TEXT("CreateTexture2D() failed with 0x%X %s"), Result, *GetComErrorDescription(Result));
-				PostError(Result, "Failed to create software decode output texture", ERRCODE_INTERNAL_COULD_NOT_CREATE_OUTPUTBUFFER, UEMEDIA_ERROR_OK);
-				return;
-			}
-			if (FAILED(Result = MFCreateDXGISurfaceBuffer(__uuidof(ID3D11Texture2D), Texture, 0, false, MediaBuffer.GetInitReference())))
-			{
-				UE_LOG(LogElectraPlayer, Error, TEXT("MFCreateDXGISurfaceBuffer() failed with 0x%X %s"), Result, *GetComErrorDescription(Result));
-				PostError(Result, "Failed to create software decode output texture", ERRCODE_INTERNAL_COULD_NOT_CREATE_OUTPUTBUFFER, UEMEDIA_ERROR_OK);
-				return;
-			}
-
-			OutputType = EOutputType::SoftwareWin8Plus;
-		}
-		else // Software decode into CPU buffer
-		{
-			DWORD AllocSize;
-
-			if (TargetBufferAllocSize.IsSet())
-			{
-				AllocSize = TargetBufferAllocSize.GetValue();
-			}
-			else
-			{
-				AllocSize = Dim.X * Dim.Y * sizeof(uint8*) * 2;
-			}
-
-			// SW decode results are just delivered in a simple CPU-side buffer. Create the decoder side version of this...
-			if (FAILED(Result = MFCreateMemoryBuffer(AllocSize, MediaBuffer.GetInitReference())))
-			{
-				UE_LOG(LogElectraPlayer, Error, TEXT("MFCreateMemoryBuffer() failed with 0x%X %s"), Result, *GetComErrorDescription(Result));
-				PostError(Result, "Failed to create software decode output buffer", ERRCODE_INTERNAL_COULD_NOT_CREATE_OUTPUTBUFFER, UEMEDIA_ERROR_OK);
-				return;
-			}
-			OutputType = EOutputType::SoftwareWin7;
-		}
-
-		if (SUCCEEDED(Result = MFCreateSample(MFSample.GetInitReference())))
-		{
-			if (FAILED(Result = MFSample->AddBuffer(MediaBuffer)))
-			{
-				UE_LOG(LogElectraPlayer, Error, TEXT("IMFSample::AddBuffer() failed with 0x%X %s"), Result, *GetComErrorDescription(Result));
-				PostError(Result, "Failed to add output buffer to software decoder sample", ERRCODE_INTERNAL_COULD_NOT_ADD_OUTPUT_BUFFER_TO_SAMPLE, UEMEDIA_ERROR_OK);
-				return;
-			}
-		}
-		else
-		{
-			UE_LOG(LogElectraPlayer, Error, TEXT("MFCreateSample() failed with 0x%X %s"), Result, *GetComErrorDescription(Result));
-			PostError(Result, "Failed to create software decode output sample", ERRCODE_INTERNAL_COULD_NOT_CREATE_OUTPUT_SAMPLE, UEMEDIA_ERROR_OK);
-			return;
-		}
-	}
-}
-
-
-void FElectraPlayerVideoDecoderOutputPC::ProcessDecodeOutput(FIntPoint OutputDim, Electra::FParamDict* InParamDict)
-{
-	check(OutputType == EOutputType::SoftwareWin7 || OutputType == EOutputType::SoftwareWin8Plus);
-
-	FVideoDecoderOutput::Initialize(InParamDict);
-
-	// This needs to be a multiple of 16
-	SampleDim.X = (OutputDim.X + 15) & ~15;
-	SampleDim.Y = ((OutputDim.Y + 15) & ~15) * 3 / 2;
-
-	Stride = SampleDim.X * sizeof(uint8);
-
-	if (OutputType == EOutputType::SoftwareWin7) // Win7 & DX12 (SW)
-	{
-		// Retrieve frame data and store it in Buffer for rendering later
-		TRefCountPtr<IMFMediaBuffer> MediaBuffer;
-		if (MFSample->GetBufferByIndex(0, MediaBuffer.GetInitReference()) != S_OK)
-		{
-			return;
-		}
-		DWORD BufferSize = 0;
-		if (MediaBuffer->GetCurrentLength(&BufferSize) != S_OK)
-		{
-			return;
-		}
-		check(SampleDim.X * SampleDim.Y <= (int32)BufferSize);
-
-		// MediaBuffer reports incorrect buffer sizes (too large) for some resolutions: we use our computed values!
-		BufferSize = SampleDim.X * SampleDim.Y;
-
-		uint8* Data = nullptr;
-		Buffer.Reset(BufferSize);
-		if (MediaBuffer->Lock(&Data, NULL, NULL) == S_OK)
-		{
-			Buffer.Append(Data, BufferSize);
-			MediaBuffer->Unlock();
-		}
-
-		// The output IMFSample needs to be released (and recreated for the next frame) for unknown reason.
-		// If not destroyed decoder throws an unknown error later on.
-		MFSample = nullptr;
-	}
-	else
-	{
-		// SW decode into texture need Win8+ (but no additional processing)
-		check(Electra::IsWindows8Plus());
-	}
-}
-
-
-void FElectraPlayerVideoDecoderOutputPC::InitializeWithBuffer(const void* InBuffer, uint32 InSize, uint32 InStride, FIntPoint Dim, Electra::FParamDict* InParamDict)
-{
-	FVideoDecoderOutput::Initialize(InParamDict);
-
-	OutputType = EOutputType::HardwareDX9_DX12;
-
-	Buffer.Reset(InSize);
-	Buffer.Append((uint8*)InBuffer, InSize);
-
-	SampleDim = Dim;
-	Stride = InStride;
-}
-
-
-void FElectraPlayerVideoDecoderOutputPC::InitializeWithSharedTexture(const TRefCountPtr<ID3D11Device>& InD3D11Device, const TRefCountPtr<ID3D11DeviceContext> InDeviceContext, const TRefCountPtr<IMFSample>& InMFSample, const FIntPoint& OutputDim, Electra::FParamDict* InParamDict)
-{
-	FVideoDecoderOutput::Initialize(InParamDict);
-
-	OutputType = EOutputType::HardwareWin8Plus;
-
-	bool bNeedsNew = !Texture.IsValid() || (SampleDim.X != OutputDim.X || SampleDim.Y != OutputDim.Y);
-
-	if (bNeedsNew)
-	{
-		SampleDim = OutputDim;
-
-		// Make a texture we pass on as output
-		D3D11_TEXTURE2D_DESC TextureDesc;
-		TextureDesc.Width = SampleDim.X;
-		TextureDesc.Height = SampleDim.Y;
-		TextureDesc.MipLevels = 1;
-		TextureDesc.ArraySize = 1;
-		TextureDesc.Format = DXGI_FORMAT_NV12;
-		TextureDesc.SampleDesc.Count = 1;
-		TextureDesc.SampleDesc.Quality = 0;
-		TextureDesc.Usage = D3D11_USAGE_DEFAULT;
-		TextureDesc.BindFlags = 0;
-		TextureDesc.CPUAccessFlags = 0;
-		TextureDesc.MiscFlags = D3D11_RESOURCE_MISC_SHARED_KEYEDMUTEX;
-
-		if (InD3D11Device->CreateTexture2D(&TextureDesc, nullptr, Texture.GetInitReference()) != S_OK)
-		{
-			return;
-		}
-
-		D3D11Device = InD3D11Device;
-	}
-
-	// If we got a texture, copy the data from the decoder into it...
-	if (Texture)
-	{
-		// Get output texture from decoder...
-		TRefCountPtr<IMFMediaBuffer> MediaBuffer;
-		if (InMFSample->GetBufferByIndex(0, MediaBuffer.GetInitReference()) != S_OK)
-		{
-			return;
-		}
-		TRefCountPtr<IMFDXGIBuffer> DXGIBuffer;
-		if (MediaBuffer->QueryInterface(__uuidof(IMFDXGIBuffer), (void**)DXGIBuffer.GetInitReference()) != S_OK)
-		{
-			return;
-		}
-		TRefCountPtr<ID3D11Texture2D> DecoderTexture;
-		if (DXGIBuffer->GetResource(IID_PPV_ARGS(DecoderTexture.GetInitReference())) != S_OK)
-		{
-			return;
-		}
-		uint32 ViewIndex = 0;
-		if (DXGIBuffer->GetSubresourceIndex(&ViewIndex) != S_OK)
-		{
-			return;
-		}
-
-		// Source data may be larger than desired output, but crop area will be aligned such that we can always copy from 0,0
-		D3D11_BOX SrcBox;
-		SrcBox.left = 0;
-		SrcBox.top = 0;
-		SrcBox.front = 0;
-		SrcBox.right = OutputDim.X;
-		SrcBox.bottom = OutputDim.Y;
-		SrcBox.back = 1;
-
-		TRefCountPtr<IDXGIKeyedMutex> KeyedMutex;
-		HRESULT res = Texture->QueryInterface(_uuidof(IDXGIKeyedMutex), (void**)&KeyedMutex);
-
-		if (KeyedMutex)
-		{
-			// No wait on acquire since sample is new and key is 0.
-			res = KeyedMutex->AcquireSync(0, 0);
-			if (res == S_OK)
-			{
-				// Copy texture using the decoder DX11 device... (and apply any cropping - see above note)
-				InDeviceContext->CopySubresourceRegion(Texture, 0, 0, 0, 0, DecoderTexture, ViewIndex, &SrcBox);
-
-				// Mark texture as updated with key of 1
-				// Sample will be read in Convert() method of texture sample
-				KeyedMutex->ReleaseSync(1);
-			}
-		}
-
-		// Make sure texture is updated before giving access to the sample in the rendering thread.
-		InDeviceContext->Flush();
-	}
-}
-
-void FElectraPlayerVideoDecoderOutputPC::ShutdownPoolable()
-{
-	TSharedPtr<IDecoderOutputOwner, ESPMode::ThreadSafe> lockedVideoRenderer = OwningRenderer.Pin();
-	if (lockedVideoRenderer.IsValid())
-	{
-		lockedVideoRenderer->SampleReleasedToPool(GetDuration());
-	}
-
-	if (OutputType != EOutputType::HardwareWin8Plus)
-	{
-		return;
-	}
-
-	// Correctly release the keyed mutex when the sample is returned to the pool
-	TRefCountPtr<IDXGIResource> OtherResource(nullptr);
-	if (Texture)
-	{
-		Texture->QueryInterface(__uuidof(IDXGIResource), (void**)&OtherResource);
-	}
-
-	if (OtherResource)
-	{
-		HANDLE SharedHandle = nullptr;
-		if (OtherResource->GetSharedHandle(&SharedHandle) == S_OK)
-		{
-			TRefCountPtr<ID3D11Resource> SharedResource;
-			D3D11Device->OpenSharedResource(SharedHandle, __uuidof(ID3D11Texture2D), (void**)&SharedResource);
-			if (SharedResource)
-			{
-				TRefCountPtr<IDXGIKeyedMutex> KeyedMutex;
-				OtherResource->QueryInterface(_uuidof(IDXGIKeyedMutex), (void**)&KeyedMutex);
-
-				if (KeyedMutex)
-				{
-					// Reset keyed mutex
-					if (KeyedMutex->AcquireSync(1, 0) == S_OK)
-					{
-						// Texture was never read
-						KeyedMutex->ReleaseSync(0);
-					}
-					else if (KeyedMutex->AcquireSync(2, 0) == S_OK)
-					{
-						// Texture was read at least once
-						KeyedMutex->ReleaseSync(0);
-					}
-				}
-			}
-		}
-	}
-
-}
-
-// -----------------------------------------------------------------------------------------------------------------------------
-
-FVideoDecoderOutput* FElectraPlayerPlatformVideoDecoderOutputFactory::Create()
-{
-	return new FElectraPlayerVideoDecoderOutputPC();
-}
 
 
 namespace Electra
diff --git a/Engine/Plugins/Media/ElectraPlayer/Source/ElectraPlayerRuntime/Private/Runtime/Decoder/Windows/VideoDecoderH265_PC.cpp b/Engine/Plugins/Media/ElectraPlayer/Source/ElectraPlayerRuntime/Private/Runtime/Decoder/Windows/VideoDecoderH265_PC.cpp
--- a/Engine/Plugins/Media/ElectraPlayer/Source/ElectraPlayerRuntime/Private/Runtime/Decoder/Windows/VideoDecoderH265_PC.cpp
+++ b/Engine/Plugins/Media/ElectraPlayer/Source/ElectraPlayerRuntime/Private/Runtime/Decoder/Windows/VideoDecoderH265_PC.cpp
@@ -19,420 +19,15 @@
 /*********************************************************************************************************************/
 /*********************************************************************************************************************/
 
-#include "Windows/AllowWindowsPlatformTypes.h"
-#include "Windows/WindowsHWrapper.h"
-#include "HAL/LowLevelMemTracker.h"
-
-THIRD_PARTY_INCLUDES_START
-#include "mftransform.h"
-#include "mfapi.h"
-#include "mferror.h"
-#include "mfidl.h"
-THIRD_PARTY_INCLUDES_END
+#include "ElectraVideoDecoder_PC.h"
 
 #include "Decoder/DX/MediaFoundationGUIDs.h"
-
-#include "Windows/HideWindowsPlatformTypes.h"
-
 #include "Decoder/DX/VideoDecoderH265_DX.h"
 #include "Decoder/Windows/PlatformVideoDecoderH265.h"
 
 
 #define ELECTRA_USE_IMF2DBUFFER	1		// set to 1 to use IMF2DBuffer interface rather then plain media buffer when retrieving HW decoded data for CPU use
 
-/*********************************************************************************************************************/
-/*********************************************************************************************************************/
-/*********************************************************************************************************************/
-
-class FElectraPlayerVideoDecoderOutputH265PC : public FVideoDecoderOutputPC
-{
-public:
-	FElectraPlayerVideoDecoderOutputH265PC()
-		: OutputType(EOutputType::Unknown)
-		, Stride(0)
-		, SampleDim(0,0)
-	{
-	}
-
-	// Hardware decode to buffer (DX12)
-	void InitializeWithBuffer(const void* InBuffer, uint32 InSize, uint32 InStride, FIntPoint Dim, Electra::FParamDict* InParamDict);
-
-	// Hardware decode to shared DX11 texture (Win8+)
-	void InitializeWithSharedTexture(const TRefCountPtr<ID3D11Device>& InD3D11Device, const TRefCountPtr<ID3D11DeviceContext> InDeviceContext, const TRefCountPtr<IMFSample>& MFSample, const FIntPoint& OutputDim, Electra::FParamDict* InParamDict);
-
-	// Software decode (into texture if DX11 device specified - available only Win8+)
-	void PreInitForDecode(FIntPoint OutputDim, const TFunction<void(int32 /*ApiReturnValue*/, const FString& /*Message*/, uint16 /*Code*/, UEMediaError /*Error*/)>& PostError);
-	void ProcessDecodeOutput(FIntPoint OutputDim, Electra::FParamDict* InParamDict);
-
-	void SetOwner(const TSharedPtr<IDecoderOutputOwner, ESPMode::ThreadSafe>& InOwningRenderer) override
-	{
-		OwningRenderer = InOwningRenderer;
-	}
-
-	void ShutdownPoolable() override;
-
-	virtual EOutputType GetOutputType() const override
-	{
-		return OutputType;
-	}
-
-	virtual TRefCountPtr<IMFSample> GetMFSample() const override
-	{
-		check(OutputType == EOutputType::SoftwareWin8Plus || OutputType == EOutputType::SoftwareWin7);
-		return MFSample;
-	}
-
-	virtual const TArray<uint8>& GetBuffer() const override
-	{
-		check(OutputType == EOutputType::SoftwareWin8Plus || OutputType == EOutputType::SoftwareWin7 || OutputType == EOutputType::HardwareDX9_DX12);
-		return Buffer;
-	}
-
-	virtual uint32 GetStride() const override
-	{
-		return Stride;
-	}
-
-	virtual TRefCountPtr<ID3D11Texture2D> GetTexture() const override
-	{
-		return Texture;
-	}
-
-	virtual TRefCountPtr<ID3D11Device> GetDevice() const override
-	{
-		check(OutputType == EOutputType::SoftwareWin8Plus || OutputType == EOutputType::HardwareWin8Plus);
-		return D3D11Device;
-	}
-
-	virtual FIntPoint GetDim() const override
-	{
-		return SampleDim;
-	}
-
-private:
-	// Decoder output type
-	EOutputType OutputType;
-
-	// Output texture (with device that created it) for SW decoder output Win8+
-	TRefCountPtr<ID3D11Texture2D> Texture;
-	TRefCountPtr<ID3D11Texture2D> SharedTexture;
-	TRefCountPtr<ID3D11Device> D3D11Device;
-
-	// CPU-side buffer to 
-	TArray<uint8> Buffer;
-	uint32 Stride;
-
-	// WMF sample (owned by this class if SW decoder is used)
-	TRefCountPtr<IMFSample> MFSample;
-
-	// Dimension of any internally allocated buffer - stored explicitly to cover various special cases for DX
-	FIntPoint SampleDim;
-
-	// We hold a weak reference to the video renderer. During destruction the video renderer could be destroyed while samples are still out there..
-	TWeakPtr<IDecoderOutputOwner, ESPMode::ThreadSafe> OwningRenderer;
-};
-
-
-void FElectraPlayerVideoDecoderOutputH265PC::PreInitForDecode(FIntPoint OutputDim, const TFunction<void(int32 /*ApiReturnValue*/, const FString& /*Message*/, uint16 /*Code*/, UEMediaError /*Error*/ )>& PostError)
-{
-	FIntPoint Dim;
-	Dim.X = OutputDim.X;
-	Dim.Y = OutputDim.Y * 3 / 2;	// adjust height to encompass Y and UV planes
-
-	bool bNeedNew = !MFSample.IsValid() || SampleDim != Dim;
-
-	if (bNeedNew)
-	{
-		HRESULT Result;
-		MFSample = nullptr;
-		Texture = nullptr;
-		TRefCountPtr<IMFMediaBuffer> MediaBuffer;
-
-		SampleDim = Dim;
-
-		// Windows 8+ and a valid DX11 rendering device?
-		// (Rendering device will be null if the rendering device is not DX11 (DX12, Vulkan...))
-		if (Electra::IsWindows8Plus() && Electra::FDXDeviceInfo::s_DXDeviceInfo->RenderingDx11Device)
-		{
-			// Software decode into DX11 texture
-
-			// We SW decode right into a DX11 texture (via the media buffer created from it below)
-			/*
-				As we cannot use RHI at this point, we cannot create a texture RHI can readily use,
-				but can just generate a plain D3D one. We cannot easily use a delegate to ask UE to do
-				it for use as we want to keep the "how" of the decoder hidden behind the abstract interfaces
-				so we can also encapsulate non-UE-aware systems easily.
-
-				So: we will later need to copy this once. (TODO: see if the render team can provide a way RHI could use
-				a native texture like this directly)
-			*/
-			D3D11_TEXTURE2D_DESC TextureDesc;
-			TextureDesc.Width = Dim.X;
-			TextureDesc.Height = Dim.Y;
-			TextureDesc.MipLevels = 1;
-			TextureDesc.ArraySize = 1;
-			TextureDesc.Format = DXGI_FORMAT_R8_UINT;
-			TextureDesc.SampleDesc.Count = 1;
-			TextureDesc.SampleDesc.Quality = 0;
-			TextureDesc.Usage = D3D11_USAGE_DEFAULT;
-			TextureDesc.BindFlags = 0;
-			TextureDesc.CPUAccessFlags = 0;
-			TextureDesc.MiscFlags = 0;
-
-			if (FAILED(Result = Electra::FDXDeviceInfo::s_DXDeviceInfo->RenderingDx11Device->CreateTexture2D(&TextureDesc, nullptr, Texture.GetInitReference())))
-			{
-				UE_LOG(LogElectraPlayer, Error, TEXT("CreateTexture2D() failed with 0x%X %s"), Result, *GetComErrorDescription(Result));
-				PostError(Result, "Failed to create software decode output texture", ERRCODE_INTERNAL_COULD_NOT_CREATE_OUTPUTBUFFER, UEMEDIA_ERROR_OK);
-				return;
-			}
-			if (FAILED(Result = MFCreateDXGISurfaceBuffer(__uuidof(ID3D11Texture2D), Texture, 0, false, MediaBuffer.GetInitReference())))
-			{
-				UE_LOG(LogElectraPlayer, Error, TEXT("MFCreateDXGISurfaceBuffer() failed with 0x%X %s"), Result, *GetComErrorDescription(Result));
-				PostError(Result, "Failed to create software decode output texture", ERRCODE_INTERNAL_COULD_NOT_CREATE_OUTPUTBUFFER, UEMEDIA_ERROR_OK);
-				return;
-			}
-
-			OutputType = EOutputType::SoftwareWin8Plus;
-		}
-		else // Software decode into CPU buffer
-		{
-			// SW decode results are just delivered in a simple CPU-side buffer. Create the decoder side version of this...
-			if (FAILED(Result = MFCreateMemoryBuffer(Dim.X * Dim.Y * sizeof(uint8), MediaBuffer.GetInitReference())))
-			{
-				UE_LOG(LogElectraPlayer, Error, TEXT("MFCreateMemoryBuffer() failed with 0x%X %s"), Result, *GetComErrorDescription(Result));
-				PostError(Result, "Failed to create software decode output buffer", ERRCODE_INTERNAL_COULD_NOT_CREATE_OUTPUTBUFFER, UEMEDIA_ERROR_OK);
-				return;
-			}
-			OutputType = EOutputType::SoftwareWin7;
-		}
-
-		if (SUCCEEDED(Result = MFCreateSample(MFSample.GetInitReference())))
-		{
-			if (FAILED(Result = MFSample->AddBuffer(MediaBuffer)))
-			{
-				UE_LOG(LogElectraPlayer, Error, TEXT("IMFSample::AddBuffer() failed with 0x%X %s"), Result, *GetComErrorDescription(Result));
-				PostError(Result, "Failed to add output buffer to software decoder sample", ERRCODE_INTERNAL_COULD_NOT_ADD_OUTPUT_BUFFER_TO_SAMPLE, UEMEDIA_ERROR_OK);
-				return;
-			}
-		}
-		else
-		{
-			UE_LOG(LogElectraPlayer, Error, TEXT("MFCreateSample() failed with 0x%X %s"), Result, *GetComErrorDescription(Result));
-			PostError(Result, "Failed to create software decode output sample", ERRCODE_INTERNAL_COULD_NOT_CREATE_OUTPUT_SAMPLE, UEMEDIA_ERROR_OK);
-			return;
-		}
-	}
-}
-
-
-void FElectraPlayerVideoDecoderOutputH265PC::ProcessDecodeOutput(FIntPoint OutputDim, Electra::FParamDict* InParamDict)
-{
-	check(OutputType == EOutputType::SoftwareWin7 || OutputType == EOutputType::SoftwareWin8Plus);
-
-	FVideoDecoderOutput::Initialize(InParamDict);
-
-	SampleDim.X = OutputDim.X;
-	SampleDim.Y = OutputDim.Y * 3 / 2;
-
-	Stride = SampleDim.X * sizeof(uint8);
-
-	if (OutputType == EOutputType::SoftwareWin7) // Win7 & DX12 (SW)
-	{
-		// Retrieve frame data and store it in Buffer for rendering later
-		TRefCountPtr<IMFMediaBuffer> MediaBuffer;
-		if (MFSample->GetBufferByIndex(0, MediaBuffer.GetInitReference()) != S_OK)
-		{
-			return;
-		}
-		DWORD BufferSize = 0;
-		if (MediaBuffer->GetCurrentLength(&BufferSize) != S_OK)
-		{
-			return;
-		}
-		check(SampleDim.X * SampleDim.Y <= (int32)BufferSize);
-
-		// MediaBuffer reports incorrect buffer sizes (too large) for some resolutions: we use our computed values!
-		BufferSize = SampleDim.X * SampleDim.Y;
-
-		uint8* Data = nullptr;
-		Buffer.Reset(BufferSize);
-		if (MediaBuffer->Lock(&Data, NULL, NULL) == S_OK)
-		{
-			Buffer.Append(Data, BufferSize);
-			MediaBuffer->Unlock();
-		}
-
-		// The output IMFSample needs to be released (and recreated for the next frame) for unknown reason.
-		// If not destroyed decoder throws an unknown error later on.
-		MFSample = nullptr;
-	}
-	else
-	{
-		// SW decode into texture need Win8+ (but no additional processing)
-		check(Electra::IsWindows8Plus());
-	}
-}
-
-
-void FElectraPlayerVideoDecoderOutputH265PC::InitializeWithBuffer(const void* InBuffer, uint32 InSize, uint32 InStride, FIntPoint Dim, Electra::FParamDict* InParamDict)
-{
-	FVideoDecoderOutput::Initialize(InParamDict);
-
-	OutputType = EOutputType::HardwareDX9_DX12;
-
-	Buffer.Reset(InSize);
-	Buffer.Append((uint8*)InBuffer, InSize);
-
-	SampleDim = Dim;
-	Stride = InStride;
-}
-
-
-void FElectraPlayerVideoDecoderOutputH265PC::InitializeWithSharedTexture(const TRefCountPtr<ID3D11Device>& InD3D11Device, const TRefCountPtr<ID3D11DeviceContext> InDeviceContext, const TRefCountPtr<IMFSample>& InMFSample, const FIntPoint& OutputDim, Electra::FParamDict* InParamDict)
-{
-	FVideoDecoderOutput::Initialize(InParamDict);
-
-	OutputType = EOutputType::HardwareWin8Plus;
-
-	bool bNeedsNew = !Texture.IsValid() || (SampleDim.X != OutputDim.X || SampleDim.Y != OutputDim.Y);
-
-	if (bNeedsNew)
-	{
-		SampleDim = OutputDim;
-
-		// Make a texture we pass on as output
-		D3D11_TEXTURE2D_DESC TextureDesc;
-		TextureDesc.Width = SampleDim.X;
-		TextureDesc.Height = SampleDim.Y;
-		TextureDesc.MipLevels = 1;
-		TextureDesc.ArraySize = 1;
-		TextureDesc.Format = DXGI_FORMAT_NV12;
-		TextureDesc.SampleDesc.Count = 1;
-		TextureDesc.SampleDesc.Quality = 0;
-		TextureDesc.Usage = D3D11_USAGE_DEFAULT;
-		TextureDesc.BindFlags = 0;
-		TextureDesc.CPUAccessFlags = 0;
-		TextureDesc.MiscFlags = D3D11_RESOURCE_MISC_SHARED_KEYEDMUTEX;
-
-		if (InD3D11Device->CreateTexture2D(&TextureDesc, nullptr, Texture.GetInitReference()) != S_OK)
-		{
-			return;
-		}
-
-		D3D11Device = InD3D11Device;
-	}
-
-	// If we got a texture, copy the data from the decoder into it...
-	if (Texture)
-	{
-		// Get output texture from decoder...
-		TRefCountPtr<IMFMediaBuffer> MediaBuffer;
-		if (InMFSample->GetBufferByIndex(0, MediaBuffer.GetInitReference()) != S_OK)
-		{
-			return;
-		}
-		TRefCountPtr<IMFDXGIBuffer> DXGIBuffer;
-		if (MediaBuffer->QueryInterface(__uuidof(IMFDXGIBuffer), (void**)DXGIBuffer.GetInitReference()) != S_OK)
-		{
-			return;
-		}
-		TRefCountPtr<ID3D11Texture2D> DecoderTexture;
-		if (DXGIBuffer->GetResource(IID_PPV_ARGS(DecoderTexture.GetInitReference())) != S_OK)
-		{
-			return;
-		}
-		uint32 ViewIndex = 0;
-		if (DXGIBuffer->GetSubresourceIndex(&ViewIndex) != S_OK)
-		{
-			return;
-		}
-
-		// Source data may be larger than desired output, but crop area will be aligned such that we can always copy from 0,0
-		D3D11_BOX SrcBox;
-		SrcBox.left = 0;
-		SrcBox.top = 0;
-		SrcBox.front = 0;
-		SrcBox.right = OutputDim.X;
-		SrcBox.bottom = OutputDim.Y;
-		SrcBox.back = 1;
-
-		TRefCountPtr<IDXGIKeyedMutex> KeyedMutex;
-		HRESULT res = Texture->QueryInterface(_uuidof(IDXGIKeyedMutex), (void**)&KeyedMutex);
-
-		if (KeyedMutex)
-		{
-			// No wait on acquire since sample is new and key is 0.
-			res = KeyedMutex->AcquireSync(0, 0);
-			if (res == S_OK)
-			{
-				// Copy texture using the decoder DX11 device... (and apply any cropping - see above note)
-				InDeviceContext->CopySubresourceRegion(Texture, 0, 0, 0, 0, DecoderTexture, ViewIndex, &SrcBox);
-
-				// Mark texture as updated with key of 1
-				// Sample will be read in Convert() method of texture sample
-				KeyedMutex->ReleaseSync(1);
-			}
-		}
-
-		// Make sure texture is updated before giving access to the sample in the rendering thread.
-		InDeviceContext->Flush();
-	}
-}
-
-void FElectraPlayerVideoDecoderOutputH265PC::ShutdownPoolable()
-{
-	TSharedPtr<IDecoderOutputOwner, ESPMode::ThreadSafe> lockedVideoRenderer = OwningRenderer.Pin();
-	if (lockedVideoRenderer.IsValid())
-	{
-		lockedVideoRenderer->SampleReleasedToPool(GetDuration());
-	}
-
-	if (OutputType != EOutputType::HardwareWin8Plus)
-	{
-		return;
-	}
-
-	// Correctly release the keyed mutex when the sample is returned to the pool
-	TRefCountPtr<IDXGIResource> OtherResource(nullptr);
-	if (Texture)
-	{
-		Texture->QueryInterface(__uuidof(IDXGIResource), (void**)&OtherResource);
-	}
-
-	if (OtherResource)
-	{
-		HANDLE SharedHandle = nullptr;
-		if (OtherResource->GetSharedHandle(&SharedHandle) == S_OK)
-		{
-			TRefCountPtr<ID3D11Resource> SharedResource;
-			D3D11Device->OpenSharedResource(SharedHandle, __uuidof(ID3D11Texture2D), (void**)&SharedResource);
-			if (SharedResource)
-			{
-				TRefCountPtr<IDXGIKeyedMutex> KeyedMutex;
-				OtherResource->QueryInterface(_uuidof(IDXGIKeyedMutex), (void**)&KeyedMutex);
-
-				if (KeyedMutex)
-				{
-					// Reset keyed mutex
-					if (KeyedMutex->AcquireSync(1, 0) == S_OK)
-					{
-						// Texture was never read
-						KeyedMutex->ReleaseSync(0);
-					}
-					else if (KeyedMutex->AcquireSync(2, 0) == S_OK)
-					{
-						// Texture was read at least once
-						KeyedMutex->ReleaseSync(0);
-					}
-				}
-			}
-		}
-	}
-
-}
-
-// -----------------------------------------------------------------------------------------------------------------------------
-
 
 namespace Electra
 {
@@ -590,7 +185,7 @@ bool FVideoDecoderH265_PC::CreateDecoderOutputBuffer()
 
 void FVideoDecoderH265_PC::PreInitDecodeOutputForSW(const FIntPoint& Dim)
 {
-	TSharedPtr<FElectraPlayerVideoDecoderOutputH265PC, ESPMode::ThreadSafe> DecoderOutput = CurrentRenderOutputBuffer->GetBufferProperties().GetValue("texture").GetSharedPointer<FElectraPlayerVideoDecoderOutputH265PC>();
+	TSharedPtr<FElectraPlayerVideoDecoderOutputPC, ESPMode::ThreadSafe> DecoderOutput = CurrentRenderOutputBuffer->GetBufferProperties().GetValue("texture").GetSharedPointer<FElectraPlayerVideoDecoderOutputPC>();
 	check(DecoderOutput);
 	DecoderOutput->PreInitForDecode(Dim, [this](int32 ApiReturnValue, const FString& Message, uint16 Code, UEMediaError Error) { PostError(ApiReturnValue, Message, Code, Error); });
 }
@@ -611,7 +206,7 @@ bool FVideoDecoderH265_PC::SetupDecodeOutputData(const FIntPoint & ImageDim, con
 bool FVideoDecoderH265_PC::CopyTexture(const TRefCountPtr<IMFSample>& Sample, Electra::FParamDict* ParamDict, FIntPoint OutputDim)
 {
 #if !UE_SERVER
-	TSharedPtr<FElectraPlayerVideoDecoderOutputH265PC, ESPMode::ThreadSafe> DecoderOutput = CurrentRenderOutputBuffer->GetBufferProperties().GetValue("texture").GetSharedPointer<FElectraPlayerVideoDecoderOutputH265PC>();
+	TSharedPtr<FElectraPlayerVideoDecoderOutputPC, ESPMode::ThreadSafe> DecoderOutput = CurrentRenderOutputBuffer->GetBufferProperties().GetValue("texture").GetSharedPointer<FElectraPlayerVideoDecoderOutputPC>();
 	check(DecoderOutput);
 
 	DWORD BuffersNum = 0;
