ElectraPlayer: Added an IMediaOptions holder around the POD pointer to check if it is still valid after garbage collection.

#rnx

diff --git a/Engine/Plugins/Media/ElectraPlayer/Source/ElectraPlayerPlugin/Private/ElectraPlayerPlugin.cpp b/Engine/Plugins/Media/ElectraPlayer/Source/ElectraPlayerPlugin/Private/ElectraPlayerPlugin.cpp
--- a/Engine/Plugins/Media/ElectraPlayer/Source/ElectraPlayerPlugin/Private/ElectraPlayerPlugin.cpp
+++ b/Engine/Plugins/Media/ElectraPlayer/Source/ElectraPlayerPlugin/Private/ElectraPlayerPlugin.cpp
@@ -85,7 +85,7 @@ FElectraPlayerPlugin::~FElectraPlayerPlugin()
 {
 	CallbackPointerLock.Lock();
 	EventSink = nullptr;
-	OptionInterface = nullptr;
+	OptionInterface.Reset();
 	CallbackPointerLock.Unlock();
 	if (Player.IsValid())
 	{
@@ -216,69 +216,76 @@ Electra::FVariantValue FElectraPlayerPlugin::FPlayerAdapterDelegate::QueryOption
 	TSharedPtr<FElectraPlayerPlugin, ESPMode::ThreadSafe> PinnedHost = Host.Pin();
 	if (PinnedHost.IsValid())
 	{
-		FScopeLock lock(&PinnedHost->CallbackPointerLock);
-		if (PinnedHost->OptionInterface)
+		PinnedHost->CallbackPointerLock.Lock();
+		TSharedPtr<IElectraSafeMediaOptionInterface, ESPMode::ThreadSafe> SafeOptionInterface = PinnedHost->OptionInterface.Pin();
+		PinnedHost->CallbackPointerLock.Unlock();
+		if (SafeOptionInterface.IsValid())
 		{
-			switch (Type)
+			IElectraSafeMediaOptionInterface::FScopedLock SafeLock(SafeOptionInterface);
+			IMediaOptions *SafeOptions = SafeOptionInterface->GetMediaOptionInterface();
+			if (SafeOptions)
 			{
-				case EOptionType::MaxVerticalStreamResolution:
-				{
-					static const FName MaxResolutionOptionKey = TEXT("MaxResolutionForMediaStreaming");
-					return FVariantValue((int64)PinnedHost->OptionInterface->GetMediaOption(MaxResolutionOptionKey, (int64)0));
-				}
-
-				case EOptionType::MaxBandwidthForStreaming:
+				switch (Type)
 				{
-					static const FName MaxBandwidthOptionKey = TEXT("ElectraMaxStreamingBandwidth");
-					return FVariantValue((int64)PinnedHost->OptionInterface->GetMediaOption(MaxBandwidthOptionKey, (int64)0));
-				}
+					case EOptionType::MaxVerticalStreamResolution:
+					{
+						static const FName MaxResolutionOptionKey = TEXT("MaxResolutionForMediaStreaming");
+						return FVariantValue((int64)SafeOptions->GetMediaOption(MaxResolutionOptionKey, (int64)0));
+					}
 
-				case EOptionType::PlayListData:
-				{
-					static const FName PlaylistOptionKey = TEXT("ElectraGetPlaylistData");
-					if (PinnedHost->OptionInterface->HasMediaOption(PlaylistOptionKey))
+					case EOptionType::MaxBandwidthForStreaming:
 					{
-						check(Param.IsType(FVariantValue::EDataType::TypeFString));
-						return FVariantValue(PinnedHost->OptionInterface->GetMediaOption(PlaylistOptionKey, Param.GetFString()));
+						static const FName MaxBandwidthOptionKey = TEXT("ElectraMaxStreamingBandwidth");
+						return FVariantValue((int64)SafeOptions->GetMediaOption(MaxBandwidthOptionKey, (int64)0));
 					}
-					break;
-				}
 
-				case EOptionType::LicenseKeyData:
-				{
-					static const FName LicenseKeyDataOptionKey = TEXT("ElectraGetLicenseKeyData");
-					if (PinnedHost->OptionInterface->HasMediaOption(LicenseKeyDataOptionKey))
+					case EOptionType::PlayListData:
 					{
-						check(Param.IsType(FVariantValue::EDataType::TypeFString));
-						return FVariantValue(PinnedHost->OptionInterface->GetMediaOption(LicenseKeyDataOptionKey, Param.GetFString()));
+						static const FName PlaylistOptionKey = TEXT("ElectraGetPlaylistData");
+						if (SafeOptions->HasMediaOption(PlaylistOptionKey))
+						{
+							check(Param.IsType(FVariantValue::EDataType::TypeFString));
+							return FVariantValue(SafeOptions->GetMediaOption(PlaylistOptionKey, Param.GetFString()));
+						}
+						break;
 					}
-					break;
-				}
 
-				case EOptionType::PlaystartPosFromSeekPositions:
-				{
-					static const FName PlaystartOptionKey = TEXT("ElectraGetPlaystartPosFromSeekPositions");
-					if (PinnedHost->OptionInterface->HasMediaOption(PlaystartOptionKey))
+					case EOptionType::LicenseKeyData:
 					{
-						check(Param.IsType(FVariantValue::EDataType::TypeSharedPointer));
+						static const FName LicenseKeyDataOptionKey = TEXT("ElectraGetLicenseKeyData");
+						if (SafeOptions->HasMediaOption(LicenseKeyDataOptionKey))
+						{
+							check(Param.IsType(FVariantValue::EDataType::TypeFString));
+							return FVariantValue(SafeOptions->GetMediaOption(LicenseKeyDataOptionKey, Param.GetFString()));
+						}
+						break;
+					}
 
-						TSharedPtr<TArray<FTimespan>, ESPMode::ThreadSafe> PosArray = Param.GetSharedPointer<TArray<FTimespan>>();
-						if (PosArray.IsValid())
+					case EOptionType::PlaystartPosFromSeekPositions:
+					{
+						static const FName PlaystartOptionKey = TEXT("ElectraGetPlaystartPosFromSeekPositions");
+						if (SafeOptions->HasMediaOption(PlaystartOptionKey))
 						{
-							TSharedPtr<FElectraSeekablePositions, ESPMode::ThreadSafe> Res = StaticCastSharedPtr<FElectraSeekablePositions, IMediaOptions::FDataContainer, ESPMode::ThreadSafe>(PinnedHost->OptionInterface->GetMediaOption(PlaystartOptionKey, MakeShared<FElectraSeekablePositions, ESPMode::ThreadSafe>(*PosArray)));
-							if (Res.IsValid() && Res->Data.Num())
+							check(Param.IsType(FVariantValue::EDataType::TypeSharedPointer));
+
+							TSharedPtr<TArray<FTimespan>, ESPMode::ThreadSafe> PosArray = Param.GetSharedPointer<TArray<FTimespan>>();
+							if (PosArray.IsValid())
 							{
-								return FVariantValue(int64(Res->Data[0].GetTicks())); // return HNS
+								TSharedPtr<FElectraSeekablePositions, ESPMode::ThreadSafe> Res = StaticCastSharedPtr<FElectraSeekablePositions, IMediaOptions::FDataContainer, ESPMode::ThreadSafe>(SafeOptions->GetMediaOption(PlaystartOptionKey, MakeShared<FElectraSeekablePositions, ESPMode::ThreadSafe>(*PosArray)));
+								if (Res.IsValid() && Res->Data.Num())
+								{
+									return FVariantValue(int64(Res->Data[0].GetTicks())); // return HNS
+								}
 							}
+							return FVariantValue();
 						}
-						return FVariantValue();
+						break;
 					}
-					break;
-				}
 
-				default:
-				{
-					break;
+					default:
+					{
+						break;
+					}
 				}
 			}
 		}
@@ -471,9 +478,9 @@ bool FElectraPlayerPlugin::Open(const FString& Url, const IMediaOptions* Options
 
 bool FElectraPlayerPlugin::Open(const FString& Url, const IMediaOptions* Options, const FMediaPlayerOptions* InPlayerOptions)
 {
-	// Remember the option interface to poll for changes during playback.
+	// Get the safe option interface to poll for changes during playback.
 	CallbackPointerLock.Lock();
-	OptionInterface = Options;
+	OptionInterface = StaticCastSharedPtr<IElectraSafeMediaOptionInterface>(Options->GetMediaOption(TEXT("GetSafeMediaOptions"), TSharedPtr<IElectraSafeMediaOptionInterface, ESPMode::ThreadSafe>()));
 	CallbackPointerLock.Unlock();
 	UE_LOG(LogElectraPlayerPlugin, Log, TEXT("[%p] IMediaPlayer::Open: Options@%p"), this, Options);
 
@@ -616,7 +623,7 @@ bool FElectraPlayerPlugin::Open(const TSharedRef<FArchive, ESPMode::ThreadSafe>&
 void FElectraPlayerPlugin::Close()
 {
 	CallbackPointerLock.Lock();
-	OptionInterface = nullptr;
+	OptionInterface.Reset();
 	CallbackPointerLock.Unlock();
 	Player->CloseInternal(true);
 }
diff --git a/Engine/Plugins/Media/ElectraPlayer/Source/ElectraPlayerPlugin/Private/ElectraPlayerPlugin.h b/Engine/Plugins/Media/ElectraPlayer/Source/ElectraPlayerPlugin/Private/ElectraPlayerPlugin.h
--- a/Engine/Plugins/Media/ElectraPlayer/Source/ElectraPlayerPlugin/Private/ElectraPlayerPlugin.h
+++ b/Engine/Plugins/Media/ElectraPlayer/Source/ElectraPlayerPlugin/Private/ElectraPlayerPlugin.h
@@ -19,6 +19,7 @@
 #include "Logging/LogMacros.h"
 
 class IElectraPlayerRuntimeModule;
+class IElectraSafeMediaOptionInterface;
 class FElectraPlayerResourceDelegate;
 
 
@@ -188,7 +189,7 @@ private:
 	FCriticalSection CallbackPointerLock;
 
 	/** Option interface */
-	const IMediaOptions* OptionInterface = nullptr;
+	TWeakPtr<IElectraSafeMediaOptionInterface, ESPMode::ThreadSafe> OptionInterface;
 
 	/** The media event handler */
 	IMediaEventSink* EventSink = nullptr;
diff --git a/Engine/Plugins/Media/ElectraPlayer/Source/ElectraPlayerPlugin/Public/IElectraPlayerPluginModule.h b/Engine/Plugins/Media/ElectraPlayer/Source/ElectraPlayerPlugin/Public/IElectraPlayerPluginModule.h
--- a/Engine/Plugins/Media/ElectraPlayer/Source/ElectraPlayerPlugin/Public/IElectraPlayerPluginModule.h
+++ b/Engine/Plugins/Media/ElectraPlayer/Source/ElectraPlayerPlugin/Public/IElectraPlayerPluginModule.h
@@ -12,6 +12,90 @@ class IAnalyticsProviderET;
 class IMediaEventSink;
 class IMediaPlayer;
 
+
+/**
+ * This class is used to get safe access to an IMediaOptions interface.
+ * When passing IMediaOptions through media framework as a POD pointer there is the risk
+ * that it is actually from a derived class like UMediaSource which is subject to GC.
+ * Thus it is not safe to keep the POD IMediaOptions around.
+ * This class is intended to be derived and instantiated from and stored as a TSharedPtr
+ * in a derived UMediaSource class.
+ * Then, as the media player is opened and the POD IMediaOptions is passed along, a
+ * GetMediaOption("GetSafeMediaOptions") query will be made on it to get this instance.
+ * If future access to the IMediaOptions is needed it will be made through this class
+ * by first locking, getting and using the IMediaOptions pointer returned here if it
+ * not null and unlocking afterwards.
+ */
+class IElectraSafeMediaOptionInterface : public IMediaOptions::FDataContainer
+{
+public:
+	virtual ~IElectraSafeMediaOptionInterface() = default;
+	virtual void Lock() = 0;
+	virtual void Unlock() = 0;
+	virtual IMediaOptions* GetMediaOptionInterface() = 0;
+
+	class FScopedLock
+	{
+	public:
+		FScopedLock(const TSharedPtr<IElectraSafeMediaOptionInterface, ESPMode::ThreadSafe>& InSafeMediaOptionInterface)
+			: SafeMediaOptionInterface(InSafeMediaOptionInterface)
+		{
+			if (SafeMediaOptionInterface.IsValid())
+			{
+				SafeMediaOptionInterface->Lock();
+			}
+		}
+		~FScopedLock()
+		{
+			if (SafeMediaOptionInterface.IsValid())
+			{
+				SafeMediaOptionInterface->Unlock();
+			}
+		}
+	
+	private:
+		TSharedPtr<IElectraSafeMediaOptionInterface, ESPMode::ThreadSafe> SafeMediaOptionInterface;
+	};
+};
+
+
+/**
+ * A ready to use implementation of the IElectraSafeMediaOptionInterface
+ */
+class FElectraSafeMediaOptionInterface : public IElectraSafeMediaOptionInterface
+{
+public:
+	FElectraSafeMediaOptionInterface(IMediaOptions* InOwner)
+		: Owner(InOwner)
+	{ }
+	virtual ~FElectraSafeMediaOptionInterface()
+	{
+		ClearOwner();
+	}
+	virtual void ClearOwner()
+	{
+		FScopeLock lock(&OwnerLock);
+		Owner = nullptr;
+	}
+	virtual void Lock() override
+	{
+		OwnerLock.Lock();
+	}
+	virtual void Unlock() override
+	{
+		OwnerLock.Unlock();
+	}
+	virtual IMediaOptions* GetMediaOptionInterface() override
+	{
+		return Owner;
+	}
+private:
+	FCriticalSection OwnerLock;
+	IMediaOptions* Owner = nullptr;
+};
+
+
+
 //! Data type for use with media options interface
 class FElectraSeekablePositions : public IMediaOptions::FDataContainer
 {
