ElectraPlayerPlugin: Added a missing include causing CI errors

#rnx

diff --git a/Engine/Plugins/Media/ElectraPlayer/Source/ElectraPlayerPlugin/Public/IElectraPlayerPluginModule.h b/Engine/Plugins/Media/ElectraPlayer/Source/ElectraPlayerPlugin/Public/IElectraPlayerPluginModule.h
--- a/Engine/Plugins/Media/ElectraPlayer/Source/ElectraPlayerPlugin/Public/IElectraPlayerPluginModule.h
+++ b/Engine/Plugins/Media/ElectraPlayer/Source/ElectraPlayerPlugin/Public/IElectraPlayerPluginModule.h
@@ -4,6 +4,7 @@
 
 #include "Templates/SharedPointer.h"
 #include "Modules/ModuleInterface.h"
+#include "Misc/ScopeLock.h"
 #include "Misc/Timespan.h"
 #include "Containers/Array.h"
 #include "IMediaOptions.h"
