ElectraPlayer: Improved retaining of selected subtitle track by also considering codec along with language; support for sideloaded subtitles in DASH

#rnx

diff --git a/Engine/Plugins/Media/ElectraPlayer/Source/ElectraPlayerRuntime/Private/ElectraPlayer.cpp b/Engine/Plugins/Media/ElectraPlayer/Source/ElectraPlayerRuntime/Private/ElectraPlayer.cpp
--- a/Engine/Plugins/Media/ElectraPlayer/Source/ElectraPlayerRuntime/Private/ElectraPlayer.cpp
+++ b/Engine/Plugins/Media/ElectraPlayer/Source/ElectraPlayerRuntime/Private/ElectraPlayer.cpp
@@ -1348,6 +1348,9 @@ bool FElectraPlayer::SelectTrack(EPlayerTrackType TrackType, int32 TrackIndex)
 						AudioAttributes.Language_ISO639 = Meta->Language;
 						PlaystartOptions.InitialAudioTrackAttributes.Language_ISO639 = Meta->Language;
 					}
+					AudioAttributes.Codec = Meta->HighestBandwidthCodec.GetCodecName();
+					PlaystartOptions.InitialAudioTrackAttributes.Codec = Meta->HighestBandwidthCodec.GetCodecName();
+
 					TSharedPtr<FInternalPlayerImpl, ESPMode::ThreadSafe> LockedPlayer = CurrentPlayer;
 					if (LockedPlayer.IsValid() && LockedPlayer->AdaptivePlayer.IsValid())
 					{
@@ -1399,6 +1402,9 @@ bool FElectraPlayer::SelectTrack(EPlayerTrackType TrackType, int32 TrackIndex)
 						SubtitleAttributes.Language_ISO639 = Meta->Language;
 						PlaystartOptions.InitialSubtitleTrackAttributes.Language_ISO639 = Meta->Language;
 					}
+					SubtitleAttributes.Codec = Meta->HighestBandwidthCodec.GetCodecName();
+					PlaystartOptions.InitialSubtitleTrackAttributes.Codec = Meta->HighestBandwidthCodec.GetCodecName();
+
 					TSharedPtr<FInternalPlayerImpl, ESPMode::ThreadSafe> LockedPlayer = CurrentPlayer;
 					if (LockedPlayer.IsValid() && LockedPlayer->AdaptivePlayer.IsValid())
 					{
diff --git a/Engine/Plugins/Media/ElectraPlayer/Source/ElectraPlayerRuntime/Private/Runtime/Decoder/SubtitleDecoder.cpp b/Engine/Plugins/Media/ElectraPlayer/Source/ElectraPlayerRuntime/Private/Runtime/Decoder/SubtitleDecoder.cpp
--- a/Engine/Plugins/Media/ElectraPlayer/Source/ElectraPlayerRuntime/Private/Runtime/Decoder/SubtitleDecoder.cpp
+++ b/Engine/Plugins/Media/ElectraPlayer/Source/ElectraPlayerRuntime/Private/Runtime/Decoder/SubtitleDecoder.cpp
@@ -113,7 +113,7 @@ public:
 
 	virtual void SetPlayerSessionServices(IPlayerSessionServices* SessionServices) override;
 
-	virtual void Open() override;
+	virtual void Open(const FParamDict& Options) override;
 	virtual void Close() override;
 	virtual void Start() override;
 	virtual void Stop() override;
@@ -144,6 +144,7 @@ private:
 	void OnDecodedSubtitleReceived(ISubtitleDecoderOutputPtr DecodedSubtitle);
 
 	TSharedPtr<IElectraSubtitleDecoder, ESPMode::ThreadSafe> PluginDecoder;
+	FParamDict PluginDecoderOptions;
 	IPlayerSessionServices* PlayerSessionServices = nullptr;
 	TArray<uint8> RawCSD;
 	FStreamCodecInformation CodecInfo;
@@ -164,7 +165,17 @@ ISubtitleDecoder* ISubtitleDecoder::Create(FAccessUnit* AccessUnit)
 	check(AccessUnit);
 	if (AccessUnit && AccessUnit->AUCodecData.IsValid())
 	{
-		TSharedPtr<IElectraSubtitleDecoder, ESPMode::ThreadSafe> PluginDecoder = FSubtitleDecoderPlugins::Get().CreateDecoder(AccessUnit->AUCodecData->ParsedInfo.GetCodecSpecifierRFC6381());
+		TSharedPtr<IElectraSubtitleDecoder, ESPMode::ThreadSafe> PluginDecoder;
+		// For sideloaded data prefer to create the decoder via the Mime type. This allows the decoder to
+		// know that the data is there all at once and it needs to handle it in a non-streaming fashion.
+		if (AccessUnit->bIsSideloaded)
+		{
+			PluginDecoder = FSubtitleDecoderPlugins::Get().CreateDecoder(AccessUnit->AUCodecData->ParsedInfo.GetMimeType());
+		}
+		if (!PluginDecoder.IsValid())
+		{
+			PluginDecoder = FSubtitleDecoderPlugins::Get().CreateDecoder(AccessUnit->AUCodecData->ParsedInfo.GetCodecSpecifierRFC6381());
+		}
 		if (PluginDecoder.IsValid())
 		{
 			FSubtitleDecoder* Decoder = new FSubtitleDecoder(MoveTemp(PluginDecoder));
@@ -204,13 +215,14 @@ void FSubtitleDecoder::SetPlayerSessionServices(IPlayerSessionServices* InSessio
 	PlayerSessionServices = InSessionServices;
 }
 
-void FSubtitleDecoder::Open()
+void FSubtitleDecoder::Open(const FParamDict& Options)
 {
 	check(PlayerSessionServices);
 	check(PluginDecoder.IsValid());
 	if (PluginDecoder.IsValid())
 	{
-		FParamDict Addtl;
+		PluginDecoderOptions = Options;
+		FParamDict Addtl(Options);
 		Addtl.Set(TEXT("width"), FVariantValue((int64) CodecInfo.GetResolution().Width));
 		Addtl.Set(TEXT("height"), FVariantValue((int64) CodecInfo.GetResolution().Height));
 		Addtl.Set(TEXT("offset_x"), FVariantValue((int64) CodecInfo.GetTranslation().GetX()));
@@ -306,7 +318,7 @@ bool FSubtitleDecoder::ResetForTrackChange(FAccessUnit* AccessUnit)
 		{
 			SetCSD(AccessUnit->AUCodecData->RawCSD);
 			SetCodecInfo(AccessUnit->AUCodecData->ParsedInfo);
-			Open();
+			Open(PluginDecoderOptions);
 			if (bWasStarted)
 			{
 				Start();
@@ -332,6 +344,13 @@ IAccessUnitBufferInterface::EAUpushResult FSubtitleDecoder::AUdataPushAU(FAccess
 				}
 			}
 			FParamDict Addtl;
+			// For sideloaded subtitles we pass an ID to the plugin via the additional options. This allows the plugin to identify
+			// whether or not it has already parsed this data before (when seeking for instance).
+			if (AccessUnit->bIsSideloaded && AccessUnit->BufferSourceInfo.IsValid())
+			{
+				Addtl.Set(FString(TEXT("SideloadedID")), FVariantValue(AccessUnit->BufferSourceInfo->PeriodAdaptationSetID));
+			}
+			Addtl.Set(FString(TEXT("PTO")), FVariantValue(AccessUnit->PTO));
 			PluginDecoder->AddStreamedSubtitleData(TArray<uint8>((const uint8*)AccessUnit->AUData, (int32)AccessUnit->AUSize), AccessUnit->PTS, AccessUnit->Duration, Addtl);
 		}
 		return IAccessUnitBufferInterface::EAUpushResult::Ok;
diff --git a/Engine/Plugins/Media/ElectraPlayer/Source/ElectraPlayerRuntime/Private/Runtime/Decoder/SubtitleDecoder.h b/Engine/Plugins/Media/ElectraPlayer/Source/ElectraPlayerRuntime/Private/Runtime/Decoder/SubtitleDecoder.h
--- a/Engine/Plugins/Media/ElectraPlayer/Source/ElectraPlayerRuntime/Private/Runtime/Decoder/SubtitleDecoder.h
+++ b/Engine/Plugins/Media/ElectraPlayer/Source/ElectraPlayerRuntime/Private/Runtime/Decoder/SubtitleDecoder.h
@@ -32,7 +32,7 @@ namespace Electra
 		// State methods
 		//
 		virtual void SetPlayerSessionServices(IPlayerSessionServices* SessionServices) = 0;
-		virtual void Open() = 0;
+		virtual void Open(const FParamDict& Options) = 0;
 		virtual void Close() = 0;
 		virtual void Start() = 0;
 		virtual void Stop() = 0;
diff --git a/Engine/Plugins/Media/ElectraPlayer/Source/ElectraPlayerRuntime/Private/Runtime/Demuxer/ParserISO14496-12.cpp b/Engine/Plugins/Media/ElectraPlayer/Source/ElectraPlayerRuntime/Private/Runtime/Demuxer/ParserISO14496-12.cpp
--- a/Engine/Plugins/Media/ElectraPlayer/Source/ElectraPlayerRuntime/Private/Runtime/Demuxer/ParserISO14496-12.cpp
+++ b/Engine/Plugins/Media/ElectraPlayer/Source/ElectraPlayerRuntime/Private/Runtime/Demuxer/ParserISO14496-12.cpp
@@ -100,8 +100,24 @@ namespace Electra
 			return UEMEDIA_ERROR_READ_ERROR;
 		}
 
-		UEMediaError ReadString(FString& OutString, uint32 MaxBytes)
+		UEMediaError ReadString(FString& OutString, uint32 MaxBytes, bool bIsUTF8=true)
 		{
+			TArray<uint8> Buf;
+			Buf.Reserve(MaxBytes);
+
+			auto ConvFromBuf = [](const TArray<uint8>& SrcBuf, bool bIsUTF8) -> FString
+			{
+				if (bIsUTF8)
+				{
+					FUTF8ToTCHAR cnv((const ANSICHAR*)SrcBuf.GetData(), SrcBuf.Num());
+					return FString(cnv.Length(), cnv.Get());
+				}
+				else
+				{
+					return FString((const char*)SrcBuf.GetData());
+				}
+			};
+
 			OutString.Empty();
 			if (MaxBytes == 0)
 			{
@@ -118,13 +134,14 @@ namespace Electra
 				}
 				if (NextChar == 0)
 				{
+					OutString = ConvFromBuf(Buf, bIsUTF8);
 					return UEMEDIA_ERROR_OK;
 				}
 				if (i == 0)
 				{
 					FirstByte = NextChar;
 				}
-				OutString.AppendChar((TCHAR)NextChar);
+				Buf.Add(NextChar);
 			}
 			// NOTE: This is mostly for the 'qt' brand. ISO files should not be using such strings.
 			// FIXME: Do this only for brands that may require this?
@@ -132,6 +149,7 @@ namespace Electra
 			if (FirstByte + 1 == MaxBytes)
 			{
 				// NOTE: We do _not_ remove the length from the string!
+				OutString = ConvFromBuf(Buf, bIsUTF8);
 				return UEMEDIA_ERROR_OK;
 			}
 			return UEMEDIA_ERROR_INSUFFICIENT_DATA;	// Did not find a NUL char to terminate the string in the maximum number of characters allowed to read.
@@ -1715,6 +1733,38 @@ private:
 	};
 
 
+	/**
+	 * 'sthd' box.
+	 * ISO/IEC 14496-12:2014 - 12.6.2 - Subtitle media header
+	 */
+	class FMP4BoxSTHD : public FMP4BoxFull
+	{
+	public:
+		FMP4BoxSTHD(IParserISO14496_12::FBoxType InBoxType, int64 InBoxSize, int64 InStartOffset, int64 InDataOffset, bool bInIsLeafBox)
+			: FMP4BoxFull(InBoxType, InBoxSize, InStartOffset, InDataOffset, bInIsLeafBox)
+		{
+		}
+
+		virtual ~FMP4BoxSTHD()
+		{
+		}
+
+	private:
+		FMP4BoxSTHD() = delete;
+		FMP4BoxSTHD(const FMP4BoxSTHD&) = delete;
+
+	protected:
+		virtual UEMediaError ReadAndParseAttributes(FMP4ParseInfo* ParseInfo) override
+		{
+			UEMediaError Error = UEMEDIA_ERROR_OK;
+			RETURN_IF_ERROR(FMP4BoxFull::ReadAndParseAttributes(ParseInfo));
+			return Error;
+		}
+
+	private:
+	};
+
+
 	/**
 	 * 'dref' box.
 	 * ISO/IEC 14496-12:2014 - 8.7.2 - Data Reference Box
@@ -2366,6 +2416,146 @@ private:
 	};
 
 
+	/**
+	 * XMLSubtitleSampleEntry sample entry (ISO/IEC 14496-12 Section 12.6.3 Sample entry)
+	 */
+	class FMP4BoxXMLSubtitleSampleEntry : public FMP4BoxSampleEntry
+	{
+	public:
+		FMP4BoxXMLSubtitleSampleEntry(IParserISO14496_12::FBoxType InBoxType, int64 InBoxSize, int64 InStartOffset, int64 InDataOffset, bool bInIsLeafBox, int32 InSTSDBoxVersion)
+			: FMP4BoxSampleEntry(InBoxType, InBoxSize, InStartOffset, InDataOffset, bInIsLeafBox)
+			, STSDBoxVersion(InSTSDBoxVersion)
+		{
+		}
+
+		virtual ~FMP4BoxXMLSubtitleSampleEntry()
+		{
+		}
+
+		const TArray<uint8> GetRawBoxData() const
+		{
+			return BoxData;
+		}
+
+		int32 GetTranslationX() const
+		{
+			return TranslationX;
+		}
+
+		int32 GetTranslationY() const
+		{
+			return TranslationY;
+		}
+
+		uint16 GetWidth() const
+		{
+			return Width;
+		}
+
+		uint16 GetHeight() const
+		{
+			return Height;
+		}
+
+		uint32 GetTimescale() const
+		{
+			return Timescale;
+		}
+
+		const FString& GetNamespace() const
+		{
+			return Namespace;
+		}
+		const FString& GetSchemaLocation() const
+		{
+			return SchemaLocation;
+		}
+		const TArray<FString>& GetAuxiliaryMimeTypes() const
+		{
+			return AuxiliaryMimeTypes;
+		}
+
+	private:
+		FMP4BoxXMLSubtitleSampleEntry() = delete;
+		FMP4BoxXMLSubtitleSampleEntry(const FMP4BoxXMLSubtitleSampleEntry&) = delete;
+
+	protected:
+		virtual UEMediaError ReadAndParseAttributes(FMP4ParseInfo* ParseInfo) override
+		{
+			UEMediaError Error = UEMEDIA_ERROR_OK;
+
+			// Duplicate all the data we are reading into a buffer that we need to provide to the subtitle plugin.
+			FMP4BoxReader::ScopedDataDuplication DuplicateBoxData(ParseInfo->Reader(), BoxData);
+
+			RETURN_IF_ERROR(FMP4BoxSampleEntry::ReadAndParseAttributes(ParseInfo));
+
+			const FMP4BoxTKHD* TKHD = ParseInfo->GetCurrentTrackBox() ? static_cast<const FMP4BoxTKHD*>(ParseInfo->GetCurrentTrackBox()->GetBoxPath(FMP4Box::kBox_tkhd)) : nullptr;
+			if (TKHD)
+			{
+				Width = TKHD->GetWidth();
+				Height = TKHD->GetHeight();
+				// Translation must not be using fractional digits and can thus be treated like integers.
+				TranslationX = (int32) TKHD->GetMatrixValue(6);
+				TranslationY = (int32) TKHD->GetMatrixValue(7);
+			}
+
+			const FMP4BoxMDHD* MDHD = ParseInfo->GetCurrentTrackBox() ? static_cast<const FMP4BoxMDHD*>(ParseInfo->GetCurrentTrackBox()->FindBox(FMP4Box::kBox_mdhd)) : nullptr;
+			if (MDHD)
+			{
+				Timescale = MDHD->GetTimescale();
+			}
+
+			int32 BytesRemaining = (int32) (BoxSize - (ParseInfo->Reader()->GetCurrentReadOffset() - StartOffset));
+			RETURN_IF_ERROR(ParseInfo->Reader()->ReadString(Namespace, BytesRemaining));
+			BytesRemaining = (int32) (BoxSize - (ParseInfo->Reader()->GetCurrentReadOffset() - StartOffset));
+			RETURN_IF_ERROR(ParseInfo->Reader()->ReadString(SchemaLocation, BytesRemaining));
+			FString AuxMimes;
+			BytesRemaining = (int32) (BoxSize - (ParseInfo->Reader()->GetCurrentReadOffset() - StartOffset));
+			RETURN_IF_ERROR(ParseInfo->Reader()->ReadString(AuxMimes, BytesRemaining));
+			AuxMimes.ParseIntoArrayWS(AuxiliaryMimeTypes);
+			BytesRemaining = (int32) (BoxSize - (ParseInfo->Reader()->GetCurrentReadOffset() - StartOffset));
+
+			// There may be child boxes here now, more precisely a 'btrt' box.
+			if (BytesRemaining > 0)
+			{
+				// Now that there are additional boxes we ourselves are no longer a leaf box.
+				bIsLeafBox = false;
+				// This will read all the boxes in here and add them as children.
+				for(; BytesRemaining>0; BytesRemaining = (int32) (BoxSize - (ParseInfo->Reader()->GetCurrentReadOffset() - StartOffset)))
+				{
+					IParserISO14496_12::FBoxType	ChildBoxType;
+					int64							ChildBoxSize;
+					uint8							ChildBoxUUID[16];
+
+					int64 BoxStartOffset = ParseInfo->Reader()->GetCurrentReadOffset();
+					RETURN_IF_ERROR(ParseInfo->ReadBoxTypeAndSize(ChildBoxType, ChildBoxSize, ChildBoxUUID));
+#if MEDIA_DEBUG_HAS_BOX_NAMES
+					char boxName[4];
+					*((uint32*)&boxName) = MEDIA_TO_BIG_ENDIAN(ChildBoxType);
+#endif
+					int64 BoxDataOffset = ParseInfo->Reader()->GetCurrentReadOffset();
+					FMP4Box* NextBox = new FMP4BoxIgnored(ChildBoxType, ChildBoxSize, BoxStartOffset, BoxDataOffset, true);
+
+					AddChildBox(NextBox);
+					RETURN_IF_ERROR(ParseInfo->ReadAndParseNextBox(NextBox));
+				}
+			}
+			return Error;
+		}
+
+		TArray<uint8>		BoxData;
+		FString				Namespace;
+		FString				SchemaLocation;
+		TArray<FString>		AuxiliaryMimeTypes;
+		int32				STSDBoxVersion;
+		uint32				Timescale = 0;
+		uint16				Width = 0;
+		uint16				Height = 0;
+		int32				TranslationX = 0;
+		int32				TranslationY = 0;
+	};
+
+
 	/**
 	 * AVC Decoder Configuration box (ISO/IEC 14496-15:2014 - 5.4.2.1.2)
 	 */
@@ -2705,7 +2895,7 @@ private:
 			UEMediaError Error = UEMEDIA_ERROR_OK;
 			RETURN_IF_ERROR(FMP4BoxFull::ReadAndParseAttributes(ParseInfo));
 			uint32 BytesRemaining = BoxSize - (ParseInfo->Reader()->GetCurrentReadOffset() - StartOffset);
-			RETURN_IF_ERROR(ParseInfo->Reader()->ReadString(ExtendedLanguage, BytesRemaining));
+			RETURN_IF_ERROR(ParseInfo->Reader()->ReadString(ExtendedLanguage, BytesRemaining, false));
 			return Error;
 		}
 
@@ -2996,6 +3186,10 @@ private:
 						{
 							NextBox = new FMP4BoxWVTTSampleEntry(ChildBoxType, ChildBoxSize, BoxStartOffset, BoxDataOffset, true, Version);
 						}
+						else if (ChildBoxType == FMP4Box::kSample_stpp)
+						{
+							NextBox = new FMP4BoxXMLSubtitleSampleEntry(ChildBoxType, ChildBoxSize, BoxStartOffset, BoxDataOffset, true, Version);
+						}
 						else
 						{
 							NextBox = new FMP4BoxIgnored(ChildBoxType, ChildBoxSize, BoxStartOffset, BoxDataOffset, true);
@@ -4851,6 +5045,9 @@ private:
 						case FMP4Box::kBox_smhd:
 							NextBox = new FMP4BoxSMHD(BoxType, BoxSize, BoxStartOffset, BoxDataOffset, true);
 							break;
+						case FMP4Box::kBox_sthd:
+							NextBox = new FMP4BoxSTHD(BoxType, BoxSize, BoxStartOffset, BoxDataOffset, true);
+							break;
 						case FMP4Box::kBox_dref:
 							NextBox = new FMP4BoxDREF(BoxType, BoxSize, BoxStartOffset, BoxDataOffset, true);		// 'dref' is not a true leaf box, but its elements are a list that must be parsed.
 							break;
@@ -5389,11 +5586,12 @@ private:
 		};
 
 
-		UEMediaError ParseAVC1SampleType(IPlayerSessionServices* PlayerSession, FTrack* Track, const FMP4Box* SampleBox);
+		UEMediaError ParseAVCxSampleType(IPlayerSessionServices* PlayerSession, FTrack* Track, const FMP4Box* SampleBox, int32 InAvcType);
 		UEMediaError ParseHVC1SampleType(IPlayerSessionServices* PlayerSession, FTrack* Track, const FMP4Box* SampleBox);
 		UEMediaError ParseMP4ASampleType(IPlayerSessionServices* PlayerSession, FTrack* Track, const FMP4Box* SampleBox);
 		UEMediaError ParseTX3GSampleType(IPlayerSessionServices* PlayerSession, FTrack* Track, const FMP4Box* SampleBox);
 		UEMediaError ParseWVTTSampleType(IPlayerSessionServices* PlayerSession, FTrack* Track, const FMP4Box* SampleBox);
+		UEMediaError ParseSTPPSampleType(IPlayerSessionServices* PlayerSession, FTrack* Track, const FMP4Box* SampleBox);
 
 		FMP4ParseInfo* ParsedData;
 
@@ -6268,8 +6466,9 @@ private:
 	}
 
 
-	UEMediaError FParserISO14496_12::ParseAVC1SampleType(IPlayerSessionServices* PlayerSession, FTrack* Track, const FMP4Box* SampleBox)
+	UEMediaError FParserISO14496_12::ParseAVCxSampleType(IPlayerSessionServices* PlayerSession, FTrack* Track, const FMP4Box* SampleBox, int32 InAvcType)
 	{
+		check(InAvcType == 1 || InAvcType == 3);
 		const FMP4BoxVisualSampleEntry* VisualSampleEntry = static_cast<const FMP4BoxVisualSampleEntry*>(SampleBox);
 		check(VisualSampleEntry->GetNumberOfChildren() > 0);
 		if (VisualSampleEntry->GetNumberOfChildren() > 0)
@@ -6295,23 +6494,44 @@ private:
 							Track->CodecInformation.SetStreamLanguageCode(Track->GetLanguage());
 							if (Track->CodecSpecificDataAVC.GetNumberOfSPS() == 0)
 							{
-								return UEMEDIA_ERROR_FORMAT_ERROR;
+								if (InAvcType == 3)
+								{
+									Track->CodecInformation.SetProfile(Track->CodecSpecificDataAVC.GetAVCProfileIndication());
+									Track->CodecInformation.SetProfileLevel(Track->CodecSpecificDataAVC.GetAVCLevelIndication());
+									Track->CodecInformation.SetProfileConstraints(Track->CodecSpecificDataAVC.GetProfileCompatibility());
+									Track->CodecInformation.SetCodecSpecifierRFC6381(FString::Printf(TEXT("avc%d.%02x%02x%02x"), InAvcType, Track->CodecSpecificDataAVC.GetAVCProfileIndication(), Track->CodecSpecificDataAVC.GetProfileCompatibility(), Track->CodecSpecificDataAVC.GetAVCLevelIndication()));
+									if (Track->TKHDBox)
+									{
+										Track->CodecInformation.SetResolution(FStreamCodecInformation::FResolution(Track->TKHDBox->GetWidth(), Track->TKHDBox->GetHeight()));
+										Track->CodecInformation.SetCrop(FStreamCodecInformation::FCrop(0, 0, Track->TKHDBox->GetWidth() % 16, Track->TKHDBox->GetHeight() % 16));
+										Track->CodecInformation.SetAspectRatio(FStreamCodecInformation::FAspectRatio(1, 1));
+									}
+
+									bGotVideoFormat = true;
+								}
+								else
+								{
+									return UEMEDIA_ERROR_FORMAT_ERROR;
+								}
+							}
+							else
+							{
+								const MPEG::FISO14496_10_seq_parameter_set_data& sps = Track->CodecSpecificDataAVC.GetParsedSPS(0);
+								int32 CropL, CropR, CropT, CropB;
+								sps.GetCrop(CropL, CropR, CropT, CropB);
+								Track->CodecInformation.SetResolution(FStreamCodecInformation::FResolution(sps.GetWidth() - CropL - CropR, sps.GetHeight() - CropT - CropB));
+								Track->CodecInformation.SetCrop(FStreamCodecInformation::FCrop(CropL, CropT, CropR, CropB));
+								FStreamCodecInformation::FAspectRatio ar;
+								sps.GetAspect(ar.Width, ar.Height);
+								Track->CodecInformation.SetAspectRatio(ar);
+								Track->CodecInformation.SetFrameRate(sps.GetTiming());
+								Track->CodecInformation.SetProfile(sps.profile_idc);
+								Track->CodecInformation.SetProfileLevel(sps.level_idc);
+								uint8 Constraints = (sps.constraint_set0_flag << 7) | (sps.constraint_set1_flag << 6) | (sps.constraint_set2_flag << 5) | (sps.constraint_set3_flag << 4) | (sps.constraint_set4_flag << 3) | (sps.constraint_set5_flag << 2);
+								Track->CodecInformation.SetProfileConstraints(Constraints);
+								Track->CodecInformation.SetCodecSpecifierRFC6381(FString::Printf(TEXT("avc%d.%02x%02x%02x"), InAvcType, sps.profile_idc, Constraints, sps.level_idc));
+								bGotVideoFormat = true;
 							}
-							const MPEG::FISO14496_10_seq_parameter_set_data& sps = Track->CodecSpecificDataAVC.GetParsedSPS(0);
-							int32 CropL, CropR, CropT, CropB;
-							sps.GetCrop(CropL, CropR, CropT, CropB);
-							Track->CodecInformation.SetResolution(FStreamCodecInformation::FResolution(sps.GetWidth() - CropL - CropR, sps.GetHeight() - CropT - CropB));
-							Track->CodecInformation.SetCrop(FStreamCodecInformation::FCrop(CropL, CropT, CropR, CropB));
-							FStreamCodecInformation::FAspectRatio ar;
-							sps.GetAspect(ar.Width, ar.Height);
-							Track->CodecInformation.SetAspectRatio(ar);
-							Track->CodecInformation.SetFrameRate(sps.GetTiming());
-							Track->CodecInformation.SetProfile(sps.profile_idc);
-							Track->CodecInformation.SetProfileLevel(sps.level_idc);
-							uint8 Constraints = (sps.constraint_set0_flag << 7) | (sps.constraint_set1_flag << 6) | (sps.constraint_set2_flag << 5) | (sps.constraint_set3_flag << 4) | (sps.constraint_set4_flag << 3) | (sps.constraint_set5_flag << 2);
-							Track->CodecInformation.SetProfileConstraints(Constraints);
-							Track->CodecInformation.SetCodecSpecifierRFC6381(FString::Printf(TEXT("avc1.%02x%02x%02x"), sps.profile_idc, Constraints, sps.level_idc));
-							bGotVideoFormat = true;
 						}
 						else
 						{
@@ -6568,6 +6788,27 @@ private:
 		return UEMEDIA_ERROR_NOT_SUPPORTED;
 	}
 
+	UEMediaError FParserISO14496_12::ParseSTPPSampleType(IPlayerSessionServices* PlayerSession, FTrack* Track, const FMP4Box* SampleBox)
+	{
+		// Need to check if there is a subtitle decoder capable of decoding this.
+		if (ISubtitleDecoder::IsSupported(TEXT(""), TEXT("stpp")))
+		{
+			const FMP4BoxXMLSubtitleSampleEntry* SubtitleSampleEntry = static_cast<const FMP4BoxXMLSubtitleSampleEntry*>(Track->STSDBox->GetChildBox(0));
+
+			Track->CodecSpecificDataRAW = SubtitleSampleEntry->GetRawBoxData();
+			Track->CodecInformation.SetStreamType(EStreamType::Subtitle);
+			Track->CodecInformation.SetCodec(FStreamCodecInformation::ECodec::TTML);
+			Track->CodecInformation.SetCodecSpecificData(Track->CodecSpecificDataRAW);
+			Track->CodecInformation.SetStreamLanguageCode(Track->GetLanguage());
+			Track->CodecInformation.SetCodecSpecifierRFC6381(TEXT("stpp"));
+			Track->CodecInformation.SetResolution(FStreamCodecInformation::FResolution(SubtitleSampleEntry->GetWidth(), SubtitleSampleEntry->GetHeight()));
+			Track->CodecInformation.SetTranslation(FStreamCodecInformation::FTranslation(SubtitleSampleEntry->GetTranslationX(), SubtitleSampleEntry->GetTranslationY()));
+			Track->CodecInformation.SetFrameRate(FTimeFraction(1, SubtitleSampleEntry->GetTimescale()));
+			return UEMEDIA_ERROR_OK;
+		}
+		return UEMEDIA_ERROR_NOT_SUPPORTED;
+	}
+
 	UEMediaError FParserISO14496_12::PrepareTracks(IPlayerSessionServices* PlayerSession, TSharedPtrTS<const IParserISO14496_12> OptionalMP4InitSegment)
 	{
 		delete ParsedTrackInfo;
@@ -6782,11 +7023,15 @@ private:
 								}
 								// Expecting avc1/avc3/hvc1/hev1 for video right now.
 								UEMediaError Error = UEMEDIA_ERROR_FORMAT_ERROR;
-								if (FRMABox->GetDataFormat() == FMP4Box::kSample_avc1 || FRMABox->GetDataFormat() == FMP4Box::kSample_avc3)
+								if (FRMABox->GetDataFormat() == FMP4Box::kSample_avc1)
+								{
+									Error = ParseAVCxSampleType(PlayerSession, Track.Get(), STSDFirstChildBox, 1);
+								}
+								else if (FRMABox->GetDataFormat() == FMP4Box::kSample_avc3)
 								{
-									Error = ParseAVC1SampleType(PlayerSession, Track.Get(), STSDFirstChildBox);
+									Error = ParseAVCxSampleType(PlayerSession, Track.Get(), STSDFirstChildBox, 3);
 								}
-								else if ( FRMABox->GetDataFormat() == FMP4Box::kSample_hvc1 || FRMABox->GetDataFormat() == FMP4Box::kSample_hev1)
+								else if (FRMABox->GetDataFormat() == FMP4Box::kSample_hvc1 || FRMABox->GetDataFormat() == FMP4Box::kSample_hev1)
 								{
 									Error = ParseHVC1SampleType(PlayerSession, Track.Get(), STSDFirstChildBox);
 								}
@@ -6843,7 +7088,7 @@ private:
 							case FMP4Box::kSample_avc1:
 							case FMP4Box::kSample_avc3:
 							{
-								UEMediaError Error = ParseAVC1SampleType(PlayerSession, Track.Get(), STSDFirstChildBox);
+								UEMediaError Error = ParseAVCxSampleType(PlayerSession, Track.Get(), STSDFirstChildBox, STSDFirstChildBox->GetType() == FMP4Box::kSample_avc1 ? 1 : 3);
 								/*
 								if (Error == UEMEDIA_ERROR_NOT_SUPPORTED)
 								{
@@ -6885,6 +7130,15 @@ private:
 							}
 							case FMP4Box::kSample_stpp:
 							{
+								UEMediaError Error = ParseSTPPSampleType(PlayerSession, Track.Get(), STSDFirstChildBox);
+								if (Error == UEMEDIA_ERROR_NOT_SUPPORTED)
+								{
+									bIsSupported = false;
+								}
+								else if (Error != UEMEDIA_ERROR_OK)
+								{
+									return Error;
+								}
 								break;
 							}
 							case FMP4Box::kSample_tx3g:
diff --git a/Engine/Plugins/Media/ElectraPlayer/Source/ElectraPlayerRuntime/Private/Runtime/Player/AdaptiveStreamingPlayer.cpp b/Engine/Plugins/Media/ElectraPlayer/Source/ElectraPlayerRuntime/Private/Runtime/Player/AdaptiveStreamingPlayer.cpp
--- a/Engine/Plugins/Media/ElectraPlayer/Source/ElectraPlayerRuntime/Private/Runtime/Player/AdaptiveStreamingPlayer.cpp
+++ b/Engine/Plugins/Media/ElectraPlayer/Source/ElectraPlayerRuntime/Private/Runtime/Player/AdaptiveStreamingPlayer.cpp
@@ -736,31 +736,44 @@ void FAdaptiveStreamingPlayer::GetSelectedTrackAttributes(FStreamSelectionAttrib
 {
 	TSharedPtrTS<FBufferSourceInfo> BufferInfo;
 
-	auto SetOutAttributes = [&](const TSharedPtrTS<FBufferSourceInfo>& BufferInfo, const FStreamSelectionAttributes& CurrentAttr) -> void
+	auto SetOutAttributes = [&](const TSharedPtrTS<FBufferSourceInfo>& BufferInfo, const FStreamSelectionAttributes& CurrentAttr, bool bIsValidInfo) -> void
 	{
 		if (BufferInfo.IsValid())
 		{
 			OutAttributes.Kind = BufferInfo->Kind;
 			OutAttributes.Language_ISO639 = BufferInfo->Language;
+			OutAttributes.Codec = BufferInfo->Codec;
 			OutAttributes.OverrideIndex = BufferInfo->HardIndex;
 		}
 		else
 		{
 			OutAttributes = CurrentAttr;
+			// If the information is known to be valid then the stream is currently not available.
+			// We set the override index to reflect this and return the intended selection.
+			if (bIsValidInfo)
+			{
+				OutAttributes.OverrideIndex = -1;
+			}
 		}
 	};
 
+	FTimeValue Now = PlaybackState.GetPlayPosition();
+	bool bHavePeriods;
+	bool bFoundTime;
+	BufferInfo = GetStreamBufferInfoAtTime(bHavePeriods, bFoundTime, StreamType, Now);
+	bool bBufferInfoValid = bHavePeriods && bFoundTime;
+
 	if (StreamType == EStreamType::Video)
 	{
-		SetOutAttributes(MultiStreamBufferVid.GetActiveOutputBufferInfo(), SelectedStreamAttributesVid);
+		SetOutAttributes(BufferInfo, SelectedStreamAttributesVid, bBufferInfoValid);
 	}
 	else if (StreamType == EStreamType::Audio)
 	{
-		SetOutAttributes(MultiStreamBufferAud.GetActiveOutputBufferInfo(), SelectedStreamAttributesAud);
+		SetOutAttributes(BufferInfo, SelectedStreamAttributesAud, bBufferInfoValid);
 	}
 	else if (StreamType == EStreamType::Subtitle)
 	{
-		SetOutAttributes(MultiStreamBufferTxt.GetActiveOutputBufferInfo(), SelectedStreamAttributesTxt);
+		SetOutAttributes(BufferInfo, SelectedStreamAttributesTxt, bBufferInfoValid);
 	}
 }
 
@@ -2176,26 +2189,26 @@ void FAdaptiveStreamingPlayer::InternalHandlePendingStartRequest(const FTimeValu
 						MultiStreamBufferTxt.SelectTrackWhenAvailable(BufferSourceInfoTxt);
 						if (BufferSourceInfoVid.IsValid())
 						{
-							SelectedStreamAttributesVid.UpdateWith(BufferSourceInfoVid->Kind, BufferSourceInfoVid->Language, BufferSourceInfoVid->HardIndex);
+							SelectedStreamAttributesVid.UpdateWith(BufferSourceInfoVid->Kind, BufferSourceInfoVid->Language, BufferSourceInfoVid->Codec, BufferSourceInfoVid->HardIndex);
 							if (ManifestType != EMediaFormatType::ISOBMFF)
 							{
-								StreamSelectionAttributesVid.UpdateIfOverrideSet(BufferSourceInfoVid->Kind, BufferSourceInfoVid->Language);
+								StreamSelectionAttributesVid.UpdateIfOverrideSet(BufferSourceInfoVid->Kind, BufferSourceInfoVid->Language, BufferSourceInfoVid->Codec);
 							}
 						}
 						if (BufferSourceInfoAud.IsValid())
 						{
-							SelectedStreamAttributesAud.UpdateWith(BufferSourceInfoAud->Kind, BufferSourceInfoAud->Language, BufferSourceInfoAud->HardIndex);
+							SelectedStreamAttributesAud.UpdateWith(BufferSourceInfoAud->Kind, BufferSourceInfoAud->Language, BufferSourceInfoAud->Codec, BufferSourceInfoAud->HardIndex);
 							if (ManifestType != EMediaFormatType::ISOBMFF)
 							{
-								StreamSelectionAttributesAud.UpdateIfOverrideSet(BufferSourceInfoAud->Kind, BufferSourceInfoAud->Language);
+								StreamSelectionAttributesAud.UpdateIfOverrideSet(BufferSourceInfoAud->Kind, BufferSourceInfoAud->Language, BufferSourceInfoAud->Codec);
 							}
 						}
 						if (BufferSourceInfoTxt.IsValid())
 						{
-							SelectedStreamAttributesTxt.UpdateWith(BufferSourceInfoTxt->Kind, BufferSourceInfoTxt->Language, BufferSourceInfoTxt->HardIndex);
+							SelectedStreamAttributesTxt.UpdateWith(BufferSourceInfoTxt->Kind, BufferSourceInfoTxt->Language, BufferSourceInfoTxt->Codec, BufferSourceInfoTxt->HardIndex);
 							if (ManifestType != EMediaFormatType::ISOBMFF)
 							{
-								StreamSelectionAttributesTxt.UpdateIfOverrideSet(BufferSourceInfoTxt->Kind, BufferSourceInfoTxt->Language);
+								StreamSelectionAttributesTxt.UpdateIfOverrideSet(BufferSourceInfoTxt->Kind, BufferSourceInfoTxt->Language, BufferSourceInfoTxt->Codec);
 							}
 						}
 
@@ -2256,6 +2269,10 @@ void FAdaptiveStreamingPlayer::InternalHandlePendingStartRequest(const FTimeValu
 									// Add this period to the list of upcoming ones. This ensures the period metadata gets reported
 									// if the period playback starts in is not the first for which metadata was initially reported.
 									AddUpcomingPeriod(InitialPlayPeriod);
+									// Remember the selected buffer source infos.
+									UpdatePeriodStreamBufferSourceInfo(InitialPlayPeriod, EStreamType::Video, BufferSourceInfoVid);
+									UpdatePeriodStreamBufferSourceInfo(InitialPlayPeriod, EStreamType::Audio, BufferSourceInfoAud);
+									UpdatePeriodStreamBufferSourceInfo(InitialPlayPeriod, EStreamType::Subtitle, BufferSourceInfoTxt);
 
 									PendingFirstSegmentRequest = FirstSegmentRequest;
 									// No longer need the start request.
@@ -2484,6 +2501,19 @@ void FAdaptiveStreamingPlayer::HandlePendingMediaSegmentRequests()
 	InternalHandlePendingFirstSegmentRequest(CurrentTime);
 	InternalHandleSegmentTrackChanges(CurrentTime);
 
+	FTimeValue CurrentPlaybackPos = GetPlayPosition();
+
+	// Add the requests that are play position triggered to the pending requests to process next.
+	for(int32 i=0; i<PlayPosPendingSegmentRequests.Num(); ++i)
+	{
+		if (CurrentPlaybackPos >= PlayPosPendingSegmentRequests[i].StartoverPosition.Time)
+		{
+			NextPendingSegmentRequests.Enqueue(PlayPosPendingSegmentRequests[i]);
+			PlayPosPendingSegmentRequests.RemoveAtSwap(i);
+			--i;
+		}
+	}
+
 	// Are there completed downloads for which we need to find the next segment?
 	TQueue<FPendingSegmentRequest> pendingRequests;
 	Swap(pendingRequests, NextPendingSegmentRequests);
@@ -2526,7 +2556,7 @@ void FAdaptiveStreamingPlayer::HandlePendingMediaSegmentRequests()
 		else if (SegmentPlayPeriod.IsValid())
 		{
 			// Is this a startover for a new period?
-			if (FinishedReq.bStartOver && !FinishedReq.Period.IsValid())
+			if ((FinishedReq.bStartOver || FinishedReq.bPlayPosAutoReselect) && !FinishedReq.Period.IsValid())
 			{
 				// Locate the period the startover time is in.
 				TSharedPtrTS<IManifest::IPlayPeriod> StartoverPeriod;
@@ -2592,7 +2622,7 @@ void FAdaptiveStreamingPlayer::HandlePendingMediaSegmentRequests()
 						CurrentPlayPeriodVideo = FinishedReq.Period;
 						StreamSelector->SetCurrentPlaybackPeriod(EStreamType::Video, CurrentPlayPeriodVideo);
 						UpdateStreamResolutionLimit();
-						if (!FinishedReq.bStartOver)
+						if (!FinishedReq.bStartOver || FinishedReq.bPlayPosAutoReselect)
 						{
 							MultiStreamBufferVid.AddUpcomingBuffer(SegmentPlayPeriod->GetSelectedStreamBufferSourceInfo(StreamType));
 						}
@@ -2602,7 +2632,7 @@ void FAdaptiveStreamingPlayer::HandlePendingMediaSegmentRequests()
 						SegmentPlayPeriod = FinishedReq.Period;
 						CurrentPlayPeriodAudio = FinishedReq.Period;
 						StreamSelector->SetCurrentPlaybackPeriod(EStreamType::Audio, CurrentPlayPeriodAudio);
-						if (!FinishedReq.bStartOver)
+						if (!FinishedReq.bStartOver || FinishedReq.bPlayPosAutoReselect)
 						{
 							MultiStreamBufferAud.AddUpcomingBuffer(SegmentPlayPeriod->GetSelectedStreamBufferSourceInfo(StreamType));
 						}
@@ -2612,17 +2642,18 @@ void FAdaptiveStreamingPlayer::HandlePendingMediaSegmentRequests()
 						SegmentPlayPeriod = FinishedReq.Period;
 						CurrentPlayPeriodText = FinishedReq.Period;
 						//StreamSelector->SetCurrentPlaybackPeriod(EStreamType::Subtitle, CurrentPlayPeriodText);
-						if (!FinishedReq.bStartOver)
+						if (!FinishedReq.bStartOver || FinishedReq.bPlayPosAutoReselect)
 						{
 							MultiStreamBufferTxt.AddUpcomingBuffer(SegmentPlayPeriod->GetSelectedStreamBufferSourceInfo(StreamType));
 						}
 					}
 					// Add this to the upcoming periods the play position will move into and metadata will need to be updated then
 					// unless this is a startover request.
-					if (!FinishedReq.bStartOver)
+					if (!FinishedReq.bStartOver || FinishedReq.bPlayPosAutoReselect)
 					{
 						AddUpcomingPeriod(SegmentPlayPeriod);
 					}
+					UpdatePeriodStreamBufferSourceInfo(SegmentPlayPeriod, StreamType, SegmentPlayPeriod->GetSelectedStreamBufferSourceInfo(StreamType));
 				}
 			}
 
@@ -2641,7 +2672,7 @@ void FAdaptiveStreamingPlayer::HandlePendingMediaSegmentRequests()
 
 			TSharedPtrTS<IStreamSegment> NextSegment;
 			IManifest::FResult Result;
-			if (!FinishedReq.bStartOver)
+			if (!FinishedReq.bStartOver && !FinishedReq.bPlayPosAutoReselect)
 			{
 				if (Action == IAdaptiveStreamSelector::ESegmentAction::FetchNext)
 				{
@@ -2671,43 +2702,47 @@ void FAdaptiveStreamingPlayer::HandlePendingMediaSegmentRequests()
 					{
 						// Tell the multi stream buffer where to switch to and update the selection defaults.
 						TSharedPtrTS<FBufferSourceInfo> BufferSourceInfo = SegmentPlayPeriod->GetSelectedStreamBufferSourceInfo(StreamType);
-						switch(StreamType)
+						if (BufferSourceInfo.IsValid())
 						{
-							case EStreamType::Video:
+							switch(StreamType)
 							{
-								MultiStreamBufferVid.SelectTrackWhenAvailable(BufferSourceInfo);
-								SelectedStreamAttributesVid.UpdateWith(BufferSourceInfo->Kind, BufferSourceInfo->Language, BufferSourceInfo->HardIndex);
-								if (ManifestType != EMediaFormatType::ISOBMFF)
+								case EStreamType::Video:
 								{
-									StreamSelectionAttributesVid.ClearOverrideIndex();
+									MultiStreamBufferVid.SelectTrackWhenAvailable(BufferSourceInfo);
+									SelectedStreamAttributesVid.UpdateWith(BufferSourceInfo->Kind, BufferSourceInfo->Language, BufferSourceInfo->Codec, BufferSourceInfo->HardIndex);
+									if (ManifestType != EMediaFormatType::ISOBMFF)
+									{
+										StreamSelectionAttributesVid.ClearOverrideIndex();
+									}
+									break;
 								}
-								break;
-							}
-							case EStreamType::Audio:
-							{
-								MultiStreamBufferAud.SelectTrackWhenAvailable(BufferSourceInfo);
-								SelectedStreamAttributesAud.UpdateWith(BufferSourceInfo->Kind, BufferSourceInfo->Language, BufferSourceInfo->HardIndex);
-								if (ManifestType != EMediaFormatType::ISOBMFF)
+								case EStreamType::Audio:
 								{
-									StreamSelectionAttributesAud.ClearOverrideIndex();
+									MultiStreamBufferAud.SelectTrackWhenAvailable(BufferSourceInfo);
+									SelectedStreamAttributesAud.UpdateWith(BufferSourceInfo->Kind, BufferSourceInfo->Language, BufferSourceInfo->Codec, BufferSourceInfo->HardIndex);
+									if (ManifestType != EMediaFormatType::ISOBMFF)
+									{
+										StreamSelectionAttributesAud.ClearOverrideIndex();
+									}
+									break;
 								}
-								break;
-							}
-							case EStreamType::Subtitle:
-							{
-								MultiStreamBufferTxt.SelectTrackWhenAvailable(BufferSourceInfo);
-								SelectedStreamAttributesTxt.UpdateWith(BufferSourceInfo->Kind, BufferSourceInfo->Language, BufferSourceInfo->HardIndex);
-								if (ManifestType != EMediaFormatType::ISOBMFF)
+								case EStreamType::Subtitle:
 								{
-									StreamSelectionAttributesTxt.ClearOverrideIndex();
+									MultiStreamBufferTxt.SelectTrackWhenAvailable(BufferSourceInfo);
+									SelectedStreamAttributesTxt.UpdateWith(BufferSourceInfo->Kind, BufferSourceInfo->Language, BufferSourceInfo->Codec, BufferSourceInfo->HardIndex);
+									if (ManifestType != EMediaFormatType::ISOBMFF)
+									{
+										StreamSelectionAttributesTxt.ClearOverrideIndex();
+									}
+									break;
+								}
+								default:
+								{
+									break;
 								}
-								break;
-							}
-							default:
-							{
-								break;
 							}
 						}
+						UpdatePeriodStreamBufferSourceInfo(SegmentPlayPeriod, StreamType, BufferSourceInfo);
 					}
 
 					NextSegment->SetExecutionDelay(ActionDelay);
@@ -2761,6 +2796,35 @@ void FAdaptiveStreamingPlayer::HandlePendingMediaSegmentRequests()
 					break;
 				}
 				case IManifest::FResult::EType::NotFound:
+				{
+					// Not found here should indicate that the type of stream is not available (or not selected) in the new period.
+					if (StreamType == EStreamType::Subtitle)
+					{
+						TSharedPtrTS<ITimelineMediaAsset> ThisPeriod = SegmentPlayPeriod->GetMediaAsset();
+						if (ThisPeriod.IsValid())
+						{
+							FTimeRange ThisPeriodTimeRange = ThisPeriod->GetTimeRange();
+							// Attempt to reselect this stream with the beginning of the following period.
+							FPendingSegmentRequest ReselectReq;
+							ReselectReq.bPlayPosAutoReselect = true;
+							ReselectReq.StartoverPosition.Time = ThisPeriodTimeRange.End;
+							ReselectReq.StreamType = StreamType;
+							PlayPosPendingSegmentRequests.Emplace(MoveTemp(ReselectReq));
+							// Remember that this period has no selected track.
+							UpdatePeriodStreamBufferSourceInfo(SegmentPlayPeriod, StreamType, TSharedPtrTS<FBufferSourceInfo>());
+						}
+					}
+					else
+					{
+						// Throw a playback error for now.
+						FErrorDetail err;
+						err.SetFacility(Facility::EFacility::Player);
+						err.SetMessage("Could not locate next segment");
+						err.SetCode(INTERR_FRAGMENT_NOT_AVAILABLE);
+						PostError(err);
+					}
+					break;
+				}
 				case IManifest::FResult::EType::BeforeStart:
 				case IManifest::FResult::EType::NotLoaded:
 				{
@@ -2803,8 +2867,7 @@ void FAdaptiveStreamingPlayer::HandlePendingMediaSegmentRequests()
 			// Fetch up to 20 seconds ahead of time.
 			const FTimeValue AdvanceFetchThreshold(FTimeValue::MillisecondsToHNS(1000 * 20));
 			FTimeValue NextPTS = ReadyReq->GetFirstPTS();
-			FTimeValue CurrentPos = GetPlayPosition();
-			if (CurrentPos.IsValid() && NextPTS.IsValid() && CurrentPos + AdvanceFetchThreshold < NextPTS)
+			if (CurrentPlaybackPos.IsValid() && NextPTS.IsValid() && CurrentPlaybackPos + AdvanceFetchThreshold < NextPTS)
 			{
 				bExecuteNow = false;
 			}
@@ -2862,6 +2925,14 @@ void FAdaptiveStreamingPlayer::CancelPendingMediaSegmentRequests(EStreamType Str
 			ReadyWaitingSegmentRequests.Enqueue(ReadyReq);
 		}
 	}
+	for(int32 i=0; i<PlayPosPendingSegmentRequests.Num(); ++i)
+	{
+		if (PlayPosPendingSegmentRequests[i].StreamType == StreamType)
+		{
+			PlayPosPendingSegmentRequests.RemoveAtSwap(i);
+			--i;
+		}
+	}
 }
 
 
@@ -2899,7 +2970,7 @@ void FAdaptiveStreamingPlayer::InternalHandleSegmentTrackChanges(const FTimeValu
 				MultiStreamBufferAud.SelectTrackWhenAvailable(BufferSourceInfo);
 				StreamSelectionAttributesAud = *PendingTrackSelectionAud;
 			}
-			SelectedStreamAttributesAud.UpdateWith(BufferSourceInfo->Kind, BufferSourceInfo->Language, BufferSourceInfo->HardIndex);
+			SelectedStreamAttributesAud.UpdateWith(BufferSourceInfo->Kind, BufferSourceInfo->Language, BufferSourceInfo->Codec, BufferSourceInfo->HardIndex);
 			PendingTrackSelectionAud.Reset();
 		}
 
@@ -2912,7 +2983,7 @@ void FAdaptiveStreamingPlayer::InternalHandleSegmentTrackChanges(const FTimeValu
 				MultiStreamBufferTxt.SelectTrackWhenAvailable(BufferSourceInfo);
 				StreamSelectionAttributesTxt = *PendingTrackSelectionTxt;
 			}
-			SelectedStreamAttributesTxt.UpdateWith(BufferSourceInfo->Kind, BufferSourceInfo->Language, BufferSourceInfo->HardIndex);
+			SelectedStreamAttributesTxt.UpdateWith(BufferSourceInfo->Kind, BufferSourceInfo->Language, BufferSourceInfo->Codec, BufferSourceInfo->HardIndex);
 			SelectedStreamAttributesTxt.Select();
 			PendingTrackSelectionTxt.Reset();
 		}
@@ -2927,7 +2998,7 @@ void FAdaptiveStreamingPlayer::InternalHandleSegmentTrackChanges(const FTimeValu
 			TSharedPtrTS<FBufferSourceInfo> BufferSourceInfo = CurrentPlayPeriodAudio->GetSelectedStreamBufferSourceInfo(EStreamType::Audio);
 			MultiStreamBufferAud.SelectTrackWhenAvailable(BufferSourceInfo);
 			StreamSelectionAttributesAud = *PendingTrackSelectionAud;
-			SelectedStreamAttributesAud.UpdateWith(BufferSourceInfo->Kind, BufferSourceInfo->Language, BufferSourceInfo->HardIndex);
+			SelectedStreamAttributesAud.UpdateWith(BufferSourceInfo->Kind, BufferSourceInfo->Language, BufferSourceInfo->Codec, BufferSourceInfo->HardIndex);
 			PendingTrackSelectionAud.Reset();
 		}
 		// Ignore video and subtitle changes for now.
@@ -3009,12 +3080,28 @@ void FAdaptiveStreamingPlayer::InternalHandleSegmentTrackChanges(const FTimeValu
  */
 void FAdaptiveStreamingPlayer::AddUpcomingPeriod(TSharedPtrTS<IManifest::IPlayPeriod> InUpcomingPeriod)
 {
+	/*
+		We maintain two lists of periods.
+		 - ActivePeriods
+		     holds all the periods for which segments are to be requested. As the play position moves forward
+			 periods in this get removed when they have fallen out of range with a threshold.
+			 This list is iterated when searching for codec information and selected buffer source information.
+		 - UpcomingPeriods
+		     this is used solely to report changes in metadata to the user as playback progresses.
+			 entries here are removed immediately without threshold, so this list tends to be shorter than
+			 ActivePeriods.
+	*/
+
 	TSharedPtrTS<ITimelineMediaAsset> Period = InUpcomingPeriod->GetMediaAsset();
 	if (Period.IsValid())
 	{
 		FString PeriodID = Period->GetUniqueIdentifier();
 		// Only add if it is not already in the list.
-		if (!UpcomingPeriods.ContainsByPredicate([PeriodID](const FPeriodInformation& e){ return e.ID.Equals(PeriodID); }))
+		ActivePeriodCriticalSection.Lock();
+		bool bAlreadyThere = ActivePeriods.ContainsByPredicate([PeriodID](const FPeriodInformation& e){ return e.ID.Equals(PeriodID); });
+		ActivePeriodCriticalSection.Unlock();
+
+		if (!bAlreadyThere && !UpcomingPeriods.ContainsByPredicate([PeriodID](const FPeriodInformation& e){ return e.ID.Equals(PeriodID); }))
 		{
 			FPeriodInformation Next;
 			Next.ID = MoveTemp(PeriodID);
@@ -3028,13 +3115,103 @@ void FAdaptiveStreamingPlayer::AddUpcomingPeriod(TSharedPtrTS<IManifest::IPlayPe
 
 			ActivePeriodCriticalSection.Lock();
 			ActivePeriods.Add(Next);
+			ActivePeriods.Sort([](const FPeriodInformation& e1, const FPeriodInformation& e2)
+			{
+				return e1.TimeRange.Start < e2.TimeRange.Start;
+			});
 			ActivePeriodCriticalSection.Unlock();
+
 			UpcomingPeriods.Emplace(MoveTemp(Next));
+			UpcomingPeriods.Sort([](const FPeriodInformation& e1, const FPeriodInformation& e2)
+			{
+				return e1.TimeRange.Start < e2.TimeRange.Start;
+			});
 		}
 	}
 }
 
 
+//-----------------------------------------------------------------------------
+/**
+ * Updates the active or upcoming period's selected buffer source information
+ * for the given stream type.
+ */
+void FAdaptiveStreamingPlayer::UpdatePeriodStreamBufferSourceInfo(TSharedPtrTS<IManifest::IPlayPeriod> InForPeriod, EStreamType InStreamType, TSharedPtrTS<FBufferSourceInfo> InBufferSourceInfo)
+{
+	TSharedPtrTS<ITimelineMediaAsset> Period = InForPeriod->GetMediaAsset();
+	if (Period.IsValid())
+	{
+		FString PeriodID = Period->GetUniqueIdentifier();
+
+		auto UpdatePeriods = [&](TArray<FPeriodInformation>& Periods) -> void
+		{
+			for(auto &P : Periods)
+			{
+				if (P.ID.Equals(PeriodID))
+				{
+					switch(InStreamType)
+					{
+						case EStreamType::Video:
+							P.BufferSourceInfoVid = InBufferSourceInfo;
+							break;
+						case EStreamType::Audio:
+							P.BufferSourceInfoAud = InBufferSourceInfo;
+							break;
+						case EStreamType::Subtitle:
+							P.BufferSourceInfoTxt = InBufferSourceInfo;
+							break;
+						default:
+							break;
+					}
+				}
+			}
+		};
+
+		// Check if we have to update upcoming periods
+		UpdatePeriods(UpcomingPeriods);
+
+		// Update the active periods if necessary.
+		ActivePeriodCriticalSection.Lock();
+		UpdatePeriods(ActivePeriods);
+		ActivePeriodCriticalSection.Unlock();
+	}
+}
+
+
+
+TSharedPtrTS<FBufferSourceInfo> FAdaptiveStreamingPlayer::GetStreamBufferInfoAtTime(bool& bOutHavePeriods, bool& bOutFoundTime, EStreamType InStreamType, const FTimeValue& InAtTime) const
+{
+	TSharedPtrTS<FBufferSourceInfo> BufferSourceInfo;
+	// Locate the period for the specified time.
+	ActivePeriodCriticalSection.Lock();
+	bOutHavePeriods = ActivePeriods.Num() != 0;
+	bOutFoundTime = false;
+	for(int32 i=0; i<ActivePeriods.Num(); ++i)
+	{
+		if (InAtTime >= ActivePeriods[i].TimeRange.Start && InAtTime < (ActivePeriods[i].TimeRange.End.IsValid() ? ActivePeriods[i].TimeRange.End : FTimeValue::GetPositiveInfinity()))
+		{
+			bOutFoundTime = true;
+			switch(InStreamType)
+			{
+				case EStreamType::Video:
+					BufferSourceInfo = ActivePeriods[i].BufferSourceInfoVid;
+					break;
+				case EStreamType::Audio:
+					BufferSourceInfo = ActivePeriods[i].BufferSourceInfoAud;
+					break;
+				case EStreamType::Subtitle:
+					BufferSourceInfo = ActivePeriods[i].BufferSourceInfoTxt;
+					break;
+				default:
+					break;
+			}
+		}
+	}
+	ActivePeriodCriticalSection.Unlock();
+	return BufferSourceInfo;
+}
+
+
 
 //-----------------------------------------------------------------------------
 /**
@@ -3240,6 +3417,7 @@ bool FAdaptiveStreamingPlayer::InternalStartAt(const FSeekParam& NewPosition)
 	check(ReadyWaitingSegmentRequests.IsEmpty());
 	check(CompletedSegmentRequests.Num() == 0);
 	NextPendingSegmentRequests.Empty();
+	PlayPosPendingSegmentRequests.Empty();
 	ReadyWaitingSegmentRequests.Empty();
 	PendingFirstSegmentRequest.Reset();
 	CompletedSegmentRequests.Empty();
@@ -3307,6 +3485,7 @@ void FAdaptiveStreamingPlayer::InternalPause()
 	// NOTE: A decoder may be in InputNeeded() / FeedDecoder() at this point!
 	//       This is to prevent their next ask for new data.
 	DecoderState = EDecoderState::eDecoder_Paused;
+	SubtitleDecoder.Stop();
 
 	// Stop rendering.
 	StopRendering();
@@ -3339,6 +3518,7 @@ void FAdaptiveStreamingPlayer::InternalResume()
 
 		// Resume decoders.
 		DecoderState = EDecoderState::eDecoder_Running;
+		SubtitleDecoder.Start();
 
 		if (PrerollVars.bIsVeryFirstStart)
 		{
@@ -3393,6 +3573,7 @@ void FAdaptiveStreamingPlayer::InternalStop(bool bHoldCurrentFrame)
 	PendingStartRequest.Reset();
 	PendingFirstSegmentRequest.Reset();
 	NextPendingSegmentRequests.Empty();
+	PlayPosPendingSegmentRequests.Empty();
 	ReadyWaitingSegmentRequests.Empty();
 	CompletedSegmentRequests.Empty();
 	InitialPlayPeriod.Reset();
diff --git a/Engine/Plugins/Media/ElectraPlayer/Source/ElectraPlayerRuntime/Private/Runtime/Player/AdaptiveStreamingPlayerDecoders.cpp b/Engine/Plugins/Media/ElectraPlayer/Source/ElectraPlayerRuntime/Private/Runtime/Player/AdaptiveStreamingPlayerDecoders.cpp
--- a/Engine/Plugins/Media/ElectraPlayer/Source/ElectraPlayerRuntime/Private/Runtime/Player/AdaptiveStreamingPlayerDecoders.cpp
+++ b/Engine/Plugins/Media/ElectraPlayer/Source/ElectraPlayerRuntime/Private/Runtime/Player/AdaptiveStreamingPlayerDecoders.cpp
@@ -179,13 +179,19 @@ bool FAdaptiveStreamingPlayer::FindMatchingStreamInfo(FStreamCodecInformation& O
 		ActivePeriodCriticalSection.Unlock();
 		return false;
 	}
+	bool bFound = false;
 	for(int32 i=0; i<ActivePeriods.Num(); ++i)
 	{
 		if (AtTime >= ActivePeriods[i].TimeRange.Start && AtTime < (ActivePeriods[i].TimeRange.End.IsValid() ? ActivePeriods[i].TimeRange.End : FTimeValue::GetPositiveInfinity()))
 		{
 			Asset = ActivePeriods[i].Period;
+			bFound = true;
 		}
 	}
+	if (!bFound)
+	{
+		Asset = ActivePeriods[0].Period;
+	}
 	ActivePeriodCriticalSection.Unlock();
 	if (Asset.IsValid())
 	{
@@ -498,12 +504,18 @@ int32 FAdaptiveStreamingPlayer::CreateDecoder(EStreamType type)
 				SubtitleDecoder.Decoder = ISubtitleDecoder::Create(AccessUnit);
 				if (SubtitleDecoder.Decoder)
 				{
+					FParamDict Options;
+					Options.Set(TEXT("sendEmptySubtitleDuringGaps"), Electra::FVariantValue(true));
+
 					SubtitleDecoder.Decoder->SetPlayerSessionServices(this);
 					SubtitleDecoder.Decoder->GetDecodedSubtitleReceiveDelegate().BindRaw(this, &FAdaptiveStreamingPlayer::OnDecodedSubtitleReceived);
 					SubtitleDecoder.Decoder->GetDecodedSubtitleFlushDelegate().BindRaw(this, &FAdaptiveStreamingPlayer::OnFlushSubtitleReceivers);
-					SubtitleDecoder.Decoder->Open();
+					SubtitleDecoder.Decoder->Open(Options);
 					// Start the decoder, if decoding is currently running.
-					SubtitleDecoder.Start(DecoderState == EDecoderState::eDecoder_Running);
+					if (DecoderState == EDecoderState::eDecoder_Running)
+					{
+						SubtitleDecoder.Start();
+					}
 				}
 				else
 				{
@@ -720,9 +732,23 @@ void FAdaptiveStreamingPlayer::FeedDecoder(EStreamType Type, FMultiTrackAccessUn
 		if (PeekedAU)
 		{
 			// Change in codec?
-			if (CurrentCodecInfo && PeekedAU->AUCodecData.IsValid() && PeekedAU->AUCodecData->ParsedInfo.GetCodec() != CurrentCodecInfo->GetCodec())
+			if (CurrentCodecInfo && PeekedAU->AUCodecData.IsValid())
+			{
+				// Actual different codec?
+				if (PeekedAU->AUCodecData->ParsedInfo.GetCodec() != CurrentCodecInfo->GetCodec())
+				{
+					bCodecChangeDetected = true;
+				}
+				// Mimetype change in subtitles?
+				else if (Type == EStreamType::Subtitle && !PeekedAU->AUCodecData->ParsedInfo.GetMimeType().Equals(CurrentCodecInfo->GetMimeType()))
+				{
+					bCodecChangeDetected = true;
+				}
+			}
+
+			// Is there a change?
+			if (bCodecChangeDetected)
 			{
-				bCodecChangeDetected = true;
 				if (Decoder)
 				{
 					// Check type of stream. We can currently change the video and subtitle codecs only.
diff --git a/Engine/Plugins/Media/ElectraPlayer/Source/ElectraPlayerRuntime/Private/Runtime/Player/AdaptiveStreamingPlayerInternal.h b/Engine/Plugins/Media/ElectraPlayer/Source/ElectraPlayerRuntime/Private/Runtime/Player/AdaptiveStreamingPlayerInternal.h
--- a/Engine/Plugins/Media/ElectraPlayer/Source/ElectraPlayerRuntime/Private/Runtime/Player/AdaptiveStreamingPlayerInternal.h
+++ b/Engine/Plugins/Media/ElectraPlayer/Source/ElectraPlayerRuntime/Private/Runtime/Player/AdaptiveStreamingPlayerInternal.h
@@ -1003,6 +1003,7 @@ private:
 				Decoder->GetDecodedSubtitleFlushDelegate().Unbind();
 				delete Decoder;
 				Decoder = nullptr;
+				bIsRunning = false;
 			}
 		}
 		void Flush()
@@ -1013,15 +1014,15 @@ private:
 			}
 		}
 		
-		void Start(bool bInitialState=true)
+		void Start()
 		{
-			if (bInitialState != bIsRunning)
+			if (Decoder)
 			{
-				if (Decoder)
+				if (!bIsRunning)
 				{
 					Decoder->Start();
+					bIsRunning = true;
 				}
-				bIsRunning = bInitialState;
 			}
 		}
 		
@@ -1433,8 +1434,9 @@ private:
 	{
 		TSharedPtrTS<IStreamSegment>				Request;
 		FTimeValue									AtTime;
-		TSharedPtrTS<IManifest::IPlayPeriod>		Period;					//!< Set if transitioning between periods. This is the new period that needs to be readied.
-		bool										bStartOver = false;
+		TSharedPtrTS<IManifest::IPlayPeriod>		Period;							//!< Set if transitioning between periods. This is the new period that needs to be readied.
+		bool										bStartOver = false;				//!< True when switching tracks within an ongoing period
+		bool										bPlayPosAutoReselect = false;	//!< True when trying to re-select a stream at the start of a new period when it was not available before.
 		EStreamType									StreamType = EStreamType::Unsupported;
 		FPlayStartPosition							StartoverPosition;
 	};
@@ -1444,6 +1446,11 @@ private:
 		TSharedPtrTS<ITimelineMediaAsset>		Period;
 		FString									ID;
 		FTimeRange								TimeRange;
+
+		// Currently selected buffer source per stream type in this period.
+		TSharedPtrTS<FBufferSourceInfo>			BufferSourceInfoVid;
+		TSharedPtrTS<FBufferSourceInfo>			BufferSourceInfoAud;
+		TSharedPtrTS<FBufferSourceInfo>			BufferSourceInfoTxt;
 	};
 
 
@@ -1487,6 +1494,9 @@ private:
 	bool FindMatchingStreamInfo(FStreamCodecInformation& OutStreamInfo, const FTimeValue& AtTime, int32 MaxWidth, int32 MaxHeight);
 	void UpdateStreamResolutionLimit();
 	void AddUpcomingPeriod(TSharedPtrTS<IManifest::IPlayPeriod> InUpcomingPeriod);
+	void UpdatePeriodStreamBufferSourceInfo(TSharedPtrTS<IManifest::IPlayPeriod> InForPeriod, EStreamType InStreamType, TSharedPtrTS<FBufferSourceInfo> InBufferSourceInfo);
+	TSharedPtrTS<FBufferSourceInfo> GetStreamBufferInfoAtTime(bool& bOutHavePeriods, bool& bOutFoundTime, EStreamType InStreamType, const FTimeValue& InAtTime) const;
+
 
 	// AU memory / generic stream reader memory
 	void* AUAllocate(IAccessUnitMemoryProvider::EDataType type, SIZE_T size, SIZE_T alignment) override;
@@ -1644,7 +1654,7 @@ private:
 	FLoopParam															CurrentLoopParam;
 	TMediaQueueDynamicNoLock<FPlayerLoopState>							NextLoopStates;
 
-	FCriticalSection													ActivePeriodCriticalSection;
+	mutable FCriticalSection											ActivePeriodCriticalSection;
 	TArray<FPeriodInformation>											ActivePeriods;
 	TArray<FPeriodInformation>											UpcomingPeriods;
 	TSharedPtrTS<IManifest::IPlayPeriod>								InitialPlayPeriod;
@@ -1654,6 +1664,7 @@ private:
 	TSharedPtrTS<FPendingStartRequest>									PendingStartRequest;
 	TSharedPtrTS<IStreamSegment>										PendingFirstSegmentRequest;
 	TQueue<FPendingSegmentRequest>										NextPendingSegmentRequests;
+	TArray<FPendingSegmentRequest>										PlayPosPendingSegmentRequests;
 	TQueue<TSharedPtrTS<IStreamSegment>>								ReadyWaitingSegmentRequests;
 	TMultiMap<EStreamType, TSharedPtrTS<IStreamSegment>>				CompletedSegmentRequests;
 	bool																bFirstSegmentRequestIsForLooping;
diff --git a/Engine/Plugins/Media/ElectraPlayer/Source/ElectraPlayerRuntime/Private/Runtime/Player/DASH/ManifestDASH.cpp b/Engine/Plugins/Media/ElectraPlayer/Source/ElectraPlayerRuntime/Private/Runtime/Player/DASH/ManifestDASH.cpp
--- a/Engine/Plugins/Media/ElectraPlayer/Source/ElectraPlayerRuntime/Private/Runtime/Player/DASH/ManifestDASH.cpp
+++ b/Engine/Plugins/Media/ElectraPlayer/Source/ElectraPlayerRuntime/Private/Runtime/Player/DASH/ManifestDASH.cpp
@@ -510,7 +510,7 @@ void FDASHPlayPeriod::PrioritizeSelection(TArray<FPrioritizedSelection>& Selecti
 			Candidate.Priority = NewPriority;
 		}
 	}
-	// Sort forst by descending bitrate?
+	// Sort first by descending bitrate?
 	if (bSortByBitrateDescending)
 	{
 		Selection.StableSort([](const FPrioritizedSelection& a, const FPrioritizedSelection& b)
@@ -584,11 +584,32 @@ TSharedPtrTS<FManifestDASHInternal::FAdaptationSet> FDASHPlayPeriod::SelectAdapt
 				PrioritizeSelection(Selection, StreamType, true, false);
 				if (Selection.Num())
 				{
-					// Take the highest prioritized set.
+					// Take the highest prioritized set for now.
 					SelectedTypeIndex = Selection[0].Index;
+
+					// Check if there is a preferred codec set. If we have multiple AdaptationSets for the same language using different codecs
+					// we pick the one for which there is a preference. If there is none matching the first is chosen.
+					// This is primarily to ensure the same track is kept when seeking or across period boundaries.
+					if (Attributes.Codec.IsSet())
+					{
+						FString PreferredCodec = Attributes.Codec.GetValue();
+						if (!PreferredCodec.IsEmpty())
+						{
+							for(int32 i=0; i<Selection.Num(); ++i)
+							{
+								if (Selection[i].CodecInfo.GetCodecName().Equals(PreferredCodec))
+								{
+									SelectedTypeIndex = i;
+									break;
+								}
+							}
+						}
+					}
+
 					AS = StaticCastSharedPtr<FManifestDASHInternal::FAdaptationSet>(Period->GetAdaptationSetByTypeAndIndex(StreamType, SelectedTypeIndex));
 				}
 			}
+
 			// Matching language not found. Is there an explicit index given?
 			// Note: for now we use the explicit override even if a matching language was already found.
 			//       there's a reason the override is specified, like to enforce a specific codec for the same language.
@@ -625,6 +646,7 @@ TSharedPtrTS<FManifestDASHInternal::FAdaptationSet> FDASHPlayPeriod::SelectAdapt
 			OutBufferSourceInfo = MakeSharedTS<FBufferSourceInfo>();
 			OutBufferSourceInfo->Kind = tm.Kind;
 			OutBufferSourceInfo->Language = tm.Language;
+			OutBufferSourceInfo->Codec = tm.HighestBandwidthCodec.GetCodecName();
 			OutBufferSourceInfo->HardIndex = SelectedTypeIndex;
 			OutBufferSourceInfo->PeriodAdaptationSetID = Period->GetUniqueIdentifier() + TEXT("/") + AS->GetUniqueIdentifier();
 		}
@@ -1004,7 +1026,7 @@ IManifest::FResult FDASHPlayPeriod::GetStartingSegment(TSharedPtrTS<IStreamSegme
 			SearchOpt.PeriodDuration = Period->GetDuration();
 			if (!SearchOpt.PeriodDuration.IsValid() || SearchOpt.PeriodDuration.IsPositiveInfinity())
 			{
-				SearchOpt.PeriodDuration = Manifest->GetLastPeriodEndTime() - AST;
+				SearchOpt.PeriodDuration = Manifest->GetLastPeriodEndTime() - AST - Period->GetStart();
 			}
 			SearchOpt.bHasFollowingPeriod = Period->GetHasFollowingPeriod();
 			SearchOpt.SearchType = SearchType;
@@ -1024,6 +1046,7 @@ IManifest::FResult FDASHPlayPeriod::GetStartingSegment(TSharedPtrTS<IStreamSegme
 			{
 				TSharedPtrTS<FStreamSegmentRequestFMP4DASH> SegmentRequest = MakeSharedTS<FStreamSegmentRequestFMP4DASH>();
 				SegmentRequest->StreamType = ActiveSelection[i].StreamType;
+				SegmentRequest->CodecInfo = Repr->GetCodecInformation();
 				SegmentRequest->Representation = Repr;
 				SegmentRequest->AdaptationSet = Adapt;
 				SegmentRequest->Period = Period;
@@ -1074,6 +1097,7 @@ IManifest::FResult FDASHPlayPeriod::GetStartingSegment(TSharedPtrTS<IStreamSegme
 
 				TSharedPtrTS<FStreamSegmentRequestFMP4DASH> SegmentRequest = MakeSharedTS<FStreamSegmentRequestFMP4DASH>();
 				SegmentRequest->StreamType = ActiveSelection[i].StreamType;
+				SegmentRequest->CodecInfo = Repr->GetCodecInformation();
 				SegmentRequest->Representation = Repr;
 				SegmentRequest->AdaptationSet = Adapt;
 				SegmentRequest->Period = Period;
@@ -1211,6 +1235,12 @@ IManifest::FResult FDASHPlayPeriod::GetNextOrRetrySegment(TSharedPtrTS<IStreamSe
 		SearchOpt.PeriodLocalTime.SetFromND(Current->Segment.Time - Current->Segment.PTO + Current->Segment.Duration*3/4, Current->Segment.Timescale);
 		SearchOpt.PeriodDuration = Period->GetDuration();
 		SearchOpt.SearchType = IManifest::ESearchType::After;
+		// If this representation is side loaded there is nothing to fetch next.
+		// Set the period search time such that the search will have to return EOS.
+		if (Repr.IsValid() && Repr->IsSideloadedSubtitle())
+		{
+			SearchOpt.PeriodLocalTime.SetToPositiveInfinity();
+		}
 	}
 	else if (InNextType == ENextSegType::SamePeriodRetry)
 	{
@@ -1248,7 +1278,15 @@ IManifest::FResult FDASHPlayPeriod::GetNextOrRetrySegment(TSharedPtrTS<IStreamSe
 	}
 	SearchOpt.bHasFollowingPeriod = Period->GetHasFollowingPeriod();
 
-	FManifestDASHInternal::FRepresentation::ESearchResult SearchResult = Repr->FindSegment(PlayerSessionServices, SegmentInfo, RemoteElementLoadRequests, SearchOpt);
+	FManifestDASHInternal::FRepresentation::ESearchResult SearchResult;
+	if (Repr.IsValid())
+	{
+		SearchResult = Repr->FindSegment(PlayerSessionServices, SegmentInfo, RemoteElementLoadRequests, SearchOpt);
+	}
+	else
+	{
+		SearchResult = FManifestDASHInternal::FRepresentation::ESearchResult::PastEOS;
+	}
 	if (SearchResult == FManifestDASHInternal::FRepresentation::ESearchResult::NeedElement)
 	{
 		TSharedPtrTS<IPlaylistReader> ManifestReader = PlayerSessionServices->GetManifestReader();
@@ -1296,6 +1334,7 @@ IManifest::FResult FDASHPlayPeriod::GetNextOrRetrySegment(TSharedPtrTS<IStreamSe
 		TSharedPtrTS<FStreamSegmentRequestFMP4DASH> SegmentRequest = MakeSharedTS<FStreamSegmentRequestFMP4DASH>();
 		SegmentRequest->PlayerLoopState = Current->PlayerLoopState;
 		SegmentRequest->StreamType = Current->GetType();
+		SegmentRequest->CodecInfo = Repr->GetCodecInformation();
 		SegmentRequest->Representation = Repr;
 		SegmentRequest->AdaptationSet = Adapt;
 		SegmentRequest->Period = Period;
@@ -2087,32 +2126,40 @@ FManifestDASHInternal::FRepresentation::ESearchResult FManifestDASHInternal::FRe
 			}
 		}
 	}
-	// Remove empty hierarchy levels
-	SegmentBase.Remove(nullptr);
-	SegmentTemplate.Remove(nullptr);
-	// Nothing? Bad MPD.
-	if (SegmentBase.Num() == 0 && SegmentTemplate.Num() == 0)
-	{
-		PostError(InPlayerSessionServices, FString::Printf(TEXT("Representation \"%s\" does not have one of <SegmentBase> or <SegmentTemplate> anywhere in the MPD hierarchy!"), *MPDRepresentation->GetID()), ERRCODE_DASH_MPD_BAD_REPRESENTATION);
-		bIsUsable = false;
-		return FManifestDASHInternal::FRepresentation::ESearchResult::BadType;
-	}
 
-	if (SegmentBase.Num())
+	if (IsSideloadedSubtitle())
 	{
-		return FindSegment_Base(InPlayerSessionServices, OutSegmentInfo, OutRemoteElementLoadRequests, InSearchOptions, MPDRepresentation, SegmentBase);
+		return SetupSideloadedFile(InPlayerSessionServices, OutSegmentInfo, InSearchOptions, MPDRepresentation);
 	}
 	else
 	{
-		// Get the segment timeline, if one is used.
-		TSharedPtrTS<FDashMPD_SegmentTimelineType> SegmentTimeline = GET_ATTR(SegmentTemplate, GetSegmentTimeline(), IsValid(), TSharedPtrTS<FDashMPD_SegmentTimelineType>());
-		if (SegmentTimeline.IsValid())
+		// Remove empty hierarchy levels
+		SegmentBase.Remove(nullptr);
+		SegmentTemplate.Remove(nullptr);
+		// Nothing? Bad MPD.
+		if (SegmentBase.Num() == 0 && SegmentTemplate.Num() == 0)
+		{
+			PostError(InPlayerSessionServices, FString::Printf(TEXT("Representation \"%s\" does not have one of <SegmentBase> or <SegmentTemplate> anywhere in the MPD hierarchy!"), *MPDRepresentation->GetID()), ERRCODE_DASH_MPD_BAD_REPRESENTATION);
+			bIsUsable = false;
+			return FManifestDASHInternal::FRepresentation::ESearchResult::BadType;
+		}
+
+		if (SegmentBase.Num())
 		{
-			return FindSegment_Timeline(InPlayerSessionServices, OutSegmentInfo, OutRemoteElementLoadRequests, InSearchOptions, MPDRepresentation, SegmentTemplate, SegmentTimeline);
+			return FindSegment_Base(InPlayerSessionServices, OutSegmentInfo, OutRemoteElementLoadRequests, InSearchOptions, MPDRepresentation, SegmentBase);
 		}
 		else
 		{
-			return FindSegment_Template(InPlayerSessionServices, OutSegmentInfo, OutRemoteElementLoadRequests, InSearchOptions, MPDRepresentation, SegmentTemplate);
+			// Get the segment timeline, if one is used.
+			TSharedPtrTS<FDashMPD_SegmentTimelineType> SegmentTimeline = GET_ATTR(SegmentTemplate, GetSegmentTimeline(), IsValid(), TSharedPtrTS<FDashMPD_SegmentTimelineType>());
+			if (SegmentTimeline.IsValid())
+			{
+				return FindSegment_Timeline(InPlayerSessionServices, OutSegmentInfo, OutRemoteElementLoadRequests, InSearchOptions, MPDRepresentation, SegmentTemplate, SegmentTimeline);
+			}
+			else
+			{
+				return FindSegment_Template(InPlayerSessionServices, OutSegmentInfo, OutRemoteElementLoadRequests, InSearchOptions, MPDRepresentation, SegmentTemplate);
+			}
 		}
 	}
 }
@@ -2369,8 +2416,8 @@ FManifestDASHInternal::FRepresentation::ESearchResult FManifestDASHInternal::FRe
 	}
 	uint32 SegmentDuration = Duration.Value();
 
-	// Get the period local time into media local timescale
-	int64 MediaLocalSearchTime = InSearchOptions.PeriodLocalTime.GetAsTimebase(MPDTimescale);
+	// Get the period local time into media local timescale and add the PTO.
+	int64 MediaLocalSearchTime = InSearchOptions.PeriodLocalTime.GetAsTimebase(MPDTimescale) + PTO;
 	// If the first media segment does not fall onto the period start there will be an EPT delta that is usually negative.
 	// To simplify calculation of the segment index we shift the search time such that 0 would correspond to the EPT.
 	int32 EPTdelta = EptDelta.GetWithDefault(0);
@@ -2380,16 +2427,16 @@ FManifestDASHInternal::FRepresentation::ESearchResult FManifestDASHInternal::FRe
 		MediaLocalSearchTime = 0;
 	}
 
-	int64 MediaLocalPeriodDuration;
+	int64 MediaLocalPeriodEnd;
 	uint32 MaxSegmentsInPeriod;
 	if (InSearchOptions.PeriodDuration.IsValid() && !InSearchOptions.PeriodDuration.IsInfinity())
 	{
-		MediaLocalPeriodDuration = InSearchOptions.PeriodDuration.GetAsTimebase(MPDTimescale) - EPTdelta;
-		MaxSegmentsInPeriod = (MediaLocalPeriodDuration + SegmentDuration - 1) / SegmentDuration;
+		MediaLocalPeriodEnd = InSearchOptions.PeriodDuration.GetAsTimebase(MPDTimescale) - EPTdelta + PTO;
+		MaxSegmentsInPeriod = (MediaLocalPeriodEnd + SegmentDuration - 1) / SegmentDuration;
 	}
 	else
 	{
-		MediaLocalPeriodDuration = TNumericLimits<int64>::Max() - PTO;
+		MediaLocalPeriodEnd = TNumericLimits<int64>::Max();
 		MaxSegmentsInPeriod = TNumericLimits<uint32>::Max();
 	}
 	// Clamp against the number of segments described by EndNumber.
@@ -2460,13 +2507,13 @@ FManifestDASHInternal::FRepresentation::ESearchResult FManifestDASHInternal::FRe
 		return FManifestDASHInternal::FRepresentation::ESearchResult::PastEOS;
 	}
 
-	OutSegmentInfo.Time = PTO + EPTdelta + SegmentNum * (int64)SegmentDuration;
+	OutSegmentInfo.Time = EPTdelta + SegmentNum * (int64)SegmentDuration;
 	OutSegmentInfo.PTO = PTO;
 	OutSegmentInfo.EPTdelta = EPTdelta;
 	OutSegmentInfo.Duration = SegmentDuration;
 	OutSegmentInfo.Number = StartNumber + SegmentNum;
-	OutSegmentInfo.MediaLocalFirstAUTime = MediaLocalSearchTime + PTO;
-	OutSegmentInfo.MediaLocalLastAUTime = MediaLocalPeriodDuration + PTO;
+	OutSegmentInfo.MediaLocalFirstAUTime = MediaLocalSearchTime;
+	OutSegmentInfo.MediaLocalLastAUTime = MediaLocalPeriodEnd;
 	OutSegmentInfo.Timescale = MPDTimescale;
 	OutSegmentInfo.bMayBeMissing = SegmentNum + 1 >= MaxSegmentsInPeriod;
 	OutSegmentInfo.bIsLastInPeriod = OutSegmentInfo.bMayBeMissing && InSearchOptions.bHasFollowingPeriod;
@@ -2816,6 +2863,67 @@ FManifestDASHInternal::FRepresentation::ESearchResult FManifestDASHInternal::FRe
 }
 
 
+FManifestDASHInternal::FRepresentation::ESearchResult FManifestDASHInternal::FRepresentation::SetupSideloadedFile(IPlayerSessionServices* InPlayerSessionServices, FSegmentInformation& OutSegmentInfo, const FSegmentSearchOption& InSearchOptions, const TSharedPtrTS<FDashMPD_RepresentationType>& MPDRepresentation)
+{
+	// As per DASH-IF-IOP 4.3, Section 6.4.5. Guidelines for side-loaded TTML and WebVTT files
+	// side loaded subtitles must not use @presentationTimeOffset, which makes sense since this attribute is not defined on <Representation> elements.
+	// However, we allow its use as well as a @timescale to give the value in.
+	uint32 Timescale = 1000;
+	uint64 PTO = 0;
+	const TArray<IDashMPDElement::FXmlAttribute>& OtherAttributes = MPDRepresentation->GetOtherAttributes();
+	for(auto &Attribute : OtherAttributes)
+	{
+		if (Attribute.GetName().Equals(TEXT("presentationTimeOffset")))
+		{
+			LexFromString(PTO, *Attribute.GetValue());
+		}
+		else if (Attribute.GetName().Equals(TEXT("timescale")))
+		{
+			LexFromString(Timescale, *Attribute.GetValue());
+		}
+	}
+	if (!Timescale)
+	{
+		Timescale = 1;
+	}
+
+	int64 MediaLocalSearchTime = InSearchOptions.PeriodLocalTime.GetAsTimebase(Timescale);
+	if (MediaLocalSearchTime < 0)
+	{
+		MediaLocalSearchTime = 0;
+	}
+	int64 MediaLocalPeriodEnd = InSearchOptions.PeriodDuration.IsValid() && !InSearchOptions.PeriodDuration.IsInfinity() ? InSearchOptions.PeriodDuration.GetAsTimebase(Timescale) : TNumericLimits<int64>::Max();
+	if (MediaLocalSearchTime >= MediaLocalPeriodEnd)
+	{
+		return FManifestDASHInternal::FRepresentation::ESearchResult::PastEOS;
+	}
+
+	// The time must be set to the PTO so it won't cause a problem searching for the next period (where the PTO is subtracted from the time)
+	OutSegmentInfo.Time = (int64) PTO;
+	OutSegmentInfo.PTO = (int64) PTO;
+	OutSegmentInfo.Duration = MediaLocalPeriodEnd;
+	OutSegmentInfo.Number = 0;
+	OutSegmentInfo.NumberOfBytes = 0;
+	OutSegmentInfo.FirstByteOffset = 0;
+	OutSegmentInfo.MediaLocalFirstAUTime = MediaLocalSearchTime;
+	OutSegmentInfo.MediaLocalLastAUTime = MediaLocalPeriodEnd;
+	OutSegmentInfo.Timescale = Timescale;
+	OutSegmentInfo.ATO.SetToZero();
+	OutSegmentInfo.bIsSideload = true;
+	OutSegmentInfo.bIsLastInPeriod = true;
+	TArray<TSharedPtrTS<FDashMPD_SegmentBaseType>> SegmentBase;
+	if (!PrepareDownloadURLs(InPlayerSessionServices, OutSegmentInfo, SegmentBase))
+	{
+		bIsUsable = false;
+		return FManifestDASHInternal::FRepresentation::ESearchResult::BadType;
+	}
+	else
+	{
+		return FManifestDASHInternal::FRepresentation::ESearchResult::Found;
+	}
+}
+
+
 
 void FManifestDASHInternal::FRepresentation::GetSegmentInformation(TArray<IManifest::IPlayPeriod::FSegmentInformation>& OutSegmentInformation, FTimeValue& OutAverageSegmentDuration, TSharedPtrTS<const IStreamSegment> CurrentSegment, const FTimeValue& LookAheadTime, const TSharedPtrTS<IPlaybackAssetAdaptationSet>& AdaptationSet)
 {
diff --git a/Engine/Plugins/Media/ElectraPlayer/Source/ElectraPlayerRuntime/Private/Runtime/Player/DASH/PlaylistReaderDASH.cpp b/Engine/Plugins/Media/ElectraPlayer/Source/ElectraPlayerRuntime/Private/Runtime/Player/DASH/PlaylistReaderDASH.cpp
--- a/Engine/Plugins/Media/ElectraPlayer/Source/ElectraPlayerRuntime/Private/Runtime/Player/DASH/PlaylistReaderDASH.cpp
+++ b/Engine/Plugins/Media/ElectraPlayer/Source/ElectraPlayerRuntime/Private/Runtime/Player/DASH/PlaylistReaderDASH.cpp
@@ -1089,7 +1089,8 @@ FErrorDetail FPlaylistReaderDASH::GetXMLResponseString(FString& OutXMLString, FR
 			// UTF-32 LE BOM
 			return CreateErrorAndLog(FString::Printf(TEXT("Document has unsupported UTF-32 LE BOM!")), ERRCODE_DASH_MPD_UNSUPPORTED_DOCUMENT_ENCODING);
 		}
-		FString XML(NumResponseBytes, (TCHAR*)FUTF8ToTCHAR((const ANSICHAR*)ResponseBytes, NumResponseBytes).Get());
+		FUTF8ToTCHAR TextConv((const ANSICHAR*)ResponseBytes, NumResponseBytes);
+		FString XML(TextConv.Length(), TextConv.Get());
 		OutXMLString = MoveTemp(XML);
 	}
 	return FErrorDetail();
@@ -1103,7 +1104,8 @@ FErrorDetail FPlaylistReaderDASH::GetResponseString(FString& OutString, FResourc
 		TSharedPtrTS<IElectraHttpManager::FReceiveBuffer> ResponseBuffer = FromRequest->Request->GetResponseBuffer();
 		int32 NumResponseBytes = ResponseBuffer->Buffer.Num();
 		const uint8* ResponseBytes = (const uint8*)ResponseBuffer->Buffer.GetLinearReadData();
-		FString UTF8Text(NumResponseBytes, (TCHAR*)FUTF8ToTCHAR((const ANSICHAR*)ResponseBytes, NumResponseBytes).Get());
+		FUTF8ToTCHAR TextConv((const ANSICHAR*)ResponseBytes, NumResponseBytes);
+		FString UTF8Text(TextConv.Length(), TextConv.Get());
 		OutString = MoveTemp(UTF8Text);
 	}
 	return FErrorDetail();
diff --git a/Engine/Plugins/Media/ElectraPlayer/Source/ElectraPlayerRuntime/Private/Runtime/Player/DASH/PlaylistReaderDASH_Internal.h b/Engine/Plugins/Media/ElectraPlayer/Source/ElectraPlayerRuntime/Private/Runtime/Player/DASH/PlaylistReaderDASH_Internal.h
--- a/Engine/Plugins/Media/ElectraPlayer/Source/ElectraPlayerRuntime/Private/Runtime/Player/DASH/PlaylistReaderDASH_Internal.h
+++ b/Engine/Plugins/Media/ElectraPlayer/Source/ElectraPlayerRuntime/Private/Runtime/Player/DASH/PlaylistReaderDASH_Internal.h
@@ -34,7 +34,8 @@ struct FMPDLoadRequestDASH : public IHTTPResourceRequestObject
 		XLink_InitializationSet,
 		Callback,
 		Segment,
-		TimeSync
+		TimeSync,
+		Sideload
 	};
 	const TCHAR* const GetRequestTypeName() const
 	{
@@ -51,6 +52,7 @@ struct FMPDLoadRequestDASH : public IHTTPResourceRequestObject
 			case ELoadType::Callback:				return TEXT("Callback");
 			case ELoadType::Segment:				return TEXT("Segment");
 			case ELoadType::TimeSync:				return TEXT("Time sync");
+			case ELoadType::Sideload:				return TEXT("Sideload");
 			default:								return TEXT("<unknown>");
 		}
 	}
@@ -154,6 +156,7 @@ public:
 		int64 MediaLocalFirstAUTime = 0;			//!< Time of the first AU to use in this segment in media local time
 		int64 MediaLocalLastAUTime = 0;				//!< Time at which the last AU to use in thie segment ends in media local time
 		uint32 Timescale = 0;						//!< Local media timescale
+		bool bIsSideload = false;					//!< true if this is a side-loaded resource to be fetched and cached.
 		bool bIsLastInPeriod = false;				//!< true if known to be the last segment in the period.
 		bool bMayBeMissing = false;					//!< true if the last segment in <SegmentTemplate> that might not exist.
 		bool bIsMissing = false;					//!< Set to true if known to be missing.
@@ -261,6 +264,7 @@ public:
 		ESearchResult FindSegment_Base(IPlayerSessionServices* PlayerSessionServices, FSegmentInformation& OutSegmentInfo, TArray<TWeakPtrTS<FMPDLoadRequestDASH>>& OutRemoteElementLoadRequests, const FSegmentSearchOption& InSearchOptions, const TSharedPtrTS<FDashMPD_RepresentationType>& MPDRepresentation, const TArray<TSharedPtrTS<FDashMPD_SegmentBaseType>>& SegmentBase);
 		ESearchResult FindSegment_Template(IPlayerSessionServices* PlayerSessionServices, FSegmentInformation& OutSegmentInfo, TArray<TWeakPtrTS<FMPDLoadRequestDASH>>& OutRemoteElementLoadRequests, const FSegmentSearchOption& InSearchOptions, const TSharedPtrTS<FDashMPD_RepresentationType>& MPDRepresentation, const TArray<TSharedPtrTS<FDashMPD_SegmentTemplateType>>& SegmentTemplate);
 		ESearchResult FindSegment_Timeline(IPlayerSessionServices* PlayerSessionServices, FSegmentInformation& OutSegmentInfo, TArray<TWeakPtrTS<FMPDLoadRequestDASH>>& OutRemoteElementLoadRequests, const FSegmentSearchOption& InSearchOptions, const TSharedPtrTS<FDashMPD_RepresentationType>& MPDRepresentation, const TArray<TSharedPtrTS<FDashMPD_SegmentTemplateType>>& SegmentTemplate, const TSharedPtrTS<FDashMPD_SegmentTimelineType>& SegmentTimeline);
+		ESearchResult SetupSideloadedFile(IPlayerSessionServices* PlayerSessionServices, FSegmentInformation& OutSegmentInfo, const FSegmentSearchOption& InSearchOptions, const TSharedPtrTS<FDashMPD_RepresentationType>& MPDRepresentation);
 
 		bool PrepareDownloadURLs(IPlayerSessionServices* PlayerSessionServices, FSegmentInformation& InOutSegmentInfo, const TArray<TSharedPtrTS<FDashMPD_SegmentBaseType>>& SegmentBase);
 		bool PrepareDownloadURLs(IPlayerSessionServices* PlayerSessionServices, FSegmentInformation& InOutSegmentInfo, const TArray<TSharedPtrTS<FDashMPD_SegmentTemplateType>>& SegmentTemplate);
diff --git a/Engine/Plugins/Media/ElectraPlayer/Source/ElectraPlayerRuntime/Private/Runtime/Player/DASH/StreamReaderFMP4DASH.cpp b/Engine/Plugins/Media/ElectraPlayer/Source/ElectraPlayerRuntime/Private/Runtime/Player/DASH/StreamReaderFMP4DASH.cpp
--- a/Engine/Plugins/Media/ElectraPlayer/Source/ElectraPlayerRuntime/Private/Runtime/Player/DASH/StreamReaderFMP4DASH.cpp
+++ b/Engine/Plugins/Media/ElectraPlayer/Source/ElectraPlayerRuntime/Private/Runtime/Player/DASH/StreamReaderFMP4DASH.cpp
@@ -557,6 +557,112 @@ FErrorDetail FStreamReaderFMP4DASH::FStreamHandler::GetInitSegment(TSharedPtrTS<
 	}
 }
 
+
+FErrorDetail FStreamReaderFMP4DASH::FStreamHandler::RetrieveSideloadedFile(TSharedPtrTS<const TArray<uint8>>& OutData, const TSharedPtrTS<FStreamSegmentRequestFMP4DASH>& Request)
+{
+	if (Request->Segment.MediaURL.URL.IsEmpty())
+	{
+		return FErrorDetail();
+	}
+
+	// Check if we already have this cached.
+	IPlayerEntityCache::FCacheItem CachedItem;
+	if (PlayerSessionService->GetEntityCache()->GetCachedEntity(CachedItem, Request->Segment.MediaURL.URL, Request->Segment.MediaURL.Range))
+	{
+		// Already cached. Use it.
+		OutData = CachedItem.RawPayloadData;
+		return FErrorDetail();
+	}
+
+	// Get the manifest reader. We need it to handle the download.
+	TSharedPtrTS<IPlaylistReader> ManifestReader = PlayerSessionService->GetManifestReader();
+	if (!ManifestReader.IsValid())
+	{
+		return PostError(PlayerSessionService, TEXT("Entity loader disappeared"), 0);
+	}
+
+	// Create a download finished structure we can wait upon for completion.
+	struct FLoadResult
+	{
+		FMediaEvent Event;
+		bool bSuccess = false;
+	};
+	TSharedPtrTS<FLoadResult> LoadResult = MakeSharedTS<FLoadResult>();
+
+	// Create the download request.
+	TSharedPtrTS<FMPDLoadRequestDASH> LoadReq = MakeSharedTS<FMPDLoadRequestDASH>();
+	LoadReq->LoadType = FMPDLoadRequestDASH::ELoadType::Sideload;
+	LoadReq->URL = Request->Segment.MediaURL.URL;
+	LoadReq->Range = Request->Segment.MediaURL.Range;
+	if (Request->Segment.MediaURL.CustomHeader.Len())
+	{
+		LoadReq->Headers.Emplace(HTTP::FHTTPHeader({DASH::HTTPHeaderOptionName, Request->Segment.MediaURL.CustomHeader}));
+	}
+	LoadReq->PlayerSessionServices = PlayerSessionService;
+	LoadReq->CompleteCallback.BindLambda([LoadResult](TSharedPtrTS<FMPDLoadRequestDASH> LoadRequest, bool bSuccess)
+	{
+		LoadResult->bSuccess = bSuccess;
+		LoadResult->Event.Signal();
+	});
+	// Issue the download request.
+	check(ManifestReader->GetPlaylistType().Equals(TEXT("dash")));
+	IPlaylistReaderDASH* Reader = static_cast<IPlaylistReaderDASH*>(ManifestReader.Get());
+	Reader->AddElementLoadRequests(TArray<TWeakPtrTS<FMPDLoadRequestDASH>>({LoadReq}));
+
+	// Wait for completion or abort
+	while(!HasReadBeenAborted())
+	{
+		if (LoadResult->Event.WaitTimeout(1000 * 100))
+		{
+			break;
+		}
+	}
+	if (HasReadBeenAborted())
+	{
+		return FErrorDetail();
+	}
+	
+	// Set up download stats to be sent to the stream selector.
+	const HTTP::FConnectionInfo* ci = LoadReq->GetConnectionInfo();
+	Metrics::FSegmentDownloadStats ds;
+	ds.StatsID = FMediaInterlockedIncrement(UniqueDownloadID);
+	ds.StreamType = Request->GetType();
+	ds.SegmentType = Metrics::ESegmentType::Media;
+	ds.URL = Request->Segment.MediaURL.URL;
+	ds.Range = Request->Segment.MediaURL.Range;
+	ds.CDN = Request->Segment.MediaURL.CDN;
+	ds.bWasSuccessful = LoadResult->bSuccess;
+	ds.HTTPStatusCode = ci->StatusInfo.HTTPStatus;
+	ds.TimeToFirstByte = ci->TimeUntilFirstByte;
+	ds.TimeToDownload = (ci->RequestEndTime - ci->RequestStartTime).GetAsSeconds();
+	ds.ByteSize = ci->ContentLength;
+	ds.NumBytesDownloaded = ci->BytesReadSoFar;
+	ds.ThroughputBps = ci->Throughput.GetThroughput() ? ci->Throughput.GetThroughput() : (ds.TimeToDownload > 0.0 ? 8 * ds.NumBytesDownloaded / ds.TimeToDownload : 0);
+	ds.MediaAssetID = Request->Period.IsValid() ? Request->Period->GetUniqueIdentifier() : "";
+	ds.AdaptationSetID = Request->AdaptationSet.IsValid() ? Request->AdaptationSet->GetUniqueIdentifier() : "";
+	ds.RepresentationID = Request->Representation.IsValid() ? Request->Representation->GetUniqueIdentifier() : "";
+	ds.Bitrate = Request->GetBitrate();
+
+	if (!LoadResult->bSuccess)
+	{
+		Request->ConnectionInfo = *ci;
+		StreamSelector->ReportDownloadEnd(ds);
+		return CreateError(FString::Printf(TEXT("Sideloaded media download error: %s"), *ci->StatusInfo.ErrorDetail.GetMessage()), INTERNAL_ERROR_INIT_SEGMENT_DOWNLOAD_ERROR);
+	}
+
+	OutData = MakeSharedTS<const TArray<uint8>>(TArrayView<uint8>(LoadReq->Request->GetResponseBuffer()->Buffer.GetLinearReadData(), LoadReq->Request->GetResponseBuffer()->Buffer.GetLinearReadSize()));
+
+	// Add this to the entity cache in case it needs to be retrieved again.
+	IPlayerEntityCache::FCacheItem CacheItem;
+	CacheItem.URL = LoadReq->URL;
+	CacheItem.Range = LoadReq->Range;
+	CacheItem.RawPayloadData = OutData;
+	PlayerSessionService->GetEntityCache()->CacheEntity(CacheItem);
+
+	StreamSelector->ReportDownloadEnd(ds);
+	return FErrorDetail();
+}
+
 void FStreamReaderFMP4DASH::FStreamHandler::CheckForInbandDASHEvents()
 {
 	SCOPE_CYCLE_COUNTER(STAT_ElectraPlayer_DASH_StreamReader);
@@ -766,7 +872,55 @@ void FStreamReaderFMP4DASH::FStreamHandler::HandleRequest()
 
 	if (InitSegmentError.IsOK() && !HasReadBeenAborted())
 	{
-		if (!bIsEmptyFillerSegment)
+		if (Request->Segment.bIsSideload)
+		{
+			FErrorDetail SideloadError;
+			TSharedPtrTS<const TArray<uint8>> SideloadedData;
+			SideloadError = RetrieveSideloadedFile(SideloadedData, Request);
+			if (SideloadError.IsOK())
+			{
+				// CSD is only partially available for sideloaded files.
+				CSD->ParsedInfo = Request->CodecInfo;
+
+				uint32 TrackTimescale = Request->Segment.Timescale;
+
+				// Create an access unit
+				FAccessUnit *AccessUnit = FAccessUnit::Create(Parameters.MemoryProvider);
+				AccessUnit->ESType = Request->GetType();
+				AccessUnit->AUSize = (uint32) SideloadedData->Num();
+				AccessUnit->AUData = AccessUnit->AllocatePayloadBuffer(AccessUnit->AUSize);
+				FMemory::Memcpy(AccessUnit->AUData, SideloadedData->GetData(), SideloadedData->Num());
+				AccessUnit->AUCodecData = CSD;
+				AccessUnit->bIsFirstInSequence = true;
+				AccessUnit->bIsSyncSample = true;
+				AccessUnit->bIsDummyData = false;
+				AccessUnit->bIsSideloaded = true;
+				AccessUnit->BufferSourceInfo = Request->SourceBufferInfo;
+				AccessUnit->PlayerLoopState = PlayerLoopState;
+				AccessUnit->DropState = FAccessUnit::EDropState::None;
+
+				// Sideloaded files coincide with the period start
+				AccessUnit->DTS = TimeOffset;
+				AccessUnit->PTS = TimeOffset;
+				AccessUnit->PTO.SetFromND(Request->Segment.PTO, TrackTimescale);
+
+				FTimeValue Duration(Request->Segment.Duration, TrackTimescale);
+				AccessUnit->Duration = Duration;
+
+				DurationSuccessfullyRead += Duration;
+				NextExpectedDTS = AccessUnit->DTS + Duration;
+				LastKnownAUDuration = Duration;
+				AccessUnitFIFO.Push(AccessUnit);
+				// Clear the pointer since we must not touch this AU again after we handed it off.
+				AccessUnit = nullptr;
+			}
+			else
+			{
+				CurrentRequest->ConnectionInfo.StatusInfo.ErrorDetail = SideloadError;
+				bHasErrored = true;
+			}
+		}
+		else if (!bIsEmptyFillerSegment)
 		{
 			ReadBuffer.Reset();
 			ReadBuffer.ReceiveBuffer = MakeSharedTS<IElectraHttpManager::FReceiveBuffer>();
@@ -924,6 +1078,7 @@ void FStreamReaderFMP4DASH::FStreamHandler::HandleRequest()
 									AccessUnit->DTS += TimeOffset;
 									AccessUnit->PTS.SetFromND(AUPTS - PTO, TrackTimescale);
 									AccessUnit->PTS += TimeOffset;
+									AccessUnit->PTO.SetFromND(PTO, TrackTimescale);
 
 									ElectraCDM::FMediaCDMSampleInfo SampleEncryptionInfo;
 									bool bIsSampleEncrypted = TrackIterator->GetEncryptionInfo(SampleEncryptionInfo);
diff --git a/Engine/Plugins/Media/ElectraPlayer/Source/ElectraPlayerRuntime/Private/Runtime/Player/DASH/StreamReaderFMP4DASH.h b/Engine/Plugins/Media/ElectraPlayer/Source/ElectraPlayerRuntime/Private/Runtime/Player/DASH/StreamReaderFMP4DASH.h
--- a/Engine/Plugins/Media/ElectraPlayer/Source/ElectraPlayerRuntime/Private/Runtime/Player/DASH/StreamReaderFMP4DASH.h
+++ b/Engine/Plugins/Media/ElectraPlayer/Source/ElectraPlayerRuntime/Private/Runtime/Player/DASH/StreamReaderFMP4DASH.h
@@ -40,6 +40,7 @@ public:
 	virtual void GetDownloadStats(Metrics::FSegmentDownloadStats& OutStats) const override;
 
 	EStreamType												StreamType = EStreamType::Unsupported;				//!< Type of stream (video, audio, etc.)
+	FStreamCodecInformation									CodecInfo;											//!< Partial codec info as can be collected from the MPD.
 	TSharedPtrTS<IPlaybackAssetRepresentation>				Representation;										//!< The representation this request belongs to.
 	TSharedPtrTS<IPlaybackAssetAdaptationSet>				AdaptationSet;										//!< The adaptation set the representation belongs to.
 	TSharedPtrTS<ITimelineMediaAsset>						Period;												//!< The period the adaptation set belongs to.
@@ -162,6 +163,7 @@ private:
 		void HandleRequest();
 
 		FErrorDetail GetInitSegment(TSharedPtrTS<const IParserISO14496_12>& OutMP4InitSegment, const TSharedPtrTS<FStreamSegmentRequestFMP4DASH>& InRequest);
+		FErrorDetail RetrieveSideloadedFile(TSharedPtrTS<const TArray<uint8>>& OutData, const TSharedPtrTS<FStreamSegmentRequestFMP4DASH>& InRequest);
 		void CheckForInbandDASHEvents();
 		void HandleEventMessages();
 
diff --git a/Engine/Plugins/Media/ElectraPlayer/Source/ElectraPlayerRuntime/Private/Runtime/Player/HLS/PlaylistReaderHLS.cpp b/Engine/Plugins/Media/ElectraPlayer/Source/ElectraPlayerRuntime/Private/Runtime/Player/HLS/PlaylistReaderHLS.cpp
--- a/Engine/Plugins/Media/ElectraPlayer/Source/ElectraPlayerRuntime/Private/Runtime/Player/HLS/PlaylistReaderHLS.cpp
+++ b/Engine/Plugins/Media/ElectraPlayer/Source/ElectraPlayerRuntime/Private/Runtime/Player/HLS/PlaylistReaderHLS.cpp
@@ -1038,7 +1038,9 @@ FErrorDetail FPlaylistReaderHLS::ParsePlaylist(FPlaylistRequestPtr FromRequest)
 	{
 		int32 RequestBytes = FromRequest->GetReceiveBuffer()->Buffer.Num();
 		// The FromRequest buffer is not zero terminated. We need to pick the correct FString constructor for converting the chars into TCHARs while adding the terminating zero!
-		ParseError = Parser.Parse(FString(RequestBytes, (TCHAR*)FUTF8ToTCHAR((const ANSICHAR*)FromRequest->GetReceiveBuffer()->Buffer.GetLinearReadData(), RequestBytes).Get()), Playlist);
+		FUTF8ToTCHAR TextConv((const ANSICHAR*)FromRequest->GetReceiveBuffer()->Buffer.GetLinearReadData(), RequestBytes);
+		FString UTF8String(TextConv.Length(), TextConv.Get());
+		ParseError = Parser.Parse(UTF8String, Playlist);
 	}
 
 	if (ParseError == HLSPlaylistParser::EPlaylistError::None)
diff --git a/Engine/Plugins/Media/ElectraPlayer/Source/ElectraPlayerRuntime/Private/Runtime/Player/mp4/ManifestMP4.cpp b/Engine/Plugins/Media/ElectraPlayer/Source/ElectraPlayerRuntime/Private/Runtime/Player/mp4/ManifestMP4.cpp
--- a/Engine/Plugins/Media/ElectraPlayer/Source/ElectraPlayerRuntime/Private/Runtime/Player/mp4/ManifestMP4.cpp
+++ b/Engine/Plugins/Media/ElectraPlayer/Source/ElectraPlayerRuntime/Private/Runtime/Player/mp4/ManifestMP4.cpp
@@ -378,6 +378,7 @@ void FManifestMP4Internal::FPlayPeriodMP4::MakeBufferSourceInfoFromMetadata(EStr
 		OutBufferSourceInfo = MakeSharedTS<FBufferSourceInfo>();
 		OutBufferSourceInfo->Kind = InMetadata->Kind;
 		OutBufferSourceInfo->Language = InMetadata->Language;
+		OutBufferSourceInfo->Codec = InMetadata->HighestBandwidthCodec.GetCodecName();
 		TSharedPtrTS<FTimelineAssetMP4> Asset = MediaAsset.Pin();
 		OutBufferSourceInfo->PeriodAdaptationSetID = Asset->GetUniqueIdentifier() + TEXT(".") + InMetadata->ID;
 		TArray<FTrackMetadata> Metadata;
diff --git a/Engine/Plugins/Media/ElectraPlayer/Source/ElectraPlayerRuntime/Private/Runtime/StreamAccessUnitBuffer.cpp b/Engine/Plugins/Media/ElectraPlayer/Source/ElectraPlayerRuntime/Private/Runtime/StreamAccessUnitBuffer.cpp
--- a/Engine/Plugins/Media/ElectraPlayer/Source/ElectraPlayerRuntime/Private/Runtime/StreamAccessUnitBuffer.cpp
+++ b/Engine/Plugins/Media/ElectraPlayer/Source/ElectraPlayerRuntime/Private/Runtime/StreamAccessUnitBuffer.cpp
@@ -216,24 +216,27 @@ namespace Electra
 
 	void FMultiTrackAccessUnitBuffer::AddUpcomingBuffer(TSharedPtrTS<FBufferSourceInfo> InBufferSourceInfo)
 	{
-		FMediaCriticalSection::ScopedLock lock(AccessLock);
-		// Add to the current or the switch-over chain?
-		if (SwitchOverBufferChain.Num())
+		if (InBufferSourceInfo.IsValid())
 		{
-			if (!SwitchOverBufferChain.Last().Info->PeriodAdaptationSetID.Equals(InBufferSourceInfo->PeriodAdaptationSetID))
+			FMediaCriticalSection::ScopedLock lock(AccessLock);
+			// Add to the current or the switch-over chain?
+			if (SwitchOverBufferChain.Num())
 			{
-				FQueuedBuffer Next;
-				Next.Info = MoveTemp(InBufferSourceInfo);
-				SwitchOverBufferChain.Emplace(MoveTemp(Next));
+				if (!SwitchOverBufferChain.Last().Info->PeriodAdaptationSetID.Equals(InBufferSourceInfo->PeriodAdaptationSetID))
+				{
+					FQueuedBuffer Next;
+					Next.Info = MoveTemp(InBufferSourceInfo);
+					SwitchOverBufferChain.Emplace(MoveTemp(Next));
+				}
 			}
-		}
-		else
-		{
-			if (UpcomingBufferChain.Num() == 0 || !UpcomingBufferChain.Last().Info->PeriodAdaptationSetID.Equals(InBufferSourceInfo->PeriodAdaptationSetID))
+			else
 			{
-				FQueuedBuffer Next;
-				Next.Info = MoveTemp(InBufferSourceInfo);
-				UpcomingBufferChain.Emplace(MoveTemp(Next));
+				if (UpcomingBufferChain.Num() == 0 || !UpcomingBufferChain.Last().Info->PeriodAdaptationSetID.Equals(InBufferSourceInfo->PeriodAdaptationSetID))
+				{
+					FQueuedBuffer Next;
+					Next.Info = MoveTemp(InBufferSourceInfo);
+					UpcomingBufferChain.Emplace(MoveTemp(Next));
+				}
 			}
 		}
 	}
@@ -408,6 +411,8 @@ namespace Electra
 				bool bSwitch = false;
 				if (PopUntilDTS.IsValid() && PopUntilPTS.IsValid())
 				{
+					TSharedPtrTS<FAccessUnitBuffer> ActiveBuffer = GetSelectedTrackBuffer();
+					bool bActiveIsEmpty = !ActiveBuffer.IsValid() || ActiveBuffer->Num() == 0;
 					for(int32 nBuf=0; nBuf<SwitchOverBufferChain.Num(); ++nBuf)
 					{
 						if (SwitchOverBufferChain[nBuf].Buffer.IsValid())
@@ -421,6 +426,15 @@ namespace Electra
 								bSwitch = true;
 								break;
 							}
+							// If this buffer is empty and the switchover one is not we switch over now.
+							// This is for the cases where we are not getting any more data from the current buffer
+							// and the LastPoppedDTS/PTS are in the past and will not be contained in the new buffer,
+							// which happens mostly with sparse data like infrequent subtitle AUs.
+							if (bActiveIsEmpty && SwitchOverBufferChain[nBuf].Buffer->Num())
+							{
+								bSwitch = true;
+								break;
+							}
 							// If the buffer is now empty with another non-empty one following we can remove this one.
 							if (SwitchOverBufferChain[nBuf].Buffer->Num() == 0 && nBuf+1 < SwitchOverBufferChain.Num() && SwitchOverBufferChain[nBuf+1].Buffer.IsValid() && SwitchOverBufferChain[nBuf+1].Buffer->Num())
 							{
diff --git a/Engine/Plugins/Media/ElectraPlayer/Source/ElectraPlayerRuntime/Private/Runtime/StreamAccessUnitBuffer.h b/Engine/Plugins/Media/ElectraPlayer/Source/ElectraPlayerRuntime/Private/Runtime/StreamAccessUnitBuffer.h
--- a/Engine/Plugins/Media/ElectraPlayer/Source/ElectraPlayerRuntime/Private/Runtime/StreamAccessUnitBuffer.h
+++ b/Engine/Plugins/Media/ElectraPlayer/Source/ElectraPlayerRuntime/Private/Runtime/StreamAccessUnitBuffer.h
@@ -39,6 +39,8 @@ namespace Electra
 		// Partial track metadata. See FTrackMetadata.
 		FString		Kind;
 		FString		Language;
+		FString		Codec;
+
 		// Internal hard index, used for multiplexed streams.
 		int32		HardIndex = -1;
 	};
@@ -67,6 +69,7 @@ namespace Electra
 		FTimeValue					DTS;							//!< DTS
 		FTimeValue					Duration;						//!< Duration
 		FTimeValue					OverlapAdjust;					//!< If set this indicates by how much the AU overlaps the desired time. Negative values indicate the AU is too early. Only set when DropState is zero.
+		FTimeValue					PTO;							//!< Media local presentation time offset.
 		uint32						AUSize;							//!< Size of this access unit
 		void*						AUData;							//!< Access unit data
 		TSharedPtrTS<CodecData>		AUCodecData;					//!< If set, points to sideband data for this access unit.
@@ -76,6 +79,7 @@ namespace Electra
 		bool						bIsSyncSample;					//!< true if this is a sync sample (keyframe)
 		bool						bIsDummyData;					//!< True if the decoder must be reset before decoding this AU.
 		bool						bTrackChangeDiscontinuity;		//!< True if this is the first AU after a track change.
+		bool						bIsSideloaded;					//!< True if the payload is not streamed but loaded from a sidecar file.
 
 		TSharedPtrTS<const FBufferSourceInfo>	BufferSourceInfo;
 		TSharedPtrTS<const FPlayerLoopState>	PlayerLoopState;
@@ -154,6 +158,8 @@ namespace Electra
 			PTS.SetToInvalid();
 			DTS.SetToInvalid();
 			Duration.SetToInvalid();
+			OverlapAdjust.SetToZero();
+			PTO.SetToZero();
 			ESType = EStreamType::Unsupported;
 			AUSize = 0;
 			AUData = nullptr;
@@ -164,6 +170,7 @@ namespace Electra
 			bIsSyncSample = false;
 			bIsDummyData = false;
 			bTrackChangeDiscontinuity = false;
+			bIsSideloaded = false;
 			bHasBeenPrepared = false;
 		}
 
diff --git a/Engine/Plugins/Media/ElectraPlayer/Source/ElectraPlayerRuntime/Private/Runtime/StreamTypes.cpp b/Engine/Plugins/Media/ElectraPlayer/Source/ElectraPlayerRuntime/Private/Runtime/StreamTypes.cpp
--- a/Engine/Plugins/Media/ElectraPlayer/Source/ElectraPlayerRuntime/Private/Runtime/StreamTypes.cpp
+++ b/Engine/Plugins/Media/ElectraPlayer/Source/ElectraPlayerRuntime/Private/Runtime/StreamTypes.cpp
@@ -221,7 +221,9 @@ namespace Electra
 			Codec = ECodec::WebVTT;
 			return true;
 		}
-		else if (CodecOTI.Equals(TEXT("stpp")))		// this is indicating one of the many TTML variants, should be IMSC1, SMPTE-TT, EBU-TT
+		// This is indicating one of the many TTML variants (eg. IMSC1, SMPTE-TT, EBU-TT) and profiles (eg. stpp.ttml.im1t)
+		//	See: https://www.w3.org/TR/ttml-profile-registry/#registry-profile-designator-specifications
+		else if (CodecOTI.StartsWith(TEXT("stpp")))
 		{
 			StreamType = EStreamType::Subtitle;
 			CodecSpecifier = CodecOTI;
@@ -244,6 +246,31 @@ namespace Electra
 		}
 	}
 
+	FString FStreamCodecInformation::GetCodecName() const
+	{
+		switch(Codec)
+		{
+			case FStreamCodecInformation::ECodec::H264:
+				return FString(TEXT("avc"));
+			case FStreamCodecInformation::ECodec::H265:
+				return FString(TEXT("hevc"));
+			case FStreamCodecInformation::ECodec::AAC:
+				return FString(TEXT("aac"));
+			case FStreamCodecInformation::ECodec::EAC3:
+				return FString(TEXT("eac3"));
+			case FStreamCodecInformation::ECodec::WebVTT:
+				return FString(TEXT("wvtt"));
+			case FStreamCodecInformation::ECodec::TTML:
+				return FString(TEXT("stpp"));
+			case FStreamCodecInformation::ECodec::TX3G:
+				return FString(TEXT("tx3g"));
+			case FStreamCodecInformation::ECodec::OtherSubtitle:
+				return FString(TEXT("subt"));
+			default:
+				return FString(TEXT("unknown"));
+		}
+	}
+
 
 
 
diff --git a/Engine/Plugins/Media/ElectraPlayer/Source/ElectraPlayerRuntime/Private/Runtime/StreamTypes.h b/Engine/Plugins/Media/ElectraPlayer/Source/ElectraPlayerRuntime/Private/Runtime/StreamTypes.h
--- a/Engine/Plugins/Media/ElectraPlayer/Source/ElectraPlayerRuntime/Private/Runtime/StreamTypes.h
+++ b/Engine/Plugins/Media/ElectraPlayer/Source/ElectraPlayerRuntime/Private/Runtime/StreamTypes.h
@@ -19,7 +19,7 @@ namespace Electra
 
 	static inline const TCHAR* GetStreamTypeName(EStreamType StreamType)
 	{
-		switch (StreamType)
+		switch(StreamType)
 		{
 			case EStreamType::Video:
 				return TEXT("Video");
@@ -89,9 +89,11 @@ namespace Electra
 			Codec = InCodec;
 		}
 
+		FString GetCodecName() const;
+
 		bool IsVideoCodec() const
 		{
-			switch (GetCodec())
+			switch(GetCodec())
 			{
 				case ECodec::H264:
 				case ECodec::H265:
@@ -103,7 +105,7 @@ namespace Electra
 
 		bool IsAudioCodec() const
 		{
-			switch (GetCodec())
+			switch(GetCodec())
 			{
 				case ECodec::AAC:
 				case ECodec::EAC3:
@@ -115,7 +117,7 @@ namespace Electra
 
 		bool IsSubtitleCodec() const
 		{
-			switch (GetCodec())
+			switch(GetCodec())
 			{
 				case ECodec::WebVTT:
 				case ECodec::TTML:
@@ -655,6 +657,10 @@ namespace Electra
 		// RFC-4646 and of RFC-5646
 		TOptional<FString> Language_ISO639;
 
+		// Preferred codec. Typically set to ensure the same track remains selected after a seek in case the same content
+		// is provided with different formats. See GetCodecName()
+		TOptional<FString> Codec;
+
 		// Rarely used. Unconditionally selects a track by its index where the index is a sequential numbering from [0..n)
 		// of the tracks as they are found. If the index is invalid the selection rules for kind and language are applied.
 		TOptional<int32> OverrideIndex;
@@ -673,9 +679,12 @@ namespace Electra
 			FString Lang1 = Language_ISO639.IsSet() ? Language_ISO639.GetValue() : FString();
 			FString Lang2 = Other.Language_ISO639.IsSet() ? Other.Language_ISO639.GetValue() : FString();
 
+			FString Codec1 = Codec.IsSet() ? Codec.GetValue() : FString();
+			FString Codec2 = Other.Codec.IsSet() ? Other.Codec.GetValue() : FString();
+
 			if (Kind1.IsEmpty() || Kind2.IsEmpty() || Kind1.Equals(Kind2))
 			{
-				return Lang1.Equals(Lang2);
+				return Lang1.Equals(Lang2) && Codec1.Equals(Codec2);
 			}
 			return false;
 		}
@@ -684,9 +693,10 @@ namespace Electra
 		{
 			Kind.Reset();
 			Language_ISO639.Reset();
+			Codec.Reset();
 			OverrideIndex.Reset();
 		}
-		virtual void UpdateWith(const FString& InKind, const FString& InLanguage, int32 InOverrideIndex)
+		virtual void UpdateWith(const FString& InKind, const FString& InLanguage, const FString& InCodec, int32 InOverrideIndex)
 		{
 			Reset();
 			if (!InKind.IsEmpty())
@@ -697,18 +707,23 @@ namespace Electra
 			{
 				Language_ISO639 = InLanguage;
 			}
+			if (!InCodec.IsEmpty())
+			{
+				Codec = InCodec;
+			}
 			if (InOverrideIndex >= 0)
 			{
 				OverrideIndex = InOverrideIndex;
 			}
 		}
-		virtual void UpdateIfOverrideSet(const FString& InKind, const FString& InLanguage)
+		virtual void UpdateIfOverrideSet(const FString& InKind, const FString& InLanguage, const FString& InCodec)
 		{
 			if (OverrideIndex.IsSet())
 			{
 				ClearOverrideIndex();
 				Kind = InKind;
 				Language_ISO639 = InLanguage;
+				Codec = InCodec;
 			}
 		}
 		virtual void ClearOverrideIndex()
diff --git a/Engine/Plugins/Media/ElectraPlayer/Source/ElectraPlayerRuntime/Private/Runtime/Utilities/UtilsMPEGVideo.h b/Engine/Plugins/Media/ElectraPlayer/Source/ElectraPlayerRuntime/Private/Runtime/Utilities/UtilsMPEGVideo.h
--- a/Engine/Plugins/Media/ElectraPlayer/Source/ElectraPlayerRuntime/Private/Runtime/Utilities/UtilsMPEGVideo.h
+++ b/Engine/Plugins/Media/ElectraPlayer/Source/ElectraPlayerRuntime/Private/Runtime/Utilities/UtilsMPEGVideo.h
@@ -165,6 +165,13 @@ namespace Electra
 			int32 GetNumberOfSPS() const;
 			const FISO14496_10_seq_parameter_set_data& GetParsedSPS(int32 SpsIndex) const;
 
+			uint8 GetAVCProfileIndication() const
+			{ return AVCProfileIndication; }
+			uint8 GetProfileCompatibility() const
+			{ return ProfileCompatibility; }
+			uint8 GetAVCLevelIndication() const
+			{ return AVCLevelIndication; }
+
 		private:
 			TArray<uint8>													RawData;
 			TArray<uint8>													CodecSpecificData;
diff --git a/Engine/Plugins/Media/ElectraPlayer/Source/ElectraPlayerRuntime/Public/IElectraPlayerInterface.h b/Engine/Plugins/Media/ElectraPlayer/Source/ElectraPlayerRuntime/Public/IElectraPlayerInterface.h
--- a/Engine/Plugins/Media/ElectraPlayer/Source/ElectraPlayerRuntime/Public/IElectraPlayerInterface.h
+++ b/Engine/Plugins/Media/ElectraPlayer/Source/ElectraPlayerRuntime/Public/IElectraPlayerInterface.h
@@ -105,6 +105,7 @@ public:
 	{
 		TOptional<FString> Kind;
 		TOptional<FString> Language_ISO639;
+		TOptional<FString> Codec;
 		TOptional<int32> TrackIndexOverride;
 		void Reset()
 		{
diff --git a/Engine/Plugins/Media/ElectraUtil/Source/ElectraBase/Private/PlayerTime.cpp b/Engine/Plugins/Media/ElectraUtil/Source/ElectraBase/Private/PlayerTime.cpp
--- a/Engine/Plugins/Media/ElectraUtil/Source/ElectraBase/Private/PlayerTime.cpp
+++ b/Engine/Plugins/Media/ElectraUtil/Source/ElectraBase/Private/PlayerTime.cpp
@@ -135,7 +135,8 @@ namespace Electra
 			return 0;
 		}
 
-		TBigInt<128> n(HNS);
+		bool bIsNeg = HNS < 0;
+		TBigInt<128> n(bIsNeg ? -HNS : HNS);
 		TBigInt<128> d(10000000);
 		TBigInt<128> s(CustomTimebase);
 
@@ -143,7 +144,7 @@ namespace Electra
 		n /= d;
 
 		int64 r = n.ToInt();
-		return r;
+		return bIsNeg ? -r : r;
 	}
 
 	// ---------------------------------------------------------------------------------------------------------------------
@@ -229,7 +230,8 @@ namespace Electra
 			return Numerator >= 0 ? 0x7fffffffffffffffLL : -0x7fffffffffffffffLL;
 		}
 
-		TBigInt<128> n(Numerator);
+		bool bIsNeg = Numerator < 0;
+		TBigInt<128> n(bIsNeg ? -Numerator : Numerator);
 		TBigInt<128> d(Denominator);
 		TBigInt<128> s(CustomTimebase);
 
@@ -237,7 +239,7 @@ namespace Electra
 		n /= d;
 
 		int64 r = n.ToInt();
-		return r;
+		return bIsNeg ? -r : r;
 	}
 
 
diff --git a/Engine/Plugins/Media/ElectraUtil/Source/ElectraBase/Public/PlayerTime.h b/Engine/Plugins/Media/ElectraUtil/Source/ElectraBase/Public/PlayerTime.h
--- a/Engine/Plugins/Media/ElectraUtil/Source/ElectraBase/Public/PlayerTime.h
+++ b/Engine/Plugins/Media/ElectraUtil/Source/ElectraBase/Public/PlayerTime.h
@@ -470,6 +470,17 @@ struct FTimeRange
 	{
 		return Start.IsValid() && End.IsValid();
 	}
+
+	bool Contains(const FTimeValue& Value)
+	{
+		return Value >= Start && Value < End;
+	}
+
+	bool Overlaps(const FTimeRange& OtherRange)
+	{
+		return End > OtherRange.Start && Start < OtherRange.End;
+	}
+
 	FTimeValue	Start;
 	FTimeValue	End;
 };
