ElectraPlayer: changed terminating the event dispatcher on player shutdown to not wait until the stop event was dispatched as there are no listeners registered any more at that point and we do not need to stall there.

#rnx

diff --git a/Engine/Plugins/Media/ElectraPlayer/Source/ElectraPlayerRuntime/Private/Runtime/Player/AdaptiveStreamingPlayer.cpp b/Engine/Plugins/Media/ElectraPlayer/Source/ElectraPlayerRuntime/Private/Runtime/Player/AdaptiveStreamingPlayer.cpp
--- a/Engine/Plugins/Media/ElectraPlayer/Source/ElectraPlayerRuntime/Private/Runtime/Player/AdaptiveStreamingPlayer.cpp
+++ b/Engine/Plugins/Media/ElectraPlayer/Source/ElectraPlayerRuntime/Private/Runtime/Player/AdaptiveStreamingPlayer.cpp
@@ -425,17 +425,6 @@ void FAdaptiveStreamingPlayer::DispatchEvent(TSharedPtrTS<FMetricEvent> Event)
 	}
 }
 
-void FAdaptiveStreamingPlayer::DispatchEventAndWait(TSharedPtrTS<FMetricEvent> Event)
-{
-	if (EventDispatcher.IsValid())
-	{
-		Event->Player = SharedThis(this);
-		EventDispatcher->DispatchEventAndWait(Event);
-	}
-}
-
-
-
 
 //-----------------------------------------------------------------------------
 /**
@@ -469,11 +458,7 @@ void FAdaptiveStreamingPlayer::StopWorkerThread()
 		SharedWorkerThread->RemovePlayerInstance(this);
 		SharedWorkerThread.Reset();
 	}
-	if (EventDispatcher.IsValid())
-	{
-		EventDispatcher->DispatchEventAndWait(TSharedPtrTS<FMetricEvent>());
-		EventDispatcher.Reset();
-	}
+	EventDispatcher.Reset();
 }
 
 
@@ -1278,11 +1263,7 @@ bool FAdaptiveStreamingPlayer::InternalHandleThreadMessages()
 					bIsClosing = true;
 					InternalStop(false);
 					InternalClose();
-					// This is the end of the rope anyway so let's terminate the event dispatcher.
-					// This ensures that the client will not get out of the Stop() call without having
-					// received the stop message.
-					DispatchEventAndWait(FMetricEvent::ReportPlaybackStopped());
-					EventDispatcher.Reset();
+					DispatchEvent(FMetricEvent::ReportPlaybackStopped());
 				}
 				if (msg.Data.MediaEvent.Event)
 				{
@@ -3912,14 +3893,6 @@ void FAdaptiveStreamingPlayerEventHandler::DispatchEvent(TSharedPtrTS<FMetricEve
 	EventQueue.SendMessage(InEvent);
 }
 
-void FAdaptiveStreamingPlayerEventHandler::DispatchEventAndWait(TSharedPtrTS<FMetricEvent> InEvent)
-{
-	FMediaEvent Sig;
-	InEvent->EventSignal = &Sig;
-	EventQueue.SendMessage(InEvent);
-	Sig.Wait();
-}
-
 FAdaptiveStreamingPlayerEventHandler::~FAdaptiveStreamingPlayerEventHandler()
 {
 	StopWorkerThread();
diff --git a/Engine/Plugins/Media/ElectraPlayer/Source/ElectraPlayerRuntime/Private/Runtime/Player/AdaptiveStreamingPlayerInternal.h b/Engine/Plugins/Media/ElectraPlayer/Source/ElectraPlayerRuntime/Private/Runtime/Player/AdaptiveStreamingPlayerInternal.h
--- a/Engine/Plugins/Media/ElectraPlayer/Source/ElectraPlayerRuntime/Private/Runtime/Player/AdaptiveStreamingPlayerInternal.h
+++ b/Engine/Plugins/Media/ElectraPlayer/Source/ElectraPlayerRuntime/Private/Runtime/Player/AdaptiveStreamingPlayerInternal.h
@@ -665,7 +665,6 @@ public:
 	static TSharedPtrTS<FAdaptiveStreamingPlayerEventHandler> Create();
 
 	void DispatchEvent(TSharedPtrTS<FMetricEvent> InEvent);
-	void DispatchEventAndWait(TSharedPtrTS<FMetricEvent> InEvent);
 
 	virtual ~FAdaptiveStreamingPlayerEventHandler();
 private:
@@ -1546,7 +1545,6 @@ private:
 	void InternalDeselectStream(EStreamType StreamType);
 
 	void DispatchEvent(TSharedPtrTS<FMetricEvent> Event);
-	void DispatchEventAndWait(TSharedPtrTS<FMetricEvent> Event);
 	void DispatchBufferingEvent(bool bBegin, EPlayerState Reason);
 	void DispatchSegmentDownloadedEvent(TSharedPtrTS<IStreamSegment> Request);
 	void DispatchBufferUtilizationEvent(EStreamType BufferType);
