ElectraPlayer: removed debug timeouts on suspend that can lead to unnecessary long delays

#rnx

diff --git a/Engine/Plugins/Media/ElectraPlayer/Source/ElectraPlayerRuntime/Private/Runtime/Decoder/DX/AudioDecoderAAC_DX.cpp b/Engine/Plugins/Media/ElectraPlayer/Source/ElectraPlayerRuntime/Private/Runtime/Decoder/DX/AudioDecoderAAC_DX.cpp
--- a/Engine/Plugins/Media/ElectraPlayer/Source/ElectraPlayerRuntime/Private/Runtime/Decoder/DX/AudioDecoderAAC_DX.cpp
+++ b/Engine/Plugins/Media/ElectraPlayer/Source/ElectraPlayerRuntime/Private/Runtime/Decoder/DX/AudioDecoderAAC_DX.cpp
@@ -979,7 +979,6 @@ namespace Electra
 	{
 		ApplicationSuspendConfirmedSignal.Reset();
 		ApplicationRunningSignal.Reset();
-		ApplicationSuspendConfirmedSignal.WaitTimeoutAndReset(1000 * 500);
 	}
 
 	
diff --git a/Engine/Plugins/Media/ElectraPlayer/Source/ElectraPlayerRuntime/Private/Runtime/Decoder/DX/VideoDecoderH264_DX.cpp b/Engine/Plugins/Media/ElectraPlayer/Source/ElectraPlayerRuntime/Private/Runtime/Decoder/DX/VideoDecoderH264_DX.cpp
--- a/Engine/Plugins/Media/ElectraPlayer/Source/ElectraPlayerRuntime/Private/Runtime/Decoder/DX/VideoDecoderH264_DX.cpp
+++ b/Engine/Plugins/Media/ElectraPlayer/Source/ElectraPlayerRuntime/Private/Runtime/Decoder/DX/VideoDecoderH264_DX.cpp
@@ -1298,10 +1298,6 @@ void FVideoDecoderH264::HandleApplicationWillEnterBackground()
 {
 	ApplicationSuspendConfirmedSignal.Reset();
 	ApplicationRunningSignal.Reset();
-	if (!ApplicationSuspendConfirmedSignal.WaitTimeoutAndReset(1000 * 500))
-	{
-		UE_LOG(LogElectraPlayer, Log, TEXT("FVideoDecoderH264(%p): Did not suspend in time"), this);
-	}
 }
 
 
diff --git a/Engine/Plugins/Media/ElectraPlayer/Source/ElectraPlayerRuntime/Private/Runtime/Decoder/DX/VideoDecoderH265_DX.cpp b/Engine/Plugins/Media/ElectraPlayer/Source/ElectraPlayerRuntime/Private/Runtime/Decoder/DX/VideoDecoderH265_DX.cpp
--- a/Engine/Plugins/Media/ElectraPlayer/Source/ElectraPlayerRuntime/Private/Runtime/Decoder/DX/VideoDecoderH265_DX.cpp
+++ b/Engine/Plugins/Media/ElectraPlayer/Source/ElectraPlayerRuntime/Private/Runtime/Decoder/DX/VideoDecoderH265_DX.cpp
@@ -1084,10 +1084,6 @@ void FVideoDecoderH265::HandleApplicationWillEnterBackground()
 {
 	ApplicationSuspendConfirmedSignal.Reset();
 	ApplicationRunningSignal.Reset();
-	if (!ApplicationSuspendConfirmedSignal.WaitTimeoutAndReset(1000 * 500))
-	{
-		UE_LOG(LogElectraPlayer, Log, TEXT("FVideoDecoderH265(%p): Did not suspend in time"), this);
-	}
 }
 
 
