ElectraPlayer: Live streams will now play synchronized to wallclock time with a safe margin. If synchronization cannot be achieved in a reasonable timeframe, playback will start regardless.

#rnx

diff --git a/Engine/Plugins/Media/ElectraPlayer/Source/ElectraPlayerRuntime/Private/Runtime/Player/AdaptiveStreamingPlayer.cpp b/Engine/Plugins/Media/ElectraPlayer/Source/ElectraPlayerRuntime/Private/Runtime/Player/AdaptiveStreamingPlayer.cpp
--- a/Engine/Plugins/Media/ElectraPlayer/Source/ElectraPlayerRuntime/Private/Runtime/Player/AdaptiveStreamingPlayer.cpp
+++ b/Engine/Plugins/Media/ElectraPlayer/Source/ElectraPlayerRuntime/Private/Runtime/Player/AdaptiveStreamingPlayer.cpp
@@ -1720,8 +1720,12 @@ void FAdaptiveStreamingPlayer::HandleMetadataChanges()
 				(!ActivePeriods[i].TimeRange.End.IsValid() && (i+1) < ActivePeriods.Num() && PastTime >= ActivePeriods[i+1].TimeRange.Start)) &&
 				ActivePeriods[i].LoopCount <= CurrentLoop.Count)
 			{
-				ActivePeriods.RemoveAt(i);
-				--i;
+				// Do not remove the only active period.
+				if (ActivePeriods.Num() > 1)
+				{
+					ActivePeriods.RemoveAt(i);
+					--i;
+				}
 			}
 		}
 	}
@@ -1842,6 +1846,98 @@ double FAdaptiveStreamingPlayer::GetMinBufferTimeBeforePlayback()
 	return kMinBufferBeforePlayback;
 }
 
+
+void FAdaptiveStreamingPlayer::InternalHandleLiveSync()
+{
+	// If we have to start playback at the current time NOW we need to tag all buffered access units
+	// that are older than NOW so that they will be skipped over in decoding or rendering.
+	// They may have been tagged already if frame accurate seeking is enabled, but the real time NOW
+	// between the time the segment download was started until, well, NOW, has continued to elapse
+	// so any access units in between need to be tagged as well.
+	if (LiveSyncVars.bSyncLiveToNow)
+	{
+		FTimeRange Seekable = Manifest->GetSeekableTimeRange();
+		FTimeRange TotalRange = Manifest->GetTotalTimeRange();
+		double minBufTime = Manifest ? Manifest->GetMinBufferTime().GetAsSeconds() : 1.0;
+
+		// For now attempting to sync to Live works only when there is a sufficiently large distance
+		// between the Live edge and the playback position. The difference can be found by examining
+		// the end value of the total vs. seekable range.
+		FTimeValue DistanceToLiveEdge = TotalRange.End - Seekable.End;
+		// FIXME: The *2 and *0.99 fudge values need to be adjusted once chunked download brings in new data constantly instead of whole segments.
+		if (!DistanceToLiveEdge.IsValid() || DistanceToLiveEdge.GetAsSeconds() < 2 * minBufTime * 0.99)
+		{
+			PostLog(Facility::EFacility::Player, IInfoLog::ELevel::Warning, FString::Printf(TEXT("Live sync not possible due to insufficient distance to Live edge (%.3fs) compared to minimum buffer time (%.3fs)"), DistanceToLiveEdge.GetAsSeconds(), minBufTime));
+			LiveSyncVars.Clear();
+			return;
+		}
+		
+		TSharedPtrTS<FMultiTrackAccessUnitBuffer> VidBuffer = GetCurrentReceiveStreamBuffer(EStreamType::Video);
+		TSharedPtrTS<FMultiTrackAccessUnitBuffer> AudBuffer = GetCurrentReceiveStreamBuffer(EStreamType::Audio);
+		TSharedPtrTS<FMultiTrackAccessUnitBuffer> TxtBuffer = GetCurrentReceiveStreamBuffer(EStreamType::Subtitle);
+
+		// See if we should perform a sync. For now we do this only if there is a video track.
+		bool bPerformSync = LiveSyncVars.StartTime.IsValid();
+		if (!bPerformSync && VidBuffer.IsValid())
+		{
+			FAccessUnitBufferInfo bi;
+			VidBuffer->GetStats(bi);
+			//bool bHaveEnough = bi.bLastPushWasBlocked || bi.bEndOfData || bi.PlayableDuration.GetAsSeconds() >= minBufTime;
+			bool bHaveEnough = bi.bLastPushWasBlocked || bi.bEndOfData || bi.PlayableDuration.GetAsSeconds() > 0.0;
+			if (bHaveEnough)
+			{
+				bPerformSync = true;
+				LiveSyncVars.StartTime = MEDIAutcTime::Current();
+			}
+		}
+
+		if (bPerformSync)
+		{
+			if (VidBuffer.IsValid())
+			{
+				LiveSyncVars.DroppedDurationVid += VidBuffer->PrepareForDecodeStartingAt(Seekable.End);
+			}
+			if (AudBuffer.IsValid())
+			{
+				LiveSyncVars.DroppedDurationAud += AudBuffer->PrepareForDecodeStartingAt(Seekable.End);
+			}
+			if (TxtBuffer.IsValid())
+			{
+				// Subtitles we discard.
+				FMultiTrackAccessUnitBuffer::FScopedLock lock(TxtBuffer);
+				TxtBuffer->PopDiscardUntil(Seekable.End);
+			}
+			UpdateDiagnostics();
+
+			// Check if we are spending too much time trying to perform the sync.
+			FTimeValue ElapsedTime = MEDIAutcTime::Current() - LiveSyncVars.StartTime;
+			if (ElapsedTime > DistanceToLiveEdge)
+			{
+				PostLog(Facility::EFacility::Player, IInfoLog::ELevel::Warning, FString::Printf(TEXT("Live sync is taking too long. Bandwidth to CDN is most likey insufficient. Starting now.")));
+				LiveSyncVars.Clear();
+			}
+
+			FTimeValue VidDropped = LiveSyncVars.DroppedDurationVid;
+			FTimeValue VidTotal = VideoBufferStats.StreamBuffer.PlayableDuration + VidDropped;
+			double DropRatio = VidDropped.GetAsSeconds() / VidTotal.GetAsSeconds();
+			if (LiveSyncVars.LastDropRatio < DropRatio && DropRatio >= 1.0)
+			{
+				// Force a rate scale in the ABR in the hope to get new data faster on a lower quality level.
+				if (StreamSelector.IsValid())
+				{
+					IAdaptiveStreamSelector::FBufferingQuality RateScale;
+					RateScale.BandwidthScaleFactor = 0.45;
+					StreamSelector->SwitchBufferingQuality(RateScale);
+				}
+			}
+			LiveSyncVars.LastDropRatio = DropRatio;
+
+			//PostLog(Facility::EFacility::Player, IInfoLog::ELevel::Info, FString::Printf(TEXT("%3.3fs: %.3f drop, %.3f total, %.3f ratio"), ElapsedTime.GetAsSeconds(), VidDropped.GetAsSeconds(), VidTotal.GetAsSeconds(), DropRatio));
+		}
+	}
+}
+
+
 //-----------------------------------------------------------------------------
 /**
  * Checks if buffers have enough data to advance the play state.
@@ -1853,6 +1949,9 @@ void FAdaptiveStreamingPlayer::HandleNewBufferedData()
 	{
 		if (DecoderState == EDecoderState::eDecoder_Paused)
 		{
+			// Synchronize buffers to wallclock time if necessary.
+			InternalHandleLiveSync();
+
 			check(LastBufferingState == EPlayerState::eState_Buffering || LastBufferingState == EPlayerState::eState_Rebuffering || LastBufferingState == EPlayerState::eState_Seeking);
 			// Yes. Check if we have enough data buffered up to begin handing off data to the decoders.
 			double kMinBufferBeforePlayback = GetMinBufferTimeBeforePlayback();
@@ -1968,6 +2067,8 @@ void FAdaptiveStreamingPlayer::HandleNewBufferedData()
 				PipelineState  		= EPipelineState::ePipeline_Prerolling;
 				PrerollVars.StartTime = tNow;
 				PostrollVars.Clear();
+				// We are synced up now and can clear the state.
+				LiveSyncVars.Clear();
 
 				// Send buffering end event
 				DispatchBufferingEvent(false, LastBufferingState);
@@ -2224,6 +2325,12 @@ void FAdaptiveStreamingPlayer::InternalHandlePendingStartRequest(const FTimeValu
 							{
 								check(Seekable.End.IsValid());
 								PendingStartRequest->StartAt.Time = Seekable.End;
+								// DASH Live streams provide the means to join at an exact wallclock time, so enable this.
+								if (ManifestType == EMediaFormatType::DASH)
+								{
+									LiveSyncVars.Clear();
+									LiveSyncVars.bSyncLiveToNow = true;
+								}
 							}
 						}
 					}
@@ -4083,6 +4190,7 @@ void FAdaptiveStreamingPlayer::InternalClose()
 	bRebufferPending   = false;
 	LastBufferingState = EPlayerState::eState_Buffering;
 	RebufferDetectedAtPlayPos.SetToInvalid();
+	LiveSyncVars.Clear();
 	PrerollVars.Clear();
 	PostrollVars.Clear();
 
@@ -4220,6 +4328,14 @@ void FAdaptiveStreamingPlayer::InternalRebuffer()
 				}
 			}
 
+			// Force a rate scale in the ABR. This may get overridden by the starting bitrate later.
+			if (StreamSelector.IsValid())
+			{
+				IAdaptiveStreamSelector::FBufferingQuality RateScale;
+				RateScale.BandwidthScaleFactor = 0.6;
+				StreamSelector->SwitchBufferingQuality(RateScale);
+			}
+
 			InternalStartAt(StartAtTime);
 		}
 	}
@@ -4286,19 +4402,19 @@ void FAdaptiveStreamingPlayer::DebugPrint(void* pPlayer, void (*pPrintFN)(void*
 	if (StreamReaderHandler && bHaveVideoReader.GetWithDefault(false))
 	{
 		pD = &VideoBufferStats.StreamBuffer;
-		pPrintFN(pPlayer, "Video buffer  : %3u AUs; %8u/%8u bytes in; eod %d; %#7.4fs", (uint32)pD->NumCurrentAccessUnits, (uint32)pD->CurrentMemInUse, (uint32)pD->MaxDataSize, pD->bEndOfData, pD->PushedDuration.GetAsSeconds());
+		pPrintFN(pPlayer, "Video buffer  : %3u AUs; %8u/%8u bytes in; eod %d; %#7.4fs", (uint32)pD->NumCurrentAccessUnits, (uint32)pD->CurrentMemInUse, (uint32)pD->MaxDataSize, pD->bEndOfData, pD->PlayableDuration.GetAsSeconds());
 		pPrintFN(pPlayer, "Video decoder : %2u in decoder, %zu total, %s; EOD in %d; EOD out %d", (uint32)VideoBufferStats.DecoderOutputBuffer.NumElementsInDecoder, (uint32)VideoBufferStats.DecoderOutputBuffer.MaxDecodedElementsReady, VideoBufferStats.DecoderOutputBuffer.bOutputStalled?"    stalled":"not stalled", VideoBufferStats.DecoderInputBuffer.bEODReached, VideoBufferStats.DecoderOutputBuffer.bEODreached);
 	}
 	if (StreamReaderHandler && bHaveAudioReader.GetWithDefault(false))
 	{
 		pD = &AudioBufferStats.StreamBuffer;
-		pPrintFN(pPlayer, "Audio buffer  : %3u AUs; %8u/%8u bytes in; eod %d; %#7.4fs", (uint32)pD->NumCurrentAccessUnits, (uint32)pD->CurrentMemInUse, (uint32)pD->MaxDataSize, pD->bEndOfData, pD->PushedDuration.GetAsSeconds());
+		pPrintFN(pPlayer, "Audio buffer  : %3u AUs; %8u/%8u bytes in; eod %d; %#7.4fs", (uint32)pD->NumCurrentAccessUnits, (uint32)pD->CurrentMemInUse, (uint32)pD->MaxDataSize, pD->bEndOfData, pD->PlayableDuration.GetAsSeconds());
 		pPrintFN(pPlayer, "Audio decoder : %2u in decoder, %2u total, %s; EOD in %d; EOD out %d", (uint32)AudioBufferStats.DecoderOutputBuffer.NumElementsInDecoder, (uint32)AudioBufferStats.DecoderOutputBuffer.MaxDecodedElementsReady, AudioBufferStats.DecoderOutputBuffer.bOutputStalled ? "    stalled" : "not stalled", AudioBufferStats.DecoderInputBuffer.bEODReached, AudioBufferStats.DecoderOutputBuffer.bEODreached);
 	}
 	if (StreamReaderHandler && bHaveTextReader.GetWithDefault(false))
 	{
 		pD = &TextBufferStats.StreamBuffer;
-		pPrintFN(pPlayer, "Text buffer   : %3u AUs; %8u/%8u bytes in; eod %d; %#7.4fs", (uint32)pD->NumCurrentAccessUnits, (uint32)pD->CurrentMemInUse, (uint32)pD->MaxDataSize, pD->bEndOfData, pD->PushedDuration.GetAsSeconds());
+		pPrintFN(pPlayer, "Text buffer   : %3u AUs; %8u/%8u bytes in; eod %d; %#7.4fs", (uint32)pD->NumCurrentAccessUnits, (uint32)pD->CurrentMemInUse, (uint32)pD->MaxDataSize, pD->bEndOfData, pD->PlayableDuration.GetAsSeconds());
 	}
 	DiagnosticsCriticalSection.Unlock();
 
diff --git a/Engine/Plugins/Media/ElectraPlayer/Source/ElectraPlayerRuntime/Private/Runtime/Player/AdaptiveStreamingPlayerABR.cpp b/Engine/Plugins/Media/ElectraPlayer/Source/ElectraPlayerRuntime/Private/Runtime/Player/AdaptiveStreamingPlayerABR.cpp
--- a/Engine/Plugins/Media/ElectraPlayer/Source/ElectraPlayerRuntime/Private/Runtime/Player/AdaptiveStreamingPlayerABR.cpp
+++ b/Engine/Plugins/Media/ElectraPlayer/Source/ElectraPlayerRuntime/Private/Runtime/Player/AdaptiveStreamingPlayerABR.cpp
@@ -129,6 +129,7 @@ namespace Electra
 		virtual void SetCurrentPlaybackPeriod(EStreamType InStreamType, TSharedPtrTS<IManifest::IPlayPeriod> InCurrentPlayPeriod) override;
 		virtual void SetBandwidth(int64 bitsPerSecond) override;
 		virtual void SetForcedNextBandwidth(int64 bitsPerSecond, double minBufferTimeBeforePlayback) override;
+		virtual void SwitchBufferingQuality(const FBufferingQuality& InQualityChange) override;
 		virtual void MarkStreamAsUnavailable(const FBlacklistedStream& BlacklistedStream) override;
 		virtual void MarkStreamAsAvailable(const FBlacklistedStream& NoLongerBlacklistedStream) override;
 		virtual int64 GetLastBandwidth() override;
@@ -235,6 +236,7 @@ namespace Electra
 		int32												ForcedInitialBandwidth;
 		double												ForcedInitialBandwidthUntilSecondsBuffered;
 		bool												bForcePlaystartBitrate;
+		bool												bBufferingQualityChanged;
 		FStreamCodecInformation::FResolution				MaxStreamResolution;
 
 		TSimpleMovingAverage<double>						AverageBandwidth;
@@ -316,6 +318,7 @@ namespace Electra
 		, ForcedInitialBandwidth(0)
 		, ForcedInitialBandwidthUntilSecondsBuffered(0.0)
 		, bForcePlaystartBitrate(false)
+		, bBufferingQualityChanged(false)
 	{
 		AverageBandwidth.SetMaxSMASampleCount(Config.MaxBandwidthHistorySize);
 		AverageLatency.SetMaxSMASampleCount(Config.MaxBandwidthHistorySize);
@@ -479,6 +482,7 @@ namespace Electra
 		AverageLatency.InitializeTo((double)0.1);
 		AverageThroughput.InitializeTo(bitsPerSecond);
 		bForcePlaystartBitrate = true;
+		bBufferingQualityChanged = false;
 	}
 
 
@@ -638,14 +642,18 @@ namespace Electra
 			double MaxSegmentDownloadDuration = SegmentDownloadStats.Duration;
 			// If there is a forced initial bitrate to start playback with then we allow for the download duration to take longer than normally
 			// as to not abort the first download unless it's download time is really getting problematic.
-			if (ForcedInitialBandwidth && SegmentDownloadStats.StreamType == EStreamType::Video)
+			// If the player asked for a buffering quality change things are problematic, so don't do it!
+			if (!bBufferingQualityChanged)
 			{
-				MaxSegmentDownloadDuration *= 2.0;
-			}
-			// When buffering we allow for some more time as well.
-			else if (bPlayerIsRebuffering)
-			{
-				MaxSegmentDownloadDuration *= 2.0;
+				if (ForcedInitialBandwidth && SegmentDownloadStats.StreamType == EStreamType::Video)
+				{
+					MaxSegmentDownloadDuration *= 2.0;
+				}
+				// When buffering we allow for some more time as well.
+				else if (bPlayerIsRebuffering)
+				{
+					MaxSegmentDownloadDuration *= 2.0;
+				}
 			}
 
 			// Rebuffering is set to just continue loading data? No timeout then.
@@ -718,6 +726,7 @@ namespace Electra
 		ForcedInitialBandwidth = 0;
 		ForcedInitialBandwidthUntilSecondsBuffered = 0.0;
 		bForcePlaystartBitrate = false;
+		bBufferingQualityChanged = false;
 	}
 	void FAdaptiveStreamSelector::ReportPlaybackEnded()
 	{
@@ -725,6 +734,7 @@ namespace Electra
 		ForcedInitialBandwidth = 0;
 		ForcedInitialBandwidthUntilSecondsBuffered = 0.0;
 		bForcePlaystartBitrate = false;
+		bBufferingQualityChanged = false;
 	}
 
 	void FAdaptiveStreamSelector::ReportDownloadEnd(const Metrics::FSegmentDownloadStats& SegmentDownloadStats)
@@ -815,6 +825,39 @@ namespace Electra
 		}
 	}
 
+
+	void FAdaptiveStreamSelector::SwitchBufferingQuality(const FBufferingQuality& InQualityChange)
+	{
+		double avgCurrent = AverageBandwidth.GetLastSample();
+		double newRate = avgCurrent;
+
+		// In case the average bandwidth is higher than what is needed for the best video stream
+		// we clamp it to there. Otherwise we may rebuffer the highest quality stream which is
+		// not likely to be useful since rebuffering means there was a drastic change in bandwidth.
+		if (StreamInformationVideo.Num())
+		{
+			if (newRate > StreamInformationVideo.Last()->Bitrate)
+			{
+				newRate = (double)StreamInformationVideo.Last()->Bitrate;
+			}
+		}
+
+		if (InQualityChange.BandwidthScaleFactor.IsSet())
+		{
+			newRate *= InQualityChange.BandwidthScaleFactor.GetValue();
+		}
+		else if (InQualityChange.AbsoluteBandwidth.IsSet())
+		{
+			newRate = (double)InQualityChange.AbsoluteBandwidth.GetValue();
+		}
+
+		double Ratio = avgCurrent > 0.0 ? newRate / avgCurrent : 1.0;
+		AverageBandwidth.InitializeTo(newRate);
+		AverageThroughput.InitializeTo((int64)(AverageThroughput.GetLastSample() * Ratio));
+		LogMessage(IInfoLog::ELevel::Info, FString::Printf(TEXT("Adjusting buffering bandwidth from %d to %d"), (int)avgCurrent, (int)newRate));
+		bBufferingQualityChanged = true;
+	}
+
 	void FAdaptiveStreamSelector::ReportBufferingStart(Metrics::EBufferingReason BufferingReason)
 	{
 		/*
@@ -827,11 +870,6 @@ namespace Electra
 		if (BufferingReason == Metrics::EBufferingReason::Rebuffering)
 		{
 			bPlayerIsRebuffering = true;
-			double Scale = 0.6;
-			double avgCurrent = AverageBandwidth.GetLastSample();
-			AverageBandwidth.InitializeTo(avgCurrent * Scale);
-			AverageThroughput.InitializeTo((int64)(AverageThroughput.GetLastSample() * Scale));
-			LogMessage(IInfoLog::ELevel::Info, FString::Printf(TEXT("Rebuffering cuts bandwidth by %.2f from %d to %d"), Scale, (int)avgCurrent, (int)(avgCurrent * Scale)));
 		}
 	}
 
diff --git a/Engine/Plugins/Media/ElectraPlayer/Source/ElectraPlayerRuntime/Private/Runtime/Player/AdaptiveStreamingPlayerABR.h b/Engine/Plugins/Media/ElectraPlayer/Source/ElectraPlayerRuntime/Private/Runtime/Player/AdaptiveStreamingPlayerABR.h
--- a/Engine/Plugins/Media/ElectraPlayer/Source/ElectraPlayerRuntime/Private/Runtime/Player/AdaptiveStreamingPlayerABR.h
+++ b/Engine/Plugins/Media/ElectraPlayer/Source/ElectraPlayerRuntime/Private/Runtime/Player/AdaptiveStreamingPlayerABR.h
@@ -66,6 +66,13 @@ public:
 	//! Sets a forced bitrate for the next segment fetches until the given duration of playable content has been received.
 	virtual void SetForcedNextBandwidth(int64 bitsPerSecond, double minBufferTimeBeforePlayback) = 0;
 
+	struct FBufferingQuality
+	{
+		TOptional<double> BandwidthScaleFactor;
+		TOptional<int64> AbsoluteBandwidth;
+	};
+	virtual void SwitchBufferingQuality(const FBufferingQuality& InQualityChange) = 0;
+
 	struct FBlacklistedStream
 	{
 		FString		AssetUniqueID;
diff --git a/Engine/Plugins/Media/ElectraPlayer/Source/ElectraPlayerRuntime/Private/Runtime/Player/AdaptiveStreamingPlayerInternal.h b/Engine/Plugins/Media/ElectraPlayer/Source/ElectraPlayerRuntime/Private/Runtime/Player/AdaptiveStreamingPlayerInternal.h
--- a/Engine/Plugins/Media/ElectraPlayer/Source/ElectraPlayerRuntime/Private/Runtime/Player/AdaptiveStreamingPlayerInternal.h
+++ b/Engine/Plugins/Media/ElectraPlayer/Source/ElectraPlayerRuntime/Private/Runtime/Player/AdaptiveStreamingPlayerInternal.h
@@ -1460,6 +1460,27 @@ private:
 		FStallMonitor										DecoderOutputStalledMonitor;
 	};
 
+	struct FLiveSyncVars
+	{
+		FLiveSyncVars()
+		{
+			Clear();
+		}
+		void Clear()
+		{
+			bSyncLiveToNow = false;
+			StartTime.SetToInvalid();
+			DroppedDurationVid.SetToZero();
+			DroppedDurationAud.SetToZero();
+			LastDropRatio = 0.0;
+		}
+		bool bSyncLiveToNow;
+		FTimeValue StartTime;
+		FTimeValue DroppedDurationVid;
+		FTimeValue DroppedDurationAud;
+		double LastDropRatio;
+	};
+
 	struct FPrerollVars
 	{
 		FPrerollVars()
@@ -1658,6 +1679,7 @@ private:
 	void InternalHandleCompletedSegmentRequests(const FTimeValue& CurrentTime);
 	void InternalHandleSegmentTrackChanges(const FTimeValue& CurrentTime);
 	void InternalStartoverAtCurrentPosition();
+	void InternalHandleLiveSync();
 
 	void UpdateDiagnostics();
 	void HandleNewBufferedData();
@@ -1814,6 +1836,7 @@ private:
 	EDecoderState														DecoderState;
 	EStreamState														StreamState;
 
+	FLiveSyncVars														LiveSyncVars;
 	FPrerollVars														PrerollVars;
 	FPostrollVars														PostrollVars;
 	EPlayerState														LastBufferingState;
diff --git a/Engine/Plugins/Media/ElectraPlayer/Source/ElectraPlayerRuntime/Private/Runtime/Player/DASH/ManifestBuilderDASH.cpp b/Engine/Plugins/Media/ElectraPlayer/Source/ElectraPlayerRuntime/Private/Runtime/Player/DASH/ManifestBuilderDASH.cpp
--- a/Engine/Plugins/Media/ElectraPlayer/Source/ElectraPlayerRuntime/Private/Runtime/Player/DASH/ManifestBuilderDASH.cpp
+++ b/Engine/Plugins/Media/ElectraPlayer/Source/ElectraPlayerRuntime/Private/Runtime/Player/DASH/ManifestBuilderDASH.cpp
@@ -930,24 +930,32 @@ void FManifestDASHInternal::TransformIntoEpicEvent()
 	// For this to work the presentation must be 'static'
 	if (PresentationType == EPresentationType::Static)
 	{
-		FString Time;
-		bool bDynamic = false;
+		FString Arg;
+		bool bStaticStart = false;
+		bool bDynamicStart = false;
 		for(int32 i=0,iMax=URLFragmentComponents.Num(); i<iMax; ++i)
 		{
-			bool bStaticStart = URLFragmentComponents[i].Name.Equals(Custom_EpicStaticStart);
-			bool bDynamicStart = URLFragmentComponents[i].Name.Equals(Custom_EpicDynamicStart);
+			bStaticStart = URLFragmentComponents[i].Name.Equals(Custom_EpicStaticStart);
+			bDynamicStart = URLFragmentComponents[i].Name.Equals(Custom_EpicDynamicStart);
 			if (bStaticStart || bDynamicStart)
 			{
-				bDynamic = bDynamicStart;
-				Time = URLFragmentComponents[i].Value;
+				Arg = URLFragmentComponents[i].Value;
 				break;
 			}
 		}
-		if (!Time.IsEmpty())
+		TArray<FString> Params;
+		if (!Arg.IsEmpty())
+		{
+			const TCHAR* const Delimiter = TEXT(",");
+			Arg.ParseIntoArray(Params, Delimiter, true);
+		}
+		if (Params.Num())
 		{
+			EpicEventType = bStaticStart ? EEpicEventType::Static : EEpicEventType::Dynamic;
 			// Get the event start time. This is either a Posix time in seconds since the Epoch (1/1/1970) or the special
 			// word 'now' optionally followed by a value to be added or subtracted from now.
 			FTimeValue Start;
+			FString Time = Params[0];
 			if (Time.StartsWith(TEXT("now")))
 			{
 				Time.RightChopInline(3);
@@ -971,12 +979,26 @@ void FManifestDASHInternal::TransformIntoEpicEvent()
 			if (Start.IsValid())
 			{
 				MPDRoot->SetAvailabilityStartTime(Start);
-				if (bDynamic)
+				if (EpicEventType == EEpicEventType::Dynamic)
 				{
 					MPDRoot->SetPublishTime(Start);
 					MPDRoot->SetType(TEXT("dynamic"));
 					PresentationType = EPresentationType::Dynamic;
-					bIsEventType = true;
+					if (Params.Num() > 1)
+					{
+						FTimeValue spd = FTimeValue().SetFromTimeFraction(FTimeFraction().SetFromFloatString(Params[1]));
+						MPDRoot->SetSuggestedPresentationDelay(spd);
+					}
+					else
+					{
+						MPDRoot->SetSuggestedPresentationDelay(MPDRoot->GetMinBufferTime());
+					}
+				}
+				// Adjust period start times with the fake AST
+				for(auto &Period : GetPeriods())
+				{
+					Period->StartAST = Period->Start + Start;
+					Period->EndAST = Period->End + Start;
 				}
 			}
 		}
@@ -1051,6 +1073,8 @@ FErrorDetail FManifestDASHInternal::BuildAfterInitialRemoteElementDownload()
 
 	FErrorDetail Error;
 
+	FTimeValue AST = GetAnchorTime();
+
 	bool bWarnedPresentationDuration = false;
 	// Go over the periods one at a time as XLINK attributes could bring in additional periods.
 	FTimeValue CombinedPeriodDuration(FTimeValue::GetZero());
@@ -1185,6 +1209,10 @@ FErrorDetail FManifestDASHInternal::BuildAfterInitialRemoteElementDownload()
 			}
 		}
 
+		// Set the time range including AST
+		p->StartAST = p->Start + AST;
+		p->EndAST = p->End + AST;
+
 		if (p->Start.IsValid() && p->End.IsValid())
 		{
 			FTimeValue periodDur = p->End - p->Start;
@@ -2094,6 +2122,31 @@ FTimeValue FManifestDASHInternal::CalculateDistanceToLiveEdge() const
 	if (!Distance.IsValid())
 	{
 		Distance = MPDRoot->GetSuggestedPresentationDelay();
+
+		/*
+			FIXME:
+				This fudge is here to prevent the use of the live segment.
+				At the moment the HTTP module only allows us access to the data once the segment has been fetched in its entirety,
+				which tends to be too late, especially if the current segment is only used partially (we may have a 2s segment
+				fetched but start playback 1s into is, giving us only 1s worth of usable data).
+				For real low-latency playback chunked transfer is needed with access to the partially downloaded data.
+
+				If the @suggestedPresentationDelay is set equal to or shorter than the @minBufferTime, which we assume to indicate the duration
+				of a single segment we would need low-latency.
+				Since we cannot force the MPD@suggestedPresentationDelay to be set such that none of this will be a problem
+				we adjust that value to twice the MPD@minBufferTime until we can get chunked transfer and partial data access.
+		*/
+			FTimeValue mbt_times2 = MPDRoot->GetMinBufferTime() * 2;
+			if (Distance < mbt_times2)
+			{
+				if (!bWarnedAboutTooSmallSuggestedPresentationDelay)
+				{
+					bWarnedAboutTooSmallSuggestedPresentationDelay = true;
+					LogMessage(PlayerSessionServices, IInfoLog::ELevel::Info, FString::Printf(TEXT("Adjusting the MPD@suggestedPresentationDelay from %.3f seconds to %.3f to avoid Live edge buffering issues"), Distance.GetAsSeconds(), mbt_times2.GetAsSeconds()));
+				}
+				Distance = mbt_times2;
+			}
+
 	}
 	// If not set see if there is an MPD@maxSegmentDuration and use that.
 	if (!Distance.IsValid())
@@ -2241,7 +2294,7 @@ FTimeRange FManifestDASHInternal::GetSeekableTimeRange() const
 		// future segments already (a pre-existing presentation made available over time) we would not
 		// be able to fetch them anyway on account of their availability window not being valid yet.
 		// Typically we would want to play some distance away from the bleeding Live edge.
-		bool bIsUpdating = AreUpdatesExpected();
+		bool bIsUpdating = AreUpdatesExpected() || EpicEventType == EEpicEventType::Dynamic;
 		FTimeValue Distance = bIsUpdating ? CalculateDistanceToLiveEdge() : FixedSeekEndDistance;
 		FTimeValue Now = PlayerSessionServices->GetSynchronizedUTCTime()->GetTime();
 		FTimeValue LastEnd = GetLastPeriodEndTime();
@@ -2400,7 +2453,6 @@ FTimeRange FManifestDASHInternal::GetPlayTimesFromURI() const
 		return FromTo;
 	}
 
-//	FTimeRange AvailableTimeRange = GetSeekableTimeRange();
 	FTimeRange AvailableTimeRange = GetTotalTimeRange();
 	// Is the time specified as a POSIX time?
 	TArray<FString> TimeRange;
@@ -2415,14 +2467,20 @@ FTimeRange FManifestDASHInternal::GetPlayTimesFromURI() const
 			// Is the start time the special time 'now'?
 			if (TimeRange[0].Equals(TEXT("now")))
 			{
-				/*
-				FTimeValue LastEnd = GetLastPeriodEndTime();
-				if (LastEnd.IsValid() && Now > LastEnd)
+				// A static event will not use an updated wallclock NOW, so if 'now' is used we do an init
+				// with the current time.
+				if (EpicEventType == EEpicEventType::Static)
 				{
-					LastEnd = Now;
+					FromTo.Start = Now;
+				}
+				else
+				{
+					// 'now' is dynamic. The time will continue to flow between here where we set the value and
+					// the moment playback will begin with buffered data.
+					// We do not lock 'now' with the current time but leave it unset. This results in the start
+					// time to be the 
+					FromTo.Start.SetToInvalid();
 				}
-				*/
-				FromTo.Start = Now;
 			}
 			else
 			{
diff --git a/Engine/Plugins/Media/ElectraPlayer/Source/ElectraPlayerRuntime/Private/Runtime/Player/DASH/ManifestDASH.cpp b/Engine/Plugins/Media/ElectraPlayer/Source/ElectraPlayerRuntime/Private/Runtime/Player/DASH/ManifestDASH.cpp
--- a/Engine/Plugins/Media/ElectraPlayer/Source/ElectraPlayerRuntime/Private/Runtime/Player/DASH/ManifestDASH.cpp
+++ b/Engine/Plugins/Media/ElectraPlayer/Source/ElectraPlayerRuntime/Private/Runtime/Player/DASH/ManifestDASH.cpp
@@ -1012,7 +1012,7 @@ IManifest::FResult FDASHPlayPeriod::GetStartingSegment(TSharedPtrTS<IStreamSegme
 	PlayRangeEnd -= Period->GetStart();
 
 	bool bUsesAST = Manifest->UsesAST();
-	bool bIsStaticType = Manifest->IsStaticType() || Manifest->IsEventType();
+	bool bIsStaticType = Manifest->IsStaticType() || Manifest->IsDynamicEpicEvent();
 
 	// Create a segment request to which the individual stream segment requests will add themselves as
 	// dependent streams. This is a special case for playback start.
@@ -1268,7 +1268,7 @@ IManifest::FResult FDASHPlayPeriod::GetNextOrRetrySegment(TSharedPtrTS<IStreamSe
 
 	FTimeValue AST = Manifest->GetAnchorTime();
 	bool bUsesAST = Manifest->UsesAST();
-	bool bIsStaticType = Manifest->IsStaticType() || Manifest->IsEventType();
+	bool bIsStaticType = Manifest->IsStaticType() || Manifest->IsDynamicEpicEvent();
 
 	// Frame accurate seek required?
 	bool bFrameAccurateSearch = PlayerSessionServices->GetOptions().GetValue(OptionKeyFrameAccurateSeek).SafeGetBool(false);
@@ -2383,6 +2383,7 @@ FManifestDASHInternal::FRepresentation::ESearchResult FManifestDASHInternal::FRe
 	// Did we find it?
 	if (bFound && CurrentT < MediaLocalEndTime)
 	{
+		OutSegmentInfo.PeriodLocalSegmentStartTime.SetFromND(CurrentT - PTO, SidxTimescale);
 		OutSegmentInfo.Time = CurrentT;
 		OutSegmentInfo.PTO = PTO;
 		OutSegmentInfo.Duration = CurrentD;
@@ -2530,6 +2531,7 @@ FManifestDASHInternal::FRepresentation::ESearchResult FManifestDASHInternal::FRe
 		return FManifestDASHInternal::FRepresentation::ESearchResult::PastEOS;
 	}
 
+	OutSegmentInfo.PeriodLocalSegmentStartTime.SetFromND(SegmentNum * (int64)SegmentDuration - PTO, MPDTimescale);
 	OutSegmentInfo.Time = EPTdelta + SegmentNum * (int64)SegmentDuration;
 	OutSegmentInfo.PTO = PTO;
 	OutSegmentInfo.EPTdelta = EPTdelta;
@@ -2855,6 +2857,7 @@ FManifestDASHInternal::FRepresentation::ESearchResult FManifestDASHInternal::FRe
 	// Did we find it?
 	if (bFound && CurrentT < MediaLocalEndTime)
 	{
+		OutSegmentInfo.PeriodLocalSegmentStartTime.SetFromND(CurrentT - PTO, MPDTimescale);
 		OutSegmentInfo.Time = CurrentT;
 		OutSegmentInfo.PTO = PTO;
 		OutSegmentInfo.Duration = CurrentD;
diff --git a/Engine/Plugins/Media/ElectraPlayer/Source/ElectraPlayerRuntime/Private/Runtime/Player/DASH/PlaylistReaderDASH.cpp b/Engine/Plugins/Media/ElectraPlayer/Source/ElectraPlayerRuntime/Private/Runtime/Player/DASH/PlaylistReaderDASH.cpp
--- a/Engine/Plugins/Media/ElectraPlayer/Source/ElectraPlayerRuntime/Private/Runtime/Player/DASH/PlaylistReaderDASH.cpp
+++ b/Engine/Plugins/Media/ElectraPlayer/Source/ElectraPlayerRuntime/Private/Runtime/Player/DASH/PlaylistReaderDASH.cpp
@@ -1184,13 +1184,15 @@ void FPlaylistReaderDASH::ManifestDownloadCompleted(FResourceLoadRequestPtr Requ
 					FetchTime = PlayerSessionServices->GetSynchronizedUTCTime()->GetTime();
 					Manifest->GetMPDRoot()->SetFetchTime(FetchTime);
 					MostRecentMPDUpdateTime = FetchTime;
-					// Do a one-time preparation of the default start time from the #t= URL fragment.
-					// This is not to be repeated for manifest updates.
-					Manifest->PrepareDefaultStartTime();
 
 					// If the URL has a special fragment part to turn this presentation into an event, do so.
+					// Do this before preparing the default start times!
 					// Note: This MAY require the client clock to be synced to the server IF the special keyword 'now' is used.
 					Manifest->TransformIntoEpicEvent();
+
+					// Do a one-time preparation of the default start time from the #t= URL fragment.
+					// This is not to be repeated for manifest updates.
+					Manifest->PrepareDefaultStartTime();
 				}
 
 				PlayerManifest = FManifestDASH::Create(PlayerSessionServices, Manifest);
diff --git a/Engine/Plugins/Media/ElectraPlayer/Source/ElectraPlayerRuntime/Private/Runtime/Player/DASH/PlaylistReaderDASH_Internal.h b/Engine/Plugins/Media/ElectraPlayer/Source/ElectraPlayerRuntime/Private/Runtime/Player/DASH/PlaylistReaderDASH_Internal.h
--- a/Engine/Plugins/Media/ElectraPlayer/Source/ElectraPlayerRuntime/Private/Runtime/Player/DASH/PlaylistReaderDASH_Internal.h
+++ b/Engine/Plugins/Media/ElectraPlayer/Source/ElectraPlayerRuntime/Private/Runtime/Player/DASH/PlaylistReaderDASH_Internal.h
@@ -145,6 +145,7 @@ public:
 		FURL InitializationURL;
 		FURL MediaURL;
 		FTimeValue ATO;
+		FTimeValue PeriodLocalSegmentStartTime;
 		int64 Time = 0;								//!< Time value T in timescale units
 		int64 PTO = 0;								//!< PresentationTimeOffset
 		int64 EPTdelta = 0;
@@ -634,7 +635,8 @@ public:
 		//
 		virtual FTimeRange GetTimeRange() const override
 		{
-			return FTimeRange({Start, End});
+			// Per convention the time range includes the AST
+			return FTimeRange({StartAST, EndAST});
 		}
 		virtual FTimeValue GetDuration() const override
 		{
@@ -699,6 +701,8 @@ public:
 		FString ID;
 		FTimeValue Start;
 		FTimeValue End;
+		FTimeValue StartAST;
+		FTimeValue EndAST;
 		FTimeValue Duration;
 		bool bIsEarlyPeriod = false;
 		bool bHasFollowingPeriod = false;
@@ -721,9 +725,9 @@ public:
 		return PresentationType;
 	}
 
-	bool IsEventType() const
+	bool IsDynamicEpicEvent() const
 	{
-		return bIsEventType;
+		return EpicEventType == EEpicEventType::Dynamic;
 	}
 
 	const TArray<TSharedPtrTS<FPeriod>>& GetPeriods() const
@@ -816,6 +820,12 @@ public:
 	void EndPresentationAt(const FTimeValue& EndsAt, const FString& InPeriod);
 
 private:
+	enum class EEpicEventType
+	{
+		None,
+		Static,
+		Dynamic
+	};
 	FErrorDetail PrepareRemoteElementLoadRequest(TArray<TWeakPtrTS<FMPDLoadRequestDASH>>& OutRemoteElementLoadRequests, TWeakPtrTS<IDashMPDElement> ElementWithXLink, int64 RequestID);
 
 	int32 ReplaceElementWithRemoteEntities(TSharedPtrTS<IDashMPDElement> Element, const FDashMPD_RootEntities& NewRootEntities, int64 OldResolveID, int64 NewResolveID);
@@ -837,13 +847,14 @@ private:
 
 	// Type of the presentation.
 	EPresentationType PresentationType;
-	bool bIsEventType = false;
+	EEpicEventType EpicEventType = EEpicEventType::None;
 
 	TArray<TSharedPtrTS<FPeriod>> Periods;
 
 	mutable FTimeRange TotalTimeRange;
 	mutable FTimeRange SeekableTimeRange;
 	FTimeValue DefaultStartTime;
+	mutable bool bWarnedAboutTooSmallSuggestedPresentationDelay = false;
 };
 
 
diff --git a/Engine/Plugins/Media/ElectraPlayer/Source/ElectraPlayerRuntime/Private/Runtime/Player/DASH/StreamReaderFMP4DASH.cpp b/Engine/Plugins/Media/ElectraPlayer/Source/ElectraPlayerRuntime/Private/Runtime/Player/DASH/StreamReaderFMP4DASH.cpp
--- a/Engine/Plugins/Media/ElectraPlayer/Source/ElectraPlayerRuntime/Private/Runtime/Player/DASH/StreamReaderFMP4DASH.cpp
+++ b/Engine/Plugins/Media/ElectraPlayer/Source/ElectraPlayerRuntime/Private/Runtime/Player/DASH/StreamReaderFMP4DASH.cpp
@@ -956,6 +956,7 @@ void FStreamReaderFMP4DASH::FStreamHandler::HandleRequest()
 			bool bDone = false;
 			bool bIsFirstAU = true;
 			bool bReadPastLastPTS = false;
+			FTimeValue TimelineOffset;
 
 			// Encrypted?
 			TSharedPtr<ElectraCDM::IMediaCDMDecrypter, ESPMode::ThreadSafe> Decrypter;
@@ -1029,6 +1030,14 @@ void FStreamReaderFMP4DASH::FStreamHandler::HandleRequest()
 
 								for(Error = TrackIterator->StartAtFirst(false); Error == UEMEDIA_ERROR_OK; Error = TrackIterator->Next())
 								{
+									if (bIsFirstAU)
+									{
+										FTimeValue BaseMediaDecodeTime;
+										BaseMediaDecodeTime.SetFromND(TrackIterator->GetBaseMediaDecodeTime(), TrackIterator->GetTimescale());
+										TimelineOffset = BaseMediaDecodeTime - Request->Segment.PeriodLocalSegmentStartTime;
+										//LogMessage(IInfoLog::ELevel::Info, FString::Printf(TEXT("Segment to timeline delta = %.3fs"), TimelineOffset.GetAsSeconds()));
+									}
+
 									// Get the DTS and PTS. Those are 0-based in a fragment and offset by the base media decode time of the fragment.
 									int64 AUDTS = TrackIterator->GetDTS();
 									int64 AUPTS = TrackIterator->GetPTS();
@@ -1093,6 +1102,8 @@ void FStreamReaderFMP4DASH::FStreamHandler::HandleRequest()
 									AccessUnit->LatestPTS += TimeOffset;
 									AccessUnit->LatestPTS.SetSequenceIndex(Request->TimestampSequenceIndex);
 
+									AccessUnit->OffsetFromSegmentStart = TimelineOffset;
+
 									ElectraCDM::FMediaCDMSampleInfo SampleEncryptionInfo;
 									bool bIsSampleEncrypted = TrackIterator->GetEncryptionInfo(SampleEncryptionInfo);
 
@@ -1524,11 +1535,6 @@ void FStreamReaderFMP4DASH::FStreamHandler::HandleRequest()
 		}
 	}
 
-	if (!bSilentCancellation)
-	{
-		StreamSelector->ReportDownloadEnd(ds);
-	}
-
 	// Clean out everything before reporting OnFragmentClose().
 	TSharedPtrTS<FStreamSegmentRequestFMP4DASH> FinishedRequest = CurrentRequest;
 	CurrentRequest.Reset();
@@ -1538,6 +1544,7 @@ void FStreamReaderFMP4DASH::FStreamHandler::HandleRequest()
 
 	if (!bSilentCancellation)
 	{
+		StreamSelector->ReportDownloadEnd(ds);
 		Parameters.EventListener->OnFragmentClose(FinishedRequest);
 	}
 }
diff --git a/Engine/Plugins/Media/ElectraPlayer/Source/ElectraPlayerRuntime/Private/Runtime/StreamAccessUnitBuffer.cpp b/Engine/Plugins/Media/ElectraPlayer/Source/ElectraPlayerRuntime/Private/Runtime/StreamAccessUnitBuffer.cpp
--- a/Engine/Plugins/Media/ElectraPlayer/Source/ElectraPlayerRuntime/Private/Runtime/StreamAccessUnitBuffer.cpp
+++ b/Engine/Plugins/Media/ElectraPlayer/Source/ElectraPlayerRuntime/Private/Runtime/StreamAccessUnitBuffer.cpp
@@ -540,6 +540,18 @@ namespace Electra
 		}
 	}
 
+	FTimeValue FMultiTrackAccessUnitBuffer::PrepareForDecodeStartingAt(FTimeValue DecodeStartTime)
+	{
+		FTimeValue TaggedDuration(FTimeValue::GetZero());
+		// We lock the access mutex here.
+		FScopedLock lock(AsShared());
+		for(auto& It : TrackBuffers)
+		{
+			TSharedPtrTS<FAccessUnitBuffer> Buf = It.Value;
+			TaggedDuration += Buf->PrepareForDecodeStartingAt(DecodeStartTime);
+		}
+		return TaggedDuration;
+	}
 
 	bool FMultiTrackAccessUnitBuffer::IsEODFlagSet()
 	{
diff --git a/Engine/Plugins/Media/ElectraPlayer/Source/ElectraPlayerRuntime/Private/Runtime/StreamAccessUnitBuffer.h b/Engine/Plugins/Media/ElectraPlayer/Source/ElectraPlayerRuntime/Private/Runtime/StreamAccessUnitBuffer.h
--- a/Engine/Plugins/Media/ElectraPlayer/Source/ElectraPlayerRuntime/Private/Runtime/StreamAccessUnitBuffer.h
+++ b/Engine/Plugins/Media/ElectraPlayer/Source/ElectraPlayerRuntime/Private/Runtime/StreamAccessUnitBuffer.h
@@ -70,6 +70,7 @@ namespace Electra
 		FTimeValue					PTO;							//!< Media local presentation time offset.
 		FTimeValue					EarliestPTS;					//!< Earliest PTS at which to present samples. If this is larger than PTS the sample is not to be presented.
 		FTimeValue					LatestPTS;						//!< Latest PTS at which to present samples. If this is less than PTS the sample is not to be presented.
+		FTimeValue					OffsetFromSegmentStart;			//!< If set, the difference between the first segment AU's PTS and the expected time according to the playlist.
 		int64						SequenceIndex;
 		uint32						AUSize;							//!< Size of this access unit
 		void*						AUData;							//!< Access unit data
@@ -156,6 +157,7 @@ namespace Electra
 			PTO.SetToZero();
 			EarliestPTS.SetToInvalid();
 			LatestPTS.SetToInvalid();
+			OffsetFromSegmentStart.SetToInvalid();
 			SequenceIndex = 0;
 			ESType = EStreamType::Unsupported;
 			AUSize = 0;
@@ -475,6 +477,56 @@ namespace Electra
 			}
 		}
 
+		// Tags all AUs such that the decoders will skip over them until the desired decode time.
+		FTimeValue PrepareForDecodeStartingAt(FTimeValue DecodeStartTime)
+		{
+			FTimeValue TaggedDuration(FTimeValue::GetZero());
+			FMediaCriticalSection::ScopedLock Lock(AccessLock);
+			int32 RemoveUpTo = -1;
+			for(int32 i=0; i<AccessUnits.Num(); ++i)
+			{
+				FAccessUnit* AU = AccessUnits[i];
+
+				// Check if the first AU has a negative offset compared to the expected segment start time.
+				// If so, we have to adjust the given decode start time by that offset as to not discard the data
+				// that has an internal timestamp deviation from the timeline of the playlist.
+				// Positive offsets are currently assumed to not be an issue, unless they were unreasonably large,
+				// but what constitutes "reasonably" is not clear. We're leaving things at that.
+				if (i == 0 && AU->OffsetFromSegmentStart.IsValid() && AU->OffsetFromSegmentStart < FTimeValue::GetZero())
+				{
+					DecodeStartTime += AU->OffsetFromSegmentStart;
+				}
+
+				if (AU->PTS < DecodeStartTime)
+				{
+					// If the AU was not set with a drop flag yet we do it now.
+					// Drop-flagged AUs do not count towards playable duration, so we need to adjust this now.
+					if (AU->DropState == 0)
+					{
+						PlayableDuration -= AU->Duration;
+						TaggedDuration += AU->Duration;
+					}
+					AU->DropState |= FAccessUnit::EDropState::TooEarly;
+
+					AU->EarliestPTS = DecodeStartTime;
+
+					// If this is a sync sample we can remove all preceding AUs until this one.
+					if (AU->bIsSyncSample)
+					{
+						RemoveUpTo = i;
+					}
+				}
+			}
+			for(int32 i=0; i<RemoveUpTo; ++i)
+			{
+				FAccessUnit* AU = nullptr;
+				Pop(AU);
+				FAccessUnit::Release(AU);
+			}
+			return TaggedDuration;
+		}
+
+
 		//! Waits for data to arrive. Returns true if data is present. False if not and timeout expired.
 		bool WaitForData(int64 waitForMicroseconds = -1)
 		{
@@ -547,6 +599,13 @@ namespace Electra
 		{
 			check(Limit->MaxDuration > FTimeValue::GetZero());
 			check(AU->Duration.IsValid() && !AU->Duration.IsInfinity());
+			// If tagged with a drop state we unconditionally accept the AU. These are early ones that won't get decoded.
+			// They should appear at the beginning and end of a stream only and we do need the data that will be decoded
+			// so the ones to be dropped should not be taken into account.
+			if (AU->DropState != 0)
+			{
+				return true;
+			}
 			if (ExternalInfo == nullptr)
 			{
 				// Memory ok?
@@ -631,6 +690,7 @@ namespace Electra
 		bool PeekAndAddRef(FAccessUnit*& OutAU);
 		bool Pop(FAccessUnit*& OutAU);
 		void PopDiscardUntil(FTimeValue UntilTime);
+		FTimeValue PrepareForDecodeStartingAt(FTimeValue DecodeStartTime);
 		bool IsEODFlagSet();
 		int32 Num();
 		bool WasLastPushBlocked();
