ElectraPlayer: added a missing mutex lock

#rnx

diff --git a/Engine/Plugins/Media/ElectraPlayer/Source/ElectraPlayerRuntime/Private/ElectraPlayer.cpp b/Engine/Plugins/Media/ElectraPlayer/Source/ElectraPlayerRuntime/Private/ElectraPlayer.cpp
--- a/Engine/Plugins/Media/ElectraPlayer/Source/ElectraPlayerRuntime/Private/ElectraPlayer.cpp
+++ b/Engine/Plugins/Media/ElectraPlayer/Source/ElectraPlayerRuntime/Private/ElectraPlayer.cpp
@@ -2699,6 +2699,7 @@ void FElectraPlayer::SendAnalyticMetrics(const TSharedPtr<IAnalyticsProviderET>&
 
 	TArray<FAnalyticsEventAttribute> ParamArray;
 	AddCommonAnalyticsAttributes(ParamArray);
+	StatisticsLock.Lock();
 	ParamArray.Add(FAnalyticsEventAttribute(TEXT("URL"), Statistics.InitialURL));
 	ParamArray.Add(FAnalyticsEventAttribute(TEXT("LastState"), Statistics.LastState));
 	ParamArray.Add(FAnalyticsEventAttribute(TEXT("LastError"), Statistics.LastError));
@@ -2746,6 +2747,7 @@ void FElectraPlayer::SendAnalyticMetrics(const TSharedPtr<IAnalyticsProviderET>&
 	ParamArray.Add(FAnalyticsEventAttribute(TEXT("SubtitlesURL"), Statistics.SubtitlesURL));
 	ParamArray.Add(FAnalyticsEventAttribute(TEXT("SubtitlesResponseTime"), Statistics.SubtitlesResponseTime));
 	ParamArray.Add(FAnalyticsEventAttribute(TEXT("SubtitlesLastError"), Statistics.SubtitlesLastError));
+	StatisticsLock.Unlock();
 
 	AnalyticsProvider->RecordEvent(TEXT("Electra.FinalMetrics"), MoveTemp(ParamArray));
 }
