ElectraPlayer: removed unnecessary check()'s triggering on Gauntlet

#rnx

diff --git a/Engine/Plugins/Media/ElectraPlayer/Source/ElectraPlayerRuntime/Private/Runtime/Player/AdaptiveStreamingPlayer.cpp b/Engine/Plugins/Media/ElectraPlayer/Source/ElectraPlayerRuntime/Private/Runtime/Player/AdaptiveStreamingPlayer.cpp
--- a/Engine/Plugins/Media/ElectraPlayer/Source/ElectraPlayerRuntime/Private/Runtime/Player/AdaptiveStreamingPlayer.cpp
+++ b/Engine/Plugins/Media/ElectraPlayer/Source/ElectraPlayerRuntime/Private/Runtime/Player/AdaptiveStreamingPlayer.cpp
@@ -3370,7 +3370,6 @@ void FAdaptiveStreamingPlayer::InternalHandleSegmentTrackChanges(const FTimeValu
 			if (Result == IManifest::IPlayPeriod::ETrackChangeResult::Changed)
 			{
 				TSharedPtrTS<FMultiTrackAccessUnitBuffer> StreamRcvBuf = GetCurrentReceiveStreamBuffer(EStreamType::Audio);
-				check(StreamRcvBuf.IsValid());
 				if (StreamRcvBuf.IsValid())
 				{
 					StreamRcvBuf->SelectTrackWhenAvailable(BufferSourceInfo);
@@ -3389,7 +3388,6 @@ void FAdaptiveStreamingPlayer::InternalHandleSegmentTrackChanges(const FTimeValu
 			if (Result == IManifest::IPlayPeriod::ETrackChangeResult::Changed)
 			{
 				TSharedPtrTS<FMultiTrackAccessUnitBuffer> StreamRcvBuf = GetCurrentReceiveStreamBuffer(EStreamType::Subtitle);
-				check(StreamRcvBuf.IsValid());
 				if (StreamRcvBuf.IsValid())
 				{
 					StreamRcvBuf->SelectTrackWhenAvailable(BufferSourceInfo);
@@ -3411,7 +3409,6 @@ void FAdaptiveStreamingPlayer::InternalHandleSegmentTrackChanges(const FTimeValu
 		{
 			TSharedPtrTS<FBufferSourceInfo> BufferSourceInfo = CurrentPlayPeriodAudio->GetSelectedStreamBufferSourceInfo(EStreamType::Audio);
 			TSharedPtrTS<FMultiTrackAccessUnitBuffer> StreamRcvBuf = GetCurrentReceiveStreamBuffer(EStreamType::Audio);
-			check(StreamRcvBuf.IsValid());
 			if (StreamRcvBuf.IsValid())
 			{
 				StreamRcvBuf->SelectTrackWhenAvailable(BufferSourceInfo);
