ElectraPlayer: Unescaping file:// URI to handle percent encoded characters

#rnx

diff --git a/Engine/Plugins/Media/ElectraPlayer/Source/ElectraPlayerRuntime/Private/Runtime/HTTP/HTTPManager.cpp b/Engine/Plugins/Media/ElectraPlayer/Source/ElectraPlayerRuntime/Private/Runtime/HTTP/HTTPManager.cpp
--- a/Engine/Plugins/Media/ElectraPlayer/Source/ElectraPlayerRuntime/Private/Runtime/HTTP/HTTPManager.cpp
+++ b/Engine/Plugins/Media/ElectraPlayer/Source/ElectraPlayerRuntime/Private/Runtime/HTTP/HTTPManager.cpp
@@ -3,6 +3,7 @@
 #include "HTTP/HTTPManager.h"
 #include "HTTP/HTTPResponseCache.h"
 #include "Utilities/StringHelpers.h"
+#include "Utilities/URLParser.h"
 #include "PlayerTime.h"
 #include "Player/PlayerSessionServices.h"
 #include "HAL/LowLevelMemTracker.h"
@@ -479,13 +480,23 @@ namespace Electra
 		}
 		else
 		{
-			FFileStream* FileStream = new FFileStream;
-			Handle->LocalByteStream = FileStream;
-			FileStream->Filename = Request->Parameters.URL.Mid(7); /* file:// */
-			FileStream->Archive = MakeShareable(IFileManager::Get().CreateFileReader(*FileStream->Filename));
-			if (!FileStream->Archive.IsValid())
+			FString Filename;
+			// Unescape percent encoded characters in the URI, such as %20 representing a space.
+			if (FURL_RFC3986::UrlDecode(Filename, Request->Parameters.URL))
 			{
-				OutError.Set(ERRCODE_HTTP_FILE_COULDNT_READ_FILE, FString::Printf(TEXT("Failed to open media file \"%s\""), *Request->Parameters.URL));
+				FFileStream* FileStream = new FFileStream;
+				Handle->LocalByteStream = FileStream;
+				FileStream->Filename = Filename.Mid(7); /* file:// */
+				FileStream->Archive = MakeShareable(IFileManager::Get().CreateFileReader(*FileStream->Filename));
+				if (!FileStream->Archive.IsValid())
+				{
+					OutError.Set(ERRCODE_HTTP_FILE_COULDNT_READ_FILE, FString::Printf(TEXT("Failed to open media file \"%s\""), *Request->Parameters.URL));
+					return nullptr;
+				}
+			}
+			else
+			{
+				OutError.Set(ERRCODE_HTTP_FILE_COULDNT_READ_FILE, FString::Printf(TEXT("Failed to parse file name \"%s\""), *Request->Parameters.URL));
 				return nullptr;
 			}
 		}
