ElectraPlayer: Added base time to timed metadata to enable relative time offsets in the metadata to calculate an effective absolute time.

#rnx

diff --git a/Engine/Plugins/Media/ElectraPlayer/Source/ElectraPlayerPlugin/Private/ElectraPlayerPlugin.cpp b/Engine/Plugins/Media/ElectraPlayer/Source/ElectraPlayerPlugin/Private/ElectraPlayerPlugin.cpp
--- a/Engine/Plugins/Media/ElectraPlayer/Source/ElectraPlayerPlugin/Private/ElectraPlayerPlugin.cpp
+++ b/Engine/Plugins/Media/ElectraPlayer/Source/ElectraPlayerPlugin/Private/ElectraPlayerPlugin.cpp
@@ -103,12 +103,12 @@ class FElectraBinarySample : public IElectraBinarySample
 {
 public:
 	virtual ~FElectraBinarySample() = default;
-	virtual const void* GetData() override						{ return Metadata->GetData(); }
-	virtual uint32 GetSize() const override						{ return Metadata->GetSize(); }
-	virtual FGuid GetGUID() const override						{ return IElectraBinarySample::GetSampleTypeGUID(); }
-	virtual const FString& GetSchemeIdUri() const override		{ return Metadata->GetSchemeIdUri(); }
-	virtual const FString& GetValue() const override			{ return Metadata->GetValue(); }
-	virtual const FString& GetID() const override				{ return Metadata->GetID(); }
+	virtual const void* GetData() override							{ return Metadata->GetData(); }
+	virtual uint32 GetSize() const override							{ return Metadata->GetSize(); }
+	virtual FGuid GetGUID() const override							{ return IElectraBinarySample::GetSampleTypeGUID(); }
+	virtual const FString& GetSchemeIdUri() const override			{ return Metadata->GetSchemeIdUri(); }
+	virtual const FString& GetValue() const override				{ return Metadata->GetValue(); }
+	virtual const FString& GetID() const override					{ return Metadata->GetID(); }
 
 	virtual EDispatchedMode GetDispatchedMode() const override
 	{
@@ -164,6 +164,17 @@ public:
 		return Duration;
 	}
 
+	virtual TOptional<FMediaTimeStamp> GetTrackBaseTime() const	override
+	{ 
+		TOptional<FMediaTimeStamp> ms;
+		TOptional<FDecoderTimeStamp> ts = Metadata->GetTime(); 
+		if (ts.IsSet())
+		{
+			ms = FMediaTimeStamp(ts.GetValue().Time, ts.GetValue().SequenceIndex);
+		}
+		return ms;
+	}
+
 	IMetaDataDecoderOutputPtr Metadata;
 };
 
diff --git a/Engine/Plugins/Media/ElectraPlayer/Source/ElectraPlayerRuntime/Private/ElectraPlayer.cpp b/Engine/Plugins/Media/ElectraPlayer/Source/ElectraPlayerRuntime/Private/ElectraPlayer.cpp
--- a/Engine/Plugins/Media/ElectraPlayer/Source/ElectraPlayerRuntime/Private/ElectraPlayer.cpp
+++ b/Engine/Plugins/Media/ElectraPlayer/Source/ElectraPlayerRuntime/Private/ElectraPlayer.cpp
@@ -84,15 +84,16 @@ class FMetaDataDecoderOutput : public IMetaDataDecoderOutput
 {
 public:
 	virtual ~FMetaDataDecoderOutput() = default;
-	virtual const void* GetData() override						{ return Data.GetData(); }
-	virtual FTimespan GetDuration() const override				{ return Duration; }
-	virtual uint32 GetSize() const override						{ return (uint32) Data.Num(); }
-	virtual FDecoderTimeStamp GetTime() const override			{ return PresentationTime; }
-	virtual EOrigin GetOrigin() const override					{ return Origin; }
-	virtual EDispatchedMode GetDispatchedMode() const override	{ return DispatchedMode; }
-	virtual const FString& GetSchemeIdUri() const override		{ return SchemeIdUri; }
-	virtual const FString& GetValue() const override			{ return Value; }
-	virtual const FString& GetID() const override				{ return ID; }
+	virtual const void* GetData() override									{ return Data.GetData(); }
+	virtual FTimespan GetDuration() const override							{ return Duration; }
+	virtual uint32 GetSize() const override									{ return (uint32) Data.Num(); }
+	virtual FDecoderTimeStamp GetTime() const override						{ return PresentationTime; }
+	virtual EOrigin GetOrigin() const override								{ return Origin; }
+	virtual EDispatchedMode GetDispatchedMode() const override				{ return DispatchedMode; }
+	virtual const FString& GetSchemeIdUri() const override					{ return SchemeIdUri; }
+	virtual const FString& GetValue() const override						{ return Value; }
+	virtual const FString& GetID() const override							{ return ID; }
+	virtual TOptional<FDecoderTimeStamp> GetTrackBaseTime() const override	{ return TrackBaseTime; }
 
 	TArray<uint8> Data;
 	FDecoderTimeStamp PresentationTime;
@@ -102,6 +103,7 @@ public:
 	FString SchemeIdUri;
 	FString Value;
 	FString ID;
+	TOptional<FDecoderTimeStamp> TrackBaseTime;
 };
 
 //-----------------------------------------------------------------------------
@@ -1572,8 +1574,12 @@ void FElectraPlayer::OnMediaPlayerEventReceived(TSharedPtrTS<IAdaptiveStreamingP
 #endif
 
 	TSharedPtr<IElectraPlayerAdapterDelegate, ESPMode::ThreadSafe> PinnedAdapterDelegate = AdapterDelegate.Pin();
-	if (PinnedAdapterDelegate.IsValid())
+	TSharedPtr<FInternalPlayerImpl, ESPMode::ThreadSafe> LockedPlayer = CurrentPlayer;
+	if (PinnedAdapterDelegate.IsValid() && LockedPlayer.IsValid() && LockedPlayer->AdaptivePlayer.IsValid())
 	{
+		Electra::FTimeRange MediaTimeline;
+		LockedPlayer->AdaptivePlayer->GetTimelineRange(MediaTimeline);
+
 		// Create a binary media sample of our extended format and pass it up.
 		TSharedPtr<FMetaDataDecoderOutput, ESPMode::ThreadSafe> Meta = MakeShared<FMetaDataDecoderOutput, ESPMode::ThreadSafe>();
 		switch(InDispatchMode)
@@ -1615,6 +1621,10 @@ void FElectraPlayer::OnMediaPlayerEventReceived(TSharedPtrTS<IAdaptiveStreamingP
 		Meta->ID = InEvent->GetID(),
 		Meta->Duration = InEvent->GetDuration().GetAsTimespan();
 		Meta->PresentationTime = FDecoderTimeStamp(InEvent->GetPresentationTime().GetAsTimespan(), 0);
+		// Set the current timeline start as the metadata track's zero point. This is only useful if the timeline does not
+		// actually change over time. The use of the base time is therefore tied to knowledge by the using code that the
+		// timeline will be fixed.
+		Meta->TrackBaseTime = FDecoderTimeStamp(MediaTimeline.Start.GetAsTimespan(), MediaTimeline.Start.GetSequenceIndex());
 		PinnedAdapterDelegate->PresentMetadataSample(Meta);
 	}
 }
diff --git a/Engine/Plugins/Media/ElectraUtil/Source/ElectraBase/Public/MediaMetaDataDecoderOutput.h b/Engine/Plugins/Media/ElectraUtil/Source/ElectraBase/Public/MediaMetaDataDecoderOutput.h
--- a/Engine/Plugins/Media/ElectraUtil/Source/ElectraBase/Public/MediaMetaDataDecoderOutput.h
+++ b/Engine/Plugins/Media/ElectraUtil/Source/ElectraBase/Public/MediaMetaDataDecoderOutput.h
@@ -44,6 +44,10 @@ public:
 
 	// Returns the event's ID.
 	virtual const FString& GetID() const = 0;
+
+	// Returns the base time (the "zero point") of the track, if available, to calculate the relative time
+	// of this metadata from GetTime() if this is based on some non-zero anchor.
+	virtual TOptional<FDecoderTimeStamp> GetTrackBaseTime() const = 0;
 };
 
 using IMetaDataDecoderOutputPtr = TSharedPtr<IMetaDataDecoderOutput, ESPMode::ThreadSafe>;
diff --git a/Engine/Plugins/Media/ElectraUtil/Source/ElectraSamples/Public/IElectraMetadataSample.h b/Engine/Plugins/Media/ElectraUtil/Source/ElectraSamples/Public/IElectraMetadataSample.h
--- a/Engine/Plugins/Media/ElectraUtil/Source/ElectraSamples/Public/IElectraMetadataSample.h
+++ b/Engine/Plugins/Media/ElectraUtil/Source/ElectraSamples/Public/IElectraMetadataSample.h
@@ -54,4 +54,8 @@ public:
 
 	// Returns the event's ID.
 	virtual const FString& GetID() const = 0;
+
+	// Returns the base time (the "zero point") of the track, if available, to calculate the relative time
+	// of this metadata from GetTime() if this is based on some non-zero anchor.
+	virtual TOptional<FMediaTimeStamp> GetTrackBaseTime() const = 0;
 };
