Electra: Cleanup of *.Build.cs scripts

#rnx

diff --git a/Engine/Plugins/Media/ElectraPlayer/Source/ElectraPlayerPlugin/ElectraPlayerPlugin.Build.cs b/Engine/Plugins/Media/ElectraPlayer/Source/ElectraPlayerPlugin/ElectraPlayerPlugin.Build.cs
--- a/Engine/Plugins/Media/ElectraPlayer/Source/ElectraPlayerPlugin/ElectraPlayerPlugin.Build.cs
+++ b/Engine/Plugins/Media/ElectraPlayer/Source/ElectraPlayerPlugin/ElectraPlayerPlugin.Build.cs
@@ -1,6 +1,5 @@
 // Copyright Epic Games, Inc. All Rights Reserved.
-using EpicGames.Core;
-using System.IO;
+using UnrealBuildTool;
 
 namespace UnrealBuildTool.Rules
 {
@@ -8,11 +7,10 @@ namespace UnrealBuildTool.Rules
 	{
 		public ElectraPlayerPlugin(ReadOnlyTargetRules Target) : base(Target)
 		{
+			bLegalToDistributeObjectCode = true;
+			PCHUsage = PCHUsageMode.NoPCHs;
 
-            bLegalToDistributeObjectCode = true;
-            PCHUsage = PCHUsageMode.NoPCHs;
-
-            DynamicallyLoadedModuleNames.AddRange(
+			DynamicallyLoadedModuleNames.AddRange(
 				new string[] {
 					"Media",
 				});
@@ -33,7 +31,7 @@ namespace UnrealBuildTool.Rules
 			PrivateIncludePathModuleNames.AddRange(
 			new string[] {
 				"Media",
-            });
+			});
 
 			if (Target.bCompileAgainstEngine)
 			{
@@ -41,14 +39,14 @@ namespace UnrealBuildTool.Rules
 			}
 
 			if (Target.Platform == UnrealTargetPlatform.Mac)
-            {
-                PrivateDependencyModuleNames.Add("MetalRHI");
-            }
+			{
+				PrivateDependencyModuleNames.Add("MetalRHI");
+			}
 
 			if (Target.Platform == UnrealTargetPlatform.IOS )
-            {
-                PrivateDependencyModuleNames.Add("MetalRHI");
-            }
+			{
+				PrivateDependencyModuleNames.Add("MetalRHI");
+			}
 
 			if (Target.Type == TargetType.Editor)
 			{
diff --git a/Engine/Plugins/Media/ElectraPlayer/Source/ElectraPlayerRuntime/ElectraPlayerRuntime.Build.cs b/Engine/Plugins/Media/ElectraPlayer/Source/ElectraPlayerRuntime/ElectraPlayerRuntime.Build.cs
--- a/Engine/Plugins/Media/ElectraPlayer/Source/ElectraPlayerRuntime/ElectraPlayerRuntime.Build.cs
+++ b/Engine/Plugins/Media/ElectraPlayer/Source/ElectraPlayerRuntime/ElectraPlayerRuntime.Build.cs
@@ -1,5 +1,5 @@
 // Copyright Epic Games, Inc. All Rights Reserved.
-using EpicGames.Core;
+using UnrealBuildTool;
 using System.IO;
 
 namespace UnrealBuildTool.Rules
@@ -12,8 +12,8 @@ namespace UnrealBuildTool.Rules
 			// Common setup...
 			//
 
-            bLegalToDistributeObjectCode = true;
-            PCHUsage = PCHUsageMode.NoPCHs;
+			bLegalToDistributeObjectCode = true;
+			PCHUsage = PCHUsageMode.NoPCHs;
 
 			PrivateDependencyModuleNames.AddRange(
 				new string[] {
@@ -42,19 +42,19 @@ namespace UnrealBuildTool.Rules
 			//
 
 			if (Target.Platform.IsInGroup(UnrealPlatformGroup.Windows))
-            {
-                string DirectXSDKDir = Target.UEThirdPartySourceDirectory + "Windows/DirectX";
-                PublicSystemIncludePaths.Add(DirectXSDKDir + "/include");
+			{
+				string DirectXSDKDir = Target.UEThirdPartySourceDirectory + "Windows/DirectX";
+				PublicSystemIncludePaths.Add(DirectXSDKDir + "/include");
 
 				if (Target.Platform == UnrealTargetPlatform.Win64)
-                {
-                    DirectXSDKDir += "/Lib/x64/";
-                }
+				{
+					DirectXSDKDir += "/Lib/x64/";
+				}
 
 				PrivateDefinitions.Add("_CRT_SECURE_NO_WARNINGS=1");
 
-                PrivateDefinitions.Add("ELECTRA_ENABLE_MFDECODER");
-                PrivateDefinitions.Add("ELECTRA_ENABLE_SWDECODE");
+				PrivateDefinitions.Add("ELECTRA_ENABLE_MFDECODER");
+				PrivateDefinitions.Add("ELECTRA_ENABLE_SWDECODE");
 				PrivateDefinitions.Add("ELECTRA_PLATFORM_HAS_H265_DECODER=1");
 
 				AddEngineThirdPartyPrivateStaticDependencies(Target, "WinHttp");
@@ -62,12 +62,12 @@ namespace UnrealBuildTool.Rules
 				if (WinSupportsDX9())
 				{
 					AddEngineThirdPartyPrivateStaticDependencies(Target, "DX9");
-					PrivateDefinitions.Add("ELECTRA_HAVE_DX9");     // video decoding on Win7 enabled - uses DX9
+					PrivateDefinitions.Add("ELECTRA_HAVE_DX9"); 	// video decoding on Win7 enabled - uses DX9
 				}
 
 				if (WinSupportsDX11())
 				{
-					PrivateDefinitions.Add("ELECTRA_HAVE_DX11");    // video decoding for DX11 enabled (Win8+)
+					PrivateDefinitions.Add("ELECTRA_HAVE_DX11");	// video decoding for DX11 enabled (Win8+)
 				}
 
 				if (Target.Platform == UnrealTargetPlatform.Win64)
@@ -78,29 +78,29 @@ namespace UnrealBuildTool.Rules
 				}
 
 				PublicSystemLibraries.AddRange(new string[] {
-                    "strmiids.lib",
-                    "legacy_stdio_definitions.lib",
-                    "Dxva2.lib",
-                });
+					"strmiids.lib",
+					"legacy_stdio_definitions.lib",
+					"Dxva2.lib",
+				});
 
-                // Delay-load all MF DLLs to be able to check Windows version for compatibility in `StartupModule` before loading them manually
-                PublicSystemLibraries.Add("mfplat.lib");
-                PublicDelayLoadDLLs.Add("mfplat.dll");
-                PublicSystemLibraries.Add("mfuuid.lib");
+				// Delay-load all MF DLLs to be able to check Windows version for compatibility in `StartupModule` before loading them manually
+				PublicSystemLibraries.Add("mfplat.lib");
+				PublicDelayLoadDLLs.Add("mfplat.dll");
+				PublicSystemLibraries.Add("mfuuid.lib");
 
 				PublicIncludePaths.Add("$(ModuleDir)/Public/Windows");
 				PrivateIncludePaths.Add("ElectraPlayerRuntime/Private/Runtime/Decoder/Windows");
 				PrivateIncludePaths.Add("ElectraPlayerRuntime/Private/Windows");
 			}
 			else if (Target.Platform == UnrealTargetPlatform.Mac)
-            {
-                PublicFrameworks.AddRange(
+			{
+				PublicFrameworks.AddRange(
 				new string[] {
 								"CoreMedia",
 								"CoreVideo",
 								"AVFoundation",
 								"AudioToolbox",
-                                "VideoToolbox",
+								"VideoToolbox",
 								"QuartzCore"
 				});
 
@@ -111,14 +111,14 @@ namespace UnrealBuildTool.Rules
 				PrivateIncludePaths.Add("ElectraPlayerRuntime/Private/Runtime/Decoder/Apple");
 			}
 			else if (Target.Platform == UnrealTargetPlatform.IOS )
-            {
+			{
 				PublicFrameworks.AddRange(
 				new string[] {
 								"CoreMedia",
 								"CoreVideo",
 								"AVFoundation",
 								"AudioToolbox",
-                                "VideoToolbox",
+								"VideoToolbox",
 								"QuartzCore"
 				});
 
@@ -129,16 +129,16 @@ namespace UnrealBuildTool.Rules
 				PrivateIncludePaths.Add("ElectraPlayerRuntime/Private/Runtime/Decoder/Apple");
 			}
 			else if (Target.IsInPlatformGroup(UnrealPlatformGroup.Android) )
-            {
+			{
 				PublicDefinitions.Add("CURL_ENABLE_DEBUG_CALLBACK=1");
 
-                if (Target.Configuration != UnrealTargetConfiguration.Shipping)
-                {
-                    PublicDefinitions.Add("CURL_ENABLE_NO_TIMEOUTS_OPTION=1");
-                }
+				if (Target.Configuration != UnrealTargetConfiguration.Shipping)
+				{
+					PublicDefinitions.Add("CURL_ENABLE_NO_TIMEOUTS_OPTION=1");
+				}
 
-                string PluginPath = Utils.MakePathRelativeTo(ModuleDirectory, Target.RelativeEnginePath);
-                AdditionalPropertiesForReceipt.Add("AndroidPlugin", Path.Combine(PluginPath, "ElectraPlayerRuntime_UPL.xml"));
+				string PluginPath = Utils.MakePathRelativeTo(ModuleDirectory, Target.RelativeEnginePath);
+				AdditionalPropertiesForReceipt.Add("AndroidPlugin", Path.Combine(PluginPath, "ElectraPlayerRuntime_UPL.xml"));
 
 				PrivateDefinitions.Add("ELECTRA_PLATFORM_HAS_H265_DECODER=1");
 
@@ -147,13 +147,13 @@ namespace UnrealBuildTool.Rules
 				PrivateIncludePaths.Add("ElectraPlayerRuntime/Private/Runtime/Decoder/Android");
 			}
 			else if (Target.IsInPlatformGroup(UnrealPlatformGroup.Unix) )
-            {
+			{
 				PublicDefinitions.Add("CURL_ENABLE_DEBUG_CALLBACK=1");
 
-                if (Target.Configuration != UnrealTargetConfiguration.Shipping)
-                {
-                    PublicDefinitions.Add("CURL_ENABLE_NO_TIMEOUTS_OPTION=1");
-                }
+				if (Target.Configuration != UnrealTargetConfiguration.Shipping)
+				{
+					PublicDefinitions.Add("CURL_ENABLE_NO_TIMEOUTS_OPTION=1");
+				}
 
 				PublicIncludePaths.Add("$(ModuleDir)/Public/Unix");
 
diff --git a/Engine/Plugins/Media/ElectraUtil/Source/ElectraSamples/ElectraSamples.Build.cs b/Engine/Plugins/Media/ElectraUtil/Source/ElectraSamples/ElectraSamples.Build.cs
--- a/Engine/Plugins/Media/ElectraUtil/Source/ElectraSamples/ElectraSamples.Build.cs
+++ b/Engine/Plugins/Media/ElectraUtil/Source/ElectraSamples/ElectraSamples.Build.cs
@@ -1,5 +1,5 @@
 // Copyright Epic Games, Inc. All Rights Reserved.
-using EpicGames.Core;
+using UnrealBuildTool;
 using System.IO;
 
 namespace UnrealBuildTool.Rules
@@ -9,10 +9,10 @@ namespace UnrealBuildTool.Rules
 		public ElectraSamples(ReadOnlyTargetRules Target) : base(Target)
 		{
 
-            bLegalToDistributeObjectCode = true;
-            PCHUsage = PCHUsageMode.NoPCHs;
+			bLegalToDistributeObjectCode = true;
+			PCHUsage = PCHUsageMode.NoPCHs;
 
-            DynamicallyLoadedModuleNames.AddRange(
+			DynamicallyLoadedModuleNames.AddRange(
 				new string[] {
 					"Media",
 				});
@@ -31,22 +31,22 @@ namespace UnrealBuildTool.Rules
 			PrivateIncludePathModuleNames.AddRange(
 				new string[] {
 					"Media",
-                });
+				});
 
 			if (Target.bCompileAgainstEngine)
 			{
 				PrivateDependencyModuleNames.Add("Engine");
 			}
 
-            if (Target.Platform.IsInGroup(UnrealPlatformGroup.Windows))
-            {
-                string DirectXSDKDir = Target.UEThirdPartySourceDirectory + "Windows/DirectX";
-                PublicSystemIncludePaths.Add(DirectXSDKDir + "/include");
+			if (Target.Platform.IsInGroup(UnrealPlatformGroup.Windows))
+			{
+				string DirectXSDKDir = Target.UEThirdPartySourceDirectory + "Windows/DirectX";
+				PublicSystemIncludePaths.Add(DirectXSDKDir + "/include");
 
 				if (Target.Platform == UnrealTargetPlatform.Win64)
-                {
-                    DirectXSDKDir += "/Lib/x64/";
-                }
+				{
+					DirectXSDKDir += "/Lib/x64/";
+				}
 
 				if (Target.Platform == UnrealTargetPlatform.Win64)
 				{
@@ -63,16 +63,16 @@ namespace UnrealBuildTool.Rules
 				}
 
 				PublicSystemLibraries.AddRange(new string[] {
-                    "strmiids.lib",
-                    "legacy_stdio_definitions.lib",
-                    "Dxva2.lib",
-                });
+					"strmiids.lib",
+					"legacy_stdio_definitions.lib",
+					"Dxva2.lib",
+				});
 
 				PublicIncludePaths.Add("$(ModuleDir)/Public/Windows");
 			}
 			else if (Target.Platform == UnrealTargetPlatform.Mac || Target.Platform == UnrealTargetPlatform.IOS)
-            {
-                PrivateDependencyModuleNames.Add("MetalRHI");
+			{
+				PrivateDependencyModuleNames.Add("MetalRHI");
 				PublicIncludePaths.Add("$(ModuleDir)/Public/Apple");
 			}
 			else if (Target.IsInPlatformGroup(UnrealPlatformGroup.Android))
