ElectraPlayer: Improved HTTP response cache to provide any byte range of cached data instead of it having to be an exact match

#rnx

diff --git a/Engine/Plugins/Media/ElectraPlayer/Source/ElectraPlayerRuntime/Private/Runtime/HTTP/HTTPManager.cpp b/Engine/Plugins/Media/ElectraPlayer/Source/ElectraPlayerRuntime/Private/Runtime/HTTP/HTTPManager.cpp
--- a/Engine/Plugins/Media/ElectraPlayer/Source/ElectraPlayerRuntime/Private/Runtime/HTTP/HTTPManager.cpp
+++ b/Engine/Plugins/Media/ElectraPlayer/Source/ElectraPlayerRuntime/Private/Runtime/HTTP/HTTPManager.cpp
@@ -170,14 +170,42 @@ namespace Electra
 				if (HttpResponseCache.IsValid())
 				{
 					TSharedPtrTS<IHTTPResponseCache::FCacheItem> CacheItem;
-					if (HttpResponseCache->GetCachedEntity(CacheItem, ActiveResponse.URL, ActiveResponse.Range))
+					IHTTPResponseCache::EScatterResult Result;
+					Result = HttpResponseCache->GetScatteredCacheEntity(CacheItem, ActiveResponse.URL, ActiveResponse.Range);
+					if (Result == IHTTPResponseCache::EScatterResult::FullHit)
 					{
+						check(CacheItem.IsValid());
 						if (CacheItem.IsValid())
 						{
 							ActiveResponse.CacheResponse = MoveTemp(CacheItem);
 							return true;
 						}
 					}
+					else if (Result == IHTTPResponseCache::EScatterResult::PartialHit)
+					{
+						check(CacheItem.IsValid());
+						if (CacheItem.IsValid())
+						{
+							// Partial cached responses require the missing partial data to be fetched via sub range requests.
+							ActiveResponse.bIsSubRangeRequest = true;
+
+							// If there is no cached response it means that the first bytes are not in the cache but later
+							// data is. The first part must be requested. Otherwise we can just return the cached response.
+							if (CacheItem->Response.IsValid())
+							{
+								ActiveResponse.CacheResponse = MoveTemp(CacheItem);
+								ActiveResponse.ReceivedContentRange = ActiveResponse.CacheResponse->Range;
+								return true;
+							}
+							else
+							{
+								// The missing range we need to fetch has been set up by the cache so we just need to use
+								// that range in the request.
+								ActiveResponse.Range = CacheItem->Range;
+								HttpRequest->SetHeader(TEXT("Range"), FString::Printf(TEXT("bytes=%s"), *ActiveResponse.Range.GetString()));
+							}
+						}
+					}
 				}
 				return HttpRequest->ProcessRequest();
 			}
@@ -230,6 +258,13 @@ namespace Electra
 					return 0;
 				}
 
+				void ClearForNextSubRange()
+				{
+					CacheResponse.Reset();
+					bHitCache = false;
+					bWasAddedToCache = false;
+				}
+
 			};
 
 			enum class EHandleType
@@ -277,6 +312,7 @@ namespace Electra
 				bResponseReceived = false;
 				bIsConnected = false;
 				bHaveResponseHeaders = false;
+				ActiveResponse.ClearForNextSubRange();
 			}
 		};
 
@@ -1324,6 +1360,16 @@ namespace Electra
 						TSharedPtrTS<IHTTPResponseCache::FCacheItem> CacheItem = MakeSharedTS<IHTTPResponseCache::FCacheItem>();
 						CacheItem->URL = Request->Parameters.URL;
 						CacheItem->Range = Handle->ActiveResponse.Range;
+						// Make sure the range is always set, even for non-partial responses.
+						CacheItem->Range.DocumentSize = ci.ContentLength > Handle->ActiveResponse.ReceivedContentRange.GetDocumentSize() ? ci.ContentLength : Handle->ActiveResponse.ReceivedContentRange.GetDocumentSize();
+						if (CacheItem->Range.GetStart() < 0)
+						{
+							CacheItem->Range.SetStart(0);
+						}
+						if (CacheItem->Range.GetEndIncluding() < 0)
+						{
+							CacheItem->Range.SetEndIncluding(CacheItem->Range.DocumentSize - 1);
+						}
 						CacheItem->Response = Handle->ActiveResponse.Response;
 						Request->ResponseCache->CacheEntity(CacheItem);
 					}
@@ -1332,30 +1378,24 @@ namespace Electra
 					TSharedPtrTS<FReceiveBuffer> ReceiveBuffer = Request->ReceiveBuffer.Pin();
 					if (ReceiveBuffer.IsValid())
 					{
+						bool bCheckForNext = false;
 						// Is it to be used as a linear buffer?
 						if (!ReceiveBuffer->bEnableRingbuffer)
 						{
 							const TArray<uint8>& Content = Handle->ActiveResponse.Response->GetContent();
-							// Set the actual content length now. What we got in the header before may have been incorrect if the content was zipped
-							// or was transferred using chunked content encoding.
-							ci.ContentLength = Content.Num();
-							if (ReceiveBuffer->Buffer.EnlargeTo(ci.ContentLength))
+
+							if (!ReceiveBuffer->Buffer.EnlargeTo(ReceiveBuffer->Buffer.Num() + Content.Num()) ||
+								!ReceiveBuffer->Buffer.PushData(Content.GetData(), Content.Num()))
 							{
-								if (ReceiveBuffer->Buffer.PushData(Content.GetData(), Content.Num()))
-								{
-									ci.BytesReadSoFar = ci.ContentLength;
-								}
-								else
-								{
-									ci.bWasAborted = true;
-								}
+								ci.bWasAborted = true;
+								bHasFinished = true;
 							}
 							else
 							{
-								// When OOM we just claim we've been aborted.
-								ci.bWasAborted = true;
+								ci.ContentLength = ReceiveBuffer->Buffer.Num();
+								ci.BytesReadSoFar = ci.ContentLength;
+								bCheckForNext = true;
 							}
-							bHasFinished = true;
 						}
 						else
 						{
@@ -1374,32 +1414,36 @@ namespace Electra
 								// All copied out now?
 								if (Handle->ActiveResponse.NumBytesPassedOut >= Content.Num())
 								{
-									// Check if this was a sub ranged request and if there is still data to go for the original request.
-									if (Handle->ActiveResponse.SizeRemaining() == 0)
+									bCheckForNext = true;
+								}
+							}
+						}
+						if (bCheckForNext)
+						{
+							// Check if this was a sub ranged request and if there is still data to go for the original request.
+							if (Handle->ActiveResponse.SizeRemaining() == 0)
+							{
+								// All done now.
+								bHasFinished = true;
+							}
+							else
+							{
+								// Still another sub range to go.
+								if (PrepareHTTPModuleHandle(Now, Handle, Request, false))
+								{
+									if (!Handle->ProcessHTTPRequest())
 									{
-										// All done now.
+										ci.StatusInfo.ErrorDetail.SetError(UEMEDIA_ERROR_INTERNAL).SetFacility(Facility::EFacility::HTTPReader).SetMessage("HTTP request failed on ProcessRequest()");
 										bHasFinished = true;
 									}
-									else
-									{
-										// Still another sub range to go.
-										if (PrepareHTTPModuleHandle(Now, Handle, Request, false))
-										{
-											if (!Handle->ProcessHTTPRequest())
-											{
-												ci.StatusInfo.ErrorDetail.SetError(UEMEDIA_ERROR_INTERNAL).SetFacility(Facility::EFacility::HTTPReader).SetMessage("HTTP request failed on ProcessRequest()");
-												bHasFinished = true;
-											}
-										}
-										else
-										{
-											// 
-											ci.StatusInfo.ErrorCode = ERRCODE_WRITE_ERROR;
-											ci.StatusInfo.bReadError = true;
-											ci.StatusInfo.ErrorDetail.SetMessage(FString::Printf(TEXT("Error setting up the next sub range request")));
-											bHasFinished = true;
-										}
-									}
+								}
+								else
+								{
+									// 
+									ci.StatusInfo.ErrorCode = ERRCODE_WRITE_ERROR;
+									ci.StatusInfo.bReadError = true;
+									ci.StatusInfo.ErrorDetail.SetMessage(FString::Printf(TEXT("Error setting up the next sub range request")));
+									bHasFinished = true;
 								}
 							}
 						}
diff --git a/Engine/Plugins/Media/ElectraPlayer/Source/ElectraPlayerRuntime/Private/Runtime/HTTP/HTTPResponseCache.cpp b/Engine/Plugins/Media/ElectraPlayer/Source/ElectraPlayerRuntime/Private/Runtime/HTTP/HTTPResponseCache.cpp
--- a/Engine/Plugins/Media/ElectraPlayer/Source/ElectraPlayerRuntime/Private/Runtime/HTTP/HTTPResponseCache.cpp
+++ b/Engine/Plugins/Media/ElectraPlayer/Source/ElectraPlayerRuntime/Private/Runtime/HTTP/HTTPResponseCache.cpp
@@ -19,11 +19,78 @@ public:
 
 	virtual void HandleEntityExpiration() override;
 	virtual void CacheEntity(TSharedPtrTS<FCacheItem> EntityToAdd) override;
-	virtual bool GetCachedEntity(TSharedPtrTS<FCacheItem>& OutCachedEntity, const FString& URL, const IElectraHttpManager::FParams::FRange& Range) override;
+	virtual EScatterResult GetScatteredCacheEntity(TSharedPtrTS<FCacheItem>& OutScatteredCachedEntity, const FString& URL, const IElectraHttpManager::FParams::FRange& Range) override;
 
 private:
+	class FResponseFromCache : public IHttpResponse
+	{
+	public:
+		virtual ~FResponseFromCache() = default;
+		FResponseFromCache() = default;
+
+		virtual int32 GetResponseCode() const override
+		{ return HTTPResponseCode; }
+
+		virtual FString GetContentAsString() const override
+		{ check(!"do not call"); return FString(); }
+
+		virtual FString GetURL() const override
+		{ return URL; }
+
+		virtual FString GetURLParameter(const FString& ParameterName) const override
+		{ check(!"do not call"); return FString(); }
+
+		virtual FString GetHeader(const FString& HeaderName) const override
+		{ check(!"do not call"); return FString(); }
+
+		virtual TArray<FString> GetAllHeaders() const override
+		{ return Headers; }
+
+		virtual FString GetContentType() const override
+		{ return GetHeader(TEXT("Content-Type")); }
+
+		virtual int32 GetContentLength() const override
+		{ return Payload.Num(); }
+
+		virtual const TArray<uint8>& GetContent() const override
+		{ return Payload; }
+
+		void CopyHeadersFromExceptContentSizes(FHttpResponsePtr CachedResponse)
+		{
+			check(CachedResponse);
+			if (CachedResponse.IsValid())
+			{
+				// Copy all headers except for those containing the content size or length.
+				for(auto &Header : CachedResponse->GetAllHeaders())
+				{
+					if (!Header.StartsWith(TEXT("Content-Length:"), ESearchCase::IgnoreCase) &&
+						!Header.StartsWith(TEXT("Content-Range:"), ESearchCase::IgnoreCase))
+					{
+						Headers.Emplace(MoveTemp(Header));
+					}
+				}
+			}
+		}
+		void AddHeader(FString InHeader)
+		{ Headers.Emplace(MoveTemp(InHeader)); }
+
+		void SetPayload(TArray<uint8> InPayload)
+		{ Payload = MoveTemp(InPayload); }
+
+		void SetURL(const FString& InUrl)
+		{ URL = InUrl; }
+
+		void SetHTTPResponseCode(int32 InResponseCode)
+		{ HTTPResponseCode = InResponseCode; }
+
+	private:
+		TArray<uint8> Payload;
+		TArray<FString> Headers;
+		FString URL;
+		int32 HTTPResponseCode = 200;
+	};
+
 	void EvictToAddSize(int64 ResponseSize);
-	void RemoveIfExists(TSharedPtrTS<FCacheItem> Entity);
 
 	IPlayerSessionServices* SessionServices = nullptr;
 	int64 MaxElementSize = 0;
@@ -83,21 +150,6 @@ void FHTTPResponseCache::EvictToAddSize(int64 ResponseSize)
 	}
 }
 
-void FHTTPResponseCache::RemoveIfExists(TSharedPtrTS<FCacheItem> Entity)
-{
-	FScopeLock ScopeLock(&Lock);
-	TSharedPtrTS<FCacheItem> Item;
-	for(int32 i=0, iMax=Cache.Num(); i<iMax; ++i)
-	{
-		if (Cache[i]->URL.Equals(Entity->URL) && Cache[i]->Range.Equals(Entity->Range))
-		{
-			SizeInUse -= Cache[i]->Response->GetContent().Num();
-			Cache.RemoveAt(i);
-			break;
-		}
-	}
-}
-
 
 void FHTTPResponseCache::HandleEntityExpiration()
 {
@@ -108,7 +160,6 @@ void FHTTPResponseCache::CacheEntity(TSharedPtrTS<FCacheItem> EntityToAdd)
 {
 	if (EntityToAdd.IsValid())
 	{
-		RemoveIfExists(EntityToAdd);
 		int64 SizeRequired = EntityToAdd->Response->GetContent().Num();
 		EvictToAddSize(SizeRequired);
 		FScopeLock ScopeLock(&Lock);
@@ -120,29 +171,154 @@ void FHTTPResponseCache::CacheEntity(TSharedPtrTS<FCacheItem> EntityToAdd)
 	}
 }
 
-bool FHTTPResponseCache::GetCachedEntity(TSharedPtrTS<FCacheItem>& OutCachedEntity, const FString& URL, const IElectraHttpManager::FParams::FRange& Range)
+FHTTPResponseCache::EScatterResult FHTTPResponseCache::GetScatteredCacheEntity(TSharedPtrTS<FCacheItem>& OutScatteredCachedEntity, const FString& URL, const IElectraHttpManager::FParams::FRange& InRange)
 {
-	FScopeLock ScopeLock(&Lock);
-	TSharedPtrTS<FCacheItem> Item;
-	for(int32 i=0, iMax=Cache.Num(); i<iMax; ++i)
+	FScopeLock lock(&Lock);
+	// Get a list of all cached blocks for this URL.
+	TArray<TSharedPtrTS<FCacheItem>> CachedBlocks;
+	for(auto &Entry : Cache)
 	{
-		if (Cache[i]->URL.Equals(URL) && Cache[i]->Range.Equals(Range))
+		if (Entry->URL.Equals(URL))
 		{
-			Item = Cache[i];
-			// Move to the beginning
-			if (i > 0)
+			CachedBlocks.Add(Entry);
+		}
+	}
+
+	IElectraHttpManager::FParams::FRange Range(InRange);
+
+	// Set up the default output that encompasses the request.
+	// This is what needs to be fetched for cache misses.
+	OutScatteredCachedEntity = MakeSharedTS<IHTTPResponseCache::FCacheItem>();
+	OutScatteredCachedEntity->URL = URL;
+	OutScatteredCachedEntity->Range = Range;
+
+	// Quick out if there are no cached blocks at all.
+	if (CachedBlocks.Num() == 0)
+	{
+		return EScatterResult::Miss;
+	}
+
+	// Adjust the input range to valid values.
+	Range.DocumentSize = CachedBlocks[0]->Range.GetDocumentSize();
+	if (Range.GetStart() < 0)
+	{
+		Range.SetStart(0);
+	}
+	if (Range.GetEndIncluding() < 0)
+	{
+		Range.SetEndIncluding(Range.GetDocumentSize() - 1);
+	}
+
+	// Sort the blocks by ascending range.
+	CachedBlocks.Sort([](const TSharedPtrTS<FCacheItem>& e1, const TSharedPtrTS<FCacheItem>& e2)
+	{
+		return e1->Range.GetStart() < e2->Range.GetStart();
+	});
+
+	// Find a cached block that overlaps the requested range.
+	int64 rs = Range.GetStart();
+	int64 re = Range.GetEndIncluding();
+	int32 firstIndex;
+	for(firstIndex=0; firstIndex<CachedBlocks.Num(); ++firstIndex)
+	{
+		// Any overlap?
+		const int64 cbstart = CachedBlocks[firstIndex]->Range.GetStart();
+		const int64 cbend = CachedBlocks[firstIndex]->Range.GetEndIncluding();
+		if (re >= cbstart && rs <= cbend)
+		{
+			// Is the request start before the cached block, ie we need to get additional data at the beginning?
+			if (rs < cbstart)
+			{
+				// The end of the requested range may also not be covered by the cache, but that is not relevant
+				// here as we move forward sequentially, one overlap at a time.
+				// Data up to the beginning of this cached block must be fetched first.
+				OutScatteredCachedEntity->Range = Range;
+				OutScatteredCachedEntity->Range.SetEndIncluding(cbstart - 1);
+				return EScatterResult::PartialHit;
+			}
+			else
 			{
-				Cache.RemoveAt(i);
-				Cache.Insert(Item, 0);
+				break;
 			}
-			OutCachedEntity = MoveTemp(Item);
-			return true;
 		}
 	}
-	return false;
-}
+	// No overlap anywhere?
+	if (firstIndex == CachedBlocks.Num())
+	{
+		return EScatterResult::Miss;
+	}
+
+	// Exact match?
+	if (rs == CachedBlocks[firstIndex]->Range.GetStart() && re == CachedBlocks[firstIndex]->Range.GetEndIncluding())
+	{
+		OutScatteredCachedEntity = CachedBlocks[firstIndex];
+		Cache.Remove(CachedBlocks[firstIndex]);
+		Cache.Insert(CachedBlocks[firstIndex], 0);
+		return EScatterResult::FullHit;
+	}
+
+	// Set up the partial cache response.
+	TSharedPtr<FResponseFromCache, ESPMode::ThreadSafe> Response = MakeShared<FResponseFromCache, ESPMode::ThreadSafe>();
+	Response->SetURL(URL);
+	Response->CopyHeadersFromExceptContentSizes(CachedBlocks[firstIndex]->Response);
+
+	// There is some overlap. Find how many consecutive bytes we can get from the cached blocks.
+	TArray<uint8> Payload;
+	TArray<TSharedPtrTS<FCacheItem>> TouchedBlocks;
+	while(1)
+	{
+		int64 cbstart = CachedBlocks[firstIndex]->Range.GetStart();
+		int64 cbend = CachedBlocks[firstIndex]->Range.GetEndIncluding();
+		int64 last = re < cbend ? re : cbend;
+		int64 offset = rs - cbstart;
+		int64 numAvail = last + 1 - cbstart - offset;
 
+		// Copy the response data into the gathering block.
+		const TArray<uint8>& CachedContent = CachedBlocks[firstIndex]->Response->GetContent();
+		Payload.Append(CachedContent.GetData() + offset, numAvail);
 
+		// Remember which blocks we just touched.
+		TouchedBlocks.Add(CachedBlocks[firstIndex]);
+
+		// Move up the start to the end of this cached block.
+		rs = last + 1;
+
+		// End fully included?
+		if (re <= cbend)
+		{
+			break;
+		}
+		// Is there another cached block?
+		if (++firstIndex >= CachedBlocks.Num())
+		{
+			break;
+		}
+		// Does its start coincide with what we need next?
+		if (CachedBlocks[firstIndex]->Range.GetStart() != rs)
+		{
+			break;
+		}
+	}
+
+	// Move the blocks we touched to the front of the cached blocks (LRU).
+	for(auto &Touched : TouchedBlocks)
+	{
+		Cache.Remove(Touched);
+		Cache.Insert(Touched, 0);
+	}
+
+	Range.SetEndIncluding(rs - 1);
+	if (InRange.IsSet())
+	{
+		Response->AddHeader(FString::Printf(TEXT("Content-Range: bytes %lld-%lld/%lld"), (long long int)Range.GetStart(), (long long int)Range.GetEndIncluding(), (long long int)Range.GetDocumentSize()));
+		Response->SetHTTPResponseCode(206);
+	}
+	Response->AddHeader(FString::Printf(TEXT("Content-Length: %lld"), (long long int)Payload.Num()));
+	Response->SetPayload(MoveTemp(Payload));
+	OutScatteredCachedEntity->Response = Response;
+	OutScatteredCachedEntity->Range = Range;
+	return EScatterResult::PartialHit;
+}
 
 
 } // namespace Electra
diff --git a/Engine/Plugins/Media/ElectraPlayer/Source/ElectraPlayerRuntime/Private/Runtime/HTTP/HTTPResponseCache.h b/Engine/Plugins/Media/ElectraPlayer/Source/ElectraPlayerRuntime/Private/Runtime/HTTP/HTTPResponseCache.h
--- a/Engine/Plugins/Media/ElectraPlayer/Source/ElectraPlayerRuntime/Private/Runtime/HTTP/HTTPResponseCache.h
+++ b/Engine/Plugins/Media/ElectraPlayer/Source/ElectraPlayerRuntime/Private/Runtime/HTTP/HTTPResponseCache.h
@@ -36,11 +36,24 @@ public:
 	 * Add an entity to the cache.
 	 */
 	virtual void CacheEntity(TSharedPtrTS<FCacheItem> EntityToAdd) = 0;
+
+
+	enum EScatterResult
+	{
+		FullHit,			//!< The entire request is present in the cache
+		PartialHit,			//!< Partial data is present in the cache
+		Miss				//!< No part of the request is present in the cache
+	};
+
 	/**
-	 * Return a cached entity.
-	 * Returns true if found, false if not.
+	 * Returns the first cached portion of a request.
+	 * This is used to create responses for which partially cached data already exist and break the request apart
+	 * to fetch the next portion that is not cached yet.
+	 * 
+	 * Returns false to indicate that the range specified in the (otherwise empty) cached response must be requested from the origin.
+	 * Returns true if a - possibly partial - cached response exists. The range that is cached is given in the cache item.
 	 */
-	virtual bool GetCachedEntity(TSharedPtrTS<FCacheItem>& OutCachedEntity, const FString& URL, const IElectraHttpManager::FParams::FRange& Range) = 0;
+	virtual EScatterResult GetScatteredCacheEntity(TSharedPtrTS<FCacheItem>& OutScatteredCachedEntity, const FString& URL, const IElectraHttpManager::FParams::FRange& Range) = 0;
 
 protected:
 	IHTTPResponseCache() = default;
diff --git a/Engine/Plugins/Media/ElectraPlayer/Source/ElectraPlayerRuntime/Private/Runtime/Player/mp4/PlaylistReaderMP4.cpp b/Engine/Plugins/Media/ElectraPlayer/Source/ElectraPlayerRuntime/Private/Runtime/Player/mp4/PlaylistReaderMP4.cpp
--- a/Engine/Plugins/Media/ElectraPlayer/Source/ElectraPlayerRuntime/Private/Runtime/Player/mp4/PlaylistReaderMP4.cpp
+++ b/Engine/Plugins/Media/ElectraPlayer/Source/ElectraPlayerRuntime/Private/Runtime/Player/mp4/PlaylistReaderMP4.cpp
@@ -337,6 +337,7 @@ void FPlaylistReaderMP4::ReadNextChunk(int64 InFromOffset, int64 ChunkSize)
 	Request->Parameters.Range.SetEndIncluding(LastByte);
 	Request->ReceiveBuffer = ReceiveBuffer;
 	Request->ProgressListener = ProgressListener;
+	Request->ResponseCache = PlayerSessionServices->GetHTTPResponseCache();
 	PlayerSessionServices->GetHTTPManager()->AddRequest(Request, false);
 }
 
