ElectraPlayer: Optimized seeking for faster frame scrubbing via seek parameter

#rnx

diff --git a/Engine/Plugins/Media/ElectraPlayer/Source/ElectraPlayerRuntime/Private/ElectraPlayer.cpp b/Engine/Plugins/Media/ElectraPlayer/Source/ElectraPlayerRuntime/Private/ElectraPlayer.cpp
--- a/Engine/Plugins/Media/ElectraPlayer/Source/ElectraPlayerRuntime/Private/ElectraPlayer.cpp
+++ b/Engine/Plugins/Media/ElectraPlayer/Source/ElectraPlayerRuntime/Private/ElectraPlayer.cpp
@@ -1067,23 +1067,13 @@ bool FElectraPlayer::Seek(const FTimespan& Time)
 {
 	if (CurrentPlayer.Get())
 	{
-		FTimespan CurrentPos = GetTime();
-		double Distance = (Time - CurrentPos).GetTotalSeconds();
-		// If the new position is close to the current position do not perform a seek.
-		if (Distance < -0.01 || Distance > 0.01)
-		{
-			FTimespan Target;
-			CalculateTargetSeekTime(Target, Time);
-			Electra::IAdaptiveStreamingPlayer::FSeekParam seek;
-			seek.Time.SetFromTimespan(Target);
-			bInitialSeekPerformed = true;
-			CurrentPlayer->AdaptivePlayer->SeekTo(seek);
-			return true;
-		}
-		else
-		{
-			UE_LOG(LogElectraPlayer, Verbose, TEXT("[%p][%p] IMediaPlayer::Seek(%.3f) ignored. Only %.3f seconds away from current position @%.3f"), this, CurrentPlayer.Get(), Time.GetTotalSeconds(), Distance, CurrentPos.GetTotalSeconds());
-		}
+		FTimespan Target;
+		CalculateTargetSeekTime(Target, Time);
+		Electra::IAdaptiveStreamingPlayer::FSeekParam seek;
+		seek.Time.SetFromTimespan(Target);
+		bInitialSeekPerformed = true;
+		CurrentPlayer->AdaptivePlayer->SeekTo(seek);
+		return true;
 	}
 	return false;
 }
@@ -1423,13 +1413,13 @@ FString FElectraPlayer::GetTrackName(EPlayerTrackType TrackType, int32 TrackInde
 
 /**
  * Selects a specified track for playback.
- * 
+ *
  * Note:
  *   There is currently no concept of selecting a track based on metadata, only by index.
  *   The idea being that before selecting a track by index the application needs to check
  *   the metadata beforehand (eg. call GetTrackLanguage()) to figure out the index of the
  *   track it wants to play.
- * 
+ *
  *   The underlying player however needs to select tracks based on metadata alone instead
  *   of an index in case the track layout changes dynamically during playback.
  *   For example, a part of the presentation could have both English and French audio,
@@ -1438,7 +1428,7 @@ FString FElectraPlayer::GetTrackName(EPlayerTrackType TrackType, int32 TrackInde
  *   player needs to automatically switch from French to English and back to French, or
  *   index 1 -> 0 -> 1 (assuming French was the starting language of choice).
  *   Indices are therefore meaningless to the underlying player.
- * 
+ *
  *   SelectTrack() is currently called implicitly by FMediaPlayerFacade::SelectDefaultTracks()
  *   when EMediaEvent::TracksChanged is received. This is why this event is NOT sent out
  *   in HandlePlayerEventTracksChanged() when the underlying player notifies us about a
@@ -1599,17 +1589,17 @@ void FElectraPlayer::OnMediaPlayerEventReceived(TSharedPtrTS<IAdaptiveStreamingP
 		switch(InEvent->GetOrigin())
 		{
 			default:
-			case IAdaptiveStreamingPlayerAEMSEvent::EOrigin::TimedMetadata:		
+			case IAdaptiveStreamingPlayerAEMSEvent::EOrigin::TimedMetadata:
 			{
 				Meta->Origin = FMetaDataDecoderOutput::EOrigin::TimedMetadata;
 				break;
 			}
-			case IAdaptiveStreamingPlayerAEMSEvent::EOrigin::EventStream:		
+			case IAdaptiveStreamingPlayerAEMSEvent::EOrigin::EventStream:
 			{
 				Meta->Origin = FMetaDataDecoderOutput::EOrigin::EventStream;
 				break;
 			}
-			case IAdaptiveStreamingPlayerAEMSEvent::EOrigin::InbandEventStream:	
+			case IAdaptiveStreamingPlayerAEMSEvent::EOrigin::InbandEventStream:
 			{
 				Meta->Origin = FMetaDataDecoderOutput::EOrigin::InbandEventStream;
 				break;
@@ -1796,6 +1786,11 @@ void FElectraPlayer::HandleDeferredPlayerEvents()
 				HandlePlayerEventPlaybackStopped();
 				break;
 			}
+			case FPlayerMetricEventBase::EType::SeekCompleted:
+			{
+				HandlePlayerEventSeekCompleted();
+				break;
+			}
 			case FPlayerMetricEventBase::EType::Error:
 			{
 				FPlayerMetricEvent_Error* Ev = static_cast<FPlayerMetricEvent_Error*>(Event.Get());
@@ -2103,7 +2098,7 @@ void FElectraPlayer::HandlePlayerEventBufferingEnd(Electra::Metrics::EBufferingR
 			Statistics.InitialBufferingDuration = BufferingDuration;
 			break;
 		case Electra::Metrics::EBufferingReason::Seeking:
-			MediaStateOnSeekFinished();
+			// End of seek buffering is not relevant here.
 			break;
 		case Electra::Metrics::EBufferingReason::Rebuffering:
 			if (BufferingDuration > Statistics.LongestRebufferingDuration)
@@ -2475,6 +2470,11 @@ void FElectraPlayer::HandlePlayerEventPlaybackStopped()
 	}
 }
 
+void FElectraPlayer::HandlePlayerEventSeekCompleted()
+{
+	UE_LOG(LogElectraPlayer, Log, TEXT("[%p][%p] Seek completed"), this, CurrentPlayer.Get());
+	MediaStateOnSeekFinished();
+}
 
 void FElectraPlayer::HandlePlayerEventError(const FString& ErrorReason)
 {
diff --git a/Engine/Plugins/Media/ElectraPlayer/Source/ElectraPlayerRuntime/Private/ElectraPlayer.h b/Engine/Plugins/Media/ElectraPlayer/Source/ElectraPlayerRuntime/Private/ElectraPlayer.h
--- a/Engine/Plugins/Media/ElectraPlayer/Source/ElectraPlayerRuntime/Private/ElectraPlayer.h
+++ b/Engine/Plugins/Media/ElectraPlayer/Source/ElectraPlayerRuntime/Private/ElectraPlayer.h
@@ -196,6 +196,7 @@ private:
 			PlaybackEnded,
 			JumpInPlayPosition,
 			PlaybackStopped,
+			SeekCompleted,
 			Error,
 			LogMessage,
 			DroppedVideoFrame,
@@ -346,6 +347,8 @@ private:
 	{ DeferredPlayerEvents.Enqueue(MakeSharedTS<FPlayerMetricEvent_JumpInPlayPosition>(ToNewTime, FromTime, TimejumpReason)); }
 	virtual void ReportPlaybackStopped() override
 	{ DeferredPlayerEvents.Enqueue(MakeSharedTS<FPlayerMetricEventBase>(FPlayerMetricEventBase::EType::PlaybackStopped)); }
+	virtual void ReportSeekCompleted() override
+	{ DeferredPlayerEvents.Enqueue(MakeSharedTS<FPlayerMetricEventBase>(FPlayerMetricEventBase::EType::SeekCompleted)); }
 	virtual void ReportError(const FString& ErrorReason) override
 	{ DeferredPlayerEvents.Enqueue(MakeSharedTS<FPlayerMetricEvent_Error>(ErrorReason)); }
 	virtual void ReportLogMessage(IInfoLog::ELevel InLogLevel, const FString& InLogMessage, int64 InPlayerWallclockMilliseconds) override
@@ -753,6 +756,7 @@ private:
 	void HandlePlayerEventPlaybackEnded();
 	void HandlePlayerEventJumpInPlayPosition(const FTimeValue& ToNewTime, const FTimeValue& FromTime, Metrics::ETimeJumpReason TimejumpReason);
 	void HandlePlayerEventPlaybackStopped();
+	void HandlePlayerEventSeekCompleted();
 	void HandlePlayerEventError(const FString& ErrorReason);
 	void HandlePlayerEventLogMessage(IInfoLog::ELevel InLogLevel, const FString& InLogMessage, int64 InPlayerWallclockMilliseconds);
 	void HandlePlayerEventDroppedVideoFrame();
diff --git a/Engine/Plugins/Media/ElectraPlayer/Source/ElectraPlayerRuntime/Private/Runtime/Decoder/Apple/VideoDecoderH264_Apple.cpp b/Engine/Plugins/Media/ElectraPlayer/Source/ElectraPlayerRuntime/Private/Runtime/Decoder/Apple/VideoDecoderH264_Apple.cpp
--- a/Engine/Plugins/Media/ElectraPlayer/Source/ElectraPlayerRuntime/Private/Runtime/Decoder/Apple/VideoDecoderH264_Apple.cpp
+++ b/Engine/Plugins/Media/ElectraPlayer/Source/ElectraPlayerRuntime/Private/Runtime/Decoder/Apple/VideoDecoderH264_Apple.cpp
@@ -1330,7 +1330,7 @@ FVideoDecoderH264::EDecodeResult FVideoDecoderH264::Decode(TSharedPtr<FDecoderIn
 */
 // TODO: set the proper flags. This may require a larger intermediate output queue. Needs experimenting.
 	VTDecodeFrameFlags DecodeFlags = kVTDecodeFrame_EnableAsynchronousDecompression | kVTDecodeFrame_EnableTemporalProcessing;
-	if (bRecreatingSession)
+	if (bRecreatingSession || !AU->AdjustedPTS.IsValid())
 	{
 		DecodeFlags = kVTDecodeFrame_DoNotOutputFrame;
 	}
diff --git a/Engine/Plugins/Media/ElectraPlayer/Source/ElectraPlayerRuntime/Private/Runtime/Decoder/Apple/VideoDecoderH265_Apple.cpp b/Engine/Plugins/Media/ElectraPlayer/Source/ElectraPlayerRuntime/Private/Runtime/Decoder/Apple/VideoDecoderH265_Apple.cpp
--- a/Engine/Plugins/Media/ElectraPlayer/Source/ElectraPlayerRuntime/Private/Runtime/Decoder/Apple/VideoDecoderH265_Apple.cpp
+++ b/Engine/Plugins/Media/ElectraPlayer/Source/ElectraPlayerRuntime/Private/Runtime/Decoder/Apple/VideoDecoderH265_Apple.cpp
@@ -1312,7 +1312,7 @@ FVideoDecoderH265::EDecodeResult FVideoDecoderH265::Decode(TSharedPtr<FDecoderIn
 */
 // TODO: set the proper flags. This may require a larger intermediate output queue. Needs experimenting.
 	VTDecodeFrameFlags DecodeFlags = kVTDecodeFrame_EnableAsynchronousDecompression | kVTDecodeFrame_EnableTemporalProcessing;
-	if (bRecreatingSession)
+	if (bRecreatingSession || !AU->AdjustedPTS.IsValid())
 	{
 		DecodeFlags = kVTDecodeFrame_DoNotOutputFrame;
 	}
diff --git a/Engine/Plugins/Media/ElectraPlayer/Source/ElectraPlayerRuntime/Private/Runtime/Player/AdaptivePlayerOptionKeynames.h b/Engine/Plugins/Media/ElectraPlayer/Source/ElectraPlayerRuntime/Private/Runtime/Player/AdaptivePlayerOptionKeynames.h
--- a/Engine/Plugins/Media/ElectraPlayer/Source/ElectraPlayerRuntime/Private/Runtime/Player/AdaptivePlayerOptionKeynames.h
+++ b/Engine/Plugins/Media/ElectraPlayer/Source/ElectraPlayerRuntime/Private/Runtime/Player/AdaptivePlayerOptionKeynames.h
@@ -26,6 +26,9 @@ namespace Electra
 
 	//! (bool) true to perform frame accurate seeks (slow as decoding and discarding data from a preceeding keyframe is required)
 	const TCHAR* const OptionKeyFrameAccurateSeek = TEXT("frame_accurate_seeking");
+
+	//! (bool) true to optimize seeking for faster frame scrubbing, false to optimize for playback.
+	const TCHAR* const OptionKeyFrameOptimizeSeekForScrubbing = TEXT("optimize_seek_for_scrubbing");
 	
 	//! (FTimeValue) absolute start time (including) of the range to limit playback to
 	const TCHAR* const OptionPlayRangeStart = TEXT("play_range_start");
diff --git a/Engine/Plugins/Media/ElectraPlayer/Source/ElectraPlayerRuntime/Private/Runtime/Player/AdaptiveStreamingPlayer.cpp b/Engine/Plugins/Media/ElectraPlayer/Source/ElectraPlayerRuntime/Private/Runtime/Player/AdaptiveStreamingPlayer.cpp
--- a/Engine/Plugins/Media/ElectraPlayer/Source/ElectraPlayerRuntime/Private/Runtime/Player/AdaptiveStreamingPlayer.cpp
+++ b/Engine/Plugins/Media/ElectraPlayer/Source/ElectraPlayerRuntime/Private/Runtime/Player/AdaptiveStreamingPlayer.cpp
@@ -55,13 +55,13 @@ namespace
 	{
 		switch(StreamType)
 		{
-			case EStreamType::Video: 
+			case EStreamType::Video:
 				return 0;
-			case EStreamType::Audio: 
+			case EStreamType::Audio:
 				return 1;
-			case EStreamType::Subtitle: 
+			case EStreamType::Subtitle:
 				return 2;
-			default: 
+			default:
 				return 3;
 		}
 	}
@@ -131,7 +131,8 @@ FAdaptiveStreamingPlayer::FAdaptiveStreamingPlayer(const IAdaptiveStreamingPlaye
 
 	bShouldBePaused  = false;
 	bShouldBePlaying = false;
-	bSeekPending	 = false;
+
+	SeekVars.Clear();
 
 	CurrentLoopState = {};
 	NextLoopStates.Empty();
@@ -393,7 +394,7 @@ void FAdaptiveStreamingPlayer::Initialize(const FParamDict& Options)
 	GetCodecSelectionPriorities(CodecPrioritiesAudio,     PlayerOptions, TEXT("preferred_codecs_audio"),     TEXT("audio"));
 	GetCodecSelectionPriorities(CodecPrioritiesSubtitles, PlayerOptions, TEXT("preferred_codecs_subtitles"), TEXT("subtitle"));
 
-	
+
 	// Unless already specified in the options to either value, enable frame accurate seeking now.
 	if (!PlayerOptions.HaveKey(OptionKeyFrameAccurateSeek))
 	{
@@ -1204,28 +1205,32 @@ bool FAdaptiveStreamingPlayer::OnFragmentAccessUnitReceived(FAccessUnit* InAcces
 	{
 		return true;
 	}
-	
-	// Get the current buffer stats. Although they are updated regularly only once per tick and would not reflect the current
+
+	// Get the current buffer stats. Although they are updated regularly once per tick this would not reflect the current
 	// state until the the next tick, which would allow the limits to be ignored as new access units arrive quickly.
 	FAccessUnitBufferInfo bufferInfo;
+	const FAccessUnitBuffer::FConfiguration* BufferConfig = nullptr;
 	if (InAccessUnit->ESType == EStreamType::Video)
 	{
 		GetStreamBufferUtilization(bufferInfo, InAccessUnit->ESType);
+		BufferConfig = &PlayerConfig.StreamBufferConfigVideo;
 	}
 	else if (InAccessUnit->ESType == EStreamType::Audio)
 	{
 		GetStreamBufferUtilization(bufferInfo, InAccessUnit->ESType);
+		BufferConfig = &PlayerConfig.StreamBufferConfigAudio;
 	}
 	else if (InAccessUnit->ESType == EStreamType::Subtitle)
 	{
 		GetStreamBufferUtilization(bufferInfo, InAccessUnit->ESType);
+		BufferConfig = &PlayerConfig.StreamBufferConfigText;
 	}
 	FAccessUnitBuffer::FExternalBufferInfo TotalUtilization;
 	TotalUtilization.DataSize = bufferInfo.CurrentMemInUse;
 	TotalUtilization.Duration = bufferInfo.PlayableDuration;
 
 	// Try to push the data into the receiving buffer
-	return Buffer->Push(InAccessUnit, &TotalUtilization);
+	return Buffer->Push(InAccessUnit, BufferConfig, &TotalUtilization);
 }
 
 //-----------------------------------------------------------------------------
@@ -1371,6 +1376,8 @@ bool FAdaptiveStreamingPlayer::InternalHandleThreadMessages()
 			{
 				FSeekParam StartAtTime = msg.Data.StartPlay.Position;
 
+				SeekVars.Clear();
+
 				if (bIsPlayStart)
 				{
 					CurrentState = EPlayerState::eState_Buffering;
@@ -1385,6 +1392,15 @@ bool FAdaptiveStreamingPlayer::InternalHandleThreadMessages()
 						Manifest->UpdateDynamicRefetchCounter();
 					}
 					CurrentState = EPlayerState::eState_Seeking;
+
+					if (StartAtTime.bOptimizeForScrubbing.IsSet())
+					{
+						SeekVars.bForScrubbing = StartAtTime.bOptimizeForScrubbing.GetValue();
+					}
+					else
+					{
+						SeekVars.bForScrubbing = PlayerOptions.GetValue(OptionKeyFrameOptimizeSeekForScrubbing).SafeGetBool(false);
+					}
 					InternalStartAt(StartAtTime);
 				}
 				break;
@@ -1412,7 +1428,7 @@ bool FAdaptiveStreamingPlayer::InternalHandleThreadMessages()
 					msg.Data.MediaEvent.Event->Signal();
 				}
 				break;
-				}
+			}
 
 			case FWorkerThreadMessages::FMessage::EType::ChangeBitrate:
 			{
@@ -1575,7 +1591,7 @@ bool FAdaptiveStreamingPlayer::InternalHandleThreadMessages()
 				break;
 			}
 		}
-			
+
 		InternalHandleOnce();
 	}
 	return bGotMessage;
@@ -1692,7 +1708,7 @@ void FAdaptiveStreamingPlayer::HandleMetadataChanges()
 				break;
 			}
 		}
-		
+
 		// Remove the periods that we have passed for which the loop counter is not in the future.
 		// With looping playback the periods to loop back to will be added/updated and must not be removed
 		// if they are used in a future loop cycle.
@@ -1871,7 +1887,7 @@ void FAdaptiveStreamingPlayer::InternalHandleLiveSync()
 			LiveSyncVars.Clear();
 			return;
 		}
-		
+
 		TSharedPtrTS<FMultiTrackAccessUnitBuffer> VidBuffer = GetCurrentReceiveStreamBuffer(EStreamType::Video);
 		TSharedPtrTS<FMultiTrackAccessUnitBuffer> AudBuffer = GetCurrentReceiveStreamBuffer(EStreamType::Audio);
 		TSharedPtrTS<FMultiTrackAccessUnitBuffer> TxtBuffer = GetCurrentReceiveStreamBuffer(EStreamType::Subtitle);
@@ -1938,112 +1954,117 @@ void FAdaptiveStreamingPlayer::InternalHandleLiveSync()
 }
 
 
-//-----------------------------------------------------------------------------
-/**
- * Checks if buffers have enough data to advance the play state.
- */
-void FAdaptiveStreamingPlayer::HandleNewBufferedData()
+bool FAdaptiveStreamingPlayer::HaveEnoughBufferedDataToStartPlayback()
 {
-	// Are we in buffering state? (this includes rebuffering)
-	if (CurrentState == EPlayerState::eState_Buffering)
-	{
-		if (DecoderState == EDecoderState::eDecoder_Paused)
-		{
-			// Synchronize buffers to wallclock time if necessary.
-			InternalHandleLiveSync();
+	double kMinBufferBeforePlayback = GetMinBufferTimeBeforePlayback();
 
-			check(LastBufferingState == EPlayerState::eState_Buffering || LastBufferingState == EPlayerState::eState_Rebuffering || LastBufferingState == EPlayerState::eState_Seeking);
-			// Yes. Check if we have enough data buffered up to begin handing off data to the decoders.
-			double kMinBufferBeforePlayback = GetMinBufferTimeBeforePlayback();
-
-			DiagnosticsCriticalSection.Lock();
+	DiagnosticsCriticalSection.Lock();
 
-			bool bHaveEnoughVideo = false;
-			if (bHaveVideoReader.IsSet())
+	bool bHaveEnoughVideo = false;
+	if (bHaveVideoReader.IsSet())
+	{
+		if (StreamReaderHandler && bHaveVideoReader.Value())
+		{
+			// If the video stream has reached the end then there won't be any new data and whatever we have is all there is.
+			if (VideoBufferStats.StreamBuffer.bEndOfData)
+			{
+				bHaveEnoughVideo = true;
+			}
+			else
 			{
-				if (StreamReaderHandler && bHaveVideoReader.Value())
+				// Does the buffer have the requested amount of data?
+				if (VideoBufferStats.StreamBuffer.PlayableDuration.GetAsSeconds() >= kMinBufferBeforePlayback)
 				{
-					// If the video stream has reached the end then there won't be any new data and whatever we have is all there is.
-					if (VideoBufferStats.StreamBuffer.bEndOfData)
-					{
-						bHaveEnoughVideo = true;
-					}
-					else
-					{
-						// Does the buffer have the requested amount of data?
-						if (VideoBufferStats.StreamBuffer.PlayableDuration.GetAsSeconds() >= kMinBufferBeforePlayback)
-						{
-							bHaveEnoughVideo = true;
-						}
-						else
-						{
-							// There is a chance that the input stream buffer is too small to hold the requested amount of material.
-							if (VideoBufferStats.StreamBuffer.bLastPushWasBlocked)
-							{
-								// We won't be able to receive any more, so we have enough as it is.
-								bHaveEnoughVideo = true;
-							}
-						}
-					}
+					bHaveEnoughVideo = true;
 				}
 				else
 				{
-					bHaveEnoughVideo = true;
+					// There is a chance that the input stream buffer is too small to hold the requested amount of material.
+					if (VideoBufferStats.StreamBuffer.bLastPushWasBlocked)
+					{
+						// We won't be able to receive any more, so we have enough as it is.
+						bHaveEnoughVideo = true;
+					}
 				}
 			}
+		}
+		else
+		{
+			bHaveEnoughVideo = true;
+		}
+	}
 
-
-			bool bHaveEnoughAudio = false;
-			if (bHaveAudioReader.IsSet())
+	bool bHaveEnoughAudio = false;
+	if (bHaveAudioReader.IsSet())
+	{
+		if (StreamReaderHandler && bHaveAudioReader.Value())
+		{
+			// If the audio stream has reached the end then there won't be any new data and whatever we have is all there is.
+			if (AudioBufferStats.StreamBuffer.bEndOfData)
+			{
+				bHaveEnoughAudio = true;
+			}
+			else
 			{
-				if (StreamReaderHandler && bHaveAudioReader.Value())
+				// Does the buffer have the requested amount of data?
+				if (AudioBufferStats.StreamBuffer.PlayableDuration.GetAsSeconds() >= kMinBufferBeforePlayback)
 				{
-					// If the audio stream has reached the end then there won't be any new data and whatever we have is all there is.
-					if (AudioBufferStats.StreamBuffer.bEndOfData)
-					{
-						bHaveEnoughAudio = true;
-					}
-					else
-					{
-						// Does the buffer have the requested amount of data?
-						if (AudioBufferStats.StreamBuffer.PlayableDuration.GetAsSeconds() >= kMinBufferBeforePlayback)
-						{
-							bHaveEnoughAudio = true;
-						}
-						else
-						{
-							// There is a chance that the input stream buffer is too small to hold the requested amount of material.
-							if (AudioBufferStats.StreamBuffer.bLastPushWasBlocked)
-							{
-								// We won't be able to receive any more, so we have enough as it is.
-								bHaveEnoughAudio = true;
-							}
-						}
-					}
+					bHaveEnoughAudio = true;
 				}
 				else
 				{
-					bHaveEnoughAudio = true;
+					// There is a chance that the input stream buffer is too small to hold the requested amount of material.
+					if (AudioBufferStats.StreamBuffer.bLastPushWasBlocked)
+					{
+						// We won't be able to receive any more, so we have enough as it is.
+						bHaveEnoughAudio = true;
+					}
 				}
 			}
+		}
+		else
+		{
+			bHaveEnoughAudio = true;
+		}
+	}
 
-			// When we are dealing with a single multiplexed stream and one buffer is blocked then essentially all buffers
-			// must be considered blocked since demuxing cannot continue.
-			// FIXME: do this more elegant somehow
-			if (ManifestType == EMediaFormatType::ISOBMFF)
-			{
-				if (VideoBufferStats.StreamBuffer.bLastPushWasBlocked ||
-					AudioBufferStats.StreamBuffer.bLastPushWasBlocked ||
-					TextBufferStats.StreamBuffer.bLastPushWasBlocked)
-				{
-					bHaveEnoughVideo = bHaveEnoughAudio = true;
-				}
-			}
+	// When we are dealing with a single multiplexed stream and one buffer is blocked then essentially all buffers
+	// must be considered blocked since demuxing cannot continue.
+	// FIXME: do this more elegant somehow
+	if (ManifestType == EMediaFormatType::ISOBMFF)
+	{
+		if (VideoBufferStats.StreamBuffer.bLastPushWasBlocked ||
+			AudioBufferStats.StreamBuffer.bLastPushWasBlocked ||
+			TextBufferStats.StreamBuffer.bLastPushWasBlocked)
+		{
+			bHaveEnoughVideo = bHaveEnoughAudio = true;
+		}
+	}
 
-			DiagnosticsCriticalSection.Unlock();
+	DiagnosticsCriticalSection.Unlock();
+
+	return bHaveEnoughVideo && bHaveEnoughAudio;
+}
 
-			int64 tNow = MEDIAutcTime::CurrentMSec();
-			if (bHaveEnoughVideo && bHaveEnoughAudio)
+
+//-----------------------------------------------------------------------------
+/**
+ * Checks if buffers have enough data to advance the play state.
+ */
+void FAdaptiveStreamingPlayer::HandleNewBufferedData()
+{
+	// Are we in buffering state? (this includes rebuffering)
+	if (CurrentState == EPlayerState::eState_Buffering && DecoderState == EDecoderState::eDecoder_Paused)
+	{
+		if (!SeekVars.bForScrubbing)
+		{
+			// Synchronize buffers to wallclock time if necessary.
+			InternalHandleLiveSync();
+
+			check(LastBufferingState == EPlayerState::eState_Buffering || LastBufferingState == EPlayerState::eState_Rebuffering || LastBufferingState == EPlayerState::eState_Seeking);
+			// Check if we have enough data buffered up to begin handing off data to the decoders.
+			bool bHaveEnoughData = HaveEnoughBufferedDataToStartPlayback();
+			if (bHaveEnoughData)
 			{
 				PlaybackState.SetIsBuffering(false);
 				PlaybackState.SetIsSeeking(false);
@@ -2063,20 +2084,48 @@ void FAdaptiveStreamingPlayer::HandleNewBufferedData()
 				}
 				DataBuffersCriticalSection.Unlock();
 
-				DecoderState   		= EDecoderState::eDecoder_Running;
-				PipelineState  		= EPipelineState::ePipeline_Prerolling;
-				PrerollVars.StartTime = tNow;
+				DecoderState = EDecoderState::eDecoder_Running;
+				PipelineState = EPipelineState::ePipeline_Prerolling;
+				PrerollVars.StartTime = MEDIAutcTime::CurrentMSec();
 				PostrollVars.Clear();
 				// We are synced up now and can clear the state.
 				LiveSyncVars.Clear();
 
 				// Send buffering end event
 				DispatchBufferingEvent(false, LastBufferingState);
-
 				// Send pre-roll begin event.
 				DispatchEvent(FMetricEvent::ReportPrerollStart());
 			}
 		}
+		else
+		{
+			// Optimized seek for scrubbing is not possible at playstart so we can skip a few things here.
+			if (!SeekVars.bScrubPrerollDone)
+			{
+				DataBuffersCriticalSection.Lock();
+				if (!ActiveDataOutputBuffers.IsValid())
+				{
+					TSharedPtrTS<FStreamDataBuffers> NewOutputBuffers;
+					if (NextDataBuffers.Num())
+					{
+						NewOutputBuffers = NextDataBuffers[0];
+						NextDataBuffers.RemoveAt(0);
+					}
+					ActiveDataOutputBuffers = MoveTemp(NewOutputBuffers);
+				}
+				bool bHaveActiveOutputBuffer = ActiveDataOutputBuffers.IsValid();
+				DataBuffersCriticalSection.Unlock();
+				if (bHaveActiveOutputBuffer)
+				{
+					DecoderState = EDecoderState::eDecoder_Running;
+					PipelineState = EPipelineState::ePipeline_Prerolling;
+					PrerollVars.StartTime = MEDIAutcTime::CurrentMSec();
+					PostrollVars.Clear();
+					LiveSyncVars.Clear();
+					DispatchEvent(FMetricEvent::ReportPrerollStart());
+				}
+			}
+		}
 	}
 }
 
@@ -2268,19 +2317,54 @@ void FAdaptiveStreamingPlayer::HandleNewOutputData()
 		// Ready to go?
 		if (bHaveEnoughVideo && bHaveEnoughAudio)
 		{
-			if (!PrerollVars.bIsMidSequencePreroll)
+			if (!SeekVars.bForScrubbing)
 			{
-				PipelineState = EPipelineState::ePipeline_Stopped;
-				CurrentState = EPlayerState::eState_Paused;
-				PlaybackRate = 0.0;
-				PrerollVars.Clear();
-				DispatchEvent(FMetricEvent::ReportPrerollEnd());
+				if (!PrerollVars.bIsMidSequencePreroll)
+				{
+					PipelineState = EPipelineState::ePipeline_Stopped;
+					CurrentState = EPlayerState::eState_Paused;
+					PlaybackRate = 0.0;
+					PrerollVars.Clear();
+					if (LastBufferingState == EPlayerState::eState_Seeking)
+					{
+						DispatchEvent(FMetricEvent::ReportSeekCompleted());
+					}
+					DispatchEvent(FMetricEvent::ReportPrerollEnd());
+				}
+				else
+				{
+					PrerollVars.Clear();
+					StartRendering();
+					PipelineState = EPipelineState::ePipeline_Running;
+				}
 			}
 			else
 			{
-				PrerollVars.Clear();
-				StartRendering();
-				PipelineState = EPipelineState::ePipeline_Running;
+				// The first time around we send the seek complete event and stop the feeding of the decoders
+				// to prevent them consuming and sending output to the renderers while we await the buffers
+				// to fill to the amount we need.
+				if (!SeekVars.bScrubPrerollDone)
+				{
+					SeekVars.bScrubPrerollDone = true;
+					DispatchEvent(FMetricEvent::ReportSeekCompleted());
+					// Pause feeding the decoder to stop draining buffers we want to fill.
+					DecoderState = EDecoderState::eDecoder_Paused;
+				}
+
+				if (HaveEnoughBufferedDataToStartPlayback())
+				{
+					PlaybackState.SetIsBuffering(false);
+					PlaybackState.SetIsSeeking(false);
+
+					SeekVars.Clear();
+					PipelineState = EPipelineState::ePipeline_Stopped;
+					CurrentState = EPlayerState::eState_Paused;
+					DecoderState = EDecoderState::eDecoder_Running;
+					PlaybackRate = 0.0;
+					PrerollVars.Clear();
+					DispatchBufferingEvent(false, LastBufferingState);
+					DispatchEvent(FMetricEvent::ReportPrerollEnd());
+				}
 			}
 		}
 	}
@@ -2670,10 +2754,6 @@ void FAdaptiveStreamingPlayer::InternalHandlePendingFirstSegmentRequest(const FT
 			NewReceiveBuffers->AudBuffer = MakeSharedTS<FMultiTrackAccessUnitBuffer>();
 			NewReceiveBuffers->TxtBuffer = MakeSharedTS<FMultiTrackAccessUnitBuffer>();
 
-			NewReceiveBuffers->VidBuffer->CapacitySet(PlayerConfig.StreamBufferConfigVideo);
-			NewReceiveBuffers->AudBuffer->CapacitySet(PlayerConfig.StreamBufferConfigAudio);
-			NewReceiveBuffers->TxtBuffer->CapacitySet(PlayerConfig.StreamBufferConfigText);
-	
 			if (ManifestType == EMediaFormatType::ISOBMFF)
 			{
 				// mp4 has multiple tracks in a single stream that provide data all at the same time.
@@ -3448,7 +3528,7 @@ void FAdaptiveStreamingPlayer::InternalHandleSegmentTrackChanges(const FTimeValu
 					NextReq.StartoverPosition.Time = PlaybackState.GetPlayPosition();
 					NextReq.StreamType = EStreamType::Audio;
 					NextPendingSegmentRequests.Enqueue(MoveTemp(NextReq));
-	
+
 					bChangeMade = true;
 				}
 			}
@@ -3862,7 +3942,7 @@ void FAdaptiveStreamingPlayer::CheckForStreamEnd()
 						DataBuffersCriticalSection.Lock();
 						ActiveDataOutputBuffers = MoveTemp(NewOutputBuffers);
 						DataBuffersCriticalSection.Unlock();
-						
+
 						// Update the current loop state.
 						FInternalLoopState newLoopState;
 						if (NextLoopStates.Peek(newLoopState))
@@ -4118,6 +4198,7 @@ void FAdaptiveStreamingPlayer::InternalStop(bool bHoldCurrentFrame)
 	UpcomingPeriods.Empty();
 	MetadataHandlingState.Reset();
 	NextLoopStates.Empty();
+	SeekVars.Clear();
 
 	PlaybackRate = 0.0;
 	CurrentState = EPlayerState::eState_Paused;
@@ -4193,7 +4274,6 @@ void FAdaptiveStreamingPlayer::InternalClose()
 
 	bShouldBePaused = false;
 	bShouldBePlaying = false;
-	bSeekPending = false;
 	CurrentLoopState = {};
 	NextLoopStates.Empty();
 
@@ -4399,19 +4479,19 @@ void FAdaptiveStreamingPlayer::DebugPrint(void* pPlayer, void (*pPrintFN)(void*
 	if (StreamReaderHandler && bHaveVideoReader.GetWithDefault(false))
 	{
 		pD = &VideoBufferStats.StreamBuffer;
-		pPrintFN(pPlayer, "Video buffer  : %3u AUs; %8u/%8u bytes in; eod %d; %#7.4fs", (uint32)pD->NumCurrentAccessUnits, (uint32)pD->CurrentMemInUse, (uint32)pD->MaxDataSize, pD->bEndOfData, pD->PlayableDuration.GetAsSeconds());
+		pPrintFN(pPlayer, "Video buffer  : %3u AUs; %8u/%8u bytes in; eod %d; %#7.4fs", (uint32)pD->NumCurrentAccessUnits, (uint32)pD->CurrentMemInUse, (uint32)PlayerConfig.StreamBufferConfigVideo.MaxDataSize, pD->bEndOfData, pD->PlayableDuration.GetAsSeconds());
 		pPrintFN(pPlayer, "Video decoder : %2u in decoder, %zu total, %s; EOD in %d; EOD out %d", (uint32)VideoBufferStats.DecoderOutputBuffer.NumElementsInDecoder, (uint32)VideoBufferStats.DecoderOutputBuffer.MaxDecodedElementsReady, VideoBufferStats.DecoderOutputBuffer.bOutputStalled?"    stalled":"not stalled", VideoBufferStats.DecoderInputBuffer.bEODReached, VideoBufferStats.DecoderOutputBuffer.bEODreached);
 	}
 	if (StreamReaderHandler && bHaveAudioReader.GetWithDefault(false))
 	{
 		pD = &AudioBufferStats.StreamBuffer;
-		pPrintFN(pPlayer, "Audio buffer  : %3u AUs; %8u/%8u bytes in; eod %d; %#7.4fs", (uint32)pD->NumCurrentAccessUnits, (uint32)pD->CurrentMemInUse, (uint32)pD->MaxDataSize, pD->bEndOfData, pD->PlayableDuration.GetAsSeconds());
+		pPrintFN(pPlayer, "Audio buffer  : %3u AUs; %8u/%8u bytes in; eod %d; %#7.4fs", (uint32)pD->NumCurrentAccessUnits, (uint32)pD->CurrentMemInUse, (uint32)PlayerConfig.StreamBufferConfigAudio.MaxDataSize, pD->bEndOfData, pD->PlayableDuration.GetAsSeconds());
 		pPrintFN(pPlayer, "Audio decoder : %2u in decoder, %2u total, %s; EOD in %d; EOD out %d", (uint32)AudioBufferStats.DecoderOutputBuffer.NumElementsInDecoder, (uint32)AudioBufferStats.DecoderOutputBuffer.MaxDecodedElementsReady, AudioBufferStats.DecoderOutputBuffer.bOutputStalled ? "    stalled" : "not stalled", AudioBufferStats.DecoderInputBuffer.bEODReached, AudioBufferStats.DecoderOutputBuffer.bEODreached);
 	}
 	if (StreamReaderHandler && bHaveTextReader.GetWithDefault(false))
 	{
 		pD = &TextBufferStats.StreamBuffer;
-		pPrintFN(pPlayer, "Text buffer   : %3u AUs; %8u/%8u bytes in; eod %d; %#7.4fs", (uint32)pD->NumCurrentAccessUnits, (uint32)pD->CurrentMemInUse, (uint32)pD->MaxDataSize, pD->bEndOfData, pD->PlayableDuration.GetAsSeconds());
+		pPrintFN(pPlayer, "Text buffer   : %3u AUs; %8u/%8u bytes in; eod %d; %#7.4fs", (uint32)pD->NumCurrentAccessUnits, (uint32)pD->CurrentMemInUse, (uint32)PlayerConfig.StreamBufferConfigText.MaxDataSize, pD->bEndOfData, pD->PlayableDuration.GetAsSeconds());
 	}
 	DiagnosticsCriticalSection.Unlock();
 
@@ -4643,6 +4723,9 @@ void FAdaptiveStreamingPlayerEventHandler::WorkerThreadFN()
 						case FMetricEvent::EType::PlaybackStopped:
 							MetricListeners[i]->ReportPlaybackStopped();
 							break;
+						case FMetricEvent::EType::SeekCompleted:
+							MetricListeners[i]->ReportSeekCompleted();
+							break;
 						case FMetricEvent::EType::Errored:
 							MetricListeners[i]->ReportError(pEvt->Param.ErrorDetail.GetPrintable());
 							break;
diff --git a/Engine/Plugins/Media/ElectraPlayer/Source/ElectraPlayerRuntime/Private/Runtime/Player/AdaptiveStreamingPlayer.h b/Engine/Plugins/Media/ElectraPlayer/Source/ElectraPlayerRuntime/Private/Runtime/Player/AdaptiveStreamingPlayer.h
--- a/Engine/Plugins/Media/ElectraPlayer/Source/ElectraPlayerRuntime/Private/Runtime/Player/AdaptiveStreamingPlayer.h
+++ b/Engine/Plugins/Media/ElectraPlayer/Source/ElectraPlayerRuntime/Private/Runtime/Player/AdaptiveStreamingPlayer.h
@@ -112,9 +112,11 @@ public:
 		{
 			Time.SetToInvalid();
 			StartingBitrate.Reset();
+			bOptimizeForScrubbing.Reset();
 		}
 		FTimeValue	Time;
 		TOptional<int32> StartingBitrate;
+		TOptional<bool> bOptimizeForScrubbing;
 	};
 	//! Seek to a new position and play from there. This includes first playstart.
 	//! Playback is initially paused on first player use and must be resumed to begin.
diff --git a/Engine/Plugins/Media/ElectraPlayer/Source/ElectraPlayerRuntime/Private/Runtime/Player/AdaptiveStreamingPlayerABR.cpp b/Engine/Plugins/Media/ElectraPlayer/Source/ElectraPlayerRuntime/Private/Runtime/Player/AdaptiveStreamingPlayerABR.cpp
--- a/Engine/Plugins/Media/ElectraPlayer/Source/ElectraPlayerRuntime/Private/Runtime/Player/AdaptiveStreamingPlayerABR.cpp
+++ b/Engine/Plugins/Media/ElectraPlayer/Source/ElectraPlayerRuntime/Private/Runtime/Player/AdaptiveStreamingPlayerABR.cpp
@@ -166,6 +166,7 @@ namespace Electra
 		virtual void ReportPlaybackEnded() override;
 		virtual void ReportJumpInPlayPosition(const FTimeValue& ToNewTime, const FTimeValue& FromTime, Metrics::ETimeJumpReason TimejumpReason) override {}
 		virtual void ReportPlaybackStopped() override {}
+		virtual void ReportSeekCompleted() override {}
 		virtual void ReportError(const FString& ErrorReason) override {}
 		virtual void ReportLogMessage(IInfoLog::ELevel InLogLevel, const FString& LogMessage, int64 PlayerWallclockMilliseconds) override {}
 		virtual void ReportDroppedVideoFrame() override {}
diff --git a/Engine/Plugins/Media/ElectraPlayer/Source/ElectraPlayerRuntime/Private/Runtime/Player/AdaptiveStreamingPlayerABR.h b/Engine/Plugins/Media/ElectraPlayer/Source/ElectraPlayerRuntime/Private/Runtime/Player/AdaptiveStreamingPlayerABR.h
--- a/Engine/Plugins/Media/ElectraPlayer/Source/ElectraPlayerRuntime/Private/Runtime/Player/AdaptiveStreamingPlayerABR.h
+++ b/Engine/Plugins/Media/ElectraPlayer/Source/ElectraPlayerRuntime/Private/Runtime/Player/AdaptiveStreamingPlayerABR.h
@@ -146,6 +146,7 @@ public:
 	virtual void ReportPlaybackEnded() = 0;
 	virtual void ReportJumpInPlayPosition(const FTimeValue& ToNewTime, const FTimeValue& FromTime, Metrics::ETimeJumpReason TimejumpReason) = 0;
 	virtual void ReportPlaybackStopped() = 0;
+	virtual void ReportSeekCompleted() = 0;
 	virtual void ReportError(const FString& ErrorReason) = 0;
 	virtual void ReportLogMessage(IInfoLog::ELevel LogLevel, const FString& LogMessage, int64 PlayerWallclockMilliseconds) = 0;
 	virtual void ReportDroppedVideoFrame() = 0;
diff --git a/Engine/Plugins/Media/ElectraPlayer/Source/ElectraPlayerRuntime/Private/Runtime/Player/AdaptiveStreamingPlayerEvents.cpp b/Engine/Plugins/Media/ElectraPlayer/Source/ElectraPlayerRuntime/Private/Runtime/Player/AdaptiveStreamingPlayerEvents.cpp
--- a/Engine/Plugins/Media/ElectraPlayer/Source/ElectraPlayerRuntime/Private/Runtime/Player/AdaptiveStreamingPlayerEvents.cpp
+++ b/Engine/Plugins/Media/ElectraPlayer/Source/ElectraPlayerRuntime/Private/Runtime/Player/AdaptiveStreamingPlayerEvents.cpp
@@ -50,25 +50,29 @@ void FAdaptiveStreamingPlayer::DispatchSegmentDownloadedEvent(TSharedPtrTS<IStre
 void FAdaptiveStreamingPlayer::DispatchBufferUtilizationEvent(EStreamType BufferType)
 {
 	const FAccessUnitBufferInfo* BufferStats = nullptr;
+	const FAccessUnitBuffer::FConfiguration* BufferConfig = nullptr;
 	if (BufferType == EStreamType::Video)
 	{
 		BufferStats = &VideoBufferStats.StreamBuffer;
+		BufferConfig = &PlayerConfig.StreamBufferConfigVideo;
 	}
 	else if (BufferType == EStreamType::Audio)
 	{
 		BufferStats = &AudioBufferStats.StreamBuffer;
+		BufferConfig = &PlayerConfig.StreamBufferConfigAudio;
 	}
 	else if (BufferType == EStreamType::Subtitle)
 	{
 		BufferStats = &TextBufferStats.StreamBuffer;
+		BufferConfig = &PlayerConfig.StreamBufferConfigText;
 	}
 	if (BufferStats)
 	{
 		Metrics::FBufferStats stats;
 		stats.BufferType = BufferType;
-		stats.MaxDurationInSeconds = BufferStats->MaxDuration.GetAsSeconds();
+		stats.MaxDurationInSeconds = BufferConfig ? BufferConfig->MaxDuration.GetAsSeconds() : 0.0;
 		stats.DurationInUse 	   = BufferStats->PushedDuration.GetAsSeconds();
-		stats.MaxByteCapacity      = BufferStats->MaxDataSize;
+		stats.MaxByteCapacity      = BufferConfig ? BufferConfig->MaxDataSize : 0;
 		stats.BytesInUse		   = BufferStats->CurrentMemInUse;
 		DispatchEvent(FMetricEvent::ReportBufferUtilization(stats));
 	}
diff --git a/Engine/Plugins/Media/ElectraPlayer/Source/ElectraPlayerRuntime/Private/Runtime/Player/AdaptiveStreamingPlayerInternal.h b/Engine/Plugins/Media/ElectraPlayer/Source/ElectraPlayerRuntime/Private/Runtime/Player/AdaptiveStreamingPlayerInternal.h
--- a/Engine/Plugins/Media/ElectraPlayer/Source/ElectraPlayerRuntime/Private/Runtime/Player/AdaptiveStreamingPlayerInternal.h
+++ b/Engine/Plugins/Media/ElectraPlayer/Source/ElectraPlayerRuntime/Private/Runtime/Player/AdaptiveStreamingPlayerInternal.h
@@ -422,6 +422,7 @@ struct FMetricEvent
 		PlaybackEnded,
 		PlaybackJumped,
 		PlaybackStopped,
+		SeekCompleted,
 		LicenseKey,
 		Errored,
 		LogMessage,
@@ -625,6 +626,12 @@ struct FMetricEvent
 		Evt->Type = EType::PlaybackStopped;
 		return Evt;
 	}
+	static TSharedPtrTS<FMetricEvent> ReportSeekCompleted()
+	{
+		TSharedPtrTS<FMetricEvent> Evt = MakeSharedTS<FMetricEvent>();
+		Evt->Type = EType::SeekCompleted;
+		return Evt;
+	}
 	static TSharedPtrTS<FMetricEvent> ReportError(const FErrorDetail& ErrorDetail)
 	{
 		TSharedPtrTS<FMetricEvent> Evt = MakeSharedTS<FMetricEvent>();
@@ -1536,6 +1543,22 @@ private:
 		FRenderState		Audio;
 	};
 
+	struct FSeekVars
+	{
+		FSeekVars()
+		{
+			Clear();
+		}
+		void Clear()
+		{
+			bForScrubbing = false;
+			bScrubPrerollDone = false;
+		}
+
+		bool	bForScrubbing;
+		bool	bScrubPrerollDone;
+	};
+
 	struct FStreamBitrateInfo
 	{
 		FStreamBitrateInfo()
@@ -1699,6 +1722,7 @@ private:
 	void CheckForErrors();
 
 	double GetMinBufferTimeBeforePlayback();
+	bool HaveEnoughBufferedDataToStartPlayback();
 
 	FTimeValue ClampTimeToCurrentRange(const FTimeValue& InTime, bool bClampToStart, bool bClampToEnd);
 
@@ -1839,6 +1863,7 @@ private:
 	FLiveSyncVars														LiveSyncVars;
 	FPrerollVars														PrerollVars;
 	FPostrollVars														PostrollVars;
+	FSeekVars															SeekVars;
 	EPlayerState														LastBufferingState;
 	double																PlaybackRate;
 	FTimeValue															RebufferDetectedAtPlayPos;
@@ -1855,7 +1880,6 @@ private:
 
 	bool																bShouldBePaused;
 	bool																bShouldBePlaying;
-	bool																bSeekPending;
 
 	FInternalLoopState													CurrentLoopState;
 	FLoopParam															CurrentLoopParam;
diff --git a/Engine/Plugins/Media/ElectraPlayer/Source/ElectraPlayerRuntime/Private/Runtime/Player/AdaptiveStreamingPlayerMetrics.h b/Engine/Plugins/Media/ElectraPlayer/Source/ElectraPlayerRuntime/Private/Runtime/Player/AdaptiveStreamingPlayerMetrics.h
--- a/Engine/Plugins/Media/ElectraPlayer/Source/ElectraPlayerRuntime/Private/Runtime/Player/AdaptiveStreamingPlayerMetrics.h
+++ b/Engine/Plugins/Media/ElectraPlayer/Source/ElectraPlayerRuntime/Private/Runtime/Player/AdaptiveStreamingPlayerMetrics.h
@@ -324,6 +324,12 @@ public:
 	 */
 	virtual void ReportPlaybackStopped() = 0;
 
+	/**
+	 * Called when a seek has completed such that the first new data is ready.
+	 * Future data may still be buffering.
+	 */
+	virtual void ReportSeekCompleted() = 0;
+
 	/**
 	 * Called when an error occurs. Errors always result in termination of playback.
 	 */
diff --git a/Engine/Plugins/Media/ElectraPlayer/Source/ElectraPlayerRuntime/Private/Runtime/StreamAccessUnitBuffer.cpp b/Engine/Plugins/Media/ElectraPlayer/Source/ElectraPlayerRuntime/Private/Runtime/StreamAccessUnitBuffer.cpp
--- a/Engine/Plugins/Media/ElectraPlayer/Source/ElectraPlayerRuntime/Private/Runtime/StreamAccessUnitBuffer.cpp
+++ b/Engine/Plugins/Media/ElectraPlayer/Source/ElectraPlayerRuntime/Private/Runtime/StreamAccessUnitBuffer.cpp
@@ -16,10 +16,7 @@ namespace Electra
 
 	TSharedPtrTS<FAccessUnitBuffer> FMultiTrackAccessUnitBuffer::CreateNewBuffer()
 	{
-		// All buffers here are internally unbounded as we need them to accept all incoming data.
-		// Limits are used on the active buffer only.
 		TSharedPtrTS<FAccessUnitBuffer> NewBuffer = MakeSharedTS<FAccessUnitBuffer>();
-		NewBuffer->CapacitySet(FAccessUnitBuffer::FConfiguration(1024 << 20, 3600.0));
 		return NewBuffer;
 	}
 
@@ -33,13 +30,7 @@ namespace Electra
 		bIsParallelTrackMode = true;
 	}
 
-
-	void FMultiTrackAccessUnitBuffer::CapacitySet(const FAccessUnitBuffer::FConfiguration& Config)
-	{
-		BufferConfiguration = Config;
-	}
-
-	bool FMultiTrackAccessUnitBuffer::Push(FAccessUnit*& AU, const FAccessUnitBuffer::FExternalBufferInfo* InCurrentTotalBufferUtilization)
+	bool FMultiTrackAccessUnitBuffer::Push(FAccessUnit*& AU, const FAccessUnitBuffer::FConfiguration* BufferConfiguration, const FAccessUnitBuffer::FExternalBufferInfo* InCurrentTotalBufferUtilization)
 	{
 		check(AU->BufferSourceInfo.IsValid());
 
@@ -58,7 +49,6 @@ namespace Electra
 		// This does not necessarily mean that the selected track is not at EOD. Just that some track is not.
 		bEndOfData = false;
 
-
 		bool bIsSwitchOverAU = false;
 		for(auto &SwitchBuf : SwitchOverBufferChain)
 		{
@@ -82,7 +72,7 @@ namespace Electra
 
 		AccessLock.Unlock();
 
-		bool bWasPushed = TrackBuffer->Push(AU, &BufferConfiguration, &EnqueuedBufferInfo);
+		bool bWasPushed = TrackBuffer->Push(AU, BufferConfiguration, &EnqueuedBufferInfo);
 		bLastPushWasBlocked = TrackBuffer->WasLastPushBlocked();
 
 		return bWasPushed;
@@ -109,7 +99,7 @@ namespace Electra
 
 	void FMultiTrackAccessUnitBuffer::PushEndOfDataAll()
 	{
-		FMediaCriticalSection::ScopedLock lock(AccessLock);
+		FScopeLock lock(&AccessLock);
 		bEndOfData = true;
 		// Push an end-of-data into all tracks.
 		for(auto& It : TrackBuffers)
@@ -132,7 +122,7 @@ namespace Electra
 
 	void FMultiTrackAccessUnitBuffer::Flush()
 	{
-		FMediaCriticalSection::ScopedLock lock(AccessLock);
+		FScopeLock lock(&AccessLock);
 		if (bIsParallelTrackMode)
 		{
 			// Flush all existing buffers but keep them in the track map.
@@ -150,7 +140,7 @@ namespace Electra
 
 	void FMultiTrackAccessUnitBuffer::PurgeAll()
 	{
-		FMediaCriticalSection::ScopedLock lock(AccessLock);
+		FScopeLock lock(&AccessLock);
 		// Get rid of all buffers
 		TrackBuffers.Empty();
 		Clear();
@@ -161,7 +151,7 @@ namespace Electra
 	{
 		if (InBufferSourceInfo.IsValid())
 		{
-			FMediaCriticalSection::ScopedLock lock(AccessLock);
+			FScopeLock lock(&AccessLock);
 			if (SwitchOverBufferChain.Num() == 0)
 			{
 				FQueuedBuffer Next;
@@ -224,7 +214,7 @@ namespace Electra
 	{
 		if (InBufferSourceInfo.IsValid())
 		{
-			FMediaCriticalSection::ScopedLock lock(AccessLock);
+			FScopeLock lock(&AccessLock);
 			// Add to the current or the switch-over chain?
 			if (SwitchOverBufferChain.Num())
 			{
@@ -249,20 +239,20 @@ namespace Electra
 
 	FTimeValue FMultiTrackAccessUnitBuffer::GetLastPoppedPTS()
 	{
-		FMediaCriticalSection::ScopedLock lock(AccessLock);
+		FScopeLock lock(&AccessLock);
 		return LastPoppedPTS;
 	}
 
 	FTimeValue FMultiTrackAccessUnitBuffer::GetLastPoppedDTS()
 	{
-		FMediaCriticalSection::ScopedLock lock(AccessLock);
+		FScopeLock lock(&AccessLock);
 		return LastPoppedDTS;
 	}
 
 	TSharedPtrTS<FAccessUnitBuffer> FMultiTrackAccessUnitBuffer::GetSelectedTrackBuffer()
 	{
 		TSharedPtrTS<FAccessUnitBuffer> TrackBuffer;
-		FMediaCriticalSection::ScopedLock lock(AccessLock);
+		FScopeLock lock(&AccessLock);
 		if (ActiveOutputBufferInfo.IsValid() && TrackBuffers.Contains(ActiveOutputBufferInfo->PeriodAdaptationSetID))
 		{
 			TrackBuffer = TrackBuffers[ActiveOutputBufferInfo->PeriodAdaptationSetID];
@@ -276,11 +266,9 @@ namespace Electra
 
 	void FMultiTrackAccessUnitBuffer::GetStats(FAccessUnitBufferInfo& OutStats)
 	{
-		FMediaCriticalSection::ScopedLock lock(AccessLock);
-		// Stats are always returned with the configuration for one buffer, not
-		// any summation of enqueued buffers.
+		FScopeLock lock(&AccessLock);
 		TSharedPtrTS<const FAccessUnitBuffer> Buf = GetSelectedTrackBuffer();
-		Buf->GetStats(OutStats, &BufferConfiguration);
+		Buf->GetStats(OutStats);
 		if (Buf == EmptyBuffer)
 		{
 			OutStats.bEndOfData = bEndOfData;
diff --git a/Engine/Plugins/Media/ElectraPlayer/Source/ElectraPlayerRuntime/Private/Runtime/StreamAccessUnitBuffer.h b/Engine/Plugins/Media/ElectraPlayer/Source/ElectraPlayerRuntime/Private/Runtime/StreamAccessUnitBuffer.h
--- a/Engine/Plugins/Media/ElectraPlayer/Source/ElectraPlayerRuntime/Private/Runtime/StreamAccessUnitBuffer.h
+++ b/Engine/Plugins/Media/ElectraPlayer/Source/ElectraPlayerRuntime/Private/Runtime/StreamAccessUnitBuffer.h
@@ -199,8 +199,6 @@ namespace Electra
 			FrontDTS.SetToInvalid();
 			PushedDuration.SetToZero();
 			PlayableDuration.SetToZero();
-			MaxDuration.SetToZero();
-			MaxDataSize = 0;
 			CurrentMemInUse = 0;
 			NumCurrentAccessUnits = 0;
 			bEndOfData = false;
@@ -210,8 +208,6 @@ namespace Electra
 		FTimeValue			FrontDTS;
 		FTimeValue			PushedDuration;
 		FTimeValue			PlayableDuration;
-		FTimeValue			MaxDuration;
-		int64				MaxDataSize;
 		int64				CurrentMemInUse;
 		int64				NumCurrentAccessUnits;
 		bool				bEndOfData;
@@ -263,14 +259,14 @@ namespace Electra
 		//! Returns the number of access units currently in the FIFO.
 		int64 Num() const
 		{
-			FMediaCriticalSection::ScopedLock Lock(AccessLock);
+			FScopeLock Lock(&AccessLock);
 			return AccessUnits.Num();
 		}
 
 		//! Returns the amount of memory currently allocated.
 		int64 AllocatedSize() const
 		{
-			FMediaCriticalSection::ScopedLock Lock(AccessLock);
+			FScopeLock Lock(&AccessLock);
 			return CurrentMemInUse;
 		}
 
@@ -281,37 +277,25 @@ namespace Electra
 		}
 
 		//! Returns all vital statistics.
-		void GetStats(FAccessUnitBufferInfo& OutStats, const FConfiguration* Limit = nullptr) const
+		void GetStats(FAccessUnitBufferInfo& OutStats) const
 		{
-			FMediaCriticalSection::ScopedLock Lock(AccessLock);
-			Limit = Limit ? Limit : &Config;
+			FScopeLock Lock(&AccessLock);
 			OutStats.FrontDTS = FrontDTS;
 			OutStats.PushedDuration = PushedDuration;
 			OutStats.PlayableDuration = PlayableDuration;
-			OutStats.MaxDuration = Limit->MaxDuration;
-			OutStats.MaxDataSize = Limit->MaxDataSize;
 			OutStats.CurrentMemInUse = CurrentMemInUse;
 			OutStats.NumCurrentAccessUnits = AccessUnits.Num();
 			OutStats.bEndOfData = bEndOfData;
 			OutStats.bLastPushWasBlocked = bLastPushWasBlocked;
 		}
 
-		//! Sets the maximum amount of memory/number of AUs this FIFO is allowed to allocate.
-		void CapacitySet(const FConfiguration& InConfig)
-		{
-			FMediaCriticalSection::ScopedLock Lock(AccessLock);
-			Flush();
-			Config = InConfig;
-		}
-
 		//! Adds an access unit to the FIFO. Returns true if successful, false if the FIFO has insufficient free space.
 		bool Push(FAccessUnit*& AU, const FConfiguration* Limit = nullptr, const FExternalBufferInfo* ExternalInfo = nullptr)
 		{
-			FMediaCriticalSection::ScopedLock Lock(AccessLock);
+			FScopeLock Lock(&AccessLock);
 			// Pushing new data unconditionally clears the EOD flag even if the buffer is currently full.
 			// The attempt to push implies there will be more data.
 			bEndOfData = false;
-			Limit = Limit ? Limit : &Config;
 			ExternalInfo = ExternalInfo ? ExternalInfo : &ZeroExternalInfo;
 			if (CanPush(AU, Limit, ExternalInfo))
 			{
@@ -349,7 +333,7 @@ namespace Electra
 		//! Removes and returns the oldest access unit from the FIFO. Returns false if the FIFO is empty.
 		bool Pop(FAccessUnit*& OutAU)
 		{
-			FMediaCriticalSection::ScopedLock Lock(AccessLock);
+			FScopeLock Lock(&AccessLock);
 			if (Num())
 			{
 				OutAU = AccessUnits.Pop();
@@ -399,7 +383,7 @@ namespace Electra
 		//!
 		bool PeekAndAddRef(FAccessUnit*& OutAU)
 		{
-			FMediaCriticalSection::ScopedLock Lock(AccessLock);
+			FScopeLock Lock(&AccessLock);
 			if (Num())
 			{
 				OutAU = AccessUnits.Front();
@@ -416,7 +400,7 @@ namespace Electra
 		//!
 		bool ContainsPTS(const FTimeValue& InPTS) const
 		{
-			FMediaCriticalSection::ScopedLock Lock(AccessLock);
+			FScopeLock Lock(&AccessLock);
 			if (AccessUnits.Num())
 			{
 				return AccessUnits.FrontRef()->PTS <= InPTS && InPTS < AccessUnits.BackRef()->PTS + AccessUnits.BackRef()->Duration;
@@ -426,7 +410,7 @@ namespace Electra
 
 		bool ContainsFuturePTS(const FTimeValue& InPTS) const
 		{
-			FMediaCriticalSection::ScopedLock Lock(AccessLock);
+			FScopeLock Lock(&AccessLock);
 			if (AccessUnits.Num())
 			{
 				return InPTS <= AccessUnits.BackRef()->PTS + AccessUnits.BackRef()->Duration;
@@ -481,7 +465,7 @@ namespace Electra
 		FTimeValue PrepareForDecodeStartingAt(FTimeValue DecodeStartTime)
 		{
 			FTimeValue TaggedDuration(FTimeValue::GetZero());
-			FMediaCriticalSection::ScopedLock Lock(AccessLock);
+			FScopeLock Lock(&AccessLock);
 			int32 RemoveUpTo = -1;
 			for(int32 i=0; i<AccessUnits.Num(); ++i)
 			{
@@ -541,7 +525,7 @@ namespace Electra
 		//! Removes all elements from the FIFO
 		void Flush()
 		{
-			FMediaCriticalSection::ScopedLock Lock(AccessLock);
+			FScopeLock Lock(&AccessLock);
 			while(AccessUnits.Num())
 			{
 				FAccessUnit::Release(AccessUnits.Pop());
@@ -558,14 +542,14 @@ namespace Electra
 		//! Checks if the buffer has reached the end-of-data marker (marker is set and no more data is in the buffer).
 		bool IsEndOfData() const
 		{
-			FMediaCriticalSection::ScopedLock Lock(AccessLock);
+			FScopeLock Lock(&AccessLock);
 			return bEndOfData && AccessUnits.IsEmpty();
 		}
 
 		// Checks if the end-of-data flag has been set. There may still be data in the buffer though!
 		bool IsEODFlagSet() const
 		{
-			FMediaCriticalSection::ScopedLock Lock(AccessLock);
+			FScopeLock Lock(&AccessLock);
 			return bEndOfData;
 		}
 
@@ -576,7 +560,7 @@ namespace Electra
 		}
 
 		// Helper class to lock the AU buffer
-		class FScopedLock : private TMediaNoncopyable<FScopedLock>
+		class FScopedLock
 		{
 		public:
 			explicit FScopedLock(const FAccessUnitBuffer& owner)
@@ -589,7 +573,9 @@ namespace Electra
 				mOwner.AccessLock.Unlock();
 			}
 		private:
-			FScopedLock();
+			FScopedLock() = delete;
+			FScopedLock(const FScopedLock&) = delete;
+			FScopedLock& operator= (const FScopedLock&) = delete;
 			const FAccessUnitBuffer& mOwner;
 		};
 
@@ -597,11 +583,16 @@ namespace Electra
 		//! Checks if an access unit can be pushed to the FIFO. Returns true if successful, false if the FIFO has insufficient free space.
 		bool CanPush(const FAccessUnit* AU, const FConfiguration* Limit, const FExternalBufferInfo* ExternalInfo)
 		{
+			check(Limit);
+			if (!Limit)
+			{
+				return false;
+			}
 			check(Limit->MaxDuration > FTimeValue::GetZero());
 			check(AU->Duration.IsValid() && !AU->Duration.IsInfinity());
 			// If tagged with a drop state we unconditionally accept the AU. These are early ones that won't get decoded.
-			// They should appear at the beginning and end of a stream only and we do need the data that will be decoded
-			// so the ones to be dropped should not be taken into account.
+			// They should appear at the beginning and end of a stream only and since we need the following data that
+			// will be decoded we have to accept those first.
 			if (AU->DropState != 0)
 			{
 				return true;
@@ -627,8 +618,7 @@ namespace Electra
 			return false;
 		}
 
-		mutable FMediaCriticalSection				AccessLock;
-		FConfiguration								Config;
+		mutable FCriticalSection					AccessLock;
 		FExternalBufferInfo							ZeroExternalInfo;
 		TMediaQueueDynamicNoLock<FAccessUnit*>		AccessUnits;
 		FMediaSemaphore								NumInSemaphore;
@@ -653,10 +643,9 @@ namespace Electra
 		FMultiTrackAccessUnitBuffer();
 		~FMultiTrackAccessUnitBuffer();
 		void SetParallelTrackMode();
-		void CapacitySet(const FAccessUnitBuffer::FConfiguration& Config);
 		void SelectTrackWhenAvailable(TSharedPtrTS<FBufferSourceInfo> InBufferSourceInfo);
 		void AddUpcomingBuffer(TSharedPtrTS<FBufferSourceInfo> InBufferSourceInfo);
-		bool Push(FAccessUnit*& AU, const FAccessUnitBuffer::FExternalBufferInfo* InCurrentTotalBufferUtilization);
+		bool Push(FAccessUnit*& AU, const FAccessUnitBuffer::FConfiguration* BufferConfiguration, const FAccessUnitBuffer::FExternalBufferInfo* InCurrentTotalBufferUtilization);
 		void PushEndOfDataFor(TSharedPtrTS<const FBufferSourceInfo> InStreamSourceInfo);
 		void PushEndOfDataAll();
 		void Flush();
@@ -667,7 +656,7 @@ namespace Electra
 		{ return ActiveOutputBufferInfo; }
 
 		// Helper class to lock the AU buffer
-		class FScopedLock : private TMediaNoncopyable<FScopedLock>
+		class FScopedLock
 		{
 		public:
 			explicit FScopedLock(TSharedPtrTS<FMultiTrackAccessUnitBuffer> Self)
@@ -713,8 +702,7 @@ namespace Electra
 		void RemoveUnusedBuffers();
 		void GetEnqueuedBufferInfo(FAccessUnitBuffer::FExternalBufferInfo& OutInfo, bool bForSwitchOverChain);
 
-		FMediaCriticalSection							AccessLock;
-		FAccessUnitBuffer::FConfiguration				BufferConfiguration;				//!< Buffer configuration.
+		FCriticalSection								AccessLock;
 		TMap<FString, TSharedPtrTS<FAccessUnitBuffer>>	TrackBuffers;						//!< Map of track buffers. One per track ID.
 		TArray<FQueuedBuffer>							UpcomingBufferChain;
 		TArray<FQueuedBuffer>							SwitchOverBufferChain;
