Add another supported platform to some media player factories.

#rnx

diff --git a/Engine/Plugins/Media/ElectraPlayer/Source/ElectraPlayerFactory/Private/ElectraPlayerFactoryModule.cpp b/Engine/Plugins/Media/ElectraPlayer/Source/ElectraPlayerFactory/Private/ElectraPlayerFactoryModule.cpp
--- a/Engine/Plugins/Media/ElectraPlayer/Source/ElectraPlayerFactory/Private/ElectraPlayerFactoryModule.cpp
+++ b/Engine/Plugins/Media/ElectraPlayer/Source/ElectraPlayerFactory/Private/ElectraPlayerFactoryModule.cpp
@@ -129,6 +129,7 @@ public:
 		AddSupportedPlatform(FGuid(0x941259d5, 0x0a2746aa, 0xadc0ba84, 0x4790ad8a));
 		AddSupportedPlatform(FGuid(0xccf05903, 0x822b47e1, 0xb2236a28, 0xdfd78817));
 		AddSupportedPlatform(FGuid(0x5636fbc1, 0xd2b54f62, 0xac8e7d4f, 0xb184b45a));
+		AddSupportedPlatform(FGuid(0xb596ce6f, 0xd8324a9c, 0x84e9f880, 0x21322535));
 #endif
 
 		// supported schemes
diff --git a/Engine/Plugins/Media/MfMedia/Source/MfMediaFactory/Private/MfMediaFactoryModule.cpp b/Engine/Plugins/Media/MfMedia/Source/MfMediaFactory/Private/MfMediaFactoryModule.cpp
--- a/Engine/Plugins/Media/MfMedia/Source/MfMediaFactory/Private/MfMediaFactoryModule.cpp
+++ b/Engine/Plugins/Media/MfMedia/Source/MfMediaFactory/Private/MfMediaFactoryModule.cpp
@@ -175,6 +175,7 @@ public:
 		AddSupportedPlatform(FGuid(0x5636fbc1, 0xd2b54f62, 0xac8e7d4f, 0xb184b45a));
 		AddSupportedPlatform(FGuid(0x941259d5, 0x0a2746aa, 0xadc0ba84, 0x4790ad8a));
 		AddSupportedPlatform(FGuid(0xccf05903, 0x822b47e1, 0xb2236a28, 0xdfd78817));
+		AddSupportedPlatform(FGuid(0xb596ce6f, 0xd8324a9c, 0x84e9f880, 0x21322535));
 
 		// supported schemes
 		SupportedUriSchemes.Add(TEXT("file"));
