ElectraPlayer: Aggregating consecutive seek operations made in rapid succession

#rnx

diff --git a/Engine/Plugins/Media/ElectraPlayer/Source/ElectraPlayerFactory/Private/ElectraPlayerFactoryModule.cpp b/Engine/Plugins/Media/ElectraPlayer/Source/ElectraPlayerFactory/Private/ElectraPlayerFactoryModule.cpp
--- a/Engine/Plugins/Media/ElectraPlayer/Source/ElectraPlayerFactory/Private/ElectraPlayerFactoryModule.cpp
+++ b/Engine/Plugins/Media/ElectraPlayer/Source/ElectraPlayerFactory/Private/ElectraPlayerFactoryModule.cpp
@@ -100,9 +100,9 @@ public:
 		return ((Feature == EMediaFeature::AudioSamples) ||
 				(Feature == EMediaFeature::AudioTracks) ||
 				//(Feature == EMediaFeature::CaptionTracks) ||
-				//(Feature == EMediaFeature::MetadataTracks) ||
+				(Feature == EMediaFeature::MetadataTracks) ||
 				//(Feature == EMediaFeature::OverlaySamples) ||
-				//(Feature == EMediaFeature::SubtitleTracks) ||
+				(Feature == EMediaFeature::SubtitleTracks) ||
 				(Feature == EMediaFeature::VideoSamples) ||
 				(Feature == EMediaFeature::VideoTracks));
 	}
diff --git a/Engine/Plugins/Media/ElectraPlayer/Source/ElectraPlayerPlugin/Private/ElectraPlayerPlugin.cpp b/Engine/Plugins/Media/ElectraPlayer/Source/ElectraPlayerPlugin/Private/ElectraPlayerPlugin.cpp
--- a/Engine/Plugins/Media/ElectraPlayer/Source/ElectraPlayerPlugin/Private/ElectraPlayerPlugin.cpp
+++ b/Engine/Plugins/Media/ElectraPlayer/Source/ElectraPlayerPlugin/Private/ElectraPlayerPlugin.cpp
@@ -661,6 +661,8 @@ bool FElectraPlayerPlugin::GetPlayerFeatureFlag(EFeatureFlag flag) const
 			return Player->IsKillAfterCloseAllowed();
 		case EFeatureFlag::UsePlaybackTimingV2:
 			return true;
+		case EFeatureFlag::PlayerUsesInternalFlushOnSeek:
+			return true;
 		default:
 			break;
 	}
diff --git a/Engine/Plugins/Media/ElectraPlayer/Source/ElectraPlayerPlugin/Private/ElectraPlayerPlugin.h b/Engine/Plugins/Media/ElectraPlayer/Source/ElectraPlayerPlugin/Private/ElectraPlayerPlugin.h
--- a/Engine/Plugins/Media/ElectraPlayer/Source/ElectraPlayerPlugin/Private/ElectraPlayerPlugin.h
+++ b/Engine/Plugins/Media/ElectraPlayer/Source/ElectraPlayerPlugin/Private/ElectraPlayerPlugin.h
@@ -102,11 +102,11 @@ public:
 	void SetLastAudioRenderedSampleTime(FTimespan SampleTime) override;
 	bool FlushOnSeekStarted() const override
 	{
-		return true;
+		return false;
 	}
 	bool FlushOnSeekCompleted() const override
 	{
-		return true;
+		return false;
 	}
 
 	bool GetPlayerFeatureFlag(EFeatureFlag flag) const override;
diff --git a/Engine/Plugins/Media/ElectraPlayer/Source/ElectraPlayerRuntime/Private/Runtime/Player/AdaptiveStreamingPlayer.cpp b/Engine/Plugins/Media/ElectraPlayer/Source/ElectraPlayerRuntime/Private/Runtime/Player/AdaptiveStreamingPlayer.cpp
--- a/Engine/Plugins/Media/ElectraPlayer/Source/ElectraPlayerRuntime/Private/Runtime/Player/AdaptiveStreamingPlayer.cpp
+++ b/Engine/Plugins/Media/ElectraPlayer/Source/ElectraPlayerRuntime/Private/Runtime/Player/AdaptiveStreamingPlayer.cpp
@@ -90,7 +90,6 @@ FAdaptiveStreamingPlayer::FAdaptiveStreamingPlayer(const IAdaptiveStreamingPlaye
 	bRebufferPending   	    = false;
 	Manifest  				= nullptr;
 	LastBufferingState 		= EPlayerState::eState_Buffering;
-	bIsPlayStart   			= true;
 	bIsClosing 				= false;
 
 	bHaveVideoReader.Reset();
@@ -132,8 +131,7 @@ FAdaptiveStreamingPlayer::FAdaptiveStreamingPlayer(const IAdaptiveStreamingPlaye
 	bShouldBePaused  = false;
 	bShouldBePlaying = false;
 
-	SeekVars.Clear();
-
+	SeekVars.Reset();
 	CurrentLoopState = {};
 	NextLoopStates.Empty();
 	bFirstSegmentRequestIsForLooping = false;
@@ -656,15 +654,9 @@ void FAdaptiveStreamingPlayer::Resume()
  */
 void FAdaptiveStreamingPlayer::SeekTo(const FSeekParam& NewPosition)
 {
-	TSharedPtrTS<IManifest> CurrentManifest = Manifest;
-	if (CurrentManifest.IsValid())
-	{
-		WorkerThread.SendSeekMessage(NewPosition);
-	}
-	else
-	{
-		PostLog(Facility::EFacility::Player, IInfoLog::ELevel::Error, FString::Printf(TEXT("API error: Cannot perform seek. No playlist loaded yet. Seeking is possible only when a valid playlist is available!")));
-	}
+	FScopeLock lock(&SeekVars.Lock);
+	SeekVars.PendingRequest = NewPosition;
+	WorkerThread.TriggerSharedWorkerThread();
 }
 
 
@@ -1297,6 +1289,8 @@ void FAdaptiveStreamingPlayer::InternalHandleOnce()
 
 		// Update diag buffers
 		UpdateDiagnostics();
+		// Handle seek requests.
+		HandleSeeking();
 		// Check for end of stream.
 		CheckForStreamEnd();
 		// Check the error queue for new arrivals.
@@ -1372,39 +1366,6 @@ bool FAdaptiveStreamingPlayer::InternalHandleThreadMessages()
 				bShouldBePlaying = true;
 				break;
 			}
-			case FWorkerThreadMessages::FMessage::EType::Seek:
-			{
-				FSeekParam StartAtTime = msg.Data.StartPlay.Position;
-
-				SeekVars.Clear();
-
-				if (bIsPlayStart)
-				{
-					CurrentState = EPlayerState::eState_Buffering;
-					InternalStartAt(StartAtTime);
-				}
-				else
-				{
-					// Stop everything.
-					InternalStop(PlayerConfig.bHoldLastFrameDuringSeek);
-					if (Manifest.IsValid())
-					{
-						Manifest->UpdateDynamicRefetchCounter();
-					}
-					CurrentState = EPlayerState::eState_Seeking;
-
-					if (StartAtTime.bOptimizeForScrubbing.IsSet())
-					{
-						SeekVars.bForScrubbing = StartAtTime.bOptimizeForScrubbing.GetValue();
-					}
-					else
-					{
-						SeekVars.bForScrubbing = PlayerOptions.GetValue(OptionKeyFrameOptimizeSeekForScrubbing).SafeGetBool(false);
-					}
-					InternalStartAt(StartAtTime);
-				}
-				break;
-			}
 			case FWorkerThreadMessages::FMessage::EType::Loop:
 			{
 				InternalSetLoop(msg.Data.Looping.Loop);
@@ -1658,6 +1619,80 @@ void FAdaptiveStreamingPlayer::HandlePlayStateChanges()
 }
 
 
+void FAdaptiveStreamingPlayer::HandleSeeking()
+{
+	// Are we in a state where a seek is possible?
+	if (CurrentState == EPlayerState::eState_Idle || CurrentState == EPlayerState::eState_ParsingManifest || CurrentState == EPlayerState::eState_PreparingStreams || CurrentState == EPlayerState::eState_Error ||
+		!Manifest.IsValid())
+	{
+		// No, return.
+		return;
+	}
+
+	// When playing the last successfully seeked position is irrelevant and cannot be used as a reference any more.
+	if (PlaybackState.GetIsPlaying())
+	{
+		SeekVars.InvalidateLastFinished();
+	}
+
+	FScopeLock lock(&SeekVars.Lock);
+	// Is there a pending request?
+	if (SeekVars.PendingRequest.IsSet())
+	{
+		// If there is an active request and the new request is for scrubbing we let the active request finish first.
+		if (SeekVars.ActiveRequest.IsSet() && SeekVars.PendingRequest.GetValue().bOptimizeForScrubbing.Get(false))
+		{
+			return;
+		}
+
+		// Check the distance to the last seek performed, if there is one.
+		if (SeekVars.LastFinishedRequest.IsSet())
+		{
+			double Distance = Utils::AbsoluteValue(SeekVars.LastFinishedRequest.GetValue().Time.GetAsSeconds() - SeekVars.PendingRequest.GetValue().Time.GetAsSeconds());
+			if (Distance <= SeekVars.PendingRequest.GetValue().DistanceThreshold.Get(0.0))
+			{
+				// Already there
+				SeekVars.PendingRequest.Reset();
+				DispatchEvent(FMetricEvent::ReportSeekCompleted());
+				return;
+			}
+		}
+
+		// Trigger the seek.
+		FSeekParam SeekParam = SeekVars.PendingRequest.GetValue();
+		SeekVars.PendingRequest.Reset();
+		SeekVars.ClearWorkVars();
+
+		if (SeekVars.bIsPlayStart)
+		{
+			CurrentState = EPlayerState::eState_Buffering;
+		}
+		else
+		{
+			// Stop everything.
+			InternalStop(PlayerConfig.bHoldLastFrameDuringSeek);
+			if (Manifest.IsValid())
+			{
+				Manifest->UpdateDynamicRefetchCounter();
+			}
+			CurrentState = EPlayerState::eState_Seeking;
+
+			if (SeekParam.bOptimizeForScrubbing.IsSet())
+			{
+				SeekVars.bForScrubbing = SeekParam.bOptimizeForScrubbing.GetValue();
+			}
+			else
+			{
+				SeekVars.bForScrubbing = PlayerOptions.GetValue(OptionKeyFrameOptimizeSeekForScrubbing).SafeGetBool(false);
+			}
+		}
+		SeekVars.ActiveRequest = SeekParam;
+		InternalStartAt(SeekVars.ActiveRequest.GetValue());
+	}
+}
+
+
+
 //-----------------------------------------------------------------------------
 /**
  * Handle changes in metadata, like timeline changes or track availability.
@@ -2325,6 +2360,8 @@ void FAdaptiveStreamingPlayer::HandleNewOutputData()
 					CurrentState = EPlayerState::eState_Paused;
 					PlaybackRate = 0.0;
 					PrerollVars.Clear();
+					SeekVars.ClearWorkVars();
+					SeekVars.SetFinished();
 					if (LastBufferingState == EPlayerState::eState_Seeking)
 					{
 						DispatchEvent(FMetricEvent::ReportSeekCompleted());
@@ -2346,6 +2383,7 @@ void FAdaptiveStreamingPlayer::HandleNewOutputData()
 				if (!SeekVars.bScrubPrerollDone)
 				{
 					SeekVars.bScrubPrerollDone = true;
+					SeekVars.SetFinished();
 					DispatchEvent(FMetricEvent::ReportSeekCompleted());
 					// Pause feeding the decoder to stop draining buffers we want to fill.
 					DecoderState = EDecoderState::eDecoder_Paused;
@@ -2356,7 +2394,7 @@ void FAdaptiveStreamingPlayer::HandleNewOutputData()
 					PlaybackState.SetIsBuffering(false);
 					PlaybackState.SetIsSeeking(false);
 
-					SeekVars.Clear();
+					SeekVars.ClearWorkVars();
 					PipelineState = EPipelineState::ePipeline_Stopped;
 					CurrentState = EPlayerState::eState_Paused;
 					DecoderState = EDecoderState::eDecoder_Running;
@@ -4037,7 +4075,7 @@ void FAdaptiveStreamingPlayer::InternalStartAt(const FSeekParam& NewPosition)
 	PendingStartRequest->SearchType = IManifest::ESearchType::Closest;
 	PendingStartRequest->StartAt.Time = NewPosition.Time;
 	PendingStartRequest->StartingBitrate = NewPosition.StartingBitrate;
-	PendingStartRequest->StartType = bIsPlayStart ? FPendingStartRequest::EStartType::PlayStart : FPendingStartRequest::EStartType::Seeking;
+	PendingStartRequest->StartType = SeekVars.bIsPlayStart ? FPendingStartRequest::EStartType::PlayStart : FPendingStartRequest::EStartType::Seeking;
 
 
 	// The fragment should be the closest to the time stamp unless we are rebuffering in which case it should be the one we failed on.
@@ -4057,7 +4095,7 @@ void FAdaptiveStreamingPlayer::InternalStartAt(const FSeekParam& NewPosition)
 	AEMSEventHandler->PlaybackStartingUp();
 
 	// From this point on any further start will not be the very first one any more.
-	bIsPlayStart = false;
+	SeekVars.bIsPlayStart = false;
 }
 
 
@@ -4198,7 +4236,7 @@ void FAdaptiveStreamingPlayer::InternalStop(bool bHoldCurrentFrame)
 	UpcomingPeriods.Empty();
 	MetadataHandlingState.Reset();
 	NextLoopStates.Empty();
-	SeekVars.Clear();
+	SeekVars.Reset();
 
 	PlaybackRate = 0.0;
 	CurrentState = EPlayerState::eState_Paused;
diff --git a/Engine/Plugins/Media/ElectraPlayer/Source/ElectraPlayerRuntime/Private/Runtime/Player/AdaptiveStreamingPlayer.h b/Engine/Plugins/Media/ElectraPlayer/Source/ElectraPlayerRuntime/Private/Runtime/Player/AdaptiveStreamingPlayer.h
--- a/Engine/Plugins/Media/ElectraPlayer/Source/ElectraPlayerRuntime/Private/Runtime/Player/AdaptiveStreamingPlayer.h
+++ b/Engine/Plugins/Media/ElectraPlayer/Source/ElectraPlayerRuntime/Private/Runtime/Player/AdaptiveStreamingPlayer.h
@@ -114,13 +114,33 @@ public:
 			StartingBitrate.Reset();
 			bOptimizeForScrubbing.Reset();
 		}
+		// Time to seek to.
 		FTimeValue	Time;
+		// Maximum stream bitrate to use when seeking.
 		TOptional<int32> StartingBitrate;
+		// Optimize for frame scrubbing (faster display of frame at target time)?
 		TOptional<bool> bOptimizeForScrubbing;
+		// Allowed distance to last performed seek to save a redundant new seek.
+		TOptional<double> DistanceThreshold;
 	};
-	//! Seek to a new position and play from there. This includes first playstart.
-	//! Playback is initially paused on first player use and must be resumed to begin.
-	//! Query the seekable range (GetSeekableRange()) to get the valid time range.
+	
+	/**
+	 * Seek to a new position and play from there. This includes first playstart.
+	 * Playback is initially paused on first player use and must be resumed to begin.
+	 * Query the seekable range (GetSeekableRange()) to get the valid time range.
+	 * 
+	 * Seeks can be issued while a seek is already executing. The seek parameters
+	 * control behaviour. If seeking is performed for scrubbing any new seek will
+	 * be performed only when the previous seek has completed, otherwise the current
+	 * seek will be canceled in favor of the new.
+	 * If the new position being seeked to is within the specified distance to the
+	 * last completed seek a new seek will not be performed.
+	 * If new seeks are performed in rapid succession (as in scrubbing) not every
+	 * new position will be seeked to. Seek commands are aggregated and the position
+	 * set with the most recent call will be used as soon as any previous seek completes.
+	 * As a result you will NOT receive a ReportSeekCompleted() notification for
+	 * every seek requested, only for those that were executed and allowed to complete.
+	 */
 	virtual void SeekTo(const FSeekParam& NewPosition) = 0;
 
 	//! Pauses playback.
diff --git a/Engine/Plugins/Media/ElectraPlayer/Source/ElectraPlayerRuntime/Private/Runtime/Player/AdaptiveStreamingPlayerInternal.h b/Engine/Plugins/Media/ElectraPlayer/Source/ElectraPlayerRuntime/Private/Runtime/Player/AdaptiveStreamingPlayerInternal.h
--- a/Engine/Plugins/Media/ElectraPlayer/Source/ElectraPlayerRuntime/Private/Runtime/Player/AdaptiveStreamingPlayerInternal.h
+++ b/Engine/Plugins/Media/ElectraPlayer/Source/ElectraPlayerRuntime/Private/Runtime/Player/AdaptiveStreamingPlayerInternal.h
@@ -1163,7 +1163,6 @@ private:
 				LoadManifest,
 				Pause,
 				Resume,
-				Seek,
 				Loop,
 				Close,
 				ChangeBitrate,
@@ -1206,10 +1205,6 @@ private:
 					}
 					FMediaEvent*									Event;
 				};
-				struct FStartPlay
-				{
-					FSeekParam										Position;
-				};
 				struct FLoop
 				{
 					FLoop()
@@ -1248,7 +1243,6 @@ private:
 				FLoadManifest				ManifestToLoad;
 				FStreamReader				StreamReader;
 				FEvent						MediaEvent;
-				FStartPlay					StartPlay;
 				FLoop						Looping;
 				FBitrate					Bitrate;
 				FSession					Session;
@@ -1300,13 +1294,6 @@ private:
 			Msg.Type = FMessage::EType::Resume;
 			TriggerSharedWorkerThread(MoveTemp(Msg));
 		}
-		void SendSeekMessage(const FSeekParam& NewPosition)
-		{
-			FMessage Msg;
-			Msg.Type = FMessage::EType::Seek;
-			Msg.Data.StartPlay.Position = NewPosition;
-			TriggerSharedWorkerThread(MoveTemp(Msg));
-		}
 		void SendLoopMessage(const FLoopParam& InLoopParams)
 		{
 			FMessage Msg;
@@ -1383,6 +1370,15 @@ private:
 				SharedWorkerThread->TriggerWork();
 			}
 		}
+
+		void TriggerSharedWorkerThread()
+		{
+			if (SharedWorkerThread.IsValid())
+			{
+				SharedWorkerThread->TriggerWork();
+			}
+		}
+
 		TSharedPtrTS<FAdaptiveStreamingPlayerWorkerThread>	SharedWorkerThread;
 		TQueue<FMessage, EQueueMode::Spsc>					WorkMessages;
 	};
@@ -1545,18 +1541,43 @@ private:
 
 	struct FSeekVars
 	{
-		FSeekVars()
+		void ClearWorkVars()
 		{
-			Clear();
+			bForScrubbing = false;
+			bScrubPrerollDone = false;
 		}
-		void Clear()
+
+		void Reset()
 		{
+			FScopeLock lock(&Lock);
+			PendingRequest.Reset();
+			ActiveRequest.Reset();
+			LastFinishedRequest.Reset();
 			bForScrubbing = false;
 			bScrubPrerollDone = false;
 		}
 
-		bool	bForScrubbing;
-		bool	bScrubPrerollDone;
+		void SetFinished()
+		{
+			LastFinishedRequest = ActiveRequest;
+			ActiveRequest.Reset();
+		}
+
+		void InvalidateLastFinished()
+		{
+			LastFinishedRequest.Reset();
+		}
+
+		bool	bForScrubbing = false;
+		bool	bScrubPrerollDone = false;
+		bool	bIsPlayStart = true;
+
+		TOptional<FSeekParam> ActiveRequest;
+		TOptional<FSeekParam> LastFinishedRequest;
+
+		FCriticalSection Lock;
+		// The pending request must be accessed under lock since it is written to by the main thread and read from the worker thread!
+		TOptional<FSeekParam> PendingRequest;
 	};
 
 	struct FStreamBitrateInfo
@@ -1709,6 +1730,7 @@ private:
 	void HandleNewOutputData();
 	void HandleSessionMessage(TSharedPtrTS<IPlayerMessage> SessionMessage);
 	void HandlePlayStateChanges();
+	void HandleSeeking();
 	void HandlePendingMediaSegmentRequests();
 	void CancelPendingMediaSegmentRequests(EStreamType StreamType);
 	void HandleDeselectedBuffers();
@@ -1868,7 +1890,6 @@ private:
 	double																PlaybackRate;
 	FTimeValue															RebufferDetectedAtPlayPos;
 	bool																bRebufferPending;
-	bool																bIsPlayStart;
 	bool																bIsClosing;
 
 	TSharedPtrTS<IAdaptiveStreamSelector>								StreamSelector;
