ElectraPlayer: Fixed incorrect flag assignment which should have been |=

#rnx

diff --git a/Engine/Plugins/Media/ElectraPlayer/Source/ElectraPlayerRuntime/Private/Runtime/Decoder/Apple/VideoDecoderH264_Apple.cpp b/Engine/Plugins/Media/ElectraPlayer/Source/ElectraPlayerRuntime/Private/Runtime/Decoder/Apple/VideoDecoderH264_Apple.cpp
--- a/Engine/Plugins/Media/ElectraPlayer/Source/ElectraPlayerRuntime/Private/Runtime/Decoder/Apple/VideoDecoderH264_Apple.cpp
+++ b/Engine/Plugins/Media/ElectraPlayer/Source/ElectraPlayerRuntime/Private/Runtime/Decoder/Apple/VideoDecoderH264_Apple.cpp
@@ -1081,7 +1081,7 @@ bool FVideoDecoderH264::FlushDecoder()
 {
 	if (DecoderHandle && DecoderHandle->DecompressionSession)
 	{
-		VTDecompressionSessionFinishDelayedFrames(DecoderHandle->DecompressionSession);
+		// This call implies VTDecompressionSessionFinishDelayedFrames();
 		VTDecompressionSessionWaitForAsynchronousFrames(DecoderHandle->DecompressionSession);
 		// Push out all images we got from the decoder from our internal PTS sorting facility
 		ProcessOutput(true);
@@ -1332,7 +1332,7 @@ FVideoDecoderH264::EDecodeResult FVideoDecoderH264::Decode(TSharedPtr<FDecoderIn
 	VTDecodeFrameFlags DecodeFlags = kVTDecodeFrame_EnableAsynchronousDecompression | kVTDecodeFrame_EnableTemporalProcessing;
 	if (bRecreatingSession || !AU->AdjustedPTS.IsValid())
 	{
-		DecodeFlags = kVTDecodeFrame_DoNotOutputFrame;
+		DecodeFlags |= kVTDecodeFrame_DoNotOutputFrame;
 	}
 	VTDecodeInfoFlags  InfoFlags = 0;
 /*
diff --git a/Engine/Plugins/Media/ElectraPlayer/Source/ElectraPlayerRuntime/Private/Runtime/Decoder/Apple/VideoDecoderH265_Apple.cpp b/Engine/Plugins/Media/ElectraPlayer/Source/ElectraPlayerRuntime/Private/Runtime/Decoder/Apple/VideoDecoderH265_Apple.cpp
--- a/Engine/Plugins/Media/ElectraPlayer/Source/ElectraPlayerRuntime/Private/Runtime/Decoder/Apple/VideoDecoderH265_Apple.cpp
+++ b/Engine/Plugins/Media/ElectraPlayer/Source/ElectraPlayerRuntime/Private/Runtime/Decoder/Apple/VideoDecoderH265_Apple.cpp
@@ -1062,7 +1062,7 @@ bool FVideoDecoderH265::FlushDecoder()
 {
 	if (DecoderHandle && DecoderHandle->DecompressionSession)
 	{
-		VTDecompressionSessionFinishDelayedFrames(DecoderHandle->DecompressionSession);
+		// This call implies VTDecompressionSessionFinishDelayedFrames();
 		VTDecompressionSessionWaitForAsynchronousFrames(DecoderHandle->DecompressionSession);
 		// Push out all images we got from the decoder from our internal PTS sorting facility
 		ProcessOutput(true);
@@ -1314,7 +1314,7 @@ FVideoDecoderH265::EDecodeResult FVideoDecoderH265::Decode(TSharedPtr<FDecoderIn
 	VTDecodeFrameFlags DecodeFlags = kVTDecodeFrame_EnableAsynchronousDecompression | kVTDecodeFrame_EnableTemporalProcessing;
 	if (bRecreatingSession || !AU->AdjustedPTS.IsValid())
 	{
-		DecodeFlags = kVTDecodeFrame_DoNotOutputFrame;
+		DecodeFlags |= kVTDecodeFrame_DoNotOutputFrame;
 	}
 	VTDecodeInfoFlags  InfoFlags = 0;
 /*
