ElectraPlayer: Improved seeking close to or past the end of a video

#rnx

diff --git a/Engine/Plugins/Media/ElectraPlayer/Source/ElectraPlayerRuntime/Private/ElectraPlayer.cpp b/Engine/Plugins/Media/ElectraPlayer/Source/ElectraPlayerRuntime/Private/ElectraPlayer.cpp
--- a/Engine/Plugins/Media/ElectraPlayer/Source/ElectraPlayerRuntime/Private/ElectraPlayer.cpp
+++ b/Engine/Plugins/Media/ElectraPlayer/Source/ElectraPlayerRuntime/Private/ElectraPlayer.cpp
@@ -1064,6 +1064,12 @@ void FElectraPlayer::CalculateTargetSeekTime(FTimespan& OutTargetTime, const FTi
 
 
 bool FElectraPlayer::Seek(const FTimespan& Time)
+{
+	FSeekParam NoParams;
+	return Seek(Time, NoParams);
+}
+
+bool FElectraPlayer::Seek(const FTimespan& Time, const FSeekParam& Param)
 {
 	if (CurrentPlayer.Get())
 	{
@@ -1071,6 +1077,9 @@ bool FElectraPlayer::Seek(const FTimespan& Time)
 		CalculateTargetSeekTime(Target, Time);
 		Electra::IAdaptiveStreamingPlayer::FSeekParam seek;
 		seek.Time.SetFromTimespan(Target);
+		seek.StartingBitrate = Param.StartingBitrate;
+		seek.bOptimizeForScrubbing = Param.bOptimizeForScrubbing;
+		seek.DistanceThreshold = Param.DistanceThreshold;
 		bInitialSeekPerformed = true;
 		CurrentPlayer->AdaptivePlayer->SeekTo(seek);
 		return true;
@@ -1088,6 +1097,16 @@ void FElectraPlayer::SetFrameAccurateSeekMode(bool bEnableFrameAccuracy)
 	}
 }
 
+void FElectraPlayer::ModifyOptions(const FParamDict& InOptionsToSetOrChange, const FParamDict& InOptionsToClear)
+{
+	TSharedPtr<FInternalPlayerImpl, ESPMode::ThreadSafe> LockedPlayer = CurrentPlayer;
+	if (LockedPlayer.IsValid() && LockedPlayer->AdaptivePlayer.IsValid())
+	{
+		LockedPlayer->AdaptivePlayer->ModifyOptions(InOptionsToSetOrChange, InOptionsToClear);
+	}
+}
+
+
 void FElectraPlayer::SetPlaybackRange(const FPlaybackRange& InPlaybackRange)
 {
 	CurrentPlaybackRange = InPlaybackRange;
@@ -2137,6 +2156,11 @@ void FElectraPlayer::HandlePlayerEventBufferUtilization(const Electra::Metrics::
 
 void FElectraPlayer::HandlePlayerEventSegmentDownload(const Electra::Metrics::FSegmentDownloadStats& SegmentDownloadStats)
 {
+	// Cached responses are not actual network traffic, so we ignore them.
+	if (SegmentDownloadStats.bIsCachedResponse)
+	{
+		return;
+	}
 	// Update statistics
 	FScopeLock Lock(&StatisticsLock);
 	if (SegmentDownloadStats.StreamType == Electra::EStreamType::Video)
@@ -2874,6 +2898,7 @@ void FElectraPlayer::MediaStateOnEndReached()
 	case EPlayerState::Preparing:
 	case EPlayerState::Playing:
 	case EPlayerState::Paused:
+	case EPlayerState::Stopped:
 		DeferredEvents.Enqueue(IElectraPlayerAdapterDelegate::EPlayerEvent::PlaybackEndReached);
 		break;
 	default:
diff --git a/Engine/Plugins/Media/ElectraPlayer/Source/ElectraPlayerRuntime/Private/ElectraPlayer.h b/Engine/Plugins/Media/ElectraPlayer/Source/ElectraPlayerRuntime/Private/ElectraPlayer.h
--- a/Engine/Plugins/Media/ElectraPlayer/Source/ElectraPlayerRuntime/Private/ElectraPlayer.h
+++ b/Engine/Plugins/Media/ElectraPlayer/Source/ElectraPlayerRuntime/Private/ElectraPlayer.h
@@ -86,7 +86,7 @@ public:
 
 	// -------- PlayerAdapter (Plugin/Native) API
 
-	bool OpenInternal(const FString& Url, const FParamDict & PlayerOptions, const FPlaystartOptions & InPlaystartOptions) override;
+	bool OpenInternal(const FString& Url, const FParamDict& PlayerOptions, const FPlaystartOptions& InPlaystartOptions) override;
 	void CloseInternal(bool bKillAfterClose) override;
 
 	void Tick(FTimespan DeltaTime, FTimespan Timecode) override;
@@ -113,8 +113,11 @@ public:
 	bool SetRate(float Rate) override;
 
 	bool Seek(const FTimespan& Time) override;
+	bool Seek(const FTimespan& Time, const FSeekParam& Param) override;
 	void SetFrameAccurateSeekMode(bool bEnableFrameAccuracy) override;
 
+	void ModifyOptions(const FParamDict& InOptionsToSetOrChange, const FParamDict& InOptionsToClear) override;
+
 	bool GetAudioTrackFormat(int32 TrackIndex, int32 FormatIndex, FAudioTrackFormat& OutFormat) const override;
 	bool GetVideoTrackFormat(int32 TrackIndex, int32 FormatIndex, FVideoTrackFormat& OutFormat) const override;
 
diff --git a/Engine/Plugins/Media/ElectraPlayer/Source/ElectraPlayerRuntime/Private/Runtime/Player/AdaptiveStreamingPlayer.cpp b/Engine/Plugins/Media/ElectraPlayer/Source/ElectraPlayerRuntime/Private/Runtime/Player/AdaptiveStreamingPlayer.cpp
--- a/Engine/Plugins/Media/ElectraPlayer/Source/ElectraPlayerRuntime/Private/Runtime/Player/AdaptiveStreamingPlayer.cpp
+++ b/Engine/Plugins/Media/ElectraPlayer/Source/ElectraPlayerRuntime/Private/Runtime/Player/AdaptiveStreamingPlayer.cpp
@@ -400,6 +400,24 @@ void FAdaptiveStreamingPlayer::Initialize(const FParamDict& Options)
 	}
 }
 
+void FAdaptiveStreamingPlayer::ModifyOptions(const FParamDict& InOptionsToSetOrChange, const FParamDict& InOptionsToClear)
+{
+	// Remove options
+	TArray<FString> Keys;
+	InOptionsToClear.GetKeysStartingWith(FString(), Keys);
+	for(auto &Key : Keys)
+	{
+		PlayerOptions.Remove(Key);
+	}
+
+	InOptionsToSetOrChange.GetKeysStartingWith(FString(), Keys);
+	for(auto &Key : Keys)
+	{
+		FVariantValue Value = InOptionsToSetOrChange.GetValue(Key);
+		PlayerOptions.SetOrUpdate(Key, Value);
+	}
+}
+
 
 void FAdaptiveStreamingPlayer::SetStaticResourceProviderCallback(const TSharedPtr<IAdaptiveStreamingPlayerResourceProvider, ESPMode::ThreadSafe>& InStaticResourceProvider)
 {
@@ -1607,11 +1625,11 @@ void FAdaptiveStreamingPlayer::HandlePlayStateChanges()
 		}
 	}
 
-	// Did the play range change?
-	if (PlaybackState.GetPlayRangeHasChanged())
+	// Did the play range or loop state change such that we need to start over?
+	if (PlaybackState.GetPlayRangeHasChanged() || PlaybackState.GetLoopStateHasChanged())
 	{
-		// When in a possible playback state we need to start over at the current position.
-		if (CurrentState == EPlayerState::eState_Playing || CurrentState == EPlayerState::eState_Paused)
+		// We do not do this when paused as setting the range usually happens only when paused for seeking or scrubbing.
+		if (CurrentState == EPlayerState::eState_Playing)
 		{
 			InternalStartoverAtCurrentPosition();
 		}
@@ -1640,7 +1658,8 @@ void FAdaptiveStreamingPlayer::HandleSeeking()
 	if (SeekVars.PendingRequest.IsSet())
 	{
 		// If there is an active request and the new request is for scrubbing we let the active request finish first.
-		if (SeekVars.ActiveRequest.IsSet() && SeekVars.PendingRequest.GetValue().bOptimizeForScrubbing.Get(false))
+		bool bIsForScrubbing = SeekVars.PendingRequest.GetValue().bOptimizeForScrubbing.Get(PlayerOptions.GetValue(OptionKeyFrameOptimizeSeekForScrubbing).SafeGetBool(false));
+		if (SeekVars.ActiveRequest.IsSet() && bIsForScrubbing)
 		{
 			return;
 		}
@@ -1653,7 +1672,7 @@ void FAdaptiveStreamingPlayer::HandleSeeking()
 			{
 				// Already there
 				SeekVars.PendingRequest.Reset();
-				DispatchEvent(FMetricEvent::ReportSeekCompleted());
+				DispatchEvent(FMetricEvent::ReportSeekCompleted(true));
 				return;
 			}
 		}
@@ -1677,14 +1696,7 @@ void FAdaptiveStreamingPlayer::HandleSeeking()
 			}
 			CurrentState = EPlayerState::eState_Seeking;
 
-			if (SeekParam.bOptimizeForScrubbing.IsSet())
-			{
-				SeekVars.bForScrubbing = SeekParam.bOptimizeForScrubbing.GetValue();
-			}
-			else
-			{
-				SeekVars.bForScrubbing = PlayerOptions.GetValue(OptionKeyFrameOptimizeSeekForScrubbing).SafeGetBool(false);
-			}
+			SeekVars.bForScrubbing = bIsForScrubbing;
 		}
 		SeekVars.ActiveRequest = SeekParam;
 		InternalStartAt(SeekVars.ActiveRequest.GetValue());
@@ -2082,6 +2094,43 @@ bool FAdaptiveStreamingPlayer::HaveEnoughBufferedDataToStartPlayback()
 }
 
 
+void FAdaptiveStreamingPlayer::PrepareForPrerolling()
+{
+	PrerollVars.StartTime = MEDIAutcTime::CurrentMSec();
+
+	// We need to check if there is any data available in the buffers that can be decoded.
+	// It may be possible that playback was starting on the last segment that did not bring in any playable
+	// data or none at all if it was past end of stream.
+	DiagnosticsCriticalSection.Lock();
+
+	// Check video buffer
+	if (StreamReaderHandler && bHaveVideoReader.GetWithDefault(false))
+	{
+		if (VideoBufferStats.StreamBuffer.bEndOfData && (VideoBufferStats.StreamBuffer.PlayableDuration == FTimeValue::GetZero() || VideoBufferStats.StreamBuffer.NumCurrentAccessUnits == 0))
+		{
+			PrerollVars.bHaveEnoughVideo = true;
+			VideoBufferStats.DecoderInputBuffer.bEODSignaled = true;
+			VideoBufferStats.DecoderInputBuffer.bEODReached = true;
+			VideoBufferStats.DecoderOutputBuffer.bEODreached = true;
+		}
+	}
+
+	// Check audio buffer
+	if (StreamReaderHandler && bHaveAudioReader.GetWithDefault(false))
+	{
+		if (AudioBufferStats.StreamBuffer.bEndOfData && (AudioBufferStats.StreamBuffer.PlayableDuration == FTimeValue::GetZero() || AudioBufferStats.StreamBuffer.NumCurrentAccessUnits == 0))
+		{
+			PrerollVars.bHaveEnoughAudio = true;
+			AudioBufferStats.DecoderInputBuffer.bEODSignaled = true;
+			AudioBufferStats.DecoderInputBuffer.bEODReached = true;
+			AudioBufferStats.DecoderOutputBuffer.bEODreached = true;
+		}
+	}
+
+	DiagnosticsCriticalSection.Unlock();
+}
+
+
 //-----------------------------------------------------------------------------
 /**
  * Checks if buffers have enough data to advance the play state.
@@ -2119,9 +2168,10 @@ void FAdaptiveStreamingPlayer::HandleNewBufferedData()
 				}
 				DataBuffersCriticalSection.Unlock();
 
-				DecoderState = EDecoderState::eDecoder_Running;
+				PrepareForPrerolling();
 				PipelineState = EPipelineState::ePipeline_Prerolling;
-				PrerollVars.StartTime = MEDIAutcTime::CurrentMSec();
+				DecoderState = EDecoderState::eDecoder_Running;
+
 				PostrollVars.Clear();
 				// We are synced up now and can clear the state.
 				LiveSyncVars.Clear();
@@ -2152,9 +2202,9 @@ void FAdaptiveStreamingPlayer::HandleNewBufferedData()
 				DataBuffersCriticalSection.Unlock();
 				if (bHaveActiveOutputBuffer)
 				{
-					DecoderState = EDecoderState::eDecoder_Running;
+					PrepareForPrerolling();
 					PipelineState = EPipelineState::ePipeline_Prerolling;
-					PrerollVars.StartTime = MEDIAutcTime::CurrentMSec();
+					DecoderState = EDecoderState::eDecoder_Running;
 					PostrollVars.Clear();
 					LiveSyncVars.Clear();
 					DispatchEvent(FMetricEvent::ReportPrerollStart());
@@ -2302,9 +2352,9 @@ void FAdaptiveStreamingPlayer::HandleNewOutputData()
 			if (StreamReaderHandler && bHaveVideoReader.Value())
 			{
 				// If the video decoder output buffer is stalled then we have enough video output. There won't be any more coming in.
-				// Also when the stream has reached the end and that has propagated through the buffers.
+				// Also when the stream has reached the end and that has propagated through the decoder's buffer.
 				if (VideoBufferStats.DecoderOutputBuffer.bOutputStalled ||
-					(VideoBufferStats.StreamBuffer.bEndOfData && VideoBufferStats.DecoderInputBuffer.bEODSignaled && VideoBufferStats.DecoderOutputBuffer.bEODreached) ||
+					(VideoBufferStats.StreamBuffer.bEndOfData && VideoBufferStats.DecoderInputBuffer.bEODSignaled && VideoBufferStats.DecoderInputBuffer.bEODReached) ||
 					bIsVideoDeselected)
 				{
 					bHaveEnoughVideo = true;
@@ -2323,9 +2373,9 @@ void FAdaptiveStreamingPlayer::HandleNewOutputData()
 			if (StreamReaderHandler && bHaveAudioReader.Value())
 			{
 				// If the audio decoder output buffer is stalled then we have enough video output. There won't be any more coming in.
-				// Also when the stream has reached the end and that has propagated through the buffers.
+				// Also when the stream has reached the end and that has propagated through the decoder's buffer.
 				if (AudioBufferStats.DecoderOutputBuffer.bOutputStalled ||
-					(AudioBufferStats.StreamBuffer.bEndOfData && AudioBufferStats.DecoderInputBuffer.bEODSignaled && AudioBufferStats.DecoderOutputBuffer.bEODreached) ||
+					(AudioBufferStats.StreamBuffer.bEndOfData && AudioBufferStats.DecoderInputBuffer.bEODSignaled && AudioBufferStats.DecoderInputBuffer.bEODReached) ||
 					bIsAudioDeselected)
 				{
 					bHaveEnoughAudio = true;
@@ -2490,8 +2540,19 @@ void FAdaptiveStreamingPlayer::InternalHandlePendingStartRequest(const FTimeValu
 						}
 						case IManifest::FResult::EType::PastEOS:
 						{
-							PendingStartRequest.Reset();
-							InternalSetPlaybackEnded();
+							// If the initial start time is beyond the end of the presentation
+							// but looping is enabled we set this request up to start at the beginning.
+							if (CurrentLoopParam.bEnableLooping)
+							{
+								FTimeRange Seekable = Manifest->GetSeekableTimeRange();
+								PendingStartRequest->SearchType = IManifest::ESearchType::Closest;
+								PendingStartRequest->StartAt.Time = Seekable.Start;
+							}
+							else
+							{
+								PendingStartRequest.Reset();
+								InternalSetPlaybackEnded();
+							}
 							break;
 						}
 						case IManifest::FResult::EType::NotFound:
@@ -2577,6 +2638,7 @@ void FAdaptiveStreamingPlayer::InternalHandlePendingStartRequest(const FTimeValu
 								{
 									StartingBitrate = PlayerOptions.GetValue(OptionKeySeekStartBitrate).SafeGetInt64(StartingBitrate);
 									bForceInitial = true;
+									bForceStartRate = true;
 								}
 							}
 						}
@@ -3912,16 +3974,20 @@ void FAdaptiveStreamingPlayer::CheckForStreamEnd()
 			// Check for end of video stream
 			if (StreamReaderHandler && bHaveVid)
 			{
-				// All buffers at end of data?
-				bEndVid = (vidStats.StreamBuffer.bEndOfData && vidStats.DecoderInputBuffer.bEODSignaled && vidStats.DecoderOutputBuffer.bEODreached);
+				// Source buffer and decoder input buffer at end of data?
+				// We do not check the decoder output buffer since it is possible that no decodable output was produced
+				// which tends to happen at the end of the stream or when there is no video any more when audio is longer.
+				bEndVid = (vidStats.StreamBuffer.bEndOfData && vidStats.DecoderInputBuffer.bEODSignaled && vidStats.DecoderInputBuffer.bEODReached);
 				VidStalled = vidStats.StreamBuffer.bEndOfData ? vidStats.GetStalledDurationMillisec() : 0;
 			}
 
 			// Check for end of audio stream
 			if (StreamReaderHandler && bHaveAud)
 			{
-				// All buffers at end of data?
-				bEndAud = (audStats.StreamBuffer.bEndOfData && audStats.DecoderInputBuffer.bEODSignaled && audStats.DecoderOutputBuffer.bEODreached);
+				// Source buffer and decoder input buffer at end of data?
+				// We do not check the decoder output buffer since it is possible that no decodable output was produced
+				// which tends to happen at the end of the stream or when there is no audio any more when video is longer.
+				bEndAud = (audStats.StreamBuffer.bEndOfData && audStats.DecoderInputBuffer.bEODSignaled && audStats.DecoderInputBuffer.bEODReached);
 				AudStalled = audStats.StreamBuffer.bEndOfData ? audStats.GetStalledDurationMillisec() : 0;
 			}
 
@@ -4025,6 +4091,7 @@ void FAdaptiveStreamingPlayer::InternalStartAt(const FSeekParam& NewPosition)
 
 	PlaybackState.SetHasEnded(false);
 	PlaybackState.SetPlayRangeHasChanged(false);
+	PlaybackState.SetLoopStateHasChanged(false);
 
 	VideoBufferStats.Clear();
 	AudioBufferStats.Clear();
@@ -4112,6 +4179,12 @@ void FAdaptiveStreamingPlayer::InternalSetPlaybackEnded()
 	}
 	PlaybackState.SetPlayPosition(TimelineRange.End);
 
+	if (LastBufferingState == EPlayerState::eState_Seeking)
+	{
+		DispatchEvent(FMetricEvent::ReportSeekCompleted());
+	}
+	SeekVars.Reset();
+
 	PlaybackState.SetHasEnded(true);
 	StreamState = EStreamState::eStream_ReachedEnd;
 
@@ -4372,7 +4445,7 @@ void FAdaptiveStreamingPlayer::InternalSetLoop(const FLoopParam& LoopParam)
 		PlaybackState.SetLoopStateEnable(LoopParam.bEnableLooping);
 		if (bStartOver)
 		{
-			InternalStartoverAtCurrentPosition();
+			PlaybackState.SetLoopStateHasChanged(true);
 		}
 	}
 }
diff --git a/Engine/Plugins/Media/ElectraPlayer/Source/ElectraPlayerRuntime/Private/Runtime/Player/AdaptiveStreamingPlayer.h b/Engine/Plugins/Media/ElectraPlayer/Source/ElectraPlayerRuntime/Private/Runtime/Player/AdaptiveStreamingPlayer.h
--- a/Engine/Plugins/Media/ElectraPlayer/Source/ElectraPlayerRuntime/Private/Runtime/Player/AdaptiveStreamingPlayer.h
+++ b/Engine/Plugins/Media/ElectraPlayer/Source/ElectraPlayerRuntime/Private/Runtime/Player/AdaptiveStreamingPlayer.h
@@ -67,6 +67,11 @@ public:
 	//
 	virtual void Initialize(const FParamDict& InOptions) = 0;
 
+	/**
+	 * Modifies options. Not all options are modifiable during playback.
+	 */
+	virtual void ModifyOptions(const FParamDict& InOptionsToSetOrChange, const FParamDict& InOptionsToClear) = 0;
+
 	/**
 	 * Sets the attributes for the stream to start buffering for and playing.
 	 * This must be set before calling SeekTo() or LoadManifest() to have an immediate effect.
diff --git a/Engine/Plugins/Media/ElectraPlayer/Source/ElectraPlayerRuntime/Private/Runtime/Player/AdaptiveStreamingPlayerABR.cpp b/Engine/Plugins/Media/ElectraPlayer/Source/ElectraPlayerRuntime/Private/Runtime/Player/AdaptiveStreamingPlayerABR.cpp
--- a/Engine/Plugins/Media/ElectraPlayer/Source/ElectraPlayerRuntime/Private/Runtime/Player/AdaptiveStreamingPlayerABR.cpp
+++ b/Engine/Plugins/Media/ElectraPlayer/Source/ElectraPlayerRuntime/Private/Runtime/Player/AdaptiveStreamingPlayerABR.cpp
@@ -1216,9 +1216,9 @@ namespace Electra
 					bool bIsFeasible = (bCond1 || bCond2) && bCond3;
 
 					// If there is an enforced initial bitrate then any stream that is outside that bitrate is not feasible.
-					if (ForcedInitialBandwidth && StreamInformationVideo[nStr]->Bitrate > ForcedInitialBandwidth)
+					if (ForcedInitialBandwidth)
 					{
-						bIsFeasible = false;
+						bIsFeasible = StreamInformationVideo[nStr]->Bitrate <= ForcedInitialBandwidth;
 					}
 
 					// If there is an initial playstart bitrate set then this stream must be used, regardless of actual
diff --git a/Engine/Plugins/Media/ElectraPlayer/Source/ElectraPlayerRuntime/Private/Runtime/Player/AdaptiveStreamingPlayerInternal.h b/Engine/Plugins/Media/ElectraPlayer/Source/ElectraPlayerRuntime/Private/Runtime/Player/AdaptiveStreamingPlayerInternal.h
--- a/Engine/Plugins/Media/ElectraPlayer/Source/ElectraPlayerRuntime/Private/Runtime/Player/AdaptiveStreamingPlayerInternal.h
+++ b/Engine/Plugins/Media/ElectraPlayer/Source/ElectraPlayerRuntime/Private/Runtime/Player/AdaptiveStreamingPlayerInternal.h
@@ -181,6 +181,7 @@ struct FPlaybackState
 		bIsPlaying = false;
 		bIsPaused = false;
 		bPlayrateHasChanged = false;
+		bLoopStateHasChanged = false;
 	}
 
 	mutable FMediaCriticalSection			Lock;
@@ -196,6 +197,7 @@ struct FPlaybackState
 	bool									bIsPlaying;
 	bool									bIsPaused;
 	bool									bPlayrateHasChanged;
+	bool									bLoopStateHasChanged;
 	TArray<FTrackMetadata>					VideoTracks;
 	TArray<FTrackMetadata>					AudioTracks;
 	TArray<FTrackMetadata>					SubtitleTracks;
@@ -390,6 +392,16 @@ struct FPlaybackState
 		FMediaCriticalSection::ScopedLock lock(Lock);
 		return bPlayrateHasChanged;
 	}
+	void SetLoopStateHasChanged(bool bWasChanged)
+	{
+		FMediaCriticalSection::ScopedLock lock(Lock);
+		bLoopStateHasChanged = bWasChanged;
+	}
+	bool GetLoopStateHasChanged()
+	{
+		FMediaCriticalSection::ScopedLock lock(Lock);
+		return bLoopStateHasChanged;
+	}
 };
 
 
@@ -455,6 +467,10 @@ struct FMetricEvent
 			Metrics::ETimeJumpReason	Reason;
 
 		};
+		struct FSeekComplete
+		{
+			bool bWasAlreadyThere;
+		};
 		FString								URL;
 		Metrics::EBufferingReason			BufferingReason;
 		Metrics::FBufferStats				BufferStats;
@@ -466,6 +482,7 @@ struct FMetricEvent
 		FQualityChange						QualityChange;
 		FStreamCodecInformation				CodecFormatChange;
 		FTimeJumped							TimeJump;
+		FSeekComplete						SeekComplete;
 		FErrorDetail						ErrorDetail;
 		FLogMessage							LogMessage;
 	};
@@ -626,10 +643,11 @@ struct FMetricEvent
 		Evt->Type = EType::PlaybackStopped;
 		return Evt;
 	}
-	static TSharedPtrTS<FMetricEvent> ReportSeekCompleted()
+	static TSharedPtrTS<FMetricEvent> ReportSeekCompleted(bool bWasAlreadyThere=false)
 	{
 		TSharedPtrTS<FMetricEvent> Evt = MakeSharedTS<FMetricEvent>();
 		Evt->Type = EType::SeekCompleted;
+		Evt->Param.SeekComplete.bWasAlreadyThere = bWasAlreadyThere;
 		return Evt;
 	}
 	static TSharedPtrTS<FMetricEvent> ReportError(const FErrorDetail& ErrorDetail)
@@ -752,6 +770,7 @@ public:
 	virtual void RemoveSubtitleReceiver(TWeakPtrTS<IAdaptiveStreamingPlayerSubtitleReceiver> InReceiver) override;
 
 	virtual void Initialize(const FParamDict& Options) override;
+	virtual void ModifyOptions(const FParamDict& InOptionsToSetOrChange, const FParamDict& InOptionsToClear) override;
 
 	virtual void SetInitialStreamAttributes(EStreamType StreamType, const FStreamSelectionAttributes& InitialSelection) override;
 	virtual void EnableFrameAccurateSeeking(bool bEnabled) override;
@@ -1745,6 +1764,7 @@ private:
 
 	double GetMinBufferTimeBeforePlayback();
 	bool HaveEnoughBufferedDataToStartPlayback();
+	void PrepareForPrerolling();
 
 	FTimeValue ClampTimeToCurrentRange(const FTimeValue& InTime, bool bClampToStart, bool bClampToEnd);
 
diff --git a/Engine/Plugins/Media/ElectraPlayer/Source/ElectraPlayerRuntime/Private/Runtime/Player/DASH/ManifestParserDASH.cpp b/Engine/Plugins/Media/ElectraPlayer/Source/ElectraPlayerRuntime/Private/Runtime/Player/DASH/ManifestParserDASH.cpp
--- a/Engine/Plugins/Media/ElectraPlayer/Source/ElectraPlayerRuntime/Private/Runtime/Player/DASH/ManifestParserDASH.cpp
+++ b/Engine/Plugins/Media/ElectraPlayer/Source/ElectraPlayerRuntime/Private/Runtime/Player/DASH/ManifestParserDASH.cpp
@@ -504,14 +504,21 @@ FErrorDetail FManifestParserDASH::BuildFromMPD(TCHAR* InOutMPDXML, const TCHAR*
 	int32 ErrorLine = 0;
 	const TCHAR* NoFile = TEXT("");
 	ExpectedRootElement = InExpectedRootElement;
-	if (!FFastXml::ParseXmlFile(this, NoFile, InOutMPDXML, nullptr, false, false, ErrorMsg, ErrorLine))
+	if (InOutMPDXML)
 	{
-		// Check if we have already set an error or if this is an error generated by the XML parser.
-		if (!LastErrorDetail.IsSet())
+		if (!FFastXml::ParseXmlFile(this, NoFile, InOutMPDXML, nullptr, false, false, ErrorMsg, ErrorLine))
 		{
-			LastErrorDetail = CreateErrorAndLog(FString::Printf(TEXT("XML parser failed with '%s' on line %d"), *(ErrorMsg.ToString()), ErrorLine), ERRCODE_DASH_MPD_PARSER_XML_PARSER_FAILED);
+			// Check if we have already set an error or if this is an error generated by the XML parser.
+			if (!LastErrorDetail.IsSet())
+			{
+				LastErrorDetail = CreateErrorAndLog(FString::Printf(TEXT("XML parser failed with '%s' on line %d"), *(ErrorMsg.ToString()), ErrorLine), ERRCODE_DASH_MPD_PARSER_XML_PARSER_FAILED);
+			}
 		}
 	}
+	else
+	{
+		LastErrorDetail = CreateErrorAndLog(FString::Printf(TEXT("No XML was given to parse")), ERRCODE_DASH_MPD_PARSER_XML_PARSER_FAILED);
+	}
 	return LastErrorDetail;
 }
 
diff --git a/Engine/Plugins/Media/ElectraPlayer/Source/ElectraPlayerRuntime/Private/Runtime/Player/HLS/ManifestHLS.cpp b/Engine/Plugins/Media/ElectraPlayer/Source/ElectraPlayerRuntime/Private/Runtime/Player/HLS/ManifestHLS.cpp
--- a/Engine/Plugins/Media/ElectraPlayer/Source/ElectraPlayerRuntime/Private/Runtime/Player/HLS/ManifestHLS.cpp
+++ b/Engine/Plugins/Media/ElectraPlayer/Source/ElectraPlayerRuntime/Private/Runtime/Player/HLS/ManifestHLS.cpp
@@ -738,6 +738,7 @@ IManifest::FResult FPlayPeriodHLS::FindSegment(TSharedPtrTS<FStreamSegmentReques
 			// If there is a playback range and we have reached it we are at EOS, regardless of whether this is a Live stream or not.
 			if (SearchParam.LastPTS.IsValid() && SegmentList[SelectedSegmentIndex].AbsoluteDateTime >= SearchParam.LastPTS)
 			{
+				Req->AbsoluteDateTime = SegmentList[SelectedSegmentIndex].AbsoluteDateTime;
 				Req->bIsEOSSegment = true;
 				OutRequest = Req;
 				return IManifest::FResult(IManifest::FResult::EType::PastEOS);
@@ -775,6 +776,7 @@ IManifest::FResult FPlayPeriodHLS::FindSegment(TSharedPtrTS<FStreamSegmentReques
 			// Unless this is a VOD list or it has as ENDLIST tag we have to try this later, assuming that an updated playlist has added additional segments.
 			if (InStream->PlaylistType == FManifestHLSInternal::FMediaStream::EPlaylistType::VOD || InStream->bHasListEnd)
 			{
+				Req->AbsoluteDateTime = SearchParam.Time;
 				Req->bIsEOSSegment = true;
 				OutRequest = Req;
 				return IManifest::FResult(IManifest::FResult::EType::PastEOS);
@@ -789,6 +791,7 @@ IManifest::FResult FPlayPeriodHLS::FindSegment(TSharedPtrTS<FStreamSegmentReques
 	else
 	{
 		// No segments is not really expected. If this occurs we assume the presentation has ended.
+		Req->AbsoluteDateTime = SearchParam.Time;
 		Req->bIsEOSSegment = true;
 		OutRequest = Req;
 		return IManifest::FResult(IManifest::FResult::EType::PastEOS);
diff --git a/Engine/Plugins/Media/ElectraPlayer/Source/ElectraPlayerRuntime/Public/IElectraPlayerInterface.h b/Engine/Plugins/Media/ElectraPlayer/Source/ElectraPlayerRuntime/Public/IElectraPlayerInterface.h
--- a/Engine/Plugins/Media/ElectraPlayer/Source/ElectraPlayerRuntime/Public/IElectraPlayerInterface.h
+++ b/Engine/Plugins/Media/ElectraPlayer/Source/ElectraPlayerRuntime/Public/IElectraPlayerInterface.h
@@ -173,8 +173,19 @@ public:
 	virtual float GetRate() const = 0;
 	virtual bool SetRate(float Rate) = 0;
 
+
+	struct FSeekParam
+	{
+		TOptional<int32> StartingBitrate;
+		TOptional<bool> bOptimizeForScrubbing;
+		TOptional<double> DistanceThreshold;
+	};
+
 	virtual bool Seek(const FTimespan& Time) = 0;
+	virtual bool Seek(const FTimespan& Time, const FSeekParam& Param) = 0;
 	virtual void SetFrameAccurateSeekMode(bool bEnableFrameAccuracy) = 0;
+	
+	virtual void ModifyOptions(const Electra::FParamDict& InOptionsToSetOrChange, const Electra::FParamDict& InOptionsToClear) = 0;
 
 	struct FAudioTrackFormat
 	{
diff --git a/Engine/Plugins/Media/ElectraUtil/Source/ElectraBase/Private/ParameterDictionary.cpp b/Engine/Plugins/Media/ElectraUtil/Source/ElectraBase/Private/ParameterDictionary.cpp
--- a/Engine/Plugins/Media/ElectraUtil/Source/ElectraBase/Private/ParameterDictionary.cpp
+++ b/Engine/Plugins/Media/ElectraUtil/Source/ElectraBase/Private/ParameterDictionary.cpp
@@ -364,40 +364,75 @@ FParamDict::FParamDict()
 {
 }
 
+FParamDict::FParamDict(const FParamDict& Other)
+{
+	InternalCopy(Other);
+}
+
+FParamDict& FParamDict::operator=(const FParamDict& Other)
+{
+	if (&Other != this)
+	{
+		FScopeLock lock(&Lock);
+		InternalCopy(Other);
+	}
+	return *this;
+}
+
 FParamDict::~FParamDict()
 {
 }
 
+void FParamDict::InternalCopy(const FParamDict& Other)
+{
+	FScopeLock lock(&Other.Lock);
+	Dictionary = Other.Dictionary;
+}
+
+
+
 void FParamDict::Clear()
 {
+	FScopeLock lock(&Lock);
 	Dictionary.Empty();
 }
 
 void FParamDict::Set(const FString& Key, const FVariantValue& Value)
 {
+	FScopeLock lock(&Lock);
 	Dictionary.FindOrAdd(Key, Value);
 }
 
 bool FParamDict::HaveKey(const FString& Key) const
 {
+	FScopeLock lock(&Lock);
 	return Dictionary.Find(Key) != nullptr;
 }
 
-const FVariantValue& FParamDict::GetValue(const FString& Key) const
+FVariantValue FParamDict::GetValue(const FString& Key) const
 {
 	static FVariantValue Empty;
+	FScopeLock lock(&Lock);
 	const FVariantValue* VariantValue = Dictionary.Find(Key);
 	return VariantValue ? *VariantValue : Empty;
 }
 
-void FParamDict::GetKeysStartingWith(const FString& StartsWith, TArray<FString>& Keys)
+void FParamDict::GetKeysStartingWith(const FString& StartsWith, TArray<FString>& Keys) const
 {
+	FScopeLock lock(&Lock);
 	Keys.Empty();
-	for (const TPair<FString, FVariantValue>& Pair : Dictionary)
+	if (StartsWith.IsEmpty())
+	{
+		Dictionary.GenerateKeyArray(Keys);
+	}
+	else
 	{
-		if (Pair.Key.StartsWith(StartsWith, ESearchCase::CaseSensitive))
+		for(const TPair<FString, FVariantValue>& Pair : Dictionary)
 		{
-			Keys.Emplace(Pair.Key);
+			if (Pair.Key.StartsWith(StartsWith, ESearchCase::CaseSensitive))
+			{
+				Keys.Emplace(Pair.Key);
+			}
 		}
 	}
 }
diff --git a/Engine/Plugins/Media/ElectraUtil/Source/ElectraBase/Public/ParameterDictionary.h b/Engine/Plugins/Media/ElectraUtil/Source/ElectraBase/Public/ParameterDictionary.h
--- a/Engine/Plugins/Media/ElectraUtil/Source/ElectraBase/Public/ParameterDictionary.h
+++ b/Engine/Plugins/Media/ElectraUtil/Source/ElectraBase/Public/ParameterDictionary.h
@@ -3,6 +3,7 @@
 #pragma once
 
 #include "CoreMinimal.h"
+#include "Misc/ScopeLock.h"
 #include "PlayerTime.h"
 #include "Templates/TypeCompatibleBytes.h"
 #include "Containers/Array.h"
@@ -150,19 +151,29 @@ class ELECTRABASE_API FParamDict
 {
 public:
 	FParamDict();
+	FParamDict(const FParamDict& Other);
+	FParamDict& operator=(const FParamDict& Other);
 	~FParamDict();
 	void Clear();
 	void Set(const FString& Key, const FVariantValue& Value);
 	bool HaveKey(const FString& Key) const;
-	const FVariantValue& GetValue(const FString& Key) const;
+	FVariantValue GetValue(const FString& Key) const;
 	void Remove(const FString& Key)
-	{ Dictionary.Remove(Key); }
+	{
+		FScopeLock lock(&Lock);
+		Dictionary.Remove(Key); 
+	}
 	void SetOrUpdate(const FString& Key, const FVariantValue& Value)
-	{ Dictionary.Add(Key, Value); }
+	{ 
+		FScopeLock lock(&Lock);
+		Dictionary.Add(Key, Value); 
+	}
 
-	void GetKeysStartingWith(const FString& StartsWith, TArray<FString>& Keys);
+	void GetKeysStartingWith(const FString& StartsWith, TArray<FString>& Keys) const;
 private:
-	TMap<FString, FVariantValue>		Dictionary;
+	void InternalCopy(const FParamDict& Other);
+	mutable FCriticalSection Lock;
+	TMap<FString, FVariantValue> Dictionary;
 };
 
 
