ElectraPlayer: Fixed a crash when the audio decoder failed to create on PC but was accidentally being used.

#rnx

diff --git a/Engine/Plugins/Media/ElectraPlayer/Source/ElectraPlayerRuntime/Private/Runtime/Decoder/DX/AudioDecoderAAC_DX.cpp b/Engine/Plugins/Media/ElectraPlayer/Source/ElectraPlayerRuntime/Private/Runtime/Decoder/DX/AudioDecoderAAC_DX.cpp
--- a/Engine/Plugins/Media/ElectraPlayer/Source/ElectraPlayerRuntime/Private/Runtime/Decoder/DX/AudioDecoderAAC_DX.cpp
+++ b/Engine/Plugins/Media/ElectraPlayer/Source/ElectraPlayerRuntime/Private/Runtime/Decoder/DX/AudioDecoderAAC_DX.cpp
@@ -1256,7 +1256,7 @@ namespace Electra
 						bHaveDiscontinuity = true;
 					}
 					SequenceIndex = CurrentAccessUnit->AccessUnit->PTS.GetSequenceIndex();
-					if (!Decode(CurrentAccessUnit, false))
+					if (!bError && !Decode(CurrentAccessUnit, false))
 					{
 						bError = true;
 					}
