ElectraHTTPStream: Added a network module geared towards downloading of chunked media streams for use with the Electra media player.

#rnx

diff --git a/Engine/Plugins/Media/ElectraPlayer/Source/ElectraPlayerRuntime/ElectraPlayerRuntime.Build.cs b/Engine/Plugins/Media/ElectraPlayer/Source/ElectraPlayerRuntime/ElectraPlayerRuntime.Build.cs
--- a/Engine/Plugins/Media/ElectraPlayer/Source/ElectraPlayerRuntime/ElectraPlayerRuntime.Build.cs
+++ b/Engine/Plugins/Media/ElectraPlayer/Source/ElectraPlayerRuntime/ElectraPlayerRuntime.Build.cs
@@ -18,10 +18,9 @@ namespace UnrealBuildTool.Rules
 			PrivateDependencyModuleNames.AddRange(
 				new string[] {
 					"Core",
-					"HTTP",
-					"WebSockets",
 					"Json",
 					"ElectraBase",
+					"ElectraHTTPStream",
 					"ElectraCDM",
 					"ElectraSubtitles",
 					"XmlParser"
diff --git a/Engine/Plugins/Media/ElectraPlayer/Source/ElectraPlayerRuntime/Private/Runtime/HTTP/HTTPManager.cpp b/Engine/Plugins/Media/ElectraPlayer/Source/ElectraPlayerRuntime/Private/Runtime/HTTP/HTTPManager.cpp
--- a/Engine/Plugins/Media/ElectraPlayer/Source/ElectraPlayerRuntime/Private/Runtime/HTTP/HTTPManager.cpp
+++ b/Engine/Plugins/Media/ElectraPlayer/Source/ElectraPlayerRuntime/Private/Runtime/HTTP/HTTPManager.cpp
@@ -17,14 +17,9 @@
 #include "Serialization/Archive.h"
 #include "Serialization/ArrayReader.h"
 
-// For https:// scheme
-#include "Modules/ModuleManager.h"
-#include "HttpModule.h"
-#include "PlatformHttp.h"
-#include "Interfaces/IHttpResponse.h"
-
-
 #include "ElectraPlayerPrivate.h"
+#include "ElectraHTTPStream.h"
+
 
 #define ERRCODE_OK								0
 #define ERRCODE_HTTP_CONNECTION_TIMEOUT			1
@@ -40,10 +35,14 @@ DECLARE_CYCLE_STAT(TEXT("FElectraHttpManager::WorkerThread"), STAT_ElectraPlayer
 DEFINE_LOG_CATEGORY(LogElectraHTTPManager);
 
 
+#define ELECTRA_HTTPMANAGER_USER_AGENT					TEXT("X-UnrealEngine-Agent")
+//#define ELECTRA_HTTPMANAGER_DEBUG_PROXY_ADDRESS			TEXT("10.29.4.30:8888")
+//#define ELECTRA_HTTPMANAGER_ALLOW_UNSAFE_CONNECTIONS	1
+
 namespace Electra
 {
 
-	class FElectraHttpManager : public IElectraHttpManager, public FMediaThread, public TSharedFromThis<FElectraHttpManager, ESPMode::ThreadSafe>
+	class FElectraHttpManager : public IElectraHttpManager, public TSharedFromThis<FElectraHttpManager, ESPMode::ThreadSafe>
 	{
 	public:
 		static TSharedPtrTS<FElectraHttpManager> Create();
@@ -95,31 +94,14 @@ namespace Electra
 			{
 				Unbind();
 			}
-			void ReportRequestComplete(FHttpRequestPtr Request, FHttpResponsePtr Response, bool bConnectedSuccessfully)
-			{
-				FScopeLock lock(&Lock);
-				TSharedPtrTS<FElectraHttpManager> That = Owner.Pin();
-				if (That.IsValid())
-				{
-					That->OnProcessRequestCompleteInternal(Request, Response, bConnectedSuccessfully);
-				}
-			}
-			void ReportRequestProgress(FHttpRequestPtr Request, int32 BytesSent, int32 BytesReceived)
-			{
-				FScopeLock lock(&Lock);
-				TSharedPtrTS<FElectraHttpManager> That = Owner.Pin();
-				if (That.IsValid())
-				{
-					That->OnRequestProgressInternal(Request, BytesSent, BytesReceived);
-				}
-			}
-			void ReportRequestHeaderReceived(FHttpRequestPtr Request, const FString& HeaderName, const FString& NewHeaderValue)
+
+			void ReportRequestNotification(IElectraHTTPStreamRequestPtr InRequest, EElectraHTTPStreamNotificationReason InReason, int64 InParam)
 			{
 				FScopeLock lock(&Lock);
 				TSharedPtrTS<FElectraHttpManager> That = Owner.Pin();
 				if (That.IsValid())
 				{
-					That->OnRequestHeaderReceivedInternal(Request, HeaderName, NewHeaderValue);
+					That->OnRequestNotification(InRequest, InReason, InParam);
 				}
 			}
 
@@ -154,15 +136,16 @@ namespace Electra
 						HttpsRequestCallbackWrapper->Unbind();
 					}
 					HttpsRequestCallbackWrapper.Reset();
-					if (HttpRequest.IsValid() && !ActiveResponse.bHitCache && !EHttpRequestStatus::IsFinished(HttpRequest->GetStatus()))
+
+					if (HttpRequest.IsValid() && !ActiveResponse.bHitCache && HttpRequest->GetResponse()->GetStatus() != IElectraHTTPStreamResponse::EStatus::Completed)
 					{
-						HttpRequest->CancelRequest();
+						HttpRequest->Cancel();
 					}
 					HttpRequest.Reset();
 				}
 			}
 
-			bool ProcessHTTPRequest()
+			bool NeedsFetchOfUncachedParts()
 			{
 				check(HandleType == EHandleType::HTTPHandle);
 
@@ -178,7 +161,8 @@ namespace Electra
 						if (CacheItem.IsValid())
 						{
 							ActiveResponse.CacheResponse = MoveTemp(CacheItem);
-							return true;
+							// Have all data, no need to fetch anything.
+							return false;
 						}
 					}
 					else if (Result == IHTTPResponseCache::EScatterResult::PartialHit)
@@ -195,19 +179,21 @@ namespace Electra
 							{
 								ActiveResponse.CacheResponse = MoveTemp(CacheItem);
 								ActiveResponse.ReceivedContentRange = ActiveResponse.CacheResponse->Range;
-								return true;
+								// Don't need to fetch the first part.
+								return false;
 							}
 							else
 							{
 								// The missing range we need to fetch has been set up by the cache so we just need to use
 								// that range in the request.
 								ActiveResponse.Range = CacheItem->Range;
-								HttpRequest->SetHeader(TEXT("Range"), FString::Printf(TEXT("bytes=%s"), *ActiveResponse.Range.GetString()));
+								HttpRequest->SetRange(ActiveResponse.Range.GetString());
 							}
 						}
 					}
 				}
-				return HttpRequest->ProcessRequest();
+				// Fetch needed.
+				return true;
 			}
 
 			struct FRequestResponse
@@ -215,14 +201,14 @@ namespace Electra
 				// URL
 				FString URL;
 				// Current range requested from server.
-				FParams::FRange Range;
+				ElectraHTTPStream::FHttpRange Range;
 				// Current response received from server
-				FHttpResponsePtr Response;
+				IElectraHTTPStreamResponsePtr Response;
 				// Number of bytes passed along already.
 				int64 NumBytesPassedOut = 0;
 				// Original range requested for easier access and comparison.
-				FParams::FRange OriginalRange;
-				FParams::FRange ReceivedContentRange;
+				ElectraHTTPStream::FHttpRange OriginalRange;
+				ElectraHTTPStream::FHttpRange ReceivedContentRange;
 				//
 				bool bIsSubRangeRequest = false;
 				int32 NumSubRangeRequest = 0;
@@ -231,7 +217,6 @@ namespace Electra
 				bool bHitCache = false;
 				bool bWasAddedToCache = false;
 
-
 				//
 				int64 SizeRemaining() const
 				{
@@ -281,11 +266,11 @@ namespace Electra
 			// Local file handle (for file:// and data:)
 			FLocalByteStream*		LocalByteStream = nullptr;
 
-			// HTTP module handle (for http:// and https://)
-			FHttpRequestPtr			HttpRequest;
+			// HTTP handle
+			IElectraHTTPStreamRequestPtr							HttpRequest;
 			TSharedPtr<FHTTPCallbackWrapper, ESPMode::ThreadSafe>	HttpsRequestCallbackWrapper;
-			bool					bHttpRequestFirstEvent = true;
-			TSharedPtrTS<IHTTPResponseCache> HttpResponseCache;
+			bool													bHttpRequestFirstEvent = true;
+			TSharedPtrTS<IHTTPResponseCache>						HttpResponseCache;
 
 			FTimeValue				LastTimeDataReceived;
 			FTimeValue				TimeAtNextProgressCallback;
@@ -343,8 +328,8 @@ namespace Electra
 		};
 
 		FHandle* CreateLocalFileHandle(const FTimeValue& Now, FTransportError& OutError, const TSharedPtrTS<FRequest>& Request);
-		FHandle* CreateHTTPModuleHandle(const FTimeValue& Now, FTransportError& OutError, const TSharedPtrTS<FRequest>& Request);
-		bool PrepareHTTPModuleHandle(const FTimeValue& Now, FHandle* InHandle, const TSharedPtrTS<FRequest>& Request, bool bIsFirstSetup);
+		FHandle* CreateHTTPHandle(const FTimeValue& Now, FTransportError& OutError, const TSharedPtrTS<FRequest>& Request);
+		bool PrepareHTTPHandle(const FTimeValue& Now, FHandle* InHandle, const TSharedPtrTS<FRequest>& Request, bool bIsFirstSetup);
 		void AddPendingRequests(const FTimeValue& Now);
 		void RemovePendingRequests(const FTimeValue& Now);
 		void HandleCompletedRequests();
@@ -356,25 +341,14 @@ namespace Electra
 		void HandleHTTPResponses(const FTimeValue& Now);
 		void RemoveAllRequests();
 
-		void WorkerThread();
+		void StartHTTPManager();
+		void StopHTTPManager();
+		void ProcessHTTPManager();
 
 	public:
-		// Callbacks from the HTTP module
-		void OnProcessRequestCompleteInternal(FHttpRequestPtr Request, FHttpResponsePtr Response, bool bConnectedSuccessfully);
-		void OnRequestProgressInternal(FHttpRequestPtr Request, int32 BytesSent, int32 BytesReceived);
-		void OnRequestHeaderReceivedInternal(FHttpRequestPtr Request, const FString& HeaderName, const FString& NewHeaderValue);
-	private:
-		struct FHttpEventData
-		{
-			TArray<HTTP::FHTTPHeader>	Headers;
-			FHttpResponsePtr			Response;
-			EHttpRequestStatus::Type	Status = EHttpRequestStatus::Type::NotStarted;
-			int32						BytesSent = 0;
-			int32						BytesReceived = 0;
-			bool						bComplete = false;
-			bool						bConnectedSuccessfully = false;
-		};
+		void OnRequestNotification(IElectraHTTPStreamRequestPtr InRequest, EElectraHTTPStreamNotificationReason InReason, int64 InParam);
 
+	private:
 		struct FRemoveRequest
 		{
 			void SignalDone()
@@ -396,15 +370,16 @@ namespace Electra
 		TQueue<FRemoveRequest, EQueueMode::Mpsc>						RequestsToRemove;
 		TQueue<TSharedPtrTS<FRequest>, EQueueMode::Mpsc>				RequestsCompleted;
 
+		TSharedPtr<IElectraHTTPStream, ESPMode::ThreadSafe>				HTTPStreamHandler;
+		FCriticalSection												HTTPStreamEventLock;
+		TArray<IElectraHTTPStreamRequestPtr>							NotifiedRequests;
+
 		TMap<FHandle*, TSharedPtrTS<FRequest>>							ActiveRequests;
 		FTimeValue														ProgressInterval;
 
-		bool															bThreadStarted = false;
+		bool															bManagerStarted = false;
 		bool															bTerminate = false;
 
-		FCriticalSection												HttpModuleEventLock;
-		TMap<FHttpRequestPtr, FHttpEventData>							HttpModuleEventMap;
-
 		static TWeakPtrTS<FElectraHttpManager>							SingletonSelf;
 		static FCriticalSection											SingletonLock;
 	};
@@ -430,9 +405,9 @@ namespace Electra
 			FElectraHttpManager* Manager = new FElectraHttpManager;
 			if (Manager)
 			{
+				Self = MakeShareable(Manager);
 				Manager->Initialize();
 			}
-			Self = MakeShareable(Manager);
 			SingletonSelf = Self;
 		}
 		return Self;
@@ -442,21 +417,17 @@ namespace Electra
 	FElectraHttpManager::~FElectraHttpManager()
 	{
 		bTerminate = true;
-		if (bThreadStarted)
+		if (bManagerStarted)
 		{
-			ThreadWaitDone();
-			ThreadReset();
+			StopHTTPManager();
 		}
 	}
 
 	void FElectraHttpManager::Initialize()
 	{
 		ProgressInterval.SetFromMilliseconds(100);
-
-		ThreadSetName("ElectraPlayer::HTTPManager");
-		ThreadStart(Electra::MakeDelegate(this, &FElectraHttpManager::WorkerThread));
-
-		bThreadStarted = true;
+		StartHTTPManager();
+		bManagerStarted = true;
 	}
 
 	void FElectraHttpManager::AddRequest(TSharedPtrTS<FRequest> Request, bool bAutoRemoveWhenComplete)
@@ -539,7 +510,7 @@ namespace Electra
 		return Handle.Release();
 	}
 
-	FElectraHttpManager::FHandle* FElectraHttpManager::CreateHTTPModuleHandle(const FTimeValue& Now, FTransportError& OutError, const TSharedPtrTS<IElectraHttpManager::FRequest>& Request)
+	FElectraHttpManager::FHandle* FElectraHttpManager::CreateHTTPHandle(const FTimeValue& Now, FTransportError& OutError, const TSharedPtrTS<IElectraHttpManager::FRequest>& Request)
 	{
 		TUniquePtr<FHandle> Handle(new FHandle);
 		Handle->Owner = this;
@@ -548,50 +519,22 @@ namespace Electra
 		Handle->ActiveResponse.NumSubRangeRequest = 0;
 		Handle->ActiveResponse.OriginalRange = Request->Parameters.Range;
 
-		PrepareHTTPModuleHandle(Now, Handle.Get(), Request, true);
+		PrepareHTTPHandle(Now, Handle.Get(), Request, true);
 		return Handle.Release();
 	}
 
 
-	bool FElectraHttpManager::PrepareHTTPModuleHandle(const FTimeValue& Now, FHandle* Handle, const TSharedPtrTS<FRequest>& Request, bool bIsFirstSetup)
+	bool FElectraHttpManager::PrepareHTTPHandle(const FTimeValue& Now, FHandle* Handle, const TSharedPtrTS<FRequest>& Request, bool bIsFirstSetup)
 	{
-		FHttpModule& HttpModule = FModuleManager::LoadModuleChecked<FHttpModule>("HTTP");
-
 		Handle->HttpsRequestCallbackWrapper = MakeShared<FHTTPCallbackWrapper, ESPMode::ThreadSafe>(AsShared());
-		Handle->HttpRequest = HttpModule.Get().CreateRequest();
-		Handle->HttpRequest->OnProcessRequestComplete().BindThreadSafeSP(Handle->HttpsRequestCallbackWrapper.ToSharedRef(), &FHTTPCallbackWrapper::ReportRequestComplete);
-		Handle->HttpRequest->OnRequestProgress().BindThreadSafeSP(Handle->HttpsRequestCallbackWrapper.ToSharedRef(), &FHTTPCallbackWrapper::ReportRequestProgress);
-		Handle->HttpRequest->OnHeaderReceived().BindThreadSafeSP(Handle->HttpsRequestCallbackWrapper.ToSharedRef(), &FHTTPCallbackWrapper::ReportRequestHeaderReceived);
-		Handle->HttpRequest->SetURL(Request->Parameters.URL);
 		Handle->ActiveResponse.URL = Request->Parameters.URL;
-		if (!Request->Parameters.Verb.IsEmpty())
-		{
-			Handle->HttpRequest->SetVerb(Request->Parameters.Verb);
-			// Add POST data
-			if (Request->Parameters.Verb.Equals(TEXT("POST")))
-			{
-				Handle->HttpRequest->SetContent(MoveTemp(Request->Parameters.PostData));
-			}
-		}
-		else
-		{
-			Handle->HttpRequest->SetVerb(TEXT("GET"));
-			Handle->HttpResponseCache = Request->ResponseCache;
-		}
-
-		Handle->HttpRequest->SetHeader(TEXT("User-Agent"), "X-UnrealEngine-Agent");
-		// Set accepted encoding first. If this is also present in custom headers let those overwrite it.
-		if (Request->Parameters.AcceptEncoding.IsSet())
-		{
-			Handle->HttpRequest->SetHeader(TEXT("Accept-Encoding"), Request->Parameters.AcceptEncoding.Value());
-		}
-		for(int32 i=0; i<Request->Parameters.RequestHeaders.Num(); ++i)
+		Handle->RequestStartTime = Now;
+		Handle->TimeAtConnectionTimeoutCheck.SetToPositiveInfinity();
+		if (Request->Parameters.ConnectTimeout.IsValid())
 		{
-			Handle->HttpRequest->SetHeader(Request->Parameters.RequestHeaders[i].Header, Request->Parameters.RequestHeaders[i].Value);
+			Handle->TimeAtConnectionTimeoutCheck = Now + Request->Parameters.ConnectTimeout;
 		}
 
-		Handle->RequestStartTime = Now;
-		Handle->TimeAtConnectionTimeoutCheck.SetToPositiveInfinity();
 		// Is this the first sub range request or a continuation?
 		if (bIsFirstSetup)
 		{
@@ -623,6 +566,7 @@ namespace Electra
 		}
 		Handle->ActiveResponse.NumBytesPassedOut = 0;
 		Handle->ActiveResponse.ReceivedContentRange.Reset();
+
 		Handle->ActiveResponse.Response.Reset();
 
 		// Response not yet received.
@@ -630,50 +574,64 @@ namespace Electra
 
 		// Do a quick check that if the request is reading into a ring buffer that there is also a sub-range request size set up.
 		TSharedPtrTS<FReceiveBuffer> ReceiveBuffer = Request->ReceiveBuffer.Pin();
-		if (ReceiveBuffer.IsValid() && ReceiveBuffer->bEnableRingbuffer)
+		// Give a warning to the console about this circumstance if the original request is not ranged.
+		// If it is a range request that goes up to the end of the file we cannot say with certainty if this is a problem
+		// since we do not know how large the response is going to be. It could be a few bytes, it could be gigabytes...
+		if (ReceiveBuffer.IsValid() && ReceiveBuffer->bEnableRingbuffer && !Handle->ActiveResponse.Range.IsSet())
 		{
-			if (Request->Parameters.SubRangeRequestSize > 0)
+			UE_LOG(LogElectraHTTPManager, Warning, TEXT("Receive buffer set to ring buffer mode but no sub range request size specified!"));
+		}
+
+		// This could be for the next sub-range request that we also want to add to the cache!
+		Handle->ActiveResponse.bWasAddedToCache = false;
+
+		if (HTTPStreamHandler.IsValid())
+		{
+			Handle->HttpRequest = HTTPStreamHandler->CreateRequest();
+			if (!Handle->HttpRequest.IsValid())
+			{
+				return false;
+			}
+			Handle->HttpRequest->SetVerb(TEXT("GET"));
+			Handle->HttpResponseCache = Request->ResponseCache;
+
+			if (!Request->Parameters.Verb.IsEmpty())
 			{
-				// The requested sub range size may be smaller than the range we are supposed to get.
-				// Adjust the first range request accordingly.
-				if (Handle->ActiveResponse.Range.IsSet())
+				Handle->HttpRequest->SetVerb(Request->Parameters.Verb);
+				// Disable the cache for anything but GET
+				if (!Request->Parameters.Verb.Equals(TEXT("GET")))
 				{
-					// If the original range is larger than the sub range request, including an open ended range, we need to adjust the range.
-					int64 NumOrgBytes = Handle->ActiveResponse.Range.GetNumberOfBytes();
-					if (NumOrgBytes > Request->Parameters.SubRangeRequestSize || Handle->ActiveResponse.Range.GetEndIncluding() < 0)
-					{
-						Handle->ActiveResponse.Range.SetEndIncluding(Handle->ActiveResponse.Range.GetStart() + Request->Parameters.SubRangeRequestSize - 1);
-						Handle->ActiveResponse.bIsSubRangeRequest = true;
-					}
+					Handle->HttpResponseCache.Reset();
 				}
-				else
+
+				// Add POST data
+				if (Request->Parameters.Verb.Equals(TEXT("POST")))
 				{
-					// No range requested. We still need to perform sub range requests.
-					Handle->ActiveResponse.Range.SetStart(0);
-					Handle->ActiveResponse.Range.SetEndIncluding(Request->Parameters.SubRangeRequestSize - 1);
-					Handle->ActiveResponse.bIsSubRangeRequest = true;
+					IElectraHTTPStreamBuffer& pdb = Handle->HttpRequest->POSTDataBuffer();
+					pdb.AddData(MoveTemp(Request->Parameters.PostData));
+					pdb.SetEOS();
 				}
 			}
-			else
+
+			Handle->HttpRequest->SetURL(Request->Parameters.URL);
+			Handle->HttpRequest->SetUserAgent(ELECTRA_HTTPMANAGER_USER_AGENT);
+			Handle->HttpRequest->AllowCompression(!Request->Parameters.AcceptEncoding.GetWithDefault(TEXT("")).Equals(TEXT("identity")));
+			#if defined(ELECTRA_HTTPMANAGER_ALLOW_UNSAFE_CONNECTIONS) && ELECTRA_HTTPMANAGER_ALLOW_UNSAFE_CONNECTIONS
+				Handle->HttpRequest->AllowUnsafeRequestsForDebugging();
+			#endif
+			if (Handle->ActiveResponse.Range.IsSet())
 			{
-				// Give a warning to the console about this circumstance if the original request is not ranged.
-				// If it is a range request that goes up to the end of the file we cannot say with certainty if this is a problem
-				// since we do not know how large the response is going to be. It could be a few bytes, it could be gigabytes...
-				if (!Handle->ActiveResponse.Range.IsSet())
-				{
-					UE_LOG(LogElectraHTTPManager, Warning, TEXT("Receive buffer set to ring buffer mode but no sub range request size specified!"));
-				}
+				Handle->HttpRequest->SetRange(Handle->ActiveResponse.Range.GetString());
 			}
-		}
-		// Set the possibly adjusted range.
-		if (Handle->ActiveResponse.Range.IsSet())
-		{
-			FString Range = FString::Printf(TEXT("bytes=%s"), *Handle->ActiveResponse.Range.GetString());
-			Handle->HttpRequest->SetHeader(TEXT("Range"), Range);
+
+			for(int32 i=0; i<Request->Parameters.RequestHeaders.Num(); ++i)
+			{
+				Handle->HttpRequest->AddHeader(Request->Parameters.RequestHeaders[i].Header, Request->Parameters.RequestHeaders[i].Value, false);
+			}
+
+			Handle->HttpRequest->NotificationDelegate().BindThreadSafeSP(Handle->HttpsRequestCallbackWrapper.ToSharedRef(), &FHTTPCallbackWrapper::ReportRequestNotification);
 		}
 
-		// This could be for the next sub-range request that we also want to add to the cache!
-		Handle->ActiveResponse.bWasAddedToCache = false;
 		return true;
 	}
 
@@ -710,14 +668,21 @@ namespace Electra
 			}
 			else
 			{
-				FHandle* Handle = CreateHTTPModuleHandle(Now, HttpError, Request);
+				FHandle* Handle = CreateHTTPHandle(Now, HttpError, Request);
 				if (Handle)
 				{
 					ActiveRequests.Add(Handle, Request);
-					if (!Handle->ProcessHTTPRequest())
+					if (Handle->NeedsFetchOfUncachedParts())
 					{
-						Request->ConnectionInfo.StatusInfo.ErrorDetail.SetError(UEMEDIA_ERROR_INTERNAL).SetFacility(Facility::EFacility::HTTPReader).SetMessage("HTTP request failed on ProcessRequest()");
-						RequestsCompleted.Enqueue(Request);
+						if (HTTPStreamHandler.IsValid())
+						{
+							HTTPStreamHandler->AddRequest(Handle->HttpRequest);
+						}
+						else
+						{
+							Request->ConnectionInfo.StatusInfo.ErrorDetail.SetError(UEMEDIA_ERROR_INTERNAL).SetFacility(Facility::EFacility::HTTPReader).SetMessage("HTTP request failed on ProcessRequest()");
+							RequestsCompleted.Enqueue(Request);
+						}
 					}
 				}
 				else
@@ -874,7 +839,7 @@ namespace Electra
 				continue;
 			}
 			// HTTP transfers that are cached or not processing also do not need to be checked.
-			else if (Handle->HandleType == FHandle::EHandleType::HTTPHandle && (Handle->ActiveResponse.CacheResponse.IsValid() || (Handle->HttpRequest.IsValid() && EHttpRequestStatus::IsFinished(Handle->HttpRequest->GetStatus()))))
+			else if (Handle->HandleType == FHandle::EHandleType::HTTPHandle && (Handle->ActiveResponse.CacheResponse.IsValid() || (Handle->HttpRequest.IsValid() && Handle->HttpRequest->GetResponse()->GetState() == IElectraHTTPStreamResponse::EState::Finished)))
 			{
 				continue;
 			}
@@ -952,59 +917,70 @@ namespace Electra
 	}
 
 
-	void FElectraHttpManager::WorkerThread()
+	void FElectraHttpManager::StartHTTPManager()
 	{
 		LLM_SCOPE(ELLMTag::ElectraPlayer);
 
-		FTimeValue Now = MEDIAutcTime::Current();
-		while(!bTerminate)
-		{
-			const long kSelectTimeoutMilliseconds = 15;
-
-			{
-				SCOPE_CYCLE_COUNTER(STAT_ElectraPlayer_FElectraHttpManager_Worker);
-				CSV_SCOPED_TIMING_STAT(ElectraPlayer, ElectraHttpManager_Worker);
-
-				// Add and remove pending requests
-				Lock.Lock();
-				AddPendingRequests(Now);
-				RemovePendingRequests(Now);
-				Lock.Unlock();
-			}
-
-			// Throttle ourselves here a bit.
-			RequestChangesEvent.WaitTimeoutAndReset(kSelectTimeoutMilliseconds * 1000);
-
-			{
-				SCOPE_CYCLE_COUNTER(STAT_ElectraPlayer_FElectraHttpManager_Worker);
-				CSV_SCOPED_TIMING_STAT(ElectraPlayer, ElectraHttpManager_Worker);
+		FParamDict HttpOptions;
+		#ifdef ELECTRA_HTTPMANAGER_DEBUG_PROXY_ADDRESS
+			HttpOptions.Set(FString(TEXT("proxy")), FVariantValue(FString(ELECTRA_HTTPMANAGER_DEBUG_PROXY_ADDRESS)));
+		#endif
+		HTTPStreamHandler = IElectraHTTPStream::Create(HttpOptions);
+		HTTPStreamHandler->ThreadHandlerDelegate().BindThreadSafeSP(AsShared(), &FElectraHttpManager::ProcessHTTPManager);
+	}
+	
+	void FElectraHttpManager::StopHTTPManager()
+	{
+		LLM_SCOPE(ELLMTag::ElectraPlayer);
 
-				// Handle local file requests.
-				HandleLocalFileRequests();
+		if (HTTPStreamHandler.IsValid())
+		{
+			HTTPStreamHandler->ThreadHandlerDelegate().Unbind();
+		}
 
-				Now = MEDIAutcTime::Current();
+		RemoveAllRequests();
 
-				// Handle requests that have a cached response.
-				HandleCachedHTTPRequests(Now);
+		if (HTTPStreamHandler.IsValid())
+		{
+			HTTPStreamHandler->Close();
+			HTTPStreamHandler.Reset();
+		}
+	}
+	
+	void FElectraHttpManager::ProcessHTTPManager()
+	{
+		LLM_SCOPE(ELLMTag::ElectraPlayer);
 
-				// Handle requests to the HTTP module
-				HandleHTTPRequests(Now);
-				// Handle the responses.
-				HandleHTTPResponses(Now);
+		SCOPE_CYCLE_COUNTER(STAT_ElectraPlayer_FElectraHttpManager_Worker);
+		CSV_SCOPED_TIMING_STAT(ElectraPlayer, ElectraHttpManager_Worker);
 
-				// Handle periodic progress callbacks. Do this before handling completed requests in case a callback asks to abort.
-				HandlePeriodicCallbacks(Now);
-				// Handle timeouts after the progress callbacks.
-				HandleTimeouts(Now);
-				// Handle all finished requests.
-				HandleCompletedRequests();
-			}
-		}
-		RemoveAllRequests();
-		HttpModuleEventMap.Reset();
+		FTimeValue Now = MEDIAutcTime::Current();
+		// Add and remove pending requests
+		Lock.Lock();
+		AddPendingRequests(Now);
+		RemovePendingRequests(Now);
+		Lock.Unlock();
+
+		// Handle local file requests.
+		HandleLocalFileRequests();
+
+		// Handle requests that have a cached response.
+		HandleCachedHTTPRequests(Now);
+
+		// Handle HTTP requests
+		HandleHTTPRequests(Now);
+		HandleHTTPResponses(Now);
+
+		// Handle periodic progress callbacks. Do this before handling the completed requests in case a callback asks to abort.
+		HandlePeriodicCallbacks(Now);
+		// Handle timeouts after the progress callbacks.
+		HandleTimeouts(Now);
+		// Handle all finished requests.
+		HandleCompletedRequests();
 	}
 
 
+
 	void FElectraHttpManager::HandleLocalFileRequests()
 	{
 		for(TMap<FHandle*, TSharedPtrTS<FRequest>>::TIterator It = ActiveRequests.CreateIterator(); It; ++It)
@@ -1050,130 +1026,81 @@ namespace Electra
 			if (Handle->HandleType == FHandle::EHandleType::HTTPHandle)
 			{
 				TSharedPtrTS<FRequest> Request = It.Value();
-				if (Handle->ActiveResponse.CacheResponse.IsValid())
+				if (Handle->ActiveResponse.CacheResponse.IsValid() && !Handle->ActiveResponse.bHitCache)
 				{
-					// First send all headers again
-					FScopeLock lock(&HttpModuleEventLock);
-					FHttpEventData& Evt = HttpModuleEventMap.FindOrAdd(Handle->HttpRequest);
-					for(auto &Header : Handle->ActiveResponse.CacheResponse->Response->GetAllHeaders())
-					{
-						HTTP::FHTTPHeader Hdr;
-						Hdr.SetFromString(Header);
-						// Filter out some response headers we cannot re-use.
-						if (Hdr.Header.Equals(TEXT("Date"), ESearchCase::IgnoreCase))
-						{
-							// Date header is bad because this can be used to synchronize the player clock.
-							continue;
-						}
-						Evt.Headers.Emplace(MoveTemp(Hdr));
-					}
-					// Set the received length;
-					Evt.BytesReceived = Handle->ActiveResponse.CacheResponse->Response->GetContent().Num();
-					// And the rest
-					Evt.bComplete = true;
-					Evt.bConnectedSuccessfully = true;
-					Evt.Status = EHttpRequestStatus::Type::Succeeded;
-					Evt.Response = Handle->ActiveResponse.CacheResponse->Response;
-					Handle->ActiveResponse.CacheResponse.Reset();
 					Handle->ActiveResponse.bHitCache = true;
+					FScopeLock lock(&HTTPStreamEventLock);
+					NotifiedRequests.AddUnique(Handle->HttpRequest);
 				}
 			}
 		}
 	}
 
+
 	void FElectraHttpManager::HandleHTTPRequests(const FTimeValue& Now)
 	{
 		// Get the events that have fired so far into a local map and clear out the original.
-		HttpModuleEventLock.Lock();
-		TMap<FHttpRequestPtr, FHttpEventData> Events = MoveTemp(HttpModuleEventMap);
-		HttpModuleEventLock.Unlock();
-
-		for(auto& Evt : Events)
+		HTTPStreamEventLock.Lock();
+		TArray<IElectraHTTPStreamRequestPtr> Notifieds = MoveTemp(NotifiedRequests);
+		HTTPStreamEventLock.Unlock();
+		for(auto &Notified : Notifieds)
 		{
 			// Find the request.
 			for(TMap<FHandle*, TSharedPtrTS<FRequest>>::TIterator It = ActiveRequests.CreateIterator(); It; ++It)
 			{
 				FHandle* Handle = It.Key();
-				if (Handle->HandleType == FHandle::EHandleType::HTTPHandle && Handle->HttpRequest == Evt.Key)
+				if (Handle->HandleType == FHandle::EHandleType::HTTPHandle && Handle->HttpRequest == Notified)
 				{
 					TSharedPtrTS<FRequest> Request = It.Value();
 					HTTP::FConnectionInfo& ci = Request->ConnectionInfo;
-					const FHttpEventData& EvData = Evt.Value;
 
-					FHttpResponsePtr Response = EvData.Response;
-					const int32 /*EHttpResponseCodes::Type*/ ResponseCode = Response ? Response->GetResponseCode() : EHttpResponseCodes::Unknown;
-					const bool bHTTPResponseSuccess = (ResponseCode >= 200 && ResponseCode <= 299) || ResponseCode == 304;
+					IElectraHTTPStreamResponsePtr Response = Handle->HttpRequest->GetResponse();
 
-					//FPlatformMisc::LocalPrint(*FString::Printf(TEXT("(%d, %d, %f; %d, %d, %d, %d)\n"), Status, ResponseCode, ElapsedTime, EvData.bComplete, EvData.bConnectedSuccessfully, EvData.BytesReceived, EvData.BytesSent));
+					if (Handle->ActiveResponse.CacheResponse.IsValid() && Handle->ActiveResponse.bHitCache)
+					{
+						Response = Handle->ActiveResponse.CacheResponse->Response;
+						Handle->ActiveResponse.CacheResponse.Reset();
+					}
 
-					// Set connect state on the first event indicating any kind of progress.
-					// Note: The HTTP module does not allow access to the response headers or HTTP status code during processing,
-					//       even when those are already available.
-					if (Handle->bHttpRequestFirstEvent)
+					if (Response->GetState() > IElectraHTTPStreamResponse::EState::Connecting && Handle->bHttpRequestFirstEvent)
 					{
-						if (EvData.bComplete || EvData.BytesSent || EvData.BytesReceived || EvData.Headers.Num())
+						if (Handle->ActiveResponse.NumSubRangeRequest == 0)
 						{
-							// Any data flow means that the connection has been established.
-							if (Handle->ActiveResponse.NumSubRangeRequest == 0)
-							{
-								ci.bIsConnected = true;
-								ci.TimeUntilConnected = (Now - Handle->RequestStartTime).GetAsSeconds();
-							}
-							Handle->bIsConnected = true;
-							// That's it for the first event.
-							Handle->bHttpRequestFirstEvent = false;
+							ci.bIsConnected = true;
+							ci.TimeForDNSResolve = Response->GetTimeUntilNameResolved();
+							ci.TimeUntilConnected = Response->GetTimeUntilConnected();
 						}
+						Handle->bIsConnected = true;
+						Handle->bHttpRequestFirstEvent = false;
 					}
 
-					// Process response headers.
-					if (EvData.Headers.Num())
+					// Get the response headers.
+					if (Response->GetState() >= IElectraHTTPStreamResponse::EState::ReceivedResponseHeaders && !ci.bHaveResponseHeaders)
 					{
-						// It may not be all of them yet, but headers were received.
-						// Unfortunately we do not get the "HTTP/1.0" / "HTTP/1.1" / "HTTP/2" header here to see which protocol version we are dealing with.
-						ci.bHaveResponseHeaders = true;
-						Handle->bHaveResponseHeaders = true;
-						for(int32 i=0, iMax=EvData.Headers.Num(); i<iMax; ++i)
+						Handle->LastTimeDataReceived = Now;
+
+						TArray<FElectraHTTPStreamHeader> Headers;
+						Response->GetAllHeaders(Headers);
+
+						for(auto &Header : Headers)
 						{
-							if (!EvData.Headers[i].Header.IsEmpty())
-							{
-								ci.ResponseHeaders.Add(EvData.Headers[i]);
-								const FString& HeaderKey = EvData.Headers[i].Header;
-								const FString& HeaderValue = EvData.Headers[i].Value;
-								if (HeaderKey == TEXT("Content-Length"))
-								{
-									ci.ContentLengthHeader = TEXT("Content-Length:") + HeaderValue;
-								}
-								else if (HeaderKey == TEXT("Content-Range"))
-								{
-									ci.ContentRangeHeader = TEXT("Content-Range: ") + HeaderValue;
-								}
-								else if (HeaderKey == TEXT("Content-Type"))
-								{
-									ci.ContentType = HeaderValue;
-								}
-								else if (HeaderKey == TEXT("Accept-Ranges"))
-								{
-									if (ci.HTTPVersionReceived == 0)
-									{
-										if (HeaderValue.Find(TEXT("none"), ESearchCase::IgnoreCase, ESearchDir::FromStart, INDEX_NONE) == INDEX_NONE)
-										{
-											ci.HTTPVersionReceived = 11;
-										}
-										else
-										{
-											ci.HTTPVersionReceived = 10;
-										}
-									}
-								}
-								else if (HeaderKey == TEXT("Transfer-Encoding"))
-								{
-									if (HeaderValue.Find(TEXT("chunked"), ESearchCase::IgnoreCase, ESearchDir::FromStart, INDEX_NONE) != INDEX_NONE)
-									{
-										ci.bIsChunked = true;
-									}
-								}
-							}
+							HTTP::FHTTPHeader h;
+							h.Header = Header.Header;
+							h.Value = Header.Value;
+							ci.ResponseHeaders.Add(h);
 						}
+						ci.ContentLengthHeader = Response->GetContentLengthHeader();
+						ci.ContentRangeHeader = Response->GetContentRangeHeader();
+						ci.ContentType = Response->GetContentTypeHeader();
+						ci.bIsChunked = Response->GetTransferEncodingHeader().Find(TEXT("chunked"), ESearchCase::IgnoreCase) != INDEX_NONE;
+						ci.HTTPVersionReceived = Response->GetHTTPStatusLine().Find(TEXT("HTTP/1.1"), ESearchCase::IgnoreCase) != INDEX_NONE ? 11 :
+													Response->GetHTTPStatusLine().Find(TEXT("HTTP/2"), ESearchCase::IgnoreCase) != INDEX_NONE ? 20 :
+													Response->GetHTTPStatusLine().Find(TEXT("HTTP/1.0"), ESearchCase::IgnoreCase) != INDEX_NONE ? 10 : 11;
+
+						ci.EffectiveURL = Response->GetEffectiveURL();
+						ci.StatusInfo.HTTPStatus = Response->GetHTTPResponseCode();
+						ci.NumberOfRedirections = 0;
+						bool bContentRangeOk = ci.ContentRangeHeader.Len() ? Handle->ActiveResponse.ReceivedContentRange.ParseFromContentRangeResponse(ci.ContentRangeHeader) : true;
 
 						// Content length needs a bit of special handling.
 						if (Handle->ActiveResponse.NumSubRangeRequest == 0)
@@ -1182,7 +1109,7 @@ namespace Electra
 							// Is there a document size from a Content-Range header?
 							if (!ci.ContentRangeHeader.IsEmpty())
 							{
-								FParams::FRange ContentRange;
+								ElectraHTTPStream::FHttpRange ContentRange;
 								if (ContentRange.ParseFromContentRangeResponse(ci.ContentRangeHeader))
 								{
 									int64 ds = ContentRange.GetDocumentSize();
@@ -1218,117 +1145,81 @@ namespace Electra
 								if (!ci.ContentLengthHeader.IsEmpty())
 								{
 									// Parse from "Content-Length: " header
-									LexFromString(ci.ContentLength, *ci.ContentLengthHeader.Mid(15));
+									LexFromString(ci.ContentLength, *ci.ContentLengthHeader);
 								}
 							}
 						}
-					}
-
-					if (EvData.bComplete && Response.IsValid())
-					{
-						/*
-							We would like to get the effective URL after any redirections.
-							There doesn't seem to be a way right now to get it.
-
-							UE_LOG(LogElectraHTTPManager, Log, TEXT("Effective URL: %s"), *Response->GetURL());
-							UE_LOG(LogElectraHTTPManager, Log, TEXT("Request   URL: %s"), *Request->Parameters.URL);
-						*/
-						ci.EffectiveURL = Request->Parameters.URL;
-						ci.StatusInfo.HTTPStatus = ResponseCode;
-						// Update the time to first byte only on the first sub range request. This is the anchor for the request as a whole.
-						if (Handle->ActiveResponse.NumSubRangeRequest == 0)
-						{
-							ci.TimeUntilFirstByte = ci.TimeUntilConnected;
-							ci.TimeForDNSResolve = 0.001;			// something not zero for the lack of an actual value.
-						}
-						if (ci.HTTPVersionReceived == 0)
-						{
-							ci.HTTPVersionReceived = 11;			// for lack of an actual value let's assume it is HTTP/1.1
-						}
-						ci.NumberOfRedirections = 0;				// no official value to be had
 
-						// Check that the data we are receiving from a range request is actually that range!
-						if (!Handle->bResponseReceived && Request->Parameters.Range.IsSet() && !Request->Parameters.Range.IsEverything())
+						// If we requested a byte range we need to check if we got the correct range back.
+						// For compatibilities sake a 200 response will also be accepted as long as the number of bytes match
+						// the number requested. This then requires the Content-Length response header to be present.
+						if (Request->Parameters.Range.IsSet() && !Request->Parameters.Range.IsEverything())
 						{
-							// Result needs to be "206 - partial content" and a content range header must be present.
-							if (ci.StatusInfo.HTTPStatus != 206 || ci.ContentRangeHeader.IsEmpty())
+							if (ci.StatusInfo.HTTPStatus == 206 && ci.ContentRangeHeader.Len() && bContentRangeOk)
 							{
-								// We allow a 200 response as long as the content size matches the number of bytes we have requested.
-								if (ci.StatusInfo.HTTPStatus == 200 && ci.ContentLength == Request->Parameters.Range.GetNumberOfBytes())
-								{
-									// This primarily supports ISO/IEC-23009-1 Annex E
-								}
-								else
-								{
-									// Note that we did not receive what we thought to get. This gets checked for below.
-									ci.bResponseNotRanged = true;
-								}
+								// We assume that the returned range is what we have requested. It's possible to check for it
+								// but not trivial since the range could be open ended on either side in the request and values
+								// provided in the response.
 							}
-						}
-					}
-
-					// Update receive stats.
-					Handle->LastTimeDataReceived = Now;
-					ci.BytesReadSoFar = Handle->BytesReadSoFar + EvData.BytesReceived;
-
-					// Is the transfer done?
-					if (EvData.bComplete)
-					{
-						bool bFailure = false;
-						bool bContentRangeOk = true;
-						// If there is a Content-Range header we must be able to parse it.
-						if (!ci.ContentRangeHeader.IsEmpty())
-						{
-							bContentRangeOk = Handle->ActiveResponse.ReceivedContentRange.ParseFromContentRangeResponse(ci.ContentRangeHeader);
-						}
-
-						// Success or failure?
-						if (EvData.Status == EHttpRequestStatus::Succeeded && bHTTPResponseSuccess && ci.bResponseNotRanged == false && bContentRangeOk)
-						{
-							// Active responses are handed out in HandleHTTPResponses since they may not fit into the target receive ring buffer
-							// at once and have to be passed out in pieces.
-							if (Response.IsValid())
+							else if (ci.StatusInfo.HTTPStatus == 200 && ci.ContentLength == Request->Parameters.Range.GetNumberOfBytes())
 							{
-								Handle->ActiveResponse.Response = Response;
+								// Allow a 200 response if the number of bytes received matches the number of bytes requested.
 							}
 							else
 							{
-								// Response not valid?
-								bFailure = true;
+								// Not good.
+								ci.bResponseNotRanged = true;
 							}
 						}
-						else
+
+						ci.bHaveResponseHeaders = true;
+						Handle->bHaveResponseHeaders = true;
+
+						// Check for HTTP errors. Redirects are not really expected to reach us here. If they do there were too many.
+						static const TArray<int32> GoodHTTPResponseCodes = { 200, 204, 206, 304 };
+						bool bHTTPResponseOk = GoodHTTPResponseCodes.Contains(ci.StatusInfo.HTTPStatus);
+						bool bFailed = false;
+						if (!bHTTPResponseOk || ci.bResponseNotRanged)
 						{
-							bFailure = true;
 							ci.StatusInfo.ErrorDetail.SetError(UEMEDIA_ERROR_READ_ERROR).SetFacility(Facility::EFacility::HTTPReader).SetCode(ERRCODE_HTTPMODULE_FAILURE);
-							// Set up a bit more verbose error message
 							if (ci.StatusInfo.HTTPStatus >= 400)
 							{
-								ci.StatusInfo.ErrorCode = ERRCODE_HTTP_RETURNED_ERROR;
 								ci.StatusInfo.ErrorDetail.SetMessage(FString::Printf(TEXT("HTTP returned status %d"), ci.StatusInfo.HTTPStatus));
+								ci.StatusInfo.ErrorCode = ERRCODE_HTTP_RETURNED_ERROR;
+								bFailed = true;
 							}
-							else if (ci.bResponseNotRanged)
+							else if (ci.StatusInfo.HTTPStatus >= 300)
 							{
-								ci.StatusInfo.ErrorCode = ERRCODE_HTTP_RANGE_ERROR;
-								ci.StatusInfo.ErrorDetail.SetMessage(FString::Printf(TEXT("Did not receive HTTP 206 for range request")));
+								ci.StatusInfo.ErrorDetail.SetMessage(FString::Printf(TEXT("Reached HTTP redirection limit with returned status %d"), ci.StatusInfo.HTTPStatus));
+								ci.StatusInfo.ErrorCode = ERRCODE_HTTP_RETURNED_ERROR;
+								bFailed = true;
 							}
-							else if (!bContentRangeOk)
+							else if (ci.StatusInfo.HTTPStatus == 0)
 							{
-								ci.StatusInfo.ErrorCode = ERRCODE_HTTP_RANGE_ERROR;
-								ci.StatusInfo.ErrorDetail.SetMessage(FString::Printf(TEXT("Failed to parse HTTP response header \"%s\""), *ci.ContentRangeHeader));
+								ci.StatusInfo.ErrorDetail.SetMessage(FString::Printf(TEXT("Failed to get response from server")));
+								ci.StatusInfo.ErrorCode = ERRCODE_HTTP_RETURNED_ERROR;
+								bFailed = true;
 							}
-							else
+							else if (ci.bResponseNotRanged)
 							{
-								ci.StatusInfo.ErrorCode = ERRCODE_WRITE_ERROR;
-								ci.StatusInfo.bReadError = true;
-								ci.StatusInfo.ErrorDetail.SetMessage(FString::Printf(TEXT("Read error after receiving %lld of %lld bytes"), (long long int)ci.BytesReadSoFar, (long long int)ci.ContentLength));
+								ci.StatusInfo.ErrorDetail.SetMessage(FString::Printf(TEXT("Did not receive HTTP 206 for range request")));
+								ci.StatusInfo.ErrorCode = ERRCODE_HTTP_RANGE_ERROR;
+								bFailed = true;
 							}
 						}
-						Handle->bResponseReceived = true;
-						ci.RequestEndTime = ci.StatusInfo.OccurredAtUTC = Now;
-						// Successful responses are processed outside this method. See above.
-						if (bFailure)
+
+						// After having handled the headers we allow the response to get consumed.
+						// Further error handling (which at this point can only be a closed connection or timeouts)
+						// is handled in HandleHTTPResponses().
+						if (!Handle->HttpRequest->HasFailed() && !bFailed)
+						{
+							Handle->ActiveResponse.Response = Response;
+						}
+						else
 						{
+							ci.bHasFinished = true;
+							ci.RequestEndTime = ci.StatusInfo.OccurredAtUTC = Now;
+							Handle->HttpRequest->Cancel();
 							RequestsCompleted.Enqueue(Request);
 						}
 					}
@@ -1337,6 +1228,7 @@ namespace Electra
 		}
 	}
 
+
 	void FElectraHttpManager::HandleHTTPResponses(const FTimeValue& Now)
 	{
 		for(TMap<FHandle*, TSharedPtrTS<FRequest>>::TIterator It = ActiveRequests.CreateIterator(); It; ++It)
@@ -1345,96 +1237,156 @@ namespace Electra
 			if (Handle->HandleType == FHandle::EHandleType::HTTPHandle)
 			{
 				TSharedPtrTS<FRequest> Request = It.Value();
-				HTTP::FConnectionInfo& ci = Request->ConnectionInfo;
-				ci.bIsCachedResponse = Handle->ActiveResponse.bHitCache;
 
-				// Active response?
-				if (Handle->ActiveResponse.Response.IsValid())
+				// Active response? This will not be set if the HTTP response code indicated an error already.
+				// For our use cases we do not need to read the response body which will be an error message only anyway.
+				IElectraHTTPStreamResponsePtr Response = Handle->ActiveResponse.Response;
+				if (Response.IsValid())
 				{
-					bool bHasFinished = false;
+					HTTP::FConnectionInfo& ci = Request->ConnectionInfo;
+					ci.bIsCachedResponse = Handle->ActiveResponse.bHitCache;
 
-					// Add to response cache unless this was a cached response already.
-					if (Request->ResponseCache.IsValid() && !Handle->ActiveResponse.bHitCache && !Handle->ActiveResponse.bWasAddedToCache)
+					bool bCompletedWithInsufficientData = !Request->Parameters.Verb.Equals(TEXT("HEAD")) &&
+														  Response->GetStatus() == IElectraHTTPStreamResponse::EStatus::Completed && 
+														  Response->GetResponseData().GetLengthFromResponseHeader() > Response->GetResponseData().GetNumTotalBytesAdded();
+					// Has it failed?
+					if (Handle->HttpRequest->HasFailed() || bCompletedWithInsufficientData)
 					{
-						Handle->ActiveResponse.bWasAddedToCache = true;
-						TSharedPtrTS<IHTTPResponseCache::FCacheItem> CacheItem = MakeSharedTS<IHTTPResponseCache::FCacheItem>();
-						CacheItem->URL = Request->Parameters.URL;
-						CacheItem->Range = Handle->ActiveResponse.Range;
-						// Make sure the range is always set, even for non-partial responses.
-						CacheItem->Range.DocumentSize = ci.ContentLength > Handle->ActiveResponse.ReceivedContentRange.GetDocumentSize() ? ci.ContentLength : Handle->ActiveResponse.ReceivedContentRange.GetDocumentSize();
-						if (CacheItem->Range.GetStart() < 0)
-						{
-							CacheItem->Range.SetStart(0);
-						}
-						if (CacheItem->Range.GetEndIncluding() < 0)
-						{
-							CacheItem->Range.SetEndIncluding(CacheItem->Range.DocumentSize - 1);
-						}
-						CacheItem->Response = Handle->ActiveResponse.Response;
-						Request->ResponseCache->CacheEntity(CacheItem);
+						ci.StatusInfo.ErrorCode = ERRCODE_WRITE_ERROR;
+						ci.StatusInfo.bReadError = true;
+						ci.StatusInfo.ErrorDetail.SetMessage(FString::Printf(TEXT("Read error after receiving %lld of %lld bytes"), (long long int)ci.BytesReadSoFar, (long long int)ci.ContentLength));
+						ci.bHasFinished = true;
+						ci.RequestEndTime = ci.StatusInfo.OccurredAtUTC = Now;
+						Handle->ActiveResponse.Response.Reset();
+						RequestsCompleted.Enqueue(Request);
+						continue;
 					}
 
-					// Receive buffer still there?
-					TSharedPtrTS<FReceiveBuffer> ReceiveBuffer = Request->ReceiveBuffer.Pin();
-					if (ReceiveBuffer.IsValid())
+					bool bHasFinished = Response->GetResponseData().HasAllDataBeenConsumed();
+					int64 NumDataAvailable = Response->GetResponseData().GetNumBytesAvailableForRead();
+					if (NumDataAvailable > 0)
 					{
-						bool bCheckForNext = false;
-						// Is it to be used as a linear buffer?
-						if (!ReceiveBuffer->bEnableRingbuffer)
+						if (ci.TimeUntilFirstByte <= 0.0)
 						{
-							const TArray<uint8>& Content = Handle->ActiveResponse.Response->GetContent();
+							ci.TimeUntilFirstByte = Response->GetTimeUntilFirstByte();
+						}
+						Handle->LastTimeDataReceived = Now;
 
-							if (!ReceiveBuffer->Buffer.EnlargeTo(ReceiveBuffer->Buffer.Num() + Content.Num()) ||
-								!ReceiveBuffer->Buffer.PushData(Content.GetData(), Content.Num()))
+						// Receive buffer still there?
+						TSharedPtrTS<FReceiveBuffer> ReceiveBuffer = Request->ReceiveBuffer.Pin();
+						if (ReceiveBuffer.IsValid())
+						{
+							int64 BufferPushableSize = 0;
+							bool bBufferUsable = false;
+							// Is it to be used as a linear buffer?
+							if (!ReceiveBuffer->bEnableRingbuffer)
 							{
-								ci.bWasAborted = true;
-								bHasFinished = true;
+								int64 RequiredBufferSize = ci.ContentLength > 0 ? ci.ContentLength : 0;
+								int64 BufferSizeAfterPush = ReceiveBuffer->Buffer.Num() + NumDataAvailable;
+								if (BufferSizeAfterPush > RequiredBufferSize)
+								{
+									RequiredBufferSize = BufferSizeAfterPush;
+								}
+								bBufferUsable = ReceiveBuffer->Buffer.EnlargeTo(RequiredBufferSize);
+								BufferPushableSize = bBufferUsable ? NumDataAvailable : 0;
 							}
 							else
 							{
-								ci.ContentLength = ReceiveBuffer->Buffer.Num();
-								ci.BytesReadSoFar = ci.ContentLength;
-								bCheckForNext = true;
+								// The ring buffer must have been set up to a size the client wants it to have.
+								bBufferUsable = ReceiveBuffer->Buffer.Capacity() > 0;
+								check(bBufferUsable);
+								BufferPushableSize = bBufferUsable ? ReceiveBuffer->Buffer.Avail() : 0;
+							}
+							if (bBufferUsable)
+							{
+								if (BufferPushableSize)
+								{
+									const uint8* NewDataPtr = nullptr;
+									int64 NewDataSize = 0;
+									Response->GetResponseData().LockBuffer(NewDataPtr, NewDataSize);
+									int64 NumToCopy = BufferPushableSize < NewDataSize ? BufferPushableSize : NewDataSize;
+									bBufferUsable = ReceiveBuffer->Buffer.PushData(NewDataPtr, NumToCopy, ReceiveBuffer->bEnableRingbuffer);
+									Response->GetResponseData().UnlockBuffer(bBufferUsable ? NumToCopy : 0);
+									if (bBufferUsable)
+									{
+										Handle->ActiveResponse.NumBytesPassedOut += NumToCopy;
+										ci.BytesReadSoFar += NumToCopy;
+									}
+								}
+							}
+							if (!bBufferUsable)
+							{
+								ci.bWasAborted = true;
+								bHasFinished = true;
 							}
 						}
 						else
 						{
-							// The ring buffer must have been set up to a size the client wants it to have.
-							check(ReceiveBuffer->Buffer.Capacity() > 0);
+							// With the receive buffer having been released we can abort the transfer.
+							ci.bWasAborted = true;
+							if (Handle->HttpsRequestCallbackWrapper.IsValid())
+							{
+								Handle->HttpsRequestCallbackWrapper->Unbind();
+							}
+							bHasFinished = true;
+						}
+					}
 
-							// How much room is there available?
-							int32 NumAvail = ReceiveBuffer->Buffer.Avail();
-							if (NumAvail)
+					if (bHasFinished)
+					{
+						if (!ci.bWasAborted && Response->GetResponseData().HasAllDataBeenConsumed())
+						{
+							// Add to response cache unless this was a cached response already.
+							if (Request->ResponseCache.IsValid() && Response->GetResponseData().IsCachable() && !Handle->ActiveResponse.bHitCache && !Handle->ActiveResponse.bWasAddedToCache)
 							{
-								const TArray<uint8>& Content = Handle->ActiveResponse.Response->GetContent();
-								int32 NumToGo = Content.Num() - Handle->ActiveResponse.NumBytesPassedOut;
-								int32 NumToCopy = NumAvail < NumToGo ? NumAvail : NumToGo;
-								ReceiveBuffer->Buffer.PushData(Content.GetData() + Handle->ActiveResponse.NumBytesPassedOut, NumToCopy);
-								Handle->ActiveResponse.NumBytesPassedOut += NumToCopy;
-								// All copied out now?
-								if (Handle->ActiveResponse.NumBytesPassedOut >= Content.Num())
+								Handle->ActiveResponse.bWasAddedToCache = true;
+								TSharedPtrTS<IHTTPResponseCache::FCacheItem> CacheItem = MakeSharedTS<IHTTPResponseCache::FCacheItem>();
+								CacheItem->RequestedURL = Request->Parameters.URL;
+								CacheItem->EffectiveURL = Response->GetEffectiveURL();
+								CacheItem->Range = Handle->ActiveResponse.Range;
+								// Make sure the range is always set, even for non-partial responses.
+								CacheItem->Range.DocumentSize = ci.ContentLength > Handle->ActiveResponse.ReceivedContentRange.GetDocumentSize() ? ci.ContentLength : Handle->ActiveResponse.ReceivedContentRange.GetDocumentSize();
+								if (CacheItem->Range.GetStart() < 0)
+								{
+									CacheItem->Range.SetStart(0);
+								}
+								if (CacheItem->Range.GetEndIncluding() < 0)
 								{
-									bCheckForNext = true;
+									CacheItem->Range.SetEndIncluding(CacheItem->Range.DocumentSize - 1);
 								}
+								CacheItem->Response = Response;
+								Request->ResponseCache->CacheEntity(CacheItem);
 							}
-						}
-						if (bCheckForNext)
-						{
+ 
 							// Check if this was a sub ranged request and if there is still data to go for the original request.
 							if (Handle->ActiveResponse.SizeRemaining() == 0)
 							{
 								// All done now.
-								bHasFinished = true;
+								TSharedPtrTS<FReceiveBuffer> ReceiveBuffer = Request->ReceiveBuffer.Pin();
+								if (ReceiveBuffer.IsValid())
+								{
+									ReceiveBuffer->Buffer.SetEOD();
+								}
 							}
 							else
 							{
 								// Still another sub range to go.
-								if (PrepareHTTPModuleHandle(Now, Handle, Request, false))
+								if (PrepareHTTPHandle(Now, Handle, Request, false))
 								{
-									if (!Handle->ProcessHTTPRequest())
+									// We need to parse the headers from the new request so we have to clear the flag.
+									ci.bHaveResponseHeaders = false;
+									ci.ResponseHeaders.Empty();
+									if (HTTPStreamHandler.IsValid())
+									{
+										if (Handle->NeedsFetchOfUncachedParts())
+										{
+											HTTPStreamHandler->AddRequest(Handle->HttpRequest);
+										}
+										bHasFinished = false;
+									}
+									else
 									{
-										ci.StatusInfo.ErrorDetail.SetError(UEMEDIA_ERROR_INTERNAL).SetFacility(Facility::EFacility::HTTPReader).SetMessage("HTTP request failed on ProcessRequest()");
-										bHasFinished = true;
+										ci.StatusInfo.ErrorDetail.SetError(UEMEDIA_ERROR_INTERNAL).SetFacility(Facility::EFacility::HTTPReader).SetMessage("HTTP sub request failed on AddRequest()");
 									}
 								}
 								else
@@ -1443,26 +1395,18 @@ namespace Electra
 									ci.StatusInfo.ErrorCode = ERRCODE_WRITE_ERROR;
 									ci.StatusInfo.bReadError = true;
 									ci.StatusInfo.ErrorDetail.SetMessage(FString::Printf(TEXT("Error setting up the next sub range request")));
-									bHasFinished = true;
 								}
 							}
 						}
-					}
-					else
-					{
-						// With the receive buffer having been released we can abort the transfer.
-						ci.bWasAborted = true;
-						if (Handle->HttpsRequestCallbackWrapper.IsValid())
+
+						if (bHasFinished)
 						{
-							Handle->HttpsRequestCallbackWrapper->Unbind();
+							ci.bHasFinished = true;
+							ci.RequestEndTime = ci.StatusInfo.OccurredAtUTC = Now;
+
+							Handle->ActiveResponse.Response.Reset();
+							RequestsCompleted.Enqueue(Request);
 						}
-						bHasFinished = true;
-					}
-					if (bHasFinished)
-					{
-						ci.RequestEndTime = ci.StatusInfo.OccurredAtUTC = Now;
-						Handle->ActiveResponse.Response.Reset();
-						RequestsCompleted.Enqueue(Request);
 					}
 				}
 			}
@@ -1471,37 +1415,20 @@ namespace Electra
 
 
 
-	void FElectraHttpManager::OnProcessRequestCompleteInternal(FHttpRequestPtr Request, FHttpResponsePtr Response, bool bConnectedSuccessfully)
+	void FElectraHttpManager::OnRequestNotification(IElectraHTTPStreamRequestPtr InRequest, EElectraHTTPStreamNotificationReason InReason, int64 InParam)
 	{
-		FScopeLock lock(&HttpModuleEventLock);
-		FHttpEventData& Evt = HttpModuleEventMap.FindOrAdd(Request);
-		Evt.bComplete = true;
-		Evt.bConnectedSuccessfully = bConnectedSuccessfully;
-		Evt.Status = Request->GetStatus();
-		Evt.Response = Response;
-	}
-
-	void FElectraHttpManager::OnRequestProgressInternal(FHttpRequestPtr Request, int32 BytesSent, int32 BytesReceived)
-	{
-		FScopeLock lock(&HttpModuleEventLock);
-		FHttpEventData& Evt = HttpModuleEventMap.FindOrAdd(Request);
-		Evt.BytesReceived = BytesReceived;
-		Evt.BytesSent = BytesSent;
-		Evt.Status = Request->GetStatus();
-	}
-
-	void FElectraHttpManager::OnRequestHeaderReceivedInternal(FHttpRequestPtr Request, const FString& HeaderName, const FString& NewHeaderValue)
-	{
-		FScopeLock lock(&HttpModuleEventLock);
-		FHttpEventData& Evt = HttpModuleEventMap.FindOrAdd(Request);
-		Evt.Headers.Add({HeaderName, NewHeaderValue});
+		// Notify only for headers and completion, not data transfer.
+		if (InReason == EElectraHTTPStreamNotificationReason::ReceivedHeaders || InReason == EElectraHTTPStreamNotificationReason::Completed)
+		{
+			FScopeLock lock(&HTTPStreamEventLock);
+			NotifiedRequests.AddUnique(MoveTemp(InRequest));
+		}
 	}
 
 
 
 
 
-
 	void FElectraHttpManager::FFileStream::SetConnected(TSharedPtrTS<FRequest> Request)
 	{
 		if (!bIsConnected)
diff --git a/Engine/Plugins/Media/ElectraPlayer/Source/ElectraPlayerRuntime/Private/Runtime/HTTP/HTTPManager.h b/Engine/Plugins/Media/ElectraPlayer/Source/ElectraPlayerRuntime/Private/Runtime/HTTP/HTTPManager.h
--- a/Engine/Plugins/Media/ElectraPlayer/Source/ElectraPlayerRuntime/Private/Runtime/HTTP/HTTPManager.h
+++ b/Engine/Plugins/Media/ElectraPlayer/Source/ElectraPlayerRuntime/Private/Runtime/HTTP/HTTPManager.h
@@ -9,17 +9,16 @@
 
 #include "ParameterDictionary.h"
 #include "ErrorDetail.h"
+#include "ElectraHTTPStream.h"
+#include "Utilities/HttpRangeHeader.h"
 
 
 DECLARE_LOG_CATEGORY_EXTERN(LogElectraHTTPManager, Log, All);
 
 namespace Electra
 {
-	//
-	class IPlayerSessionServices;
 	class IHTTPResponseCache;
 
-
 	namespace HTTP
 	{
 
@@ -59,11 +58,14 @@ namespace Electra
 		};
 
 
-		struct FHTTPHeader
+		struct FHTTPHeader : public FElectraHTTPStreamHeader
 		{
-			FString	Header;
-			FString	Value;
-
+			FHTTPHeader() = default;
+			FHTTPHeader(const FString& InHeader, const FString& InValue)
+			{
+				Header = InHeader;
+				Value = InValue;
+			}
 			void SetFromString(const FString& InString)
 			{
 				int32 ColonPos;
@@ -255,159 +257,6 @@ namespace Electra
 
 		struct FParams
 		{
-			struct FRange
-			{
-				void Reset()
-				{
-					Start = -1;
-					EndIncluding = -1;
-					DocumentSize = -1;
-				}
-				bool IsSet() const
-				{
-					return Start != -1 || EndIncluding != -1;
-				}
-				bool Equals(const FRange& Other)
-				{
-					return Start == Other.Start && EndIncluding == Other.EndIncluding;
-				}
-				//! Check if the range would result in "0-" for the entire file in which case we don't need to use range request.
-				bool IsEverything() const
-				{
-					return Start <= 0 && EndIncluding < 0;
-				}
-				FString GetString() const
-				{
-					FString s, e, d(TEXT("-"));
-					if (Start >= 0)
-					{
-						// An explicit range?
-						if (EndIncluding >= Start)
-						{
-							return FString::Printf(TEXT("%lld-%lld"), (long long int)Start, (long long int)EndIncluding);
-						}
-						// Start to end, whereever that is.
-						else
-						{
-							return FString::Printf(TEXT("%lld-"), (long long int)Start);
-						}
-					}
-					// Everything
-					return FString(TEXT("0-"));
-				}
-				// Returns the number of bytes in the range, which must be fully specified. An unset or partially open range will return -1.
-				int64 GetNumberOfBytes() const
-				{
-					if (Start >= 0 && EndIncluding >= 0)
-					{
-						return EndIncluding - Start + 1;
-					}
-					return -1;
-				}
-				int64 GetStart() const
-				{
-					return Start;
-				}
-				int64 GetEndIncluding() const
-				{
-					return EndIncluding;
-				}
-				bool IsOpenEnded() const
-				{
-					return GetEndIncluding() < 0;
-				}
-				void Set(const FString& InString)
-				{
-					int32 DashPos = INDEX_NONE;
-					if (InString.FindChar(TCHAR('-'), DashPos))
-					{
-						// -end
-						if (DashPos == 0)
-						{
-							Start = 0;
-							LexFromString(EndIncluding, *InString + 1);
-						}
-						// start-
-						else if (DashPos == InString.Len()-1)
-						{
-							LexFromString(Start, *InString.Mid(0, DashPos));
-							EndIncluding = -1;
-						}
-						// start-end
-						else
-						{
-							LexFromString(Start, *InString.Mid(0, DashPos));
-							LexFromString(EndIncluding, *InString + DashPos + 1);
-						}
-					}
-				}
-				void SetStart(int64 InStart)
-				{
-					Start = InStart;
-				}
-				void SetEndIncluding(int64 InEndIncluding)
-				{
-					EndIncluding = InEndIncluding;
-				}
-				int64 GetDocumentSize() const
-				{
-					return DocumentSize;
-				}
-				bool ParseFromContentRangeResponse(const FString& ContentRangeHeader)
-				{
-					// Examples: <unit> <range-start>-<range-end>/<size>
-					//   Content-Range: bytes 26151-157222/7594984
-					//   Content-Range: bytes 26151-157222/*
-					//   Content-Range: bytes */7594984
-					//
-					Start = -1;
-					EndIncluding = -1;
-					DocumentSize = -1;
-					FString rh = ContentRangeHeader;
-					// In case the entire header is given, remove the header including the separating colon and space.
-					rh.RemoveFromStart(TEXT("Content-Range: "), ESearchCase::CaseSensitive);
-					// Split into parts
-					TArray<FString> Parts;
-					const TCHAR* const Delims[3] = {TEXT(" "),TEXT("-"),TEXT("/")};
-					int32 NumParts = rh.ParseIntoArray(Parts, Delims, 3);
-					if (NumParts)
-					{
-						if (Parts[0] == TEXT("bytes"))
-						{
-							Parts.RemoveAt(0);
-						}
-						// We should now be left with 3 remaining results, the start, end and document size.
-						// The case where we get "*/<size>" we treat as invalid.
-						if (Parts.Num() == 3)
-						{
-							if (Parts[0].IsNumeric())
-							{
-								LexFromString(Start, *Parts[0]);
-								if (Parts[1].IsNumeric())
-								{
-									LexFromString(EndIncluding, *Parts[1]);
-									if (Parts[2].IsNumeric())
-									{
-										LexFromString(DocumentSize, *Parts[2]);
-										return true;
-									}
-									else if (Parts[2] == TEXT("*"))
-									{
-										DocumentSize = -1;
-										return true;
-									}
-								}
-							}
-						}
-					}
-					UE_LOG(LogElectraHTTPManager, Error, TEXT("Failed to parse Content-Range HTTP response header \"%s\""), *ContentRangeHeader);
-					return false;
-				}
-				int64			Start = -1;
-				int64			EndIncluding = -1;
-				int64			DocumentSize = -1;
-			};
-
 			void AddFromHeaderList(const TArray<FString>& InHeaderList)
 			{
 				for(int32 i=0; i<InHeaderList.Num(); ++i)
@@ -420,8 +269,7 @@ namespace Electra
 
 			FString								URL;							//!< URL
 			FString								Verb;							//!< GET (default if not set), HEAD, OPTIONS,....
-			FRange								Range;							//!< Optional request range
-			int32								SubRangeRequestSize = 0;		//!< If not 0 the size to break the request into smaller range requests into.
+			ElectraHTTPStream::FHttpRange		Range;							//!< Optional request range
 			TArray<HTTP::FHTTPHeader>			RequestHeaders;					//!< Request headers
 			TMediaOptionalValue<FString>		AcceptEncoding;					//!< Optional accepted encoding
 			FTimeValue							ConnectTimeout;					//!< Optional timeout for connecting to the server
diff --git a/Engine/Plugins/Media/ElectraPlayer/Source/ElectraPlayerRuntime/Private/Runtime/HTTP/HTTPResponseCache.cpp b/Engine/Plugins/Media/ElectraPlayer/Source/ElectraPlayerRuntime/Private/Runtime/HTTP/HTTPResponseCache.cpp
--- a/Engine/Plugins/Media/ElectraPlayer/Source/ElectraPlayerRuntime/Private/Runtime/HTTP/HTTPResponseCache.cpp
+++ b/Engine/Plugins/Media/ElectraPlayer/Source/ElectraPlayerRuntime/Private/Runtime/HTTP/HTTPResponseCache.cpp
@@ -1,8 +1,8 @@
 // Copyright Epic Games, Inc. All Rights Reserved.
 
 #include "HTTPResponseCache.h"
+#include "ElectraHTTPStreamBuffer.h"
 #include "Player/PlayerSessionServices.h"
-#include "Interfaces/IHttpResponse.h"
 
 
 namespace Electra
@@ -19,75 +19,317 @@ public:
 
 	virtual void HandleEntityExpiration() override;
 	virtual void CacheEntity(TSharedPtrTS<FCacheItem> EntityToAdd) override;
-	virtual EScatterResult GetScatteredCacheEntity(TSharedPtrTS<FCacheItem>& OutScatteredCachedEntity, const FString& URL, const IElectraHttpManager::FParams::FRange& Range) override;
+	virtual EScatterResult GetScatteredCacheEntity(TSharedPtrTS<FCacheItem>& OutScatteredCachedEntity, const FString& URL, const ElectraHTTPStream::FHttpRange& Range) override;
 
 private:
-	class FResponseFromCache : public IHttpResponse
+	/**
+	 * This class wraps IElectraHTTPStreamBuffer with the intention that the original buffer
+	 * can be shared multiple times, each having a new set of internal read offsets.
+	 * Otherwise it would not be possible to return the buffer to multiple cached responses.
+	 * Also, since the original buffer has been read from before it was added to the cache
+	 * it will have its read offsets already pointing to the end.
+	 */
+	class FWrappedHTTPStreamBuffer : public IElectraHTTPStreamBuffer
 	{
 	public:
-		virtual ~FResponseFromCache() = default;
-		FResponseFromCache() = default;
+		FWrappedHTTPStreamBuffer(IElectraHTTPStreamBuffer& InBufferToWrap)
+			: WrappedBuffer(InBufferToWrap)
+		{
+			InBufferToWrap.GetBaseBuffer(WrappedBufferBaseAddress, WrappedBufferSize);
+		}
+		virtual ~FWrappedHTTPStreamBuffer() = default;
+
+		void AddData(const TArray<uint8>& InNewData) override
+		{ }
+		void AddData(TArray<uint8>&& InNewData) override
+		{ }
+		void AddData(const TConstArrayView<const uint8>& InNewData) override
+		{ }
+		void AddData(const IElectraHTTPStreamBuffer& Other, int64 Offset, int64 NumBytes) override
+		{ }
+		int64 GetNumTotalBytesAdded() const override
+		{ return WrappedBuffer.GetNumTotalBytesAdded(); }
+		bool GetEOS() const override
+		{ return true; }
+		void SetEOS() override
+		{ }
+		void ClearEOS() override
+		{ }
+		bool IsCachable() const override
+		{ return false; }
+		void SetIsCachable(bool bInIsCachable) override
+		{ }
+		void SetLengthFromResponseHeader(int64 InLengthFromResponseHeader) override
+		{ }
+		int64 GetLengthFromResponseHeader() const override
+		{ return WrappedBuffer.GetLengthFromResponseHeader(); }
+
+		int64 GetNumTotalBytesHandedOut() const override
+		{
+			return NumBytesHandedOut;
+		}
+		int64 GetNumBytesAvailableForRead()	const override
+		{
+			return WrappedBufferSize - NextReadPosInBuffer;
+		}
 
-		virtual int32 GetResponseCode() const override
-		{ return HTTPResponseCode; }
+		void LockBuffer(const uint8*& OutNextReadAddress, int64& OutNumBytesAvailable) override
+		{
+			OutNextReadAddress = WrappedBufferBaseAddress + NextReadPosInBuffer;
+			OutNumBytesAvailable = WrappedBufferSize - NextReadPosInBuffer;
+		}
+		void UnlockBuffer(int64 NumBytesConsumed) override
+		{
+			check(NumBytesConsumed >= 0 && NumBytesConsumed <= WrappedBufferSize - NextReadPosInBuffer);
+			if (NumBytesConsumed > 0)
+			{
+				if (NumBytesConsumed > WrappedBufferSize - NextReadPosInBuffer)
+				{
+					NumBytesConsumed = WrappedBufferSize - NextReadPosInBuffer;
+				}
+				NextReadPosInBuffer += NumBytesConsumed;
+				NumBytesHandedOut += NumBytesConsumed;
+			}
+		}
+		bool RewindToBeginning() override
+		{
+			NextReadPosInBuffer = 0;
+			NumBytesHandedOut = 0;
+			return true;
+		}
+		bool HasAllDataBeenConsumed() const override
+		{
+			return WrappedBufferSize - NextReadPosInBuffer == 0;
+		}
 
-		virtual FString GetContentAsString() const override
-		{ check(!"do not call"); return FString(); }
+	private:
+		virtual void GetBaseBuffer(const uint8*& OutBaseAddress, int64& OutBytesInBuffer) override
+		{
+			check(!"Why do you get here?");
+		}
+		IElectraHTTPStreamBuffer& WrappedBuffer;
+		const uint8* WrappedBufferBaseAddress;
+		int64 WrappedBufferSize;
 
-		virtual FString GetURL() const override
-		{ return URL; }
+		int64 NextReadPosInBuffer = 0;
+		int64 NumBytesHandedOut = 0;
+	};
 
-		virtual FString GetURLParameter(const FString& ParameterName) const override
-		{ check(!"do not call"); return FString(); }
 
-		virtual FString GetHeader(const FString& HeaderName) const override
-		{ check(!"do not call"); return FString(); }
+	class FWrappedCacheResponse : public IElectraHTTPStreamResponse
+	{
+	public:
+		FWrappedCacheResponse(IElectraHTTPStreamResponsePtr InResponseToWrap)
+			: WrappedResponse(InResponseToWrap)
+			, WrappedBuffer(InResponseToWrap->GetResponseData())
+		{
+			CopyHeadersExceptDateFrom(InResponseToWrap);
+		}
+		virtual ~FWrappedCacheResponse() = default;
+
+		EStatus GetStatus() override
+		{ return WrappedResponse->GetStatus(); }
+		EState GetState() override
+		{ return WrappedResponse->GetState(); }
+		FString GetErrorMessage() override
+		{ return WrappedResponse->GetErrorMessage(); }
+		int32 GetHTTPResponseCode() override
+		{ return WrappedResponse->GetHTTPResponseCode(); }
+		int64 GetNumResponseBytesReceived() override
+		{ return WrappedResponse->GetNumResponseBytesReceived(); }
+		int64 GetNumRequestBytesSent() override
+		{ return WrappedResponse->GetNumRequestBytesSent(); }
+		FString GetEffectiveURL() override
+		{ return WrappedResponse->GetEffectiveURL(); }
+		FString GetHTTPStatusLine() override
+		{ return WrappedResponse->GetHTTPStatusLine(); }
+		FString GetContentLengthHeader() override
+		{ return WrappedResponse->GetContentLengthHeader(); }
+		FString GetContentRangeHeader() override
+		{ return WrappedResponse->GetContentRangeHeader(); }
+		FString GetAcceptRangesHeader() override
+		{ return WrappedResponse->GetAcceptRangesHeader(); }
+		FString GetTransferEncodingHeader() override
+		{ return WrappedResponse->GetTransferEncodingHeader(); }
+		FString GetContentTypeHeader() override
+		{ return WrappedResponse->GetContentTypeHeader(); }
+		double GetTimeElapsed() override
+		{ return WrappedResponse->GetTimeElapsed(); }
+		double GetTimeSinceLastDataArrived() override
+		{ return WrappedResponse->GetTimeSinceLastDataArrived(); }
+		double GetTimeUntilNameResolved() override
+		{ return WrappedResponse->GetTimeUntilNameResolved(); }
+		double GetTimeUntilConnected() override
+		{ return WrappedResponse->GetTimeUntilConnected(); }
+		double GetTimeUntilRequestSent() override
+		{ return WrappedResponse->GetTimeUntilRequestSent(); }
+		double GetTimeUntilHeadersAvailable() override
+		{ return WrappedResponse->GetTimeUntilHeadersAvailable(); }
+		double GetTimeUntilFirstByte() override
+		{ return WrappedResponse->GetTimeUntilFirstByte(); }
+		double GetTimeUntilFinished() override
+		{ return WrappedResponse->GetTimeUntilFinished(); }
+
+		void GetAllHeaders(TArray<FElectraHTTPStreamHeader>& OutHeaders) override
+		{ OutHeaders = Headers; }
+
+		IElectraHTTPStreamBuffer& GetResponseData() override
+		{ return WrappedBuffer; }
 
-		virtual TArray<FString> GetAllHeaders() const override
-		{ return Headers; }
+	private:
+		void CopyHeadersExceptDateFrom(IElectraHTTPStreamResponsePtr CachedResponse)
+		{
+			check(CachedResponse);
+			if (CachedResponse.IsValid())
+			{
+				TArray<FElectraHTTPStreamHeader> CachedHeaders;
+				CachedResponse->GetAllHeaders(CachedHeaders);
+				// Copy all headers except for those containing the date
+				for(auto &Header : CachedHeaders)
+				{
+					if (!Header.Header.Equals(TEXT("Date"), ESearchCase::IgnoreCase))
+					{
+						Headers.Emplace(Header);
+					}
+				}
+			}
+		}
 
-		virtual FString GetContentType() const override
-		{ return GetHeader(TEXT("Content-Type")); }
+		IElectraHTTPStreamResponsePtr WrappedResponse;
+		TArray<FElectraHTTPStreamHeader> Headers;
+		FWrappedHTTPStreamBuffer WrappedBuffer;
+	};
 
-		virtual int32 GetContentLength() const override
-		{ return Payload.Num(); }
 
-		virtual const TArray<uint8>& GetContent() const override
-		{ return Payload; }
+	class FSynthesizedCacheResponse : public IElectraHTTPStreamResponse
+	{
+	public:
+		virtual ~FSynthesizedCacheResponse() = default;
+		FSynthesizedCacheResponse()
+		{
+			// The new response is itself not cachable.
+			Buffer.SetIsCachable(false);
+		}
+
+		EStatus GetStatus() override
+		{ return EStatus::Completed; }
+		EState GetState() override
+		{ return EState::Finished; }
+		FString GetErrorMessage() override
+		{ return FString(); }
+		int32 GetHTTPResponseCode() override
+		{ return HTTPResponseCode; }
+		int64 GetNumResponseBytesReceived() override
+		{ return Buffer.GetNumTotalBytesAdded(); }
+		int64 GetNumRequestBytesSent() override
+		{ return 0; }
+		FString GetEffectiveURL() override
+		{ return URL; }
+
+		void GetAllHeaders(TArray<FElectraHTTPStreamHeader>& OutHeaders) override
+		{ OutHeaders = Headers; }
+		FString GetHTTPStatusLine() override
+		{ return StatusLine; }
+		FString GetContentLengthHeader() override
+		{ return GetHeader(TEXT("Content-Length")); }
+		FString GetContentRangeHeader() override
+		{ return GetHeader(TEXT("Content-Range")); }
+		FString GetAcceptRangesHeader() override
+		{ return GetHeader(TEXT("Accept-Ranges")); }
+		FString GetTransferEncodingHeader() override
+		{ return GetHeader(TEXT("Transfer-Encoding")); }
+		FString GetContentTypeHeader() override
+		{ return GetHeader(TEXT("Content-Type")); }
 
-		void CopyHeadersFromExceptContentSizes(FHttpResponsePtr CachedResponse)
+		IElectraHTTPStreamBuffer& GetResponseData() override
+		{ return Buffer; }
+
+		double GetTimeElapsed() override
+		{ return 0.0; }
+		double GetTimeSinceLastDataArrived() override
+		{ return 0.0; }
+		double GetTimeUntilNameResolved() override
+		{ return TimeUntilNameResolved; }
+		double GetTimeUntilConnected() override
+		{ return TimeUntilConnected; }
+		double GetTimeUntilRequestSent() override
+		{ return TimeUntilRequestSent; }
+		double GetTimeUntilHeadersAvailable() override
+		{ return TimeUntilHeadersAvailable; }
+		double GetTimeUntilFirstByte() override
+		{ return TimeUntilFirstByte; }
+		double GetTimeUntilFinished() override
+		{ return TimeUntilFinished; }
+
+		void SetTimeUntilNameResolved(double InTimeUntilNameResolved)
+		{ TimeUntilNameResolved = InTimeUntilNameResolved; }
+		void SetTimeUntilConnected(double InTimeUntilConnected)
+		{ TimeUntilConnected = InTimeUntilConnected; }
+		void SetTimeUntilRequestSent(double InTimeUntilRequestSent)
+		{ TimeUntilRequestSent = InTimeUntilRequestSent; }
+		void SetTimeUntilHeadersAvailable(double InTimeUntilHeadersAvailable)
+		{ TimeUntilHeadersAvailable = InTimeUntilHeadersAvailable; }
+		void SetTimeUntilFirstByte(double InTimeUntilFirstByte)
+		{ TimeUntilFirstByte = InTimeUntilFirstByte; }
+		void SetTimeUntilFinished(double InTimeUntilFinished)
+		{ TimeUntilFinished = InTimeUntilFinished; }
+
+		void CopyHeadersExceptContentSizesAndDateFrom(IElectraHTTPStreamResponsePtr CachedResponse)
 		{
 			check(CachedResponse);
 			if (CachedResponse.IsValid())
 			{
-				// Copy all headers except for those containing the content size or length.
-				for(auto &Header : CachedResponse->GetAllHeaders())
+				TArray<FElectraHTTPStreamHeader> CachedHeaders;
+				CachedResponse->GetAllHeaders(CachedHeaders);
+				// Copy all headers except for those containing the content size, length or date.
+				for(auto &Header : CachedHeaders)
 				{
-					if (!Header.StartsWith(TEXT("Content-Length:"), ESearchCase::IgnoreCase) &&
-						!Header.StartsWith(TEXT("Content-Range:"), ESearchCase::IgnoreCase))
+					if (!Header.Header.Equals(TEXT("Content-Length"), ESearchCase::IgnoreCase) &&
+						!Header.Header.Equals(TEXT("Content-Range"), ESearchCase::IgnoreCase) &&
+						!Header.Header.Equals(TEXT("Date"), ESearchCase::IgnoreCase))
 					{
-						Headers.Emplace(MoveTemp(Header));
+						Headers.Emplace(Header);
 					}
 				}
 			}
 		}
-		void AddHeader(FString InHeader)
-		{ Headers.Emplace(MoveTemp(InHeader)); }
-
-		void SetPayload(TArray<uint8> InPayload)
-		{ Payload = MoveTemp(InPayload); }
-
+		void AddHeader(FString InHeader, FString InValue)
+		{ 
+			FElectraHTTPStreamHeader Hdr;
+			Hdr.Header = MoveTemp(InHeader);
+			Hdr.Value = MoveTemp(InValue);
+			Headers.Emplace(MoveTemp(Hdr)); 
+		}
 		void SetURL(const FString& InUrl)
 		{ URL = InUrl; }
-
+		void SetHTTPStatusLine(const FString& InStatusLine)
+		{ StatusLine = InStatusLine; }
 		void SetHTTPResponseCode(int32 InResponseCode)
 		{ HTTPResponseCode = InResponseCode; }
-
 	private:
-		TArray<uint8> Payload;
-		TArray<FString> Headers;
+		FString GetHeader(const TCHAR* const InHeader)
+		{
+			for(auto& Hdr : Headers)
+			{
+				if (Hdr.Header.Equals(InHeader, ESearchCase::IgnoreCase))
+				{
+					return Hdr.Value;
+				}
+			}
+			return FString();
+		}
+		FElectraHTTPStreamBuffer Buffer;
+		TArray<FElectraHTTPStreamHeader> Headers;
 		FString URL;
+		FString StatusLine;
 		int32 HTTPResponseCode = 200;
+
+		double TimeUntilNameResolved = 0.0;
+		double TimeUntilConnected = 0.0;
+		double TimeUntilRequestSent = 0.0;
+		double TimeUntilHeadersAvailable = 0.0;
+		double TimeUntilFirstByte = 0.0;
+		double TimeUntilFinished = 0.0;
 	};
 
 	void EvictToAddSize(int64 ResponseSize);
@@ -146,7 +388,7 @@ void FHTTPResponseCache::EvictToAddSize(int64 ResponseSize)
 	while(Cache.Num() && (MaxElementSize - SizeInUse < ResponseSize || Cache.Num() >= MaxNumElements))
 	{
 		TSharedPtrTS<FCacheItem> Item = Cache.Pop();
-		SizeInUse -= Item->Response->GetContent().Num();
+		SizeInUse -= Item->Response->GetResponseData().GetNumTotalBytesAdded();
 	}
 }
 
@@ -160,7 +402,9 @@ void FHTTPResponseCache::CacheEntity(TSharedPtrTS<FCacheItem> EntityToAdd)
 {
 	if (EntityToAdd.IsValid())
 	{
-		int64 SizeRequired = EntityToAdd->Response->GetContent().Num();
+		check(EntityToAdd->EffectiveURL.Len());
+
+		int64 SizeRequired = EntityToAdd->Response->GetResponseData().GetNumTotalBytesAdded();
 		EvictToAddSize(SizeRequired);
 		FScopeLock ScopeLock(&Lock);
 		if (SizeInUse + SizeRequired <= MaxElementSize && Cache.Num() < MaxNumElements)
@@ -171,25 +415,26 @@ void FHTTPResponseCache::CacheEntity(TSharedPtrTS<FCacheItem> EntityToAdd)
 	}
 }
 
-FHTTPResponseCache::EScatterResult FHTTPResponseCache::GetScatteredCacheEntity(TSharedPtrTS<FCacheItem>& OutScatteredCachedEntity, const FString& URL, const IElectraHttpManager::FParams::FRange& InRange)
+FHTTPResponseCache::EScatterResult FHTTPResponseCache::GetScatteredCacheEntity(TSharedPtrTS<FCacheItem>& OutScatteredCachedEntity, const FString& URL, const ElectraHTTPStream::FHttpRange& InRange)
 {
 	FScopeLock lock(&Lock);
 	// Get a list of all cached blocks for this URL.
 	TArray<TSharedPtrTS<FCacheItem>> CachedBlocks;
 	for(auto &Entry : Cache)
 	{
-		if (Entry->URL.Equals(URL))
+		if (Entry->RequestedURL.Equals(URL) || Entry->EffectiveURL.Equals(URL))
 		{
 			CachedBlocks.Add(Entry);
 		}
 	}
 
-	IElectraHttpManager::FParams::FRange Range(InRange);
+	ElectraHTTPStream::FHttpRange Range(InRange);
 
 	// Set up the default output that encompasses the request.
 	// This is what needs to be fetched for cache misses.
 	OutScatteredCachedEntity = MakeSharedTS<IHTTPResponseCache::FCacheItem>();
-	OutScatteredCachedEntity->URL = URL;
+	OutScatteredCachedEntity->RequestedURL = URL;
+	OutScatteredCachedEntity->EffectiveURL = URL;
 	OutScatteredCachedEntity->Range = Range;
 
 	// Quick out if there are no cached blocks at all.
@@ -234,6 +479,7 @@ FHTTPResponseCache::EScatterResult FHTTPResponseCache::GetScatteredCacheEntity(T
 				// Data up to the beginning of this cached block must be fetched first.
 				OutScatteredCachedEntity->Range = Range;
 				OutScatteredCachedEntity->Range.SetEndIncluding(cbstart - 1);
+				OutScatteredCachedEntity->EffectiveURL = CachedBlocks[firstIndex]->EffectiveURL;
 				return EScatterResult::PartialHit;
 			}
 			else
@@ -251,19 +497,21 @@ FHTTPResponseCache::EScatterResult FHTTPResponseCache::GetScatteredCacheEntity(T
 	// Exact match?
 	if (rs == CachedBlocks[firstIndex]->Range.GetStart() && re == CachedBlocks[firstIndex]->Range.GetEndIncluding())
 	{
-		OutScatteredCachedEntity = CachedBlocks[firstIndex];
+		OutScatteredCachedEntity->Response = MakeShared<FWrappedCacheResponse, ESPMode::ThreadSafe>(CachedBlocks[firstIndex]->Response);
+		OutScatteredCachedEntity->EffectiveURL = CachedBlocks[firstIndex]->EffectiveURL;
 		Cache.Remove(CachedBlocks[firstIndex]);
 		Cache.Insert(CachedBlocks[firstIndex], 0);
 		return EScatterResult::FullHit;
 	}
 
 	// Set up the partial cache response.
-	TSharedPtr<FResponseFromCache, ESPMode::ThreadSafe> Response = MakeShared<FResponseFromCache, ESPMode::ThreadSafe>();
-	Response->SetURL(URL);
-	Response->CopyHeadersFromExceptContentSizes(CachedBlocks[firstIndex]->Response);
+	TSharedPtr<FSynthesizedCacheResponse, ESPMode::ThreadSafe> Response = MakeShared<FSynthesizedCacheResponse, ESPMode::ThreadSafe>();
+	Response->SetURL(CachedBlocks[firstIndex]->EffectiveURL);
+	OutScatteredCachedEntity->EffectiveURL = CachedBlocks[firstIndex]->EffectiveURL;
+	Response->CopyHeadersExceptContentSizesAndDateFrom(CachedBlocks[firstIndex]->Response);
+	Response->SetHTTPStatusLine(CachedBlocks[firstIndex]->Response->GetHTTPStatusLine());
 
 	// There is some overlap. Find how many consecutive bytes we can get from the cached blocks.
-	TArray<uint8> Payload;
 	TArray<TSharedPtrTS<FCacheItem>> TouchedBlocks;
 	while(1)
 	{
@@ -273,9 +521,8 @@ FHTTPResponseCache::EScatterResult FHTTPResponseCache::GetScatteredCacheEntity(T
 		int64 offset = rs - cbstart;
 		int64 numAvail = last + 1 - cbstart - offset;
 
-		// Copy the response data into the gathering block.
-		const TArray<uint8>& CachedContent = CachedBlocks[firstIndex]->Response->GetContent();
-		Payload.Append(CachedContent.GetData() + offset, numAvail);
+		// Add the bytes from the cache to the new response.
+		Response->GetResponseData().AddData(CachedBlocks[firstIndex]->Response->GetResponseData(), offset, numAvail);
 
 		// Remember which blocks we just touched.
 		TouchedBlocks.Add(CachedBlocks[firstIndex]);
@@ -310,13 +557,28 @@ FHTTPResponseCache::EScatterResult FHTTPResponseCache::GetScatteredCacheEntity(T
 	Range.SetEndIncluding(rs - 1);
 	if (InRange.IsSet())
 	{
-		Response->AddHeader(FString::Printf(TEXT("Content-Range: bytes %lld-%lld/%lld"), (long long int)Range.GetStart(), (long long int)Range.GetEndIncluding(), (long long int)Range.GetDocumentSize()));
+		Response->AddHeader(TEXT("Content-Range"), FString::Printf(TEXT("bytes %lld-%lld/%lld"), (long long int)Range.GetStart(), (long long int)Range.GetEndIncluding(), (long long int)Range.GetDocumentSize()));
 		Response->SetHTTPResponseCode(206);
+		FString StatusLine = Response->GetHTTPStatusLine();
+		StatusLine.ReplaceInline(TEXT("200"), TEXT("206"));
+		Response->SetHTTPStatusLine(StatusLine);
 	}
-	Response->AddHeader(FString::Printf(TEXT("Content-Length: %lld"), (long long int)Payload.Num()));
-	Response->SetPayload(MoveTemp(Payload));
+	Response->AddHeader(TEXT("Content-Length"), FString::Printf(TEXT("%lld"), (long long int)Response->GetResponseData().GetNumTotalBytesAdded()));
+	Response->GetResponseData().SetLengthFromResponseHeader(Response->GetResponseData().GetNumTotalBytesAdded());
+	Response->GetResponseData().SetEOS();
+
+/*
+	Response->SetTimeUntilNameResolved(double InTimeUntilNameResolved);
+	Response->SetTimeUntilConnected(double InTimeUntilConnected);
+	Response->SetTimeUntilRequestSent(double InTimeUntilRequestSent);
+	Response->SetTimeUntilHeadersAvailable(double InTimeUntilHeadersAvailable);
+	Response->SetTimeUntilFirstByte(double InTimeUntilFirstByte);
+	Response->SetTimeUntilFinished(double InTimeUntilFinished);
+*/
+
 	OutScatteredCachedEntity->Response = Response;
 	OutScatteredCachedEntity->Range = Range;
+
 	return EScatterResult::PartialHit;
 }
 
diff --git a/Engine/Plugins/Media/ElectraPlayer/Source/ElectraPlayerRuntime/Private/Runtime/HTTP/HTTPResponseCache.h b/Engine/Plugins/Media/ElectraPlayer/Source/ElectraPlayerRuntime/Private/Runtime/HTTP/HTTPResponseCache.h
--- a/Engine/Plugins/Media/ElectraPlayer/Source/ElectraPlayerRuntime/Private/Runtime/HTTP/HTTPResponseCache.h
+++ b/Engine/Plugins/Media/ElectraPlayer/Source/ElectraPlayerRuntime/Private/Runtime/HTTP/HTTPResponseCache.h
@@ -5,8 +5,8 @@
 #include "PlayerCore.h"
 #include "PlayerTime.h"
 #include "ParameterDictionary.h"
-#include "HTTP/HTTPManager.h"
-#include "Interfaces/IHttpRequest.h"
+#include "ElectraHTTPStream.h"
+#include "Utilities/HttpRangeHeader.h"
 
 namespace Electra
 {
@@ -17,10 +17,11 @@ class IHTTPResponseCache
 public:
 	struct FCacheItem
 	{
-		FString URL;
-		IElectraHttpManager::FParams::FRange Range;
+		FString RequestedURL;
+		FString EffectiveURL;
+		ElectraHTTPStream::FHttpRange Range;
 		FTimeValue ExpiresAtUTC;
-		FHttpResponsePtr Response;
+		IElectraHTTPStreamResponsePtr Response;
 	};
 
 	static TSharedPtrTS<IHTTPResponseCache> Create(IPlayerSessionServices* SessionServices, const FParamDict& Options);
@@ -53,7 +54,7 @@ public:
 	 * Returns false to indicate that the range specified in the (otherwise empty) cached response must be requested from the origin.
 	 * Returns true if a - possibly partial - cached response exists. The range that is cached is given in the cache item.
 	 */
-	virtual EScatterResult GetScatteredCacheEntity(TSharedPtrTS<FCacheItem>& OutScatteredCachedEntity, const FString& URL, const IElectraHttpManager::FParams::FRange& Range) = 0;
+	virtual EScatterResult GetScatteredCacheEntity(TSharedPtrTS<FCacheItem>& OutScatteredCachedEntity, const FString& URL, const ElectraHTTPStream::FHttpRange& Range) = 0;
 
 protected:
 	IHTTPResponseCache() = default;
diff --git a/Engine/Plugins/Media/ElectraPlayer/Source/ElectraPlayerRuntime/Private/Runtime/Player/AdaptiveStreamingPlayerManifest.cpp b/Engine/Plugins/Media/ElectraPlayer/Source/ElectraPlayerRuntime/Private/Runtime/Player/AdaptiveStreamingPlayerManifest.cpp
--- a/Engine/Plugins/Media/ElectraPlayer/Source/ElectraPlayerRuntime/Private/Runtime/Player/AdaptiveStreamingPlayerManifest.cpp
+++ b/Engine/Plugins/Media/ElectraPlayer/Source/ElectraPlayerRuntime/Private/Runtime/Player/AdaptiveStreamingPlayerManifest.cpp
@@ -95,7 +95,16 @@ void FAdaptiveStreamingPlayer::OnManifestGetMimeTypeComplete(TSharedPtrTS<FHTTPR
 			const HTTP::FConnectionInfo* ci = InRequest->GetConnectionInfo();
 			if (ci)
 			{
-				mimeType = ci->ContentType;
+				// Get the content type, if it exists.
+				if (ci->ContentType.Len())
+				{
+					mimeType = ci->ContentType;
+				}
+				// Use the effective URL after potential redirections.
+				if (ci->EffectiveURL.Len())
+				{
+					ManifestURL = ci->EffectiveURL;
+				}
 			}
 		}
 		WorkerThread.SendLoadManifestMessage(ManifestURL, mimeType);
diff --git a/Engine/Plugins/Media/ElectraPlayer/Source/ElectraPlayerRuntime/Private/Runtime/Player/AdaptiveStreamingPlayerResourceRequest.cpp b/Engine/Plugins/Media/ElectraPlayer/Source/ElectraPlayerRuntime/Private/Runtime/Player/AdaptiveStreamingPlayerResourceRequest.cpp
--- a/Engine/Plugins/Media/ElectraPlayer/Source/ElectraPlayerRuntime/Private/Runtime/Player/AdaptiveStreamingPlayerResourceRequest.cpp
+++ b/Engine/Plugins/Media/ElectraPlayer/Source/ElectraPlayerRuntime/Private/Runtime/Player/AdaptiveStreamingPlayerResourceRequest.cpp
@@ -124,7 +124,7 @@ void FHTTPResourceRequest::HTTPCompletionCallback(const IElectraHttpManager::FRe
 					}
 					else
 					{
-						Error = ConnInfo.StatusInfo.HTTPStatus;
+						Error = ConnInfo.StatusInfo.HTTPStatus ? ConnInfo.StatusInfo.HTTPStatus : 4;
 					}
 				}
 
diff --git a/Engine/Plugins/Media/ElectraPlayer/Source/ElectraPlayerRuntime/Private/Runtime/Player/AdaptiveStreamingPlayerResourceRequest.h b/Engine/Plugins/Media/ElectraPlayer/Source/ElectraPlayerRuntime/Private/Runtime/Player/AdaptiveStreamingPlayerResourceRequest.h
--- a/Engine/Plugins/Media/ElectraPlayer/Source/ElectraPlayerRuntime/Private/Runtime/Player/AdaptiveStreamingPlayerResourceRequest.h
+++ b/Engine/Plugins/Media/ElectraPlayer/Source/ElectraPlayerRuntime/Private/Runtime/Player/AdaptiveStreamingPlayerResourceRequest.h
@@ -9,6 +9,8 @@
 
 namespace Electra
 {
+class IPlayerSessionServices;
+
 
 class IAdaptiveStreamingPlayerResourceRequest : public TSharedFromThis<IAdaptiveStreamingPlayerResourceRequest, ESPMode::ThreadSafe>
 {
diff --git a/Engine/Plugins/Media/ElectraPlayer/Source/ElectraPlayerRuntime/Private/Runtime/Player/DASH/ManifestDASH.cpp b/Engine/Plugins/Media/ElectraPlayer/Source/ElectraPlayerRuntime/Private/Runtime/Player/DASH/ManifestDASH.cpp
--- a/Engine/Plugins/Media/ElectraPlayer/Source/ElectraPlayerRuntime/Private/Runtime/Player/DASH/ManifestDASH.cpp
+++ b/Engine/Plugins/Media/ElectraPlayer/Source/ElectraPlayerRuntime/Private/Runtime/Player/DASH/ManifestDASH.cpp
@@ -1721,7 +1721,7 @@ FManifestDASHInternal::FRepresentation::ESearchResult FManifestDASHInternal::FRe
 		}
 		if (!IndexRange.IsEmpty())
 		{
-			IElectraHttpManager::FParams::FRange r;
+			ElectraHTTPStream::FHttpRange r;
 			r.Set(IndexRange);
 			if (r.IsSet())
 			{
@@ -1781,7 +1781,7 @@ FManifestDASHInternal::FRepresentation::ESearchResult FManifestDASHInternal::FRe
 				LoadReq->Range = IndexRange;
 				if (RequestHeader.Len())
 				{
-					LoadReq->Headers.Emplace(HTTP::FHTTPHeader({DASH::HTTPHeaderOptionName, RequestHeader}));
+					LoadReq->Headers.Emplace(HTTP::FHTTPHeader(DASH::HTTPHeaderOptionName, RequestHeader));
 				}
 				LoadReq->PlayerSessionServices = InPlayerSessionServices;
 				LoadReq->XLinkElement = MPDRepresentation;
@@ -1850,7 +1850,7 @@ bool FManifestDASHInternal::FRepresentation::PrepareDownloadURLs(IPlayerSessionS
 
 	if (InOutSegmentInfo.FirstByteOffset && InOutSegmentInfo.NumberOfBytes)
 	{
-		IElectraHttpManager::FParams::FRange r;
+		ElectraHTTPStream::FHttpRange r;
 		r.SetStart(InOutSegmentInfo.FirstByteOffset);
 		r.SetEndIncluding(InOutSegmentInfo.FirstByteOffset + InOutSegmentInfo.NumberOfBytes - 1);
 		InOutSegmentInfo.MediaURL.Range = r.GetString();
@@ -1946,7 +1946,7 @@ bool FManifestDASHInternal::FRepresentation::PrepareDownloadURLs(IPlayerSessionS
 
 	if (InOutSegmentInfo.FirstByteOffset && InOutSegmentInfo.NumberOfBytes)
 	{
-		IElectraHttpManager::FParams::FRange r;
+		ElectraHTTPStream::FHttpRange r;
 		r.SetStart(InOutSegmentInfo.FirstByteOffset);
 		r.SetEndIncluding(InOutSegmentInfo.FirstByteOffset + InOutSegmentInfo.NumberOfBytes - 1);
 		InOutSegmentInfo.MediaURL.Range = r.GetString();
diff --git a/Engine/Plugins/Media/ElectraPlayer/Source/ElectraPlayerRuntime/Private/Runtime/Player/DASH/OptionKeynamesDASH.h b/Engine/Plugins/Media/ElectraPlayer/Source/ElectraPlayerRuntime/Private/Runtime/Player/DASH/OptionKeynamesDASH.h
--- a/Engine/Plugins/Media/ElectraPlayer/Source/ElectraPlayerRuntime/Private/Runtime/Player/DASH/OptionKeynamesDASH.h
+++ b/Engine/Plugins/Media/ElectraPlayer/Source/ElectraPlayerRuntime/Private/Runtime/Player/DASH/OptionKeynamesDASH.h
@@ -9,12 +9,17 @@ namespace Electra
 namespace DASH
 {
 
-const TCHAR* const OptionKeyMPDLoadConnectTimeout = TEXT("mpd_connection_timeout");				//!< (FTimeValue) value specifying connection timeout fetching the MPD
-const TCHAR* const OptionKeyMPDLoadNoDataTimeout = TEXT("mpd_nodata_timeout");					//!< (FTimeValue) value specifying no-data timeout fetching the MPD
-const TCHAR* const OptionKeyMPDReloadConnectTimeout = TEXT("mpd_update_connection_timeout");	//!< (FTimeValue) value specifying connection timeout fetching the MPD repeatedly
-const TCHAR* const OptionKeyMPDReloadNoDataTimeout = TEXT("mpd_update_connection_timeout");		//!< (FTimeValue) value specifying no-data timeout fetching the MPD repeatedly
+const TCHAR* const OptionKeyMPDLoadConnectTimeout = TEXT("mpd_connection_timeout");					//!< (FTimeValue) value specifying connection timeout fetching the MPD
+const TCHAR* const OptionKeyMPDLoadNoDataTimeout = TEXT("mpd_nodata_timeout");						//!< (FTimeValue) value specifying no-data timeout fetching the MPD
+const TCHAR* const OptionKeyMPDReloadConnectTimeout = TEXT("mpd_update_connection_timeout");		//!< (FTimeValue) value specifying connection timeout fetching the MPD repeatedly
+const TCHAR* const OptionKeyMPDReloadNoDataTimeout = TEXT("mpd_update_nodata_timeout");				//!< (FTimeValue) value specifying no-data timeout fetching the MPD repeatedly
 
-const TCHAR* const OptionKey_MinTimeBetweenMPDUpdates = TEXT("dash:min_mpd_update_interval");	//!< (FTimeValue) value limiting the time between MPD updates unless these are required updates.
+const TCHAR* const OptionKeyInitSegmentConnectTimeout = TEXT("initsegment_connection_timeout");		//!< (FTimeValue) value specifying connection timeout fetching an init segment
+const TCHAR* const OptionKeyInitSegmentNoDataTimeout = TEXT("initsegment_nodata_timeout");			//!< (FTimeValue) value specifying no-data timeout fetching an init segment
+const TCHAR* const OptionKeyMediaSegmentConnectTimeout = TEXT("mediasegment_connection_timeout");	//!< (FTimeValue) value specifying connection timeout fetching a media segment
+const TCHAR* const OptionKeyMediaSegmentNoDataTimeout = TEXT("mediasegment_nodata_timeout");		//!< (FTimeValue) value specifying no-data timeout fetching a media segment
+
+const TCHAR* const OptionKey_MinTimeBetweenMPDUpdates = TEXT("dash:min_mpd_update_interval");		//!< (FTimeValue) value limiting the time between MPD updates unless these are required updates.
 
 const TCHAR* const OptionKey_CurrentCDN = TEXT("dash:current_cdn");			//!< A string giving the currently selected CDN to be used with <BaseURL@serviceLocation> elements.
 
diff --git a/Engine/Plugins/Media/ElectraPlayer/Source/ElectraPlayerRuntime/Private/Runtime/Player/DASH/PlaylistReaderDASH.cpp b/Engine/Plugins/Media/ElectraPlayer/Source/ElectraPlayerRuntime/Private/Runtime/Player/DASH/PlaylistReaderDASH.cpp
--- a/Engine/Plugins/Media/ElectraPlayer/Source/ElectraPlayerRuntime/Private/Runtime/Player/DASH/PlaylistReaderDASH.cpp
+++ b/Engine/Plugins/Media/ElectraPlayer/Source/ElectraPlayerRuntime/Private/Runtime/Player/DASH/PlaylistReaderDASH.cpp
@@ -505,13 +505,13 @@ void FPlaylistReaderDASH::CheckForMPDUpdate()
 				PlaylistLoadRequest->URL = URL;
 				if (RequestHeader.Len())
 				{
-					PlaylistLoadRequest->Headers.Emplace(HTTP::FHTTPHeader({DASH::HTTPHeaderOptionName, RequestHeader}));
+					PlaylistLoadRequest->Headers.Emplace(HTTP::FHTTPHeader(DASH::HTTPHeaderOptionName, RequestHeader));
 				}
 #ifndef DO_NOT_PERFORM_CONDITIONAL_GET
 				// As per the standard we should perform conditional GET requests.
 				if (ETag.Len())
 				{
-					PlaylistLoadRequest->Headers.Emplace(HTTP::FHTTPHeader({TEXT("If-None-Match"), ETag}));
+					PlaylistLoadRequest->Headers.Emplace(HTTP::FHTTPHeader(TEXT("If-None-Match"), ETag));
 				}
 #endif
 				PlaylistLoadRequest->CompleteCallback.BindThreadSafeSP(AsShared(), &FPlaylistReaderDASH::ManifestUpdateDownloadCompleted);
@@ -901,7 +901,7 @@ void FPlaylistReaderDASH::SetupRequestTimeouts(FResourceLoadRequestPtr InRequest
 		{
 			InRequest->Request->ConnectionTimeout(FTimeValue(FTimeValue::MillisecondsToHNS(1000 * 5)));
 			InRequest->Request->NoDataTimeout(FTimeValue(FTimeValue::MillisecondsToHNS(1000 * 2)));
-			//InRequest->Request->AcceptEncoding(TEXT("identity"));
+			InRequest->Request->AcceptEncoding(TEXT("identity"));
 			break;
 		}
 		case FMPDLoadRequestDASH::ELoadType::XLink_Period:
diff --git a/Engine/Plugins/Media/ElectraPlayer/Source/ElectraPlayerRuntime/Private/Runtime/Player/DASH/StreamReaderFMP4DASH.cpp b/Engine/Plugins/Media/ElectraPlayer/Source/ElectraPlayerRuntime/Private/Runtime/Player/DASH/StreamReaderFMP4DASH.cpp
--- a/Engine/Plugins/Media/ElectraPlayer/Source/ElectraPlayerRuntime/Private/Runtime/Player/DASH/StreamReaderFMP4DASH.cpp
+++ b/Engine/Plugins/Media/ElectraPlayer/Source/ElectraPlayerRuntime/Private/Runtime/Player/DASH/StreamReaderFMP4DASH.cpp
@@ -943,6 +943,10 @@ void FStreamReaderFMP4DASH::FStreamHandler::HandleRequest()
 			{
 				HTTP->Parameters.RequestHeaders.Emplace(HTTP::FHTTPHeader({DASH::HTTPHeaderOptionName, Request->Segment.MediaURL.CustomHeader}));
 			}
+			// Set timeouts for media segment retrieval
+			const FParamDict& Options = PlayerSessionService->GetOptions();
+			HTTP->Parameters.ConnectTimeout = Options.GetValue(DASH::OptionKeyMediaSegmentConnectTimeout).SafeGetTimeValue(FTimeValue().SetFromMilliseconds(1000 * 4));
+			HTTP->Parameters.NoDataTimeout = Options.GetValue(DASH::OptionKeyMediaSegmentNoDataTimeout).SafeGetTimeValue(FTimeValue().SetFromMilliseconds(1000 * 4));
 
 			ProgressReportCount = 0;
 			DownloadCompleteSignal.Reset();
diff --git a/Engine/Plugins/Media/ElectraPlayer/Source/ElectraPlayerRuntime/Private/Runtime/Player/HLS/StreamReaderHLSfmp4.h b/Engine/Plugins/Media/ElectraPlayer/Source/ElectraPlayerRuntime/Private/Runtime/Player/HLS/StreamReaderHLSfmp4.h
--- a/Engine/Plugins/Media/ElectraPlayer/Source/ElectraPlayerRuntime/Private/Runtime/Player/HLS/StreamReaderHLSfmp4.h
+++ b/Engine/Plugins/Media/ElectraPlayer/Source/ElectraPlayerRuntime/Private/Runtime/Player/HLS/StreamReaderHLSfmp4.h
@@ -40,7 +40,7 @@ public:
 	virtual void GetDownloadStats(Metrics::FSegmentDownloadStats& OutStats) const override;
 
 	FString																		URL;
-	IElectraHttpManager::FParams::FRange										Range;
+	ElectraHTTPStream::FHttpRange												Range;
 
 	EStreamType																	StreamType;							//!< Type of stream (video, audio, etc.)
 	uint32																		StreamUniqueID;						//!< The unique stream ID identifying the stream for which this is a request.
diff --git a/Engine/Plugins/Media/ElectraPlayer/Source/ElectraPlayerRuntime/Private/Runtime/Player/mp4/OptionKeynamesMP4.h b/Engine/Plugins/Media/ElectraPlayer/Source/ElectraPlayerRuntime/Private/Runtime/Player/mp4/OptionKeynamesMP4.h
new file mode 100644
--- /dev/null
+++ b/Engine/Plugins/Media/ElectraPlayer/Source/ElectraPlayerRuntime/Private/Runtime/Player/mp4/OptionKeynamesMP4.h
@@ -0,0 +1,19 @@
+// Copyright Epic Games, Inc. All Rights Reserved.
+
+#pragma once
+
+#include "PlayerCore.h"
+
+namespace Electra
+{
+namespace MP4
+{
+
+const TCHAR* const OptionKeyMP4LoadConnectTimeout = TEXT("mp4_connection_timeout");				//!< (FTimeValue) value specifying connection timeout fetching the MP4
+const TCHAR* const OptionKeyMP4LoadNoDataTimeout = TEXT("mp4_nodata_timeout");					//!< (FTimeValue) value specifying no-data timeout fetching the MP4
+
+}
+
+} // namespace Electra
+
+
diff --git a/Engine/Plugins/Media/ElectraPlayer/Source/ElectraPlayerRuntime/Private/Runtime/Player/mp4/PlaylistReaderMP4.cpp b/Engine/Plugins/Media/ElectraPlayer/Source/ElectraPlayerRuntime/Private/Runtime/Player/mp4/PlaylistReaderMP4.cpp
--- a/Engine/Plugins/Media/ElectraPlayer/Source/ElectraPlayerRuntime/Private/Runtime/Player/mp4/PlaylistReaderMP4.cpp
+++ b/Engine/Plugins/Media/ElectraPlayer/Source/ElectraPlayerRuntime/Private/Runtime/Player/mp4/PlaylistReaderMP4.cpp
@@ -13,7 +13,7 @@
 #include "Player/AdaptiveStreamingPlayerResourceRequest.h"
 #include "Player/PlayerStreamReader.h"
 #include "Player/mp4/ManifestMP4.h"
-
+#include "Player/mp4/OptionKeynamesMP4.h"
 
 #define ERRCODE_MP4_INVALID_FILE							1
 
@@ -276,13 +276,18 @@ void FPlaylistReaderMP4::HTTPCompletionCallback(const IElectraHttpManager::FRequ
 		// Set the size of the resource if we don't have it yet.
 		if (FileSize < 0)
 		{
-			IElectraHttpManager::FParams::FRange crh;
+			ElectraHTTPStream::FHttpRange crh;
 			if (crh.ParseFromContentRangeResponse(InRequest->ConnectionInfo.ContentRangeHeader))
 			{
 				FileSize = crh.GetDocumentSize();
 			}
 		}
 
+		if (ConnectionInfo.EffectiveURL.Len())
+		{
+			MasterPlaylistURL = ConnectionInfo.EffectiveURL;
+		}
+
 		// Copy the read data across.
 		int32 NumRead = ReceiveBuffer->Buffer.Num();
 		Buffer.PushData(ReceiveBuffer->Buffer.GetLinearReadData(), NumRead);
@@ -326,6 +331,8 @@ void FPlaylistReaderMP4::ReadNextChunk(int64 InFromOffset, int64 ChunkSize)
 	ReceiveBuffer = MakeSharedTS<IElectraHttpManager::FReceiveBuffer>();
 	ReceiveBuffer->Buffer.Reserve(ChunkSize);
 
+	const FParamDict& Options = PlayerSessionServices->GetOptions();
+
 	Request = MakeSharedTS<IElectraHttpManager::FRequest>();
 	Request->Parameters.URL = MasterPlaylistURL;
 	Request->Parameters.Range.SetStart(InFromOffset);
@@ -335,6 +342,8 @@ void FPlaylistReaderMP4::ReadNextChunk(int64 InFromOffset, int64 ChunkSize)
 		LastByte = FileSize - 1;
 	}
 	Request->Parameters.Range.SetEndIncluding(LastByte);
+	Request->Parameters.ConnectTimeout = Options.GetValue(MP4::OptionKeyMP4LoadConnectTimeout).SafeGetTimeValue(FTimeValue().SetFromMilliseconds(1000 * 8));
+	Request->Parameters.NoDataTimeout = Options.GetValue(MP4::OptionKeyMP4LoadNoDataTimeout).SafeGetTimeValue(FTimeValue().SetFromMilliseconds(1000 * 6));
 	Request->ReceiveBuffer = ReceiveBuffer;
 	Request->ProgressListener = ProgressListener;
 	Request->ResponseCache = PlayerSessionServices->GetHTTPResponseCache();
diff --git a/Engine/Plugins/Media/ElectraPlayer/Source/ElectraPlayerRuntime/Private/Runtime/Player/mp4/StreamReaderMP4.cpp b/Engine/Plugins/Media/ElectraPlayer/Source/ElectraPlayerRuntime/Private/Runtime/Player/mp4/StreamReaderMP4.cpp
--- a/Engine/Plugins/Media/ElectraPlayer/Source/ElectraPlayerRuntime/Private/Runtime/Player/mp4/StreamReaderMP4.cpp
+++ b/Engine/Plugins/Media/ElectraPlayer/Source/ElectraPlayerRuntime/Private/Runtime/Player/mp4/StreamReaderMP4.cpp
@@ -13,6 +13,7 @@
 #include "Player/AdaptiveStreamingPlayerABR.h"
 #include "Player/mp4/ManifestMP4.h"
 #include "Player/mp4/StreamReaderMP4.h"
+#include "Player/mp4/OptionKeynamesMP4.h"
 #include "Player/PlayerSessionServices.h"
 #include "Player/PlayerStreamFilter.h"
 
@@ -399,35 +400,21 @@ void FStreamReaderMP4::HandleRequest()
 	ReadBuffer.ReceiveBuffer->Buffer.Reserve(4 << 20);
 	ReadBuffer.ReceiveBuffer->bEnableRingbuffer = true;
 	ReadBuffer.SetCurrentPos(Request->FileStartOffset);
+
+	const FParamDict& Options = PlayerSessionServices->GetOptions();
+
 	TSharedPtrTS<IElectraHttpManager::FRequest> HTTP(new IElectraHttpManager::FRequest);
 	HTTP->Parameters.URL				= TimelineAsset->GetMediaURL();
 	HTTP->Parameters.Range.Start		= Request->FileStartOffset;
 	HTTP->Parameters.Range.EndIncluding = Request->FileEndOffset;
+	// No compression as this would not yield much with already compressed video/audio data.
+	HTTP->Parameters.AcceptEncoding.Set(TEXT("identity"));
+	// Timeouts
+	HTTP->Parameters.ConnectTimeout = Options.GetValue(MP4::OptionKeyMP4LoadConnectTimeout).SafeGetTimeValue(FTimeValue().SetFromMilliseconds(1000 * 8));
+	HTTP->Parameters.NoDataTimeout = Options.GetValue(MP4::OptionKeyMP4LoadNoDataTimeout).SafeGetTimeValue(FTimeValue().SetFromMilliseconds(1000 * 6));
+
 	// Explicit range?
 	int64 NumRequestedBytes = HTTP->Parameters.Range.GetNumberOfBytes();
-	int32 SubRequestSize = 0;
-	if (NumRequestedBytes > 0)
-	{
-		if (Request->bIsFirstSegment && !Request->bIsLastSegment)
-		{
-			SubRequestSize = 512 << 10;
-		}
-		else if (Request->bIsFirstSegment && Request->bIsLastSegment)
-		{
-			SubRequestSize = 2 << 20;
-		}
-	}
-	else
-	{
-		if (Request->SegmentInternalSize < 0)
-		{
-			SubRequestSize = 2 << 20;
-		}
-	}
-	if (SubRequestSize)
-	{
-		HTTP->Parameters.SubRangeRequestSize = SubRequestSize;
-	}
 
 	HTTP->ReceiveBuffer = ReadBuffer.ReceiveBuffer;
 	HTTP->ProgressListener = ProgressListener;
