ElectraPlayer: Keeping currently active playback range with the seek request to prevent issues if the range is changed immediately after seeking which could put the seek position outside the newly set range.

#rnx

diff --git a/Engine/Plugins/Media/ElectraPlayer/Source/ElectraPlayerRuntime/Private/Runtime/Player/AdaptiveStreamingPlayer.cpp b/Engine/Plugins/Media/ElectraPlayer/Source/ElectraPlayerRuntime/Private/Runtime/Player/AdaptiveStreamingPlayer.cpp
--- a/Engine/Plugins/Media/ElectraPlayer/Source/ElectraPlayerRuntime/Private/Runtime/Player/AdaptiveStreamingPlayer.cpp
+++ b/Engine/Plugins/Media/ElectraPlayer/Source/ElectraPlayerRuntime/Private/Runtime/Player/AdaptiveStreamingPlayer.cpp
@@ -2506,24 +2506,7 @@ void FAdaptiveStreamingPlayer::InternalHandlePendingStartRequest(const FTimeValu
 							}
 						}
 					}
-
-					// Clamp to within any set playback range
-					FTimeValue RangeStart = GetOptions().GetValue(OptionPlayRangeStart).SafeGetTimeValue(FTimeValue());
-					FTimeValue RangeEnd = GetOptions().GetValue(OptionPlayRangeEnd).SafeGetTimeValue(FTimeValue());
-					if (RangeStart.IsValid() && PendingStartRequest->StartAt.Time < RangeStart)
-					{
-						PendingStartRequest->StartAt.Time = RangeStart;
-					}
-					if (RangeEnd.IsValid() && PendingStartRequest->StartAt.Time > RangeEnd)
-					{
-						PendingStartRequest->StartAt.Time = RangeEnd;
-						// Going to the end will of course result in an 'end reached' and end of playback.
-						// If the player is set to loop then we should start from the beginning.
-						if (CurrentLoopParam.bEnableLooping)
-						{
-							PendingStartRequest->StartAt.Time = RangeStart;
-						}
-					}
+					ClampStartRequestTime();
 
 					IManifest::FResult Result = Manifest->FindPlayPeriod(InitialPlayPeriod, PendingStartRequest->StartAt, PendingStartRequest->SearchType);
 					switch(Result.GetType())
@@ -3239,13 +3222,15 @@ void FAdaptiveStreamingPlayer::HandlePendingMediaSegmentRequests()
 			IManifest::FResult Result;
 			if (!FinishedReq.bStartOver && !FinishedReq.bPlayPosAutoReselect)
 			{
+				FPlayStartOptions Options;
+				SetPlaystartOptions(Options);
 				if (Action == IAdaptiveStreamSelector::ESegmentAction::FetchNext)
 				{
-					Result = SegmentPlayPeriod->GetNextSegment(NextSegment, FinishedReq.Request);
+					Result = SegmentPlayPeriod->GetNextSegment(NextSegment, FinishedReq.Request, Options);
 				}
 				else if (Action == IAdaptiveStreamSelector::ESegmentAction::Retry || Action == IAdaptiveStreamSelector::ESegmentAction::Fill)
 				{
-					Result = SegmentPlayPeriod->GetRetrySegment(NextSegment, FinishedReq.Request, Action == IAdaptiveStreamSelector::ESegmentAction::Fill);
+					Result = SegmentPlayPeriod->GetRetrySegment(NextSegment, FinishedReq.Request, Options, Action == IAdaptiveStreamSelector::ESegmentAction::Fill);
 				}
 			}
 			else
@@ -3527,6 +3512,18 @@ void FAdaptiveStreamingPlayer::InternalDeselectStream(EStreamType StreamType)
 
 void FAdaptiveStreamingPlayer::InternalStartoverAtCurrentPosition()
 {
+	{
+		// Check if a seek is currently pending. If so then let us perform the seek instead.
+		FScopeLock lock(&SeekVars.Lock);
+		// Is there a pending request?
+		if (SeekVars.PendingRequest.IsSet())
+		{
+			// To ensure the seek will execute we clear out the active and last finished requests.
+			SeekVars.ActiveRequest.Reset();
+			SeekVars.LastFinishedRequest.Reset();
+			return;
+		}
+	}
 	FSeekParam NewPosition;
 	NewPosition.Time = PlaybackState.GetPlayPosition();
 	InternalStop(PlayerConfig.bHoldLastFrameDuringSeek);
@@ -3626,6 +3623,7 @@ void FAdaptiveStreamingPlayer::InternalHandleSegmentTrackChanges(const FTimeValu
 					FPendingSegmentRequest NextReq;
 					NextReq.bStartOver = true;
 					NextReq.StartoverPosition.Time = PlaybackState.GetPlayPosition();
+					SetPlaystartOptions(NextReq.StartoverPosition.Options);
 					NextReq.StreamType = EStreamType::Audio;
 					NextPendingSegmentRequests.Enqueue(MoveTemp(NextReq));
 
@@ -3653,6 +3651,7 @@ void FAdaptiveStreamingPlayer::InternalHandleSegmentTrackChanges(const FTimeValu
 					FPendingSegmentRequest NextReq;
 					NextReq.bStartOver = true;
 					NextReq.StartoverPosition.Time = PlaybackState.GetPlayPosition();
+					SetPlaystartOptions(NextReq.StartoverPosition.Options);
 					NextReq.StreamType = EStreamType::Subtitle;
 					NextPendingSegmentRequests.Enqueue(MoveTemp(NextReq));
 
@@ -3850,6 +3849,41 @@ TSharedPtrTS<FBufferSourceInfo> FAdaptiveStreamingPlayer::GetStreamBufferInfoAtT
 }
 
 
+void FAdaptiveStreamingPlayer::SetPlaystartOptions(FPlayStartOptions& OutOptions)
+{
+	FTimeValue RangeStart = GetOptions().GetValue(OptionPlayRangeStart).SafeGetTimeValue(FTimeValue());
+	FTimeValue RangeEnd = GetOptions().GetValue(OptionPlayRangeEnd).SafeGetTimeValue(FTimeValue());
+
+	OutOptions.PlaybackRange.Start = RangeStart.IsValid() ? RangeStart : FTimeValue::GetZero();
+	OutOptions.PlaybackRange.End = RangeEnd.IsValid() ? RangeEnd : FTimeValue::GetPositiveInfinity();
+
+	OutOptions.bFrameAccuracy = GetOptions().GetValue(OptionKeyFrameAccurateSeek).SafeGetBool(false);
+}
+
+void FAdaptiveStreamingPlayer::ClampStartRequestTime()
+{
+	if (PendingStartRequest.IsValid())
+	{
+		// Clamp to within any set playback range
+		FTimeValue RangeStart = PendingStartRequest->StartAt.Options.PlaybackRange.Start;
+		FTimeValue RangeEnd = PendingStartRequest->StartAt.Options.PlaybackRange.End;
+		if (PendingStartRequest->StartAt.Time < RangeStart)
+		{
+			PendingStartRequest->StartAt.Time = RangeStart;
+		}
+		else if (PendingStartRequest->StartAt.Time > RangeEnd)
+		{
+			PendingStartRequest->StartAt.Time = RangeEnd;
+			// Going to the end will of course result in an 'end reached' and end of playback.
+			// If the player is set to loop then we should start from the beginning.
+			if (CurrentLoopParam.bEnableLooping)
+			{
+				PendingStartRequest->StartAt.Time = RangeStart;
+			}
+		}
+	}
+}
+
 FTimeValue FAdaptiveStreamingPlayer::ClampTimeToCurrentRange(const FTimeValue& InTime, bool bClampToStart, bool bClampToEnd)
 {
 	FTimeValue Out(InTime);
@@ -4143,7 +4177,7 @@ void FAdaptiveStreamingPlayer::InternalStartAt(const FSeekParam& NewPosition)
 	PendingStartRequest->StartAt.Time = NewPosition.Time;
 	PendingStartRequest->StartingBitrate = NewPosition.StartingBitrate;
 	PendingStartRequest->StartType = SeekVars.bIsPlayStart ? FPendingStartRequest::EStartType::PlayStart : FPendingStartRequest::EStartType::Seeking;
-
+	SetPlaystartOptions(PendingStartRequest->StartAt.Options);
 
 	// The fragment should be the closest to the time stamp unless we are rebuffering in which case it should be the one we failed on.
 /*
diff --git a/Engine/Plugins/Media/ElectraPlayer/Source/ElectraPlayerRuntime/Private/Runtime/Player/AdaptiveStreamingPlayerInternal.h b/Engine/Plugins/Media/ElectraPlayer/Source/ElectraPlayerRuntime/Private/Runtime/Player/AdaptiveStreamingPlayerInternal.h
--- a/Engine/Plugins/Media/ElectraPlayer/Source/ElectraPlayerRuntime/Private/Runtime/Player/AdaptiveStreamingPlayerInternal.h
+++ b/Engine/Plugins/Media/ElectraPlayer/Source/ElectraPlayerRuntime/Private/Runtime/Player/AdaptiveStreamingPlayerInternal.h
@@ -180,7 +180,7 @@ struct FPlaybackState
 		bIsBuffering = false;
 		bIsPlaying = false;
 		bIsPaused = false;
-		bPlayrateHasChanged = false;
+		bPlayrangeHasChanged = false;
 		bLoopStateHasChanged = false;
 	}
 
@@ -196,7 +196,7 @@ struct FPlaybackState
 	bool									bIsBuffering;
 	bool									bIsPlaying;
 	bool									bIsPaused;
-	bool									bPlayrateHasChanged;
+	bool									bPlayrangeHasChanged;
 	bool									bLoopStateHasChanged;
 	TArray<FTrackMetadata>					VideoTracks;
 	TArray<FTrackMetadata>					AudioTracks;
@@ -385,12 +385,12 @@ struct FPlaybackState
 	void SetPlayRangeHasChanged(bool bWasChanged)
 	{
 		FMediaCriticalSection::ScopedLock lock(Lock);
-		bPlayrateHasChanged = bWasChanged;
+		bPlayrangeHasChanged = bWasChanged;
 	}
 	bool GetPlayRangeHasChanged()
 	{
 		FMediaCriticalSection::ScopedLock lock(Lock);
-		return bPlayrateHasChanged;
+		return bPlayrangeHasChanged;
 	}
 	void SetLoopStateHasChanged(bool bWasChanged)
 	{
@@ -1766,6 +1766,8 @@ private:
 	bool HaveEnoughBufferedDataToStartPlayback();
 	void PrepareForPrerolling();
 
+	void SetPlaystartOptions(FPlayStartOptions& OutOptions);
+	void ClampStartRequestTime();
 	FTimeValue ClampTimeToCurrentRange(const FTimeValue& InTime, bool bClampToStart, bool bClampToEnd);
 
 	TSharedPtrTS<FMultiTrackAccessUnitBuffer> GetStreamBuffer(EStreamType InStreamType, const TSharedPtrTS<FStreamDataBuffers>& InFromStreamBuffers)
diff --git a/Engine/Plugins/Media/ElectraPlayer/Source/ElectraPlayerRuntime/Private/Runtime/Player/DASH/ManifestDASH.cpp b/Engine/Plugins/Media/ElectraPlayer/Source/ElectraPlayerRuntime/Private/Runtime/Player/DASH/ManifestDASH.cpp
--- a/Engine/Plugins/Media/ElectraPlayer/Source/ElectraPlayerRuntime/Private/Runtime/Player/DASH/ManifestDASH.cpp
+++ b/Engine/Plugins/Media/ElectraPlayer/Source/ElectraPlayerRuntime/Private/Runtime/Player/DASH/ManifestDASH.cpp
@@ -123,8 +123,8 @@ public:
 	virtual IManifest::FResult GetStartingSegment(TSharedPtrTS<IStreamSegment>& OutSegment, const FPlayerSequenceState& InSequenceState, const FPlayStartPosition& StartPosition, IManifest::ESearchType SearchType) override;
 	virtual IManifest::FResult GetContinuationSegment(TSharedPtrTS<IStreamSegment>& OutSegment, EStreamType StreamType, const FPlayerSequenceState& SequenceState, const FPlayStartPosition& StartPosition, IManifest::ESearchType SearchType) override;
 	virtual IManifest::FResult GetLoopingSegment(TSharedPtrTS<IStreamSegment>& OutSegment, const FPlayerSequenceState& SequenceState, const TMultiMap<EStreamType, TSharedPtrTS<IStreamSegment>>& InFinishedSegments, const FPlayStartPosition& StartPosition, IManifest::ESearchType SearchType) override;
-	virtual IManifest::FResult GetNextSegment(TSharedPtrTS<IStreamSegment>& OutSegment, TSharedPtrTS<const IStreamSegment> CurrentSegment) override;
-	virtual IManifest::FResult GetRetrySegment(TSharedPtrTS<IStreamSegment>& OutSegment, TSharedPtrTS<const IStreamSegment> CurrentSegment, bool bReplaceWithFillerData) override;
+	virtual IManifest::FResult GetNextSegment(TSharedPtrTS<IStreamSegment>& OutSegment, TSharedPtrTS<const IStreamSegment> CurrentSegment, const FPlayStartOptions& Options) override;
+	virtual IManifest::FResult GetRetrySegment(TSharedPtrTS<IStreamSegment>& OutSegment, TSharedPtrTS<const IStreamSegment> CurrentSegment, const FPlayStartOptions& Options, bool bReplaceWithFillerData) override;
 	virtual void IncreaseSegmentFetchDelay(const FTimeValue& IncreaseAmount) override;
 	virtual void GetSegmentInformation(TArray<FSegmentInformation>& OutSegmentInformation, FTimeValue& OutAverageSegmentDuration, TSharedPtrTS<const IStreamSegment> CurrentSegment, const FTimeValue& LookAheadTime, const FString& AdaptationSetID, const FString& RepresentationID) override;
 
@@ -151,7 +151,7 @@ private:
 		SamePeriodStartOver,
 		NextPeriod,
 	};
-	IManifest::FResult GetNextOrRetrySegment(TSharedPtrTS<IStreamSegment>& OutSegment, TSharedPtrTS<const IStreamSegment> InCurrentSegment, ENextSegType InNextType);
+	IManifest::FResult GetNextOrRetrySegment(TSharedPtrTS<IStreamSegment>& OutSegment, TSharedPtrTS<const IStreamSegment> InCurrentSegment, ENextSegType InNextType, const FPlayStartOptions& Options);
 
 	bool PrepareDRM(const TArray<FManifestDASHInternal::FAdaptationSet::FContentProtection>& InContentProtections);
 
@@ -335,8 +335,8 @@ IManifest::FResult FManifestDASH::FindPlayPeriod(TSharedPtrTS<IPlayPeriod>& OutP
 		return IManifest::FResult(IManifest::FResult::EType::TryAgainLater).RetryAfterMilliseconds(1000);
 	}
 
-	// Get the end of the playback end, if any.
-	FTimeValue PlayRangeEnd = PlayerSessionServices->GetOptions().GetValue(OptionPlayRangeEnd).SafeGetTimeValue(FTimeValue::GetPositiveInfinity());
+	FTimeValue PlayRangeEnd = StartPosition.Options.PlaybackRange.End;
+	check(PlayRangeEnd.IsValid());
 
 	FTimeValue StartTime = StartPosition.Time;
 
@@ -981,15 +981,14 @@ IManifest::FResult FDASHPlayPeriod::GetStartingSegment(TSharedPtrTS<IStreamSegme
 	Manifest->PreparePeriodAdaptationSets(Period, false);
 
 	// Frame accurate seek required?
-	bool bFrameAccurateSearch = PlayerSessionServices->GetOptions().GetValue(OptionKeyFrameAccurateSeek).SafeGetBool(false);
+	bool bFrameAccurateSearch = StartPosition.Options.bFrameAccuracy;
 	if (bFrameAccurateSearch)
 	{
 		// Get the segment that starts on or before the search time.
 		SearchType = IManifest::ESearchType::Before;
 	}
-
-	// Get the end of the playback end, if any.
-	FTimeValue PlayRangeEnd = PlayerSessionServices->GetOptions().GetValue(OptionPlayRangeEnd).SafeGetTimeValue(FTimeValue::GetPositiveInfinity());
+	FTimeValue PlayRangeEnd = StartPosition.Options.PlaybackRange.End;
+	check(PlayRangeEnd.IsValid());
 
 	FTimeValue AST = Manifest->GetAnchorTime();
 	FTimeValue StartTime = StartPosition.Time;
@@ -1219,12 +1218,12 @@ IManifest::FResult FDASHPlayPeriod::GetContinuationSegment(TSharedPtrTS<IStreamS
 	DummyReq->StreamType = StreamType;
 	DummyReq->PeriodStart = StartPosition.Time;
 	DummyReq->TimestampSequenceIndex = SequenceState.SequenceIndex;
-	return GetNextOrRetrySegment(OutSegment, DummyReq, ENextSegType::SamePeriodStartOver);
+	return GetNextOrRetrySegment(OutSegment, DummyReq, ENextSegType::SamePeriodStartOver, StartPosition.Options);
 }
 
 
 
-IManifest::FResult FDASHPlayPeriod::GetNextOrRetrySegment(TSharedPtrTS<IStreamSegment>& OutSegment, TSharedPtrTS<const IStreamSegment> InCurrentSegment, ENextSegType InNextType)
+IManifest::FResult FDASHPlayPeriod::GetNextOrRetrySegment(TSharedPtrTS<IStreamSegment>& OutSegment, TSharedPtrTS<const IStreamSegment> InCurrentSegment, ENextSegType InNextType, const FPlayStartOptions& Options)
 {
 	TSharedPtrTS<const FStreamSegmentRequestFMP4DASH> Current = StaticCastSharedPtr<const FStreamSegmentRequestFMP4DASH>(InCurrentSegment);
 	if (Current->bIsInitialStartRequest)
@@ -1271,9 +1270,9 @@ IManifest::FResult FDASHPlayPeriod::GetNextOrRetrySegment(TSharedPtrTS<IStreamSe
 	bool bIsStaticType = Manifest->IsStaticType() || Manifest->IsDynamicEpicEvent();
 
 	// Frame accurate seek required?
-	bool bFrameAccurateSearch = PlayerSessionServices->GetOptions().GetValue(OptionKeyFrameAccurateSeek).SafeGetBool(false);
-	// Get the end of the playback end, if any.
-	FTimeValue PlayRangeEnd = PlayerSessionServices->GetOptions().GetValue(OptionPlayRangeEnd).SafeGetTimeValue(FTimeValue::GetPositiveInfinity());
+	bool bFrameAccurateSearch = Options.bFrameAccuracy;
+	FTimeValue PlayRangeEnd = Options.PlaybackRange.End;
+	check(PlayRangeEnd.IsValid());
 	PlayRangeEnd -= AST;
 	PlayRangeEnd -= Period->GetStart();
 
@@ -1437,7 +1436,7 @@ IManifest::FResult FDASHPlayPeriod::GetNextOrRetrySegment(TSharedPtrTS<IStreamSe
 }
 
 
-IManifest::FResult FDASHPlayPeriod::GetNextSegment(TSharedPtrTS<IStreamSegment>& OutSegment, TSharedPtrTS<const IStreamSegment> InCurrentSegment)
+IManifest::FResult FDASHPlayPeriod::GetNextSegment(TSharedPtrTS<IStreamSegment>& OutSegment, TSharedPtrTS<const IStreamSegment> InCurrentSegment, const FPlayStartOptions& Options)
 {
 	if (!InCurrentSegment.IsValid())
 	{
@@ -1454,16 +1453,16 @@ IManifest::FResult FDASHPlayPeriod::GetNextSegment(TSharedPtrTS<IStreamSegment>&
 		{
 			return IManifest::FResult(IManifest::FResult::EType::PastEOS);
 		}
-		return GetNextOrRetrySegment(OutSegment, InCurrentSegment, ENextSegType::SamePeriodNext);
+		return GetNextOrRetrySegment(OutSegment, InCurrentSegment, ENextSegType::SamePeriodNext, Options);
 	}
 	else
 	{
 		// Moved into a new period. This here is the new period.
-		return GetNextOrRetrySegment(OutSegment, InCurrentSegment, ENextSegType::NextPeriod);
+		return GetNextOrRetrySegment(OutSegment, InCurrentSegment, ENextSegType::NextPeriod, Options);
 	}
 }
 
-IManifest::FResult FDASHPlayPeriod::GetRetrySegment(TSharedPtrTS<IStreamSegment>& OutSegment, TSharedPtrTS<const IStreamSegment> InCurrentSegment, bool bReplaceWithFillerData)
+IManifest::FResult FDASHPlayPeriod::GetRetrySegment(TSharedPtrTS<IStreamSegment>& OutSegment, TSharedPtrTS<const IStreamSegment> InCurrentSegment, const FPlayStartOptions& Options, bool bReplaceWithFillerData)
 {
 	if (!InCurrentSegment.IsValid())
 	{
@@ -1481,7 +1480,7 @@ IManifest::FResult FDASHPlayPeriod::GetRetrySegment(TSharedPtrTS<IStreamSegment>
 		OutSegment = NewRequest;
 		return IManifest::FResult(IManifest::FResult::EType::Found);
 	}
-	return GetNextOrRetrySegment(OutSegment, InCurrentSegment, ENextSegType::SamePeriodRetry);
+	return GetNextOrRetrySegment(OutSegment, InCurrentSegment, ENextSegType::SamePeriodRetry, Options);
 }
 
 
diff --git a/Engine/Plugins/Media/ElectraPlayer/Source/ElectraPlayerRuntime/Private/Runtime/Player/HLS/ManifestHLS.cpp b/Engine/Plugins/Media/ElectraPlayer/Source/ElectraPlayerRuntime/Private/Runtime/Player/HLS/ManifestHLS.cpp
--- a/Engine/Plugins/Media/ElectraPlayer/Source/ElectraPlayerRuntime/Private/Runtime/Player/HLS/ManifestHLS.cpp
+++ b/Engine/Plugins/Media/ElectraPlayer/Source/ElectraPlayerRuntime/Private/Runtime/Player/HLS/ManifestHLS.cpp
@@ -45,8 +45,8 @@ public:
 
 	virtual IManifest::FResult GetStartingSegment(TSharedPtrTS<IStreamSegment>& OutSegment, const FPlayerSequenceState& InSequenceState, const FPlayStartPosition& StartPosition, IManifest::ESearchType SearchType) override;
 	virtual IManifest::FResult GetContinuationSegment(TSharedPtrTS<IStreamSegment>& OutSegment, EStreamType StreamType, const FPlayerSequenceState& SequenceState, const FPlayStartPosition& StartPosition, IManifest::ESearchType SearchType) override;
-	virtual IManifest::FResult GetNextSegment(TSharedPtrTS<IStreamSegment>& OutSegment, TSharedPtrTS<const IStreamSegment> CurrentSegment) override;
-	virtual IManifest::FResult GetRetrySegment(TSharedPtrTS<IStreamSegment>& OutSegment, TSharedPtrTS<const IStreamSegment> CurrentSegment, bool bReplaceWithFillerData) override;
+	virtual IManifest::FResult GetNextSegment(TSharedPtrTS<IStreamSegment>& OutSegment, TSharedPtrTS<const IStreamSegment> CurrentSegment, const FPlayStartOptions& Options) override;
+	virtual IManifest::FResult GetRetrySegment(TSharedPtrTS<IStreamSegment>& OutSegment, TSharedPtrTS<const IStreamSegment> CurrentSegment, const FPlayStartOptions& Options, bool bReplaceWithFillerData) override;
 	virtual IManifest::FResult GetLoopingSegment(TSharedPtrTS<IStreamSegment>& OutSegment, const FPlayerSequenceState& SequenceState, const TMultiMap<EStreamType, TSharedPtrTS<IStreamSegment>>& InFinishedSegments, const FPlayStartPosition& StartPosition, IManifest::ESearchType SearchType) override;
 	virtual void IncreaseSegmentFetchDelay(const FTimeValue& IncreaseAmount) override;
 
@@ -84,7 +84,7 @@ private:
 
 	IManifest::FResult GetMediaStreamForID(TSharedPtrTS<FManifestHLSInternal::FPlaylistBase>& OutPlaylist, TSharedPtrTS<FManifestHLSInternal::FMediaStream>& OutMediaStream, uint32 UniqueID) const;
 
-	IManifest::FResult GetNextOrRetrySegment(TSharedPtrTS<IStreamSegment>& OutSegment, TSharedPtrTS<const IStreamSegment> CurrentSegment, bool bRetry);
+	IManifest::FResult GetNextOrRetrySegment(TSharedPtrTS<IStreamSegment>& OutSegment, TSharedPtrTS<const IStreamSegment> CurrentSegment, const FPlayStartOptions& Options, bool bRetry);
 
 	IManifest::FResult FindSegment(TSharedPtrTS<FStreamSegmentRequestHLSfmp4>& OutRequest, TSharedPtrTS<FManifestHLSInternal::FPlaylistBase> InPlaylist, TSharedPtrTS<FManifestHLSInternal::FMediaStream> InStream, uint32 StreamUniqueID, EStreamType StreamType, const FSegSearchParam& SearchParam, IManifest::ESearchType SearchType);
 
@@ -821,16 +821,15 @@ IManifest::FResult FPlayPeriodHLS::GetStartingSegment(TSharedPtrTS<IStreamSegmen
 		FSegSearchParam								SearchParam;
 
 		// Frame accurate seek required?
-		bool bFrameAccurateSearch = SessionServices->GetOptions().GetValue(OptionKeyFrameAccurateSeek).SafeGetBool(false);
+		bool bFrameAccurateSearch = StartPosition.Options.bFrameAccuracy;
 		if (bFrameAccurateSearch)
 		{
 			// Get the segment that starts on or before the search time.
 			SearchType = IManifest::ESearchType::Before;
 		}
 
-		// Get the end of the playback end, if any.
-		FTimeValue PlayRangeEnd = SessionServices->GetOptions().GetValue(OptionPlayRangeEnd).SafeGetTimeValue(FTimeValue::GetPositiveInfinity());
-
+		FTimeValue PlayRangeEnd = StartPosition.Options.PlaybackRange.End;
+		check(PlayRangeEnd.IsValid());
 
 		SearchParam.Time = StartPosition.Time;
 		SearchParam.LastPTS = PlayRangeEnd;
@@ -1024,7 +1023,7 @@ IManifest::FResult FPlayPeriodHLS::GetLoopingSegment(TSharedPtrTS<IStreamSegment
 
 
 
-IManifest::FResult FPlayPeriodHLS::GetNextOrRetrySegment(TSharedPtrTS<IStreamSegment>& OutSegment, TSharedPtrTS<const IStreamSegment> InCurrentSegment, bool bRetry)
+IManifest::FResult FPlayPeriodHLS::GetNextOrRetrySegment(TSharedPtrTS<IStreamSegment>& OutSegment, TSharedPtrTS<const IStreamSegment> InCurrentSegment, const FPlayStartOptions& Options, bool bRetry)
 {
 	// Need to have a current segment to find the next one.
 	if (!InCurrentSegment.IsValid())
@@ -1077,9 +1076,9 @@ IManifest::FResult FPlayPeriodHLS::GetNextOrRetrySegment(TSharedPtrTS<IStreamSeg
 		FSegSearchParam								SearchParam;
 
 		// Frame accurate seek required?
-		bool bFrameAccurateSearch = SessionServices->GetOptions().GetValue(OptionKeyFrameAccurateSeek).SafeGetBool(false);
-		// Get the end of the playback end, if any.
-		FTimeValue PlayRangeEnd = SessionServices->GetOptions().GetValue(OptionPlayRangeEnd).SafeGetTimeValue(FTimeValue::GetPositiveInfinity());
+		bool bFrameAccurateSearch = Options.bFrameAccuracy;
+		FTimeValue PlayRangeEnd = Options.PlaybackRange.End;
+		check(PlayRangeEnd.IsValid());
 		SearchParam.LastPTS = PlayRangeEnd;
 		SearchParam.bFrameAccurateSearch = bFrameAccurateSearch;
 
@@ -1149,16 +1148,17 @@ IManifest::FResult FPlayPeriodHLS::GetNextOrRetrySegment(TSharedPtrTS<IStreamSeg
  *
  * @param OutSegment
  * @param InCurrentSegment
+ * @param Options
  *
  * @return
  */
-IManifest::FResult FPlayPeriodHLS::GetNextSegment(TSharedPtrTS<IStreamSegment>& OutSegment, TSharedPtrTS<const IStreamSegment> InCurrentSegment)
+IManifest::FResult FPlayPeriodHLS::GetNextSegment(TSharedPtrTS<IStreamSegment>& OutSegment, TSharedPtrTS<const IStreamSegment> InCurrentSegment, const FPlayStartOptions& Options)
 {
-	return GetNextOrRetrySegment(OutSegment, InCurrentSegment, false);
+	return GetNextOrRetrySegment(OutSegment, InCurrentSegment, Options, false);
 }
 
 
-IManifest::FResult FPlayPeriodHLS::GetRetrySegment(TSharedPtrTS<IStreamSegment>& OutSegment, TSharedPtrTS<const IStreamSegment> InCurrentSegment, bool bReplaceWithFillerData)
+IManifest::FResult FPlayPeriodHLS::GetRetrySegment(TSharedPtrTS<IStreamSegment>& OutSegment, TSharedPtrTS<const IStreamSegment> InCurrentSegment, const FPlayStartOptions& Options, bool bReplaceWithFillerData)
 {
 	// To insert filler data we can use the current request over again.
 	if (bReplaceWithFillerData)
@@ -1172,7 +1172,7 @@ IManifest::FResult FPlayPeriodHLS::GetRetrySegment(TSharedPtrTS<IStreamSegment>&
 		OutSegment = NewRequest;
 		return IManifest::FResult(IManifest::FResult::EType::Found);
 	}
-	return GetNextOrRetrySegment(OutSegment, InCurrentSegment, true);
+	return GetNextOrRetrySegment(OutSegment, InCurrentSegment, Options, true);
 }
 
 
diff --git a/Engine/Plugins/Media/ElectraPlayer/Source/ElectraPlayerRuntime/Private/Runtime/Player/Manifest.h b/Engine/Plugins/Media/ElectraPlayer/Source/ElectraPlayerRuntime/Private/Runtime/Player/Manifest.h
--- a/Engine/Plugins/Media/ElectraPlayer/Source/ElectraPlayerRuntime/Private/Runtime/Player/Manifest.h
+++ b/Engine/Plugins/Media/ElectraPlayer/Source/ElectraPlayerRuntime/Private/Runtime/Player/Manifest.h
@@ -12,12 +12,6 @@
 #include "ParameterDictionary.h"
 
 
-
-//
-//
-//
-
-
 namespace Electra
 {
 
@@ -32,9 +26,21 @@ namespace Electra
 	};
 
 
+	struct FPlayStartOptions
+	{
+		FPlayStartOptions()
+		{
+			PlaybackRange.Start = FTimeValue::GetZero();
+			PlaybackRange.End = FTimeValue::GetPositiveInfinity();
+		}
+		FTimeRange		PlaybackRange;
+		bool			bFrameAccuracy = false;
+	};
+
 	struct FPlayStartPosition
 	{
-		FTimeValue		Time;
+		FTimeValue			Time;
+		FPlayStartOptions	Options;
 	};
 
 
@@ -385,18 +391,19 @@ namespace Electra
 			 *
 			 * @return
 			 */
-			virtual FResult GetNextSegment(TSharedPtrTS<IStreamSegment>& OutSegment, TSharedPtrTS<const IStreamSegment> CurrentSegment) = 0;
+			virtual FResult GetNextSegment(TSharedPtrTS<IStreamSegment>& OutSegment, TSharedPtrTS<const IStreamSegment> CurrentSegment, const FPlayStartOptions& Options) = 0;
 
 			/**
 			 * Gets the segment request for the same segment on a different quality level or CDN.
 			 *
 			 * @param OutSegment
 			 * @param CurrentSegment
+			 * @param Options
 			 * @param bReplaceWithFillerData
 			 *
 			 * @return
 			 */
-			virtual FResult GetRetrySegment(TSharedPtrTS<IStreamSegment>& OutSegment, TSharedPtrTS<const IStreamSegment> CurrentSegment, bool bReplaceWithFillerData) = 0;
+			virtual FResult GetRetrySegment(TSharedPtrTS<IStreamSegment>& OutSegment, TSharedPtrTS<const IStreamSegment> CurrentSegment, const FPlayStartOptions& Options, bool bReplaceWithFillerData) = 0;
 
 			/**
 			 * Called by the ABR to increase the delay in fetching the next segment in case the segment returned a 404 when fetched at
diff --git a/Engine/Plugins/Media/ElectraPlayer/Source/ElectraPlayerRuntime/Private/Runtime/Player/mp4/ManifestMP4.cpp b/Engine/Plugins/Media/ElectraPlayer/Source/ElectraPlayerRuntime/Private/Runtime/Player/mp4/ManifestMP4.cpp
--- a/Engine/Plugins/Media/ElectraPlayer/Source/ElectraPlayerRuntime/Private/Runtime/Player/mp4/ManifestMP4.cpp
+++ b/Engine/Plugins/Media/ElectraPlayer/Source/ElectraPlayerRuntime/Private/Runtime/Player/mp4/ManifestMP4.cpp
@@ -550,13 +550,14 @@ void FManifestMP4Internal::FPlayPeriodMP4::IncreaseSegmentFetchDelay(const FTime
  *
  * @param OutSegment
  * @param CurrentSegment
+ * @param Options
  *
  * @return
  */
-IManifest::FResult FManifestMP4Internal::FPlayPeriodMP4::GetNextSegment(TSharedPtrTS<IStreamSegment>& OutSegment, TSharedPtrTS<const IStreamSegment> CurrentSegment)
+IManifest::FResult FManifestMP4Internal::FPlayPeriodMP4::GetNextSegment(TSharedPtrTS<IStreamSegment>& OutSegment, TSharedPtrTS<const IStreamSegment> CurrentSegment, const FPlayStartOptions& Options)
 {
 	TSharedPtrTS<FTimelineAssetMP4> ma = MediaAsset.Pin();
-	return ma.IsValid() ? ma->GetNextSegment(OutSegment, CurrentSegment) : IManifest::FResult(IManifest::FResult::EType::NotFound);
+	return ma.IsValid() ? ma->GetNextSegment(OutSegment, CurrentSegment, Options) : IManifest::FResult(IManifest::FResult::EType::NotFound);
 }
 
 
@@ -566,14 +567,15 @@ IManifest::FResult FManifestMP4Internal::FPlayPeriodMP4::GetNextSegment(TSharedP
  *
  * @param OutSegment
  * @param CurrentSegment
+ * @param Options
  * @param bReplaceWithFillerData
  *
  * @return
  */
-IManifest::FResult FManifestMP4Internal::FPlayPeriodMP4::GetRetrySegment(TSharedPtrTS<IStreamSegment>& OutSegment, TSharedPtrTS<const IStreamSegment> CurrentSegment, bool bReplaceWithFillerData)
+IManifest::FResult FManifestMP4Internal::FPlayPeriodMP4::GetRetrySegment(TSharedPtrTS<IStreamSegment>& OutSegment, TSharedPtrTS<const IStreamSegment> CurrentSegment, const FPlayStartOptions& Options, bool bReplaceWithFillerData)
 {
 	TSharedPtrTS<FTimelineAssetMP4> ma = MediaAsset.Pin();
-	return ma.IsValid() ? ma->GetRetrySegment(OutSegment, CurrentSegment, bReplaceWithFillerData) : IManifest::FResult(IManifest::FResult::EType::NotFound);
+	return ma.IsValid() ? ma->GetRetrySegment(OutSegment, CurrentSegment, Options, bReplaceWithFillerData) : IManifest::FResult(IManifest::FResult::EType::NotFound);
 }
 
 
@@ -780,8 +782,9 @@ IManifest::FResult FManifestMP4Internal::FTimelineAssetMP4::GetStartingSegment(T
 // TODO: If there is a SIDX box we will look in there.
 
 	// Frame accurate seek required?
-	bool bFrameAccurateSearch = AtAbsoluteFilePos < 0 ? PlayerSessionServices->GetOptions().GetValue(OptionKeyFrameAccurateSeek).SafeGetBool(false) : false;
-	FTimeValue PlayRangeEnd = PlayerSessionServices->GetOptions().GetValue(OptionPlayRangeEnd).SafeGetTimeValue(FTimeValue::GetPositiveInfinity());
+	bool bFrameAccurateSearch = AtAbsoluteFilePos < 0 ? StartPosition.Options.bFrameAccuracy : false;
+	FTimeValue PlayRangeEnd = StartPosition.Options.PlaybackRange.End;
+	check(PlayRangeEnd.IsValid());
 	UpdatePlayRangeEndInfo(PlayRangeEnd);
 
 	// Look at the actual tracks. If there is video search there first for a keyframe/IDR frame.
@@ -1084,7 +1087,7 @@ IManifest::FResult FManifestMP4Internal::FTimelineAssetMP4::GetStartingSegment(T
 					   .SetMessage(FString::Printf(TEXT("Could not find start segment for time %lld, no valid tracks"), (long long int)StartPosition.Time.GetAsHNS())));
 }
 
-IManifest::FResult FManifestMP4Internal::FTimelineAssetMP4::GetNextSegment(TSharedPtrTS<IStreamSegment>& OutSegment, TSharedPtrTS<const IStreamSegment> CurrentSegment)
+IManifest::FResult FManifestMP4Internal::FTimelineAssetMP4::GetNextSegment(TSharedPtrTS<IStreamSegment>& OutSegment, TSharedPtrTS<const IStreamSegment> CurrentSegment, const FPlayStartOptions& Options)
 {
 	const FStreamSegmentRequestMP4* Request = static_cast<const FStreamSegmentRequestMP4*>(CurrentSegment.Get());
 	if (Request)
@@ -1094,6 +1097,7 @@ IManifest::FResult FManifestMP4Internal::FTimelineAssetMP4::GetNextSegment(TShar
 		{
 			FPlayStartPosition dummyPos;
 			FPlayerSequenceState seqState;
+			dummyPos.Options = Options;
 			seqState.SequenceIndex = Request->TimestampSequenceIndex;
 			IManifest::FResult res = GetStartingSegment(OutSegment, seqState, dummyPos, ESearchType::Same, Request->FileEndOffset + 1);
 			if (res.GetType() == IManifest::FResult::EType::Found)
@@ -1113,13 +1117,14 @@ IManifest::FResult FManifestMP4Internal::FTimelineAssetMP4::GetNextSegment(TShar
 	return IManifest::FResult(IManifest::FResult::EType::PastEOS);
 }
 
-IManifest::FResult FManifestMP4Internal::FTimelineAssetMP4::GetRetrySegment(TSharedPtrTS<IStreamSegment>& OutSegment, TSharedPtrTS<const IStreamSegment> CurrentSegment, bool bReplaceWithFillerData)
+IManifest::FResult FManifestMP4Internal::FTimelineAssetMP4::GetRetrySegment(TSharedPtrTS<IStreamSegment>& OutSegment, TSharedPtrTS<const IStreamSegment> CurrentSegment, const FPlayStartOptions& Options, bool bReplaceWithFillerData)
 {
 	const FStreamSegmentRequestMP4* Request = static_cast<const FStreamSegmentRequestMP4*>(CurrentSegment.Get());
 	if (Request)
 	{
 		FPlayStartPosition dummyPos;
 		FPlayerSequenceState seqState;
+		dummyPos.Options = Options;
 		seqState.SequenceIndex = Request->TimestampSequenceIndex;
 		IManifest::FResult res = GetStartingSegment(OutSegment, seqState, dummyPos, ESearchType::Same, Request->CurrentIteratorBytePos);
 		if (res.GetType() == IManifest::FResult::EType::Found)
diff --git a/Engine/Plugins/Media/ElectraPlayer/Source/ElectraPlayerRuntime/Private/Runtime/Player/mp4/ManifestMP4.h b/Engine/Plugins/Media/ElectraPlayer/Source/ElectraPlayerRuntime/Private/Runtime/Player/mp4/ManifestMP4.h
--- a/Engine/Plugins/Media/ElectraPlayer/Source/ElectraPlayerRuntime/Private/Runtime/Player/mp4/ManifestMP4.h
+++ b/Engine/Plugins/Media/ElectraPlayer/Source/ElectraPlayerRuntime/Private/Runtime/Player/mp4/ManifestMP4.h
@@ -224,8 +224,8 @@ public:
 
 
 		FResult GetStartingSegment(TSharedPtrTS<IStreamSegment>& OutSegment, const FPlayerSequenceState& InSequenceState, const FPlayStartPosition& StartPosition, ESearchType SearchType, int64 AtAbsoluteFilePos);
-		FResult GetNextSegment(TSharedPtrTS<IStreamSegment>& OutSegment, TSharedPtrTS<const IStreamSegment> CurrentSegment);
-		FResult GetRetrySegment(TSharedPtrTS<IStreamSegment>& OutSegment, TSharedPtrTS<const IStreamSegment> CurrentSegment, bool bReplaceWithFillerData);
+		FResult GetNextSegment(TSharedPtrTS<IStreamSegment>& OutSegment, TSharedPtrTS<const IStreamSegment> CurrentSegment, const FPlayStartOptions& Options);
+		FResult GetRetrySegment(TSharedPtrTS<IStreamSegment>& OutSegment, TSharedPtrTS<const IStreamSegment> CurrentSegment, const FPlayStartOptions& Options, bool bReplaceWithFillerData);
 		FResult GetLoopingSegment(TSharedPtrTS<IStreamSegment>& OutSegment, const FPlayerSequenceState& InSequenceState, const TMultiMap<EStreamType, TSharedPtrTS<IStreamSegment>>& InFinishedSegments, const FPlayStartPosition& StartPosition, ESearchType SearchType);
 
 		void GetSegmentInformation(TArray<IManifest::IPlayPeriod::FSegmentInformation>& OutSegmentInformation, FTimeValue& OutAverageSegmentDuration, TSharedPtrTS<const IStreamSegment> CurrentSegment, const FTimeValue& LookAheadTime, const FString& AdaptationSetID, const FString& RepresentationID);
@@ -279,8 +279,8 @@ public:
 		virtual void SelectStream(const FString& AdaptationSetID, const FString& RepresentationID) override;
 		virtual FResult GetStartingSegment(TSharedPtrTS<IStreamSegment>& OutSegment, const FPlayerSequenceState& InSequenceState, const FPlayStartPosition& StartPosition, ESearchType SearchType) override;
 		virtual FResult GetContinuationSegment(TSharedPtrTS<IStreamSegment>& OutSegment, EStreamType StreamType, const FPlayerSequenceState& LoopState, const FPlayStartPosition& StartPosition, ESearchType SearchType) override;
-		virtual FResult GetNextSegment(TSharedPtrTS<IStreamSegment>& OutSegment, TSharedPtrTS<const IStreamSegment> CurrentSegment) override;
-		virtual FResult GetRetrySegment(TSharedPtrTS<IStreamSegment>& OutSegment, TSharedPtrTS<const IStreamSegment> CurrentSegment, bool bReplaceWithFillerData) override;
+		virtual FResult GetNextSegment(TSharedPtrTS<IStreamSegment>& OutSegment, TSharedPtrTS<const IStreamSegment> CurrentSegment, const FPlayStartOptions& Options) override;
+		virtual FResult GetRetrySegment(TSharedPtrTS<IStreamSegment>& OutSegment, TSharedPtrTS<const IStreamSegment> CurrentSegment, const FPlayStartOptions& Options, bool bReplaceWithFillerData) override;
 		virtual FResult GetLoopingSegment(TSharedPtrTS<IStreamSegment>& OutSegment, const FPlayerSequenceState& InSequenceState, const TMultiMap<EStreamType, TSharedPtrTS<IStreamSegment>>& InFinishedSegments, const FPlayStartPosition& StartPosition, ESearchType SearchType) override;
 		virtual void IncreaseSegmentFetchDelay(const FTimeValue& IncreaseAmount) override;
 		virtual void GetSegmentInformation(TArray<FSegmentInformation>& OutSegmentInformation, FTimeValue& OutAverageSegmentDuration, TSharedPtrTS<const IStreamSegment> CurrentSegment, const FTimeValue& LookAheadTime, const FString& AdaptationSetID, const FString& RepresentationID) override;
