ElectraPlayer: Fixed looping if a play range is set which was being ignored.

#rnx

diff --git a/Engine/Plugins/Media/ElectraPlayer/Source/ElectraPlayerRuntime/Private/Runtime/Player/AdaptiveStreamingPlayer.cpp b/Engine/Plugins/Media/ElectraPlayer/Source/ElectraPlayerRuntime/Private/Runtime/Player/AdaptiveStreamingPlayer.cpp
--- a/Engine/Plugins/Media/ElectraPlayer/Source/ElectraPlayerRuntime/Private/Runtime/Player/AdaptiveStreamingPlayer.cpp
+++ b/Engine/Plugins/Media/ElectraPlayer/Source/ElectraPlayerRuntime/Private/Runtime/Player/AdaptiveStreamingPlayer.cpp
@@ -2506,6 +2506,8 @@ void FAdaptiveStreamingPlayer::InternalHandlePendingStartRequest(const FTimeValu
 							}
 						}
 					}
+					// Set the current playback range with the start options and clamp the start time to the range.
+					SetPlaystartOptions(PendingStartRequest->StartAt.Options);
 					ClampStartRequestTime();
 
 					IManifest::FResult Result = Manifest->FindPlayPeriod(InitialPlayPeriod, PendingStartRequest->StartAt, PendingStartRequest->SearchType);
@@ -3020,11 +3022,12 @@ void FAdaptiveStreamingPlayer::InternalHandleCompletedSegmentRequests(const FTim
 					// the time will be clamped to that. The in-point cannot be before the start point.
 					PendingStartRequest->SearchType = IManifest::ESearchType::Closest;
 					PendingStartRequest->StartAt.Time = Seekable.Start;
+					SetPlaystartOptions(PendingStartRequest->StartAt.Options);
 					PendingStartRequest->StartType = FPendingStartRequest::EStartType::LoopPoint;
 					PendingStartRequest->FinishedRequests = MoveTemp(LocalFinishedRequests);
 					// Increase the internal loop count
 					++CurrentLoopState.Count;
-					CurrentLoopState.To = Seekable.Start;
+					CurrentLoopState.To = ClampTimeToCurrentRange(Seekable.Start, true, false);
 					CurrentLoopState.From = ClampTimeToCurrentRange(FTimeValue::GetPositiveInfinity(), false, true);
 				}
 			}
@@ -3871,7 +3874,7 @@ void FAdaptiveStreamingPlayer::ClampStartRequestTime()
 		{
 			PendingStartRequest->StartAt.Time = RangeStart;
 		}
-		else if (PendingStartRequest->StartAt.Time > RangeEnd)
+		if (PendingStartRequest->StartAt.Time > RangeEnd)
 		{
 			PendingStartRequest->StartAt.Time = RangeEnd;
 			// Going to the end will of course result in an 'end reached' and end of playback.
