Fixing up ElectraPlayer's close/shutdown code so it can also succeed during the engine's shutdown (PreExit) phase

diff --git a/Engine/Plugins/Media/ElectraPlayer/Source/ElectraPlayerRuntime/Private/ElectraPlayer.cpp b/Engine/Plugins/Media/ElectraPlayer/Source/ElectraPlayerRuntime/Private/ElectraPlayer.cpp
--- a/Engine/Plugins/Media/ElectraPlayer/Source/ElectraPlayerRuntime/Private/ElectraPlayer.cpp
+++ b/Engine/Plugins/Media/ElectraPlayer/Source/ElectraPlayerRuntime/Private/ElectraPlayer.cpp
@@ -437,7 +437,16 @@ void FElectraPlayer::FInternalPlayerImpl::DoCloseAsync(TSharedPtr<FInternalPlaye
 			AsyncResourceReleaseNotification->Signal(ResourceFlags_Decoder);
 		}
 	};
-	Async(EAsyncExecution::ThreadPool, MoveTemp(CloseTask));
+
+	// Fallback to simple, sequential execution if the engine is already shutting down...
+	if (GIsRunning)
+	{
+		Async(EAsyncExecution::ThreadPool, MoveTemp(CloseTask));
+	}
+	else
+	{
+		CloseTask();
+	}
 
 #if !(UE_BUILD_SHIPPING || UE_BUILD_TEST)
 	FInternalPlayerImpl* PlayerImpl = Player.Get();
