Reworking ElectraPlayer's resource delegate to be more thread safe

diff --git a/Engine/Plugins/Media/ElectraPlayer/Source/ElectraPlayerPlugin/Private/ElectraPlayerPlugin.cpp b/Engine/Plugins/Media/ElectraPlayer/Source/ElectraPlayerPlugin/Private/ElectraPlayerPlugin.cpp
--- a/Engine/Plugins/Media/ElectraPlayer/Source/ElectraPlayerPlugin/Private/ElectraPlayerPlugin.cpp
+++ b/Engine/Plugins/Media/ElectraPlayer/Source/ElectraPlayerPlugin/Private/ElectraPlayerPlugin.cpp
@@ -71,9 +71,11 @@ bool FElectraPlayerPlugin::Initialize(IMediaEventSink& InEventSink,
 	EventSink = &InEventSink;
 	CallbackPointerLock.Unlock();
 
+	OutputTexturePool = MakeShareable(new FElectraTextureSamplePool);
+
 	MediaSamples.Reset(new FMediaSamples);
 
-	PlayerResourceDelegate.Reset(PlatformCreatePlayerResourceDelegate());
+	PlayerResourceDelegate = MakeShareable(PlatformCreatePlayerResourceDelegate());
 
 	PlayerDelegate = MakeShareable(new FPlayerAdapterDelegate(AsShared()));
 	Player = MakeShareable(FElectraPlayerRuntimeFactory::CreatePlayer(PlayerDelegate, InSendAnalyticMetricsDelegate, InSendAnalyticMetricsPerMinuteDelegate, InReportVideoStreamingErrorDelegate, InReportSubtitlesFileMetricsDelegate));
@@ -363,7 +365,7 @@ void FElectraPlayerPlugin::FPlayerAdapterDelegate::PresentVideoFrame(const FVide
 	TSharedPtr<FElectraPlayerPlugin, ESPMode::ThreadSafe> PinnedHost = Host.Pin();
 	if (PinnedHost.IsValid())
 	{
-		FElectraTextureSampleRef TextureSample = PinnedHost->OutputTexturePool.AcquireShared();
+		FElectraTextureSampleRef TextureSample = PinnedHost->OutputTexturePool->AcquireShared();
 		TextureSample->Initialize(InVideoFrame.Get());
 		PinnedHost->MediaSamples->AddVideo(TextureSample);
 	}
@@ -431,7 +433,7 @@ void FElectraPlayerPlugin::FPlayerAdapterDelegate::PrepareForDecoderShutdown()
 	TSharedPtr<FElectraPlayerPlugin, ESPMode::ThreadSafe> PinnedHost = Host.Pin();
 	if (PinnedHost.IsValid())
 	{
-		PinnedHost->OutputTexturePool.PrepareForDecoderShutdown();
+		PinnedHost->OutputTexturePool->PrepareForDecoderShutdown();
 	}
 }
 
@@ -442,12 +444,12 @@ FString FElectraPlayerPlugin::FPlayerAdapterDelegate::GetVideoAdapterName() cons
 }
 
 
-IElectraPlayerResourceDelegate* FElectraPlayerPlugin::FPlayerAdapterDelegate::GetResourceDelegate() const
+TSharedPtr<IElectraPlayerResourceDelegate, ESPMode::ThreadSafe> FElectraPlayerPlugin::FPlayerAdapterDelegate::GetResourceDelegate() const
 {
 	TSharedPtr<FElectraPlayerPlugin, ESPMode::ThreadSafe> PinnedHost = Host.Pin();
 	if (PinnedHost.IsValid())
 	{
-		return PinnedHost->PlayerResourceDelegate.Get();
+		return PinnedHost->PlayerResourceDelegate;
 	}
 	return nullptr;
 }
@@ -645,7 +647,7 @@ void FElectraPlayerPlugin::Close()
  */
 void FElectraPlayerPlugin::TickInput(FTimespan DeltaTime, FTimespan Timecode)
 {
-	OutputTexturePool.Tick();
+	OutputTexturePool->Tick();
 	Player->Tick(DeltaTime, Timecode);
 }
 
diff --git a/Engine/Plugins/Media/ElectraPlayer/Source/ElectraPlayerPlugin/Private/ElectraPlayerPlugin.h b/Engine/Plugins/Media/ElectraPlayer/Source/ElectraPlayerPlugin/Private/ElectraPlayerPlugin.h
--- a/Engine/Plugins/Media/ElectraPlayer/Source/ElectraPlayerPlugin/Private/ElectraPlayerPlugin.h
+++ b/Engine/Plugins/Media/ElectraPlayer/Source/ElectraPlayerPlugin/Private/ElectraPlayerPlugin.h
@@ -172,7 +172,7 @@ private:
 		bool CanReceiveAudioSamples(int32 NumFrames) override;
 		void PrepareForDecoderShutdown() override;
 		FString GetVideoAdapterName() const override;
-		IElectraPlayerResourceDelegate* GetResourceDelegate() const override;
+		TSharedPtr<IElectraPlayerResourceDelegate, ESPMode::ThreadSafe> GetResourceDelegate() const override;
 
 	private:
 		TWeakPtr<FElectraPlayerPlugin, ESPMode::ThreadSafe> Host;
@@ -198,9 +198,9 @@ private:
 	TSharedPtr<IElectraPlayerInterface, ESPMode::ThreadSafe> Player;
 
 	/** Output sample pools */
-	FElectraTextureSamplePool OutputTexturePool;
+	TSharedPtr<FElectraTextureSamplePool, ESPMode::ThreadSafe> OutputTexturePool;
 	FElectraPlayerAudioSamplePool OutputAudioPool;
 
-	TUniquePtr<IElectraPlayerResourceDelegate> PlayerResourceDelegate;
+	TSharedPtr<IElectraPlayerResourceDelegate,ESPMode::ThreadSafe> PlayerResourceDelegate;
 };
 
diff --git a/Engine/Plugins/Media/ElectraPlayer/Source/ElectraPlayerPlugin/Private/ElectraPlayerPlugin_Android.cpp b/Engine/Plugins/Media/ElectraPlayer/Source/ElectraPlayerPlugin/Private/ElectraPlayerPlugin_Android.cpp
--- a/Engine/Plugins/Media/ElectraPlayer/Source/ElectraPlayerPlugin/Private/ElectraPlayerPlugin_Android.cpp
+++ b/Engine/Plugins/Media/ElectraPlayer/Source/ElectraPlayerPlugin/Private/ElectraPlayerPlugin_Android.cpp
@@ -14,7 +14,7 @@ public:
 
 	virtual jobject GetCodecSurface() override
 	{
-		return (jobject)Owner->OutputTexturePool.GetCodecSurface();
+		return (jobject)Owner->OutputTexturePool->GetCodecSurface();
 	}
 
 private:
diff --git a/Engine/Plugins/Media/ElectraPlayer/Source/ElectraPlayerRuntime/Private/ElectraPlayer_Android.cpp b/Engine/Plugins/Media/ElectraPlayer/Source/ElectraPlayerRuntime/Private/ElectraPlayer_Android.cpp
--- a/Engine/Plugins/Media/ElectraPlayer/Source/ElectraPlayerRuntime/Private/ElectraPlayer_Android.cpp
+++ b/Engine/Plugins/Media/ElectraPlayer/Source/ElectraPlayerRuntime/Private/ElectraPlayer_Android.cpp
@@ -49,7 +49,7 @@ public:
 	{
 		if (const TSharedPtr<IElectraPlayerAdapterDelegate, ESPMode::ThreadSafe>& PinnedAdapterDelegate = PlayerAdapterDelegate.Pin())
 		{
-			IElectraPlayerResourceDelegate* ResourceDelegate = PinnedAdapterDelegate->GetResourceDelegate();
+			TSharedPtr<IElectraPlayerResourceDelegate, ESPMode::ThreadSafe> ResourceDelegate = PinnedAdapterDelegate->GetResourceDelegate();
 			if (ResourceDelegate)
 			{
 				return ResourceDelegate->GetCodecSurface();
diff --git a/Engine/Plugins/Media/ElectraPlayer/Source/ElectraPlayerRuntime/Public/IElectraPlayerInterface.h b/Engine/Plugins/Media/ElectraPlayer/Source/ElectraPlayerRuntime/Public/IElectraPlayerInterface.h
--- a/Engine/Plugins/Media/ElectraPlayer/Source/ElectraPlayerRuntime/Public/IElectraPlayerInterface.h
+++ b/Engine/Plugins/Media/ElectraPlayer/Source/ElectraPlayerRuntime/Public/IElectraPlayerInterface.h
@@ -79,7 +79,7 @@ public:
 
 	virtual FString GetVideoAdapterName() const = 0;
 
-	virtual IElectraPlayerResourceDelegate* GetResourceDelegate() const = 0;
+	virtual TSharedPtr<IElectraPlayerResourceDelegate, ESPMode::ThreadSafe> GetResourceDelegate() const = 0;
 };
 
 // Container class to be passed through options as shared pointer to allow passing any non-standard ref-counted entities to the player
