Adding some thread safety (during player async closes) to Android version as well

diff --git a/Engine/Plugins/Media/ElectraPlayer/Source/ElectraPlayerPlugin/Private/ElectraPlayerPlugin_Android.cpp b/Engine/Plugins/Media/ElectraPlayer/Source/ElectraPlayerPlugin/Private/ElectraPlayerPlugin_Android.cpp
--- a/Engine/Plugins/Media/ElectraPlayer/Source/ElectraPlayerPlugin/Private/ElectraPlayerPlugin_Android.cpp
+++ b/Engine/Plugins/Media/ElectraPlayer/Source/ElectraPlayerPlugin/Private/ElectraPlayerPlugin_Android.cpp
@@ -10,15 +10,19 @@
 class FElectraPlayerResourceDelegate : public IElectraPlayerResourceDelegate
 {
 public:
-	FElectraPlayerResourceDelegate(FElectraPlayerPlugin* InOwner) : Owner(InOwner) {}
+	FElectraPlayerResourceDelegate(FElectraPlayerPlugin* InOwner) : Owner(InOwner->AsWeak()) {}
 
 	virtual jobject GetCodecSurface() override
 	{
-		return (jobject)Owner->OutputTexturePool->GetCodecSurface();
+		if (auto PinnedOwner = Owner.Pin())
+		{
+			return (jobject)PinnedOwner->OutputTexturePool->GetCodecSurface();
+		}
+		return (jobject)nullptr;
 	}
 
 private:
-	FElectraPlayerPlugin* Owner;
+	TWeakPtr<FElectraPlayerPlugin, ESPMode::ThreadSafe> Owner;
 };
 
 
