Reviving Electra okayer playback as RHI changes rendered its init code inactive
#rnx

diff --git a/Engine/Plugins/Media/ElectraPlayer/Source/ElectraPlayerPlugin/Private/ElectraPlayerPluginModule.cpp b/Engine/Plugins/Media/ElectraPlayer/Source/ElectraPlayerPlugin/Private/ElectraPlayerPluginModule.cpp
--- a/Engine/Plugins/Media/ElectraPlayer/Source/ElectraPlayerPlugin/Private/ElectraPlayerPluginModule.cpp
+++ b/Engine/Plugins/Media/ElectraPlayer/Source/ElectraPlayerPlugin/Private/ElectraPlayerPluginModule.cpp
@@ -87,7 +87,7 @@ public:
 		{
 //WE COULD ALSO CHECK FOR SINGLE THREADING ISSUES HERE!? -- kinda piggy back on this?
 			// to detect cooking and other commandlets that run with NullRHI
-			if (GDynamicRHI == nullptr || RHIGetInterfaceType() == ERHIInterfaceType::Null || RHIGetInterfaceType() == ERHIInterfaceType::Hidden)
+			if (GDynamicRHI == nullptr || RHIGetInterfaceType() == ERHIInterfaceType::Null)
 			{
 				UE_LOG(LogElectraPlayerPlugin, Log, TEXT("Dummy Dynamic RHI detected. Electra Player plugin is not initialised"));
 				return;
