ElectraPlayer: Update from Dev-EngineMerge to bring in the new HTTP streamer for Electra

#rnx

diff --git a/Engine/Plugins/Media/ElectraPlayer/Source/ElectraPlayerRuntime/ElectraPlayerRuntime.Build.cs b/Engine/Plugins/Media/ElectraPlayer/Source/ElectraPlayerRuntime/ElectraPlayerRuntime.Build.cs
--- a/Engine/Plugins/Media/ElectraPlayer/Source/ElectraPlayerRuntime/ElectraPlayerRuntime.Build.cs
+++ b/Engine/Plugins/Media/ElectraPlayer/Source/ElectraPlayerRuntime/ElectraPlayerRuntime.Build.cs
@@ -45,7 +45,7 @@ namespace UnrealBuildTool.Rules
 				string DirectXSDKDir = Target.UEThirdPartySourceDirectory + "Windows/DirectX";
 				PublicSystemIncludePaths.Add(DirectXSDKDir + "/include");
 
-				if (Target.Platform.IsInGroup(UnrealPlatformGroup.Windows))
+				if (Target.Platform == UnrealTargetPlatform.Win64)
 				{
 					DirectXSDKDir += "/Lib/x64/";
 				}
@@ -69,7 +69,7 @@ namespace UnrealBuildTool.Rules
 					PrivateDefinitions.Add("ELECTRA_HAVE_DX11");	// video decoding for DX11 enabled (Win8+)
 				}
 
-				if (Target.Platform.IsInGroup(UnrealPlatformGroup.Windows))
+				if (Target.Platform == UnrealTargetPlatform.Win64)
 				{
 					PublicAdditionalLibraries.AddRange(new string[] {
 						DirectXSDKDir + "dxerr.lib",
