ElectraPlayer: Added DASH low latency Live streaming

diff --git a/Engine/Plugins/Media/ElectraPlayer/Source/ElectraPlayerRuntime/Private/ElectraPlayer.cpp b/Engine/Plugins/Media/ElectraPlayer/Source/ElectraPlayerRuntime/Private/ElectraPlayer.cpp
--- a/Engine/Plugins/Media/ElectraPlayer/Source/ElectraPlayerRuntime/Private/ElectraPlayer.cpp
+++ b/Engine/Plugins/Media/ElectraPlayer/Source/ElectraPlayerRuntime/Private/ElectraPlayer.cpp
@@ -1035,6 +1035,12 @@ void FElectraPlayer::CalculateTargetSeekTime(FTimespan& OutTargetTime, const FTi
 		// into content already aired.
 		if (IsLive())
 		{
+			// If the target is maximum we treat it as going to the Live edge.
+			if (InTime == FTimespan::MaxValue())
+			{
+				OutTargetTime = InTime;
+				return;
+			}
 			// In case the seek time has been given as a negative number we negate it.
 			if (newTime.GetAsHNS() < 0)
 			{
@@ -1085,7 +1091,10 @@ bool FElectraPlayer::Seek(const FTimespan& Time, const FSeekParam& Param)
 		FTimespan Target;
 		CalculateTargetSeekTime(Target, Time);
 		Electra::IAdaptiveStreamingPlayer::FSeekParam seek;
-		seek.Time.SetFromTimespan(Target);
+		if (Target != FTimespan::MaxValue())
+		{
+			seek.Time.SetFromTimespan(Target);
+		}
 		seek.StartingBitrate = Param.StartingBitrate;
 		seek.bOptimizeForScrubbing = Param.bOptimizeForScrubbing;
 		seek.DistanceThreshold = Param.DistanceThreshold;
diff --git a/Engine/Plugins/Media/ElectraPlayer/Source/ElectraPlayerRuntime/Private/ElectraPlayerAudioDecoderOutput.cpp b/Engine/Plugins/Media/ElectraPlayer/Source/ElectraPlayerRuntime/Private/ElectraPlayerAudioDecoderOutput.cpp
--- a/Engine/Plugins/Media/ElectraPlayer/Source/ElectraPlayerRuntime/Private/ElectraPlayerAudioDecoderOutput.cpp
+++ b/Engine/Plugins/Media/ElectraPlayer/Source/ElectraPlayerRuntime/Private/ElectraPlayerAudioDecoderOutput.cpp
@@ -6,12 +6,7 @@
 class FElectraPlayerAudioDecoderOutput : public IAudioDecoderOutput
 {
 public:
-	FElectraPlayerAudioDecoderOutput()
-	: BufferMaxSizeBytes(0)
-	, BufferUsedSizeBytes(0)
-	, Buffer(nullptr)
-	{
-	}
+	FElectraPlayerAudioDecoderOutput() = default;
 
 	~FElectraPlayerAudioDecoderOutput() override
 	{
@@ -120,18 +115,18 @@ public:
 		TSharedPtr<IDecoderOutputOwner, ESPMode::ThreadSafe> lockedAudioRenderer = OwningRenderer.Pin();
 		if (lockedAudioRenderer.IsValid())
 		{
-			lockedAudioRenderer->SampleReleasedToPool(Duration);
+			lockedAudioRenderer->SampleReleasedToPool(this);
 		}
 	}
 
 private:
-	uint32 BufferMaxSizeBytes;
-	uint32 BufferUsedSizeBytes;
-	void *Buffer;
+	uint32 BufferMaxSizeBytes = 0;
+	uint32 BufferUsedSizeBytes = 0;
+	void* Buffer = nullptr;
 
-	ESampleFormat Format;
-	uint32 NumChannels;
-	uint32 SampleRate;
+	ESampleFormat Format = ESampleFormat::Undefined;
+	uint32 NumChannels = 0;
+	uint32 SampleRate = 0;
 
 	FDecoderTimeStamp Time;
 	FTimespan Duration;
@@ -144,8 +139,7 @@ private:
 /*
 	Create one and the same audio sample for all platforms - we so far do not need any specialization here (otherwise this should be with the respective decoder - much like the class definition / impl above)
 */
-IAudioDecoderOutput *FElectraPlayerPlatformAudioDecoderOutputFactory::Create()
+IAudioDecoderOutput* FElectraPlayerPlatformAudioDecoderOutputFactory::Create()
 {
 	return new FElectraPlayerAudioDecoderOutput();
 }
-
diff --git a/Engine/Plugins/Media/ElectraPlayer/Source/ElectraPlayerRuntime/Private/Runtime/Core/MediaQueue.h b/Engine/Plugins/Media/ElectraPlayer/Source/ElectraPlayerRuntime/Private/Runtime/Core/MediaQueue.h
--- a/Engine/Plugins/Media/ElectraPlayer/Source/ElectraPlayerRuntime/Private/Runtime/Core/MediaQueue.h
+++ b/Engine/Plugins/Media/ElectraPlayer/Source/ElectraPlayerRuntime/Private/Runtime/Core/MediaQueue.h
@@ -55,7 +55,7 @@ public:
 	using MyType = TMediaQueue<T, L, M>;
 	using ElementType = T;
 
-	TMediaQueue(SIZE_T InitialCapacity = 0)
+	TMediaQueue(int32 InitialCapacity = 0)
 	{
 		Resize(InitialCapacity);
 	}
@@ -72,7 +72,7 @@ public:
 	void Clear()
 	{
 		L::Lock();
-		for(uint32 i = 0; i < NumIn; ++i)
+		for(int32 i=0; i<NumIn; ++i)
 		{
 			Elements[IdxOut].~T();
 			if (++IdxOut == MaxNum)
@@ -84,7 +84,7 @@ public:
 		L::Unlock();
 	}
 
-	void Resize(SIZE_T newCapacity)
+	void Resize(int32 newCapacity)
 	{
 		L::Lock();
 		Clear();
@@ -109,10 +109,10 @@ public:
 		L::Lock();
 		// Remember our values.
 		T* pElements = Elements;
-		SIZE_T  nCapacity = MaxNum;
-		SIZE_T  nNumIn = NumIn;
-		SIZE_T  nIdxIn = IdxIn;
-		SIZE_T  nIdxOut = IdxOut;
+		int32  nCapacity = MaxNum;
+		int32  nNumIn = NumIn;
+		int32  nIdxIn = IdxIn;
+		int32  nIdxOut = IdxOut;
 		// Get the values from the other queue.
 		Elements = rhs.Elements;
 		MaxNum = rhs.MaxNum;
@@ -128,12 +128,12 @@ public:
 		L::Unlock();
 	}
 
-	SIZE_T Capacity() const
+	int32 Capacity() const
 	{
 		return MaxNum;
 	}
 
-	SIZE_T Num() const
+	int32 Num() const
 	{
 		return NumIn;
 	}
@@ -152,7 +152,7 @@ public:
 	{
 		L::Lock();
 		check(!IsFull());
-		SIZE_T pos = IdxIn;
+		int32 pos = IdxIn;
 		if (++IdxIn == MaxNum)
 		{
 			IdxIn = 0;
@@ -175,7 +175,7 @@ public:
 		{
 			--IdxOut;
 		}
-		SIZE_T pos = IdxOut;
+		int32 pos = IdxOut;
 		++NumIn;
 		// Copying the element can NOT take place outside the lock in case multiple threads push at the same time and a Resize() will happen.
 		new ((void*)(Elements + pos)) T(Element);
@@ -208,10 +208,10 @@ public:
 
 protected:
 	T* 		Elements = nullptr;
-	SIZE_T	MaxNum = 0;
-	SIZE_T	NumIn = 0;
-	SIZE_T	IdxIn = 0;
-	SIZE_T	IdxOut = 0;
+	int32	MaxNum = 0;
+	int32	NumIn = 0;
+	int32	IdxIn = 0;
+	int32	IdxOut = 0;
 };
 
 
@@ -237,7 +237,7 @@ public:
 	T& PushRef()
 	{
 		check(!B::IsFull());
-		SIZE_T pos = B::IdxIn;
+		int32 pos = B::IdxIn;
 		if (++B::IdxIn == B::Capacity())
 		{
 			B::IdxIn = 0;
@@ -271,13 +271,13 @@ public:
 		return B::Elements[(B::IdxIn + B::Capacity() - 1) % B::Capacity()];
 	}
 
-	void Erase(SIZE_T i)
+	void Erase(int32 i)
 	{
 		check(!B::IsEmpty());
 		check(i < B::NumIn);
 		--B::NumIn;
-		SIZE_T i0 = (i + B::IdxOut) % B::Capacity();
-		SIZE_T i1 = (i + B::IdxOut + 1) % B::Capacity();
+		int32 i0 = (i + B::IdxOut) % B::Capacity();
+		int32 i1 = (i + B::IdxOut + 1) % B::Capacity();
 		for(; i < B::NumIn; ++i)
 		{
 			B::Elements[i0] = B::Elements[i1];
@@ -298,7 +298,7 @@ public:
 		}
 	}
 
-	const T& operator[](SIZE_T i) const
+	const T& operator[](int32 i) const
 	{
 		check(!B::IsEmpty());
 		check(i < B::NumIn);
@@ -310,7 +310,7 @@ public:
 		return B::Elements[i];
 	}
 
-	T& operator[](SIZE_T i)
+	T& operator[](int32 i)
 	{
 		check(!B::IsEmpty());
 		check(i < B::NumIn);
@@ -322,12 +322,12 @@ public:
 		return B::Elements[i];
 	}
 
-	void AppendFirstElements(const MyType& FromOther, SIZE_T MaxElements = ~SIZE_T(0))
+	void AppendFirstElements(const MyType& FromOther, int32 MaxElements = ~int32(0))
 	{
 		// NOTE: Your code needs to make sure the source queue is protected by some external lock so does not get
 		//       modified while we are appending its elements.
-		SIZE_T NumAvail = B::Capacity() - B::Num();
-		SIZE_T NumToAppend = FromOther.Num();
+		int32 NumAvail = B::Capacity() - B::Num();
+		int32 NumToAppend = FromOther.Num();
 		if (NumToAppend > NumAvail)
 		{
 			NumToAppend = NumAvail;
@@ -336,18 +336,18 @@ public:
 		{
 			NumToAppend = MaxElements;
 		}
-		for(SIZE_T i = 0; i < NumToAppend; ++i)
+		for(int32 i = 0; i < NumToAppend; ++i)
 		{
 			B::Push(FromOther[i]);
 		}
 	}
 
-	void AppendLastElements(const MyType& FromOther, SIZE_T MaxElements = ~SIZE_T(0))
+	void AppendLastElements(const MyType& FromOther, int32 MaxElements = ~int32(0))
 	{
 		// NOTE: Your code needs to make sure the source queue is protected by some external lock so does not get
 		//       modified while we're appending its elements.
-		SIZE_T NumAvail = B::Capacity() - B::Num();
-		SIZE_T NumToAppend = FromOther.Num();
+		int32 NumAvail = B::Capacity() - B::Num();
+		int32 NumToAppend = FromOther.Num();
 		if (NumToAppend > NumAvail)
 		{
 			NumToAppend = NumAvail;
@@ -356,7 +356,7 @@ public:
 		{
 			NumToAppend = MaxElements;
 		}
-		for(SIZE_T i = 0, j = FromOther.Num() - NumToAppend; i < NumToAppend; ++i, ++j)
+		for(int32 i = 0, j = FromOther.Num() - NumToAppend; i < NumToAppend; ++i, ++j)
 		{
 			B::Push(FromOther[j]);
 		}
@@ -417,7 +417,7 @@ public:
 	void Clear()
 	{
 		L::Lock();
-		for(uint32 i = 0; i < NumIn; ++i)
+		for(int32 i=0; i<NumIn; ++i)
 		{
 			Elements[IdxOut].~T();
 			if (++IdxOut == CAPACITY)
@@ -429,12 +429,12 @@ public:
 		L::Unlock();
 	}
 
-	SIZE_T Capacity() const
+	int32 Capacity() const
 	{
 		return CAPACITY;
 	}
 
-	SIZE_T Num() const
+	int32 Num() const
 	{
 		return NumIn;
 	}
@@ -453,7 +453,7 @@ public:
 	{
 		L::Lock();
 		check(!IsFull());
-		SIZE_T pos = IdxIn;
+		int32 pos = IdxIn;
 		if (++IdxIn == CAPACITY)
 		{
 			IdxIn = 0;
@@ -476,7 +476,7 @@ public:
 		{
 			--IdxOut;
 		}
-		SIZE_T pos = IdxOut;
+		int32 pos = IdxOut;
 		++NumIn;
 		// Copying the element can NOT take place outside the lock in case multiple threads push at the same time and a Resize() will happen.
 		new ((void*)(Elements + pos)) T(element);
@@ -513,9 +513,9 @@ protected:
 	uint8	FixedElements[CAPACITY * sizeof(T)]
 	GCC_ALIGN(16);
 	T*		Elements = (T*)FixedElements;
-	SIZE_T	NumIn = 0;
-	SIZE_T	IdxIn = 0;
-	SIZE_T	IdxOut = 0;
+	int32	NumIn = 0;
+	int32	IdxIn = 0;
+	int32	IdxOut = 0;
 };
 
 
@@ -534,7 +534,7 @@ public:
 	T& PushRef()
 	{
 		check(!B::IsFull());
-		SIZE_T pos = B::IdxIn;
+		int32 pos = B::IdxIn;
 		if (++B::IdxIn == CAPACITY)
 		{
 			B::IdxIn = 0;
@@ -568,13 +568,13 @@ public:
 		return B::Elements[(B::IdxIn + B::Capacity() - 1) % B::Capacity()];
 	}
 
-	void Erase(SIZE_T i)
+	void Erase(int32 i)
 	{
 		check(!B::IsEmpty());
 		check(i < B::NumIn);
 		--B::NumIn;
-		SIZE_T i0 = (i + B::IdxOut) % CAPACITY;
-		SIZE_T i1 = (i + B::IdxOut + 1) % CAPACITY;
+		int32 i0 = (i + B::IdxOut) % CAPACITY;
+		int32 i1 = (i + B::IdxOut + 1) % CAPACITY;
 		for(; i < B::NumIn; ++i)
 		{
 			B::Elements[i0] = B::Elements[i1];
@@ -595,7 +595,7 @@ public:
 		}
 	}
 
-	const T& operator[](SIZE_T i) const
+	const T& operator[](int32 i) const
 	{
 		check(!B::IsEmpty());
 		check(i < B::NumIn);
@@ -607,7 +607,7 @@ public:
 		return B::Elements[i];
 	}
 
-	T& operator[](SIZE_T i)
+	T& operator[](int32 i)
 	{
 		check(!B::IsEmpty());
 		check(i < B::NumIn);
@@ -648,7 +648,7 @@ public:
 	using MyType = TMediaQueueDynamic<T, L, M>;
 	using ElementType = T;
 
-	TMediaQueueDynamic(SIZE_T initialCapacity = 0, SIZE_T increment = 32)
+	TMediaQueueDynamic(int32 initialCapacity = 0, int32 increment = 32)
 		: IncrementBy(increment)
 	{
 		Resize(initialCapacity);
@@ -661,7 +661,7 @@ public:
 		L::Unlock();
 	}
 
-	void SetIncrement(SIZE_T newIncrement)
+	void SetIncrement(int32 newIncrement)
 	{
 		L::Lock();
 		IncrementBy = newIncrement;
@@ -671,7 +671,7 @@ public:
 	void Clear()
 	{
 		L::Lock();
-		for(uint32 i = 0; i < NumIn; ++i)
+		for(int32 i=0; i<NumIn; ++i)
 		{
 			Elements[IdxOut].~T();
 			if (++IdxOut == MaxNum)
@@ -683,7 +683,7 @@ public:
 		L::Unlock();
 	}
 
-	void Resize(SIZE_T newCapacity)
+	void Resize(int32 newCapacity)
 	{
 		L::Lock();
 		if (newCapacity == 0)
@@ -698,7 +698,7 @@ public:
 			// Create a new array and copy all active elements over.
 			checkf(newCapacity >= Num(), TEXT("Cannot shrink the dynamic queue to %u with %u elements currently managed!"), (uint32)newCapacity, (uint32)Num());
 			T* pNewElements = M::Allocate(newCapacity);
-			for(SIZE_T i = 0, j = IdxOut; i < NumIn; ++i)
+			for(int32 i = 0, j = IdxOut; i < NumIn; ++i)
 			{
 				new ((void*)(pNewElements + i)) T(Elements[j]);
 				Elements[j].~T();
@@ -716,17 +716,17 @@ public:
 		L::Unlock();
 	}
 
-	SIZE_T Increment() const
+	int32 Increment() const
 	{
 		return IncrementBy;
 	}
 
-	SIZE_T Capacity() const
+	int32 Capacity() const
 	{
 		return MaxNum;
 	}
 
-	SIZE_T Num() const
+	int32 Num() const
 	{
 		return NumIn;
 	}
@@ -748,7 +748,7 @@ public:
 		{
 			Resize(Capacity() + Increment());
 		}
-		SIZE_T pos = IdxIn;
+		int32 pos = IdxIn;
 		if (++IdxIn == MaxNum)
 		{
 			IdxIn = 0;
@@ -774,7 +774,7 @@ public:
 		{
 			--IdxOut;
 		}
-		SIZE_T pos = IdxOut;
+		int32 pos = IdxOut;
 		++NumIn;
 		// Copying the element can NOT take place outside the lock in case multiple threads push at the same time and a Resize() will happen.
 		new ((void*)(Elements + pos)) T(element);
@@ -807,11 +807,11 @@ public:
 
 protected:
 	T* Elements = nullptr;
-	SIZE_T	MaxNum = 0;
-	SIZE_T	NumIn = 0;
-	SIZE_T	IdxIn = 0;
-	SIZE_T	IdxOut = 0;
-	SIZE_T	IncrementBy = 32;
+	int32	MaxNum = 0;
+	int32	NumIn = 0;
+	int32	IdxIn = 0;
+	int32	IdxOut = 0;
+	int32	IncrementBy = 32;
 };
 
 
@@ -857,7 +857,7 @@ public:
 		return B::Elements[(B::IdxIn + B::Capacity() - 1) % B::Capacity()];
 	}
 
-	const T& operator[](SIZE_T i) const
+	const T& operator[](int32 i) const
 	{
 		check(!B::IsEmpty());
 		check(i < B::NumIn);
@@ -869,7 +869,7 @@ public:
 		return B::Elements[i];
 	}
 
-	T& operator[](SIZE_T i)
+	T& operator[](int32 i)
 	{
 		check(!B::IsEmpty());
 		check(i < B::NumIn);
diff --git a/Engine/Plugins/Media/ElectraPlayer/Source/ElectraPlayerRuntime/Private/Runtime/Decoder/Android/AudioDecoderAAC_Android.cpp b/Engine/Plugins/Media/ElectraPlayer/Source/ElectraPlayerRuntime/Private/Runtime/Decoder/Android/AudioDecoderAAC_Android.cpp
--- a/Engine/Plugins/Media/ElectraPlayer/Source/ElectraPlayerRuntime/Private/Runtime/Decoder/Android/AudioDecoderAAC_Android.cpp
+++ b/Engine/Plugins/Media/ElectraPlayer/Source/ElectraPlayerRuntime/Private/Runtime/Decoder/Android/AudioDecoderAAC_Android.cpp
@@ -326,6 +326,8 @@ bool FAudioDecoderAAC::CreateDecodedSamplePool()
 	uint32 frameSize = sizeof(int16) * 8 * 2048;
 	poolOpts.Set("max_buffer_size", FVariantValue((int64) frameSize));
 	poolOpts.Set("num_buffers", FVariantValue((int64) 8));
+	poolOpts.Set("samples_per_block", FVariantValue((int64) 2048));
+	poolOpts.Set("max_channels", FVariantValue((int64) 8));
 
 	UEMediaError Error = Renderer->CreateBufferPool(poolOpts);
 	check(Error == UEMEDIA_ERROR_OK);
diff --git a/Engine/Plugins/Media/ElectraPlayer/Source/ElectraPlayerRuntime/Private/Runtime/Decoder/Android/ElectraVideoDecoder_Android.cpp b/Engine/Plugins/Media/ElectraPlayer/Source/ElectraPlayerRuntime/Private/Runtime/Decoder/Android/ElectraVideoDecoder_Android.cpp
--- a/Engine/Plugins/Media/ElectraPlayer/Source/ElectraPlayerRuntime/Private/Runtime/Decoder/Android/ElectraVideoDecoder_Android.cpp
+++ b/Engine/Plugins/Media/ElectraPlayer/Source/ElectraPlayerRuntime/Private/Runtime/Decoder/Android/ElectraVideoDecoder_Android.cpp
@@ -32,7 +32,7 @@ void FElectraPlayerVideoDecoderOutputAndroid::ShutdownPoolable()
 	TSharedPtr<IDecoderOutputOwner, ESPMode::ThreadSafe> lockedVideoRenderer = OwningRenderer.Pin();
 	if (lockedVideoRenderer.IsValid())
 	{
-		lockedVideoRenderer->SampleReleasedToPool(GetDuration());
+		lockedVideoRenderer->SampleReleasedToPool(this);
 	}
 }
 
@@ -52,4 +52,3 @@ FVideoDecoderOutput* FElectraPlayerPlatformVideoDecoderOutputFactory::Create()
 }
 
 #endif
-
diff --git a/Engine/Plugins/Media/ElectraPlayer/Source/ElectraPlayerRuntime/Private/Runtime/Decoder/Apple/AudioDecoderAAC_Apple.cpp b/Engine/Plugins/Media/ElectraPlayer/Source/ElectraPlayerRuntime/Private/Runtime/Decoder/Apple/AudioDecoderAAC_Apple.cpp
--- a/Engine/Plugins/Media/ElectraPlayer/Source/ElectraPlayerRuntime/Private/Runtime/Decoder/Apple/AudioDecoderAAC_Apple.cpp
+++ b/Engine/Plugins/Media/ElectraPlayer/Source/ElectraPlayerRuntime/Private/Runtime/Decoder/Apple/AudioDecoderAAC_Apple.cpp
@@ -348,6 +348,8 @@ bool FAudioDecoderAAC::CreateDecodedSamplePool()
 	uint32 frameSize = sizeof(int16) * 8 * 2048;
 	poolOpts.Set("max_buffer_size", FVariantValue((int64) frameSize));
 	poolOpts.Set("num_buffers", FVariantValue((int64) 8));
+	poolOpts.Set("samples_per_block", FVariantValue((int64) 2048));
+	poolOpts.Set("max_channels", FVariantValue((int64) 8));
 
 	UEMediaError Error = Renderer->CreateBufferPool(poolOpts);
 	check(Error == UEMEDIA_ERROR_OK);
diff --git a/Engine/Plugins/Media/ElectraPlayer/Source/ElectraPlayerRuntime/Private/Runtime/Decoder/Apple/ElectraVideoDecoder_Apple.cpp b/Engine/Plugins/Media/ElectraPlayer/Source/ElectraPlayerRuntime/Private/Runtime/Decoder/Apple/ElectraVideoDecoder_Apple.cpp
--- a/Engine/Plugins/Media/ElectraPlayer/Source/ElectraPlayerRuntime/Private/Runtime/Decoder/Apple/ElectraVideoDecoder_Apple.cpp
+++ b/Engine/Plugins/Media/ElectraPlayer/Source/ElectraPlayerRuntime/Private/Runtime/Decoder/Apple/ElectraVideoDecoder_Apple.cpp
@@ -54,7 +54,7 @@ void FElectraPlayerVideoDecoderOutputApple::ShutdownPoolable()
 	TSharedPtr<IDecoderOutputOwner, ESPMode::ThreadSafe> lockedVideoRenderer = OwningRenderer.Pin();
 	if (lockedVideoRenderer.IsValid())
 	{
-		lockedVideoRenderer->SampleReleasedToPool(GetDuration());
+		lockedVideoRenderer->SampleReleasedToPool(this);
 	}
 }
 
diff --git a/Engine/Plugins/Media/ElectraPlayer/Source/ElectraPlayerRuntime/Private/Runtime/Decoder/DX/AudioDecoderAAC_DX.cpp b/Engine/Plugins/Media/ElectraPlayer/Source/ElectraPlayerRuntime/Private/Runtime/Decoder/DX/AudioDecoderAAC_DX.cpp
--- a/Engine/Plugins/Media/ElectraPlayer/Source/ElectraPlayerRuntime/Private/Runtime/Decoder/DX/AudioDecoderAAC_DX.cpp
+++ b/Engine/Plugins/Media/ElectraPlayer/Source/ElectraPlayerRuntime/Private/Runtime/Decoder/DX/AudioDecoderAAC_DX.cpp
@@ -352,8 +352,9 @@ namespace Electra
 		FParamDict poolOpts;
 		uint32 frameSize = sizeof(int16) * AACDEC_MAX_CHANNELS * AACDEC_PCM_SAMPLE_SIZE;
 		poolOpts.Set("max_buffer_size", FVariantValue((int64)frameSize));
-		//	poolOpts.Set("buffer_alignment", FVariantValue((int64) 32));
 		poolOpts.Set("num_buffers", FVariantValue((int64)8));
+		poolOpts.Set("samples_per_block", FVariantValue((int64)AACDEC_PCM_SAMPLE_SIZE));
+		poolOpts.Set("max_channels", FVariantValue((int64)AACDEC_MAX_CHANNELS));
 
 		UEMediaError Error = Renderer->CreateBufferPool(poolOpts);
 		check(Error == UEMEDIA_ERROR_OK);
diff --git a/Engine/Plugins/Media/ElectraPlayer/Source/ElectraPlayerRuntime/Private/Runtime/Decoder/Windows/ElectraVideoDecoder_PC.cpp b/Engine/Plugins/Media/ElectraPlayer/Source/ElectraPlayerRuntime/Private/Runtime/Decoder/Windows/ElectraVideoDecoder_PC.cpp
--- a/Engine/Plugins/Media/ElectraPlayer/Source/ElectraPlayerRuntime/Private/Runtime/Decoder/Windows/ElectraVideoDecoder_PC.cpp
+++ b/Engine/Plugins/Media/ElectraPlayer/Source/ElectraPlayerRuntime/Private/Runtime/Decoder/Windows/ElectraVideoDecoder_PC.cpp
@@ -322,7 +322,7 @@ void FElectraPlayerVideoDecoderOutputPC::ShutdownPoolable()
 	TSharedPtr<IDecoderOutputOwner, ESPMode::ThreadSafe> lockedVideoRenderer = OwningRenderer.Pin();
 	if (lockedVideoRenderer.IsValid())
 	{
-		lockedVideoRenderer->SampleReleasedToPool(GetDuration());
+		lockedVideoRenderer->SampleReleasedToPool(this);
 	}
 
 	if (OutputType != EOutputType::HardwareWin8Plus)
diff --git a/Engine/Plugins/Media/ElectraPlayer/Source/ElectraPlayerRuntime/Private/Runtime/Demuxer/ParserISO14496-12.cpp b/Engine/Plugins/Media/ElectraPlayer/Source/ElectraPlayerRuntime/Private/Runtime/Demuxer/ParserISO14496-12.cpp
--- a/Engine/Plugins/Media/ElectraPlayer/Source/ElectraPlayerRuntime/Private/Runtime/Demuxer/ParserISO14496-12.cpp
+++ b/Engine/Plugins/Media/ElectraPlayer/Source/ElectraPlayerRuntime/Private/Runtime/Demuxer/ParserISO14496-12.cpp
@@ -517,6 +517,11 @@ private:
 			return BoxSize;
 		}
 
+		const FMP4Box* GetParentBox() const
+		{
+			return ParentBox;
+		}
+
 		void AddChildBox(FMP4Box* Child)
 		{
 			Child->ParentBox = this;
@@ -835,6 +840,11 @@ private:
 		{
 		}
 
+		uint32 GetFlags() const
+		{
+			return Flags;
+		}
+
 		virtual ~FMP4BoxFull()
 		{
 		}
@@ -2921,6 +2931,21 @@ private:
 		{
 		}
 
+		uint32 GetReferenceTrackID() const
+		{
+			return ReferenceTrackID;
+		}
+
+		uint64 GetNtpTimestamp() const
+		{
+			return NtpTimestamp;
+		}
+
+		uint64 GetMediaTime() const
+		{
+			return MediaTime;
+		}
+
 	private:
 		FMP4BoxPRFT() = delete;
 		FMP4BoxPRFT(const FMP4BoxPRFT&) = delete;
@@ -5476,6 +5501,7 @@ private:
 			virtual const FBitrateInfo& GetBitrateInfo() const override;
 			virtual const FString GetLanguage() const override;
 			virtual void GetPSSHBoxes(TArray<TArray<uint8>>& OutBoxes, bool bFromMOOV, bool bFromMOOF) const override;
+			virtual void GetPRFTBoxes(TArray<ITrack::FProducerReferenceTime>& OutBoxes) const override;
 
 			//private:
 			const FMP4BoxMVHD* MVHDBox = nullptr;
@@ -5506,6 +5532,7 @@ private:
 			const FMP4BoxTFHD* TFHDBox = nullptr;
 			const FMP4BoxTFDT* TFDTBox = nullptr;
 			TArray<const FMP4BoxTRUN*>	TRUNBoxes;
+			TArray<ITrack::FProducerReferenceTime> PRFTBoxes;
 			// Future boxes: sbgp, sgpd, subs
 
 			FStreamCodecInformation					CodecInformation;
@@ -5769,6 +5796,12 @@ private:
 		}
 	}
 
+	void FParserISO14496_12::FTrack::GetPRFTBoxes(TArray<ITrack::FProducerReferenceTime>& OutBoxes) const
+	{
+		OutBoxes = PRFTBoxes;
+	}
+
+
 	/***************************************************************************************************************************************************/
 	/***************************************************************************************************************************************************/
 	/***************************************************************************************************************************************************/
@@ -7247,7 +7280,7 @@ private:
 									}
 
 									// Get the PSSH boxes from the track fragment.
-									Box->GetAllBoxInstances(AllPSSHBoxes, FMP4Box::kBox_pssh);
+									MOOFBox->GetAllBoxInstances(AllPSSHBoxes, FMP4Box::kBox_pssh);
 									for(int32 nPSSHs=0; nPSSHs<AllPSSHBoxes.Num(); ++nPSSHs)
 									{
 										Track->PSSHBoxes.Add(static_cast<const FMP4BoxPSSH*>(AllPSSHBoxes[nPSSHs]));
@@ -7261,6 +7294,36 @@ private:
 											return Error;
 										}
 									}
+
+									// Get the PRFT boxes from the track fragment that apply to this track.
+									TArray<const FMP4Box*> AllPRFTBoxes;
+									MOOFBox->GetParentBox()->GetAllBoxInstances(AllPRFTBoxes, FMP4Box::kBox_prft);
+									for(int32 nPRFTs=0; nPRFTs<AllPRFTBoxes.Num(); ++nPRFTs)
+									{
+										const FMP4BoxPRFT* PRFTBox = static_cast<const FMP4BoxPRFT*>(AllPRFTBoxes[nPRFTs]);
+										if (PRFTBox->GetReferenceTrackID() == Track->GetID())
+										{
+											ITrack::FProducerReferenceTime prft;
+											prft.NtpTimestamp = PRFTBox->GetNtpTimestamp();
+											prft.MediaTime = PRFTBox->GetMediaTime();
+											switch(PRFTBox->GetFlags() & (8 | 16))
+											{
+												case 0:
+													prft.Reference = ITrack::FProducerReferenceTime::EReferenceType::Encoder;
+													break;
+												case 8|16:
+													prft.Reference = ITrack::FProducerReferenceTime::EReferenceType::Captured;
+													break;
+												case 16:
+													prft.Reference = ITrack::FProducerReferenceTime::EReferenceType::Application;
+													break;
+												default:
+													prft.Reference = ITrack::FProducerReferenceTime::EReferenceType::Undefined;
+													break;
+											}
+											Track->PRFTBoxes.Emplace(MoveTemp(prft));
+										}
+									}
 								}
 								else
 								{
diff --git a/Engine/Plugins/Media/ElectraPlayer/Source/ElectraPlayerRuntime/Private/Runtime/Demuxer/ParserISO14496-12.h b/Engine/Plugins/Media/ElectraPlayer/Source/ElectraPlayerRuntime/Private/Runtime/Demuxer/ParserISO14496-12.h
--- a/Engine/Plugins/Media/ElectraPlayer/Source/ElectraPlayerRuntime/Private/Runtime/Demuxer/ParserISO14496-12.h
+++ b/Engine/Plugins/Media/ElectraPlayer/Source/ElectraPlayerRuntime/Private/Runtime/Demuxer/ParserISO14496-12.h
@@ -90,6 +90,7 @@ namespace Electra
 		static const FBoxType BoxType_styp = MAKE_BOX_ATOM('s', 't', 'y', 'p');
 		static const FBoxType BoxType_moov = MAKE_BOX_ATOM('m', 'o', 'o', 'v');
 		static const FBoxType BoxType_sidx = MAKE_BOX_ATOM('s', 'i', 'd', 'x');
+		static const FBoxType BoxType_prft = MAKE_BOX_ATOM('p', 'r', 'f', 't');
 		static const FBoxType BoxType_moof = MAKE_BOX_ATOM('m', 'o', 'o', 'f');
 		static const FBoxType BoxType_mdat = MAKE_BOX_ATOM('m', 'd', 'a', 't');
 
@@ -204,12 +205,24 @@ namespace Electra
 			// Bitrate box from either 'btrt' or the 'esds' DecoderConfigDescriptor
 			struct FBitrateInfo
 			{
-				FBitrateInfo() : BufferSizeDB(0), MaxBitrate(0), AvgBitrate(0)
+				uint32	BufferSizeDB = 0;
+				uint32	MaxBitrate = 0;
+				uint32	AvgBitrate = 0;
+			};
+
+			// Producer reference time box ('prft') applicable to this track.
+			struct FProducerReferenceTime
+			{
+				enum class EReferenceType
 				{
-				}
-				uint32	BufferSizeDB;
-				uint32	MaxBitrate;
-				uint32	AvgBitrate;
+					Undefined,
+					Captured,
+					Encoder,
+					Application
+				};
+				uint64 NtpTimestamp = 0;
+				uint64 MediaTime = 0;
+				EReferenceType Reference = EReferenceType::Undefined;
 			};
 
 			virtual ~ITrack() = default;
@@ -226,6 +239,7 @@ namespace Electra
 			virtual const FBitrateInfo& GetBitrateInfo() const = 0;
 			virtual const FString GetLanguage() const = 0;
 			virtual void GetPSSHBoxes(TArray<TArray<uint8>>& OutBoxes, bool bFromMOOV, bool bFromMOOF) const = 0;
+			virtual void GetPRFTBoxes(TArray<FProducerReferenceTime>& OutBoxes) const = 0;
 		};
 
 
diff --git a/Engine/Plugins/Media/ElectraPlayer/Source/ElectraPlayerRuntime/Private/Runtime/HTTP/HTTPManager.cpp b/Engine/Plugins/Media/ElectraPlayer/Source/ElectraPlayerRuntime/Private/Runtime/HTTP/HTTPManager.cpp
--- a/Engine/Plugins/Media/ElectraPlayer/Source/ElectraPlayerRuntime/Private/Runtime/HTTP/HTTPManager.cpp
+++ b/Engine/Plugins/Media/ElectraPlayer/Source/ElectraPlayerRuntime/Private/Runtime/HTTP/HTTPManager.cpp
@@ -613,6 +613,10 @@ namespace Electra
 				}
 			}
 
+			if (Request->Parameters.bCollectTimingTraces)
+			{
+				Handle->HttpRequest->EnableTimingTraces();
+			}
 			Handle->HttpRequest->SetURL(Request->Parameters.URL);
 			Handle->HttpRequest->SetUserAgent(ELECTRA_HTTPMANAGER_USER_AGENT);
 			Handle->HttpRequest->AllowCompression(!Request->Parameters.AcceptEncoding.GetWithDefault(TEXT("")).Equals(TEXT("identity")));
@@ -777,28 +781,6 @@ namespace Electra
 		for(TMap<FHandle*, TSharedPtrTS<FRequest>>::TIterator It = ActiveRequests.CreateIterator(); It; ++It)
 		{
 			FHandle* Handle = It.Key();
-
-			// Update throughput stats when something was read.
-			HTTP::FConnectionInfo& ci = It.Value()->ConnectionInfo;
-			if (ci.Throughput.LastReadCount != ci.BytesReadSoFar)
-			{
-				FTimeValue DiffTime = Now - ci.Throughput.LastCheckTime;
-				if (DiffTime > FTimeValue::GetZero())
-				{
-					ci.Throughput.ActiveReadTime += DiffTime;
-					int64 BytesPerSecond = 1000 * ci.Throughput.LastReadCount / ci.Throughput.ActiveReadTime.GetAsMilliseconds();
-					ci.Throughput.AccumulatedBytesPerSec += BytesPerSecond;
-					if (!ci.Throughput.bIsFirst)
-					{
-						ci.Throughput.AccumulatedBytesPerSecWithoutFirst += BytesPerSecond;
-					}
-					ci.Throughput.bIsFirst = false;
-					++ci.Throughput.NumAccumulatedBytesPerSec;
-				}
-				ci.Throughput.LastReadCount = ci.BytesReadSoFar;
-			}
-			ci.Throughput.LastCheckTime = Now;
-
 			// Fire periodic progress callback?
 			if (Now >= Handle->TimeAtNextProgressCallback)
 			{
@@ -1246,6 +1228,9 @@ namespace Electra
 					HTTP::FConnectionInfo& ci = Request->ConnectionInfo;
 					ci.bIsCachedResponse = Handle->ActiveResponse.bHitCache;
 
+					// Copy all new timing traces across.
+					Response->GetTimingTraces(&ci.TimingTraces, TNumericLimits<int32>::Max());
+
 					bool bCompletedWithInsufficientData = !Request->Parameters.Verb.Equals(TEXT("HEAD")) &&
 														  Response->GetStatus() == IElectraHTTPStreamResponse::EStatus::Completed && 
 														  Response->GetResponseData().GetLengthFromResponseHeader() > Response->GetResponseData().GetNumTotalBytesAdded();
diff --git a/Engine/Plugins/Media/ElectraPlayer/Source/ElectraPlayerRuntime/Private/Runtime/HTTP/HTTPManager.h b/Engine/Plugins/Media/ElectraPlayer/Source/ElectraPlayerRuntime/Private/Runtime/HTTP/HTTPManager.h
--- a/Engine/Plugins/Media/ElectraPlayer/Source/ElectraPlayerRuntime/Private/Runtime/HTTP/HTTPManager.h
+++ b/Engine/Plugins/Media/ElectraPlayer/Source/ElectraPlayerRuntime/Private/Runtime/HTTP/HTTPManager.h
@@ -83,47 +83,6 @@ namespace Electra
 
 		struct FConnectionInfo
 		{
-			struct FThroughput
-			{
-				FThroughput()
-				{
-					Clear();
-				}
-				void Clear()
-				{
-					ActiveReadTime.SetToZero();
-					LastCheckTime.SetToInvalid();
-					LastReadCount = 0;
-					AccumulatedBytesPerSec = 0;
-					AccumulatedBytesPerSecWithoutFirst = 0;
-					NumAccumulatedBytesPerSec = 0;
-					bIsFirst = true;
-				}
-				int64 GetThroughput() const
-				{
-					if (NumAccumulatedBytesPerSec > 1)
-					{
-						return AccumulatedBytesPerSecWithoutFirst / (NumAccumulatedBytesPerSec - 1);
-					}
-					else if (NumAccumulatedBytesPerSec)
-					{
-						return AccumulatedBytesPerSec / NumAccumulatedBytesPerSec;
-					}
-					else
-					{
-						// Did not get to collect a throughput sample.
-						return 0;
-					}
-				}
-				FTimeValue				ActiveReadTime;
-				FTimeValue				LastCheckTime;
-				int64					LastReadCount;
-				int64					AccumulatedBytesPerSec;
-				int64					AccumulatedBytesPerSecWithoutFirst;
-				int32					NumAccumulatedBytesPerSec;
-				bool					bIsFirst;
-			};
-
 			TArray<FHTTPHeader>					ResponseHeaders;					//!< Response headers
 			FString								EffectiveURL;						//!< Effective URL after all redirections
 			FString								ContentType;						//!< Content-Type header
@@ -131,74 +90,82 @@ namespace Electra
 			FString  							ContentLengthHeader;
 			FTimeValue							RequestStartTime;					//!< Time at which the request was started
 			FTimeValue							RequestEndTime;						//!< Time at which the request ended
-			double								TimeForDNSResolve;					//!< Time it took to resolve DNS
-			double								TimeUntilConnected;					//!< Time it took until connected to the server
-			double								TimeUntilFirstByte;					//!< Time it took until received the first response data byte
-			int64								ContentLength;						//!< Content length, if known. Chunked transfers may have no length (set to -1). Compressed file will store compressed size here!
-			int64								BytesReadSoFar;						//!< Number of bytes read so far.
-			int32								NumberOfRedirections;				//!< Number of redirections performed
-			int32								HTTPVersionReceived;				//!< Version of HTTP header received (10 = 1.0, 11=1.1, 20=2.0)
-			bool								bIsConnected;						//!< true once connected to the server
-			bool								bHaveResponseHeaders;				//!< true when response headers have been received
-			bool								bIsChunked;							//!< true if response is received in chunks
-			bool								bWasAborted;						//!< true if transfer was aborted.
-			bool								bHasFinished;						//!< true once the connection is closed regardless of state.
-			bool								bResponseNotRanged;					//!< true if the response is not a range as was requested.
-			bool								bIsCachedResponse;					//!< true if the response came from the cache.
+			double								TimeForDNSResolve = 0.0;			//!< Time it took to resolve DNS
+			double								TimeUntilConnected = 0.0;			//!< Time it took until connected to the server
+			double								TimeUntilFirstByte = 0.0;			//!< Time it took until received the first response data byte
+			int64								ContentLength = -1;					//!< Content length, if known. Chunked transfers may have no length (set to -1). Compressed file will store compressed size here!
+			int64								BytesReadSoFar = 0;					//!< Number of bytes read so far.
+			int32								NumberOfRedirections = 0;			//!< Number of redirections performed
+			int32								HTTPVersionReceived = 0;			//!< Version of HTTP header received (10 = 1.0, 11=1.1, 20=2.0)
+			bool								bIsConnected = false;				//!< true once connected to the server
+			bool								bHaveResponseHeaders = false;		//!< true when response headers have been received
+			bool								bIsChunked = false;					//!< true if response is received in chunks
+			bool								bWasAborted = false;				//!< true if transfer was aborted.
+			bool								bHasFinished = false;				//!< true once the connection is closed regardless of state.
+			bool								bResponseNotRanged = false;			//!< true if the response is not a range as was requested.
+			bool								bIsCachedResponse = false;			//!< true if the response came from the cache.
 			FStatusInfo							StatusInfo;
 			TSharedPtrTS<FRetryInfo>			RetryInfo;
 
-			FThroughput							Throughput;
+			mutable FCriticalSection			Lock;
+			TArray<IElectraHTTPStreamResponse::FTimingTrace> TimingTraces;
+
+			void GetTimingTraces(TArray<IElectraHTTPStreamResponse::FTimingTrace>& OutTimingTraces) const
+			{
+				FScopeLock lock(&Lock);
+				OutTimingTraces = TimingTraces;
+			}
+
+			FConnectionInfo() = default;
 
-			FConnectionInfo()
+			FConnectionInfo(const FConnectionInfo& rhs)
 			{
-				TimeForDNSResolve = 0.0;
-				TimeUntilConnected = 0.0;
-				TimeUntilFirstByte = 0.0;
-				NumberOfRedirections = 0;
-				HTTPVersionReceived = 0;
-				ContentLength = -1;
-				BytesReadSoFar = 0;
-				bIsConnected = false;
-				bHaveResponseHeaders = false;
-				bIsChunked = false;
-				bWasAborted = false;
-				bHasFinished = false;
-				bResponseNotRanged = false;
-				bIsCachedResponse = false;
+				CopyFrom(rhs);
 			}
 
-			FConnectionInfo& CopyFrom(const FConnectionInfo& rhs)
+			FConnectionInfo& operator = (const FConnectionInfo& rhs)
 			{
-				ResponseHeaders = rhs.ResponseHeaders;
-				EffectiveURL = rhs.EffectiveURL;
-				ContentType = rhs.ContentType;
-				ContentRangeHeader = rhs.ContentRangeHeader;
-				ContentLengthHeader = rhs.ContentLengthHeader;
-				RequestStartTime = rhs.RequestStartTime;
-				RequestEndTime = rhs.RequestEndTime;
-				TimeForDNSResolve = rhs.TimeForDNSResolve;
-				TimeUntilConnected = rhs.TimeUntilConnected;
-				TimeUntilFirstByte = rhs.TimeUntilFirstByte;
-				ContentLength = rhs.ContentLength;
-				BytesReadSoFar = rhs.BytesReadSoFar;
-				NumberOfRedirections = rhs.NumberOfRedirections;
-				HTTPVersionReceived = rhs.HTTPVersionReceived;
-				bIsConnected = rhs.bIsConnected;
-				bHaveResponseHeaders = rhs.bHaveResponseHeaders;
-				bIsChunked = rhs.bIsChunked;
-				bWasAborted = rhs.bWasAborted;
-				bHasFinished = rhs.bHasFinished;
-				bResponseNotRanged = rhs.bResponseNotRanged;
-				bIsCachedResponse = rhs.bIsCachedResponse;
-				StatusInfo = rhs.StatusInfo;
-				Throughput = rhs.Throughput;
-				if (rhs.RetryInfo.IsValid())
+				if (this != &rhs)
 				{
-					RetryInfo = MakeSharedTS<FRetryInfo>(*rhs.RetryInfo);
+					CopyFrom(rhs);
 				}
 				return *this;
 			}
+
+			private:
+				FConnectionInfo& CopyFrom(const FConnectionInfo& rhs)
+				{
+					FScopeLock l1(&Lock);
+					FScopeLock l2(&rhs.Lock);
+					ResponseHeaders = rhs.ResponseHeaders;
+					EffectiveURL = rhs.EffectiveURL;
+					ContentType = rhs.ContentType;
+					ContentRangeHeader = rhs.ContentRangeHeader;
+					ContentLengthHeader = rhs.ContentLengthHeader;
+					RequestStartTime = rhs.RequestStartTime;
+					RequestEndTime = rhs.RequestEndTime;
+					TimeForDNSResolve = rhs.TimeForDNSResolve;
+					TimeUntilConnected = rhs.TimeUntilConnected;
+					TimeUntilFirstByte = rhs.TimeUntilFirstByte;
+					ContentLength = rhs.ContentLength;
+					BytesReadSoFar = rhs.BytesReadSoFar;
+					NumberOfRedirections = rhs.NumberOfRedirections;
+					HTTPVersionReceived = rhs.HTTPVersionReceived;
+					bIsConnected = rhs.bIsConnected;
+					bHaveResponseHeaders = rhs.bHaveResponseHeaders;
+					bIsChunked = rhs.bIsChunked;
+					bWasAborted = rhs.bWasAborted;
+					bHasFinished = rhs.bHasFinished;
+					bResponseNotRanged = rhs.bResponseNotRanged;
+					bIsCachedResponse = rhs.bIsCachedResponse;
+					StatusInfo = rhs.StatusInfo;
+					if (rhs.RetryInfo.IsValid())
+					{
+						RetryInfo = MakeSharedTS<FRetryInfo>(*rhs.RetryInfo);
+					}
+					TimingTraces = rhs.TimingTraces;
+					return *this;
+				}
 		};
 
 	} // namespace HTTP
@@ -275,6 +242,7 @@ namespace Electra
 			FTimeValue							ConnectTimeout;					//!< Optional timeout for connecting to the server
 			FTimeValue							NoDataTimeout;					//!< Optional timeout when no data is being received
 			TArray<uint8>						PostData;						//!< Data for POST
+			bool								bCollectTimingTraces = false;	//!< Whether or not to collect download timing traces.
 		};
 
 
diff --git a/Engine/Plugins/Media/ElectraPlayer/Source/ElectraPlayerRuntime/Private/Runtime/HTTP/HTTPResponseCache.cpp b/Engine/Plugins/Media/ElectraPlayer/Source/ElectraPlayerRuntime/Private/Runtime/HTTP/HTTPResponseCache.cpp
--- a/Engine/Plugins/Media/ElectraPlayer/Source/ElectraPlayerRuntime/Private/Runtime/HTTP/HTTPResponseCache.cpp
+++ b/Engine/Plugins/Media/ElectraPlayer/Source/ElectraPlayerRuntime/Private/Runtime/HTTP/HTTPResponseCache.cpp
@@ -170,6 +170,11 @@ private:
 		double GetTimeUntilFinished() override
 		{ return WrappedResponse->GetTimeUntilFinished(); }
 
+		int32 GetTimingTraces(TArray<FTimingTrace>* OutTraces, int32 InNumToRemove) override
+		{ return /*WrappedResponse->GetTimeElapsed(OutTraces, 0);*/ 0; }
+		void SetEnableTimingTraces() override
+		{ WrappedResponse->SetEnableTimingTraces(); }
+
 		void GetAllHeaders(TArray<FElectraHTTPStreamHeader>& OutHeaders) override
 		{ OutHeaders = Headers; }
 
@@ -260,6 +265,10 @@ private:
 		{ return TimeUntilFirstByte; }
 		double GetTimeUntilFinished() override
 		{ return TimeUntilFinished; }
+		int32 GetTimingTraces(TArray<FTimingTrace>* OutTraces, int32 InNumToRemove) override
+		{ return 0; }
+		void SetEnableTimingTraces() override
+		{ }
 
 		void SetTimeUntilNameResolved(double InTimeUntilNameResolved)
 		{ TimeUntilNameResolved = InTimeUntilNameResolved; }
diff --git a/Engine/Plugins/Media/ElectraPlayer/Source/ElectraPlayerRuntime/Private/Runtime/Player/AdaptivePlayerOptionKeynames.h b/Engine/Plugins/Media/ElectraPlayer/Source/ElectraPlayerRuntime/Private/Runtime/Player/AdaptivePlayerOptionKeynames.h
--- a/Engine/Plugins/Media/ElectraPlayer/Source/ElectraPlayerRuntime/Private/Runtime/Player/AdaptivePlayerOptionKeynames.h
+++ b/Engine/Plugins/Media/ElectraPlayer/Source/ElectraPlayerRuntime/Private/Runtime/Player/AdaptivePlayerOptionKeynames.h
@@ -21,9 +21,12 @@ namespace Electra
 	//! (FTimeValue) value specifying how many seconds away from the Live media timeline the seekable range should end.
 	const TCHAR* const OptionKeyLiveSeekableEndOffset = TEXT("seekable_range_live_end_offset");
 
-	//! (bool) absolute end time (excluding) of the range to limit playback to
+	//! (bool) true to just finish the currently loading segment when rebuffering. false to start over with.
 	const TCHAR* const OptionRebufferingContinuesLoading = TEXT("rebuffering_continues_loading");
 
+	//! (bool) true to throw a playback error when rebuffering occurs, false to continue normally.
+	const TCHAR* const OptionThrowErrorWhenRebuffering = TEXT("throw_error_when_rebuffering");
+
 	//! (bool) true to perform frame accurate seeks (slow as decoding and discarding data from a preceeding keyframe is required)
 	const TCHAR* const OptionKeyFrameAccurateSeek = TEXT("frame_accurate_seeking");
 
diff --git a/Engine/Plugins/Media/ElectraPlayer/Source/ElectraPlayerRuntime/Private/Runtime/Player/AdaptiveStreamingPlayer.cpp b/Engine/Plugins/Media/ElectraPlayer/Source/ElectraPlayerRuntime/Private/Runtime/Player/AdaptiveStreamingPlayer.cpp
--- a/Engine/Plugins/Media/ElectraPlayer/Source/ElectraPlayerRuntime/Private/Runtime/Player/AdaptiveStreamingPlayer.cpp
+++ b/Engine/Plugins/Media/ElectraPlayer/Source/ElectraPlayerRuntime/Private/Runtime/Player/AdaptiveStreamingPlayer.cpp
@@ -79,13 +79,14 @@ FAdaptiveStreamingPlayer::FAdaptiveStreamingPlayer(const IAdaptiveStreamingPlaye
 
 	ManifestType			= EMediaFormatType::Unknown;
 
-	VideoRender.Renderer    = InCreateParameters.VideoRenderer;
-	AudioRender.Renderer    = InCreateParameters.AudioRenderer;
+	VideoRender.Renderer    = CreateWrappedRenderer(InCreateParameters.VideoRenderer, EStreamType::Video);
+	AudioRender.Renderer    = CreateWrappedRenderer(InCreateParameters.AudioRenderer, EStreamType::Audio);
 	CurrentState   		    = EPlayerState::eState_Idle;
 	PipelineState  		    = EPipelineState::ePipeline_Stopped;
 	DecoderState   		    = EDecoderState::eDecoder_Paused;
 	StreamState			    = EStreamState::eStream_Running;
 	PlaybackRate   		    = 0.0;
+	RenderRateScale			= 1.0;
 	StreamReaderHandler     = nullptr;
 	bRebufferPending   	    = false;
 	Manifest  				= nullptr;
@@ -142,6 +143,7 @@ FAdaptiveStreamingPlayer::FAdaptiveStreamingPlayer(const IAdaptiveStreamingPlaye
 
 	TMediaInterlockedExchangePointer(PointerToLatestPlayer, this);
 
+	bUseSharedWorkerThread = true;
 	StartWorkerThread();
 }
 
@@ -297,8 +299,7 @@ void FAdaptiveStreamingPlayer::Initialize(const FParamDict& Options)
 	HttpResponseCache = IHTTPResponseCache::Create(this, PlayerOptions);
 
 	// Create the ABR stream selector.
-	PlayerConfig.StreamSelectorConfig.bRebufferingJustResumesLoading = PlayerOptions.GetValue(OptionRebufferingContinuesLoading).SafeGetBool(false);
-	StreamSelector = IAdaptiveStreamSelector::Create(this, PlayerConfig.StreamSelectorConfig);
+	StreamSelector = IAdaptiveStreamSelector::Create(this, this);
 	AddMetricsReceiver(StreamSelector.Get());
 
 	// Create renderers.
@@ -538,7 +539,7 @@ void FAdaptiveStreamingPlayer::StartWorkerThread()
 	}
 	if (!SharedWorkerThread.IsValid())
 	{
-		SharedWorkerThread = FAdaptiveStreamingPlayerWorkerThread::Create();
+		SharedWorkerThread = FAdaptiveStreamingPlayerWorkerThread::Create(bUseSharedWorkerThread);
 		WorkerThread.SetSharedWorkerThread(SharedWorkerThread);
 		SharedWorkerThread->AddPlayerInstance(this);
 	}
@@ -1116,9 +1117,9 @@ TSharedPtr<IAdaptiveStreamingPlayerResourceProvider, ESPMode::ThreadSafe> FAdapt
 	return StaticResourceProvider.Pin();
 }
 
-IElectraHttpManager* FAdaptiveStreamingPlayer::GetHTTPManager()
+TSharedPtrTS<IElectraHttpManager> FAdaptiveStreamingPlayer::GetHTTPManager()
 {
-	return HttpManager.Get();
+	return HttpManager;
 }
 
 TSharedPtrTS<IAdaptiveStreamSelector> FAdaptiveStreamingPlayer::GetStreamSelector()
@@ -1158,28 +1159,6 @@ TSharedPtrTS<FDRMManager> FAdaptiveStreamingPlayer::GetDRMManager()
 
 
 
-void FAdaptiveStreamingPlayer::GetStreamBufferStats(FAccessUnitBufferInfo& OutBufferStats, EStreamType ForStream)
-{
-	UpdateDiagnostics();
-	FMediaCriticalSection::ScopedLock lock(DiagnosticsCriticalSection);
-	switch(ForStream)
-	{
-		case EStreamType::Video:
-			OutBufferStats = VideoBufferStats.StreamBuffer;
-			break;
-		case EStreamType::Audio:
-			OutBufferStats = AudioBufferStats.StreamBuffer;
-			break;
-		case EStreamType::Subtitle:
-			OutBufferStats = TextBufferStats.StreamBuffer;
-			break;
-		default:
-			OutBufferStats.Clear();
-			break;
-	}
-}
-
-
 //-----------------------------------------------------------------------------
 /**
  * Notified when the given fragment will be opened.
@@ -1289,7 +1268,12 @@ void FAdaptiveStreamingPlayer::InternalHandleOnce()
 		{
 			if (RenderClock->IsRunning())
 			{
-				FTimeValue playPos = RenderClock->GetInterpolatedRenderTime(IMediaRenderClock::ERendererType::Video);
+				IMediaRenderClock::ERendererType ClockSource = IMediaRenderClock::ERendererType::Video;
+				if (DataAvailabilityStateAud.Availability == Metrics::FDataAvailabilityChange::EAvailability::DataAvailable)
+				{
+					ClockSource = IMediaRenderClock::ERendererType::Audio;
+				}
+				FTimeValue playPos = RenderClock->GetInterpolatedRenderTime(ClockSource);
 				PlaybackState.SetPlayPosition(ClampTimeToCurrentRange(playPos, false, true));
 			}
 		}
@@ -1344,6 +1328,15 @@ void FAdaptiveStreamingPlayer::InternalHandleOnce()
 		{
 			DrmManager->Tick();
 		}
+		// Handle ABR.
+		if (StreamSelector.IsValid())
+		{
+			IAdaptiveStreamSelector::EHandlingAction ABRAction = StreamSelector->PeriodicHandle();
+			if (ABRAction == IAdaptiveStreamSelector::EHandlingAction::SeekToLive)
+			{
+				InternalStartoverAtLiveEdge();
+			}
+		}
 	}
 }
 
@@ -1611,6 +1604,7 @@ void FAdaptiveStreamingPlayer::HandlePlayStateChanges()
 		// Are we not paused and in a state in which we can pause?
 		if (CurrentState == EPlayerState::eState_Playing)
 		{
+			PlaybackState.SetShouldPlayOnLiveEdge(false);
 			InternalPause();
 		}
 	}
@@ -1634,6 +1628,9 @@ void FAdaptiveStreamingPlayer::HandlePlayStateChanges()
 			InternalStartoverAtCurrentPosition();
 		}
 	}
+
+	// Update the current live latency.
+	PlaybackState.SetCurrentLiveLatency(CalculateCurrentLiveLatency(false));
 }
 
 
@@ -1677,6 +1674,8 @@ void FAdaptiveStreamingPlayer::HandleSeeking()
 			}
 		}
 
+		// Not playing on the Live edge. This may get set to true later on handling the start segment.
+		PlaybackState.SetShouldPlayOnLiveEdge(false);
 		// Trigger the seek.
 		FSeekParam SeekParam = SeekVars.PendingRequest.GetValue();
 		SeekVars.PendingRequest.Reset();
@@ -1893,6 +1892,41 @@ void FAdaptiveStreamingPlayer::HandleSessionMessage(TSharedPtrTS<IPlayerMessage>
 
 
 
+
+FTimeValue FAdaptiveStreamingPlayer::CalculateCurrentLiveLatency(bool bViaLatencyElement)
+{
+	FTimeValue LiveLatency;
+	if (Manifest.IsValid())
+	{
+		if (Manifest->GetPresentationType() != IManifest::EType::OnDemand)
+		{
+			FTimeValue UTCNow = GetSynchronizedUTCTime()->GetTime();
+			FTimeValue PlayPosNow = PlaybackState.GetPlayPosition();
+			LiveLatency = UTCNow - PlayPosNow;
+
+			if (bViaLatencyElement)
+			{
+				TSharedPtrTS<const FLowLatencyDescriptor> llDesc = Manifest->GetLowLatencyDescriptor();
+				if (llDesc.IsValid())
+				{
+					// Low latency Live
+					TSharedPtrTS<IProducerReferenceTimeInfo> ProdRefTime = Manifest->GetProducerReferenceTimeInfo(llDesc->Latency.ReferenceID);
+					if (ProdRefTime.IsValid())
+					{
+						FTimeValue EncoderLatency = PlaybackState.GetEncoderLatency();
+						if (EncoderLatency.IsValid())
+						{
+							LiveLatency += EncoderLatency;
+						}
+					}
+				}
+			}
+		}
+	}
+	return LiveLatency;
+}
+
+
 //-----------------------------------------------------------------------------
 /**
  * Returns the required duration to be available in the buffer before playback can begin.
@@ -1903,9 +1937,10 @@ double FAdaptiveStreamingPlayer::GetMinBufferTimeBeforePlayback()
 	double kMinBufferBeforePlayback = LastBufferingState == EPlayerState::eState_Seeking ? PlayerConfig.SeekBufferMinTimeAvailBeforePlayback :
 										LastBufferingState == EPlayerState::eState_Rebuffering ? PlayerConfig.RebufferMinTimeAvailBeforePlayback : PlayerConfig.InitialBufferMinTimeAvailBeforePlayback;
 
-	// The buffer times may not always exactly align with segment durations, especially when different timescales are involved.
-	// For the time being we tweak the buffering value a bit to account for rounding errors.
-	kMinBufferBeforePlayback *= 0.98;
+	FTimeValue mbtAbr = StreamSelector.IsValid() ? StreamSelector->GetMinBufferTimeForPlayback(LastBufferingState == EPlayerState::eState_Seeking ? IAdaptiveStreamSelector::EMinBufferType::Seeking
+																							 : LastBufferingState == EPlayerState::eState_Rebuffering ? IAdaptiveStreamSelector::EMinBufferType::Rebuffering 
+																							 :	IAdaptiveStreamSelector::EMinBufferType::Initial, Manifest->GetMinBufferTime()) : FTimeValue();
+	kMinBufferBeforePlayback = mbtAbr.IsValid() ? mbtAbr.GetAsSeconds() : kMinBufferBeforePlayback;
 	return kMinBufferBeforePlayback;
 }
 
@@ -1917,18 +1952,17 @@ void FAdaptiveStreamingPlayer::InternalHandleLiveSync()
 	// They may have been tagged already if frame accurate seeking is enabled, but the real time NOW
 	// between the time the segment download was started until, well, NOW, has continued to elapse
 	// so any access units in between need to be tagged as well.
-	if (LiveSyncVars.bSyncLiveToNow)
+	if (LiveSyncVars.bSyncLiveToNow && Manifest.IsValid())
 	{
 		FTimeRange Seekable = Manifest->GetSeekableTimeRange();
 		FTimeRange TotalRange = Manifest->GetTotalTimeRange();
-		double minBufTime = Manifest ? Manifest->GetMinBufferTime().GetAsSeconds() : 1.0;
+		double minBufTime = Manifest->GetMinBufferTime().GetAsSeconds();
 
 		// For now attempting to sync to Live works only when there is a sufficiently large distance
 		// between the Live edge and the playback position. The difference can be found by examining
 		// the end value of the total vs. seekable range.
 		FTimeValue DistanceToLiveEdge = TotalRange.End - Seekable.End;
-		// FIXME: The *2 and *0.99 fudge values need to be adjusted once chunked download brings in new data constantly instead of whole segments.
-		if (!DistanceToLiveEdge.IsValid() || DistanceToLiveEdge.GetAsSeconds() < 2 * minBufTime * 0.99)
+		if (!DistanceToLiveEdge.IsValid() || DistanceToLiveEdge.GetAsSeconds() < minBufTime)
 		{
 			PostLog(Facility::EFacility::Player, IInfoLog::ELevel::Warning, FString::Printf(TEXT("Live sync not possible due to insufficient distance to Live edge (%.3fs) compared to minimum buffer time (%.3fs)"), DistanceToLiveEdge.GetAsSeconds(), minBufTime));
 			LiveSyncVars.Clear();
@@ -2248,26 +2282,6 @@ FTimeValue FAdaptiveStreamingPlayer::GetCurrentPlayTime()
 	{
 		Current = PlaybackState.GetPlayPosition();
 	}
-
-	// Check for clock drift
-	// FIXME: This is until such time when the application gives us feedback at which time it is currently
-	//        rendering video or audio. We have no control here over how the application handles the decoded
-	//        output or how deep its pipeline is. Our own timestamps for decoding are ahead of the time
-	//        the output is rendered, but we do not know by how much.
-	if (Current.IsValid() && RenderClock->IsRunning())
-	{
-		FTimeValue clockPos = RenderClock->GetInterpolatedRenderTime(IMediaRenderClock::ERendererType::Video);
-		if (clockPos.IsValid())
-		{
-			static const FTimeValue ClockDriftThreshold(FTimeValue::MillisecondsToHNS(200));
-			if (Current.GetSequenceIndex() != clockPos.GetSequenceIndex() || clockPos > Current || clockPos + ClockDriftThreshold < Current)
-			{
-				RenderClock->SetCurrentTime(IMediaRenderClock::ERendererType::Video, Current);
-				RenderClock->SetCurrentTime(IMediaRenderClock::ERendererType::Audio, Current);
-				PlaybackState.SetPlayPosition(Current);
-			}
-		}
-	}
 	return Current;
 }
 
@@ -2465,95 +2479,111 @@ void FAdaptiveStreamingPlayer::InternalHandlePendingStartRequest(const FTimeValu
 	// Is there a start request?
 	if (PendingStartRequest.IsValid())
 	{
+		FPlayStartPosition StartAt = PendingStartRequest->StartAt;
+
 		// Need to have a manifest.
 		if (Manifest.IsValid())
 		{
+			if (PendingStartRequest->RetryAtTime.IsValid() && CurrentTime < PendingStartRequest->RetryAtTime)
+			{
+				return;
+			}
+
 			// For sanities sake disable any loop flag that might have been set before loading the playlist if the presentation isn't VoD.
 			if (Manifest->GetPresentationType() != IManifest::EType::OnDemand)
 			{
 				CurrentLoopParam.bEnableLooping = false;
 			}
 
-			// Do we have a play period yet?
-			if (!InitialPlayPeriod.IsValid())
+			// If the starting time has not been set we now check if we are dealing with a VoD or a Live stream and choose a starting point.
+			if (!StartAt.Time.IsValid())
 			{
-				// No. Check if we should ask for one now.
-				if (!PendingStartRequest->RetryAtTime.IsValid() || CurrentTime >= PendingStartRequest->RetryAtTime)
+				// Use the presentation provided start time, if it has one.
+				StartAt.Time = Manifest->GetDefaultStartTime();
+				if (!StartAt.Time.IsValid())
 				{
-					// If the starting time has not been set we now check if we are dealing with a VoD or a Live stream and choose a starting point.
-					if (!PendingStartRequest->StartAt.Time.IsValid())
+					FTimeRange Seekable = Manifest->GetSeekableTimeRange();
+					if (Manifest->GetPresentationType() == IManifest::EType::OnDemand)
 					{
-						// Use the presentation provided start time, if it has one.
-						PendingStartRequest->StartAt.Time = Manifest->GetDefaultStartTime();
-						if (!PendingStartRequest->StartAt.Time.IsValid())
-						{
-							FTimeRange Seekable = Manifest->GetSeekableTimeRange();
-							if (Manifest->GetPresentationType() == IManifest::EType::OnDemand)
-							{
-								check(Seekable.Start.IsValid());
-								PendingStartRequest->StartAt.Time = Seekable.Start;
-							}
-							else
-							{
-								check(Seekable.End.IsValid());
-								PendingStartRequest->StartAt.Time = Seekable.End;
-								// DASH Live streams provide the means to join at an exact wallclock time, so enable this.
-								if (ManifestType == EMediaFormatType::DASH)
-								{
-									LiveSyncVars.Clear();
-									LiveSyncVars.bSyncLiveToNow = true;
-								}
-							}
-						}
+						check(Seekable.Start.IsValid());
+						StartAt.Time = Seekable.Start;
 					}
-					// Set the current playback range with the start options and clamp the start time to the range.
-					SetPlaystartOptions(PendingStartRequest->StartAt.Options);
-					ClampStartRequestTime();
-
-					IManifest::FResult Result = Manifest->FindPlayPeriod(InitialPlayPeriod, PendingStartRequest->StartAt, PendingStartRequest->SearchType);
-					switch(Result.GetType())
+					else
 					{
-						case IManifest::FResult::EType::TryAgainLater:
+						/*
+							We need to wait for the timeline to become valid. It may not be yet if the presentation
+							is scheduled to start in the future or if the intended distance to the Live edge can not
+							be maintained yet when the presentation has just begun.
+						*/
+						if (Seekable.End - Seekable.Start <= FTimeValue::GetZero())
 						{
-							PendingStartRequest->RetryAtTime = Result.GetRetryAgainAtTime();
-							break;
+							return;
 						}
-						case IManifest::FResult::EType::Found:
+
+						check(Seekable.End.IsValid());
+						StartAt.Time = Seekable.End;
+						// DASH Live streams provide the means to join at an exact wallclock time, so enable this.
+						if (ManifestType == EMediaFormatType::DASH)
 						{
-							// NOTE: Do *not* reset the pending start request yet. It is still needed.
-							break;
+							LiveSyncVars.Clear();
+							// Live sync is no longer required since DASH Live maintains sync with the Live edge now.
+							//LiveSyncVars.bSyncLiveToNow = true;
+							PlaybackState.SetShouldPlayOnLiveEdge(true);
 						}
-						case IManifest::FResult::EType::PastEOS:
+					}
+				}
+			}
+			// Set the current playback range with the start options and clamp the start time to the range.
+			SetPlaystartOptions(StartAt.Options);
+			ClampStartRequestTime(StartAt.Time);
+
+
+			// Do we have a play period yet?
+			if (!InitialPlayPeriod.IsValid())
+			{
+				IManifest::FResult Result = Manifest->FindPlayPeriod(InitialPlayPeriod, StartAt, PendingStartRequest->SearchType);
+				switch(Result.GetType())
+				{
+					case IManifest::FResult::EType::TryAgainLater:
+					{
+						PendingStartRequest->RetryAtTime = Result.GetRetryAgainAtTime();
+						break;
+					}
+					case IManifest::FResult::EType::Found:
+					{
+						// NOTE: Do *not* reset the pending start request yet. It is still needed.
+						break;
+					}
+					case IManifest::FResult::EType::PastEOS:
+					{
+						// If the initial start time is beyond the end of the presentation
+						// but looping is enabled we set this request up to start at the beginning.
+						if (CurrentLoopParam.bEnableLooping)
 						{
-							// If the initial start time is beyond the end of the presentation
-							// but looping is enabled we set this request up to start at the beginning.
-							if (CurrentLoopParam.bEnableLooping)
-							{
-								FTimeRange Seekable = Manifest->GetSeekableTimeRange();
-								PendingStartRequest->SearchType = IManifest::ESearchType::Closest;
-								PendingStartRequest->StartAt.Time = Seekable.Start;
-							}
-							else
-							{
-								PendingStartRequest.Reset();
-								InternalSetPlaybackEnded();
-							}
-							break;
+							FTimeRange Seekable = Manifest->GetSeekableTimeRange();
+							PendingStartRequest->SearchType = IManifest::ESearchType::Closest;
+							StartAt.Time = Seekable.Start;
 						}
-						case IManifest::FResult::EType::NotFound:
-						case IManifest::FResult::EType::BeforeStart:
-						case IManifest::FResult::EType::NotLoaded:
+						else
 						{
-							// Throw a playback error for now.
-							FErrorDetail err;
-							err.SetFacility(Facility::EFacility::Player);
-							err.SetMessage("Could not locate the playback period to begin playback at.");
-							err.SetCode(INTERR_COULD_NOT_LOCATE_START_PERIOD);
-							PostError(err);
-							InitialPlayPeriod.Reset();
 							PendingStartRequest.Reset();
-							break;
+							InternalSetPlaybackEnded();
 						}
+						break;
+					}
+					case IManifest::FResult::EType::NotFound:
+					case IManifest::FResult::EType::BeforeStart:
+					case IManifest::FResult::EType::NotLoaded:
+					{
+						// Throw a playback error for now.
+						FErrorDetail err;
+						err.SetFacility(Facility::EFacility::Player);
+						err.SetMessage("Could not locate the playback period to begin playback at.");
+						err.SetCode(INTERR_COULD_NOT_LOCATE_START_PERIOD);
+						PostError(err);
+						InitialPlayPeriod.Reset();
+						PendingStartRequest.Reset();
+						break;
 					}
 				}
 			}
@@ -2673,15 +2703,11 @@ void FAdaptiveStreamingPlayer::InternalHandlePendingStartRequest(const FTimeValu
 						CurrentPlayPeriodVideo = InitialPlayPeriod;
 						CurrentPlayPeriodAudio = InitialPlayPeriod;
 						CurrentPlayPeriodText = InitialPlayPeriod;
+
 						// Tell the ABR the current playback period. At playback start it is the same for all stream types.
 						StreamSelector->SetCurrentPlaybackPeriod(EStreamType::Video, CurrentPlayPeriodVideo);
 						StreamSelector->SetCurrentPlaybackPeriod(EStreamType::Audio, CurrentPlayPeriodAudio);
 
-						// Single stream mp4 files cannot switch streams. Failures need to retry the same file.
-						if (ManifestType == EMediaFormatType::ISOBMFF)
-						{
-							StreamSelector->SetCanSwitchToAlternateStreams(false);
-						}
 						TSharedPtrTS<FBufferSourceInfo> BufferSourceInfoVid = CurrentPlayPeriodVideo->GetSelectedStreamBufferSourceInfo(EStreamType::Video);
 						TSharedPtrTS<FBufferSourceInfo> BufferSourceInfoAud = CurrentPlayPeriodAudio->GetSelectedStreamBufferSourceInfo(EStreamType::Audio);
 						TSharedPtrTS<FBufferSourceInfo> BufferSourceInfoTxt = CurrentPlayPeriodText->GetSelectedStreamBufferSourceInfo(EStreamType::Subtitle);
@@ -2713,19 +2739,6 @@ void FAdaptiveStreamingPlayer::InternalHandlePendingStartRequest(const FTimeValu
 						// Apply any set resolution limit at the start of the new play period.
 						UpdateStreamResolutionLimit();
 
-						//
-						switch(Manifest->GetPresentationType())
-						{
-							case IManifest::EType::OnDemand:
-								StreamSelector->SetPresentationType(IAdaptiveStreamSelector::EMediaPresentationType::OnDemand);
-								break;
-							// FIXME: Need to differentiate between "Live" and "Realtime" here!
-							default:
-							case IManifest::EType::Live:
-								StreamSelector->SetPresentationType(IAdaptiveStreamSelector::EMediaPresentationType::Realtime);
-								break;
-						}
-
 						// Have the ABR select the initial bandwidth. Pass it an empty segment request to indicate this.
 						FTimeValue ActionDelay;
 						IAdaptiveStreamSelector::ESegmentAction Action = StreamSelector->SelectSuitableStreams(ActionDelay, TSharedPtrTS<IStreamSegment>());
@@ -2747,14 +2760,14 @@ void FAdaptiveStreamingPlayer::InternalHandlePendingStartRequest(const FTimeValu
 							IManifest::FResult Result;
 							if (PendingStartRequest->StartType != FPendingStartRequest::EStartType::LoopPoint)
 							{
-								Result = InitialPlayPeriod->GetStartingSegment(FirstSegmentRequest, CurrentPlaybackSequenceState, PendingStartRequest->StartAt, PendingStartRequest->SearchType);
+								Result = InitialPlayPeriod->GetStartingSegment(FirstSegmentRequest, CurrentPlaybackSequenceState, StartAt, PendingStartRequest->SearchType);
 								bFirstSegmentRequestIsForLooping = false;
 							}
 							else
 							{
 								FPlayerSequenceState NextPlaybackState = CurrentPlaybackSequenceState;
 								++NextPlaybackState.SequenceIndex;
-								Result = InitialPlayPeriod->GetLoopingSegment(FirstSegmentRequest, NextPlaybackState, PendingStartRequest->FinishedRequests, PendingStartRequest->StartAt, PendingStartRequest->SearchType);
+								Result = InitialPlayPeriod->GetLoopingSegment(FirstSegmentRequest, NextPlaybackState, PendingStartRequest->FinishedRequests, StartAt, PendingStartRequest->SearchType);
 								if (Result.GetType() == IManifest::FResult::EType::Found)
 								{
 									CurrentPlaybackSequenceState = NextPlaybackState;
@@ -2770,6 +2783,52 @@ void FAdaptiveStreamingPlayer::InternalHandlePendingStartRequest(const FTimeValu
 								}
 								case IManifest::FResult::EType::Found:
 								{
+									// For DASH Live check if it would be better to start on the next segment.
+									if (ManifestType == EMediaFormatType::DASH &&
+										Manifest->GetPresentationType() != IManifest::EType::OnDemand)
+									{
+										struct FSegTime { FTimeValue Start, Into, Duration; };
+										FSegTime t0;
+										if (FirstSegmentRequest->GetStartupDelay(t0.Start, t0.Into, t0.Duration))
+										{
+											TSharedPtrTS<IStreamSegment> NextStartSegmentRequest;
+											FPlayStartPosition LaterStart = StartAt;
+											FTimeValue RemainingSegmentDuration = t0.Duration - t0.Into;
+											LaterStart.Time += RemainingSegmentDuration;
+											LaterStart.Options.bFrameAccuracy = false;
+											IManifest::FResult LaterResult = InitialPlayPeriod->GetStartingSegment(NextStartSegmentRequest, CurrentPlaybackSequenceState, LaterStart, IManifest::ESearchType::Closest);
+											if (LaterResult.GetType() == IManifest::FResult::EType::Found && NextStartSegmentRequest.IsValid())
+											{
+												// Are we skipping more than half the segment's duration?
+												if (t0.Into >= t0.Duration/2)
+												{
+													// Start with the next segment if this segment has less than 1s of usable duration.
+													bool bStartWithNextSegment = RemainingSegmentDuration.GetAsSeconds() < 1.0;
+
+													// If using the next segment would put us under the minimum required latency we cannot start with it
+													// if we are to play on the Live edge.
+													if (ABRShouldPlayOnLiveEdge())
+													{
+														double NextSegStartLiveLatency = (ABRGetWallclockTime() - NextStartSegmentRequest->GetFirstPTS()).GetAsSeconds();
+														double MinLiveLatency = 0.0;
+														TSharedPtrTS<const FLowLatencyDescriptor> llDesc = ABRGetLowLatencyDescriptor();
+														MinLiveLatency = llDesc.IsValid() ? llDesc->GetLatencyMin().GetAsSeconds(MinLiveLatency) : MinLiveLatency;
+														if (NextSegStartLiveLatency < MinLiveLatency)
+														{
+															bStartWithNextSegment = false;
+														}
+													}
+
+													if (bStartWithNextSegment)
+													{
+														PostLog(Facility::EFacility::Player, IInfoLog::ELevel::Verbose, FString::Printf(TEXT("Using next segment since it has a nearer start time.")));
+														FirstSegmentRequest = NextStartSegmentRequest;
+													}
+												}
+											}
+										}
+									}
+
 									// Add this period to the list of upcoming ones. This ensures the period metadata gets reported
 									// if the period playback starts in is not the first for which metadata was initially reported.
 									AddUpcomingPeriod(InitialPlayPeriod);
@@ -3047,6 +3106,7 @@ void FAdaptiveStreamingPlayer::HandlePendingMediaSegmentRequests()
 		return;
 	}
 
+	FTimeValue UTCNow = GetSynchronizedUTCTime()->GetTime();
 	FTimeValue CurrentTime = MEDIAutcTime::Current();
 	InternalHandlePendingStartRequest(CurrentTime);
 	InternalHandlePendingFirstSegmentRequest(CurrentTime);
@@ -3300,7 +3360,7 @@ void FAdaptiveStreamingPlayer::HandlePendingMediaSegmentRequests()
 						UpdatePeriodStreamBufferSourceInfo(SegmentPlayPeriod, StreamType, BufferSourceInfo);
 					}
 
-					NextSegment->SetExecutionDelay(ActionDelay);
+					NextSegment->SetExecutionDelay(UTCNow, ActionDelay);
 					ReadyWaitingSegmentRequests.Enqueue(NextSegment);
 					break;
 				}
@@ -3399,7 +3459,6 @@ void FAdaptiveStreamingPlayer::HandlePendingMediaSegmentRequests()
 	// Go through the queue of ready-to-fetch segments to see if they should be fetched now.
 	TQueue<TSharedPtrTS<IStreamSegment>> readyRequests;
 	Swap(readyRequests, ReadyWaitingSegmentRequests);
-	FTimeValue UTCNow = GetSynchronizedUTCTime()->GetTime();
 	while(!readyRequests.IsEmpty())
 	{
 		TSharedPtrTS<IStreamSegment> ReadyReq;
@@ -3536,6 +3595,15 @@ void FAdaptiveStreamingPlayer::InternalStartoverAtCurrentPosition()
 	InternalStartAt(NewPosition);
 }
 
+void FAdaptiveStreamingPlayer::InternalStartoverAtLiveEdge()
+{
+	FSeekParam NewPosition;
+	InternalStop(PlayerConfig.bHoldLastFrameDuringSeek);
+	CurrentState = EPlayerState::eState_Seeking;
+	PlaybackState.GetLoopState(CurrentLoopState);
+	InternalStartAt(NewPosition);
+}
+
 
 void FAdaptiveStreamingPlayer::InternalHandleSegmentTrackChanges(const FTimeValue& CurrentTime)
 {
@@ -3863,25 +3931,25 @@ void FAdaptiveStreamingPlayer::SetPlaystartOptions(FPlayStartOptions& OutOptions
 	OutOptions.bFrameAccuracy = GetOptions().GetValue(OptionKeyFrameAccurateSeek).SafeGetBool(false);
 }
 
-void FAdaptiveStreamingPlayer::ClampStartRequestTime()
+void FAdaptiveStreamingPlayer::ClampStartRequestTime(FTimeValue& InOutTimeToClamp)
 {
 	if (PendingStartRequest.IsValid())
 	{
 		// Clamp to within any set playback range
 		FTimeValue RangeStart = PendingStartRequest->StartAt.Options.PlaybackRange.Start;
 		FTimeValue RangeEnd = PendingStartRequest->StartAt.Options.PlaybackRange.End;
-		if (PendingStartRequest->StartAt.Time < RangeStart)
+		if (InOutTimeToClamp < RangeStart)
 		{
-			PendingStartRequest->StartAt.Time = RangeStart;
+			InOutTimeToClamp = RangeStart;
 		}
-		if (PendingStartRequest->StartAt.Time > RangeEnd)
+		if (InOutTimeToClamp > RangeEnd)
 		{
-			PendingStartRequest->StartAt.Time = RangeEnd;
+			InOutTimeToClamp = RangeEnd;
 			// Going to the end will of course result in an 'end reached' and end of playback.
 			// If the player is set to loop then we should start from the beginning.
 			if (CurrentLoopParam.bEnableLooping)
 			{
-				PendingStartRequest->StartAt.Time = RangeStart;
+				InOutTimeToClamp = RangeStart;
 			}
 		}
 	}
@@ -4129,6 +4197,7 @@ void FAdaptiveStreamingPlayer::InternalStartAt(const FSeekParam& NewPosition)
 	PlaybackState.SetHasEnded(false);
 	PlaybackState.SetPlayRangeHasChanged(false);
 	PlaybackState.SetLoopStateHasChanged(false);
+	PlaybackState.SetEncoderLatency(FTimeValue::GetInvalid());
 
 	VideoBufferStats.Clear();
 	AudioBufferStats.Clear();
@@ -4223,6 +4292,7 @@ void FAdaptiveStreamingPlayer::InternalSetPlaybackEnded()
 	SeekVars.Reset();
 
 	PlaybackState.SetHasEnded(true);
+	PlaybackState.SetShouldPlayOnLiveEdge(false);
 	StreamState = EStreamState::eStream_ReachedEnd;
 
 	// Go to pause mode now.
@@ -4402,13 +4472,13 @@ void FAdaptiveStreamingPlayer::InternalClose()
 		DrmManager.Reset();
 	}
 
+	RemoveMetricsReceiver(StreamSelector.Get());
+	StreamSelector.Reset();
+
 	HttpManager.Reset();
 	EntityCache.Reset();
 	HttpResponseCache.Reset();
 
-	RemoveMetricsReceiver(StreamSelector.Get());
-	StreamSelector.Reset();
-
 	// Reset remaining internal state
 	CurrentState   	= EPlayerState::eState_Idle;
 	PlaybackRate   	= 0.0;
@@ -4495,9 +4565,25 @@ void FAdaptiveStreamingPlayer::InternalSetLoop(const FLoopParam& LoopParam)
  */
 void FAdaptiveStreamingPlayer::InternalRebuffer()
 {
+
+	IAdaptiveStreamSelector::FRebufferAction Action;
+	check(StreamSelector.IsValid());
+	if (StreamSelector.IsValid())
+	{
+		Action = StreamSelector->GetRebufferAction(GetOptions());
+	}
+
 	// Check if we are configured to throw a playback error instead of doing a rebuffer.
 	// This can be used by a higher logic to force playback at a different position instead.
-	if (PlayerOptions.GetValue("throw_error_when_rebuffering").SafeGetBool(false) == false)
+	if (Action.Action == IAdaptiveStreamSelector::FRebufferAction::EAction::ThrowError)
+	{
+		FErrorDetail err;
+		err.SetFacility(Facility::EFacility::Player);
+		err.SetMessage("Rebuffering is set to generate a playback error");
+		err.SetCode(INTERR_REBUFFER_SHALL_THROW_ERROR);
+		PostError(err);
+	}
+	else
 	{
 		// Pause decoders to stop feeding them data.
 		// NOTE: A decoder may be in InputNeeded() / FeedDecoder() at this point!
@@ -4513,8 +4599,8 @@ void FAdaptiveStreamingPlayer::InternalRebuffer()
 
 		bRebufferPending = false;
 
-		// Check if we should just continue loading data (for VoD presentations only!)
-		if (GetOptions().GetValue(OptionRebufferingContinuesLoading).SafeGetBool(false) && Manifest->GetPresentationType() == IManifest::EType::OnDemand)
+		// Check if we should just continue loading data
+		if (Action.Action == IAdaptiveStreamSelector::FRebufferAction::EAction::ContinueLoading)
 		{
 			RebufferDetectedAtPlayPos.SetToInvalid();
 			PostLog(Facility::EFacility::Player, IInfoLog::ELevel::Info, FString::Printf(TEXT("Rebuffering is waiting for new data.")));
@@ -4523,14 +4609,12 @@ void FAdaptiveStreamingPlayer::InternalRebuffer()
 		}
 		else
 		{
-			// We do not wait on the current segment that caused the rebuffering to complete.
-			// Instead we perform a seek to that time and start over from scratch.
+			// We do not wait on the current segment that caused the rebuffering to complete and start over.
 			FSeekParam StartAtTime;
-			if (Manifest->GetPresentationType() == IManifest::EType::Live)
+			if (Action.Action == IAdaptiveStreamSelector::FRebufferAction::EAction::GoToLive)
 			{
-				// A Live presentation shall implicitly go to the Live edge.
 				StartAtTime.Time.SetToInvalid();
-				PostLog(Facility::EFacility::Player, IInfoLog::ELevel::Info, FString::Printf(TEXT("Rebuffering Live stream implicitly goes to the live edge")));
+				PostLog(Facility::EFacility::Player, IInfoLog::ELevel::Info, FString::Printf(TEXT("Rebuffering jumps to the live edge")));
 			}
 			else
 			{
@@ -4539,41 +4623,154 @@ void FAdaptiveStreamingPlayer::InternalRebuffer()
 					RebufferDetectedAtPlayPos = PlaybackState.GetPlayPosition();
 				}
 				StartAtTime.Time = RebufferDetectedAtPlayPos;
-				PostLog(Facility::EFacility::Player, IInfoLog::ELevel::Info, FString::Printf(TEXT("Rebuffering on-demand stream at current position")));
+				PostLog(Facility::EFacility::Player, IInfoLog::ELevel::Info, FString::Printf(TEXT("Rebuffering starts over at current position")));
 			}
 			RebufferDetectedAtPlayPos.SetToInvalid();
 			InternalStop(PlayerConfig.bHoldLastFrameDuringSeek);
 			CurrentState = EPlayerState::eState_Rebuffering;
-			if (PlayerOptions.HaveKey(OptionKeyRebufferStartBitrate))
+			if (Action.SuggestedRestartBitrate > 0)
 			{
-				int64 StartingBitrate = PlayerOptions.GetValue(OptionKeyRebufferStartBitrate).SafeGetInt64(0);
-				if (StartingBitrate > 0)
-				{
-					StartAtTime.StartingBitrate = (int32)StartingBitrate;
-				}
-			}
-
-			// Force a rate scale in the ABR. This may get overridden by the starting bitrate later.
-			if (StreamSelector.IsValid())
-			{
-				IAdaptiveStreamSelector::FBufferingQuality RateScale;
-				RateScale.BandwidthScaleFactor = 0.6;
-				StreamSelector->SwitchBufferingQuality(RateScale);
+				StartAtTime.StartingBitrate = Action.SuggestedRestartBitrate;
 			}
 
 			InternalStartAt(StartAtTime);
 		}
 	}
-	else
+}
+
+
+
+
+
+
+
+
+FTimeValue FAdaptiveStreamingPlayer::ABRGetPlayPosition() const
+{
+	return PlaybackState.GetPlayPosition();
+}
+FTimeRange FAdaptiveStreamingPlayer::ABRGetTimeline() const
+{
+	FTimeRange Timeline;
+	PlaybackState.GetTimelineRange(Timeline);
+	return Timeline;
+}
+FTimeValue FAdaptiveStreamingPlayer::ABRGetWallclockTime() const
+{
+	return SynchronizedUTCTime->GetTime();
+}
+bool FAdaptiveStreamingPlayer::ABRIsLive() const
+{
+	TSharedPtrTS<IManifest> Current = Manifest;
+	return Current.IsValid() ? Current->GetPresentationType() == IManifest::EType::Live : false;
+}
+bool FAdaptiveStreamingPlayer::ABRShouldPlayOnLiveEdge() const
+{
+	return PlaybackState.GetShouldPlayOnLiveEdge();
+}
+TSharedPtrTS<const FLowLatencyDescriptor> FAdaptiveStreamingPlayer::ABRGetLowLatencyDescriptor() const
+{
+	TSharedPtrTS<IManifest> Current = Manifest;
+	return Current.IsValid() ? Current->GetLowLatencyDescriptor() : nullptr;
+}
+FTimeValue FAdaptiveStreamingPlayer::ABRGetDesiredLiveEdgeLatency() const
+{
+	TSharedPtrTS<IManifest> Current = Manifest;
+	return Current.IsValid() ? Current->GetDesiredLiveLatency() : FTimeValue();
+}
+FTimeValue FAdaptiveStreamingPlayer::ABRGetLatency() const
+{
+	return PlaybackState.GetCurrentLiveLatency();
+}
+FTimeValue FAdaptiveStreamingPlayer::ABRGetPlaySpeed() const
+{
+	return FTimeValue(RenderClock->IsRunning() ? 1.0 : 0.0);
+}
+void FAdaptiveStreamingPlayer::ABRGetStreamBufferStats(FAccessUnitBufferInfo& OutBufferStats, EStreamType ForStream)
+{
+	UpdateDiagnostics();
+	FMediaCriticalSection::ScopedLock lock(DiagnosticsCriticalSection);
+	switch(ForStream)
 	{
-		FErrorDetail err;
-		err.SetFacility(Facility::EFacility::Player);
-		err.SetMessage("Rebuffering is set to generate a playback error");
-		err.SetCode(INTERR_REBUFFER_SHALL_THROW_ERROR);
-		PostError(err);
+		case EStreamType::Video:
+			OutBufferStats = VideoBufferStats.StreamBuffer;
+			break;
+		case EStreamType::Audio:
+			OutBufferStats = AudioBufferStats.StreamBuffer;
+			break;
+		case EStreamType::Subtitle:
+			OutBufferStats = TextBufferStats.StreamBuffer;
+			break;
+		default:
+			OutBufferStats.Clear();
+			break;
+	}
+}
+
+void FAdaptiveStreamingPlayer::ABRGetStreamBufferStats(IAdaptiveStreamSelector::IPlayerLiveControl::FABRBufferStats& OutBufferStats, EStreamType ForStream)
+{
+	FAccessUnitBufferInfo bi;
+	ABRGetStreamBufferStats(bi, ForStream);
+	OutBufferStats.PlayableContentDuration = bi.PlayableDuration;
+	if (ForStream == EStreamType::Video && VideoRender.Renderer.IsValid())
+	{
+		OutBufferStats.PlayableContentDuration += VideoRender.Renderer->GetEnqueuedSampleDuration();
+	}
+	else if (ForStream == EStreamType::Audio && AudioRender.Renderer.IsValid())
+	{
+		OutBufferStats.PlayableContentDuration += AudioRender.Renderer->GetEnqueuedSampleDuration();
 	}
 }
 
+void FAdaptiveStreamingPlayer::ABRSetRenderRateScale(double InRenderRateScale)
+{
+	RenderRateScale = InRenderRateScale;
+	if (VideoRender.Renderer.IsValid())
+	{
+		VideoRender.Renderer->SetPlayRateScale(InRenderRateScale);
+	}
+	if (AudioRender.Renderer.IsValid())
+	{
+		AudioRender.Renderer->SetPlayRateScale(InRenderRateScale);
+	}
+}
+FTimeRange FAdaptiveStreamingPlayer::ABRGetSupportedRenderRateScale()
+{
+	FTimeRange Range;
+	if (AudioRender.Renderer.IsValid())
+	{
+		Range = AudioRender.Renderer->GetSupportedRenderRateScale();
+	}
+	else if (VideoRender.Renderer.IsValid())
+	{
+		Range = VideoRender.Renderer->GetSupportedRenderRateScale();
+	}
+	return Range;
+}
+
+double FAdaptiveStreamingPlayer::ABRGetRenderRateScale() const
+{
+	return RenderRateScale;
+}
+
+void FAdaptiveStreamingPlayer::ABRTriggerClockSync(IAdaptiveStreamSelector::IPlayerLiveControl::EClockSyncType InClockSyncType)
+{
+	if (Manifest.IsValid())
+	{
+		Manifest->TriggerClockSync(InClockSyncType == IAdaptiveStreamSelector::IPlayerLiveControl::EClockSyncType::Required ?
+									IManifest::EClockSyncType::Required : IManifest::EClockSyncType::Recommended);
+	}
+}
+
+
+
+
+
+
+
+
+
+
 //-----------------------------------------------------------------------------
 /**
 */
@@ -4622,6 +4819,7 @@ FParamDict& FAdaptiveStreamingPlayer::Android_Workarounds(FStreamCodecInformatio
  */
 void FAdaptiveStreamingPlayer::DebugPrint(void* pPlayer, void (*pPrintFN)(void* pPlayer, const char *pFmt, ...))
 {
+#if 1
 	DiagnosticsCriticalSection.Lock();
 	const FAccessUnitBufferInfo *pD = nullptr;
 	if (StreamReaderHandler && bHaveVideoReader.GetWithDefault(false))
@@ -4629,12 +4827,14 @@ void FAdaptiveStreamingPlayer::DebugPrint(void* pPlayer, void (*pPrintFN)(void*
 		pD = &VideoBufferStats.StreamBuffer;
 		pPrintFN(pPlayer, "Video buffer  : %3u AUs; %8u/%8u bytes in; eod %d; %#7.4fs", (uint32)pD->NumCurrentAccessUnits, (uint32)pD->CurrentMemInUse, (uint32)PlayerConfig.StreamBufferConfigVideo.MaxDataSize, pD->bEndOfData, pD->PlayableDuration.GetAsSeconds());
 		pPrintFN(pPlayer, "Video decoder : %2u in decoder, %zu total, %s; EOD in %d; EOD out %d", (uint32)VideoBufferStats.DecoderOutputBuffer.NumElementsInDecoder, (uint32)VideoBufferStats.DecoderOutputBuffer.MaxDecodedElementsReady, VideoBufferStats.DecoderOutputBuffer.bOutputStalled?"    stalled":"not stalled", VideoBufferStats.DecoderInputBuffer.bEODReached, VideoBufferStats.DecoderOutputBuffer.bEODreached);
+		pPrintFN(pPlayer, "Video renderer: %#7.4fs enqueued", VideoRender.Renderer.IsValid() ? VideoRender.Renderer->GetEnqueuedSampleDuration().GetAsSeconds() : 0.0);
 	}
 	if (StreamReaderHandler && bHaveAudioReader.GetWithDefault(false))
 	{
 		pD = &AudioBufferStats.StreamBuffer;
 		pPrintFN(pPlayer, "Audio buffer  : %3u AUs; %8u/%8u bytes in; eod %d; %#7.4fs", (uint32)pD->NumCurrentAccessUnits, (uint32)pD->CurrentMemInUse, (uint32)PlayerConfig.StreamBufferConfigAudio.MaxDataSize, pD->bEndOfData, pD->PlayableDuration.GetAsSeconds());
 		pPrintFN(pPlayer, "Audio decoder : %2u in decoder, %2u total, %s; EOD in %d; EOD out %d", (uint32)AudioBufferStats.DecoderOutputBuffer.NumElementsInDecoder, (uint32)AudioBufferStats.DecoderOutputBuffer.MaxDecodedElementsReady, AudioBufferStats.DecoderOutputBuffer.bOutputStalled ? "    stalled" : "not stalled", AudioBufferStats.DecoderInputBuffer.bEODReached, AudioBufferStats.DecoderOutputBuffer.bEODreached);
+		pPrintFN(pPlayer, "Audio renderer: %#7.4fs enqueued", AudioRender.Renderer.IsValid() ? AudioRender.Renderer->GetEnqueuedSampleDuration().GetAsSeconds() : 0.0);
 	}
 	if (StreamReaderHandler && bHaveTextReader.GetWithDefault(false))
 	{
@@ -4666,6 +4866,11 @@ void FAdaptiveStreamingPlayer::DebugPrint(void* pPlayer, void (*pPrintFN)(void*
 	pPrintFN(pPlayer, "     has ended: %s", HasEnded() ? "Yes" : "No");
 	pPrintFN(pPlayer, "clock (%s): %.3f", RenderClock->IsRunning()?"running":"paused", RenderClock->GetInterpolatedRenderTime(IMediaRenderClock::ERendererType::Video).GetAsSeconds(-1.0));
 	pPrintFN(pPlayer, " internal time: [%d] %.3f", (int32) GetCurrentPlayTime().GetSequenceIndex(), GetCurrentPlayTime().GetAsSeconds());
+#endif
+	if (StreamSelector.IsValid())
+	{
+		StreamSelector->DebugPrint(pPlayer, pPrintFN);
+	}
 }
 void FAdaptiveStreamingPlayer::DebugHandle(void* pPlayer, void (*debugDrawPrintf)(void* pPlayer, const char *pFmt, ...))
 {
@@ -4685,18 +4890,29 @@ void FAdaptiveStreamingPlayer::DebugHandle(void* pPlayer, void (*debugDrawPrintf
 TWeakPtrTS<FAdaptiveStreamingPlayerWorkerThread>	FAdaptiveStreamingPlayerWorkerThread::SingletonSelf;
 FCriticalSection									FAdaptiveStreamingPlayerWorkerThread::SingletonLock;
 
-TSharedPtrTS<FAdaptiveStreamingPlayerWorkerThread> FAdaptiveStreamingPlayerWorkerThread::Create()
+TSharedPtrTS<FAdaptiveStreamingPlayerWorkerThread> FAdaptiveStreamingPlayerWorkerThread::Create(bool bUseSharedWorkerThread)
 {
-	FScopeLock lock(&SingletonLock);
-	TSharedPtrTS<FAdaptiveStreamingPlayerWorkerThread> Self = SingletonSelf.Pin();
-	if (!Self.IsValid())
+	if (bUseSharedWorkerThread)
+	{
+		FScopeLock lock(&SingletonLock);
+		TSharedPtrTS<FAdaptiveStreamingPlayerWorkerThread> Self = SingletonSelf.Pin();
+		if (!Self.IsValid())
+		{
+			FAdaptiveStreamingPlayerWorkerThread* Handler = new FAdaptiveStreamingPlayerWorkerThread;
+			Handler->bIsDedicatedWorker = !bUseSharedWorkerThread;
+			Handler->StartWorkerThread();
+			Self = MakeShareable(Handler);
+			SingletonSelf = Self;
+		}
+		return Self;
+	}
+	else
 	{
 		FAdaptiveStreamingPlayerWorkerThread* Handler = new FAdaptiveStreamingPlayerWorkerThread;
+		Handler->bIsDedicatedWorker = !bUseSharedWorkerThread;
 		Handler->StartWorkerThread();
-		Self = MakeShareable(Handler);
-		SingletonSelf = Self;
+		return MakeShareable(Handler);
 	}
-	return Self;
 }
 
 FAdaptiveStreamingPlayerWorkerThread::~FAdaptiveStreamingPlayerWorkerThread()
@@ -4706,7 +4922,14 @@ FAdaptiveStreamingPlayerWorkerThread::~FAdaptiveStreamingPlayerWorkerThread()
 
 void FAdaptiveStreamingPlayerWorkerThread::StartWorkerThread()
 {
-	WorkerThread.ThreadSetName("ElectraPlayer::SharedWorker");
+	if (bIsDedicatedWorker)
+	{
+		WorkerThread.ThreadSetName("ElectraPlayer::Worker");
+	}
+	else
+	{
+		WorkerThread.ThreadSetName("ElectraPlayer::SharedWorker");
+	}
 	WorkerThread.ThreadStart(Electra::MakeDelegate(this, &FAdaptiveStreamingPlayerWorkerThread::WorkerThreadFN));
 }
 
@@ -4721,17 +4944,38 @@ void FAdaptiveStreamingPlayerWorkerThread::StopWorkerThread()
 void FAdaptiveStreamingPlayerWorkerThread::WorkerThreadFN()
 {
 	LLM_SCOPE(ELLMTag::ElectraPlayer);
+
+	auto RemoveInstances = [this]() -> void
+	{
+		FInstanceToRemove DelInstance;
+		while(InstancesToRemove.Dequeue(DelInstance))
+		{
+			PlayerInstances.Remove(DelInstance.Player);
+			if (DelInstance.DoneSignal)
+			{
+				DelInstance.DoneSignal->Signal();
+			}
+		}
+	};
+
 	while(!bTerminate)
 	{
 		HaveWorkSignal.WaitTimeoutAndReset(1000 * 20);
 
-		SingletonLock.Lock();
+		FAdaptiveStreamingPlayer* NewInstance;
+		while(InstancesToAdd.Dequeue(NewInstance))
+		{
+			PlayerInstances.Add(NewInstance);
+		}
+
 		for(int32 i=0; i<PlayerInstances.Num(); ++i)
 		{
 			PlayerInstances[i]->HandleOnce();
 		}
-		SingletonLock.Unlock();
+
+		RemoveInstances();
 	}
+	RemoveInstances();
 }
 
 
diff --git a/Engine/Plugins/Media/ElectraPlayer/Source/ElectraPlayerRuntime/Private/Runtime/Player/AdaptiveStreamingPlayer.h b/Engine/Plugins/Media/ElectraPlayer/Source/ElectraPlayerRuntime/Private/Runtime/Player/AdaptiveStreamingPlayer.h
--- a/Engine/Plugins/Media/ElectraPlayer/Source/ElectraPlayerRuntime/Private/Runtime/Player/AdaptiveStreamingPlayer.h
+++ b/Engine/Plugins/Media/ElectraPlayer/Source/ElectraPlayerRuntime/Private/Runtime/Player/AdaptiveStreamingPlayer.h
@@ -134,6 +134,9 @@ public:
 	 * Playback is initially paused on first player use and must be resumed to begin.
 	 * Query the seekable range (GetSeekableRange()) to get the valid time range.
 	 * 
+	 * If the seek-to time is not set the seek will start at the beginning for
+	 * on-demand presentations and on the Live edge for Live presentations.
+	 * 
 	 * Seeks can be issued while a seek is already executing. The seek parameters
 	 * control behaviour. If seeking is performed for scrubbing any new seek will
 	 * be performed only when the previous seek has completed, otherwise the current
