ElectraPlayer: Added DASH low latency Live streaming

diff --git a/Engine/Plugins/Media/ElectraPlayer/Source/ElectraPlayerRuntime/Private/Runtime/Player/AdaptiveStreamingPlayerABR.cpp b/Engine/Plugins/Media/ElectraPlayer/Source/ElectraPlayerRuntime/Private/Runtime/Player/AdaptiveStreamingPlayerABR.cpp
--- a/Engine/Plugins/Media/ElectraPlayer/Source/ElectraPlayerRuntime/Private/Runtime/Player/AdaptiveStreamingPlayerABR.cpp
+++ b/Engine/Plugins/Media/ElectraPlayer/Source/ElectraPlayerRuntime/Private/Runtime/Player/AdaptiveStreamingPlayerABR.cpp
@@ -4,329 +4,170 @@
 
 #include "Player/Manifest.h"
 #include "Player/PlaybackTimeline.h"
+#include "Player/AdaptivePlayerOptionKeynames.h"
+#include "Player/ABRRules/ABRInternal.h"
+//#include "Player/ABRRules/ABRLeakyBucket.h"
+#include "Player/ABRRules/ABRVoDPlus.h"
+#include "Player/ABRRules/ABRLiveStream.h"
+#include "Player/ABRRules/ABRFixedStream.h"
+
 #include "StreamAccessUnitBuffer.h"
 #include "Utilities/Utilities.h"
-
+#include "Utilities/URLParser.h"
 
 
 namespace Electra
 {
 
-	template <typename T>
-	class TSimpleMovingAverage
-	{
-	public:
-		//! History sample
-		TSimpleMovingAverage(int32 maxNumSMASamples = kMaxSMASamples)
-			: MaxSamples(maxNumSMASamples)
-		{
-			SimpleMovingAverage.Resize(MaxSamples);
-			LastSample = T(0);
-		}
-
-		void SetMaxSMASampleCount(int32 maxSMASamples)
-		{
-			FMediaCriticalSection::ScopedLock lock(CriticalSection);
-			MaxSamples = maxSMASamples;
-			SimpleMovingAverage.Clear();
-			SimpleMovingAverage.Resize(MaxSamples);
-		}
-
-		//! Clears everything
-		void Clear()
-		{
-			FMediaCriticalSection::ScopedLock lock(CriticalSection);
-			SimpleMovingAverage.Clear();
-		}
-
-		void InitializeTo(const T& v)
-		{
-			FMediaCriticalSection::ScopedLock lock(CriticalSection);
-			Clear();
-			AddSample(v);
-		}
-
-		//! Adds a sample to the mean.
-		void AddSample(const T& value)
-		{
-			FMediaCriticalSection::ScopedLock lock(CriticalSection);
-
-			LastSample = value;
-
-			// Update SMA (http://en.wikipedia.org/wiki/Moving_average#Simple_moving_average)
-			if (SimpleMovingAverage.Capacity())
-			{
-				if (SimpleMovingAverage.Num() >= MaxSamples)
-				{
-					SimpleMovingAverage.Pop();
-				}
-				check(!SimpleMovingAverage.IsFull());
-				SimpleMovingAverage.Push(LastSample);
-			}
-		}
-
-		//! Returns the sample value that was added last, if one exists
-		T GetLastSample() const
-		{
-			FMediaCriticalSection::ScopedLock lock(CriticalSection);
-			return LastSample;
-		}
-
-		//! Returns the simple moving average from all the collected history samples.
-		T GetSMA() const
-		{
-			FMediaCriticalSection::ScopedLock lock(CriticalSection);
-			if (SimpleMovingAverage.Num())
-			{
-				T sum = 0;
-				for(SIZE_T i=0; i<SimpleMovingAverage.Num(); ++i)
-				{
-					sum += SimpleMovingAverage[i];
-				}
-				return T(sum / SimpleMovingAverage.Num());
-			}
-			return T(0);
-		}
-
-	private:
-		enum
-		{
-			kMaxSMASamples = 5,		//!< Default amount of SMA samples
-		};
-
-		mutable FMediaCriticalSection	CriticalSection;
-		TMediaQueueNoLock<T>			SimpleMovingAverage;	//!< Simple moving average
-		T								LastSample;				//!< The last sample that was added.
-		int32							MaxSamples;				//!< Limit number of SMA samples (0 < mMaxSMASamples <= kMaxSMASamples)
-	};
-
-
-	struct FSimulationResult
-	{
-		FSimulationResult()
-			: SecondsAdded(0.0)
-			, SecondsGained(0.0)
-			, SmallestGain(0.0)
-			, DownloadTime(0.0)
-		{
-		}
-		double		SecondsAdded;
-		double		SecondsGained;
-		double		SmallestGain;
-		double		DownloadTime;
-	};
-
-
-
-
-	class FAdaptiveStreamSelector : public IAdaptiveStreamSelector
+	class FAdaptiveStreamSelector : public IAdaptiveStreamSelector, public IABRInfoInterface
 	{
 	public:
-		FAdaptiveStreamSelector(IPlayerSessionServices* PlayerSessionServices, const IAdaptiveStreamSelector::FConfiguration& config);
+		FAdaptiveStreamSelector(IPlayerSessionServices* PlayerSessionServices, IAdaptiveStreamSelector::IPlayerLiveControl* InPlayerLiveControl);
 		virtual ~FAdaptiveStreamSelector();
-		virtual void SetPresentationType(EMediaPresentationType PresentationType) override;
-		virtual void SetCanSwitchToAlternateStreams(bool bCanSwitch) override;
-		virtual void SetCurrentPlaybackPeriod(EStreamType InStreamType, TSharedPtrTS<IManifest::IPlayPeriod> InCurrentPlayPeriod) override;
-		virtual void SetBandwidth(int64 bitsPerSecond) override;
-		virtual void SetForcedNextBandwidth(int64 bitsPerSecond, double minBufferTimeBeforePlayback) override;
-		virtual void SwitchBufferingQuality(const FBufferingQuality& InQualityChange) override;
-		virtual void MarkStreamAsUnavailable(const FBlacklistedStream& BlacklistedStream) override;
-		virtual void MarkStreamAsAvailable(const FBlacklistedStream& NoLongerBlacklistedStream) override;
-		virtual int64 GetLastBandwidth() override;
-		virtual int64 GetAverageBandwidth() override;
-		virtual int64 GetAverageThroughput() override;
-		virtual double GetAverageLatency() override;
-
-		virtual void SetBandwidthCeiling(int32 HighestManifestBitrate) override;
-		virtual void SetMaxVideoResolution(int32 MaxWidth, int32 MaxHeight) override;
-
-		virtual ESegmentAction SelectSuitableStreams(FTimeValue& OutDelay, TSharedPtrTS<const IStreamSegment> CurrentSegment) override;
-
-		virtual FABRDownloadProgressDecision ReportDownloadProgress(const Metrics::FSegmentDownloadStats& SegmentDownloadStats) override;
-		virtual void ReportDownloadEnd(const Metrics::FSegmentDownloadStats& SegmentDownloadStats) override;
-		virtual void ReportBufferingStart(Metrics::EBufferingReason BufferingReason) override;
-		virtual void ReportBufferingEnd(Metrics::EBufferingReason BufferingReason) override;
-		virtual void ReportOpenSource(const FString& URL) override {}
-		virtual void ReportReceivedMasterPlaylist(const FString& EffectiveURL) override {}
-		virtual void ReportReceivedPlaylists() override {}
-		virtual void ReportTracksChanged() override {}
-		virtual void ReportPlaylistDownload(const Metrics::FPlaylistDownloadStats& PlaylistDownloadStats) override {}
-		virtual void ReportBandwidth(int64 EffectiveBps, int64 ThroughputBps, double LatencyInSeconds) override {}
-		virtual void ReportBufferUtilization(const Metrics::FBufferStats& BufferStats) override {}
-		virtual void ReportSegmentDownload(const Metrics::FSegmentDownloadStats& SegmentDownloadStats) override {}
-		virtual void ReportLicenseKey(const Metrics::FLicenseKeyStats& LicenseKeyStats) override {}
-		virtual void ReportDataAvailabilityChange(const Metrics::FDataAvailabilityChange& DataAvailability) override {}
-		virtual void ReportVideoQualityChange(int32 NewBitrate, int32 PreviousBitrate, bool bIsDrasticDownswitch) override {}
-		virtual void ReportDecodingFormatChange(const FStreamCodecInformation& NewDecodingFormat) override {}
-		virtual void ReportPrerollStart() override {}
-		virtual void ReportPrerollEnd() override {}
-		virtual void ReportPlaybackStart() override {}
-		virtual void ReportPlaybackPaused() override {}
-		virtual void ReportPlaybackResumed() override;
-		virtual void ReportPlaybackEnded() override;
-		virtual void ReportJumpInPlayPosition(const FTimeValue& ToNewTime, const FTimeValue& FromTime, Metrics::ETimeJumpReason TimejumpReason) override {}
-		virtual void ReportPlaybackStopped() override {}
-		virtual void ReportSeekCompleted() override {}
-		virtual void ReportError(const FString& ErrorReason) override {}
-		virtual void ReportLogMessage(IInfoLog::ELevel InLogLevel, const FString& LogMessage, int64 PlayerWallclockMilliseconds) override {}
-		virtual void ReportDroppedVideoFrame() override {}
-		virtual void ReportDroppedAudioFrame() override {}
+		void SetFormatType(EMediaFormatType FormatType) override;
+		void SetBandwidthCeiling(int32 HighestManifestBitrate) override;
+		void SetMaxVideoResolution(int32 MaxWidth, int32 MaxHeight) override;
+		void SetCurrentPlaybackPeriod(EStreamType InStreamType, TSharedPtrTS<IManifest::IPlayPeriod> InCurrentPlayPeriod) override;
+		void SetBandwidth(int64 bitsPerSecond) override;
+		void SetForcedNextBandwidth(int64 bitsPerSecond, double minBufferTimeBeforePlayback) override;
+		FTimeValue GetMinBufferTimeForPlayback(EMinBufferType InBufferingType, FTimeValue InDefaultMBT) override;
+		FRebufferAction GetRebufferAction(const FParamDict& CurrentPlayerOptions) override;
+		EHandlingAction PeriodicHandle() override;
+		void SwitchBufferingQuality(const FBufferingQuality& InQualityChange) override;
+		void MarkStreamAsUnavailable(const FBlacklistedStream& BlacklistedStream) override;
+		void MarkStreamAsAvailable(const FBlacklistedStream& NoLongerBlacklistedStream) override;
+		int64 GetLastBandwidth() override;
+		int64 GetAverageBandwidth() override;
+		int64 GetAverageThroughput() override;
+		double GetAverageLatency() override;
+
+		void DebugPrint(void* pThat, void (*pPrintFN)(void* pThat, const char *pFmt, ...)) override;
+
+		ESegmentAction SelectSuitableStreams(FTimeValue& OutDelay, TSharedPtrTS<const IStreamSegment> CurrentSegment) override;
+
+		FABRDownloadProgressDecision ReportDownloadProgress(const Metrics::FSegmentDownloadStats& SegmentDownloadStats) override;
+		void ReportDownloadEnd(const Metrics::FSegmentDownloadStats& SegmentDownloadStats) override;
+		void ReportBufferingStart(Metrics::EBufferingReason BufferingReason) override;
+		void ReportBufferingEnd(Metrics::EBufferingReason BufferingReason) override;
+		void ReportOpenSource(const FString& URL) override {}
+		void ReportReceivedMasterPlaylist(const FString& EffectiveURL) override {}
+		void ReportReceivedPlaylists() override {}
+		void ReportTracksChanged() override {}
+		void ReportPlaylistDownload(const Metrics::FPlaylistDownloadStats& PlaylistDownloadStats) override {}
+		void ReportBandwidth(int64 EffectiveBps, int64 ThroughputBps, double LatencyInSeconds) override {}
+		void ReportBufferUtilization(const Metrics::FBufferStats& BufferStats) override {}
+		void ReportSegmentDownload(const Metrics::FSegmentDownloadStats& SegmentDownloadStats) override {}
+		void ReportLicenseKey(const Metrics::FLicenseKeyStats& LicenseKeyStats) override {}
+		void ReportDataAvailabilityChange(const Metrics::FDataAvailabilityChange& DataAvailability) override {}
+		void ReportVideoQualityChange(int32 NewBitrate, int32 PreviousBitrate, bool bIsDrasticDownswitch) override {}
+		void ReportDecodingFormatChange(const FStreamCodecInformation& NewDecodingFormat) override {}
+		void ReportPrerollStart() override {}
+		void ReportPrerollEnd() override {}
+		void ReportPlaybackStart() override {}
+		void ReportPlaybackPaused() override;
+		void ReportPlaybackResumed() override;
+		void ReportPlaybackEnded() override;
+		void ReportJumpInPlayPosition(const FTimeValue& ToNewTime, const FTimeValue& FromTime, Metrics::ETimeJumpReason TimejumpReason) override {}
+		void ReportPlaybackStopped() override {}
+		void ReportSeekCompleted() override {}
+		void ReportError(const FString& ErrorReason) override {}
+		void ReportLogMessage(IInfoLog::ELevel InLogLevel, const FString& LogMessage, int64 PlayerWallclockMilliseconds) override {}
+		void ReportDroppedVideoFrame() override {}
+		void ReportDroppedAudioFrame() override {}
 
 	private:
-		struct FStreamInformation
-		{
-			struct FStreamHealth
-			{
-				FStreamHealth()
-				{
-					Reset();
-				}
-				void Reset()
-				{
-					//PreviousAttempts.Reset();
-					BecomesAvailableAgainAtUTC.SetToZero();
-					LastDownloadStats = {};
-				}
-				//TSharedPtrTS<HTTP::FRetryInfo>				PreviousAttempts;
-				FTimeValue										BecomesAvailableAgainAtUTC;
-				Metrics::FSegmentDownloadStats					LastDownloadStats;
-			};
-
-			FString											AdaptationSetUniqueID;
-			FString											RepresentationUniqueID;
-
-			FStreamHealth									Health;
-			// Set for convenience.
-			FStreamCodecInformation::FResolution			Resolution;
-			int32											Bitrate;
-		};
-
-		void LogMessage(IInfoLog::ELevel Level, const FString& Message)
-		{
-			if (PlayerSessionServices)
-			{
-				PlayerSessionServices->PostLog(Facility::EFacility::ABR, Level, Message);
-			}
-		}
-
-
-		TSharedPtrTS<FStreamInformation> GetStreamInformation(const Metrics::FSegmentDownloadStats& FromDownloadStats);
-
-		void SimulateDownload(FSimulationResult& sim, const TArray<IManifest::IPlayPeriod::FSegmentInformation>& NextSegments, double secondsToSimulate, int64 atBPS);
+		FParamDict& GetPlayerOptions() override
+		{ return PlayerSessionServices->GetOptions(); }
+		void LogMessage(IInfoLog::ELevel Level, const FString& Message) override
+		{ PlayerSessionServices->PostLog(Facility::EFacility::ABR, Level, Message); }
+		TSharedPtrTS<FABRStreamInformation> GetStreamInformation(const Metrics::FSegmentDownloadStats& FromDownloadStats) override;
+		const TArray<TSharedPtrTS<FABRStreamInformation>>& GetStreamInformations(EStreamType InForStreamType) override;
+		FStreamCodecInformation::FResolution GetMaxStreamResolution() override
+		{ return MaxStreamResolution; }
+		int32 GetBandwidthCeiling() override
+		{ return BandwidthCeiling; }
+
+		FTimeValue ABRGetPlayPosition() const override
+		{ return PlayerLiveControl->ABRGetPlayPosition(); }
+		FTimeRange ABRGetTimeline() const override
+		{ return PlayerLiveControl->ABRGetTimeline(); }
+		FTimeValue ABRGetWallclockTime() const override
+		{ return PlayerLiveControl->ABRGetWallclockTime(); }
+		bool ABRIsLive() const override
+		{ return PlayerLiveControl->ABRIsLive(); }
+		bool ABRShouldPlayOnLiveEdge() const override
+		{ return PlayerLiveControl->ABRShouldPlayOnLiveEdge(); }
+		TSharedPtrTS<const FLowLatencyDescriptor> ABRGetLowLatencyDescriptor() const override
+		{ return PlayerLiveControl->ABRGetLowLatencyDescriptor(); }
+		FTimeValue ABRGetDesiredLiveEdgeLatency() const override
+		{ return PlayerLiveControl->ABRGetDesiredLiveEdgeLatency(); }
+		FTimeValue ABRGetLatency() const override
+		{ return PlayerLiveControl->ABRGetLatency(); }
+		FTimeValue ABRGetPlaySpeed() const override
+		{ return PlayerLiveControl->ABRGetPlaySpeed(); }
+		void ABRGetStreamBufferStats(FAccessUnitBufferInfo& OutBufferStats, EStreamType ForStream) override
+		{ PlayerLiveControl->ABRGetStreamBufferStats(OutBufferStats, ForStream); }
+		void ABRGetStreamBufferStats(IAdaptiveStreamSelector::IPlayerLiveControl::FABRBufferStats& OutBufferStats, EStreamType ForStream) override
+		{ PlayerLiveControl->ABRGetStreamBufferStats(OutBufferStats, ForStream); }
+		FTimeRange ABRGetSupportedRenderRateScale() override
+		{ return PlayerLiveControl->ABRGetSupportedRenderRateScale(); }
+		void ABRSetRenderRateScale(double InRenderRateScale) override
+		{ PlayerLiveControl->ABRSetRenderRateScale(InRenderRateScale); }
+		double ABRGetRenderRateScale() const override
+		{ return PlayerLiveControl->ABRGetRenderRateScale(); }
+		void ABRTriggerClockSync(IAdaptiveStreamSelector::IPlayerLiveControl::EClockSyncType InClockSyncType) override
+		{ PlayerLiveControl->ABRTriggerClockSync(InClockSyncType); }
 
 
 		IPlayerSessionServices*								PlayerSessionServices;
-		FConfiguration										Config;
+		IAdaptiveStreamSelector::IPlayerLiveControl*		PlayerLiveControl;
 
-		FMediaCriticalSection								AccessMutex;
-		EMediaPresentationType								PresentationType;
+		FCriticalSection									AccessMutex;
 		TSharedPtrTS<IManifest::IPlayPeriod>				CurrentPlayPeriodVideo;
 		TSharedPtrTS<IManifest::IPlayPeriod>				CurrentPlayPeriodAudio;
-		bool												bPlayerIsBuffering;
-		bool												bPlayerIsRebuffering;
-		bool												bCanSwitchToAlternateStreams;
 
-		TArray<TSharedPtrTS<FStreamInformation>>			StreamInformationVideo;
-		TArray<TSharedPtrTS<FStreamInformation>>			StreamInformationAudio;
+		TArray<TSharedPtrTS<FABRStreamInformation>>			StreamInformationVideo;
+		TArray<TSharedPtrTS<FABRStreamInformation>>			StreamInformationAudio;
 		FString												CurrentVideoAdaptationSetID;
 		FString												CurrentAudioAdaptationSetID;
 
 		TArray<FBlacklistedStream>							BlacklistedExternally;
 
-		int32												BandwidthCeiling;
-		int32												ForcedInitialBandwidth;
-		double												ForcedInitialBandwidthUntilSecondsBuffered;
-		bool												bForcePlaystartBitrate;
-		bool												bBufferingQualityChanged;
 		FStreamCodecInformation::FResolution				MaxStreamResolution;
+		int32												BandwidthCeiling = 0x7fffffff;
 
-		TSimpleMovingAverage<double>						AverageBandwidth;
-		TSimpleMovingAverage<double>						AverageLatency;
-		TSimpleMovingAverage<int64>							AverageThroughput;
+		TUniquePtr<IABRRule>								ABRMethod;
 	};
 
 
 
-	IAdaptiveStreamSelector::FConfiguration::FRateScale::FRateScale()
-		// Change no defaults without making adjustments to the values in Configuration::Configuration() below!
-		: LastToAvg(0.0f)				// prefer last rate
-		, Scale(1.0f)					// 100%
-		, MaxClampBPS(50000000)		// limit to 50 Mbps
-	{
-	}
-
-	double IAdaptiveStreamSelector::FConfiguration::FRateScale::Apply(const double lastBps, const double averageBps) const
-	{
-		return Utils::Max(Utils::Min((lastBps + LastToAvg * (averageBps - lastBps)) * Scale, double(MaxClampBPS)), 0.0);
-	}
-
-
-	IAdaptiveStreamSelector::FConfiguration::FConfiguration()
-	{
-		SimulationLongtermRateScale.LastToAvg = 1.0f; 			// For simulation use the average, not last bandwidth.
-		ForwardSimulationDurationMSec = 1000 * 20;		// 20 seconds
-		VideoBufferLowWatermarkMSec = 1000 * 3;
-		VideoBufferHighWatermarkMSec = 1000 * 8;
-		MaxBandwidthHistorySize = 3;				// track download speed of the last 5 fragments.
-		DiscardSmallBandwidthSamplesLess = 64 << 10;
-		FudgeSmallBandwidthSamplesLess = 128 << 10;
-		FudgeSmallBandwidthSamplesFactor = 1000;
-		bRebufferingJustResumesLoading = false;
-	}
-
-
-
-
-
-
-
-
-
-
-
-
-
-
-
-
 
 	//-----------------------------------------------------------------------------
 	/**
 	 * Create an instance of this class
 	 *
-	 * @param PlayerSessionServices
-	 * @param config
+	 * @param InPlayerSessionServices
+	 * @param InPlayerLiveControl
 	 *
 	 * @return
 	 */
-	TSharedPtrTS<IAdaptiveStreamSelector> IAdaptiveStreamSelector::Create(IPlayerSessionServices* PlayerSessionServices, const IAdaptiveStreamSelector::FConfiguration& config)
+	TSharedPtrTS<IAdaptiveStreamSelector> IAdaptiveStreamSelector::Create(IPlayerSessionServices* InPlayerSessionServices, IPlayerLiveControl* InPlayerLiveControl)
 	{
-		return TSharedPtrTS<IAdaptiveStreamSelector>(new FAdaptiveStreamSelector(PlayerSessionServices, config));
+		return TSharedPtrTS<IAdaptiveStreamSelector>(new FAdaptiveStreamSelector(InPlayerSessionServices, InPlayerLiveControl));
 	}
 
 	//-----------------------------------------------------------------------------
 	/**
 	 * CTOR
 	 *
-	 * @param config
+	 * @param InPlayerSessionServices
+	 * @param InPlayerLiveControl
 	 */
-	FAdaptiveStreamSelector::FAdaptiveStreamSelector(IPlayerSessionServices* InPlayerSessionServices, const IAdaptiveStreamSelector::FConfiguration& config)
+	FAdaptiveStreamSelector::FAdaptiveStreamSelector(IPlayerSessionServices* InPlayerSessionServices, IAdaptiveStreamSelector::IPlayerLiveControl* InPlayerLiveControl)
 		: PlayerSessionServices(InPlayerSessionServices)
-		, Config(config)
-		, bPlayerIsBuffering(true)
-		, bPlayerIsRebuffering(false)
-		, BandwidthCeiling(0x7fffffff)
-		, ForcedInitialBandwidth(0)
-		, ForcedInitialBandwidthUntilSecondsBuffered(0.0)
-		, bForcePlaystartBitrate(false)
-		, bBufferingQualityChanged(false)
+		, PlayerLiveControl(InPlayerLiveControl)
 	{
-		AverageBandwidth.SetMaxSMASampleCount(Config.MaxBandwidthHistorySize);
-		AverageLatency.SetMaxSMASampleCount(Config.MaxBandwidthHistorySize);
-		AverageThroughput.SetMaxSMASampleCount(Config.MaxBandwidthHistorySize);
-		MaxStreamResolution.Width = MaxStreamResolution.Height = 8192;
-		PresentationType = EMediaPresentationType::Realtime;
-		bCanSwitchToAlternateStreams = true;
+		MaxStreamResolution.Width = MaxStreamResolution.Height = 32768;
 	}
 
 	//-----------------------------------------------------------------------------
@@ -337,42 +178,49 @@ namespace Electra
 	{
 	}
 
-
-
-	//-----------------------------------------------------------------------------
-	/**
-	 * Sets the type of presentation.
-	 *
-	 * @param InAssetType
-	 */
-	void FAdaptiveStreamSelector::SetPresentationType(IAdaptiveStreamSelector::EMediaPresentationType InAssetType)
-	{
-		PresentationType = InAssetType;
-	}
-
-
 	//-----------------------------------------------------------------------------
 	/**
-	 * Sets whether or not a segment download failure can potentially be resolved by switching to a different stream. Defaults to yes.
+	 * Sets the presentation format (HLS, DASH, mp4, etc...)
+	 * 
+	 * This must be the first call that will determine the type of ABR algorithm to use.
 	 *
-	 * @param bCanSwitch
+	 * @param InFormatType
 	 */
-	void FAdaptiveStreamSelector::SetCanSwitchToAlternateStreams(bool bCanSwitch)
+	void FAdaptiveStreamSelector::SetFormatType(EMediaFormatType InFormatType)
 	{
-		bCanSwitchToAlternateStreams = bCanSwitch;
+		// Multiplexed stream with no bitrate switching alternatives?
+		if (InFormatType == EMediaFormatType::ISOBMFF)
+		{
+			ABRMethod.Reset(IABRFixedStream::Create(this, InFormatType, EABRPresentationType::OnDemand));
+		}
+		else
+		{
+			EABRPresentationType PresentationType;
+			// On-demand presentation?
+			if (!PlayerLiveControl->ABRIsLive())
+			{
+				PresentationType = EABRPresentationType::OnDemand;
+				ABRMethod.Reset(IABROnDemandPlus::Create(this, InFormatType, PresentationType));
+			}
+			// Live / low-latency Live
+			else
+			{
+				PresentationType = PlayerLiveControl->ABRGetLowLatencyDescriptor().IsValid() ? EABRPresentationType::LowLatency : EABRPresentationType::Live;
+				ABRMethod.Reset(IABRLiveStream::Create(this, InFormatType, PresentationType));
+			}
+		}
 	}
 
-
 	//-----------------------------------------------------------------------------
 	/**
 	 * Sets the playback period from which to select streams now.
 	 *
 	 * @param InStreamType
-	 * @param CurrentPlayPeriod
+	 * @param InCurrentPlayPeriod
 	 */
 	void FAdaptiveStreamSelector::SetCurrentPlaybackPeriod(EStreamType InStreamType, TSharedPtrTS<IManifest::IPlayPeriod> InCurrentPlayPeriod)
 	{
-		FMediaCriticalSection::ScopedLock lock(AccessMutex);
+		FScopeLock lock(&AccessMutex);
 		if (InStreamType == EStreamType::Video)
 		{
 			CurrentPlayPeriodVideo = InCurrentPlayPeriod;
@@ -391,13 +239,16 @@ namespace Electra
 							TSharedPtrTS<IPlaybackAssetAdaptationSet> Adapt = Asset->GetAdaptationSetByTypeAndIndex(InStreamType, nA);
 							if (CurrentVideoAdaptationSetID.Equals(Adapt->GetUniqueIdentifier()))
 							{
+								bool bLowLatency = Adapt->IsLowLatencyEnabled();
 								for(int32 i=0, iMax=Adapt->GetNumberOfRepresentations(); i<iMax; ++i)
 								{
 									TSharedPtrTS<IPlaybackAssetRepresentation> Repr = Adapt->GetRepresentationByIndex(i);
-									TSharedPtrTS<FStreamInformation> si = MakeSharedTS<FStreamInformation>();
+									TSharedPtrTS<FABRStreamInformation> si = MakeSharedTS<FABRStreamInformation>();
 									si->AdaptationSetUniqueID = Adapt->GetUniqueIdentifier();
 									si->RepresentationUniqueID = Repr->GetUniqueIdentifier();
 									si->Bitrate = Repr->GetBitrate();
+									si->QualityIndex = Repr->GetQualityIndex();
+									si->bLowLatencyEnabled = bLowLatency;
 									if (Repr->GetCodecInformation().IsVideoCodec())
 									{
 										si->Resolution = Repr->GetCodecInformation().GetResolution();
@@ -405,7 +256,7 @@ namespace Electra
 									StreamInformationVideo.Push(si);
 								}
 								// Sort the representations by ascending bitrate
-								StreamInformationVideo.Sort([](const TSharedPtrTS<FStreamInformation>& a, const TSharedPtrTS<FStreamInformation>& b)
+								StreamInformationVideo.Sort([](const TSharedPtrTS<FABRStreamInformation>& a, const TSharedPtrTS<FABRStreamInformation>& b)
 								{
 									return a->Bitrate < b->Bitrate;
 								});
@@ -440,17 +291,20 @@ namespace Electra
 							TSharedPtrTS<IPlaybackAssetAdaptationSet> Adapt = Asset->GetAdaptationSetByTypeAndIndex(InStreamType, nA);
 							if (CurrentAudioAdaptationSetID.Equals(Adapt->GetUniqueIdentifier()))
 							{
+								bool bLowLatency = Adapt->IsLowLatencyEnabled();
 								for(int32 i=0, iMax=Adapt->GetNumberOfRepresentations(); i<iMax; ++i)
 								{
 									TSharedPtrTS<IPlaybackAssetRepresentation> Repr = Adapt->GetRepresentationByIndex(i);
-									TSharedPtrTS<FStreamInformation> si = MakeSharedTS<FStreamInformation>();
+									TSharedPtrTS<FABRStreamInformation> si = MakeSharedTS<FABRStreamInformation>();
 									si->AdaptationSetUniqueID = Adapt->GetUniqueIdentifier();
 									si->RepresentationUniqueID = Repr->GetUniqueIdentifier();
 									si->Bitrate = Repr->GetBitrate();
+									si->QualityIndex = Repr->GetQualityIndex();
+									si->bLowLatencyEnabled = bLowLatency;
 									StreamInformationAudio.Push(si);
 								}
 								// Sort the representations by ascending bitrate
-								StreamInformationAudio.Sort([](const TSharedPtrTS<FStreamInformation>& a, const TSharedPtrTS<FStreamInformation>& b)
+								StreamInformationAudio.Sort([](const TSharedPtrTS<FABRStreamInformation>& a, const TSharedPtrTS<FABRStreamInformation>& b)
 								{
 									return a->Bitrate < b->Bitrate;
 								});
@@ -467,423 +321,213 @@ namespace Electra
 				StreamInformationAudio.Empty();
 			}
 		}
+		// Notify the ABR method that there is a different set of representations now.
+		ABRMethod->RepresentationsChanged(InStreamType, InCurrentPlayPeriod);
 	}
 
-
 	//-----------------------------------------------------------------------------
 	/**
 	 * Sets the initial bandwidth, either from guessing, past history or a current measurement.
 	 *
-	 * @param bitPerSecond
+	 * @param InBitsPerSecond
 	 */
-	void FAdaptiveStreamSelector::SetBandwidth(int64 bitsPerSecond)
+	void FAdaptiveStreamSelector::SetBandwidth(int64 InBitsPerSecond)
 	{
-		FMediaCriticalSection::ScopedLock lock(AccessMutex);
-		AverageBandwidth.InitializeTo((double)bitsPerSecond);
-		AverageLatency.InitializeTo((double)0.1);
-		AverageThroughput.InitializeTo(bitsPerSecond);
-		bForcePlaystartBitrate = true;
-		bBufferingQualityChanged = false;
+		check(ABRMethod.IsValid());
+		ABRMethod->SetBandwidth(InBitsPerSecond);
 	}
 
-
 	//-----------------------------------------------------------------------------
 	/**
 	 * Sets a forced bitrate for the next segment fetches until the given duration of playable content has been received.
+	 * The ABR algorithm may or may not use this.
 	 *
-	 * @param bitsPerSecond
-	 * @param minBufferTimeBeforePlayback
+	 * @param InBitsPerSecond
+	 * @param InMinBufferTimeBeforePlayback
 	 */
-	void FAdaptiveStreamSelector::SetForcedNextBandwidth(int64 bitsPerSecond, double minBufferTimeBeforePlayback)
+	void FAdaptiveStreamSelector::SetForcedNextBandwidth(int64 InBitsPerSecond, double InMinBufferTimeBeforePlayback)
 	{
-		FMediaCriticalSection::ScopedLock lock(AccessMutex);
-		ForcedInitialBandwidth = (int32)bitsPerSecond;
-		ForcedInitialBandwidthUntilSecondsBuffered = minBufferTimeBeforePlayback;
+		check(ABRMethod.IsValid());
+		ABRMethod->SetForcedNextBandwidth(InBitsPerSecond, InMinBufferTimeBeforePlayback);
 	}
 
-
 	//-----------------------------------------------------------------------------
 	/**
-	 * Called to (temporarily) mark a stream as unavailable.
-	 * This is used in formats like HLS when a dedicated stream playlist is in error
-	 * and thus segment information for this stream is not available.
+	 * Sets a highest bandwidth limit for a stream type. Call with 0 to disable.
 	 *
-	 * @param BlacklistedStream
+	 * @param InHighestManifestBitrate
 	 */
-	void FAdaptiveStreamSelector::MarkStreamAsUnavailable(const FBlacklistedStream& BlacklistedStream)
+	void FAdaptiveStreamSelector::SetBandwidthCeiling(int32 InHighestManifestBitrate)
 	{
-		FMediaCriticalSection::ScopedLock lock(AccessMutex);
-		BlacklistedExternally.Push(BlacklistedStream);
+		FScopeLock lock(&AccessMutex);
+		BandwidthCeiling = InHighestManifestBitrate > 0 ? InHighestManifestBitrate : 0x7fffffff;
 	}
 
 	//-----------------------------------------------------------------------------
 	/**
-	 * Marks a previously set to unavailable stream as being available again.
-	 * Usually when the stream specific playlist (in HLS) has become available again.
+	 * Limits video resolution.
 	 *
-	 * @param NoLongerBlacklistedStream
+	 * @param InMaxWidth
+	 * @param InMaxHeight
 	 */
-	void FAdaptiveStreamSelector::MarkStreamAsAvailable(const FBlacklistedStream& NoLongerBlacklistedStream)
+	void FAdaptiveStreamSelector::SetMaxVideoResolution(int32 InMaxWidth, int32 InMaxHeight)
 	{
-		FMediaCriticalSection::ScopedLock lock(AccessMutex);
-		BlacklistedExternally.Remove(NoLongerBlacklistedStream);
+		FScopeLock lock(&AccessMutex);
+		MaxStreamResolution.Width = InMaxWidth;
+		MaxStreamResolution.Height = InMaxHeight;
 	}
 
-
 	//-----------------------------------------------------------------------------
 	/**
-	 * Returns the last measured bandwidth sample in bps.
+	 * Called to (temporarily) mark a stream as unavailable.
+	 * This is used in formats like HLS when a dedicated stream playlist is in error
+	 * and thus segment information for this stream is not available.
 	 *
-	 * @return
+	 * @param InBlacklistedStream
 	 */
-	int64 FAdaptiveStreamSelector::GetLastBandwidth()
+	void FAdaptiveStreamSelector::MarkStreamAsUnavailable(const FBlacklistedStream& InBlacklistedStream)
 	{
-		FMediaCriticalSection::ScopedLock lock(AccessMutex);
-		double lastBandwidth = AverageBandwidth.GetLastSample();
-		return (int64)lastBandwidth;
+		FScopeLock lock(&AccessMutex);
+		BlacklistedExternally.Push(InBlacklistedStream);
 	}
 
-
 	//-----------------------------------------------------------------------------
 	/**
-	 * Returns the average measured bandwidth sample in bps.
+	 * Marks a previously set to unavailable stream as being available again.
+	 * Usually when the stream specific playlist (in HLS) has become available again.
 	 *
-	 * @return
+	 * @param InNoLongerBlacklistedStream
 	 */
-	int64 FAdaptiveStreamSelector::GetAverageBandwidth()
+	void FAdaptiveStreamSelector::MarkStreamAsAvailable(const FBlacklistedStream& InNoLongerBlacklistedStream)
 	{
-		FMediaCriticalSection::ScopedLock lock(AccessMutex);
-		const double averageBandwidth = AverageBandwidth.GetSMA();
-		return (int64)averageBandwidth;
+		FScopeLock lock(&AccessMutex);
+		BlacklistedExternally.Remove(InNoLongerBlacklistedStream);
 	}
 
 	//-----------------------------------------------------------------------------
 	/**
-	 * Returns the average measured throughput in bps.
-	 *
-	 * @return
+	 * Returns the last measured bandwidth sample in bits per second.
 	 */
-	int64 FAdaptiveStreamSelector::GetAverageThroughput()
+	int64 FAdaptiveStreamSelector::GetLastBandwidth()
 	{
-		FMediaCriticalSection::ScopedLock lock(AccessMutex);
-		return AverageThroughput.GetSMA();
+		return ABRMethod.IsValid() ? ABRMethod->GetLastBandwidth() : 0;
 	}
 
 	//-----------------------------------------------------------------------------
 	/**
-	 * Returns the average measured latency in seconds.
-	 *
-	 * @return
+	 * Returns the average measured bandwidth sample in bits per second.
 	 */
-	double FAdaptiveStreamSelector::GetAverageLatency()
+	int64 FAdaptiveStreamSelector::GetAverageBandwidth()
 	{
-		FMediaCriticalSection::ScopedLock lock(AccessMutex);
-		return AverageLatency.GetSMA();
+		return ABRMethod.IsValid() ? ABRMethod->GetAverageBandwidth() : 0;
 	}
 
-
-
 	//-----------------------------------------------------------------------------
 	/**
-	 * Sets a highest bandwidth limit for a stream type. Call with 0 to disable.
-	 *
-	 * @param HighestManifestBitrate
+	 * Returns the average measured throughput in bits per second.
 	 */
-	void FAdaptiveStreamSelector::SetBandwidthCeiling(int32 HighestManifestBitrate)
+	int64 FAdaptiveStreamSelector::GetAverageThroughput()
 	{
-		FMediaCriticalSection::ScopedLock lock(AccessMutex);
-		BandwidthCeiling = HighestManifestBitrate > 0 ? HighestManifestBitrate : 0x7fffffff;
+		return ABRMethod.IsValid() ? ABRMethod->GetAverageThroughput() : 0;
 	}
 
-
 	//-----------------------------------------------------------------------------
 	/**
-	 * Limits video resolution.
-	 *
-	 * @param MaxWidth
-	 * @param MaxHeight
+	 * Returns the average measured download latency in seconds.
 	 */
-	void FAdaptiveStreamSelector::SetMaxVideoResolution(int32 MaxWidth, int32 MaxHeight)
+	double FAdaptiveStreamSelector::GetAverageLatency()
 	{
-		FMediaCriticalSection::ScopedLock lock(AccessMutex);
-		MaxStreamResolution.Width = MaxWidth;
-		MaxStreamResolution.Height = MaxHeight;
+		return ABRMethod.IsValid() ? ABRMethod->GetAverageLatency() : 0.0;
 	}
 
 
-
-	FABRDownloadProgressDecision FAdaptiveStreamSelector::ReportDownloadProgress(const Metrics::FSegmentDownloadStats& SegmentDownloadStats)
+	IAdaptiveStreamSelector::EHandlingAction FAdaptiveStreamSelector::PeriodicHandle()
 	{
-		FABRDownloadProgressDecision Decision;
-		uint32 Flags = SegmentDownloadStats.ABRState.ProgressDecision.Flags;
-
-		FAccessUnitBufferInfo CurrentBufferStats;
-		PlayerSessionServices->GetStreamBufferStats(CurrentBufferStats, SegmentDownloadStats.StreamType);
-		const double secondsInBuffer = CurrentBufferStats.PlayableDuration.GetAsSeconds();
-		// Estimate how long we will take to complete the download.
-		double EstimatedTotalDownloadTime = SegmentDownloadStats.Duration;
-		//double currentPercentage = 0.0;
-		// Can we base this off the byte sizes?
-		if (SegmentDownloadStats.NumBytesDownloaded > 0 && SegmentDownloadStats.ByteSize > 0)
-		{
-			EstimatedTotalDownloadTime = (SegmentDownloadStats.TimeToDownload / SegmentDownloadStats.NumBytesDownloaded) * SegmentDownloadStats.ByteSize;
-			//currentPercentage = SegmentDownloadStats.NumBytesDownloaded / SegmentDownloadStats.ByteSize;
-		}
-		// Otherwise, can we base it off the durations?
-		else if (SegmentDownloadStats.DurationDownloaded > 0.0 && SegmentDownloadStats.Duration > 0.0)
-		{
-			EstimatedTotalDownloadTime = (SegmentDownloadStats.TimeToDownload / SegmentDownloadStats.DurationDownloaded) * SegmentDownloadStats.Duration;
-			//currentPercentage = SegmentDownloadStats.DurationDownloaded / SegmentDownloadStats.Duration;
-		}
-		double EstimatedDownloadTimeRemaining = EstimatedTotalDownloadTime - SegmentDownloadStats.TimeToDownload;
-
-		if ((SegmentDownloadStats.ABRState.ProgressDecision.Flags & FABRDownloadProgressDecision::EDecisionFlags::eABR_AbortDownload) == 0)
-		{
-			// Did we take longer to download than the segment duration?
-			double MaxSegmentDownloadDuration = SegmentDownloadStats.Duration;
-			// If there is a forced initial bitrate to start playback with then we allow for the download duration to take longer than normally
-			// as to not abort the first download unless it's download time is really getting problematic.
-			// If the player asked for a buffering quality change things are problematic, so don't do it!
-			if (!bBufferingQualityChanged)
-			{
-				if (ForcedInitialBandwidth && SegmentDownloadStats.StreamType == EStreamType::Video)
-				{
-					MaxSegmentDownloadDuration *= 2.0;
-				}
-				// When buffering we allow for some more time as well.
-				else if (bPlayerIsRebuffering)
-				{
-					MaxSegmentDownloadDuration *= 2.0;
-				}
-			}
-
-			// Rebuffering is set to just continue loading data? No timeout then.
-			if (bPlayerIsRebuffering && Config.bRebufferingJustResumesLoading)
-			{
-				MaxSegmentDownloadDuration = 1e30;
-			}
-
-			if (SegmentDownloadStats.TimeToDownload >= MaxSegmentDownloadDuration)
-			{
-				// If the estimate to complete download is longer than what is currently in the buffer abort the download and hope
-				// switching to a lower quality will help.
-		//		if (EstimatedDownloadTimeRemaining >= secondsInBuffer)
-		// Be more aggressive and demand the buffer must have the total duration! This is to avoid late switching that may be too late to try to download a whole new segment
-				if (EstimatedTotalDownloadTime >= secondsInBuffer)
-				{
-					Flags = FABRDownloadProgressDecision::EDecisionFlags::eABR_AbortDownload;
-					// If content was already emitted into the buffer we cannot switch. Insert dummy filler data instead.
-					if (SegmentDownloadStats.DurationDelivered > 0.0)
-					{
-						Flags |= FABRDownloadProgressDecision::EDecisionFlags::eABR_InsertFillerData;
-					}
-
-					Decision.Reason = FString::Printf(TEXT("Abort download of %s %s segment @ %d bps for time %.3f with %.3fs in buffer and %.3fs of %.3fs fetched after %.3fs. Estimated total download time %.3fs")
-						, GetStreamTypeName(SegmentDownloadStats.StreamType)
-						, GetSegmentTypeString(SegmentDownloadStats.SegmentType)
-						, SegmentDownloadStats.Bitrate
-						, SegmentDownloadStats.PresentationTime
-						, secondsInBuffer
-						, SegmentDownloadStats.DurationDownloaded
-						, SegmentDownloadStats.Duration
-						, SegmentDownloadStats.TimeToDownload
-						, EstimatedTotalDownloadTime);
-					LogMessage(IInfoLog::ELevel::Info, Decision.Reason);
-				}
-			}
-		}
-
-		// If the buffer is low allow the data being streamed in to already be emitted to the buffer.
-		// This will however mean that in case of remaining problems with this segment we cannot switch to another quality level.
-		if ((SegmentDownloadStats.ABRState.ProgressDecision.Flags & FABRDownloadProgressDecision::EDecisionFlags::eABR_EmitPartialData) == 0)
-		{
-			if (secondsInBuffer < 1.0 && SegmentDownloadStats.DurationDelivered == 0.0)
-			{
-				Flags |= FABRDownloadProgressDecision::EDecisionFlags::eABR_EmitPartialData;
-				/*
-							LogMessage(IInfoLog::ELevel::Info, StringHelpers::SPrintf("Emitting partial download of %s %s segment @ %d bps for time %.3f with %.3fs in buffer and %.3fs of %.3fs fetched after %.3fs. Estimated total download time %.3fs"
-								  , GetStreamTypeName(SegmentDownloadStats.StreamType)
-								  , GetSegmentTypeString(SegmentDownloadStats.SegmentType)
-								  , SegmentDownloadStats.Bitrate
-								  , SegmentDownloadStats.PresentationTime
-								  , secondsInBuffer
-								  , SegmentDownloadStats.DurationDownloaded
-								  , SegmentDownloadStats.Duration
-								  , SegmentDownloadStats.TimeToDownload
-								  , EstimatedTotalDownloadTime));
-				*/
-			}
-		}
-
-		Decision.Flags = (FABRDownloadProgressDecision::EDecisionFlags) Flags;
-		return Decision;
+		return ABRMethod.IsValid() ? ABRMethod->PeriodicHandle() : IAdaptiveStreamSelector::EHandlingAction::None;
 	}
 
-
-	void FAdaptiveStreamSelector::ReportPlaybackResumed()
+	IAdaptiveStreamSelector::FRebufferAction FAdaptiveStreamSelector::GetRebufferAction(const FParamDict& CurrentPlayerOptions)
 	{
-		// When playback resumes there must be enough data buffered to do that, so we are no longer
-		// constrained to a forced buffered bandwidth
-		ForcedInitialBandwidth = 0;
-		ForcedInitialBandwidthUntilSecondsBuffered = 0.0;
-		bForcePlaystartBitrate = false;
-		bBufferingQualityChanged = false;
+		return ABRMethod.IsValid() ? ABRMethod->GetRebufferAction(CurrentPlayerOptions) : IAdaptiveStreamSelector::FRebufferAction();
 	}
-	void FAdaptiveStreamSelector::ReportPlaybackEnded()
+
+	FTimeValue FAdaptiveStreamSelector::GetMinBufferTimeForPlayback(IAdaptiveStreamSelector::EMinBufferType InBufferingType, FTimeValue InDefaultMBT)
 	{
-		// When playback ends we are no longer constrained to a forced buffered bandwidth
-		ForcedInitialBandwidth = 0;
-		ForcedInitialBandwidthUntilSecondsBuffered = 0.0;
-		bForcePlaystartBitrate = false;
-		bBufferingQualityChanged = false;
+		return ABRMethod.IsValid() ? ABRMethod->GetMinBufferTimeForPlayback(InBufferingType, InDefaultMBT) : FTimeValue();
 	}
 
-	void FAdaptiveStreamSelector::ReportDownloadEnd(const Metrics::FSegmentDownloadStats& SegmentDownloadStats)
+	void FAdaptiveStreamSelector::DebugPrint(void* pThat, void (*pPrintFN)(void* pThat, const char *pFmt, ...))
 	{
-		if (SegmentDownloadStats.SegmentType == Metrics::ESegmentType::Media && SegmentDownloadStats.StreamType == EStreamType::Video &&
-			(SegmentDownloadStats.bWasSuccessful || SegmentDownloadStats.NumBytesDownloaded))	// any number of bytes received counts!
+		if (ABRMethod.IsValid())
 		{
-			if (!SegmentDownloadStats.bIsCachedResponse)
-			{
-				// Add bandwidth sample if it can be calculated.
-				if (SegmentDownloadStats.NumBytesDownloaded && SegmentDownloadStats.TimeToDownload > 0.0)
-				{
-					FMediaCriticalSection::ScopedLock lock(AccessMutex);
-					double v = SegmentDownloadStats.NumBytesDownloaded * 8 / SegmentDownloadStats.TimeToDownload;
-					AverageBandwidth.AddSample(v);
-					AverageLatency.AddSample(SegmentDownloadStats.TimeToFirstByte);
-					AverageThroughput.AddSample(SegmentDownloadStats.ThroughputBps > 0 ? SegmentDownloadStats.ThroughputBps : (int64)v);
-				}
-			}
-			else
-			{
-				// When we hit the cache we cannot calculate a meaningful bandwidth. We need to update the averages nonetheless
-				// to make a quality upswitch possible on one of the next uncached segments.
-				FMediaCriticalSection::ScopedLock lock(AccessMutex);
-				int32 ThisBitrate = SegmentDownloadStats.Bitrate;
-				// Get the index of this video stream. This requires the stream information to be sorted by bitrate, which it is.
-				int32 nQualityIndex = 0;
-				for(int32 i=0; i<StreamInformationVideo.Num(); ++i)
-				{
-					if (ThisBitrate >= StreamInformationVideo[i]->Bitrate)
-					{
-						nQualityIndex = i;
-					}
-				}
-				int64 LastAddedBandwidth = (int64)AverageBandwidth.GetLastSample();
-				int32 NextHighestQualityIndex = nQualityIndex + 1 < StreamInformationVideo.Num() ? nQualityIndex + 1 : nQualityIndex;
-				int32 NextHighestBitrate = StreamInformationVideo[NextHighestQualityIndex]->Bitrate;
-				// Last seen bandwidth already higher than what we might want to switch up to next?
-				if (LastAddedBandwidth > NextHighestBitrate)
-				{
-					// We want to bring the average down a bit in preparation for the next uncached segment in case the network
-					// has degraded. This way we are not overshooting the target.
-					ThisBitrate = (LastAddedBandwidth + NextHighestBitrate) / 2;
-					AverageBandwidth.AddSample(ThisBitrate);
-				}
-				else if (LastAddedBandwidth > ThisBitrate)
-				{
-					// The last bandwidth was somewhere between the bitrate of the cached segment and the next quality.
-					// We would love to switch up with the next segment. This might be possible because we got this segment's data
-					// at lightning speed and have its duration worth of time to spend on attempting the next segment download.
-					// If that doesn't finish in time we can downswitch without hurting things too much.
-					ThisBitrate = NextHighestBitrate;
-					AverageBandwidth.InitializeTo(ThisBitrate);
-				}
-				else
-				{
-					// Let's try to bring the average up a bit.
-					ThisBitrate = (int32)(ThisBitrate * 0.2 + NextHighestBitrate * 0.8);
-					AverageBandwidth.AddSample(ThisBitrate);
-				}
-				// We do NOT update the average latency.
-					//  AverageLatency.AddSample()
-				
-				// The throughput is calculated as real time.
-				AverageThroughput.AddSample((int64)(SegmentDownloadStats.NumBytesDownloaded * 8 / SegmentDownloadStats.Duration));
-			}
-		}
-
-		// Check if we can now go back to regular video bitrate selection.
-		if (SegmentDownloadStats.SegmentType == Metrics::ESegmentType::Media && SegmentDownloadStats.StreamType == EStreamType::Video)
-		{
-			// The forced playstart bitrate is intended for the very first video segment only.
-			// The general initial bitrate limit still applies to all segments until enough data is buffered.
-			bForcePlaystartBitrate = false;
-			if (ForcedInitialBandwidth)
-			{
-				FAccessUnitBufferInfo bufStats;
-				PlayerSessionServices->GetStreamBufferStats(bufStats, SegmentDownloadStats.StreamType);
-				bool bBufferedEnough = bufStats.PlayableDuration.GetAsSeconds() >= ForcedInitialBandwidthUntilSecondsBuffered;
-				bBufferedEnough |= bufStats.bEndOfData;
-				bBufferedEnough |= bufStats.bLastPushWasBlocked;
-				if (bBufferedEnough)
-				{
-					ForcedInitialBandwidth = 0;
-					ForcedInitialBandwidthUntilSecondsBuffered = 0.0;
-				}
-			}
+			ABRMethod->DebugPrint(pThat, pPrintFN);
 		}
 	}
 
 
-	void FAdaptiveStreamSelector::SwitchBufferingQuality(const FBufferingQuality& InQualityChange)
+	void FAdaptiveStreamSelector::ReportPlaybackPaused()
 	{
-		double avgCurrent = AverageBandwidth.GetLastSample();
-		double newRate = avgCurrent;
-
-		// In case the average bandwidth is higher than what is needed for the best video stream
-		// we clamp it to there. Otherwise we may rebuffer the highest quality stream which is
-		// not likely to be useful since rebuffering means there was a drastic change in bandwidth.
-		if (StreamInformationVideo.Num())
+		if (ABRMethod.IsValid())
 		{
-			if (newRate > StreamInformationVideo.Last()->Bitrate)
-			{
-				newRate = (double)StreamInformationVideo.Last()->Bitrate;
-			}
+			ABRMethod->ReportPlaybackPaused();
 		}
+	}
 
-		if (InQualityChange.BandwidthScaleFactor.IsSet())
+	void FAdaptiveStreamSelector::ReportPlaybackResumed()
+	{
+		if (ABRMethod.IsValid())
 		{
-			newRate *= InQualityChange.BandwidthScaleFactor.GetValue();
+			ABRMethod->ReportPlaybackResumed();
 		}
-		else if (InQualityChange.AbsoluteBandwidth.IsSet())
+	}
+	
+	void FAdaptiveStreamSelector::ReportPlaybackEnded()
+	{
+		if (ABRMethod.IsValid())
 		{
-			newRate = (double)InQualityChange.AbsoluteBandwidth.GetValue();
+			ABRMethod->ReportPlaybackEnded();
 		}
+	}
 
-		double Ratio = avgCurrent > 0.0 ? newRate / avgCurrent : 1.0;
-		AverageBandwidth.InitializeTo(newRate);
-		AverageThroughput.InitializeTo((int64)(AverageThroughput.GetLastSample() * Ratio));
-		LogMessage(IInfoLog::ELevel::Info, FString::Printf(TEXT("Adjusting buffering bandwidth from %d to %d"), (int)avgCurrent, (int)newRate));
-		bBufferingQualityChanged = true;
+	FABRDownloadProgressDecision FAdaptiveStreamSelector::ReportDownloadProgress(const Metrics::FSegmentDownloadStats& SegmentDownloadStats)
+	{
+		check(ABRMethod.IsValid());
+		return ABRMethod->ReportDownloadProgress(SegmentDownloadStats);
 	}
 
-	void FAdaptiveStreamSelector::ReportBufferingStart(Metrics::EBufferingReason BufferingReason)
+	void FAdaptiveStreamSelector::ReportDownloadEnd(const Metrics::FSegmentDownloadStats& SegmentDownloadStats)
 	{
-		/*
-			Initial,
-			Seeking,
-			Rebuffering,
-		*/
-		bPlayerIsBuffering = true;
+		check(ABRMethod.IsValid());
+		ABRMethod->ReportDownloadEnd(SegmentDownloadStats);
+	}
 
-		if (BufferingReason == Metrics::EBufferingReason::Rebuffering)
-		{
-			bPlayerIsRebuffering = true;
-		}
+	void FAdaptiveStreamSelector::ReportBufferingStart(Metrics::EBufferingReason BufferingReason)
+	{
+		check(ABRMethod.IsValid());
+		ABRMethod->ReportBufferingStart(BufferingReason);
 	}
 
 	void FAdaptiveStreamSelector::ReportBufferingEnd(Metrics::EBufferingReason BufferingReason)
 	{
-		bPlayerIsBuffering = false;
-		bPlayerIsRebuffering = false;
+		check(ABRMethod.IsValid());
+		ABRMethod->ReportBufferingEnd(BufferingReason);
+	}
+
+	void FAdaptiveStreamSelector::SwitchBufferingQuality(const FBufferingQuality& InQualityChange)
+	{
+		check(ABRMethod.IsValid());
+		ABRMethod->SwitchBufferingQuality(InQualityChange);
+	}
+
+	const TArray<TSharedPtrTS<FABRStreamInformation>>& FAdaptiveStreamSelector::GetStreamInformations(EStreamType InForStreamType)
+	{
+		return InForStreamType == EStreamType::Video ? StreamInformationVideo : StreamInformationAudio;
 	}
 
-	TSharedPtrTS<FAdaptiveStreamSelector::FStreamInformation> FAdaptiveStreamSelector::GetStreamInformation(const Metrics::FSegmentDownloadStats& FromDownloadStats)
+	TSharedPtrTS<FABRStreamInformation> FAdaptiveStreamSelector::GetStreamInformation(const Metrics::FSegmentDownloadStats& FromDownloadStats)
 	{
-		FMediaCriticalSection::ScopedLock lock(AccessMutex);
-		const TArray<TSharedPtrTS<FStreamInformation>>* StreamInfos = nullptr;
+		FScopeLock lock(&AccessMutex);
+		const TArray<TSharedPtrTS<FABRStreamInformation>>* StreamInfos = nullptr;
 		if (FromDownloadStats.StreamType == EStreamType::Video)
 		{
 			StreamInfos = &StreamInformationVideo;
@@ -908,68 +552,23 @@ namespace Electra
 		return nullptr;
 	}
 
-
 	//-----------------------------------------------------------------------------
 	/**
-	 * Simulates download of the next fragments at the current download rate.
-	 *
-	 * @param sim
-	 * @param NextSegments
-	 * @param secondsToSimulate
-	 * @param atBPS
-	 */
-	void FAdaptiveStreamSelector::SimulateDownload(FSimulationResult& sim, const TArray<IManifest::IPlayPeriod::FSegmentInformation>& NextSegments, double secondsToSimulate, int64 atBPS)
-	{
-		sim.SecondsAdded = 0.0;
-		sim.SecondsGained = 0.0;
-		sim.SmallestGain = 0.0;
-		sim.DownloadTime = 0.0;
-
-		bool bFirst = true;
-		int32 infoIdx = 0;
-		check(NextSegments.Num() > 0);
-		double s = 8.0 / atBPS;
-		while (secondsToSimulate > 0.0)
-		{
-			double Dur = NextSegments[infoIdx].Duration.GetAsSeconds();
-
-			double timeToDownload = NextSegments[infoIdx].ByteSize * s;
-
-			secondsToSimulate -= Dur;
-			sim.SecondsAdded += Dur;
-			sim.DownloadTime += timeToDownload;
-			sim.SecondsGained += Dur - timeToDownload;
-
-			if (sim.SecondsGained < sim.SmallestGain || bFirst)
-			{
-				sim.SmallestGain = sim.SecondsGained;
-			}
-			bFirst = false;
-
-			// Advance to the next info, if present. Otherwise repeat the last one ad nauseum.
-			if (infoIdx + 1 < NextSegments.Num())
-			{
-				++infoIdx;
-			}
-		}
-	}
-
-
-	//-----------------------------------------------------------------------------
-	/**
-	 * Returns the index of a feasible video or multiplex stream.
+	 * Chooses a stream suitable for streaming the next segment from based on network conditions.
 	 *
 	 * @param OutDelay
 	 * @param CurrentSegment
 	 *
-	 * @return
+	 * @return Action to take for the next segment, which is to load it, retry due to a failure or skip it.
 	 */
 	IAdaptiveStreamSelector::ESegmentAction FAdaptiveStreamSelector::SelectSuitableStreams(FTimeValue& OutDelay, TSharedPtrTS<const IStreamSegment> CurrentSegment)
 	{
-		OutDelay.SetToZero();
-
 		FTimeValue TimeNow = PlayerSessionServices->GetSynchronizedUTCTime()->GetTime();
-		// Make all streams available again if their blacklist time has expired
+		
+		// No delay to fetch this segment with.
+		OutDelay.SetToZero();
+		
+		// Check if streams that were temporarily marked as bad can now be used again.
 		AccessMutex.Lock();
 		for(int32 i=0, iMax=StreamInformationVideo.Num(); i<iMax; ++i)
 		{
@@ -987,10 +586,9 @@ namespace Electra
 		}
 		AccessMutex.Unlock();
 
-
-		EStreamType StreamType = CurrentSegment.IsValid() ? CurrentSegment->GetType() : StreamInformationVideo.Num() ? EStreamType::Video : EStreamType::Audio;
-		
+		// Get the currently active play period of the finished segment, or if it is the first segment the most appropriate one.
 		TSharedPtrTS<IManifest::IPlayPeriod> CurrentPlayPeriod;
+		EStreamType StreamType = CurrentSegment.IsValid() ? CurrentSegment->GetType() : StreamInformationVideo.Num() ? EStreamType::Video : EStreamType::Audio;
 		if (StreamType == EStreamType::Video)
 		{
 			CurrentPlayPeriod = CurrentPlayPeriodVideo;
@@ -999,327 +597,64 @@ namespace Electra
 		{
 			CurrentPlayPeriod = CurrentPlayPeriodAudio;
 		}
+		// No play period, nothing to do. Assume there is something coming up, so continue.
 		if (!CurrentPlayPeriod.IsValid())
 		{
 			return ESegmentAction::FetchNext;
 		}
 
-		// Was the previous segment download in error?
-		bool bRetryIfPossible = false;
-		bool bSkipWithFiller = false;
-		if (CurrentSegment.IsValid())
-		{
-			Metrics::FSegmentDownloadStats Stats;
-			CurrentSegment->GetDownloadStats(Stats);
-			// Try to get the stream information for the current downloaded segment. We may not find it on a period transition where the
-			// segment is the last of the previous period.
-			TSharedPtrTS<FAdaptiveStreamSelector::FStreamInformation> CurrentStreamInfo = GetStreamInformation(Stats);
-			if (CurrentStreamInfo.IsValid())
-			{
-				if (!Stats.bWasSuccessful)
-				{
-					// Was not successful. Figure out what to do now.
-					// In the case where the segment had an availability window set the stream reader was waiting for
-					// to be entered and we got a 404 back, the server did not manage to publish the new segment in time.
-					// We try this again after a slight delay.
-					if (Stats.AvailibilityDelay && Utils::AbsoluteValue(Stats.AvailibilityDelay) < 0.5)
-					{
-						//LogMessage(IInfoLog::ELevel::Warning, FString::Printf(TEXT("Segment not available at announced time. Trying again.")));
-						CurrentPlayPeriod->IncreaseSegmentFetchDelay(FTimeValue(FTimeValue::MillisecondsToHNS(100)));
-						OutDelay.SetFromMilliseconds(500);
-						return ESegmentAction::Retry;
-					}
-
-					// Too many failures already? It's unlikely that there is a way that would magically fix itself now.
-					if (Stats.RetryNumber >= 3)
-					{
-						LogMessage(IInfoLog::ELevel::Error, FString::Printf(TEXT("Exceeded permissable number of retries (%d). Failing now."), Stats.RetryNumber));
-						return ESegmentAction::Fail;
-					}
-
-					// If we cannot switch to another stream we need to retry this one.
-					if (!bCanSwitchToAlternateStreams)
-					{
-						return Stats.bParseFailure ? ESegmentAction::Fail : ESegmentAction::Retry;
-					}
+		// Get the initial list of representations that are currently possible to choose from.
+		TArray<TSharedPtrTS<FABRStreamInformation>> CandidateRepresentations;
+		ABRMethod->PrepareStreamCandidateList(CandidateRepresentations, StreamType, CurrentPlayPeriod, TimeNow);
 
-					// Is this an init segment?
-					if (Stats.SegmentType == Metrics::ESegmentType::Init)
-					{
-						bRetryIfPossible = true;
-						// Did the init segment fail to be parsed? If so this stream is dead for good.
-						if (Stats.bParseFailure)
-						{
-							CurrentStreamInfo->Health.BecomesAvailableAgainAtUTC.SetToInvalid();
-						}
-						// Did we abort that stream in ReportDownloadProgress() ?
-						else if (Stats.bWasAborted)
-						{
-							// Since the download was aborted for the lack of time we should not be trying anything of a same or better bitrate.
-							double bps = 8.0 * Stats.NumBytesDownloaded / Stats.TimeToDownload;
-							double fallbackBps = Stats.Bitrate / 2;
-							AverageBandwidth.InitializeTo(Utils::Min(fallbackBps / 2.0, bps));
-						}
-						// Any other case of failure
-						else
-						{
-							// Take this stream offline for a brief moment.
-							if (CurrentStreamInfo->Health.BecomesAvailableAgainAtUTC == FTimeValue::GetZero())
-							{
-								CurrentStreamInfo->Health.BecomesAvailableAgainAtUTC = TimeNow + FTimeValue().SetFromMilliseconds(1000);
-							}
-						}
-					}
-					// A media segment failure
-					else
-					{
-						bRetryIfPossible = true;
-						bSkipWithFiller = true;
-
-						// Did we abort that stream in ReportDownloadProgress() ?
-						if (Stats.bWasAborted)
-						{
-							// Since the download was aborted for the lack of time we should not be trying anything of a same or better bitrate.
-							if (Stats.StreamType == EStreamType::Video)
-							{
-								double bps = 8.0 * Stats.NumBytesDownloaded / Stats.TimeToDownload;
-								double fallbackBps = Stats.Bitrate / 2;
-								AverageBandwidth.InitializeTo(Utils::Min(fallbackBps / 2.0, bps));
-							}
-						}
-						else
-						{
-							// Take a video stream offline for a brief moment unless it is the only one. Audio is usually the only one with no alternative to switch to, so don't do that!
-							if (Stats.StreamType == EStreamType::Video)
-							{
-								if (StreamInformationVideo.Num() > 1)
-								{
-									if (CurrentStreamInfo->Health.BecomesAvailableAgainAtUTC == FTimeValue::GetZero())
-									{
-										CurrentStreamInfo->Health.BecomesAvailableAgainAtUTC = TimeNow + FTimeValue().SetFromMilliseconds(1000);
-									}
-								}
-								else
-								{
-									// Need to insert filler data instead.
-									bRetryIfPossible = false;
-								}
-							}
-							else
-							{
-								// For VoD we retry the audio segment up to a certain number of times while for Live we skip over it
-								// since retrying is a luxury we do not have in this case.
-								if (PresentationType == IAdaptiveStreamSelector::EMediaPresentationType::OnDemand)
-								{
-									bRetryIfPossible = true;
-								}
-								else
-								{
-									bRetryIfPossible = false;
-								}
-							}
-						}
-
-						// If any content has already been put out into the buffers we cannot retry the segment on another quality level.
-						if (Stats.DurationDelivered > 0.0)
-						{
-							bRetryIfPossible = false;
-							bSkipWithFiller = false;
-						}
-					}
-				}
-				// Update stats
-				CurrentStreamInfo->Health.LastDownloadStats = Stats;
-			}
-		}
-
-		// Anything besides video or audio?
-		// NOTE: For the love of God do not remove this line without re-instating the else if() below. See comments there!
-		if (StreamType != EStreamType::Video && StreamType != EStreamType::Audio)
+		// First have the ABR method evaluate this download for possible errors and decide on an action.
+		IAdaptiveStreamSelector::ESegmentAction NextSegmentAction = ABRMethod->EvaluateForError(CandidateRepresentations, StreamType, OutDelay, CurrentPlayPeriod, CurrentSegment, TimeNow);
+		// If there is an error action like retry, fill or fail return it right away.
+		if (NextSegmentAction != IAdaptiveStreamSelector::ESegmentAction::FetchNext)
 		{
-			return ESegmentAction::FetchNext;
+			return NextSegmentAction;
 		}
 
-
-		TArray<TSharedPtrTS<FStreamInformation>>			PossibleRepresentations;
-		FAccessUnitBufferInfo								CurrentBufferStats;
-
-		PlayerSessionServices->GetStreamBufferStats(CurrentBufferStats, StreamType);
-
-		// For video streams do a download simulation
-		if (StreamType == EStreamType::Video)
+		// Evaluate which stream qualities could potentially be used and select one.
+		NextSegmentAction = ABRMethod->EvaluateForQuality(CandidateRepresentations, StreamType, OutDelay, CurrentPlayPeriod, CurrentSegment, TimeNow);
+		if (NextSegmentAction != IAdaptiveStreamSelector::ESegmentAction::FetchNext)
 		{
-			const double secondsInVideoBuffer = CurrentBufferStats.PlayableDuration.GetAsSeconds();
-			// Buffer watermark levels
-			const double videoBufferLowWatermark = double(Config.VideoBufferLowWatermarkMSec) / 1000.0;
-			const double videoBufferHighWatermark = double(Config.VideoBufferHighWatermarkMSec) / 1000.0;
-
-			// Get average and last bitrate.
-			const int64 averageBandwidth = Utils::Max(GetAverageBandwidth(), (int64)100000);
-			int64 lastBandwidth = GetLastBandwidth();
-			if (lastBandwidth <= 0)
-			{
-				lastBandwidth = averageBandwidth;
-			}
-			lastBandwidth = Utils::Max(lastBandwidth, (int64)100000);
-
-			// Setup simulation and last bandwidth such that the current audio bitrate is accounted for and clamped to at least 100 Kbps.
-			const int64 simulationBandwidthLongterm = Utils::Max((int64)Config.SimulationLongtermRateScale.Apply(lastBandwidth, averageBandwidth), (int64)100000);
-			// How long to simulate download for.
-			const double simulationDuration = double(Config.ForwardSimulationDurationMSec) / 1000.0;
-
-			// Bandwidth to simulate download of just the next fragment with.
-			const int64 simulationBandwidthNextFragment = Utils::Max((int64)Config.SimulationNextFragmentRateScale.Apply(lastBandwidth, averageBandwidth), (int64)100000);
-
-			// For each video stream simulate how the buffer will behave if we were to use it.
-			// The lowest quality is always available regardless of bandwidth or resolution or other constraints.
-			if (StreamInformationVideo.Num())
-			{
-				PossibleRepresentations.Push(StreamInformationVideo[0]);
-			}
-			for(int32 nStr=1; nStr<StreamInformationVideo.Num(); ++nStr)
-			{
-				// Check if bitrate and resolution are acceptable
-				if (StreamInformationVideo[nStr]->Bitrate <= BandwidthCeiling &&
-					!StreamInformationVideo[nStr]->Resolution.ExceedsLimit(MaxStreamResolution))
-				{
-					FSimulationResult sim;
-
-					// Get information on the next segments.
-					TArray<IManifest::IPlayPeriod::FSegmentInformation> NextSegments;
-					FTimeValue AverageSegmentDuration;
-					CurrentPlayPeriod->GetSegmentInformation(NextSegments, AverageSegmentDuration, CurrentSegment, FTimeValue().SetFromSeconds(simulationDuration), StreamInformationVideo[nStr]->AdaptationSetUniqueID, StreamInformationVideo[nStr]->RepresentationUniqueID);
-
-					// Long term prediction: Simulate download (at a user-scaled bandwidth factor)
-					SimulateDownload(sim, NextSegments, simulationDuration, simulationBandwidthLongterm);
-
-					// Mid term prediction: Estimate how much we will have in the video buffer when downloading just the next fragment.
-					double secondsInBufferAfterNext;
-					if (NextSegments.Num())
-					{
-						secondsInBufferAfterNext = secondsInVideoBuffer + NextSegments[0].Duration.GetAsSeconds() - (NextSegments[0].ByteSize * (8.0 / simulationBandwidthNextFragment));
-					}
-					else
-					{
-						secondsInBufferAfterNext = secondsInVideoBuffer;
-					}
-
-					// A stream is feasible if it:
-					//    - gives at least some data when the current buffer level is lower than the low-threshold
-					//      OR
-					//    - even with the smallest gain in data the buffer level AFTER the NEXT fragment will be higher than the low-threshold
-					// AND
-					//    - the simulated gain plus what is in the buffer now will be higher than the high-threshold.
-					bool bCond1 = sim.SmallestGain > 0.0 && secondsInVideoBuffer < videoBufferLowWatermark;
-					bool bCond2 = sim.SmallestGain + secondsInBufferAfterNext > videoBufferLowWatermark;
-					bool bCond3 = sim.SecondsGained + secondsInVideoBuffer > videoBufferHighWatermark;
-					bool bIsFeasible = (bCond1 || bCond2) && bCond3;
-
-					// If there is an enforced initial bitrate then any stream that is outside that bitrate is not feasible.
-					if (ForcedInitialBandwidth)
-					{
-						bIsFeasible = StreamInformationVideo[nStr]->Bitrate <= ForcedInitialBandwidth;
-					}
-
-					// If there is an initial playstart bitrate set then this stream must be used, regardless of actual
-					// available bandwidth which we do not know anyway.
-					if (bForcePlaystartBitrate && StreamInformationVideo[nStr]->Bitrate <= ForcedInitialBandwidth)
-					{
-						bIsFeasible = true;
-					}
-
-					if (bIsFeasible)
-					{
-						PossibleRepresentations.Push(StreamInformationVideo[nStr]);
-					}
-				}
-			}
+			return NextSegmentAction;
 		}
-		// For audio we do no download simulation or anything.
-	// NOTE: removed this else if check because of static code analyzer complaining this would always be true, which it is not.
-	//       it is only true as long as the above check that if this is anything but video or audio we already return.
-	//       otherwise, with that if() above removed *NOT* having this check here will be a major bug!
-		//else if (StreamType == EStreamType::Audio)
-		else
+
+		// Filter out streams that are currently not healthy.
+		for(int32 i=0; i<CandidateRepresentations.Num(); ++i)
 		{
-			for(int32 nStr=0; nStr<StreamInformationAudio.Num(); ++nStr)
+			if (!CandidateRepresentations[i]->Health.BecomesAvailableAgainAtUTC.IsValid() ||
+				(CandidateRepresentations[i]->Health.BecomesAvailableAgainAtUTC.IsValid() && CandidateRepresentations[i]->Health.BecomesAvailableAgainAtUTC > TimeNow))
 			{
-				PossibleRepresentations.Push(StreamInformationAudio[nStr]);
+				CandidateRepresentations.RemoveAt(i);
+				--i;
 			}
 		}
 
-		// At this point we have a list of representations that we could potentially use.
-		if (PossibleRepresentations.Num())
+		// Remove those that have been externally blacklisted. This may include all of them.
+		AccessMutex.Lock();
+		if (BlacklistedExternally.Num())
 		{
-			// Remove those that have been externally blacklisted. This may include the lowest quality stream!
-			FMediaCriticalSection::ScopedLock lock(AccessMutex);
-			if (BlacklistedExternally.Num())
-			{
-				TArray<TSharedPtrTS<FStreamInformation>> RemainingCandidates;
-				for(int32 j=0, jMax=PossibleRepresentations.Num(); j<jMax; ++j)
-				{
-					bool bStillGood = true;
-					for(int32 i=0, iMax=BlacklistedExternally.Num(); i<iMax; ++i)
-					{
-						if (PossibleRepresentations[j]->AdaptationSetUniqueID == BlacklistedExternally[i].AdaptationSetUniqueID &&
-							PossibleRepresentations[j]->RepresentationUniqueID == BlacklistedExternally[i].RepresentationUniqueID)
-						{
-							bStillGood = false;
-							break;
-						}
-					}
-					if (bStillGood)
-					{
-						RemainingCandidates.Push(PossibleRepresentations[j]);
-					}
-				}
-				Swap(PossibleRepresentations, RemainingCandidates);
-			}
-
-			// Filter out streams that are currently not healthy.
-			if (1)
+			for(int32 j=0; j<CandidateRepresentations.Num(); ++j)
 			{
-				TArray<TSharedPtrTS<FStreamInformation>> RemainingCandidates;
-				for(int32 j=0, jMax=PossibleRepresentations.Num(); j<jMax; ++j)
+				for(int32 i=0, iMax=BlacklistedExternally.Num(); i<iMax; ++i)
 				{
-					bool bStillGood = PossibleRepresentations[j]->Health.BecomesAvailableAgainAtUTC <= TimeNow;
-					if (bStillGood)
+					if (CandidateRepresentations[j]->AdaptationSetUniqueID == BlacklistedExternally[i].AdaptationSetUniqueID &&
+						CandidateRepresentations[j]->RepresentationUniqueID == BlacklistedExternally[i].RepresentationUniqueID)
 					{
-						RemainingCandidates.Push(PossibleRepresentations[j]);
+						CandidateRepresentations.RemoveAt(j);
+						--j;
+						break;
 					}
 				}
-				Swap(PossibleRepresentations, RemainingCandidates);
-			}
-
-			// Still something left to use?
-			if (PossibleRepresentations.Num())
-			{
-				const TSharedPtrTS<FStreamInformation>& Candidate = PossibleRepresentations.Top();
-				CurrentPlayPeriod->SelectStream(Candidate->AdaptationSetUniqueID, Candidate->RepresentationUniqueID);
-
-				if (bRetryIfPossible)
-				{
-					return ESegmentAction::Retry;
-				}
-				else if (bSkipWithFiller)
-				{
-					return ESegmentAction::Fill;
-				}
-				else
-				{
-					return ESegmentAction::FetchNext;
-				}
-			}
-			else
-			{
-				// TODO: Are there representations we could lift the temporary blacklist from?
-				if (bSkipWithFiller)
-				{
-					return ESegmentAction::Fill;
-				}
 			}
 		}
-		return ESegmentAction::Fail;
+		AccessMutex.Unlock();
+
+		// From whatever is left, select what is deemed the best stream.
+		return ABRMethod->PerformSelection(CandidateRepresentations, StreamType, OutDelay, CurrentPlayPeriod, CurrentSegment, TimeNow);
 	}
 
 
diff --git a/Engine/Plugins/Media/ElectraPlayer/Source/ElectraPlayerRuntime/Private/Runtime/Player/AdaptiveStreamingPlayerABR.h b/Engine/Plugins/Media/ElectraPlayer/Source/ElectraPlayerRuntime/Private/Runtime/Player/AdaptiveStreamingPlayerABR.h
--- a/Engine/Plugins/Media/ElectraPlayer/Source/ElectraPlayerRuntime/Private/Runtime/Player/AdaptiveStreamingPlayerABR.h
+++ b/Engine/Plugins/Media/ElectraPlayer/Source/ElectraPlayerRuntime/Private/Runtime/Player/AdaptiveStreamingPlayerABR.h
@@ -18,44 +18,65 @@ namespace Electra
 class IAdaptiveStreamSelector : public IAdaptiveStreamingPlayerMetrics
 {
 public:
-	struct FConfiguration
+	class IPlayerLiveControl
 	{
-		FConfiguration();
-		struct FRateScale
+	public:
+		virtual ~IPlayerLiveControl() = default;
+
+		// Provides the current playback position.
+		virtual FTimeValue ABRGetPlayPosition() const = 0;
+
+		// Provides the current full media timeline, not adjusted for seeking.
+		virtual FTimeRange ABRGetTimeline() const = 0;
+
+		// Provides the current wallclock time (server real time).
+		virtual FTimeValue ABRGetWallclockTime() const = 0;
+		
+		// Returns true if this is a Live presentation, false for VoD.
+		virtual bool ABRIsLive() const = 0;
+		
+		// Returns true if playback should be on the Live edge. false if the user intentionally seeked away from the edge.
+		virtual bool ABRShouldPlayOnLiveEdge() const = 0;
+		
+		// Provides the current low latency descriptor if low latency playback is active. Returns nullptr if not.
+		virtual TSharedPtrTS<const FLowLatencyDescriptor> ABRGetLowLatencyDescriptor() const = 0;
+
+		// Provides the desired latency to the live edge for both low-latency and normal live playback.
+		virtual FTimeValue ABRGetDesiredLiveEdgeLatency() const = 0;
+
+		// Provides the currently measured latency to the live edge.
+		virtual FTimeValue ABRGetLatency() const = 0;
+
+		// Provides the current playback speed.
+		virtual FTimeValue ABRGetPlaySpeed() const = 0;
+
+		// Returns the current stream access unit buffer stats.
+		virtual void ABRGetStreamBufferStats(FAccessUnitBufferInfo& OutBufferStats, EStreamType ForStream) = 0;
+
+		struct FABRBufferStats
 		{
-			FRateScale();
-			double Apply(const double lastKbps, const double averageKbps) const;
-			float		LastToAvg;			//!< [last, average] bandwidth scaler such that Rate = last + mLastToAvg * (average - last)
-			float		Scale;				//!< Scale factor to apply next
-			uint32		MaxClampBPS;		//!< Maximum value to clamp result to (in bits per second)
+			FTimeValue PlayableContentDuration;
 		};
-		FRateScale		SimulationLongtermRateScale;
-		FRateScale		SimulationNextFragmentRateScale;
-		FRateScale		FinalRateAdjust;
-		uint32			ForwardSimulationDurationMSec;			//!< Forward simulation duration
-		uint32			VideoBufferLowWatermarkMSec;
-		uint32			VideoBufferHighWatermarkMSec;
-		uint32			MaxBandwidthHistorySize;				//!< Number of bandwidth history samples to give the estimated average bandwidth.
-		uint32			DiscardSmallBandwidthSamplesLess;
-		uint32			FudgeSmallBandwidthSamplesLess;
-		uint32			FudgeSmallBandwidthSamplesFactor;
-		bool			bRebufferingJustResumesLoading;
+		virtual void ABRGetStreamBufferStats(FABRBufferStats& OutBufferStats, EStreamType ForStream) = 0;
+
+		virtual FTimeRange ABRGetSupportedRenderRateScale() = 0;
+		virtual void ABRSetRenderRateScale(double InRenderRateScale) = 0;
+		virtual double ABRGetRenderRateScale() const = 0;
+
+		enum class EClockSyncType
+		{
+			Recommended,
+			Required
+		};
+		virtual void ABRTriggerClockSync(EClockSyncType InClockSyncType) = 0;
 	};
 
+
 	//! Create an instance of this class
-	static TSharedPtrTS<IAdaptiveStreamSelector> Create(IPlayerSessionServices* PlayerSessionServices, const FConfiguration& config);
+	static TSharedPtrTS<IAdaptiveStreamSelector> Create(IPlayerSessionServices* PlayerSessionServices, IPlayerLiveControl* PlayerLiveControl);
 	virtual ~IAdaptiveStreamSelector() = default;
 
-	enum class EMediaPresentationType
-	{
-		OnDemand,
-		Live,
-		Realtime
-	};
-	virtual void SetPresentationType(EMediaPresentationType AssetType) = 0;
-
-	//! Sets whether or not a segment download failure can potentially be resolved by switching to a different stream. Defaults to yes.
-	virtual void SetCanSwitchToAlternateStreams(bool bCanSwitch) = 0;
+	virtual void SetFormatType(EMediaFormatType FormatType) = 0;
 
 	//! Sets the manifest from which to select streams.
 	virtual void SetCurrentPlaybackPeriod(EStreamType InStreamType, TSharedPtrTS<IManifest::IPlayPeriod> CurrentPlayPeriod) = 0;
@@ -66,6 +87,39 @@ public:
 	//! Sets a forced bitrate for the next segment fetches until the given duration of playable content has been received.
 	virtual void SetForcedNextBandwidth(int64 bitsPerSecond, double minBufferTimeBeforePlayback) = 0;
 
+	enum class EMinBufferType
+	{
+		Initial,
+		Seeking,
+		Rebuffering
+	};
+	virtual FTimeValue GetMinBufferTimeForPlayback(EMinBufferType InBufferingType, FTimeValue InDefaultMBT) = 0;
+
+	enum class EHandlingAction
+	{
+		None,
+		SeekToLive
+	};
+
+	virtual EHandlingAction PeriodicHandle() = 0;
+
+	struct FRebufferAction
+	{
+		enum class EAction
+		{
+			Restart,
+			ContinueLoading,
+			GoToLive,
+			ThrowError
+		};
+		EAction Action = EAction::Restart;
+		int32 SuggestedRestartBitrate = 0;
+	};
+
+	// Call when rebuffering occurs to get the recommended action. This implies that rebuffering has actually occurred.
+	// The ABR may take internal action to switch quality when this is called!
+	virtual FRebufferAction GetRebufferAction(const FParamDict& CurrentPlayerOptions) = 0;
+
 	struct FBufferingQuality
 	{
 		TOptional<double> BandwidthScaleFactor;
@@ -125,32 +179,7 @@ public:
 	virtual ESegmentAction SelectSuitableStreams(FTimeValue& OutDelay, TSharedPtrTS<const IStreamSegment> CurrentSegment) = 0;
 
 
-	virtual void ReportOpenSource(const FString& URL) = 0;
-	virtual void ReportReceivedMasterPlaylist(const FString& EffectiveURL) = 0;
-	virtual void ReportReceivedPlaylists() = 0;
-	virtual void ReportTracksChanged() = 0;
-	virtual void ReportPlaylistDownload(const Metrics::FPlaylistDownloadStats& PlaylistDownloadStats) = 0;
-	virtual void ReportBufferingStart(Metrics::EBufferingReason BufferingReason) = 0;
-	virtual void ReportBufferingEnd(Metrics::EBufferingReason BufferingReason) = 0;
-	virtual void ReportBandwidth(int64 EffectiveBps, int64 ThroughputBps, double LatencyInSeconds) = 0;
-	virtual void ReportBufferUtilization(const Metrics::FBufferStats& BufferStats) = 0;
-	virtual void ReportSegmentDownload(const Metrics::FSegmentDownloadStats& SegmentDownloadStats) = 0;
-	virtual void ReportLicenseKey(const Metrics::FLicenseKeyStats& LicenseKeyStats) = 0;
-	virtual void ReportDataAvailabilityChange(const Metrics::FDataAvailabilityChange& DataAvailability) = 0;
-	virtual void ReportVideoQualityChange(int32 NewBitrate, int32 PreviousBitrate, bool bIsDrasticDownswitch) = 0;
-	virtual void ReportPrerollStart() = 0;
-	virtual void ReportPrerollEnd() = 0;
-	virtual void ReportPlaybackStart() = 0;
-	virtual void ReportPlaybackPaused() = 0;
-	virtual void ReportPlaybackResumed() = 0;
-	virtual void ReportPlaybackEnded() = 0;
-	virtual void ReportJumpInPlayPosition(const FTimeValue& ToNewTime, const FTimeValue& FromTime, Metrics::ETimeJumpReason TimejumpReason) = 0;
-	virtual void ReportPlaybackStopped() = 0;
-	virtual void ReportSeekCompleted() = 0;
-	virtual void ReportError(const FString& ErrorReason) = 0;
-	virtual void ReportLogMessage(IInfoLog::ELevel LogLevel, const FString& LogMessage, int64 PlayerWallclockMilliseconds) = 0;
-	virtual void ReportDroppedVideoFrame() = 0;
-	virtual void ReportDroppedAudioFrame() = 0;
+	virtual void DebugPrint(void* pThat, void (*pPrintFN)(void* pThat, const char *pFmt, ...)) = 0;
 };
 
 
diff --git a/Engine/Plugins/Media/ElectraPlayer/Source/ElectraPlayerRuntime/Private/Runtime/Player/AdaptiveStreamingPlayerDecoders.cpp b/Engine/Plugins/Media/ElectraPlayer/Source/ElectraPlayerRuntime/Private/Runtime/Player/AdaptiveStreamingPlayerDecoders.cpp
--- a/Engine/Plugins/Media/ElectraPlayer/Source/ElectraPlayerRuntime/Private/Runtime/Player/AdaptiveStreamingPlayerDecoders.cpp
+++ b/Engine/Plugins/Media/ElectraPlayer/Source/ElectraPlayerRuntime/Private/Runtime/Player/AdaptiveStreamingPlayerDecoders.cpp
@@ -752,17 +752,30 @@ void FAdaptiveStreamingPlayer::FeedDecoder(EStreamType Type, IAccessUnitBufferIn
 			bool bEODSet = FromMultistreamBuffer->IsEODFlagSet();
 			if (!bEODSet && FromMultistreamBuffer->Num() == 0)
 			{
-				// Buffer underrun.
-				bRebufferPending = true;
-				FTimeValue LastKnownPTS = FromMultistreamBuffer->GetLastPoppedPTS();
-				// Only set the 'rebuffer at' time if we have a valid last known PTS. If we don't
-				// then maybe this is a cascade failure from a previous rebuffer attempt for which
-				// we then try that time once more.
-				if (LastKnownPTS.IsValid())
+				FTimeValue EnqueuedDuration(FTimeValue::GetZero());
+				int32 NumEnqueuedSamples = 0;
+				if (Type == EStreamType::Video && VideoRender.Renderer.IsValid())
 				{
-					RebufferDetectedAtPlayPos = LastKnownPTS;
+					NumEnqueuedSamples = VideoRender.Renderer->GetNumEnqueuedSamples(&EnqueuedDuration);
+				}
+				else if (Type == EStreamType::Audio && AudioRender.Renderer.IsValid())
+				{
+					NumEnqueuedSamples = AudioRender.Renderer->GetNumEnqueuedSamples(&EnqueuedDuration);
+				}
+				if (NumEnqueuedSamples <= 1)
+				{
+					// Buffer underrun.
+					bRebufferPending = true;
+					FTimeValue LastKnownPTS = FromMultistreamBuffer->GetLastPoppedPTS();
+					// Only set the 'rebuffer at' time if we have a valid last known PTS. If we don't
+					// then maybe this is a cascade failure from a previous rebuffer attempt for which
+					// we then try that time once more.
+					if (LastKnownPTS.IsValid())
+					{
+						RebufferDetectedAtPlayPos = LastKnownPTS;
+					}
+					WorkerThread.Enqueue(FWorkerThreadMessages::FMessage::EType::BufferUnderrun);
 				}
-				WorkerThread.Enqueue(FWorkerThreadMessages::FMessage::EType::BufferUnderrun);
 			}
 		}
 
@@ -860,6 +873,15 @@ void FAdaptiveStreamingPlayer::FeedDecoder(EStreamType Type, IAccessUnitBufferIn
 				{
 					UpdateDataAvailabilityState(*pAvailability, Metrics::FDataAvailabilityChange::EAvailability::DataAvailable);
 				}
+
+				// Check for presence of encoder latency and store it.
+				// This may be updated by different streams/decoders which can't be helped. There should only be one actively
+				// observed element but we have seen manifests that have the same <ProducerReferenceTime@id> in multiple AdaptationSets.
+				if (AccessUnit->ProducerReferenceTime.IsValid())
+				{
+					PlaybackState.SetEncoderLatency(AccessUnit->DTS - AccessUnit->ProducerReferenceTime);
+				}
+
 				FAccessUnit::Release(AccessUnit);
 			}
 		}
diff --git a/Engine/Plugins/Media/ElectraPlayer/Source/ElectraPlayerRuntime/Private/Runtime/Player/AdaptiveStreamingPlayerInternal.h b/Engine/Plugins/Media/ElectraPlayer/Source/ElectraPlayerRuntime/Private/Runtime/Player/AdaptiveStreamingPlayerInternal.h
--- a/Engine/Plugins/Media/ElectraPlayer/Source/ElectraPlayerRuntime/Private/Runtime/Player/AdaptiveStreamingPlayerInternal.h
+++ b/Engine/Plugins/Media/ElectraPlayer/Source/ElectraPlayerRuntime/Private/Runtime/Player/AdaptiveStreamingPlayerInternal.h
@@ -19,6 +19,7 @@
 #include "Player/PlaylistReader.h"
 #include "Player/PlayerStreamFilter.h"
 #include "Player/PlayerEntityCache.h"
+#include "Player/AdaptiveStreamingPlayerABR.h"
 
 #include "ElectraCDM.h"
 
@@ -42,6 +43,19 @@ namespace Electra
 class FAdaptiveStreamingPlayer;
 
 
+class IAdaptiveStreamingWrappedRenderer : public IMediaRenderer
+{
+public:
+	virtual ~IAdaptiveStreamingWrappedRenderer() = default;
+	virtual FTimeValue GetEnqueuedSampleDuration() = 0;
+	virtual int32 GetNumEnqueuedSamples(FTimeValue* OutOptionalDuration) = 0;
+
+	virtual FTimeRange GetSupportedRenderRateScale() = 0;
+	virtual void SetPlayRateScale(double InNewScale) = 0;
+	virtual double GetPlayRateScale() = 0;
+};
+
+
 class FMediaRenderClock : public IMediaRenderClock
 {
 public:
@@ -169,10 +183,12 @@ struct FPlaybackState
 	}
 	void Reset()
 	{
-		FMediaCriticalSection::ScopedLock lock(Lock);
+		FScopeLock lock(&Lock);
 		SeekableRange.Reset();
 		Duration.SetToInvalid();
 		CurrentPlayPosition.SetToInvalid();
+		EncoderLatency.SetToInvalid();
+		CurrentLiveLatency.SetToInvalid();
 		LoopState = {};
 		bHaveMetadata = false;
 		bHasEnded = false;
@@ -182,13 +198,16 @@ struct FPlaybackState
 		bIsPaused = false;
 		bPlayrangeHasChanged = false;
 		bLoopStateHasChanged = false;
+		bShouldPlayOnLiveEdge = false;
 	}
 
-	mutable FMediaCriticalSection			Lock;
+	mutable FCriticalSection				Lock;
 	FTimeRange								SeekableRange;
 	FTimeRange								TimelineRange;
 	FTimeValue								Duration;
 	FTimeValue								CurrentPlayPosition;
+	FTimeValue								EncoderLatency;
+	FTimeValue								CurrentLiveLatency;
 	IAdaptiveStreamingPlayer::FLoopState	LoopState;
 	bool									bHaveMetadata;
 	bool									bHasEnded;
@@ -198,6 +217,7 @@ struct FPlaybackState
 	bool									bIsPaused;
 	bool									bPlayrangeHasChanged;
 	bool									bLoopStateHasChanged;
+	bool									bShouldPlayOnLiveEdge;
 	TArray<FTrackMetadata>					VideoTracks;
 	TArray<FTrackMetadata>					AudioTracks;
 	TArray<FTrackMetadata>					SubtitleTracks;
@@ -206,135 +226,168 @@ struct FPlaybackState
 
 	void SetSeekableRange(const FTimeRange& TimeRange)
 	{
-		FMediaCriticalSection::ScopedLock lock(Lock);
+		FScopeLock lock(&Lock);
 		SeekableRange = TimeRange;
 	}
 	void GetSeekableRange(FTimeRange& OutRange) const
 	{
-		FMediaCriticalSection::ScopedLock lock(Lock);
+		FScopeLock lock(&Lock);
 		OutRange = SeekableRange;
 	}
 
 	void SetSeekablePositions(const TArray<FTimespan>& InSeekablePositions)
 	{
-		FMediaCriticalSection::ScopedLock lock(Lock);
+		FScopeLock lock(&Lock);
 		SeekablePositions = InSeekablePositions;
 	}
 	void GetSeekablePositions(TArray<FTimespan>& OutSeekablePositions) const
 	{
-		FMediaCriticalSection::ScopedLock lock(Lock);
+		FScopeLock lock(&Lock);
 		OutSeekablePositions = SeekablePositions;
 	}
 
 	void SetTimelineRange(const FTimeRange& TimeRange)
 	{
-		FMediaCriticalSection::ScopedLock lock(Lock);
+		FScopeLock lock(&Lock);
 		TimelineRange = TimeRange;
 	}
 	void GetTimelineRange(FTimeRange& OutRange) const
 	{
-		FMediaCriticalSection::ScopedLock lock(Lock);
+		FScopeLock lock(&Lock);
 		OutRange = TimelineRange;
 	}
 
 	void SetDuration(const FTimeValue& InDuration)
 	{
-		FMediaCriticalSection::ScopedLock lock(Lock);
+		FScopeLock lock(&Lock);
 		Duration = InDuration;
 	}
 	FTimeValue GetDuration() const
 	{
-		FMediaCriticalSection::ScopedLock lock(Lock);
+		FScopeLock lock(&Lock);
 		return Duration;
 	}
 
 	void SetPlayPosition(const FTimeValue& InPosition)
 	{
-		FMediaCriticalSection::ScopedLock lock(Lock);
+		FScopeLock lock(&Lock);
 		CurrentPlayPosition = InPosition;
 	}
 	FTimeValue GetPlayPosition() const
 	{
-		FMediaCriticalSection::ScopedLock lock(Lock);
+		FScopeLock lock(&Lock);
 		return CurrentPlayPosition;
 	}
 
+	void SetEncoderLatency(const FTimeValue& InEncoderLatency)
+	{
+		FScopeLock lock(&Lock);
+		EncoderLatency = InEncoderLatency;
+	}
+	FTimeValue GetEncoderLatency() const
+	{
+		FScopeLock lock(&Lock);
+		return EncoderLatency;
+	}
+
+	void SetCurrentLiveLatency(const FTimeValue& InCurrentLiveLatency)
+	{
+		FScopeLock lock(&Lock);
+		CurrentLiveLatency = InCurrentLiveLatency;
+	}
+	FTimeValue GetCurrentLiveLatency() const
+	{
+		FScopeLock lock(&Lock);
+		return CurrentLiveLatency;
+	}
+
+	void SetShouldPlayOnLiveEdge(bool bInShouldPlayOnLiveEdge)
+	{
+		FScopeLock lock(&Lock);
+		bShouldPlayOnLiveEdge = bInShouldPlayOnLiveEdge;
+	}
+	bool GetShouldPlayOnLiveEdge() const
+	{
+		FScopeLock lock(&Lock);
+		return bShouldPlayOnLiveEdge;
+	}
+
 	void SetHaveMetadata(bool bInHaveMetadata)
 	{
-		FMediaCriticalSection::ScopedLock lock(Lock);
+		FScopeLock lock(&Lock);
 		bHaveMetadata = bInHaveMetadata;
 	}
 	bool GetHaveMetadata() const
 	{
-		FMediaCriticalSection::ScopedLock lock(Lock);
+		FScopeLock lock(&Lock);
 		return bHaveMetadata;
 	}
 
 	void SetHasEnded(bool bInHasEnded)
 	{
-		FMediaCriticalSection::ScopedLock lock(Lock);
+		FScopeLock lock(&Lock);
 		bHasEnded = bInHasEnded;
 	}
 	bool GetHasEnded() const
 	{
-		FMediaCriticalSection::ScopedLock lock(Lock);
+		FScopeLock lock(&Lock);
 		return bHasEnded;
 	}
 
 	void SetIsSeeking(bool bInIsSeeking)
 	{
-		FMediaCriticalSection::ScopedLock lock(Lock);
+		FScopeLock lock(&Lock);
 		bIsSeeking = bInIsSeeking;
 	}
 	bool GetIsSeeking() const
 	{
-		FMediaCriticalSection::ScopedLock lock(Lock);
+		FScopeLock lock(&Lock);
 		return bIsSeeking;
 	}
 
 	void SetIsBuffering(bool bInIsBuffering)
 	{
-		FMediaCriticalSection::ScopedLock lock(Lock);
+		FScopeLock lock(&Lock);
 		bIsBuffering = bInIsBuffering;
 	}
 	bool GetIsBuffering() const
 	{
-		FMediaCriticalSection::ScopedLock lock(Lock);
+		FScopeLock lock(&Lock);
 		return bIsBuffering;
 	}
 
 	void SetIsPlaying(bool bInIsPlaying)
 	{
-		FMediaCriticalSection::ScopedLock lock(Lock);
+		FScopeLock lock(&Lock);
 		bIsPlaying = bInIsPlaying;
 	}
 	bool GetIsPlaying() const
 	{
-		FMediaCriticalSection::ScopedLock lock(Lock);
+		FScopeLock lock(&Lock);
 		return bIsPlaying;
 	}
 
 	void SetIsPaused(bool bInIsPaused)
 	{
-		FMediaCriticalSection::ScopedLock lock(Lock);
+		FScopeLock lock(&Lock);
 		bIsPaused = bInIsPaused;
 	}
 	bool GetIsPaused() const
 	{
-		FMediaCriticalSection::ScopedLock lock(Lock);
+		FScopeLock lock(&Lock);
 		return bIsPaused;
 	}
 
 	void SetPausedAndPlaying(bool bInIsPaused, bool bInIsPlaying)
 	{
-		FMediaCriticalSection::ScopedLock lock(Lock);
+		FScopeLock lock(&Lock);
 		bIsPaused  = bInIsPaused;
 		bIsPlaying = bInIsPlaying;
 	}
 
 	bool SetTrackMetadata(const TArray<FTrackMetadata> &InVideoTracks, const TArray<FTrackMetadata>& InAudioTracks, const TArray<FTrackMetadata>& InSubtitleTracks)
 	{
-		FMediaCriticalSection::ScopedLock lock(Lock);
+		FScopeLock lock(&Lock);
 
 		auto ChangedTrackMetadata = [](const TArray<FTrackMetadata> &These, const TArray<FTrackMetadata>& Other) -> bool
 		{
@@ -360,7 +413,7 @@ struct FPlaybackState
 	}
 	void GetTrackMetadata(TArray<FTrackMetadata> &OutVideoTracks, TArray<FTrackMetadata>& OutAudioTracks, TArray<FTrackMetadata>& OutSubtitleTracks) const
 	{
-		FMediaCriticalSection::ScopedLock lock(Lock);
+		FScopeLock lock(&Lock);
 		OutVideoTracks = VideoTracks;
 		OutAudioTracks = AudioTracks;
 		OutSubtitleTracks = SubtitleTracks;
@@ -368,38 +421,38 @@ struct FPlaybackState
 
 	void SetLoopState(const IAdaptiveStreamingPlayer::FLoopState& InLoopState)
 	{
-		FMediaCriticalSection::ScopedLock lock(Lock);
+		FScopeLock lock(&Lock);
 		LoopState = InLoopState;
 	}
 	void GetLoopState(IAdaptiveStreamingPlayer::FLoopState& OutLoopState) const
 	{
-		FMediaCriticalSection::ScopedLock lock(Lock);
+		FScopeLock lock(&Lock);
 		OutLoopState = LoopState;
 	}
 	void SetLoopStateEnable(bool bEnable)
 	{
-		FMediaCriticalSection::ScopedLock lock(Lock);
+		FScopeLock lock(&Lock);
 		LoopState.bIsEnabled = bEnable;
 	}
 
 	void SetPlayRangeHasChanged(bool bWasChanged)
 	{
-		FMediaCriticalSection::ScopedLock lock(Lock);
+		FScopeLock lock(&Lock);
 		bPlayrangeHasChanged = bWasChanged;
 	}
 	bool GetPlayRangeHasChanged()
 	{
-		FMediaCriticalSection::ScopedLock lock(Lock);
+		FScopeLock lock(&Lock);
 		return bPlayrangeHasChanged;
 	}
 	void SetLoopStateHasChanged(bool bWasChanged)
 	{
-		FMediaCriticalSection::ScopedLock lock(Lock);
+		FScopeLock lock(&Lock);
 		bLoopStateHasChanged = bWasChanged;
 	}
 	bool GetLoopStateHasChanged()
 	{
-		FMediaCriticalSection::ScopedLock lock(Lock);
+		FScopeLock lock(&Lock);
 		return bLoopStateHasChanged;
 	}
 };
@@ -674,17 +727,22 @@ struct FMetricEvent
 class FAdaptiveStreamingPlayerWorkerThread
 {
 public:
-	static TSharedPtrTS<FAdaptiveStreamingPlayerWorkerThread> Create();
+	static TSharedPtrTS<FAdaptiveStreamingPlayerWorkerThread> Create(bool bUseSharedWorkerThread);
 
 	void AddPlayerInstance(FAdaptiveStreamingPlayer* Instance)
 	{
-		FScopeLock lock(&SingletonLock);
-		PlayerInstances.Emplace(Instance);
+		InstancesToAdd.Enqueue(Instance);
+		TriggerWork();
 	}
 	void RemovePlayerInstance(FAdaptiveStreamingPlayer* Instance)
 	{
-		FScopeLock lock(&SingletonLock);
-		PlayerInstances.Remove(Instance);
+		FInstanceToRemove That;
+		FMediaEvent Done;
+		That.Player = Instance;
+		That.DoneSignal = &Done;
+		InstancesToRemove.Enqueue(MoveTemp(That));
+		TriggerWork();
+		Done.Wait();
 	}
 
 	void TriggerWork()
@@ -698,10 +756,19 @@ private:
 	void StopWorkerThread();
 	void WorkerThreadFN();
 
-	FMediaThread						WorkerThread;
-	TArray<FAdaptiveStreamingPlayer*>	PlayerInstances;
-	FMediaEvent							HaveWorkSignal;
-	bool								bTerminate = false;
+	struct FInstanceToRemove
+	{
+		FAdaptiveStreamingPlayer* Player = nullptr;
+		FMediaEvent* DoneSignal = nullptr;
+	};
+
+	FMediaThread										WorkerThread;
+	TQueue<FAdaptiveStreamingPlayer*, EQueueMode::Mpsc>	InstancesToAdd;
+	TQueue<FInstanceToRemove, EQueueMode::Mpsc>			InstancesToRemove;
+	TArray<FAdaptiveStreamingPlayer*>					PlayerInstances;
+	FMediaEvent											HaveWorkSignal;
+	bool												bTerminate = false;
+	bool												bIsDedicatedWorker = false;
 
 	static TWeakPtrTS<FAdaptiveStreamingPlayerWorkerThread>	SingletonSelf;
 	static FCriticalSection									SingletonLock;
@@ -735,6 +802,7 @@ class FAdaptiveStreamingPlayer : public IAdaptiveStreamingPlayer
 							   , public IAccessUnitMemoryProvider
 							   , public IPlayerStreamFilter
 							   , public TSharedFromThis<FAdaptiveStreamingPlayer, ESPMode::ThreadSafe>
+							   , public IAdaptiveStreamSelector::IPlayerLiveControl
 {
 public:
 	FAdaptiveStreamingPlayer(const IAdaptiveStreamingPlayer::FCreateParam& InCreateParameters);
@@ -828,9 +896,8 @@ private:
 	virtual void GetExternalGuid(FGuid& OutExternalGuid) override;
 	virtual ISynchronizedUTCTime* GetSynchronizedUTCTime() override;
 	virtual TSharedPtr<IAdaptiveStreamingPlayerResourceProvider, ESPMode::ThreadSafe> GetStaticResourceProvider() override;
-	virtual IElectraHttpManager* GetHTTPManager() override;
+	virtual TSharedPtrTS<IElectraHttpManager> GetHTTPManager() override;
 	virtual TSharedPtrTS<IAdaptiveStreamSelector> GetStreamSelector() override;
-	virtual void GetStreamBufferStats(FAccessUnitBufferInfo& OutBufferStats, EStreamType ForStream) override;
 	virtual IPlayerStreamFilter* GetStreamFilter() override;
 	virtual const FCodecSelectionPriorities& GetCodecSelectionPriorities(EStreamType ForStream) override;
 	virtual	TSharedPtrTS<IPlaylistReader> GetManifestReader() override;
@@ -846,6 +913,24 @@ private:
 	virtual bool CanDecodeSubtitle(const FString& MimeType, const FString& Codec) const override;
 
 
+	// Methods from IAdaptiveStreamSelector::IPlayerLiveControl
+	FTimeValue ABRGetPlayPosition() const override;
+	FTimeRange ABRGetTimeline() const override;
+	FTimeValue ABRGetWallclockTime() const override;
+	bool ABRIsLive() const override;
+	bool ABRShouldPlayOnLiveEdge() const override;
+	TSharedPtrTS<const FLowLatencyDescriptor> ABRGetLowLatencyDescriptor() const override;
+	FTimeValue ABRGetDesiredLiveEdgeLatency() const override;
+	FTimeValue ABRGetLatency() const override;
+	FTimeValue ABRGetPlaySpeed() const override;
+	void ABRGetStreamBufferStats(FAccessUnitBufferInfo& OutBufferStats, EStreamType ForStream) override;
+	void ABRGetStreamBufferStats(IAdaptiveStreamSelector::IPlayerLiveControl::FABRBufferStats& OutBufferStats, EStreamType ForStream) override;
+	FTimeRange ABRGetSupportedRenderRateScale() override;
+	void ABRSetRenderRateScale(double InRenderRateScale) override;
+	double ABRGetRenderRateScale() const override;
+	void ABRTriggerClockSync(IAdaptiveStreamSelector::IPlayerLiveControl::EClockSyncType InClockSyncType) override;
+
+
 	enum class EPlayerState
 	{
 		eState_Idle,
@@ -939,7 +1024,7 @@ private:
 				Renderer->Flush(options);
 			}
 		}
-		TSharedPtr<IMediaRenderer, ESPMode::ThreadSafe>	Renderer;
+		TSharedPtr<IAdaptiveStreamingWrappedRenderer, ESPMode::ThreadSafe>	Renderer;
 	};
 
 	struct FAudioRenderer
@@ -956,7 +1041,7 @@ private:
 				Renderer->Flush(options);
 			}
 		}
-		TSharedPtr<IMediaRenderer, ESPMode::ThreadSafe>	Renderer;
+		TSharedPtr<IAdaptiveStreamingWrappedRenderer, ESPMode::ThreadSafe>	Renderer;
 	};
 
 	struct FVideoDecoder : public IAccessUnitBufferListener, public IDecoderOutputBufferListener
@@ -1399,7 +1484,7 @@ private:
 		}
 
 		TSharedPtrTS<FAdaptiveStreamingPlayerWorkerThread>	SharedWorkerThread;
-		TQueue<FMessage, EQueueMode::Spsc>					WorkMessages;
+		TQueue<FMessage, EQueueMode::Mpsc>					WorkMessages;
 	};
 
 	struct FPendingStartRequest
@@ -1742,6 +1827,7 @@ private:
 	void InternalHandleCompletedSegmentRequests(const FTimeValue& CurrentTime);
 	void InternalHandleSegmentTrackChanges(const FTimeValue& CurrentTime);
 	void InternalStartoverAtCurrentPosition();
+	void InternalStartoverAtLiveEdge();
 	void InternalHandleLiveSync();
 
 	void UpdateDiagnostics();
@@ -1762,12 +1848,13 @@ private:
 
 	void CheckForErrors();
 
+	FTimeValue CalculateCurrentLiveLatency(bool bViaLatencyElement);
 	double GetMinBufferTimeBeforePlayback();
 	bool HaveEnoughBufferedDataToStartPlayback();
 	void PrepareForPrerolling();
 
 	void SetPlaystartOptions(FPlayStartOptions& OutOptions);
-	void ClampStartRequestTime();
+	void ClampStartRequestTime(FTimeValue& InOutTimeToClamp);
 	FTimeValue ClampTimeToCurrentRange(const FTimeValue& InTime, bool bClampToStart, bool bClampToEnd);
 
 	TSharedPtrTS<FMultiTrackAccessUnitBuffer> GetStreamBuffer(EStreamType InStreamType, const TSharedPtrTS<FStreamDataBuffers>& InFromStreamBuffers)
@@ -1824,6 +1911,7 @@ private:
 
 	void UpdateDataAvailabilityState(Metrics::FDataAvailabilityChange& DataAvailabilityState, Metrics::FDataAvailabilityChange::EAvailability NewAvailability);
 
+	TSharedPtr<IAdaptiveStreamingWrappedRenderer, ESPMode::ThreadSafe> CreateWrappedRenderer(TSharedPtr<IMediaRenderer, ESPMode::ThreadSafe> RendererToWrap, EStreamType InType);
 	void StartRendering();
 	void StopRendering();
 
@@ -1841,6 +1929,7 @@ private:
 	TSharedPtrTS<FAdaptiveStreamingPlayerEventHandler>					EventDispatcher;
 	TSharedPtrTS<FAdaptiveStreamingPlayerWorkerThread>					SharedWorkerThread;
 	FWorkerThreadMessages												WorkerThread;
+	bool																bUseSharedWorkerThread;
 
 	FParamDict															PlayerOptions;
 	FPlaybackState														PlaybackState;
@@ -1910,6 +1999,7 @@ private:
 	FSeekVars															SeekVars;
 	EPlayerState														LastBufferingState;
 	double																PlaybackRate;
+	double																RenderRateScale;
 	FTimeValue															RebufferDetectedAtPlayPos;
 	bool																bRebufferPending;
 	bool																bIsClosing;
diff --git a/Engine/Plugins/Media/ElectraPlayer/Source/ElectraPlayerRuntime/Private/Runtime/Player/AdaptiveStreamingPlayerInternalConfig.h b/Engine/Plugins/Media/ElectraPlayer/Source/ElectraPlayerRuntime/Private/Runtime/Player/AdaptiveStreamingPlayerInternalConfig.h
--- a/Engine/Plugins/Media/ElectraPlayer/Source/ElectraPlayerRuntime/Private/Runtime/Player/AdaptiveStreamingPlayerInternalConfig.h
+++ b/Engine/Plugins/Media/ElectraPlayer/Source/ElectraPlayerRuntime/Private/Runtime/Player/AdaptiveStreamingPlayerInternalConfig.h
@@ -6,13 +6,10 @@
 #include "Decoder/VideoDecoderH265.h"
 #include "Decoder/AudioDecoderAAC.h"
 #include "StreamAccessUnitBuffer.h"
-#include "Player/AdaptiveStreamingPlayerABR.h"
 
 
 namespace Electra
 {
-
-
 	namespace AdaptiveStreamingPlayerConfig
 	{
 		struct FConfiguration
@@ -63,8 +60,6 @@ namespace Electra
 			FAccessUnitBuffer::FConfiguration					StreamBufferConfigAudio;
 			FAccessUnitBuffer::FConfiguration					StreamBufferConfigText;
 
-			IAdaptiveStreamSelector::FConfiguration				StreamSelectorConfig;
-
 			IVideoDecoderH264::FInstanceConfiguration			DecoderCfg264;
 			IAudioDecoderAAC::FInstanceConfiguration			DecoderCfgAAC;
 
diff --git a/Engine/Plugins/Media/ElectraPlayer/Source/ElectraPlayerRuntime/Private/Runtime/Player/AdaptiveStreamingPlayerManifest.cpp b/Engine/Plugins/Media/ElectraPlayer/Source/ElectraPlayerRuntime/Private/Runtime/Player/AdaptiveStreamingPlayerManifest.cpp
--- a/Engine/Plugins/Media/ElectraPlayer/Source/ElectraPlayerRuntime/Private/Runtime/Player/AdaptiveStreamingPlayerManifest.cpp
+++ b/Engine/Plugins/Media/ElectraPlayer/Source/ElectraPlayerRuntime/Private/Runtime/Player/AdaptiveStreamingPlayerManifest.cpp
@@ -289,6 +289,9 @@ bool FAdaptiveStreamingPlayer::SelectManifest()
 				InternalCloseManifestReader();
 			}
 
+			// Let the ABR know the format as well.
+			StreamSelector->SetFormatType(ManifestType);
+
 			return true;
 		}
 		else
diff --git a/Engine/Plugins/Media/ElectraPlayer/Source/ElectraPlayerRuntime/Private/Runtime/Player/AdaptiveStreamingPlayerMetrics.h b/Engine/Plugins/Media/ElectraPlayer/Source/ElectraPlayerRuntime/Private/Runtime/Player/AdaptiveStreamingPlayerMetrics.h
--- a/Engine/Plugins/Media/ElectraPlayer/Source/ElectraPlayerRuntime/Private/Runtime/Player/AdaptiveStreamingPlayerMetrics.h
+++ b/Engine/Plugins/Media/ElectraPlayer/Source/ElectraPlayerRuntime/Private/Runtime/Player/AdaptiveStreamingPlayerMetrics.h
@@ -7,6 +7,7 @@
 #include "InfoLog.h"
 #include "Player/AdaptiveStreamingPlayerABR_State.h"
 #include "Player/Playlist.h"
+#include "ElectraHTTPStream.h"
 
 namespace Electra
 {
@@ -90,21 +91,6 @@ namespace Metrics
 		int64			BytesInUse;
 	};
 
-	struct FABRState
-	{
-		FABRState()
-		{
-			Reset();
-		}
-
-		void Reset()
-		{
-			ProgressDecision.Reset();
-		}
-		FABRDownloadProgressDecision		ProgressDecision;
-	};
-
-
 	struct FPlaylistDownloadStats
 	{
 		FPlaylistDownloadStats()
@@ -153,7 +139,6 @@ namespace Metrics
 		double			TimeToDownload = 0.0;				//!< Total time in seconds for entire download
 		int64			ByteSize = 0;						//!< Content-Length, may be -1 if unknown (either on error or chunked transfer)
 		int64			NumBytesDownloaded = 0;				//!< Number of bytes successfully downloaded.
-		int64			ThroughputBps = 0;					//!< Estimated throughput in bits per second
 		int32			HTTPStatusCode = 0 ;				//!< HTTP status code (0 if not connected to server yet)
 		bool			bWasSuccessful = false;				//!< true if download was successful, false if not
 		bool			bWasAborted = false;				//!< true if download was aborted by ABR (not by playback!)
@@ -162,8 +147,18 @@ namespace Metrics
 		bool			bInsertedFillerData = false;
 		bool			bIsCachedResponse = false;
 
-		// ABR state
-		FABRState		ABRState;
+		// Chunk timing
+		struct FMovieChunkInfo
+		{
+			int64 HeaderOffset = 0;
+			int64 PayloadStartOffset = 0;
+			int64 PayloadEndOffset = 0;
+			int64 NumKeyframeBytes = 0;
+			FTimeValue ContentDuration;
+			FMovieChunkInfo() { ContentDuration = FTimeValue::GetZero(); }
+		};
+		TArray<IElectraHTTPStreamResponse::FTimingTrace> TimingTraces;
+		TArray<FMovieChunkInfo> MovieChunkInfos;
 	};
 
 	struct FLicenseKeyStats
