ElectraPlayer: Added DASH low latency Live streaming

diff --git a/Engine/Plugins/Media/ElectraPlayer/Source/ElectraPlayerRuntime/Private/Runtime/Player/AdaptiveStreamingPlayerRenderers.cpp b/Engine/Plugins/Media/ElectraPlayer/Source/ElectraPlayerRuntime/Private/Runtime/Player/AdaptiveStreamingPlayerRenderers.cpp
new file mode 100644
--- /dev/null
+++ b/Engine/Plugins/Media/ElectraPlayer/Source/ElectraPlayerRuntime/Private/Runtime/Player/AdaptiveStreamingPlayerRenderers.cpp
@@ -0,0 +1,658 @@
+// Copyright Epic Games, Inc. All Rights Reserved.
+
+#include "PlayerCore.h"
+#include "Player/AdaptiveStreamingPlayerInternal.h"
+#include "Player/AdaptivePlayerOptionKeynames.h"
+#include "Utilities/Utilities.h"
+#include "ParameterDictionary.h"
+#include "Renderer/RendererBase.h"
+#include "ElectraPlayerPrivate.h"
+
+#include "HAL/LowLevelMemTracker.h"
+#include "Stats/Stats.h"
+
+#include <atomic>
+
+#define VALIDITY_VALUE_KEY TEXT("$renderVV$")
+#define DURATION_VALUE_KEY TEXT("duration")
+#define TIMESTAMP_VALUE_KEY TEXT("pts")
+
+#define AUDIO_BUFFER_SIZE TEXT("max_buffer_size")
+#define AUDIO_BUFFER_NUM TEXT("num_buffers")
+#define AUDIO_BUFFER_MAX_CHANNELS TEXT("max_channels")
+#define AUDIO_BUFFER_SAMPLES_PER_BLOCK TEXT("samples_per_block")
+#define AUDIO_BUFFER_ALLOCATED_ADDRESS TEXT("address")
+#define AUDIO_BUFFER_ALLOCATED_SIZE TEXT("size")
+#define AUDIO_BUFFER_NUM_CHANNELS TEXT("num_channels")
+#define AUDIO_BUFFER_BYTE_SIZE TEXT("byte_size")
+#define AUDIO_BUFFER_SAMPLE_RATE TEXT("sample_rate")
+
+/***************************************************************************************************************************************************/
+
+DECLARE_STATS_GROUP(TEXT("Electra Audio"), STATGROUP_ElectraAudioProcessing, STATCAT_Advanced);
+DECLARE_CYCLE_STAT(TEXT("Process"), STAT_ElectraAudioProcessing_Process, STATGROUP_ElectraAudioProcessing);
+
+/***************************************************************************************************************************************************/
+
+//#define ENABLE_PLAYRATE_OVERRIDE_CVAR
+
+#ifdef ENABLE_PLAYRATE_OVERRIDE_CVAR
+#include "HAL/IConsoleManager.h"
+static TAutoConsoleVariable<float> CVarElectraPR(TEXT("Electra.PR"), 1.0f, TEXT("Playback rate"), ECVF_Default);
+#endif
+
+/***************************************************************************************************************************************************/
+
+
+namespace Electra
+{
+	class FAdaptiveStreamingWrappedRenderer : public IAdaptiveStreamingWrappedRenderer
+	{
+	public:
+		virtual ~FAdaptiveStreamingWrappedRenderer() = default;
+
+		FAdaptiveStreamingWrappedRenderer(TSharedPtr<IMediaRenderer, ESPMode::ThreadSafe> RendererToWrap, EStreamType InType)
+			: WrappedRenderer(MoveTemp(RendererToWrap))
+			, Type(InType)
+		{
+			// If there is no renderer to wrap there should not be a wrapper created in the first place!
+			check(WrappedRenderer.IsValid());
+			NumPendingReturnBuffers = 0;
+			NumEnqueuedSamples = 0;
+			EnqueuedDuration.SetToZero();
+		}
+
+	private:
+		const int32 MaxSampleRate = 48000;
+		const int32 NumInterpolationSamplesAt48kHz = 60;
+		const double MinPlaybackSpeed = 0.8;
+		const double MaxPlaybackSpeed = 1.5;
+		const double MinResampleSpeed = 0.98;
+		const double MaxResampleSpeed = 1.02;
+
+		struct FPendingReturnBuffer
+		{
+			IBuffer* Buffer = nullptr;
+			bool bRender = false;
+			FParamDict Properties;
+		};
+
+
+		FTimeValue GetEnqueuedSampleDuration() override;
+		int32 GetNumEnqueuedSamples(FTimeValue* OutOptionalDuration) override;
+
+		FTimeRange GetSupportedRenderRateScale() override;
+		void SetPlayRateScale(double InNewScale) override;
+		double GetPlayRateScale() override;
+
+		void SampleReleasedToPool(IDecoderOutput* InDecoderOutput) override;
+
+		const FParamDict& GetBufferPoolProperties() const override;
+		UEMediaError CreateBufferPool(const FParamDict& Parameters) override;
+		UEMediaError AcquireBuffer(IBuffer*& OutBuffer, int32 TimeoutInMicroseconds, const FParamDict& InParameters) override;
+		UEMediaError ReturnBuffer(IBuffer* Buffer, bool bRender, const FParamDict& InSampleProperties) override;
+		UEMediaError ReleaseBufferPool() override;
+		bool CanReceiveOutputFrames(uint64 NumFrames) const override;
+		void SetRenderClock(TSharedPtr<IMediaRenderClock, ESPMode::ThreadSafe> InRenderClock) override;
+		void SetParentRenderer(TWeakPtr<IMediaRenderer, ESPMode::ThreadSafe> ParentRenderer) override;
+		void SetNextApproximatePresentationTime(const FTimeValue& NextApproxPTS) override;
+		UEMediaError Flush(const FParamDict& InOptions) override;
+		void StartRendering(const FParamDict& InOptions) override;
+		void StopRendering(const FParamDict& InOptions) override;
+		void TickOutputBufferPool() override;
+
+		UEMediaError ReturnVideoBuffer(IBuffer* Buffer, bool bRender, FParamDict& InSampleProperties);
+		UEMediaError ReturnAudioBuffer(IBuffer* Buffer, bool bRender, FParamDict& InSampleProperties);
+		UEMediaError ReturnBufferCommon(IBuffer* Buffer, bool bRender, FParamDict& InSampleProperties);
+		void ReturnAllPendingBuffers(bool bForFlush);
+
+		bool ProcessAudio(bool& bOutNeed2ndBuffer, IBuffer* Buffer, FParamDict& InSampleProperties, double InRate);
+
+		mutable FCriticalSection Lock;
+		TSharedPtr<IMediaRenderer, ESPMode::ThreadSafe> WrappedRenderer;
+		TSharedPtr<Electra::IMediaRenderClock, ESPMode::ThreadSafe>	RenderClock;
+		int64 CurrentValidityValue = 0;
+		EStreamType Type = EStreamType::Unsupported;
+		bool bIsRunning = false;
+
+		TQueue<FPendingReturnBuffer> PendingReturnBuffers;
+		std::atomic<int32> NumPendingReturnBuffers;
+
+		std::atomic<int32> NumBuffersInCirculation;
+
+		// Audio manipulation
+		struct FAudioVars
+		{
+			struct FConfig
+			{
+				int32 SampleRate = 0;
+				int32 NumChannels = 0;
+				bool DiffersFrom(int32 InSampleRate, int32 InNumChannels)
+				{ return InSampleRate != SampleRate || InNumChannels != NumChannels; }
+				void Update(int32 InSampleRate, int32 InNumChannels)
+				{ 
+					SampleRate = InSampleRate;
+					NumChannels = InNumChannels;
+				}
+				void Reset()
+				{
+					SampleRate = 0;
+					NumChannels = 0;
+				}
+			};
+
+			enum class EState
+			{
+				Disengaged,
+				Engaged
+			};
+
+			FConfig CurrentConfig;
+			EState CurrentState = EState::Disengaged;
+			int16 LastSampleValuePerChannel[16];
+			bool bNextBlockNeedsInterpolation = false;
+
+			int64 OriginalAudioBufferNum = 0;
+			int64 OriginalAudioBufferSize = 0;
+			int64 AudioBufferSize = 0;
+			int32 NumAudioBuffersInUse = 0;
+			double RateScale = 1.0;
+			int16* AudioTempSourceBuffer = nullptr;
+			int32 MaxOutputSampleBlockSize = 0;
+
+			FAudioVars()
+			{
+				Reset();
+				CurrentConfig.Reset();
+			}
+			~FAudioVars()
+			{
+				FMemory::Free(AudioTempSourceBuffer);
+			}
+			void Reset()
+			{
+				CurrentState = EState::Disengaged;
+				FMemory::Memzero(LastSampleValuePerChannel);
+				bNextBlockNeedsInterpolation = false;
+			}
+			void UpdateLastSampleValue(const int16* InSamples, int32 NumSamples, int32 NumChannels)
+			{
+				if (NumSamples)
+				{
+					InSamples += (NumSamples - 1) * NumChannels;
+					check(NumChannels <= 16);
+					FMemory::Memcpy(LastSampleValuePerChannel, InSamples, sizeof(int16) * NumChannels);
+				}
+			}
+			void InterpolateFromLastSampleValue(int16* InSamples, int32 NumSamples, int32 NumChannels, int32 SampleRate, int32 NumInterpolationSamples)
+			{
+				// Number of samples over which to interpolate depends on sampling rate.
+				int32 NumInter = Utils::Min((int32)(NumInterpolationSamples * (SampleRate / 48000.0)), NumSamples);
+				if (NumInter)
+				{
+					const int16 *LastInterpSample = InSamples + (NumInter-1) * NumChannels;
+					const float Step = 1.0f / (NumInter-1);
+					for(int32 i=1; i<NumInter; ++i)
+					{
+						for(int32 j=0; j<NumChannels; ++j)
+						{
+							InSamples[j] = (int16)((float) LastSampleValuePerChannel[j] + ((LastInterpSample[j] - LastSampleValuePerChannel[j]) * (i*Step)));
+						}
+						InSamples += NumChannels;
+					}
+				}
+			}
+		};
+		FAudioVars AudioVars;
+
+		// Stats
+		int32 NumEnqueuedSamples = 0;
+		FTimeValue EnqueuedDuration;
+	};
+
+
+
+TSharedPtr<IAdaptiveStreamingWrappedRenderer, ESPMode::ThreadSafe>	FAdaptiveStreamingPlayer::CreateWrappedRenderer(TSharedPtr<IMediaRenderer, ESPMode::ThreadSafe> RendererToWrap, EStreamType InType)
+{
+	TSharedPtr<IAdaptiveStreamingWrappedRenderer, ESPMode::ThreadSafe> WrappedRenderer;
+	if (RendererToWrap.IsValid())
+	{
+		WrappedRenderer = MakeShared<FAdaptiveStreamingWrappedRenderer, ESPMode::ThreadSafe>(RendererToWrap, InType);
+		RendererToWrap->SetParentRenderer(WrappedRenderer);
+	}
+	return WrappedRenderer;
+}
+
+
+void FAdaptiveStreamingWrappedRenderer::SampleReleasedToPool(IDecoderOutput* InDecoderOutput)
+{
+	--NumBuffersInCirculation;
+
+	check(InDecoderOutput);
+	if (InDecoderOutput && RenderClock.IsValid())
+	{
+		int64 ValidityValue = InDecoderOutput->GetMutablePropertyDictionary().GetValue(VALIDITY_VALUE_KEY).SafeGetInt64(0);
+		if (ValidityValue == CurrentValidityValue)
+		{
+			FTimeValue RenderTime = InDecoderOutput->GetMutablePropertyDictionary().GetValue(TIMESTAMP_VALUE_KEY).SafeGetTimeValue(FTimeValue::GetInvalid());
+			FTimeValue Duration = InDecoderOutput->GetMutablePropertyDictionary().GetValue(DURATION_VALUE_KEY).SafeGetTimeValue(FTimeValue::GetZero());
+			switch(Type)
+			{
+				case EStreamType::Video:
+					RenderClock->SetCurrentTime(Electra::IMediaRenderClock::ERendererType::Video, RenderTime);
+					break;
+				case EStreamType::Audio:
+					RenderClock->SetCurrentTime(Electra::IMediaRenderClock::ERendererType::Audio, RenderTime);
+					break;
+				default:
+					break;
+			}
+
+			{
+				FScopeLock lock(&Lock);
+
+				if (--NumEnqueuedSamples < 0)
+				{
+					NumEnqueuedSamples = 0;
+				}
+
+				EnqueuedDuration -= Duration;
+				if (!EnqueuedDuration.IsValid() || EnqueuedDuration < FTimeValue::GetZero())
+				{
+					EnqueuedDuration = FTimeValue::GetZero();
+				}
+			}
+		}
+	}
+	// Note that this is called *by* the wrapped renderer so we do *not* forward this call to there!
+}
+
+const FParamDict& FAdaptiveStreamingWrappedRenderer::GetBufferPoolProperties() const
+{
+	return WrappedRenderer->GetBufferPoolProperties();
+}
+
+UEMediaError FAdaptiveStreamingWrappedRenderer::CreateBufferPool(const FParamDict& InParameters)
+{
+	LLM_SCOPE(ELLMTag::ElectraPlayer);
+
+	FParamDict Parameters(InParameters);
+
+	NumBuffersInCirculation = 0;
+	NumEnqueuedSamples = 0;
+	EnqueuedDuration.SetToZero();
+
+	// Ask for larger buffers in case of audio. For playback speed changes we may need to create artificial
+	// samples to slow down audio playback and need larger buffers for that.
+	if (Type == EStreamType::Audio)
+	{
+		FMemory::Free(AudioVars.AudioTempSourceBuffer);
+		AudioVars.AudioTempSourceBuffer = nullptr;
+		AudioVars.NumAudioBuffersInUse = 0;
+		AudioVars.OriginalAudioBufferNum = Parameters.GetValue(AUDIO_BUFFER_NUM).SafeGetInt64(0);
+		AudioVars.OriginalAudioBufferSize = Parameters.GetValue(AUDIO_BUFFER_SIZE).SafeGetInt64(0);
+		if (AudioVars.OriginalAudioBufferSize)
+		{
+			// We need an occasional extra buffer when the input sample sequence counter changes. Double the number of buffers to accommodate.
+			Parameters.SetOrUpdate(AUDIO_BUFFER_NUM, FVariantValue(AudioVars.OriginalAudioBufferNum * 2));
+			AudioVars.AudioTempSourceBuffer = (int16*)FMemory::Malloc(AudioVars.OriginalAudioBufferSize);
+
+			int32 SamplesPerBlock = (int32) Parameters.GetValue(AUDIO_BUFFER_SAMPLES_PER_BLOCK).SafeGetInt64(2048);
+			int32 MaxChannels = (int32) Parameters.GetValue(AUDIO_BUFFER_MAX_CHANNELS).SafeGetInt64(8);
+
+			// Get maximum number of samples we may produce when slowing down the most.
+			int32 NumResampleSamples = (int32)(SamplesPerBlock / MinPlaybackSpeed + 0.5);
+
+			AudioVars.MaxOutputSampleBlockSize = Utils::Max(NumResampleSamples, SamplesPerBlock);
+			AudioVars.AudioBufferSize = AudioVars.MaxOutputSampleBlockSize * MaxChannels * sizeof(int16);
+			Parameters.SetOrUpdate(AUDIO_BUFFER_SIZE, FVariantValue(AudioVars.AudioBufferSize));
+		}
+	}
+	return WrappedRenderer->CreateBufferPool(Parameters);
+}
+
+UEMediaError FAdaptiveStreamingWrappedRenderer::AcquireBuffer(IBuffer*& OutBuffer, int32 TimeoutInMicroseconds, const FParamDict& InParameters)
+{
+	LLM_SCOPE(ELLMTag::ElectraPlayer);
+
+	if (Type == EStreamType::Audio)
+	{
+		if (NumBuffersInCirculation >= AudioVars.OriginalAudioBufferNum)
+		{
+			return UEMEDIA_ERROR_INSUFFICIENT_DATA;
+		}
+		// For audio we want to return the buffer with the originally requested size.
+		// This is to prevent users to do unexpected things like putting more samples in there like they originally wanted
+		// to because they see that the buffer can accommodate more.
+		OutBuffer = nullptr;
+		UEMediaError Error = WrappedRenderer->AcquireBuffer(OutBuffer, TimeoutInMicroseconds, InParameters);
+		if (Error == UEMEDIA_ERROR_OK && OutBuffer)
+		{
+			++NumBuffersInCirculation;
+			OutBuffer->GetMutableBufferProperties().SetOrUpdate(AUDIO_BUFFER_ALLOCATED_SIZE, FVariantValue(AudioVars.OriginalAudioBufferSize));
+		}
+		return Error;
+	}
+	else
+	{
+		UEMediaError Error = WrappedRenderer->AcquireBuffer(OutBuffer, TimeoutInMicroseconds, InParameters);
+		if (Error == UEMEDIA_ERROR_OK && OutBuffer)
+		{
+			++NumBuffersInCirculation;
+		}
+		return Error;
+	}
+}
+
+UEMediaError FAdaptiveStreamingWrappedRenderer::ReturnBuffer(IBuffer* Buffer, bool bRender, const FParamDict& InSampleProperties)
+{
+	LLM_SCOPE(ELLMTag::ElectraPlayer);
+
+	FParamDict SampleProperties(InSampleProperties);
+	Lock.Lock();
+	SampleProperties.SetOrUpdate(VALIDITY_VALUE_KEY, FVariantValue(CurrentValidityValue));
+	Lock.Unlock();
+	if (Type == EStreamType::Video)
+	{
+		return ReturnVideoBuffer(Buffer, bRender, SampleProperties);
+	}
+	else if (Type == EStreamType::Audio)
+	{
+		return ReturnAudioBuffer(Buffer, bRender, SampleProperties);
+	}
+	else
+	{
+		check(!"What type is this?");
+		return UEMEDIA_ERROR_NOT_SUPPORTED;
+	}
+}
+
+UEMediaError FAdaptiveStreamingWrappedRenderer::ReturnVideoBuffer(IBuffer* Buffer, bool bRender, FParamDict& InSampleProperties)
+{
+	return ReturnBufferCommon(Buffer, bRender, InSampleProperties);
+}
+
+UEMediaError FAdaptiveStreamingWrappedRenderer::ReturnAudioBuffer(IBuffer* Buffer, bool bRender, FParamDict& InSampleProperties)
+{
+	bool bNeed2ndBuffer = false;
+	// Process audio when the block will be rendered.
+	if (bRender)
+	{
+		Lock.Lock();
+		double DesiredOutputPlayrate = AudioVars.RateScale;
+		Lock.Unlock();
+		#ifdef ENABLE_PLAYRATE_OVERRIDE_CVAR
+		//DesiredOutputPlayrate = Utils::Max(Utils::Min((double)CVarElectraPR.GetValueOnAnyThread(), MaxPlaybackSpeed), MinPlaybackSpeed);
+		#endif
+		if (ProcessAudio(bNeed2ndBuffer, Buffer, InSampleProperties, DesiredOutputPlayrate))
+		{
+			if (bNeed2ndBuffer)
+			{
+				FParamDict NoParams;
+				// Return the current buffer. This should allow us to get another one for the
+				// residuals immediately.
+				UEMediaError Error = ReturnBufferCommon(Buffer, bRender, InSampleProperties);
+				check(Error == UEMEDIA_ERROR_OK);
+				Error = WrappedRenderer->AcquireBuffer(Buffer, 0, NoParams);
+				check(Error == UEMEDIA_ERROR_OK);
+				if (Error != UEMEDIA_ERROR_OK)
+				{
+					return Error;
+				}
+				++NumBuffersInCirculation;
+				ProcessAudio(bNeed2ndBuffer, Buffer, InSampleProperties, DesiredOutputPlayrate);
+			}
+		}
+		else
+		{
+			// Audio processing produced no samples that can be output right now.
+			// Mark this for not-to-get-rendered and return it.
+			bRender = false;
+		}
+	}
+	return ReturnBufferCommon(Buffer, bRender, InSampleProperties);
+}
+
+UEMediaError FAdaptiveStreamingWrappedRenderer::ReturnBufferCommon(IBuffer* Buffer, bool bRender, FParamDict& InSampleProperties)
+{
+	FTimeValue Duration = InSampleProperties.GetValue(DURATION_VALUE_KEY).SafeGetTimeValue(FTimeValue::GetZero());
+	if (!Duration.IsValid())
+	{
+		Duration.SetToZero();
+		InSampleProperties.SetOrUpdate(DURATION_VALUE_KEY, FVariantValue(Duration));
+	}
+
+	FScopeLock lock(&Lock);
+	EnqueuedDuration += Duration;
+	++NumEnqueuedSamples;
+	if (!bIsRunning)
+	{
+		FPendingReturnBuffer pb;
+		pb.Buffer = Buffer;
+		pb.bRender = bRender;
+		pb.Properties = InSampleProperties;
+		PendingReturnBuffers.Enqueue(MoveTemp(pb));
+		++NumPendingReturnBuffers;
+		return UEMEDIA_ERROR_OK;
+	}
+	lock.Unlock();
+	return WrappedRenderer->ReturnBuffer(Buffer, bRender, InSampleProperties);
+}
+
+
+
+UEMediaError FAdaptiveStreamingWrappedRenderer::ReleaseBufferPool()
+{
+	LLM_SCOPE(ELLMTag::ElectraPlayer);
+
+	ReturnAllPendingBuffers(false);
+	NumEnqueuedSamples = 0;
+	EnqueuedDuration.SetToZero();
+	Lock.Unlock();
+
+	AudioVars.Reset();
+
+	return WrappedRenderer->ReleaseBufferPool();
+}
+
+bool FAdaptiveStreamingWrappedRenderer::CanReceiveOutputFrames(uint64 NumFrames) const
+{
+	int32 NumPending = NumPendingReturnBuffers;
+	return WrappedRenderer->CanReceiveOutputFrames(NumFrames + NumPending);
+}
+
+void FAdaptiveStreamingWrappedRenderer::SetRenderClock(TSharedPtr<IMediaRenderClock, ESPMode::ThreadSafe> InRenderClock)
+{
+	RenderClock = InRenderClock;
+	WrappedRenderer->SetRenderClock(InRenderClock);
+}
+
+void FAdaptiveStreamingWrappedRenderer::SetParentRenderer(TWeakPtr<IMediaRenderer, ESPMode::ThreadSafe> ParentRenderer)
+{
+	check(!"this must not be called");
+}
+
+void FAdaptiveStreamingWrappedRenderer::SetNextApproximatePresentationTime(const FTimeValue& NextApproxPTS)
+{
+	WrappedRenderer->SetNextApproximatePresentationTime(NextApproxPTS);
+}
+
+UEMediaError FAdaptiveStreamingWrappedRenderer::Flush(const FParamDict& InOptions)
+{
+	LLM_SCOPE(ELLMTag::ElectraPlayer);
+
+	ReturnAllPendingBuffers(true);
+	++CurrentValidityValue;
+	NumEnqueuedSamples = 0;
+	EnqueuedDuration.SetToZero();
+	Lock.Unlock();
+
+	AudioVars.Reset();
+
+	return WrappedRenderer->Flush(InOptions);
+}
+
+void FAdaptiveStreamingWrappedRenderer::StartRendering(const FParamDict& InOptions)
+{
+	ReturnAllPendingBuffers(false);
+	bIsRunning = true;
+	Lock.Unlock();
+
+	WrappedRenderer->StartRendering(InOptions);
+}
+
+void FAdaptiveStreamingWrappedRenderer::StopRendering(const FParamDict& InOptions)
+{
+	Lock.Lock();
+	bIsRunning = false;
+	Lock.Unlock();
+
+	WrappedRenderer->StopRendering(InOptions);
+}
+
+void FAdaptiveStreamingWrappedRenderer::TickOutputBufferPool()
+{
+	WrappedRenderer->TickOutputBufferPool();
+}
+
+void FAdaptiveStreamingWrappedRenderer::ReturnAllPendingBuffers(bool bForFlush)
+{
+	Lock.Lock();
+	while(NumPendingReturnBuffers > 0)
+	{
+		FPendingReturnBuffer Buffer;
+		PendingReturnBuffers.Dequeue(Buffer);
+		--NumPendingReturnBuffers;
+		Lock.Unlock();
+		WrappedRenderer->ReturnBuffer(Buffer.Buffer, bForFlush ? false : Buffer.bRender, Buffer.Properties);
+		Lock.Lock();
+	}
+}
+
+
+FTimeValue FAdaptiveStreamingWrappedRenderer::GetEnqueuedSampleDuration()
+{
+	FScopeLock lock(&Lock);
+	return EnqueuedDuration;
+}
+
+int32 FAdaptiveStreamingWrappedRenderer::GetNumEnqueuedSamples(FTimeValue* OutOptionalDuration)
+{
+	FScopeLock lock(&Lock);
+	if (OutOptionalDuration)
+	{
+		*OutOptionalDuration = EnqueuedDuration;
+	}
+	return NumEnqueuedSamples;
+}
+
+FTimeRange FAdaptiveStreamingWrappedRenderer::GetSupportedRenderRateScale()
+{
+	FScopeLock lock(&Lock);
+	FTimeRange Range;
+	if (Type == EStreamType::Audio)
+	{
+		Range.Start.SetFromSeconds(MinPlaybackSpeed);
+		Range.End.SetFromSeconds(MaxPlaybackSpeed);
+	}
+	return Range;
+}
+
+void FAdaptiveStreamingWrappedRenderer::SetPlayRateScale(double InNewScale)
+{
+	FScopeLock lock(&Lock);
+	// Clamp to within permitted range just in case.
+	InNewScale = Utils::Max(Utils::Min(InNewScale, MaxPlaybackSpeed), MinPlaybackSpeed);
+	// Quantize to 0.005 multiples.
+	AudioVars.RateScale = (double)((int32)(InNewScale * 200.0)) / 200.0;
+}
+
+double FAdaptiveStreamingWrappedRenderer::GetPlayRateScale()
+{
+	return AudioVars.RateScale;
+}
+
+
+bool FAdaptiveStreamingWrappedRenderer::ProcessAudio(bool& bOutNeed2ndBuffer, IBuffer* Buffer, FParamDict& InSampleProperties, double InRate)
+{
+	SCOPE_CYCLE_COUNTER(STAT_ElectraAudioProcessing_Process);
+
+	bool bGetResiduals = bOutNeed2ndBuffer;
+	bOutNeed2ndBuffer = false;
+
+	int32 SizeInBytes = (int32)InSampleProperties.GetValue(AUDIO_BUFFER_BYTE_SIZE).SafeGetInt64();
+	FTimeValue Timestamp = InSampleProperties.GetValue(TIMESTAMP_VALUE_KEY).SafeGetTimeValue(FTimeValue::GetInvalid());
+	FTimeValue Duration = InSampleProperties.GetValue(DURATION_VALUE_KEY).SafeGetTimeValue(FTimeValue::GetZero());
+	int32 NumChannels = (int32)InSampleProperties.GetValue(AUDIO_BUFFER_NUM_CHANNELS).SafeGetInt64();
+	int32 NumSamples = SizeInBytes / NumChannels / sizeof(int16);
+	int32 SampleRate = (int32)InSampleProperties.GetValue(AUDIO_BUFFER_SAMPLE_RATE).SafeGetInt64();
+	int16* BufferAddress = (int16*)Buffer->GetBufferProperties().GetValue(AUDIO_BUFFER_ALLOCATED_ADDRESS).GetPointer();
+
+	if (AudioVars.CurrentConfig.DiffersFrom(SampleRate, NumChannels))
+	{
+		AudioVars.CurrentConfig.Update(SampleRate, NumChannels);
+		AudioVars.Reset();
+	}
+
+	bool bUpdateProperties = false;
+
+	if (InRate != 1.0)
+	{
+		// Small enough change to use resampler where pitch changes may not be that noticeable?
+		if (InRate >= MinPlaybackSpeed && InRate <= MaxPlaybackSpeed)
+		{
+			int32 NumOutputSamples = (int32) FMath::RoundToZero(NumSamples / InRate);
+			if (NumOutputSamples > 16)
+			{
+				const int32 MaxOutSamples = AudioVars.MaxOutputSampleBlockSize;
+				check(SizeInBytes <= AudioVars.OriginalAudioBufferSize);
+				FMemory::Memcpy(AudioVars.AudioTempSourceBuffer, BufferAddress, SizeInBytes);
+
+				double Offset = 0.0;
+				int32 o = 0;
+				double Step = (double)NumSamples / (double)NumOutputSamples;
+				while(o < NumOutputSamples && o < MaxOutSamples)
+				{
+					int32 I0 = (int32)Offset;
+					if (I0+1 >= NumSamples)
+					{
+						break;
+					}
+					double F0 = Offset - I0;
+					for(int32 nC=0; nC<NumChannels; ++nC)
+					{
+						double S0 = AudioVars.AudioTempSourceBuffer[I0       * NumChannels + nC];
+						double S1 = AudioVars.AudioTempSourceBuffer[(I0 + 1) * NumChannels + nC];
+						double S = S0 + (S1-S0) * F0;
+						BufferAddress[o * NumChannels + nC] = (int16) S;
+					}
+					++o;
+					Offset += Step;
+				}
+				NumSamples = o;
+				bUpdateProperties = true;
+			}
+		}
+	}
+
+	if (bUpdateProperties)
+	{
+		check(NumSamples <= AudioVars.AudioBufferSize / NumChannels / sizeof(int16));
+		InSampleProperties.SetOrUpdate(AUDIO_BUFFER_BYTE_SIZE, FVariantValue((int64)(NumSamples * sizeof(int16) * NumChannels)));
+		InSampleProperties.SetOrUpdate(TIMESTAMP_VALUE_KEY, FVariantValue(Timestamp));
+		InSampleProperties.SetOrUpdate(DURATION_VALUE_KEY, FVariantValue(FTimeValue(NumSamples, SampleRate, 0)));
+	}
+
+	// Need to interpolate this block's start samples from the last block's last values?
+	if (AudioVars.bNextBlockNeedsInterpolation)
+	{
+		AudioVars.bNextBlockNeedsInterpolation = false;
+		AudioVars.InterpolateFromLastSampleValue(BufferAddress, NumSamples, NumChannels, SampleRate, NumInterpolationSamplesAt48kHz);
+	}
+
+	// Remember last sample value for interpolation, if necessary.
+	AudioVars.UpdateLastSampleValue(BufferAddress, NumSamples, NumChannels);
+	return true;
+}
+
+} // namespace Electra
diff --git a/Engine/Plugins/Media/ElectraPlayer/Source/ElectraPlayerRuntime/Private/Runtime/Player/AdaptiveStreamingPlayerResourceRequest.cpp b/Engine/Plugins/Media/ElectraPlayer/Source/ElectraPlayerRuntime/Private/Runtime/Player/AdaptiveStreamingPlayerResourceRequest.cpp
--- a/Engine/Plugins/Media/ElectraPlayer/Source/ElectraPlayerRuntime/Private/Runtime/Player/AdaptiveStreamingPlayerResourceRequest.cpp
+++ b/Engine/Plugins/Media/ElectraPlayer/Source/ElectraPlayerRuntime/Private/Runtime/Player/AdaptiveStreamingPlayerResourceRequest.cpp
@@ -26,7 +26,11 @@ FHTTPResourceRequest::~FHTTPResourceRequest()
 	Cancel();
 	if (bWasAdded && Request.IsValid())
 	{
-		PlayerSessionServices->GetHTTPManager()->RemoveRequest(Request, true);
+		TSharedPtrTS<IElectraHttpManager> HTTPManager = PlayerSessionServices->GetHTTPManager();
+		if (HTTPManager.IsValid())
+		{
+			HTTPManager->RemoveRequest(Request, true);
+		}
 	}
 	Request.Reset();
 }
@@ -45,7 +49,11 @@ void FHTTPResourceRequest::StartGet(IPlayerSessionServices* InPlayerSessionServi
 	else
 	{
 		bWasAdded = true;
-		PlayerSessionServices->GetHTTPManager()->AddRequest(Request, false);
+		TSharedPtrTS<IElectraHttpManager> HTTPManager = PlayerSessionServices->GetHTTPManager();
+		if (HTTPManager.IsValid())
+		{
+			HTTPManager->AddRequest(Request, false);
+		}
 	}
 }
 
@@ -56,7 +64,11 @@ void FHTTPResourceRequest::Cancel()
 	ReceiveBuffer.Reset();
 	if (bWasAdded && Request.IsValid())
 	{
-		PlayerSessionServices->GetHTTPManager()->RemoveRequest(Request, true);
+		TSharedPtrTS<IElectraHttpManager> HTTPManager = PlayerSessionServices->GetHTTPManager();
+		if (HTTPManager.IsValid())
+		{
+			HTTPManager->RemoveRequest(Request, true);
+		}
 		bWasAdded = false;
 	}
 }
@@ -71,7 +83,11 @@ void FHTTPResourceRequest::StaticDataReady()
 		{
 			// Do the actual HTTP request now.
 			bWasAdded = true;
-			PlayerSessionServices->GetHTTPManager()->AddRequest(Request, false);
+			TSharedPtrTS<IElectraHttpManager> HTTPManager = PlayerSessionServices->GetHTTPManager();
+			if (HTTPManager.IsValid())
+			{
+				HTTPManager->AddRequest(Request, false);
+			}
 		}
 		else
 		{
diff --git a/Engine/Plugins/Media/ElectraPlayer/Source/ElectraPlayerRuntime/Private/Runtime/Player/DASH/MPDElementsDASH.h b/Engine/Plugins/Media/ElectraPlayer/Source/ElectraPlayerRuntime/Private/Runtime/Player/DASH/MPDElementsDASH.h
--- a/Engine/Plugins/Media/ElectraPlayer/Source/ElectraPlayerRuntime/Private/Runtime/Player/DASH/MPDElementsDASH.h
+++ b/Engine/Plugins/Media/ElectraPlayer/Source/ElectraPlayerRuntime/Private/Runtime/Player/DASH/MPDElementsDASH.h
@@ -92,6 +92,7 @@ public:
 		PlaybackRate,
 		Latency,
 		ServiceDescription,
+		Resync,
 		Other
 	};
 
@@ -219,13 +220,13 @@ public:
 	{ return EType::BaseURL; }
 
 	// Convenience method to get the URL.
-	const FString& GetURL() const						{ return Data; }
+	const FString& GetURL() const											{ return Data; }
 
-	const FString& GetServiceLocation() const			{ return ServiceLocation; }
-	const FString& GetByteRange() const					{ return ByteRange; }
-	const FTimeValue& GetAvailabilityTimeOffset() const	{ return AvailabilityTimeOffset; }
-	bool GetAvailabilityTimeComplete() const			{ return bAvailabilityTimeComplete; }
-	const FTimeValue& GetTimeShiftBufferDepth() const	{ return TimeShiftBufferDepth; }
+	const FString& GetServiceLocation() const								{ return ServiceLocation; }
+	const FString& GetByteRange() const										{ return ByteRange; }
+	const FTimeValue& GetAvailabilityTimeOffset() const						{ return AvailabilityTimeOffset; }
+	const TMediaOptionalValue<bool>& GetAvailabilityTimeComplete() const	{ return bAvailabilityTimeComplete; }
+	const FTimeValue& GetTimeShiftBufferDepth() const						{ return TimeShiftBufferDepth; }
 
 private:
 	virtual bool ProcessElement(FManifestParserDASH* Builder, const TCHAR* ElementName, const TCHAR* ElementData, int32 XmlFileLineNumber) override;
@@ -234,7 +235,7 @@ private:
 	FString ServiceLocation;
 	FString ByteRange;
 	FTimeValue AvailabilityTimeOffset = FTimeValue::GetZero();
-	bool bAvailabilityTimeComplete = true;
+	TMediaOptionalValue<bool> bAvailabilityTimeComplete;
 	FTimeValue TimeShiftBufferDepth;
 };
 
@@ -653,7 +654,7 @@ public:
 	const FString& GetIndexRange() const									{ return IndexRange; }
 	bool GetIndexRangeExact() const											{ return bIndexRangeExact; }
 	const FTimeValue& GetAvailabilityTimeOffset() const						{ return AvailabilityTimeOffset; }
-	bool GetAvailabilityTimeComplete() const								{ return bAvailabilityTimeComplete; }
+	const TMediaOptionalValue<bool>& GetAvailabilityTimeComplete() const	{ return bAvailabilityTimeComplete; }
 
 protected:
 	virtual bool ProcessElement(FManifestParserDASH* Builder, const TCHAR* ElementName, const TCHAR* ElementData, int32 XmlFileLineNumber) override;
@@ -673,7 +674,7 @@ private:
 	FString IndexRange;
 	bool bIndexRangeExact = false;
 	FTimeValue AvailabilityTimeOffset = FTimeValue::GetZero();
-	bool bAvailabilityTimeComplete = true;
+	TMediaOptionalValue<bool> bAvailabilityTimeComplete;
 };
 
 
@@ -896,7 +897,7 @@ public:
 
 	TSharedPtrTS<FDashMPD_DescriptorType> GetUTCTiming() const	{ return UTCTiming; }
 
-	uint64 GetID() const						{ return ID; }
+	uint32 GetID() const						{ return (uint32)ID; }
 	bool GetInband() const						{ return bInband; }
 	const FString& GetType() const				{ return Type; }
 	const FString& GetApplicationScheme() const	{ return ApplicationScheme; }
@@ -1209,7 +1210,7 @@ public:
 	const TArray<TSharedPtrTS<FDashMPD_OperatingQualityType>>& GetOperatingQualities() const	{ return OperatingQualities; }
 	const TArray<TSharedPtrTS<FDashMPD_OperatingBandwidthType>>& GetOperatingBandwidths() const	{ return OperatingBandwidths; }
 
-	const TMediaOptionalValue<uint64>& GetID() const	{ return ID; }
+	uint32 GetID() const { return (uint32) ID; }
 
 private:
 	virtual bool ProcessElement(FManifestParserDASH* Builder, const TCHAR* ElementName, const TCHAR* ElementData, int32 XmlFileLineNumber) override;
@@ -1223,7 +1224,43 @@ private:
 	TArray<TSharedPtrTS<FDashMPD_OperatingBandwidthType>> OperatingBandwidths;
 
 	// Attributes
-	TMediaOptionalValue<uint64> ID;
+	uint64 ID = 0;
+};
+
+
+/**
+ * 5.3.12 - ResyncType
+ */
+class FDashMPD_ResyncType : public IDashMPDElement
+{
+	typedef IDashMPDElement Super;
+public:
+	FDashMPD_ResyncType(const TCHAR* Name, const TCHAR* Data) : Super(Name, Data)
+	{ }
+	virtual ~FDashMPD_ResyncType() = default;
+	virtual EType GetElementType() const override
+	{ return EType::Resync; }
+
+	int32 GetType() const							{ return Type; }
+	TMediaOptionalValue<uint32> GetdT() const		{ return dT; }
+	TMediaOptionalValue<double> GetdIMax() const	{ return dIMax; }
+	double GetdIMin() const							{ return dIMin; }
+	bool GetMarker() const							{ return bMarker; }
+	bool GetRangeAccess() const						{ return bRangeAccess; }
+	const FString& GetIndex() const					{ return Index; }
+
+private:
+	virtual bool ProcessElement(FManifestParserDASH* Builder, const TCHAR* ElementName, const TCHAR* ElementData, int32 XmlFileLineNumber) override;
+	virtual bool ProcessAttribute(FManifestParserDASH* Builder, const TCHAR* AttributeName, const TCHAR* AttributeValue) override;
+
+	// Attributes
+	int32 Type = 0;						// SAP type
+	TMediaOptionalValue<uint32> dT;
+	TMediaOptionalValue<double> dIMax;
+	double dIMin = 0.0;
+	bool bMarker = false;
+	bool bRangeAccess = false;
+	FString Index;
 };
 
 
@@ -1250,6 +1287,7 @@ public:
 	const TArray<TSharedPtrTS<FDashMPD_LabelType>>& GetLabels() const										{ return Labels; }
 	const TArray<TSharedPtrTS<FDashMPD_ProducerReferenceTimeType>>& GetProducerReferenceTimes() const		{ return ProducerReferenceTimes; }
 	const TArray<TSharedPtrTS<FDashMPD_ContentPopularityRateType>>& GetContentPopularityRates() const		{ return ContentPopularityRates; }
+	const TArray<TSharedPtrTS<FDashMPD_ResyncType>>& GetResyncs() const										{ return Resyncs; }
 
 	const TArray<FString>& GetProfiles() const						{ return Profiles; }
 	const TMediaOptionalValue<uint64>& GetWidth() const				{ return Width; }
@@ -1286,6 +1324,7 @@ private:
 	TArray<TSharedPtrTS<FDashMPD_LabelType>> Labels;
 	TArray<TSharedPtrTS<FDashMPD_ProducerReferenceTimeType>> ProducerReferenceTimes;
 	TArray<TSharedPtrTS<FDashMPD_ContentPopularityRateType>> ContentPopularityRates;
+	TArray<TSharedPtrTS<FDashMPD_ResyncType>> Resyncs;
 
 	// Attributes
 	TArray<FString> Profiles;
diff --git a/Engine/Plugins/Media/ElectraPlayer/Source/ElectraPlayerRuntime/Private/Runtime/Player/DASH/ManifestBuilderDASH.cpp b/Engine/Plugins/Media/ElectraPlayer/Source/ElectraPlayerRuntime/Private/Runtime/Player/DASH/ManifestBuilderDASH.cpp
--- a/Engine/Plugins/Media/ElectraPlayer/Source/ElectraPlayerRuntime/Private/Runtime/Player/DASH/ManifestBuilderDASH.cpp
+++ b/Engine/Plugins/Media/ElectraPlayer/Source/ElectraPlayerRuntime/Private/Runtime/Player/DASH/ManifestBuilderDASH.cpp
@@ -11,6 +11,7 @@
 #include "Player/DASH/OptionKeynamesDASH.h"
 #include "Player/AdaptivePlayerOptionKeynames.h"
 #include "Player/DASH/PlayerEventDASH.h"
+#include "Player/DASH/PlayerEventDASH_Internal.h"
 
 #include "ManifestBuilderDASH.h"
 #include "ManifestParserDASH.h"
@@ -78,6 +79,8 @@ namespace
 		// DVB-DASH profiles
 		TEXT("urn:dvb:dash:profile:dvb-dash:2014"),
 		TEXT("urn:dvb:dash:profile:dvb-dash:isoff-ext-live:2014"),
+		// Low latency Live
+		TEXT("http://www.dashif.org/guidelines/low-latency-live-v5"),
 	};
 
 	const TCHAR* const ScanTypeInterlace = TEXT("interlace");
@@ -100,6 +103,10 @@ namespace
     const TCHAR* const SubtitleMimeType_SideloadedTTML = TEXT("application/ttml+xml");
     const TCHAR* const SubtitleMimeType_SideloadedVTT = TEXT("text/vtt");
     const TCHAR* const SubtitleMimeType_Streamed = TEXT("application/mp4");
+
+	// Thumbnail mime types
+    const TCHAR* const ThumbnailMimeType_Jpeg = TEXT("image/jpeg");
+    const TCHAR* const ThumbnailMimeType_Png= TEXT("image/png");
 }
 
 
@@ -560,10 +567,11 @@ namespace DASHUrlHelpers
 	 * Returns true if an absolute URL could be generated, false if not.
 	 * Since the MPD URL had to be an absolute URL this cannot actually fail.
 	 */
-	bool BuildAbsoluteElementURL(FString& OutURL, FTimeValue& ATO, const FString& DocumentURL, const TArray<TSharedPtrTS<const FDashMPD_BaseURLType>>& BaseURLs, const FString& InElementURL)
+	bool BuildAbsoluteElementURL(FString& OutURL, FTimeValue& ATO, TMediaOptionalValue<bool>& bATOComplete, const FString& DocumentURL, const TArray<TSharedPtrTS<const FDashMPD_BaseURLType>>& BaseURLs, const FString& InElementURL)
 	{
 		FTimeValue SumOfUrlATOs(FTimeValue::GetZero());
 		ATO.SetToZero();
+		bATOComplete.Reset();
 		FURL_RFC3986 UrlParser;
 		FString ElementURL(InElementURL);
 		// If the element URL is empty it is specified as the first entry in the BaseURL array.
@@ -595,6 +603,17 @@ namespace DASHUrlHelpers
 					{
 						UrlParser.ResolveAgainst(BaseURLs[nBase]->GetURL());
 						SumOfUrlATOs += BaseURLs[nBase]->GetAvailabilityTimeOffset();
+						if (!bATOComplete.IsSet())
+						{
+							bATOComplete = BaseURLs[nBase]->GetAvailabilityTimeComplete();
+						}
+						else
+						{
+							if (BaseURLs[nBase]->GetAvailabilityTimeComplete().IsSet() && BaseURLs[nBase]->GetAvailabilityTimeComplete().Value() != bATOComplete.Value())
+							{
+								// Inconsistent availabilityTimeComplete across the hierarchy.
+							}
+						}
 						++nBase;
 					}
 					else
@@ -871,6 +890,7 @@ FErrorDetail FManifestDASHInternal::PrepareRemoteElementLoadRequest(TArray<TWeak
 		FString URL;
 		FString RequestHeader;
 		FTimeValue UrlATO;
+		TMediaOptionalValue<bool> bATOComplete;
 		// Get the xlink:href
 		FString XlinkHRef = xl.GetHref();
 		if (!XlinkHRef.Equals(XLinkResolveToZero))
@@ -883,7 +903,7 @@ FErrorDetail FManifestDASHInternal::PrepareRemoteElementLoadRequest(TArray<TWeak
 			{
 				TArray<TSharedPtrTS<const FDashMPD_BaseURLType>> OutBaseURLs;
 				DASHUrlHelpers::GetAllHierarchyBaseURLs(PlayerSessionServices, OutBaseURLs, XElem->GetParentElement(), PreferredServiceLocation);
-				if (!DASHUrlHelpers::BuildAbsoluteElementURL(URL, UrlATO, MPDRoot->GetDocumentURL(), OutBaseURLs, XlinkHRef))
+				if (!DASHUrlHelpers::BuildAbsoluteElementURL(URL, UrlATO, bATOComplete, MPDRoot->GetDocumentURL(), OutBaseURLs, XlinkHRef))
 				{
 					// Not resolving to an absolute URL is very unlikely as we had to load the MPD itself from somewhere.
 					return CreateErrorAndLog(PlayerSessionServices, FString::Printf(TEXT("xlink:href did not resolve to an absolute URL")), ERRCODE_DASH_MPD_BUILDER_URL_FAILED_TO_RESOLVE);
@@ -989,7 +1009,8 @@ void FManifestDASHInternal::TransformIntoEpicEvent()
 						FTimeValue spd = FTimeValue().SetFromTimeFraction(FTimeFraction().SetFromFloatString(Params[1]));
 						MPDRoot->SetSuggestedPresentationDelay(spd);
 					}
-					else
+					// In case there is a @suggestedPresentationDelay present (even if should not) we use it, otherwise we use the @minBufferTime
+					else if (!MPDRoot->GetSuggestedPresentationDelay().IsValid())
 					{
 						MPDRoot->SetSuggestedPresentationDelay(MPDRoot->GetMinBufferTime());
 					}
@@ -1074,6 +1095,7 @@ FErrorDetail FManifestDASHInternal::BuildAfterInitialRemoteElementDownload()
 	FErrorDetail Error;
 
 	FTimeValue AST = GetAnchorTime();
+	SegmentFetchDelay.SetToZero();
 
 	bool bWarnedPresentationDuration = false;
 	// Go over the periods one at a time as XLINK attributes could bring in additional periods.
@@ -1249,6 +1271,49 @@ FErrorDetail FManifestDASHInternal::BuildAfterInitialRemoteElementDownload()
 		Periods.Emplace(MoveTemp(p));
 	}
 
+
+	// Check for supported <ServiceDescription> elements on MPD level. Those on Period level we ignore for now.
+	TArray<TSharedPtrTS<FDashMPD_ServiceDescriptionType>> ServiceDescs = MPDRoot->GetServiceDescriptions().FilterByPredicate([](const TSharedPtrTS<FDashMPD_ServiceDescriptionType>& InDesc)
+	{
+		// No <Scope> means this <ServiceDescription> applies to all clients, not specific ones.
+		// If we need to look for specific scopes we can do this here.
+		return InDesc->GetScopes().Num() == 0;
+	});
+	// Are there any indicating low-latency playback?
+	TArray<TSharedPtrTS<FDashMPD_ServiceDescriptionType>> LowLatencyDescs = ServiceDescs.FilterByPredicate([](const TSharedPtrTS<FDashMPD_ServiceDescriptionType>& InDesc)
+	{
+		return InDesc->GetLatencies().Num() != 0;
+	});
+	if (LowLatencyDescs.Num())
+	{
+		// If there is more than one we just use the first.
+		if (LowLatencyDescs.Num() > 1)
+		{
+			LogMessage(PlayerSessionServices, IInfoLog::ELevel::Warning, FString::Printf(TEXT("MPD has more than one applicable low-latency <ServiceDescription> element. Using the first one.")));
+		}
+		TSharedPtrTS<FDashMPD_ServiceDescriptionType> llDesc = LowLatencyDescs[0];
+		if (llDesc->GetLatencies().Num() > 1)
+		{
+			LogMessage(PlayerSessionServices, IInfoLog::ELevel::Warning, FString::Printf(TEXT("Selected MPD <ServiceDescription> element has more than one <Latency> element. Using the first one.")));
+		}
+		if (llDesc->GetPlaybackRates().Num() > 1)
+		{
+			LogMessage(PlayerSessionServices, IInfoLog::ELevel::Warning, FString::Printf(TEXT("Selected MPD <ServiceDescription> element has more than one <PlaybackRate> element. Using the first one.")));
+		}
+
+		TSharedPtrTS<FDashMPD_LatencyType> llType = llDesc->GetLatencies()[0];
+		TSharedPtrTS<FDashMPD_PlaybackRateType> llRate = llDesc->GetPlaybackRates().Num() ? llDesc->GetPlaybackRates()[0] : nullptr;
+
+		LowLatencyDescriptor = MakeSharedTS<FLowLatencyDescriptor>();
+		LowLatencyDescriptor->Latency.ReferenceID = llType->GetReferenceID().IsSet() ? (int64)llType->GetReferenceID().Value() : -1;
+		LowLatencyDescriptor->Latency.Target = llType->GetTarget().IsSet() ? FTimeValue((int64) llType->GetTarget().Value(), (uint32)1000) : FTimeValue();
+		LowLatencyDescriptor->Latency.Min = llType->GetMin().IsSet() ? FTimeValue((int64) llType->GetMin().Value(), (uint32)1000) : FTimeValue();
+		LowLatencyDescriptor->Latency.Max = llType->GetMax().IsSet() ? FTimeValue((int64) llType->GetMax().Value(), (uint32)1000) : FTimeValue();
+		LowLatencyDescriptor->PlayRate.Min = llRate.IsValid() ? llRate->GetMin() : FTimeValue();
+		LowLatencyDescriptor->PlayRate.Max = llRate.IsValid() ? llRate->GetMax() : FTimeValue();
+	}
+
+
 	// As per ISO/IEC 23009-1:2019/DAM 1 "Change 2: Event Stream and Timed Metadata Processing" Section A.11.11 Detailed processing
 	// all events from MPD EventStreams are to be collected right now.
 	for(int32 i=0; i<Periods.Num(); ++i)
@@ -1272,6 +1337,9 @@ void FManifestDASHInternal::PreparePeriodAdaptationSets(TSharedPtrTS<FPeriod> Pe
 			return;
 		}
 
+		check(MPDRoot.IsValid());
+		const TArray<TSharedPtrTS<FDashMPD_DescriptorType>>& UTCTimings = MPDRoot->GetUTCTimings();
+
 		SCOPE_CYCLE_COUNTER(STAT_ElectraPlayer_FManifestDASHInternal_Build);
 		CSV_SCOPED_TIMING_STAT(ElectraPlayer, FManifestDASHInternal_Build);
 
@@ -1486,6 +1554,13 @@ void FManifestDASHInternal::PreparePeriodAdaptationSets(TSharedPtrTS<FPeriod> Pe
 					Representation->bIsSideloadedSubtitle = true;
 				}
 
+				// Check for thumbnails
+				if (MimeType.Equals(ThumbnailMimeType_Jpeg) || MimeType.Equals(ThumbnailMimeType_Png))
+				{
+					LogMessage(PlayerSessionServices, IInfoLog::ELevel::Verbose, FString::Printf(TEXT("Thumbnail representations are not supported, ignoring this Representation.")));
+					continue;
+				}
+
 				if (MPDCodecs.Num() == 0)
 				{
 					LogMessage(PlayerSessionServices, IInfoLog::ELevel::Info, FString::Printf(TEXT("Neither @codecs found on Representation or AdaptationSet level, ignoring this Representation.")));
@@ -1659,6 +1734,135 @@ void FManifestDASHInternal::PreparePeriodAdaptationSets(TSharedPtrTS<FPeriod> Pe
 					continue;
 				}
 
+
+				/*
+					Check if anywhere in the hierarchy (Representation->AdaptationSet->Period) there is a <SegmentTemplate> and if there is whether
+					it has @availabilityTimeComplete set to false, indicating this segment is available via chunked transfer encoding (CTE) and thus
+					potentially usable for low-latency streaming.
+				*/
+				TArray<TSharedPtrTS<FDashMPD_SegmentTemplateType>> SegmentTemplate({MPDRepresentation->GetSegmentTemplate(), MPDAdaptationSet->GetSegmentTemplate(), MPDPeriod->GetSegmentTemplate()});
+				SegmentTemplate.Remove(nullptr);
+				if (SegmentTemplate.Num())
+				{
+					TMediaOptionalValue<bool> bTimeComplete;
+					TArray<TSharedPtrTS<const FDashMPD_BaseURLType>> BaseURLElements;
+					DASHUrlHelpers::GetAllHierarchyBaseURLs(PlayerSessionServices, BaseURLElements, MPDRepresentation, nullptr);
+					for(auto &segTemp : SegmentTemplate)
+					{
+						if (!bTimeComplete.IsSet())
+						{
+							bTimeComplete = segTemp->GetAvailabilityTimeComplete();
+						}
+						else if (segTemp->GetAvailabilityTimeComplete().IsSet() && segTemp->GetAvailabilityTimeComplete().Value() != bTimeComplete.Value())
+						{
+							// Inconsistent. Emit warning and assume 'true' so segment is not available for low-latency.
+							if (!Representation->bWarnedAboutInconsistentAvailabilityTimeComplete)
+							{
+								Representation->bWarnedAboutInconsistentAvailabilityTimeComplete = true;
+								LogMessage(PlayerSessionServices, IInfoLog::ELevel::Warning, FString::Printf(TEXT("Inconsistent @availabilityTimeComplete value across <SegmentTemplate> hierarchy for representation \"%s\" in Period \"%s\", ignoring."), *MPDRepresentation->GetID(), *Period->GetID()));
+							}
+							bTimeComplete.Reset();
+							break;
+						}
+					}
+					for(auto &baseUrl : BaseURLElements)
+					{
+						if (!bTimeComplete.IsSet())
+						{
+							bTimeComplete = baseUrl->GetAvailabilityTimeComplete();
+						}
+						else if (baseUrl->GetAvailabilityTimeComplete().IsSet() && baseUrl->GetAvailabilityTimeComplete().Value() != bTimeComplete.Value())
+						{
+							// Inconsistent. Emit warning and assume 'true' so segment is not available for low-latency.
+							if (!Representation->bWarnedAboutInconsistentAvailabilityTimeComplete)
+							{
+								Representation->bWarnedAboutInconsistentAvailabilityTimeComplete = true;
+								LogMessage(PlayerSessionServices, IInfoLog::ELevel::Warning, FString::Printf(TEXT("Inconsistent @availabilityTimeComplete value across <SegmentTemplate> / <BaseURL> hierarchy for representation \"%s\" in Period \"%s\", ignoring."), *MPDRepresentation->GetID(), *Period->GetID()));
+							}
+							bTimeComplete.Reset();
+							break;
+						}
+					}
+					// Set as potentially low-latency usable?
+					if (bTimeComplete.IsSet() && bTimeComplete.Value() == false)
+					{
+						Representation->bAvailableAsLowLatency.Set(true);
+					}
+				}
+
+
+				// Any producer reference times?
+				TArray<TSharedPtrTS<FDashMPD_ProducerReferenceTimeType>> ProducerReferenceTimes = MPDRepresentation->GetProducerReferenceTimes();
+				if (ProducerReferenceTimes.Num() == 0)
+				{
+					ProducerReferenceTimes = MPDAdaptationSet->GetProducerReferenceTimes();
+				}
+				for(auto &ProdRef : ProducerReferenceTimes)
+				{
+					if (ProdRef->GetType().Equals(TEXT("encoder")) || ProdRef->GetType().Equals(TEXT("captured")))
+					{
+						FProducerReferenceTimeInfo prti;
+						prti.ID = ProdRef->GetID();
+						prti.bInband = ProdRef->GetInband();
+						prti.Type = ProdRef->GetType().Equals(TEXT("encoder")) ? FProducerReferenceTimeInfo::EType::Encoder : FProducerReferenceTimeInfo::EType::Captured;
+						prti.PresentationTime = ProdRef->GetPresentationTime();
+
+						// As per 5.12.2, if a timing element is specified then it must also be specified in the MPD.
+						FTimeValue WallclockTime;
+						if (ProdRef->GetUTCTiming().IsValid())
+						{
+							bool bFound = false;
+							for(auto &MPDUTCTiming : UTCTimings)
+							{
+								// Compare only the scheme, not the value. The scheme is needed to interpret the format the wallClockTime is
+								// specified in correctly and nothing else.
+								if (MPDUTCTiming->GetSchemeIdUri().Equals(ProdRef->GetUTCTiming()->GetSchemeIdUri()))
+								{
+									bFound = true;
+									if ((ProdRef->GetUTCTiming()->GetSchemeIdUri().Equals(DASH::Schemes::TimingSources::Scheme_urn_mpeg_dash_utc_httpiso2014)) ||
+										(ProdRef->GetUTCTiming()->GetSchemeIdUri().Equals(DASH::Schemes::TimingSources::Scheme_urn_mpeg_dash_utc_direct2014)))
+									{
+										ISO8601::ParseDateTime(WallclockTime, ProdRef->GetWallclockTime());
+									}
+									else if (ProdRef->GetUTCTiming()->GetSchemeIdUri().Equals(DASH::Schemes::TimingSources::Scheme_urn_mpeg_dash_utc_httpxsdate2014))
+									{
+										if (!ISO8601::ParseDateTime(WallclockTime, ProdRef->GetWallclockTime()))
+										{
+											UnixEpoch::ParseFloatString(WallclockTime, ProdRef->GetWallclockTime());
+										}
+									}
+									else if (ProdRef->GetUTCTiming()->GetSchemeIdUri().Equals(DASH::Schemes::TimingSources::Scheme_urn_mpeg_dash_utc_httphead2014))
+									{
+										RFC7231::ParseDateTime(WallclockTime, ProdRef->GetWallclockTime());
+									}
+									else
+									{
+										// Unsupported type. No need to emit a warning since the MPD's root <UTCTiming> element will have warned already.
+										// Just ignore it.
+									}
+									break;
+								}
+							}
+							if (!bFound)
+							{
+								LogMessage(PlayerSessionServices, IInfoLog::ELevel::Warning, FString::Printf(TEXT("Representation <ProducerReferenceTime> references a <UTCTiming> element not present in the MPD! Ignoring this element!")));
+							}
+						}
+						prti.WallclockTime = WallclockTime;
+						if (WallclockTime.IsValid())
+						{
+							// Add the parsed element to the map for fast access when handling <Latency> in the <ServiceDescription> element.
+							ProducerReferenceTimeElements.Add(prti.ID, MakeSharedTS<FProducerReferenceTimeInfo>(prti));
+							// Assign it to the representation.
+							Representation->ProducerReferenceTimeInfos.Emplace(MoveTemp(prti));
+						}
+					}
+					else
+					{
+						LogMessage(PlayerSessionServices, IInfoLog::ELevel::Verbose, FString::Printf(TEXT("Representation uses <ProducerReferenceTime> of unsupported type, ignoring.")));
+					}
+				}
+
 				// For all intents and purposes we consider this Representation as usable now.
 				Representation->bIsUsable = true;
 				AdaptationSet->Representations.Emplace(Representation);
@@ -1693,6 +1897,33 @@ void FManifestDASHInternal::PreparePeriodAdaptationSets(TSharedPtrTS<FPeriod> Pe
 					E.Value->QualityIndex = CurrentQualityIndex;
 				}
 				AdaptationSet->bIsUsable = true;
+
+				TMediaOptionalValue<bool> lowLatencyUsable;
+				bool bFirst = true;
+				for(auto &Repr : AdaptationSet->Representations)
+				{
+					if (bFirst)
+					{
+						bFirst = false;
+						lowLatencyUsable = Repr->bAvailableAsLowLatency;
+					}
+					if ((lowLatencyUsable.IsSet() && Repr->bAvailableAsLowLatency.IsSet() && lowLatencyUsable.Value() != Repr->bAvailableAsLowLatency.Value()) ||
+						(lowLatencyUsable.IsSet() != Repr->bAvailableAsLowLatency.IsSet()))
+					{
+						// Inconsistent.
+						if (!AdaptationSet->bWarnedAboutInconsistentAvailabilityTimeComplete)
+						{
+							AdaptationSet->bWarnedAboutInconsistentAvailabilityTimeComplete = true;
+							LogMessage(PlayerSessionServices, IInfoLog::ELevel::Warning, FString::Printf(TEXT("Inconsistent @availabilityTimeComplete value across the Representations in AdaptationSet \"%s\" in Period \"%s\", ignoring."), *MPDAdaptationSet->GetID_AsStr(), *Period->GetID()));
+						}
+					}
+				}
+				// If all representations are flagged as usable for low-latency then tag the adaptation set as well.
+				if (!AdaptationSet->bWarnedAboutInconsistentAvailabilityTimeComplete && lowLatencyUsable.GetWithDefault(false))
+				{
+					AdaptationSet->bAvailableAsLowLatency.Set(true);
+				}
+
 				Period->AdaptationSets.Emplace(AdaptationSet);
 			}
 		}
@@ -2113,40 +2344,43 @@ FTimeRange FManifestDASHInternal::GetTotalTimeRange() const
 	return TotalTimeRange;
 }
 
+FTimeValue FManifestDASHInternal::GetDesiredLiveLatency() const
+{
+	return CalculateDistanceToLiveEdge();
+}
+
 
 FTimeValue FManifestDASHInternal::CalculateDistanceToLiveEdge() const
 {
+	if (CalculatedLiveDistance.IsValid())
+	{
+		return CalculatedLiveDistance;
+	}
+
 	// Check if there is a user provided value. If there is it takes precedence over everything else.
 	FTimeValue Distance = PlayerSessionServices->GetOptions().GetValue(OptionKeyLiveSeekableEndOffset).SafeGetTimeValue(FTimeValue());
+
+	// If there is a low latency descriptor we use the target latency from it.
+	if (LowLatencyDescriptor.IsValid())
+	{
+		Distance = LowLatencyDescriptor->Latency.Target;
+	}
+
 	// If not set use the MPD@suggestedPresentationDelay
 	if (!Distance.IsValid())
 	{
 		Distance = MPDRoot->GetSuggestedPresentationDelay();
 
 		/*
-			FIXME:
-				This fudge is here to prevent the use of the live segment.
-				At the moment the HTTP module only allows us access to the data once the segment has been fetched in its entirety,
-				which tends to be too late, especially if the current segment is only used partially (we may have a 2s segment
-				fetched but start playback 1s into is, giving us only 1s worth of usable data).
-				For real low-latency playback chunked transfer is needed with access to the partially downloaded data.
-
-				If the @suggestedPresentationDelay is set equal to or shorter than the @minBufferTime, which we assume to indicate the duration
-				of a single segment we would need low-latency.
-				Since we cannot force the MPD@suggestedPresentationDelay to be set such that none of this will be a problem
-				we adjust that value to twice the MPD@minBufferTime until we can get chunked transfer and partial data access.
+			If the @suggestedPresentationDelay is set equal to or shorter than the @minBufferTime there is a
+			conflict. Since we have to buffer @minBufferTime worth of content, which could be arriving in real-time only
+			then how could we maintain the @suggestedPresentationDelay?
 		*/
-			FTimeValue mbt_times2 = MPDRoot->GetMinBufferTime() * 2;
-			if (Distance < mbt_times2)
-			{
-				if (!bWarnedAboutTooSmallSuggestedPresentationDelay)
-				{
-					bWarnedAboutTooSmallSuggestedPresentationDelay = true;
-					LogMessage(PlayerSessionServices, IInfoLog::ELevel::Info, FString::Printf(TEXT("Adjusting the MPD@suggestedPresentationDelay from %.3f seconds to %.3f to avoid Live edge buffering issues"), Distance.GetAsSeconds(), mbt_times2.GetAsSeconds()));
-				}
-				Distance = mbt_times2;
-			}
-
+		if (Distance < MPDRoot->GetMinBufferTime() && !bWarnedAboutTooSmallSuggestedPresentationDelay)
+		{
+			bWarnedAboutTooSmallSuggestedPresentationDelay = true;
+			LogMessage(PlayerSessionServices, IInfoLog::ELevel::Info, FString::Printf(TEXT("MPD@suggestedPresentationDelay of %.3fs is less than MPD@minBufferTime of %.3fs"), Distance.GetAsSeconds(), MPDRoot->GetMinBufferTime().GetAsSeconds()));
+		}
 	}
 	// If not set see if there is an MPD@maxSegmentDuration and use that.
 	if (!Distance.IsValid())
@@ -2173,6 +2407,7 @@ FTimeValue FManifestDASHInternal::CalculateDistanceToLiveEdge() const
 		Distance = MPDRoot->GetTimeShiftBufferDepth();
 	}
 	check(Distance.IsValid());
+	CalculatedLiveDistance = Distance;
 	return Distance;
 }
 
@@ -2478,7 +2713,7 @@ FTimeRange FManifestDASHInternal::GetPlayTimesFromURI() const
 					// 'now' is dynamic. The time will continue to flow between here where we set the value and
 					// the moment playback will begin with buffered data.
 					// We do not lock 'now' with the current time but leave it unset. This results in the start
-					// time to be the 
+					// time to be taken from the seekable range's end value which is updating dynamically.
 					FromTo.Start.SetToInvalid();
 				}
 			}
diff --git a/Engine/Plugins/Media/ElectraPlayer/Source/ElectraPlayerRuntime/Private/Runtime/Player/DASH/ManifestDASH.cpp b/Engine/Plugins/Media/ElectraPlayer/Source/ElectraPlayerRuntime/Private/Runtime/Player/DASH/ManifestDASH.cpp
--- a/Engine/Plugins/Media/ElectraPlayer/Source/ElectraPlayerRuntime/Private/Runtime/Player/DASH/ManifestDASH.cpp
+++ b/Engine/Plugins/Media/ElectraPlayer/Source/ElectraPlayerRuntime/Private/Runtime/Player/DASH/ManifestDASH.cpp
@@ -90,7 +90,7 @@ namespace DashUtils
 /***************************************************************************************************************************************************/
 /***************************************************************************************************************************************************/
 
-class FDASHPlayPeriod : public IManifest::IPlayPeriod
+class FDASHPlayPeriod : public IManifest::IPlayPeriod, public TSharedFromThis<FDASHPlayPeriod, ESPMode::ThreadSafe>
 {
 public:
 	FDASHPlayPeriod(IPlayerSessionServices* InPlayerSessionServices, const FString& SelectedPeriodID)
@@ -110,23 +110,24 @@ public:
 	//----------------------------------------------
 	// Methods from IManifest::IPlayPeriod
 	//
-	virtual void SetStreamPreferences(EStreamType ForStreamType, const FStreamSelectionAttributes& StreamAttributes) override;
-	virtual EReadyState GetReadyState() override;
-	virtual void Load() override;
-	virtual void PrepareForPlay() override;
-	virtual int64 GetDefaultStartingBitrate() const override;
-	virtual TSharedPtrTS<FBufferSourceInfo> GetSelectedStreamBufferSourceInfo(EStreamType StreamType) override;
-	virtual FString GetSelectedAdaptationSetID(EStreamType StreamType) override;
-	virtual ETrackChangeResult ChangeTrackStreamPreference(EStreamType ForStreamType, const FStreamSelectionAttributes& StreamAttributes) override;
-	virtual TSharedPtrTS<ITimelineMediaAsset> GetMediaAsset() const override;
-	virtual void SelectStream(const FString& AdaptationSetID, const FString& RepresentationID) override;
-	virtual IManifest::FResult GetStartingSegment(TSharedPtrTS<IStreamSegment>& OutSegment, const FPlayerSequenceState& InSequenceState, const FPlayStartPosition& StartPosition, IManifest::ESearchType SearchType) override;
-	virtual IManifest::FResult GetContinuationSegment(TSharedPtrTS<IStreamSegment>& OutSegment, EStreamType StreamType, const FPlayerSequenceState& SequenceState, const FPlayStartPosition& StartPosition, IManifest::ESearchType SearchType) override;
-	virtual IManifest::FResult GetLoopingSegment(TSharedPtrTS<IStreamSegment>& OutSegment, const FPlayerSequenceState& SequenceState, const TMultiMap<EStreamType, TSharedPtrTS<IStreamSegment>>& InFinishedSegments, const FPlayStartPosition& StartPosition, IManifest::ESearchType SearchType) override;
-	virtual IManifest::FResult GetNextSegment(TSharedPtrTS<IStreamSegment>& OutSegment, TSharedPtrTS<const IStreamSegment> CurrentSegment, const FPlayStartOptions& Options) override;
-	virtual IManifest::FResult GetRetrySegment(TSharedPtrTS<IStreamSegment>& OutSegment, TSharedPtrTS<const IStreamSegment> CurrentSegment, const FPlayStartOptions& Options, bool bReplaceWithFillerData) override;
-	virtual void IncreaseSegmentFetchDelay(const FTimeValue& IncreaseAmount) override;
-	virtual void GetSegmentInformation(TArray<FSegmentInformation>& OutSegmentInformation, FTimeValue& OutAverageSegmentDuration, TSharedPtrTS<const IStreamSegment> CurrentSegment, const FTimeValue& LookAheadTime, const FString& AdaptationSetID, const FString& RepresentationID) override;
+	void SetStreamPreferences(EStreamType ForStreamType, const FStreamSelectionAttributes& StreamAttributes) override;
+	EReadyState GetReadyState() override;
+	void Load() override;
+	void PrepareForPlay() override;
+	int64 GetDefaultStartingBitrate() const override;
+	TSharedPtrTS<FBufferSourceInfo> GetSelectedStreamBufferSourceInfo(EStreamType StreamType) override;
+	FString GetSelectedAdaptationSetID(EStreamType StreamType) override;
+	ETrackChangeResult ChangeTrackStreamPreference(EStreamType ForStreamType, const FStreamSelectionAttributes& StreamAttributes) override;
+	TSharedPtrTS<ITimelineMediaAsset> GetMediaAsset() const override;
+	void SelectStream(const FString& AdaptationSetID, const FString& RepresentationID) override;
+	void TriggerInitSegmentPreload(const TArray<FInitSegmentPreload>& InitSegmentsToPreload) override;
+	IManifest::FResult GetStartingSegment(TSharedPtrTS<IStreamSegment>& OutSegment, const FPlayerSequenceState& InSequenceState, const FPlayStartPosition& StartPosition, IManifest::ESearchType SearchType) override;
+	IManifest::FResult GetContinuationSegment(TSharedPtrTS<IStreamSegment>& OutSegment, EStreamType StreamType, const FPlayerSequenceState& SequenceState, const FPlayStartPosition& StartPosition, IManifest::ESearchType SearchType) override;
+	IManifest::FResult GetLoopingSegment(TSharedPtrTS<IStreamSegment>& OutSegment, const FPlayerSequenceState& SequenceState, const TMultiMap<EStreamType, TSharedPtrTS<IStreamSegment>>& InFinishedSegments, const FPlayStartPosition& StartPosition, IManifest::ESearchType SearchType) override;
+	IManifest::FResult GetNextSegment(TSharedPtrTS<IStreamSegment>& OutSegment, TSharedPtrTS<const IStreamSegment> CurrentSegment, const FPlayStartOptions& Options) override;
+	IManifest::FResult GetRetrySegment(TSharedPtrTS<IStreamSegment>& OutSegment, TSharedPtrTS<const IStreamSegment> CurrentSegment, const FPlayStartOptions& Options, bool bReplaceWithFillerData) override;
+	void IncreaseSegmentFetchDelay(const FTimeValue& IncreaseAmount) override;
+	void GetSegmentInformation(TArray<FSegmentInformation>& OutSegmentInformation, FTimeValue& OutAverageSegmentDuration, TSharedPtrTS<const IStreamSegment> CurrentSegment, const FTimeValue& LookAheadTime, const FString& AdaptationSetID, const FString& RepresentationID) override;
 
 private:
 	ELECTRA_IMPL_DEFAULT_ERROR_METHODS(DASHManifest);
@@ -139,11 +140,35 @@ private:
 		int32 Bitrate = 0;
 	};
 
+	struct FInitSegmentInfo
+	{
+		FString AdaptationSetID;
+		FString RepresentationSetID;
+		FManifestDASHInternal::FSegmentInformation InitSegmentInfo;
+		bool bRequested = false;
+		TSharedPtrTS<FMPDLoadRequestDASH> LoadRequest;
+		
+		class FAcceptBoxes : public IParserISO14496_12::IBoxCallback
+		{
+		public:
+			virtual ~FAcceptBoxes() = default;
+			IParserISO14496_12::IBoxCallback::EParseContinuation OnFoundBox(IParserISO14496_12::FBoxType Box, int64 BoxSizeInBytes, int64 FileDataOffset, int64 BoxDataOffset) override
+			{ return IParserISO14496_12::IBoxCallback::EParseContinuation::Continue; }
+			IParserISO14496_12::IBoxCallback::EParseContinuation OnEndOfBox(IParserISO14496_12::FBoxType Box, int64 BoxSizeInBytes, int64 FileDataOffset, int64 BoxDataOffset) override
+			{ return IParserISO14496_12::IBoxCallback::EParseContinuation::Continue; }
+		};
+	};
+
 	TSharedPtrTS<FManifestDASHInternal> GetCurrentManifest() const;
 
 	TSharedPtrTS<IPlaybackAssetRepresentation> GetRepresentationFromAdaptationByMaxBandwidth(TWeakPtrTS<IPlaybackAssetAdaptationSet> InAdaptationSet, int32 NotExceedingBandwidth);
 	TSharedPtrTS<IPlaybackAssetRepresentation> GetRepresentationFromAdaptationByPriorityAndMaxBandwidth(TWeakPtrTS<IPlaybackAssetAdaptationSet> InAdaptationSet, int32 NotExceedingBandwidth, EStreamType StreamType);
 
+	void GetRepresentationInitSegmentsFromAdaptation(TArray<FInitSegmentInfo>& OutInitSegInfos, TWeakPtrTS<IPlaybackAssetAdaptationSet> InAdaptationSet);
+	void MergeRepresentationInitSegments(TArray<FInitSegmentInfo>& InOutInitSegInfos, const TArray<FInitSegmentInfo>& NewInitSegInfos);
+	void HandleRepresentationInitSegmentLoading(const TArray<FInitSegmentPreload>& InitSegmentsToPreload);
+	void InitSegmentDownloadComplete(TSharedPtrTS<FMPDLoadRequestDASH> LoadRequest, bool bSuccess);
+
 	enum class ENextSegType
 	{
 		SamePeriodNext,
@@ -167,8 +192,6 @@ private:
 	FStreamSelectionAttributes AudioStreamPreferences;
 	FStreamSelectionAttributes SubtitleStreamPreferences;
 
-	FTimeValue SegmentFetchDelay = FTimeValue::GetZero();
-
 	FString PeriodID;
 
 	FString ActiveVideoAdaptationSetID;
@@ -184,6 +207,9 @@ private:
 	TSharedPtrTS<FBufferSourceInfo> SourceBufferInfoSubtitles;
 
 	TSharedPtr<ElectraCDM::IMediaCDMClient, ESPMode::ThreadSafe> DrmClient;
+
+
+	TArray<FInitSegmentInfo> VideoInitSegmentInfos;
 };
 
 /***************************************************************************************************************************************************/
@@ -222,6 +248,12 @@ IManifest::EType FManifestDASH::GetPresentationType() const
 	return IManifest::EType::OnDemand;
 }
 
+TSharedPtrTS<const FLowLatencyDescriptor> FManifestDASH::GetLowLatencyDescriptor() const
+{ 
+	TSharedPtrTS<FManifestDASHInternal> Manifest(CurrentManifest);
+	return Manifest.IsValid() ? Manifest->GetLowLatencyDescriptor() : nullptr;
+}
+
 FTimeValue FManifestDASH::GetAnchorTime() const
 {
 	TSharedPtrTS<FManifestDASHInternal> Manifest(CurrentManifest);
@@ -287,6 +319,20 @@ FTimeValue FManifestDASH::GetMinBufferTime() const
 	return FTimeValue();
 }
 
+FTimeValue FManifestDASH::GetDesiredLiveLatency() const
+{
+	TSharedPtrTS<FManifestDASHInternal> Manifest(CurrentManifest);
+	return Manifest.IsValid() ? Manifest->GetDesiredLiveLatency() : FTimeValue();
+}
+
+
+TSharedPtrTS<IProducerReferenceTimeInfo> FManifestDASH::GetProducerReferenceTimeInfo(int64 ID) const
+{
+	TSharedPtrTS<FManifestDASHInternal> Manifest(CurrentManifest);
+	return Manifest.IsValid() ? Manifest->GetProducerReferenceTimeElement(ID) : nullptr;
+}
+
+
 void FManifestDASH::GetTrackMetadata(TArray<FTrackMetadata>& OutMetadata, EStreamType StreamType) const
 {
 	TSharedPtrTS<FManifestDASHInternal> Manifest(CurrentManifest);
@@ -315,6 +361,13 @@ void FManifestDASH::UpdateDynamicRefetchCounter()
 	}
 }
 
+void FManifestDASH::TriggerClockSync(IManifest::EClockSyncType InClockSyncType)
+{
+	IPlaylistReaderDASH* Reader = static_cast<IPlaylistReaderDASH*>(PlayerSessionServices->GetManifestReader().Get());
+	Reader->RequestClockResync();
+}
+
+
 IStreamReader* FManifestDASH::CreateStreamReaderHandler()
 {
 	return new FStreamReaderFMP4DASH;
@@ -814,6 +867,11 @@ void FDASHPlayPeriod::PrepareForPlay()
 
 			TSharedPtrTS<IPlaybackAssetRepresentation> VideoRepr = GetRepresentationFromAdaptationByMaxBandwidth(VideoAS, (int32) StartingBitrate);
 			ActiveVideoRepresentationID = VideoRepr->GetUniqueIdentifier();
+
+			// Set up the list of initialization segments.
+			TArray<FInitSegmentInfo> InitSegInfos;
+			GetRepresentationInitSegmentsFromAdaptation(InitSegInfos, VideoAS);
+			MergeRepresentationInitSegments(VideoInitSegmentInfos, InitSegInfos);
 		}
 
 		TSharedPtrTS<FManifestDASHInternal::FAdaptationSet> AudioAS = Period->GetAdaptationSetByUniqueID(ActiveAudioAdaptationSetID);
@@ -830,6 +888,22 @@ void FDASHPlayPeriod::PrepareForPlay()
 			ActiveSubtitleRepresentationID = SubtitleRepr->GetUniqueIdentifier();
 		}
 
+
+		// If there is a low latency service description with a reference ID to a producer reference time, set it in the options.
+		// This will be retrieved from there in setting up the segment requests.
+		TSharedPtrTS<const FLowLatencyDescriptor> llDesc = Manifest->GetLowLatencyDescriptor();
+		if (llDesc.IsValid())
+		{
+			if (llDesc->Latency.ReferenceID >= 0)
+			{
+				PlayerSessionServices->GetOptions().SetOrUpdate(DASH::OptionKey_LatencyReferenceId, FVariantValue(llDesc->Latency.ReferenceID));
+			}
+			else
+			{
+				PlayerSessionServices->GetOptions().Remove(DASH::OptionKey_LatencyReferenceId);
+			}
+		}
+
 		// Emit all <EventStream> events of the period to the AEMS event handler.
 		Manifest->SendEventsFromAllPeriodEventStreams(Period);
 
@@ -950,6 +1024,12 @@ void FDASHPlayPeriod::SelectStream(const FString& AdaptationSetID, const FString
 	}
 }
 
+void FDASHPlayPeriod::TriggerInitSegmentPreload(const TArray<FInitSegmentPreload>& InitSegmentsToPreload)
+{
+	HandleRepresentationInitSegmentLoading(InitSegmentsToPreload);
+}
+
+
 void FDASHPlayPeriod::SetupCommonSegmentRequestInfos(TSharedPtrTS<FStreamSegmentRequestFMP4DASH>& InOutSegmentRequest)
 {
 	FManifestDASHInternal::FRepresentation* Repr = static_cast<FManifestDASHInternal::FRepresentation*>(InOutSegmentRequest->Representation.Get());
@@ -990,6 +1070,7 @@ IManifest::FResult FDASHPlayPeriod::GetStartingSegment(TSharedPtrTS<IStreamSegme
 	FTimeValue PlayRangeEnd = StartPosition.Options.PlaybackRange.End;
 	check(PlayRangeEnd.IsValid());
 
+	FTimeValue SegmentFetchDelay = Manifest->GetSegmentFetchDelay();
 	FTimeValue AST = Manifest->GetAnchorTime();
 	FTimeValue StartTime = StartPosition.Time;
 	// All time values we communicate to the outside - and therefor get from the outside - are offset by the availabilityStartTime.
@@ -1265,6 +1346,7 @@ IManifest::FResult FDASHPlayPeriod::GetNextOrRetrySegment(TSharedPtrTS<IStreamSe
 		return IManifest::FResult(IManifest::FResult::EType::NotFound).SetErrorDetail(FErrorDetail().SetMessage("No active stream found to get next segment for"));
 	}
 
+	FTimeValue SegmentFetchDelay = Manifest->GetSegmentFetchDelay();
 	FTimeValue AST = Manifest->GetAnchorTime();
 	bool bUsesAST = Manifest->UsesAST();
 	bool bIsStaticType = Manifest->IsStaticType() || Manifest->IsDynamicEpicEvent();
@@ -1285,6 +1367,7 @@ IManifest::FResult FDASHPlayPeriod::GetNextOrRetrySegment(TSharedPtrTS<IStreamSe
 		// This is to make sure the time is sufficiently large that it won't be affected by rounding errors in timescale conversions.
 		SearchOpt.PeriodLocalTime.SetFromND(Current->Segment.Time - Current->Segment.PTO + Current->Segment.Duration*3/4, Current->Segment.Timescale);
 		SearchOpt.SearchType = IManifest::ESearchType::After;
+		bFrameAccurateSearch = false;
 		// If this representation is side loaded there is nothing to fetch next.
 		// Set the period search time such that the search will have to return EOS.
 		if (Repr.IsValid() && Repr->IsSideloadedSubtitle())
@@ -1501,8 +1584,19 @@ IManifest::FResult FDASHPlayPeriod::GetLoopingSegment(TSharedPtrTS<IStreamSegmen
 
 void FDASHPlayPeriod::IncreaseSegmentFetchDelay(const FTimeValue& IncreaseAmount)
 {
-// TODO: when this starts to get too large we should trigger a resync of our clock with the server time.
-	SegmentFetchDelay += IncreaseAmount;
+	TSharedPtrTS<FManifestDASHInternal> Manifest = GetCurrentManifest();
+	if (Manifest.IsValid())
+	{
+		FTimeValue NewFetchDelay = Manifest->GetSegmentFetchDelay() + IncreaseAmount;
+		Manifest->SetSegmentFetchDelay(NewFetchDelay);
+		// If the fetch delay becomes too large then there is possibly a clock drift.
+		// Trigger a resynchronization which will reset the delay when complete.
+		if (NewFetchDelay.GetAsSeconds() > 0.5)
+		{
+			IPlaylistReaderDASH* Reader = static_cast<IPlaylistReaderDASH*>(PlayerSessionServices->GetManifestReader().Get());
+			Reader->RequestClockResync();
+		}
+	}
 }
 
 
@@ -1617,6 +1711,130 @@ TSharedPtrTS<IPlaybackAssetRepresentation> FDASHPlayPeriod::GetRepresentationFro
 	return BestRepr;
 }
 
+void FDASHPlayPeriod::GetRepresentationInitSegmentsFromAdaptation(TArray<FInitSegmentInfo>& OutInitSegInfos, TWeakPtrTS<IPlaybackAssetAdaptationSet> InAdaptationSet)
+{
+	TSharedPtrTS<IPlaybackAssetAdaptationSet> AS = InAdaptationSet.Pin();
+	if (AS.IsValid())
+	{
+		TArray<TWeakPtrTS<FMPDLoadRequestDASH>> DummyRequests;
+		FManifestDASHInternal::FSegmentSearchOption SearchOpt;
+		SearchOpt.bInitSegmentSetupOnly = true;
+		int32 NumRepr = AS->GetNumberOfRepresentations();
+		for(int32 i=0; i<NumRepr; ++i)
+		{
+			TSharedPtrTS<FManifestDASHInternal::FRepresentation> Repr = StaticCastSharedPtr<FManifestDASHInternal::FRepresentation>(AS->GetRepresentationByIndex(i));
+			if (Repr->CanBePlayed())
+			{
+				FInitSegmentInfo SegInfo;
+				if (Repr->FindSegment(PlayerSessionServices, SegInfo.InitSegmentInfo, DummyRequests, SearchOpt) == FManifestDASHInternal::FRepresentation::ESearchResult::Found)
+				{
+					SegInfo.AdaptationSetID = AS->GetUniqueIdentifier();
+					SegInfo.RepresentationSetID = Repr->GetUniqueIdentifier();
+					OutInitSegInfos.Emplace(MoveTemp(SegInfo));
+				}
+			}
+		}
+	}
+}
+
+void FDASHPlayPeriod::MergeRepresentationInitSegments(TArray<FInitSegmentInfo>& InOutInitSegInfos, const TArray<FInitSegmentInfo>& NewInitSegInfos)
+{
+	for(auto &ni : NewInitSegInfos)
+	{
+		if (!InOutInitSegInfos.ContainsByPredicate([ni](const FInitSegmentInfo& Other) { return ni.AdaptationSetID.Equals(Other.AdaptationSetID) && ni.RepresentationSetID.Equals(Other.RepresentationSetID); }))
+		{
+			InOutInitSegInfos.Add(ni);
+		}
+	}
+}
+
+void FDASHPlayPeriod::HandleRepresentationInitSegmentLoading(const TArray<FInitSegmentPreload>& InitSegmentsToPreload)
+{
+	TArray<TWeakPtrTS<FMPDLoadRequestDASH>> RemoteElementLoadRequests;
+
+	auto GetInitSegInfo = [](TArray<FInitSegmentInfo>& InitSegList, const FString& InRepID) -> FInitSegmentInfo*
+	{
+		for(int32 i=0; i<InitSegList.Num(); ++i)
+		{
+			if (InitSegList[i].RepresentationSetID == InRepID)
+			{
+				return &InitSegList[i];
+			}
+		}
+		return nullptr;
+	};
+
+	for(auto& pl : InitSegmentsToPreload)
+	{
+		FInitSegmentInfo* InitSeg = nullptr;
+
+		if (pl.AdaptationSetID == ActiveVideoAdaptationSetID)
+		{
+			InitSeg = GetInitSegInfo(VideoInitSegmentInfos, pl.RepresentationID);
+		}
+		if (!InitSeg)
+		{
+			continue;
+		}
+		FInitSegmentInfo& is = *InitSeg;
+		if (!is.bRequested)
+		{
+			is.bRequested = true;
+
+			IPlayerEntityCache::FCacheItem CachedItem;
+			if (!PlayerSessionServices->GetEntityCache()->GetCachedEntity(CachedItem, is.InitSegmentInfo.InitializationURL.URL, is.InitSegmentInfo.InitializationURL.Range))
+			{
+				is.LoadRequest = MakeSharedTS<FMPDLoadRequestDASH>();
+				is.LoadRequest->LoadType = FMPDLoadRequestDASH::ELoadType::Segment;
+				is.LoadRequest->URL = is.InitSegmentInfo.InitializationURL.URL;
+				is.LoadRequest->Range = is.InitSegmentInfo.InitializationURL.Range;
+				if (is.InitSegmentInfo.InitializationURL.CustomHeader.Len())
+				{
+					is.LoadRequest->Headers.Emplace(HTTP::FHTTPHeader({DASH::HTTPHeaderOptionName, is.InitSegmentInfo.InitializationURL.CustomHeader}));
+				}
+				is.LoadRequest->PlayerSessionServices = PlayerSessionServices;
+				is.LoadRequest->CompleteCallback.BindThreadSafeSP(AsShared(), &FDASHPlayPeriod::InitSegmentDownloadComplete);
+				
+				RemoteElementLoadRequests.Emplace(is.LoadRequest);
+			}
+		}
+	}
+
+	TSharedPtrTS<IPlaylistReader> ManifestReader = PlayerSessionServices->GetManifestReader();
+	if (ManifestReader.IsValid())
+	{
+		IPlaylistReaderDASH* Reader = static_cast<IPlaylistReaderDASH*>(ManifestReader.Get());
+		Reader->AddElementLoadRequests(RemoteElementLoadRequests);
+	}
+}
+
+void FDASHPlayPeriod::InitSegmentDownloadComplete(TSharedPtrTS<FMPDLoadRequestDASH> LoadRequest, bool bSuccess)
+{
+	if (bSuccess)
+	{
+		FInitSegmentInfo::FAcceptBoxes AllBoxes;
+		FMP4StaticDataReader StaticDataReader;
+		StaticDataReader.SetParseData(LoadRequest->Request->GetResponseBuffer());
+		TSharedPtrTS<IParserISO14496_12> Init = IParserISO14496_12::CreateParser();
+		UEMediaError parseError = Init->ParseHeader(&StaticDataReader, &AllBoxes, PlayerSessionServices, nullptr);
+		if (parseError == UEMEDIA_ERROR_OK || parseError == UEMEDIA_ERROR_END_OF_STREAM)
+		{
+			// Parse the tracks of the init segment. We do this mainly to get to the CSD we might need should we have to insert filler data later.
+			parseError = Init->PrepareTracks(PlayerSessionServices, TSharedPtrTS<const IParserISO14496_12>());
+			if (parseError == UEMEDIA_ERROR_OK)
+			{
+				// Add this to the entity cache in case it needs to be retrieved again.
+				IPlayerEntityCache::FCacheItem CacheItem;
+				CacheItem.URL = LoadRequest->URL;
+				CacheItem.Range = LoadRequest->Range;
+				CacheItem.Parsed14496_12Data = Init;
+				PlayerSessionServices->GetEntityCache()->CacheEntity(CacheItem);
+			}
+		}
+	}
+}
+
+
 
 /***************************************************************************************************************************************************/
 /***************************************************************************************************************************************************/
@@ -1683,6 +1901,48 @@ FTimeValue CalculateSegmentAvailabilityTimeOffset(const TArray<TSharedPtrTS<FDas
 	return Sum;
 }
 
+
+
+TMediaOptionalValue<bool> GetSegmentAvailabilityTimeComplete(const TArray<TSharedPtrTS<FDashMPD_SegmentBaseType>>& Arr)
+{
+	TMediaOptionalValue<bool> TimeComplete;
+	for(int32 i=0; i<Arr.Num(); ++i)
+	{
+		if (!TimeComplete.IsSet())
+		{
+			TimeComplete = Arr[i]->GetAvailabilityTimeComplete();
+		}
+		else
+		{
+			if (Arr[i]->GetAvailabilityTimeComplete().IsSet() && Arr[i]->GetAvailabilityTimeComplete().Value() != TimeComplete.Value())
+			{
+				// Inconsistent @availabilityTimeComplete
+			}
+		}
+	}
+	return TimeComplete;
+}
+
+TMediaOptionalValue<bool> GetSegmentAvailabilityTimeComplete(const TArray<TSharedPtrTS<FDashMPD_SegmentTemplateType>>& Arr)
+{
+	TMediaOptionalValue<bool> TimeComplete;
+	for(int32 i=0; i<Arr.Num(); ++i)
+	{
+		if (!TimeComplete.IsSet())
+		{
+			TimeComplete = Arr[i]->GetAvailabilityTimeComplete();
+		}
+		else
+		{
+			if (Arr[i]->GetAvailabilityTimeComplete().IsSet() && Arr[i]->GetAvailabilityTimeComplete().Value() != TimeComplete.Value())
+			{
+				// Inconsistent @availabilityTimeComplete
+			}
+		}
+	}
+	return TimeComplete;
+}
+
 }
 
 
@@ -1745,7 +2005,8 @@ FManifestDASHInternal::FRepresentation::ESearchResult FManifestDASHInternal::FRe
 			FString DocumentURL = MPDRepresentation->GetDocumentURL();
 			FString URL, RequestHeader;
 			FTimeValue UrlATO;
-			if (!DASHUrlHelpers::BuildAbsoluteElementURL(URL, UrlATO, DocumentURL, OutBaseURLs, FString()))
+			TMediaOptionalValue<bool> bATOComplete;
+			if (!DASHUrlHelpers::BuildAbsoluteElementURL(URL, UrlATO, bATOComplete, DocumentURL, OutBaseURLs, FString()))
 			{
 				PostError(InPlayerSessionServices, FString::Printf(TEXT("Representation \"%s\" failed to resolve URL to an absolute URL!"), *MPDRepresentation->GetID()), ERRCODE_DASH_MPD_BAD_REPRESENTATION);
 				bIsUsable = false;
@@ -1824,7 +2085,8 @@ bool FManifestDASHInternal::FRepresentation::PrepareDownloadURLs(IPlayerSessionS
 	FString DocumentURL = MPDRepresentation->GetDocumentURL();
 	FString URL, RequestHeader;
 	FTimeValue UrlATO;
-	if (!DASHUrlHelpers::BuildAbsoluteElementURL(URL, UrlATO, DocumentURL, OutBaseURLs, FString()))
+	TMediaOptionalValue<bool> bATOComplete;
+	if (!DASHUrlHelpers::BuildAbsoluteElementURL(URL, UrlATO, bATOComplete, DocumentURL, OutBaseURLs, FString()))
 	{
 		PostError(InPlayerSessionServices, FString::Printf(TEXT("Representation \"%s\" failed to resolve URL to an absolute URL!"), *MPDRepresentation->GetID()), ERRCODE_DASH_MPD_BAD_REPRESENTATION);
 		return false;
@@ -1860,7 +2122,11 @@ bool FManifestDASHInternal::FRepresentation::PrepareDownloadURLs(IPlayerSessionS
 
 	// Add any availabilityTimeOffset from the <BaseURL>.
 	InOutSegmentInfo.ATO += UrlATO;
-
+	if (!InOutSegmentInfo.bAvailabilityTimeComplete.IsSet())
+	{
+		InOutSegmentInfo.bAvailabilityTimeComplete = bATOComplete;
+	}
+	InOutSegmentInfo.bLowLatencyChunkedEncodingExpected = bAvailableAsLowLatency.GetWithDefault(false);
 	return true;
 }
 
@@ -1912,7 +2178,8 @@ bool FManifestDASHInternal::FRepresentation::PrepareDownloadURLs(IPlayerSessionS
 	FString DocumentURL = MPDRepresentation->GetDocumentURL();
 	FString MediaURL, MediaRequestHeader;
 	FTimeValue MediaUrlATO;
-	if (!DASHUrlHelpers::BuildAbsoluteElementURL(MediaURL, MediaUrlATO, DocumentURL, OutBaseURLs, MediaTemplateURL))
+	TMediaOptionalValue<bool> bMediaATOComplete;
+	if (!DASHUrlHelpers::BuildAbsoluteElementURL(MediaURL, MediaUrlATO, bMediaATOComplete, DocumentURL, OutBaseURLs, MediaTemplateURL))
 	{
 		PostError(InPlayerSessionServices, FString::Printf(TEXT("Representation \"%s\" failed to resolve media segment URL to an absolute URL!"), *MPDRepresentation->GetID()), ERRCODE_DASH_MPD_BAD_REPRESENTATION);
 		return false;
@@ -1927,7 +2194,8 @@ bool FManifestDASHInternal::FRepresentation::PrepareDownloadURLs(IPlayerSessionS
 	// And also generate the absolute init segment URL
 	FString InitURL, InitRequestHeader;
 	FTimeValue InitUrlATO;
-	if (!DASHUrlHelpers::BuildAbsoluteElementURL(InitURL, InitUrlATO, DocumentURL, OutBaseURLs, InitTemplateURL))
+	TMediaOptionalValue<bool> bInitATOComplete;
+	if (!DASHUrlHelpers::BuildAbsoluteElementURL(InitURL, InitUrlATO, bInitATOComplete, DocumentURL, OutBaseURLs, InitTemplateURL))
 	{
 		PostError(InPlayerSessionServices, FString::Printf(TEXT("Representation \"%s\" failed to resolve init segment URL to an absolute URL!"), *MPDRepresentation->GetID()), ERRCODE_DASH_MPD_BAD_REPRESENTATION);
 		return false;
@@ -1957,7 +2225,12 @@ bool FManifestDASHInternal::FRepresentation::PrepareDownloadURLs(IPlayerSessionS
 	// Add any availabilityTimeOffset from the <BaseURL>. We use the media ATO since we request the init segment at the same time
 	// as the media segment and the init segment needs to be available at the same time anyway.
 	InOutSegmentInfo.ATO += MediaUrlATO;
+	if (!InOutSegmentInfo.bAvailabilityTimeComplete.IsSet())
+	{
+		InOutSegmentInfo.bAvailabilityTimeComplete = bMediaATOComplete;
+	}
 
+	InOutSegmentInfo.bLowLatencyChunkedEncodingExpected = bAvailableAsLowLatency.GetWithDefault(false);
 	return true;
 }
 
@@ -2077,6 +2350,12 @@ void FManifestDASHInternal::FRepresentation::CollectInbandEventStreams(IPlayerSe
 	}
 }
 
+void FManifestDASHInternal::FRepresentation::SetupProducerReferenceTimeInfo(IPlayerSessionServices* InPlayerSessionServices, FSegmentInformation& InOutSegmentInfo)
+{
+	InOutSegmentInfo.ProducerReferenceTimeInfos = ProducerReferenceTimeInfos;
+	DashUtils::GetPlayerOption(InPlayerSessionServices, InOutSegmentInfo.MeasureLatencyViaReferenceTimeInfoID, DASH::OptionKey_LatencyReferenceId, (int64)-1);
+}
+
 
 
 FManifestDASHInternal::FRepresentation::ESearchResult FManifestDASHInternal::FRepresentation::FindSegment(IPlayerSessionServices* InPlayerSessionServices, FSegmentInformation& OutSegmentInfo, TArray<TWeakPtrTS<FMPDLoadRequestDASH>>& OutRemoteElementLoadRequests, const FSegmentSearchOption& InSearchOptions)
@@ -2171,21 +2450,35 @@ FManifestDASHInternal::FRepresentation::ESearchResult FManifestDASHInternal::FRe
 			return FManifestDASHInternal::FRepresentation::ESearchResult::BadType;
 		}
 
-		if (SegmentBase.Num())
+		if (!InSearchOptions.bInitSegmentSetupOnly)
 		{
-			return FindSegment_Base(InPlayerSessionServices, OutSegmentInfo, OutRemoteElementLoadRequests, InSearchOptions, MPDRepresentation, SegmentBase);
+			if (SegmentBase.Num())
+			{
+				return FindSegment_Base(InPlayerSessionServices, OutSegmentInfo, OutRemoteElementLoadRequests, InSearchOptions, MPDRepresentation, SegmentBase);
+			}
+			else
+			{
+				// Get the segment timeline, if one is used.
+				TSharedPtrTS<FDashMPD_SegmentTimelineType> SegmentTimeline = GET_ATTR(SegmentTemplate, GetSegmentTimeline(), IsValid(), TSharedPtrTS<FDashMPD_SegmentTimelineType>());
+				if (SegmentTimeline.IsValid())
+				{
+					return FindSegment_Timeline(InPlayerSessionServices, OutSegmentInfo, OutRemoteElementLoadRequests, InSearchOptions, MPDRepresentation, SegmentTemplate, SegmentTimeline);
+				}
+				else
+				{
+					return FindSegment_Template(InPlayerSessionServices, OutSegmentInfo, OutRemoteElementLoadRequests, InSearchOptions, MPDRepresentation, SegmentTemplate);
+				}
+			}
 		}
 		else
 		{
-			// Get the segment timeline, if one is used.
-			TSharedPtrTS<FDashMPD_SegmentTimelineType> SegmentTimeline = GET_ATTR(SegmentTemplate, GetSegmentTimeline(), IsValid(), TSharedPtrTS<FDashMPD_SegmentTimelineType>());
-			if (SegmentTimeline.IsValid())
+			if (SegmentBase.Num())
 			{
-				return FindSegment_Timeline(InPlayerSessionServices, OutSegmentInfo, OutRemoteElementLoadRequests, InSearchOptions, MPDRepresentation, SegmentTemplate, SegmentTimeline);
+				return PrepareDownloadURLs(InPlayerSessionServices, OutSegmentInfo, SegmentBase) ? FManifestDASHInternal::FRepresentation::ESearchResult::Found : FManifestDASHInternal::FRepresentation::ESearchResult::BadType;
 			}
 			else
 			{
-				return FindSegment_Template(InPlayerSessionServices, OutSegmentInfo, OutRemoteElementLoadRequests, InSearchOptions, MPDRepresentation, SegmentTemplate);
+				return PrepareDownloadURLs(InPlayerSessionServices, OutSegmentInfo, SegmentTemplate) ? FManifestDASHInternal::FRepresentation::ESearchResult::Found : FManifestDASHInternal::FRepresentation::ESearchResult::BadType;
 			}
 		}
 	}
@@ -2226,6 +2519,8 @@ FManifestDASHInternal::FRepresentation::ESearchResult FManifestDASHInternal::FRe
 		LogMessage(InPlayerSessionServices, IInfoLog::ELevel::Warning, FString::Printf(TEXT("Representation timescale (%u) in MPD is not equal to timescale used in the segment index (%u) for Representation \"%s\"."), MPDTimescale, SidxTimescale, *MPDRepresentation->GetID()));
 	}
 	FTimeValue ATO = CalculateSegmentAvailabilityTimeOffset(SegmentBase);
+	TMediaOptionalValue<bool> bATOComplete = GetSegmentAvailabilityTimeComplete(SegmentBase);
+
 
 	// Convert the local media search time to the timescale of the segment index.
 	// The PTO (presentation time offset) which maps the internal media time to the zero point of the period must be included as well.
@@ -2389,13 +2684,15 @@ FManifestDASHInternal::FRepresentation::ESearchResult FManifestDASHInternal::FRe
 		OutSegmentInfo.Number = CurrentN;
 		OutSegmentInfo.NumberOfBytes = Sidx->GetEntry(CurrentN).Size;
 		OutSegmentInfo.FirstByteOffset = Sidx->GetFirstOffset() + SegmentIndexRangeStart + SegmentIndexRangeSize + CurrentOffset;
-		OutSegmentInfo.MediaLocalFirstAUTime = MediaLocalSearchTime;
+		OutSegmentInfo.MediaLocalFirstAUTime = OutSegmentInfo.MediaLocalFirstPTS = MediaLocalSearchTime;
 		OutSegmentInfo.MediaLocalLastAUTime = MediaLocalEndTime;
 		OutSegmentInfo.Timescale = SidxTimescale;
 		OutSegmentInfo.ATO = ATO;
+		OutSegmentInfo.bAvailabilityTimeComplete = bATOComplete;
 		OutSegmentInfo.bIsLastInPeriod = CurrentT + CurrentD >= MediaLocalEndTime;
 		OutSegmentInfo.bFrameAccuracyRequired = InSearchOptions.bFrameAccurateSearch;
 		CollectInbandEventStreams(InPlayerSessionServices, OutSegmentInfo);
+		SetupProducerReferenceTimeInfo(InPlayerSessionServices, OutSegmentInfo);
 		if (!PrepareDownloadURLs(InPlayerSessionServices, OutSegmentInfo, SegmentBase))
 		{
 			bIsUsable = false;
@@ -2421,7 +2718,7 @@ FManifestDASHInternal::FRepresentation::ESearchResult FManifestDASHInternal::FRe
 	TMediaOptionalValue<uint32> Duration = GET_ATTR(SegmentTemplate, GetDuration(), IsSet(), TMediaOptionalValue<uint32>());
 	TMediaOptionalValue<int32> EptDelta = GET_ATTR(SegmentTemplate, GetEptDelta(), IsSet(), TMediaOptionalValue<int32>());
 	FTimeValue ATO = CalculateSegmentAvailabilityTimeOffset(SegmentTemplate);
-
+	TMediaOptionalValue<bool> bATOComplete = GetSegmentAvailabilityTimeComplete(SegmentTemplate);
 
 	// The timescale should in all likelihood not be 1. While certainly allowed an accuracy of only one second is more likely to be
 	// an oversight when building the MPD.
@@ -2536,14 +2833,16 @@ FManifestDASHInternal::FRepresentation::ESearchResult FManifestDASHInternal::FRe
 	OutSegmentInfo.EPTdelta = EPTdelta;
 	OutSegmentInfo.Duration = SegmentDuration;
 	OutSegmentInfo.Number = StartNumber + SegmentNum;
-	OutSegmentInfo.MediaLocalFirstAUTime = MediaLocalSearchTime;
+	OutSegmentInfo.MediaLocalFirstAUTime = OutSegmentInfo.MediaLocalFirstPTS = MediaLocalSearchTime;
 	OutSegmentInfo.MediaLocalLastAUTime = MediaLocalEndTime;
 	OutSegmentInfo.Timescale = MPDTimescale;
 	OutSegmentInfo.bMayBeMissing = SegmentNum + 1 >= MaxSegmentsInPeriod;
 	OutSegmentInfo.bIsLastInPeriod = OutSegmentInfo.bMayBeMissing && InSearchOptions.bHasFollowingPeriod;
 	OutSegmentInfo.bFrameAccuracyRequired = InSearchOptions.bFrameAccurateSearch;
 	OutSegmentInfo.ATO = ATO;
+	OutSegmentInfo.bAvailabilityTimeComplete = bATOComplete;
 	CollectInbandEventStreams(InPlayerSessionServices, OutSegmentInfo);
+	SetupProducerReferenceTimeInfo(InPlayerSessionServices, OutSegmentInfo);
 	if (!PrepareDownloadURLs(InPlayerSessionServices, OutSegmentInfo, SegmentTemplate))
 	{
 		bIsUsable = false;
@@ -2578,6 +2877,7 @@ FManifestDASHInternal::FRepresentation::ESearchResult FManifestDASHInternal::FRe
 	uint32 EndNumber = GET_ATTR(SegmentTemplate, GetEndNumber(), IsSet(), TMediaOptionalValue<uint32>(~0U)).Value();
 	//TMediaOptionalValue<uint32> Duration = GET_ATTR(SegmentTemplate, GetDuration(), IsSet(), TMediaOptionalValue<uint32>());
 	FTimeValue ATO = CalculateSegmentAvailabilityTimeOffset(SegmentTemplate);
+	TMediaOptionalValue<bool> bATOComplete = GetSegmentAvailabilityTimeComplete(SegmentTemplate);
 
 	// The timescale should in all likelihood not be 1. While certainly allowed an accuracy of only one second is more likely to be
 	// an oversight when building the MPD.
@@ -2861,7 +3161,7 @@ FManifestDASHInternal::FRepresentation::ESearchResult FManifestDASHInternal::FRe
 		OutSegmentInfo.PTO = PTO;
 		OutSegmentInfo.Duration = CurrentD;
 		OutSegmentInfo.Number = CurrentN;
-		OutSegmentInfo.MediaLocalFirstAUTime = MediaLocalSearchTime;
+		OutSegmentInfo.MediaLocalFirstAUTime = OutSegmentInfo.MediaLocalFirstPTS = MediaLocalSearchTime;
 		OutSegmentInfo.MediaLocalLastAUTime = MediaLocalEndTime;
 		OutSegmentInfo.Timescale = MPDTimescale;
 		OutSegmentInfo.bMayBeMissing = CurrentT + CurrentD >= MediaLocalEndTime;
@@ -2874,7 +3174,9 @@ FManifestDASHInternal::FRepresentation::ESearchResult FManifestDASHInternal::FRe
 			LogMessage(InPlayerSessionServices, IInfoLog::ELevel::Warning, FString::Printf(TEXT("Representation \"%s\" <SegmentTimeline> gap encountered for needed 't' value of %lld. Replacing with an empty filler segment."), *MPDRepresentation->GetID(), (long long int)CurrentT));
 		}
 		OutSegmentInfo.ATO = ATO;
+		OutSegmentInfo.bAvailabilityTimeComplete = bATOComplete;
 		CollectInbandEventStreams(InPlayerSessionServices, OutSegmentInfo);
+		SetupProducerReferenceTimeInfo(InPlayerSessionServices, OutSegmentInfo);
 		if (!PrepareDownloadURLs(InPlayerSessionServices, OutSegmentInfo, SegmentTemplate))
 		{
 			bIsUsable = false;
@@ -2934,10 +3236,11 @@ FManifestDASHInternal::FRepresentation::ESearchResult FManifestDASHInternal::FRe
 	OutSegmentInfo.Number = 0;
 	OutSegmentInfo.NumberOfBytes = 0;
 	OutSegmentInfo.FirstByteOffset = 0;
-	OutSegmentInfo.MediaLocalFirstAUTime = MediaLocalSearchTime;
+	OutSegmentInfo.MediaLocalFirstAUTime = OutSegmentInfo.MediaLocalFirstPTS = MediaLocalSearchTime;
 	OutSegmentInfo.MediaLocalLastAUTime = MediaLocalPeriodEnd;
 	OutSegmentInfo.Timescale = Timescale;
 	OutSegmentInfo.ATO.SetToZero();
+	OutSegmentInfo.bAvailabilityTimeComplete.Set(true);
 	OutSegmentInfo.bIsSideload = true;
 	OutSegmentInfo.bIsLastInPeriod = true;
 	TArray<TSharedPtrTS<FDashMPD_SegmentBaseType>> SegmentBase;
diff --git a/Engine/Plugins/Media/ElectraPlayer/Source/ElectraPlayerRuntime/Private/Runtime/Player/DASH/ManifestDASH.h b/Engine/Plugins/Media/ElectraPlayer/Source/ElectraPlayerRuntime/Private/Runtime/Player/DASH/ManifestDASH.h
--- a/Engine/Plugins/Media/ElectraPlayer/Source/ElectraPlayerRuntime/Private/Runtime/Player/DASH/ManifestDASH.h
+++ b/Engine/Plugins/Media/ElectraPlayer/Source/ElectraPlayerRuntime/Private/Runtime/Player/DASH/ManifestDASH.h
@@ -18,22 +18,26 @@ public:
 	void UpdateInternalManifest(TSharedPtrTS<FManifestDASHInternal> UpdatedManifest);
 
 	virtual ~FManifestDASH();
-	virtual EType GetPresentationType() const override;
-	virtual FTimeValue GetAnchorTime() const override;
-	virtual FTimeRange GetTotalTimeRange() const override;
-	virtual FTimeRange GetSeekableTimeRange() const override;
-	virtual FTimeRange GetPlaybackRange() const override;
-	virtual void GetSeekablePositions(TArray<FTimespan>& OutPositions) const override;
-	virtual FTimeValue GetDuration() const override;
-	virtual FTimeValue GetDefaultStartTime() const override;
-	virtual void ClearDefaultStartTime() override;
-	virtual FTimeValue GetMinBufferTime() const override;
-	virtual void GetTrackMetadata(TArray<FTrackMetadata>& OutMetadata, EStreamType StreamType) const override;
-	virtual void UpdateDynamicRefetchCounter() override;
-	virtual IStreamReader* CreateStreamReaderHandler() override;
-
-	virtual FResult FindPlayPeriod(TSharedPtrTS<IPlayPeriod>& OutPlayPeriod, const FPlayStartPosition& StartPosition, ESearchType SearchType) override;
-	virtual FResult FindNextPlayPeriod(TSharedPtrTS<IPlayPeriod>& OutPlayPeriod, TSharedPtrTS<const IStreamSegment> CurrentSegment) override;
+	EType GetPresentationType() const override;
+	TSharedPtrTS<const FLowLatencyDescriptor> GetLowLatencyDescriptor() const override;
+	FTimeValue GetAnchorTime() const override;
+	FTimeRange GetTotalTimeRange() const override;
+	FTimeRange GetSeekableTimeRange() const override;
+	FTimeRange GetPlaybackRange() const override;
+	void GetSeekablePositions(TArray<FTimespan>& OutPositions) const override;
+	FTimeValue GetDuration() const override;
+	FTimeValue GetDefaultStartTime() const override;
+	void ClearDefaultStartTime() override;
+	FTimeValue GetMinBufferTime() const override;
+	FTimeValue GetDesiredLiveLatency() const override;
+	TSharedPtrTS<IProducerReferenceTimeInfo> GetProducerReferenceTimeInfo(int64 ID) const override;
+	void GetTrackMetadata(TArray<FTrackMetadata>& OutMetadata, EStreamType StreamType) const override;
+	void UpdateDynamicRefetchCounter() override;
+	void TriggerClockSync(IManifest::EClockSyncType InClockSyncType) override;
+	IStreamReader* CreateStreamReaderHandler() override;
+
+	FResult FindPlayPeriod(TSharedPtrTS<IPlayPeriod>& OutPlayPeriod, const FPlayStartPosition& StartPosition, ESearchType SearchType) override;
+	FResult FindNextPlayPeriod(TSharedPtrTS<IPlayPeriod>& OutPlayPeriod, TSharedPtrTS<const IStreamSegment> CurrentSegment) override;
 
 private:
 	ELECTRA_IMPL_DEFAULT_ERROR_METHODS(DASHManifest);
diff --git a/Engine/Plugins/Media/ElectraPlayer/Source/ElectraPlayerRuntime/Private/Runtime/Player/DASH/ManifestParserDASH.cpp b/Engine/Plugins/Media/ElectraPlayer/Source/ElectraPlayerRuntime/Private/Runtime/Player/DASH/ManifestParserDASH.cpp
--- a/Engine/Plugins/Media/ElectraPlayer/Source/ElectraPlayerRuntime/Private/Runtime/Player/DASH/ManifestParserDASH.cpp
+++ b/Engine/Plugins/Media/ElectraPlayer/Source/ElectraPlayerRuntime/Private/Runtime/Player/DASH/ManifestParserDASH.cpp
@@ -101,6 +101,7 @@ DEF_ELEMENT(PlaybackRate);
 DEF_ELEMENT(OperatingQuality);
 DEF_ELEMENT(OperatingBandwidth);
 DEF_ELEMENT(ServiceDescription);
+DEF_ELEMENT(Resync);
 
 const TCHAR* const Element_up_UrlQueryInfo = TEXT("up:UrlQueryInfo");
 const TCHAR* const Element_up_ExtUrlQueryInfo = TEXT("up:ExtUrlQueryInfo");
@@ -225,6 +226,12 @@ DEF_ATTR(queryString);
 DEF_ATTR(includeInRequests);
 DEF_ATTR(headerParamSource);
 DEF_ATTR(sameOriginOnly);
+DEF_ATTR(dT);
+DEF_ATTR(dImin);
+DEF_ATTR(dImax);
+DEF_ATTR(marker);
+DEF_ATTR(rangeAccess);
+
 
 // xlink attributes. Due to the xlink namespace they have a colon in them
 const TCHAR* const Attr_xlink_href = TEXT("xlink:href");
@@ -270,6 +277,8 @@ namespace
 	static bool ParseUIntAttribute(TMediaOptionalValue<uint32>& OutValue, FManifestParserDASH* Builder, const TCHAR* AttributeName, const TCHAR* AttributeValue, int32 XmlLineNumber);
 	static bool ParseDoubleAttribute(FTimeFraction& OutValue, FManifestParserDASH* Builder, const TCHAR* AttributeName, const TCHAR* AttributeValue, int32 XmlLineNumber);
 	static bool ParseDoubleAttribute(FTimeValue& OutValue, FManifestParserDASH* Builder, const TCHAR* AttributeName, const TCHAR* AttributeValue, int32 XmlLineNumber);
+	static bool ParseDoubleAttribute(double& OutValue, FManifestParserDASH* Builder, const TCHAR* AttributeName, const TCHAR* AttributeValue, int32 XmlLineNumber);
+	static bool ParseDoubleAttribute(TMediaOptionalValue<double>& OutValue, FManifestParserDASH* Builder, const TCHAR* AttributeName, const TCHAR* AttributeValue, int32 XmlLineNumber);
 	static bool ParseDuration(FTimeValue& OutValue, FManifestParserDASH* Builder, const TCHAR* AttributeName, const TCHAR* AttributeValue, int32 XmlLineNumber);
 	static bool ParseDateTime(FTimeValue& OutValue, FManifestParserDASH* Builder, const TCHAR* AttributeName, const TCHAR* AttributeValue, int32 XmlLineNumber);
 	static bool ParseRatioAttribute(FTimeFraction& OutValue, FManifestParserDASH* Builder, const TCHAR* AttributeName, const TCHAR* AttributeValue, int32 XmlLineNumber);
@@ -1576,6 +1585,11 @@ bool FDashMPD_RepresentationBaseType::ProcessElement(FManifestParserDASH* Builde
 		ContentPopularityRates.Emplace(MakeSharedTS<FDashMPD_ContentPopularityRateType>(ElementName, ElementData));
 		return Builder->PushNewElement(ContentPopularityRates.Last(), XmlFileLineNumber);
 	}
+	else if (StringHelpers::StringEquals(ElementName, ELEMENT(Resync)))
+	{
+		Resyncs.Emplace(MakeSharedTS<FDashMPD_ResyncType>(ElementName, ElementData));
+		return Builder->PushNewElement(Resyncs.Last(), XmlFileLineNumber);
+	}
 	else
 	{
 		return Super::ProcessElement(Builder, ElementName, ElementData, XmlFileLineNumber);
@@ -2512,7 +2526,7 @@ bool FDashMPD_ServiceDescriptionType::ProcessElement(FManifestParserDASH* Builde
 		OperatingQualities.Emplace(MakeSharedTS<FDashMPD_OperatingQualityType>(ElementName, ElementData));
 		return Builder->PushNewElement(OperatingQualities.Last(), XmlFileLineNumber);
 	}
-	else if (StringHelpers::StringEquals(ElementName, ELEMENT(Scope)))
+	else if (StringHelpers::StringEquals(ElementName, ELEMENT(OperatingBandwidth)))
 	{
 		OperatingBandwidths.Emplace(MakeSharedTS<FDashMPD_OperatingBandwidthType>(ElementName, ElementData));
 		return Builder->PushNewElement(OperatingBandwidths.Last(), XmlFileLineNumber);
@@ -2537,6 +2551,49 @@ bool FDashMPD_ServiceDescriptionType::ProcessAttribute(FManifestParserDASH* Buil
 
 /*********************************************************************************************************************/
 
+bool FDashMPD_ResyncType::ProcessElement(FManifestParserDASH* Builder, const TCHAR* ElementName, const TCHAR* ElementData, int32 XmlFileLineNumber)
+{
+	return NoChildElementAllowed(Builder, ElementName, ElementData, XmlFileLineNumber);
+}
+
+bool FDashMPD_ResyncType::ProcessAttribute(FManifestParserDASH* Builder, const TCHAR* AttributeName, const TCHAR* AttributeValue)
+{
+	if (AttrEquals(type))
+	{
+		return ParseIntAttribute(Type, Builder, AttributeName, AttributeValue, GetXMLLineNumber());
+	}
+	else if (AttrEquals(dT))
+	{
+		return ParseUIntAttribute(dT, Builder, AttributeName, AttributeValue, GetXMLLineNumber());
+	}
+	else if (AttrEquals(dImax))
+	{
+		return ParseDoubleAttribute(dIMax, Builder, AttributeName, AttributeValue, GetXMLLineNumber());
+	}
+	else if (AttrEquals(dImin))
+	{
+		return ParseDoubleAttribute(dIMin, Builder, AttributeName, AttributeValue, GetXMLLineNumber());
+	}
+	else if (AttrEquals(marker))
+	{
+		return ParseBooleanAttribute(bMarker, Builder, AttributeName, AttributeValue, GetXMLLineNumber());
+	}
+	else if (AttrEquals(rangeAccess))
+	{
+		return ParseBooleanAttribute(bRangeAccess, Builder, AttributeName, AttributeValue, GetXMLLineNumber());
+	}
+	else if (AttrEquals(index))
+	{
+		return ParseStringAttribute(Index, Builder, AttributeName, AttributeValue, GetXMLLineNumber());
+	}
+	else
+	{
+		return Super::ProcessAttribute(Builder, AttributeName, AttributeValue);
+	}
+}
+
+/*********************************************************************************************************************/
+
 bool FDashMPD_RepresentationType::ProcessElement(FManifestParserDASH* Builder, const TCHAR* ElementName, const TCHAR* ElementData, int32 XmlFileLineNumber)
 {
 	if (StringHelpers::StringEquals(ElementName, ELEMENT(SegmentBase)))
@@ -3146,6 +3203,23 @@ static bool ParseDoubleAttribute(FTimeValue& OutValue, FManifestParserDASH* Buil
 	return false;
 }
 
+static bool ParseDoubleAttribute(double& OutValue, FManifestParserDASH* Builder, const TCHAR* AttributeName, const TCHAR* AttributeValue, int32 XmlLineNumber)
+{
+	double v;
+	LexFromString(v, AttributeValue);
+	OutValue = v;
+	return true;
+}
+
+static bool ParseDoubleAttribute(TMediaOptionalValue<double>& OutValue, FManifestParserDASH* Builder, const TCHAR* AttributeName, const TCHAR* AttributeValue, int32 XmlLineNumber)
+{
+	double v;
+	LexFromString(v, AttributeValue);
+	OutValue.Set(v);
+	return true;
+}
+
+
 static bool ParseConditionalUIntAttribute(IDashMPDElement::FBoolUInt64& OutValue, FManifestParserDASH* Builder, const TCHAR* AttributeName, const TCHAR* AttributeValue, int32 XmlLineNumber)
 {
 	if (ParseBooleanAttribute(OutValue.bBoolean, Builder, AttributeName, AttributeValue, XmlLineNumber))
diff --git a/Engine/Plugins/Media/ElectraPlayer/Source/ElectraPlayerRuntime/Private/Runtime/Player/DASH/OptionKeynamesDASH.h b/Engine/Plugins/Media/ElectraPlayer/Source/ElectraPlayerRuntime/Private/Runtime/Player/DASH/OptionKeynamesDASH.h
--- a/Engine/Plugins/Media/ElectraPlayer/Source/ElectraPlayerRuntime/Private/Runtime/Player/DASH/OptionKeynamesDASH.h
+++ b/Engine/Plugins/Media/ElectraPlayer/Source/ElectraPlayerRuntime/Private/Runtime/Player/DASH/OptionKeynamesDASH.h
@@ -23,6 +23,7 @@ const TCHAR* const OptionKey_MinTimeBetweenMPDUpdates = TEXT("dash:min_mpd_updat
 
 const TCHAR* const OptionKey_CurrentCDN = TEXT("dash:current_cdn");			//!< A string giving the currently selected CDN to be used with <BaseURL@serviceLocation> elements.
 
+const TCHAR* const OptionKey_LatencyReferenceId = TEXT("dash:latency_refId");	//!< The Id of the <ProducerReferenceTime> element used by <ServiceDescription><Latency>
 
 const TCHAR* const HTTPHeaderOptionName = TEXT("MPEG-DASH-Param");			//!< HTTP request header name carrying optional DASH specific parameters.
 
