ElectraPlayer: Added DASH low latency Live streaming

diff --git a/Engine/Plugins/Media/ElectraPlayer/Source/ElectraPlayerRuntime/Private/Runtime/Player/DASH/PlaylistReaderDASH.cpp b/Engine/Plugins/Media/ElectraPlayer/Source/ElectraPlayerRuntime/Private/Runtime/Player/DASH/PlaylistReaderDASH.cpp
--- a/Engine/Plugins/Media/ElectraPlayer/Source/ElectraPlayerRuntime/Private/Runtime/Player/DASH/PlaylistReaderDASH.cpp
+++ b/Engine/Plugins/Media/ElectraPlayer/Source/ElectraPlayerRuntime/Private/Runtime/Player/DASH/PlaylistReaderDASH.cpp
@@ -52,26 +52,27 @@ public:
 	void Initialize(IPlayerSessionServices* PlayerSessionServices);
 
 	virtual ~FPlaylistReaderDASH();
-	virtual void Close() override;
-	virtual void HandleOnce() override;
-	virtual const FString& GetPlaylistType() const override
+	void Close() override;
+	void HandleOnce() override;
+	const FString& GetPlaylistType() const override
 	{
 		static FString Type("dash");
 		return Type;
 	}
-	virtual void LoadAndParse(const FString& URL) override;
-	virtual FString GetURL() const override;
-	virtual TSharedPtrTS<IManifest> GetManifest() override;
-	virtual void AddElementLoadRequests(const TArray<TWeakPtrTS<FMPDLoadRequestDASH>>& RemoteElementLoadRequests) override;
-	virtual void RequestMPDUpdate(EMPDRequestType InRequestType) override;
-	virtual TSharedPtrTS<FManifestDASHInternal> GetCurrentMPD() override
+	void LoadAndParse(const FString& URL) override;
+	FString GetURL() const override;
+	TSharedPtrTS<IManifest> GetManifest() override;
+	void AddElementLoadRequests(const TArray<TWeakPtrTS<FMPDLoadRequestDASH>>& RemoteElementLoadRequests) override;
+	void RequestMPDUpdate(EMPDRequestType InRequestType) override;
+	void RequestClockResync() override;
+	TSharedPtrTS<FManifestDASHInternal> GetCurrentMPD() override
 	{
 		return Manifest;
 	}
-	virtual void SetStreamInbandEventUsage(EStreamType InStreamType, bool bUsesInbandDASHEvents) override;
+	void SetStreamInbandEventUsage(EStreamType InStreamType, bool bUsesInbandDASHEvents) override;
 
 
-	virtual void OnMediaPlayerEventReceived(TSharedPtrTS<IAdaptiveStreamingPlayerAEMSEvent> InEvent, IAdaptiveStreamingPlayerAEMSReceiver::EDispatchMode InDispatchMode) override;
+	void OnMediaPlayerEventReceived(TSharedPtrTS<IAdaptiveStreamingPlayerAEMSEvent> InEvent, IAdaptiveStreamingPlayerAEMSReceiver::EDispatchMode InDispatchMode) override;
 
 private:
 	using FResourceLoadRequestPtr = TSharedPtrTS<FMPDLoadRequestDASH>;
@@ -109,6 +110,8 @@ private:
 	void EnqueueInitialXLinkRequests();
 
 	void ManifestDownloadCompleted(FResourceLoadRequestPtr Request, bool bSuccess);
+	void FinishManifestBuildingAfterTimesync(bool bGotTheTime);
+	void PrepareNextTimeSync();
 	void ManifestUpdateDownloadCompleted(FResourceLoadRequestPtr Request, bool bSuccess);
 	void InitialMPDXLinkElementDownloadCompleted(FResourceLoadRequestPtr Request, bool bSuccess);
 
@@ -153,10 +156,28 @@ private:
 	TSharedPtrTS<FManifestDASH>								PlayerManifest;
 
 	// Time synchronization.
+	struct FInitialTimesync
+	{
+		FTimeValue UtcTimeThen;
+		FTimeValue HttpRequestStart;
+		FString HttpDateHeader;
+		FString Utc_direct2014;
+		TUniquePtr<HTTP::FConnectionInfo> ConnInfo;
+		void Reset()
+		{
+			UtcTimeThen.SetToInvalid();
+			HttpRequestStart.SetToInvalid();
+			HttpDateHeader.Reset();
+			Utc_direct2014.Reset();
+			ConnInfo.Reset();
+		}
+	};
+	FInitialTimesync										InitialTimesync;
+	bool													bIsInitialTimesync = true;
 	TArray<TSharedPtrTS<FDashMPD_DescriptorType>>			AttemptedTimesyncDescriptors;
 	FTimeValue												NextTimeSyncTime;
 	bool													bTimeSyncInProgress = false;
-	bool													bInitialTimeSyncedToMPDDateTimeHeader = false;
+
 
 	// Warnings
 	bool													bWarnedAboutNoUTCTimingElements = false;
@@ -230,20 +251,19 @@ void FPlaylistReaderDASH::InternalHandleOnce()
 	{
 		SCOPE_CYCLE_COUNTER(STAT_ElectraPlayer_DASH_PlaylistWorker);
 		CSV_SCOPED_TIMING_STAT(ElectraPlayer, DASH_PlaylistWorker);
+
 		FTimeValue Now = PlayerSessionServices->GetSynchronizedUTCTime()->GetTime();
+		// Check for time sync first. The time sync result and thus the clock gets changed in HandleCompletedRequest()
+		// which may trigger a jump in time, making "Now" and NextTimeSyncTime obsolete until the next loop.
+		if (NextTimeSyncTime.IsValid() && Now >= NextTimeSyncTime)
+		{
+			TriggerTimeSynchronization();
+		}
 
 		HandleCompletedRequests(Now);
-
 		ExecutePendingRequests(Now);
-
 		// Check if the MPD must be updated.
 		CheckForMPDUpdate();
-
-		// Time sync
-		if (NextTimeSyncTime.IsValid() && Now >= NextTimeSyncTime)
-		{
-			TriggerTimeSynchronization();
-		}
 	}
 }
 
@@ -490,7 +510,8 @@ void FPlaylistReaderDASH::CheckForMPDUpdate()
 			TArray<TSharedPtrTS<const FDashMPD_BaseURLType>> OutBaseURLs;
 			FString URL, RequestHeader;
 			FTimeValue UrlATO;
-			DASHUrlHelpers::BuildAbsoluteElementURL(URL, UrlATO, MPDRoot->GetDocumentURL(), OutBaseURLs, NewMPDLocation);
+			TMediaOptionalValue<bool> bATOComplete;
+			DASHUrlHelpers::BuildAbsoluteElementURL(URL, UrlATO, bATOComplete, MPDRoot->GetDocumentURL(), OutBaseURLs, NewMPDLocation);
 
 			// The URL query might need to be changed. Look for the UrlQuery properties.
 			TArray<TSharedPtrTS<FDashMPD_UrlQueryInfoType>> UrlQueries;
@@ -607,6 +628,12 @@ void FPlaylistReaderDASH::RequestMPDUpdate(IPlaylistReaderDASH::EMPDRequestType
 	UpdateRequested.Emplace(MoveTemp(Request));
 }
 
+void FPlaylistReaderDASH::RequestClockResync()
+{
+	NextTimeSyncTime.SetToZero();
+}
+
+
 void FPlaylistReaderDASH::SetStreamInbandEventUsage(EStreamType InStreamType, bool bUsesInbandDASHEvents)
 {
 	switch(InStreamType)
@@ -633,139 +660,131 @@ void FPlaylistReaderDASH::TriggerTimeSynchronization()
 	// We're doing a sync now. Clear the next time.
 	NextTimeSyncTime.SetToInvalid();
 
-	// Time sync is only necessary when there is either an MPD@availabilityStartTime or MPD@type is dynamic.
-	if (Manifest.IsValid() && (Manifest->UsesAST() || !Manifest->IsStaticType()))
+	if (Manifest.IsValid())
 	{
-		TSharedPtrTS<FDashMPD_MPDType> MPDRoot = Manifest->GetMPDRoot();
-		if (MPDRoot.IsValid())
+		// Time sync is only necessary when there is either an MPD@availabilityStartTime or MPD@type is dynamic.
+		if (Manifest->UsesAST() || !Manifest->IsStaticType())
 		{
-			bool bFoundSupportedElement = false;
-			FString MPDDocumentURL = MPDRoot->GetDocumentURL();
-			const TArray<TSharedPtrTS<FDashMPD_DescriptorType>>& UTCTimings = MPDRoot->GetUTCTimings();
-
-			// Emit a one-time only warning if there is no <UTCTiming> element in the MPD when there should be.
-			if (UTCTimings.Num() == 0 && !bWarnedAboutNoUTCTimingElements)
+			TSharedPtrTS<FDashMPD_MPDType> MPDRoot = Manifest->GetMPDRoot();
+			if (MPDRoot.IsValid())
 			{
-				bWarnedAboutNoUTCTimingElements = true;
-				LogMessage(IInfoLog::ELevel::Warning, FString::Printf(TEXT("No <UTCTiming> element found in MPD. This could result in playback failures. For detailed validation of your MPD please use https://conformance.dashif.org/")));
-			}
+				bool bFoundSupportedElement = false;
+				FString MPDDocumentURL = MPDRoot->GetDocumentURL();
+				const TArray<TSharedPtrTS<FDashMPD_DescriptorType>>& UTCTimings = MPDRoot->GetUTCTimings();
 
-			/*
-				We go over the UTC timing elements in two passes.
-				In the first we look only for the (discouraged) urn:mpeg:dash:utc:direct:2014 scheme because
-				if it is available we can set it immediately. Its value is baked into the manifest and is therefore
-				only really valid right at the time the MPD was fetched.
-				If we try other schemes first that involve a network request and those fail then there will have
-				passed that much time that the hardcoded value in the MPD will no longer make sense to associate
-				with the current system time.
-			*/
-			for(int32 i=0; i<UTCTimings.Num(); ++i)
-			{
-				if (UTCTimings[i]->GetSchemeIdUri().Equals(DASH::Schemes::TimingSources::Scheme_urn_mpeg_dash_utc_direct2014))
+				// Emit a one-time only warning if there is no <UTCTiming> element in the MPD when there should be.
+				if (UTCTimings.Num() == 0 && !bWarnedAboutNoUTCTimingElements)
 				{
-					FTimeValue DirectTime;
-					if (ISO8601::ParseDateTime(DirectTime, UTCTimings[i]->GetValue()))
+					bWarnedAboutNoUTCTimingElements = true;
+					LogMessage(IInfoLog::ELevel::Warning, FString::Printf(TEXT("No <UTCTiming> element found in MPD. This could result in playback failures. For detailed validation of your MPD please use https://conformance.dashif.org/")));
+				}
+
+				/*
+					We go over the UTC timing elements in two passes.
+					In the first we look only for the (discouraged) urn:mpeg:dash:utc:direct:2014 scheme because
+					if it is available we can set it immediately. Its value is baked into the manifest and is therefore
+					only really valid right at the time the MPD was fetched.
+					If we try other schemes first that involve a network request and those fail then there will have
+					passed that much time that the hardcoded value in the MPD will no longer make sense to associate
+					with the current system time.
+				*/
+				for(int32 i=0; i<UTCTimings.Num(); ++i)
+				{
+					if (UTCTimings[i]->GetSchemeIdUri().Equals(DASH::Schemes::TimingSources::Scheme_urn_mpeg_dash_utc_direct2014))
 					{
-						PlayerSessionServices->GetSynchronizedUTCTime()->SetTime(DirectTime);
+						InitialTimesync.Utc_direct2014 = UTCTimings[i]->GetValue();
 						// Remove this timing element. It can not be used indefinitely and we best get rid of it now.
 						MPDRoot->RemoveUTCTimingElement(UTCTimings[i]);
 						--i;
+						bFoundSupportedElement = true;
 					}
-					bFoundSupportedElement = true;
 				}
-			}
-			// 2nd pass. Use the more appropriate timing elements.
-			for(int32 i=0; i<UTCTimings.Num(); ++i)
-			{
-				auto GetRequestURL = [MPDDocumentURL](FString ListOfURLs, int32 Index) -> FString
+				// 2nd pass. Use the more appropriate timing elements.
+				for(int32 i=0; i<UTCTimings.Num(); ++i)
 				{
-					TArray<FString> URLs;
-					ListOfURLs.ParseIntoArrayWS(URLs);
-					if (Index < URLs.Num())
+					auto GetRequestURL = [MPDDocumentURL](FString ListOfURLs, int32 Index) -> FString
 					{
-						TArray<TSharedPtrTS<const FDashMPD_BaseURLType>> OutBaseURLs;
-						FString URL, RequestHeader;
-						FTimeValue UrlATO;
-						DASHUrlHelpers::BuildAbsoluteElementURL(URL, UrlATO, MPDDocumentURL, OutBaseURLs, URLs[Index]);
-						return URL;
-					}
-					return FString();
-				};
+						TArray<FString> URLs;
+						ListOfURLs.ParseIntoArrayWS(URLs);
+						if (Index < URLs.Num())
+						{
+							TArray<TSharedPtrTS<const FDashMPD_BaseURLType>> OutBaseURLs;
+							FString URL, RequestHeader;
+							FTimeValue UrlATO;
+							TMediaOptionalValue<bool> bATOComplete;
+							DASHUrlHelpers::BuildAbsoluteElementURL(URL, UrlATO, bATOComplete, MPDDocumentURL, OutBaseURLs, URLs[Index]);
+							return URL;
+						}
+						return FString();
+					};
 
 
-				const FString& SchemeIdUri = UTCTimings[i]->GetSchemeIdUri();
-				const FString& Value = UTCTimings[i]->GetValue();
+					const FString& SchemeIdUri = UTCTimings[i]->GetSchemeIdUri();
+					const FString& Value = UTCTimings[i]->GetValue();
 
-				TSharedPtrTS<FMPDLoadRequestDASH> TimeSyncRequest;
-				// Check for supported schemes. They appear in the manifest in the order the author has prioritized them.
-				if (SchemeIdUri.Equals(DASH::Schemes::TimingSources::Scheme_urn_mpeg_dash_utc_httpxsdate2014))
-				{
-					TimeSyncRequest = MakeSharedTS<FMPDLoadRequestDASH>();
-					TimeSyncRequest->LoadType = FMPDLoadRequestDASH::ELoadType::TimeSync;
-					TimeSyncRequest->URL = GetRequestURL(Value, 0);
-					TimeSyncRequest->CompleteCallback.BindThreadSafeSP(AsShared(), &FPlaylistReaderDASH::Timesync_httpxsdate_Completed);
-				}
-				else if (SchemeIdUri.Equals(DASH::Schemes::TimingSources::Scheme_urn_mpeg_dash_utc_httpiso2014))
-				{
-					TimeSyncRequest = MakeSharedTS<FMPDLoadRequestDASH>();
-					TimeSyncRequest->LoadType = FMPDLoadRequestDASH::ELoadType::TimeSync;
-					TimeSyncRequest->URL = GetRequestURL(Value, 0);
-					TimeSyncRequest->CompleteCallback.BindThreadSafeSP(AsShared(), &FPlaylistReaderDASH::Timesync_httpiso_Completed);
-				}
-				else if (SchemeIdUri.Equals(DASH::Schemes::TimingSources::Scheme_urn_mpeg_dash_utc_httphead2014))
-				{
-					TimeSyncRequest = MakeSharedTS<FMPDLoadRequestDASH>();
-					TimeSyncRequest->LoadType = FMPDLoadRequestDASH::ELoadType::TimeSync;
-					TimeSyncRequest->URL = GetRequestURL(Value, 0);
-					TimeSyncRequest->Verb = TEXT("HEAD");
-					TimeSyncRequest->CompleteCallback.BindThreadSafeSP(AsShared(), &FPlaylistReaderDASH::Timesync_httphead_Completed);
-				}
-				else
-				{
-					// Not supported.
-				}
+					TSharedPtrTS<FMPDLoadRequestDASH> TimeSyncRequest;
+					// Check for supported schemes. They appear in the manifest in the order the author has prioritized them.
+					if (SchemeIdUri.Equals(DASH::Schemes::TimingSources::Scheme_urn_mpeg_dash_utc_httpxsdate2014))
+					{
+						TimeSyncRequest = MakeSharedTS<FMPDLoadRequestDASH>();
+						TimeSyncRequest->LoadType = FMPDLoadRequestDASH::ELoadType::TimeSync;
+						TimeSyncRequest->URL = GetRequestURL(Value, 0);
+						TimeSyncRequest->CompleteCallback.BindThreadSafeSP(AsShared(), &FPlaylistReaderDASH::Timesync_httpxsdate_Completed);
+					}
+					else if (SchemeIdUri.Equals(DASH::Schemes::TimingSources::Scheme_urn_mpeg_dash_utc_httpiso2014))
+					{
+						TimeSyncRequest = MakeSharedTS<FMPDLoadRequestDASH>();
+						TimeSyncRequest->LoadType = FMPDLoadRequestDASH::ELoadType::TimeSync;
+						TimeSyncRequest->URL = GetRequestURL(Value, 0);
+						TimeSyncRequest->CompleteCallback.BindThreadSafeSP(AsShared(), &FPlaylistReaderDASH::Timesync_httpiso_Completed);
+					}
+					else if (SchemeIdUri.Equals(DASH::Schemes::TimingSources::Scheme_urn_mpeg_dash_utc_httphead2014))
+					{
+						TimeSyncRequest = MakeSharedTS<FMPDLoadRequestDASH>();
+						TimeSyncRequest->LoadType = FMPDLoadRequestDASH::ELoadType::TimeSync;
+						TimeSyncRequest->URL = GetRequestURL(Value, 0);
+						TimeSyncRequest->Verb = TEXT("HEAD");
+						TimeSyncRequest->CompleteCallback.BindThreadSafeSP(AsShared(), &FPlaylistReaderDASH::Timesync_httphead_Completed);
+					}
+					else
+					{
+						// Not supported, skip it.
+					}
 
-				// When a request has been set up, enqueue it and leave the loop.
-				if (TimeSyncRequest.IsValid())
-				{
-					bTimeSyncInProgress = true;
-					AttemptedTimesyncDescriptors.Emplace(UTCTimings[i]);
-					EnqueueResourceRequest(MoveTemp(TimeSyncRequest));
-					bFoundSupportedElement = true;
-					break;
+					// When a request has been set up, enqueue it and leave the loop.
+					if (TimeSyncRequest.IsValid())
+					{
+						bTimeSyncInProgress = true;
+						AttemptedTimesyncDescriptors.Emplace(UTCTimings[i]);
+						EnqueueResourceRequest(MoveTemp(TimeSyncRequest));
+						bFoundSupportedElement = true;
+						break;
+					}
 				}
-			}
-			// Warn once when there is no UTCTiming element we support.
-			if (!bFoundSupportedElement && !bWarnedAboutUnsupportedUTCTimingElement)
-			{
-				bWarnedAboutUnsupportedUTCTimingElement = true;
-				if (!bInitialTimeSyncedToMPDDateTimeHeader)
+				// Warn once when there is no UTCTiming element we support.
+				if (!bFoundSupportedElement && !bWarnedAboutUnsupportedUTCTimingElement)
 				{
-					LogMessage(IInfoLog::ELevel::Warning, FString::Printf(TEXT("No supported <UTCTiming> element found in MPD. This could result in playback failures.")));
+					bWarnedAboutUnsupportedUTCTimingElement = true;
+					if (InitialTimesync.HttpDateHeader.Len())
+					{
+						LogMessage(IInfoLog::ELevel::Warning, FString::Printf(TEXT("No supported <UTCTiming> element found in MPD, but time was synchronized to the MPD's HTTP response Date header. This may not be accurate enough.")));
+					}
+					else
+					{
+						LogMessage(IInfoLog::ELevel::Warning, FString::Printf(TEXT("No supported <UTCTiming> element found in MPD. This could result in playback failures.")));
+					}
 				}
-				else
+
+				// If no request was enqueued and this is the initial time sync we need to complete the manifest setup now.
+				if (!bTimeSyncInProgress)
 				{
-					LogMessage(IInfoLog::ELevel::Warning, FString::Printf(TEXT("No supported <UTCTiming> element found in MPD, but time was synchronized to the MPD's HTTP response Date header. This may not be accurate enough.")));
+					FinishManifestBuildingAfterTimesync(false);
 				}
 			}
-
-			// Do not re-sync the time too often.
-			const FTimeValue MinResyncTimeInterval(120.0);		// not sooner than every 120 seconds.
-			const FTimeValue MaxResyncTimeInterval(1800.0);	// at least once every 30 minutes.
-
-			FTimeValue ResyncTimeInterval(MinResyncTimeInterval);
-			// How frequently does the MPD update?
-			FTimeValue mup = Manifest->GetMinimumUpdatePeriod();
-			if (mup.IsValid() && mup > MinResyncTimeInterval)
-			{
-				ResyncTimeInterval = mup;
-			}
-			if (ResyncTimeInterval > MaxResyncTimeInterval)
-			{
-				ResyncTimeInterval = MaxResyncTimeInterval;
-			}
-
-			NextTimeSyncTime = PlayerSessionServices->GetSynchronizedUTCTime()->GetTime() + ResyncTimeInterval;
+		}
+		else
+		{
+			FinishManifestBuildingAfterTimesync(false);
 		}
 	}
 }
@@ -1121,8 +1140,6 @@ void FPlaylistReaderDASH::ManifestDownloadCompleted(FResourceLoadRequestPtr Requ
 			return;
 		}
 		
-		// Get the Date header from the response and set our clock to this time.
-		// Do this once only with the date from the master playlist.
 		const HTTP::FConnectionInfo* ConnInfo = Request->GetConnectionInfo();
 		if (!ConnInfo)
 		{
@@ -1130,22 +1147,15 @@ void FPlaylistReaderDASH::ManifestDownloadCompleted(FResourceLoadRequestPtr Requ
 		}
 
 		FString ETag;
-		FTimeValue FetchTime;
 		FString EffectiveURL = ConnInfo->EffectiveURL;
+		InitialTimesync.HttpRequestStart = ConnInfo->RequestStartTime;
+		InitialTimesync.UtcTimeThen = MEDIAutcTime::Current();
+		InitialTimesync.ConnInfo = MakeUnique<HTTP::FConnectionInfo>(*ConnInfo);
 		for(int32 i=0; i<ConnInfo->ResponseHeaders.Num(); ++i)
 		{
 			if (ConnInfo->ResponseHeaders[i].Header.Equals(TEXT("Date"), ESearchCase::IgnoreCase))
 			{
-				// Parse the header
-				FTimeValue DateFromHeader;
-				bool bDateParsedOk = RFC7231::ParseDateTime(DateFromHeader, ConnInfo->ResponseHeaders[i].Value);
-				if (bDateParsedOk && DateFromHeader.IsValid())
-				{
-					PlayerSessionServices->GetSynchronizedUTCTime()->SetTime(ConnInfo->RequestStartTime, DateFromHeader);
-				}
-				// The MPD 'FetchTime' is the time at which the request to get the MPD was made to the server.
-				FetchTime = PlayerSessionServices->GetSynchronizedUTCTime()->MapToSyncTime(ConnInfo->RequestStartTime);
-				bInitialTimeSyncedToMPDDateTimeHeader = true;
+				InitialTimesync.HttpDateHeader = ConnInfo->ResponseHeaders[i].Value;
 			}
 			else if (ConnInfo->ResponseHeaders[i].Header.Equals(TEXT("ETag"), ESearchCase::IgnoreCase))
 			{
@@ -1175,60 +1185,6 @@ void FPlaylistReaderDASH::ManifestDownloadCompleted(FResourceLoadRequestPtr Requ
 
 				// Set the next time sync time to zero to cause a time sync as soon as possible.
 				NextTimeSyncTime.SetToZero();
-
-				if (Manifest.IsValid())
-				{
-					// Note: This is not like the standard defines FetchTime. It should be the time the server is processing the request
-					//       but we can't know this. Also if there is a direct <UTCTiming> element in the MPD it may be different from
-					//       what our local clock is set to at the moment. So instead we set FetchTime to NOW.
-					FetchTime = PlayerSessionServices->GetSynchronizedUTCTime()->GetTime();
-					Manifest->GetMPDRoot()->SetFetchTime(FetchTime);
-					MostRecentMPDUpdateTime = FetchTime;
-
-					// If the URL has a special fragment part to turn this presentation into an event, do so.
-					// Do this before preparing the default start times!
-					// Note: This MAY require the client clock to be synced to the server IF the special keyword 'now' is used.
-					Manifest->TransformIntoEpicEvent();
-
-					// Do a one-time preparation of the default start time from the #t= URL fragment.
-					// This is not to be repeated for manifest updates.
-					Manifest->PrepareDefaultStartTime();
-				}
-
-				PlayerManifest = FManifestDASH::Create(PlayerSessionServices, Manifest);
-
-				// Check if the MPD defines an @minimumUpdatePeriod
-				FTimeValue mup = Manifest.IsValid() ? Manifest->GetMinimumUpdatePeriod() : FTimeValue::GetInvalid();
-				// If the update time is zero then updates happen just in time when segments are required or through
-				// an inband event stream. Either way, we do not need to update periodically.
-				if (mup.IsValid() && mup > FTimeValue::GetZero())
-				{
-					// Warn if MUP is really small
-					if (mup.GetAsMilliseconds() < 1000)
-					{
-						LogMessage(IInfoLog::ELevel::Warning, FString::Printf(TEXT("MPD@minimumUpdatePeriod is set to a really small value of %lld msec. This could be a performance issue."), (long long int)mup.GetAsMilliseconds()));
-					}
-					RequestsLock.Lock();
-					NextMPDUpdateTime = MostRecentMPDUpdateTime + mup;
-					RequestsLock.Unlock();
-				}
-
-				// Notify that the "master playlist" has been parsed, successfully or not. If we still need to resolve remote entities this is not an error!
-				PlayerSessionServices->SendMessageToPlayer(IPlaylistReader::PlaylistLoadedMessage::Create(!LastErrorDetail.IsTryAgain() ? LastErrorDetail : FErrorDetail(), ConnInfo, Playlist::EListType::Master, Playlist::ELoadType::Initial));
-				if (LastErrorDetail.IsOK())
-				{
-					// Notify that the "variant playlists" are ready. There are no variants in DASH, but this is the trigger that the playlists are all set up and are good to go now.
-					PlayerSessionServices->SendMessageToPlayer(IPlaylistReader::PlaylistLoadedMessage::Create(LastErrorDetail, ConnInfo, Playlist::EListType::Variant, Playlist::ELoadType::Initial));
-				}
-				// "try again" is returned when there are remote entities that need to be resolved first.
-				else if (LastErrorDetail.IsTryAgain())
-				{
-					EnqueueInitialXLinkRequests();
-				}
-				else
-				{
-					PostError(LastErrorDetail);
-				}
 			}
 			else
 			{
@@ -1247,6 +1203,130 @@ void FPlaylistReaderDASH::ManifestDownloadCompleted(FResourceLoadRequestPtr Requ
 }
 
 
+void FPlaylistReaderDASH::PrepareNextTimeSync()
+{
+	if (Manifest.IsValid())
+	{
+		// Do not re-sync the time too often.
+		const FTimeValue MinResyncTimeInterval(120.0);		// not sooner than every 120 seconds.
+		const FTimeValue MaxResyncTimeInterval(1800.0);		// at least once every 30 minutes.
+
+		FTimeValue ResyncTimeInterval(MinResyncTimeInterval);
+		// How frequently does the MPD update?
+		FTimeValue mup = Manifest->GetMinimumUpdatePeriod();
+		if (mup.IsValid() && mup > MinResyncTimeInterval)
+		{
+			ResyncTimeInterval = mup;
+		}
+		if (ResyncTimeInterval > MaxResyncTimeInterval)
+		{
+			ResyncTimeInterval = MaxResyncTimeInterval;
+		}
+
+		NextTimeSyncTime = PlayerSessionServices->GetSynchronizedUTCTime()->GetTime() + ResyncTimeInterval;
+	}
+}
+
+
+void FPlaylistReaderDASH::FinishManifestBuildingAfterTimesync(bool bGotTheTime)
+{
+	// After a time sync clear the segment fetch delay.
+	if (Manifest.IsValid())
+	{
+		Manifest->SetSegmentFetchDelay(FTimeValue::GetZero());
+	}
+
+	if (!bIsInitialTimesync)
+	{
+		PrepareNextTimeSync();
+		return;
+	}
+	bIsInitialTimesync = false;
+
+	FTimeValue Now = MEDIAutcTime::Current();
+	FTimeValue TimeDiffSinceStart = Now - InitialTimesync.UtcTimeThen;
+	if (!bGotTheTime)
+	{
+		// First check if the MPD contained a direct time value.
+		if (InitialTimesync.Utc_direct2014.Len())
+		{
+			FTimeValue DirectTime;
+			if (ISO8601::ParseDateTime(DirectTime, InitialTimesync.Utc_direct2014))
+			{
+				PlayerSessionServices->GetSynchronizedUTCTime()->SetTime(DirectTime + TimeDiffSinceStart);
+				// No need to use the Date header any more.
+				InitialTimesync.HttpDateHeader.Reset();
+			}
+		}
+		// As a last resort use the Date header from the HTTP response
+		if (InitialTimesync.HttpDateHeader.Len())
+		{
+			FTimeValue DateFromHeader;
+			bool bDateParsedOk = RFC7231::ParseDateTime(DateFromHeader, InitialTimesync.HttpDateHeader);
+			if (bDateParsedOk && DateFromHeader.IsValid())
+			{
+				PlayerSessionServices->GetSynchronizedUTCTime()->SetTime(InitialTimesync.HttpRequestStart, DateFromHeader);
+			}
+		}
+	}
+
+	if (Manifest.IsValid())
+	{
+		FTimeValue FetchTime = PlayerSessionServices->GetSynchronizedUTCTime()->GetTime() - TimeDiffSinceStart;
+		Manifest->GetMPDRoot()->SetFetchTime(FetchTime);
+		MostRecentMPDUpdateTime = FetchTime;
+
+		// If the URL has a special fragment part to turn this presentation into an event, do so.
+		// Do this before preparing the default start times!
+		// Note: This MAY require the client clock to be synced to the server IF the special keyword 'now' is used.
+		Manifest->TransformIntoEpicEvent();
+
+		// Do a one-time preparation of the default start time from the #t= URL fragment.
+		// This is not to be repeated for manifest updates.
+		Manifest->PrepareDefaultStartTime();
+	}
+
+	PlayerManifest = FManifestDASH::Create(PlayerSessionServices, Manifest);
+
+	// Check if the MPD defines an @minimumUpdatePeriod
+	FTimeValue mup = Manifest.IsValid() ? Manifest->GetMinimumUpdatePeriod() : FTimeValue::GetInvalid();
+	// If the update time is zero then updates happen just in time when segments are required or through
+	// an inband event stream. Either way, we do not need to update periodically.
+	if (mup.IsValid() && mup > FTimeValue::GetZero())
+	{
+		// Warn if MUP is really small
+		if (mup.GetAsMilliseconds() < 1000)
+		{
+			LogMessage(IInfoLog::ELevel::Warning, FString::Printf(TEXT("MPD@minimumUpdatePeriod is set to a really small value of %lld msec. This could be a performance issue."), (long long int)mup.GetAsMilliseconds()));
+		}
+		RequestsLock.Lock();
+		NextMPDUpdateTime = MostRecentMPDUpdateTime + mup;
+		RequestsLock.Unlock();
+	}
+
+	// Notify that the "master playlist" has been parsed, successfully or not. If we still need to resolve remote entities this is not an error!
+	PlayerSessionServices->SendMessageToPlayer(IPlaylistReader::PlaylistLoadedMessage::Create(!LastErrorDetail.IsTryAgain() ? LastErrorDetail : FErrorDetail(), InitialTimesync.ConnInfo.Get(), Playlist::EListType::Master, Playlist::ELoadType::Initial));
+	if (LastErrorDetail.IsOK())
+	{
+		// Notify that the "variant playlists" are ready. There are no variants in DASH, but this is the trigger that the playlists are all set up and are good to go now.
+		PlayerSessionServices->SendMessageToPlayer(IPlaylistReader::PlaylistLoadedMessage::Create(LastErrorDetail, InitialTimesync.ConnInfo.Get(), Playlist::EListType::Variant, Playlist::ELoadType::Initial));
+		PrepareNextTimeSync();
+	}
+	// "try again" is returned when there are remote entities that need to be resolved first.
+	else if (LastErrorDetail.IsTryAgain())
+	{
+		EnqueueInitialXLinkRequests();
+	}
+	else
+	{
+		PostError(LastErrorDetail);
+	}
+
+	InitialTimesync.Reset();
+}
+
+
+
 void FPlaylistReaderDASH::ManifestUpdateDownloadCompleted(FResourceLoadRequestPtr Request, bool bSuccess)
 {
 	bIsMPDUpdateInProgress = false;
@@ -1267,7 +1347,6 @@ void FPlaylistReaderDASH::ManifestUpdateDownloadCompleted(FResourceLoadRequestPt
 		FString EffectiveURL;
 		FTimeValue FetchTime;
 		EffectiveURL = ConnInfo->EffectiveURL;
-		// See above why we use NOW for FetchTime.
 		//FetchTime = PlayerSessionServices->GetSynchronizedUTCTime()->MapToSyncTime(ConnInfo->RequestStartTime);
 		FetchTime = PlayerSessionServices->GetSynchronizedUTCTime()->GetTime();
 		for(int32 i=0; i<ConnInfo->ResponseHeaders.Num(); ++i)
@@ -1385,6 +1464,7 @@ void FPlaylistReaderDASH::InitialMPDXLinkElementDownloadCompleted(FResourceLoadR
 			LastErrorDetail = ForManifest->BuildAfterInitialRemoteElementDownload();
 			// Notify that the "variant playlists" are ready. There are no variants in DASH, but this is the trigger that the playlists are all set up and are good to go now.
 			PlayerSessionServices->SendMessageToPlayer(IPlaylistReader::PlaylistLoadedMessage::Create(LastErrorDetail, nullptr, Playlist::EListType::Variant, Playlist::ELoadType::Initial));
+			PrepareNextTimeSync();
 		}
 		else if (LastErrorDetail.IsTryAgain())
 		{
@@ -1396,6 +1476,7 @@ void FPlaylistReaderDASH::InitialMPDXLinkElementDownloadCompleted(FResourceLoadR
 			PostError(LastErrorDetail);
 		}
 	}
+	InitialTimesync.Reset();
 }
 
 
@@ -1425,6 +1506,7 @@ void FPlaylistReaderDASH::Timesync_httphead_Completed(FResourceLoadRequestPtr Re
 	// Presently we do not try another time sync method if this one has failed. Clear out what we have attempted so far and be done with it.
 	AttemptedTimesyncDescriptors.Empty();
 	bTimeSyncInProgress = false;
+	FinishManifestBuildingAfterTimesync(bSuccess);
 }
 
 void FPlaylistReaderDASH::Timesync_httpxsdate_Completed(FResourceLoadRequestPtr Request, bool bSuccess)
@@ -1453,6 +1535,7 @@ void FPlaylistReaderDASH::Timesync_httpxsdate_Completed(FResourceLoadRequestPtr
 	// Presently we do not try another time sync method if this one has failed. Clear out what we have attempted so far and be done with it.
 	AttemptedTimesyncDescriptors.Empty();
 	bTimeSyncInProgress = false;
+	FinishManifestBuildingAfterTimesync(bSuccess);
 }
 
 void FPlaylistReaderDASH::Timesync_httpiso_Completed(FResourceLoadRequestPtr Request, bool bSuccess)
@@ -1472,6 +1555,7 @@ void FPlaylistReaderDASH::Timesync_httpiso_Completed(FResourceLoadRequestPtr Req
 	// Presently we do not try another time sync method if this one has failed. Clear out what we have attempted so far and be done with it.
 	AttemptedTimesyncDescriptors.Empty();
 	bTimeSyncInProgress = false;
+	FinishManifestBuildingAfterTimesync(bSuccess);
 }
 
 void FPlaylistReaderDASH::OnMediaPlayerEventReceived(TSharedPtrTS<IAdaptiveStreamingPlayerAEMSEvent> InEvent, IAdaptiveStreamingPlayerAEMSReceiver::EDispatchMode InDispatchMode)
diff --git a/Engine/Plugins/Media/ElectraPlayer/Source/ElectraPlayerRuntime/Private/Runtime/Player/DASH/PlaylistReaderDASH.h b/Engine/Plugins/Media/ElectraPlayer/Source/ElectraPlayerRuntime/Private/Runtime/Player/DASH/PlaylistReaderDASH.h
--- a/Engine/Plugins/Media/ElectraPlayer/Source/ElectraPlayerRuntime/Private/Runtime/Player/DASH/PlaylistReaderDASH.h
+++ b/Engine/Plugins/Media/ElectraPlayer/Source/ElectraPlayerRuntime/Private/Runtime/Player/DASH/PlaylistReaderDASH.h
@@ -54,6 +54,11 @@ public:
 	};
 	virtual void RequestMPDUpdate(EMPDRequestType InRequestType) = 0;
 
+	/**
+	 * Requests a resync of the player clock with the server clock.
+	 */
+	virtual void RequestClockResync() = 0;
+
 	/**
 	 * Access the current internal MPD
 	 */
diff --git a/Engine/Plugins/Media/ElectraPlayer/Source/ElectraPlayerRuntime/Private/Runtime/Player/DASH/PlaylistReaderDASH_Internal.h b/Engine/Plugins/Media/ElectraPlayer/Source/ElectraPlayerRuntime/Private/Runtime/Player/DASH/PlaylistReaderDASH_Internal.h
--- a/Engine/Plugins/Media/ElectraPlayer/Source/ElectraPlayerRuntime/Private/Runtime/Player/DASH/PlaylistReaderDASH_Internal.h
+++ b/Engine/Plugins/Media/ElectraPlayer/Source/ElectraPlayerRuntime/Private/Runtime/Player/DASH/PlaylistReaderDASH_Internal.h
@@ -110,7 +110,7 @@ namespace DASHUrlHelpers
 	};
 	void GetAllHierarchyUrlQueries(TArray<TSharedPtrTS<FDashMPD_UrlQueryInfoType>>& OutUrlQueries, TSharedPtrTS<const IDashMPDElement> StartingAt, EUrlQueryRequestType ForRequestType, bool bInclude2014);
 	FErrorDetail ApplyUrlQueries(IPlayerSessionServices* PlayerSessionServices, const FString& InMPDUrl, FString& InOutURL, FString& OutRequestHeader, const TArray<TSharedPtrTS<FDashMPD_UrlQueryInfoType>>& UrlQueries);
-	bool BuildAbsoluteElementURL(FString& OutURL, FTimeValue& ATO, const FString& DocumentURL, const TArray<TSharedPtrTS<const FDashMPD_BaseURLType>>& BaseURLs, const FString& ElementURL);
+	bool BuildAbsoluteElementURL(FString& OutURL, FTimeValue& ATO, TMediaOptionalValue<bool>& bATOComplete, const FString& DocumentURL, const TArray<TSharedPtrTS<const FDashMPD_BaseURLType>>& BaseURLs, const FString& ElementURL);
 	FString ApplyAnnexEByteRange(FString InURL, FString InRange, const TArray<TSharedPtrTS<const FDashMPD_BaseURLType>>& OutBaseURLs);
 }
 
@@ -126,6 +126,25 @@ public:
 		Dynamic
 	};
 
+	struct FProducerReferenceTimeInfo : public IProducerReferenceTimeInfo
+	{
+		FTimeValue WallclockTime;
+		uint64 PresentationTime = 0;
+		uint32 ID = 0;
+		EType Type = EType::Encoder;
+		bool bInband = false;
+		FTimeValue GetWallclockTime() const override
+		{ return WallclockTime; }
+		uint64 GetPresentationTime() const override
+		{ return PresentationTime; }
+		uint32 GetID() const override
+		{ return ID; }
+		EType GetType() const override
+		{ return Type; }
+		bool GetIsInband() const override
+		{ return bInband; }
+	};
+
 	struct FSegmentInformation
 	{
 		struct FURL
@@ -156,7 +175,10 @@ public:
 		int64 FirstByteOffset = 0;
 		int64 MediaLocalFirstAUTime = 0;			//!< Time of the first AU to use in this segment in media local time
 		int64 MediaLocalLastAUTime = 0;				//!< Time at which the last AU to use in thie segment ends in media local time
+		int64 MediaLocalFirstPTS = 0;
 		uint32 Timescale = 0;						//!< Local media timescale
+		TMediaOptionalValue<bool> bAvailabilityTimeComplete;
+		bool bLowLatencyChunkedEncodingExpected = false;
 		bool bFrameAccuracyRequired = false;		//!< true if the segment was located for frame accurate seeking.
 		bool bIsSideload = false;					//!< true if this is a side-loaded resource to be fetched and cached.
 		bool bIsLastInPeriod = false;				//!< true if known to be the last segment in the period.
@@ -164,6 +186,8 @@ public:
 		bool bIsMissing = false;					//!< Set to true if known to be missing.
 		bool bSawLMSG = false;						//!< Will be set to true by the stream reader if the 'lmsg' brand was found.
 		TArray<FInbandEventStream> InbandEventStreams;
+		TArray<FProducerReferenceTimeInfo> ProducerReferenceTimeInfos;
+		int64 MeasureLatencyViaReferenceTimeInfoID = -1;
 
 		FTimeValue CalculateASAST(const FTimeValue& AST, const FTimeValue& PeriodStart, bool bIsStatic)
 		{
@@ -213,6 +237,7 @@ public:
 		int64 RequestID = 0;						//!< Sequential request ID across all segments during playback, needed to re-resolve potential UrlQueryInfo xlinks.
 		bool bHasFollowingPeriod = false;			//!< true if we know for sure there is another period following.
 		bool bFrameAccurateSearch = false;			//!< true to prepare segments for frame-accurate decoding and rendering
+		bool bInitSegmentSetupOnly = false;			//!< true to get the initialization segment information only.
 	};
 
 	class FRepresentation : public IPlaybackAssetRepresentation, public TSharedFromThis<FRepresentation, ESPMode::ThreadSafe>
@@ -241,23 +266,23 @@ public:
 		//----------------------------------------------
 		// Methods from IPlaybackAssetRepresentation
 		//
-		virtual FString GetUniqueIdentifier() const override
+		FString GetUniqueIdentifier() const override
 		{
 			return ID;
 		}
-		virtual const FStreamCodecInformation& GetCodecInformation() const override
+		const FStreamCodecInformation& GetCodecInformation() const override
 		{
 			return CodecInfo;
 		}
-		virtual int32 GetBitrate() const override
+		int32 GetBitrate() const override
 		{
 			return Bitrate;
 		}
-		virtual int32 GetQualityIndex() const override
+		int32 GetQualityIndex() const override
 		{
 			return QualityIndex;
 		}
-		virtual bool CanBePlayed() const override
+		bool CanBePlayed() const override
 		{
 			return bIsEnabled && bIsUsable;
 		}
@@ -274,6 +299,7 @@ public:
 		bool PrepareDownloadURLs(IPlayerSessionServices* PlayerSessionServices, FSegmentInformation& InOutSegmentInfo, const TArray<TSharedPtrTS<FDashMPD_SegmentTemplateType>>& SegmentTemplate);
 		FString ApplyTemplateStrings(FString TemplateURL, const FSegmentInformation& InSegmentInfo);
 		void CollectInbandEventStreams(IPlayerSessionServices* PlayerSessionServices, FSegmentInformation& InOutSegmentInfo);
+		void SetupProducerReferenceTimeInfo(IPlayerSessionServices* PlayerSessionServices, FSegmentInformation& InOutSegmentInfo);
 
 		void SegmentIndexDownloadComplete(TSharedPtrTS<FMPDLoadRequestDASH> Request, bool bSuccess);
 		friend class FManifestDASHInternal;
@@ -285,12 +311,14 @@ public:
 		int32 SelectionPriority = 1;
 		bool bIsUsable = false;
 		bool bIsEnabled = true;
+		TMediaOptionalValue<bool> bAvailableAsLowLatency;
 		bool bWarnedAboutTimelineStartGap = false;
 		bool bWarnedAboutTimelineNoTAfterNegativeR = false;
 		bool bWarnedAboutTimelineNumberOverflow = false;
 		bool bWarnedAboutInconsistentNumbering = false;
 		bool bWarnedAboutTimelineOverlap = false;
 		bool bWarnedAboutTimescale = false;
+		bool bWarnedAboutInconsistentAvailabilityTimeComplete = false;
 		//
 		bool bNeedsSegmentIndex = true;
 		TSharedPtrTS<const IParserISO14496_12> SegmentIndex;
@@ -299,6 +327,8 @@ public:
 		TSharedPtrTS<FMPDLoadRequestDASH> PendingSegmentIndexLoadRequest;
 		//
 		bool bIsSideloadedSubtitle = false;
+
+		TArray<FProducerReferenceTimeInfo> ProducerReferenceTimeInfos;
 	};
 
 	class FAdaptationSet : public IPlaybackAssetAdaptationSet
@@ -493,23 +523,27 @@ public:
 		//----------------------------------------------
 		// Methods from IPlaybackAssetAdaptationSet
 		//
-		virtual FString GetUniqueIdentifier() const override
+		FString GetUniqueIdentifier() const override
 		{
 			return FString::Printf(TEXT("%d"), UniqueSequentialSetIndex);
 		}
-		virtual FString GetListOfCodecs() const override
+		FString GetListOfCodecs() const override
 		{
 			return Codec.GetCodecSpecifierRFC6381();
 		}
-		virtual FString GetLanguage() const override
+		FString GetLanguage() const override
 		{
 			return Language;
 		}
-		virtual int32 GetNumberOfRepresentations() const override
+		int32 GetNumberOfRepresentations() const override
 		{
 			return Representations.Num();
 		}
-		virtual TSharedPtrTS<IPlaybackAssetRepresentation> GetRepresentationByIndex(int32 RepresentationIndex) const override
+		bool IsLowLatencyEnabled() const override
+		{
+			return bAvailableAsLowLatency.GetWithDefault(false);
+		}
+		TSharedPtrTS<IPlaybackAssetRepresentation> GetRepresentationByIndex(int32 RepresentationIndex) const override
 		{
 			if (RepresentationIndex < Representations.Num())
 			{
@@ -517,7 +551,7 @@ public:
 			}
 			return nullptr;
 		}
-		virtual TSharedPtrTS<IPlaybackAssetRepresentation> GetRepresentationByUniqueIdentifier(const FString& UniqueIdentifier) const override
+		TSharedPtrTS<IPlaybackAssetRepresentation> GetRepresentationByUniqueIdentifier(const FString& UniqueIdentifier) const override
 		{
 			return GetRepresentationByUniqueID(UniqueIdentifier);
 		}
@@ -535,6 +569,8 @@ public:
 		int32 UniqueSequentialSetIndex = 0;
 		bool bIsUsable = false;
 		bool bIsEnabled = true;
+		TMediaOptionalValue<bool> bAvailableAsLowLatency;
+
 		// Encryption related
 		TArray<FContentProtection> PossibleContentProtections;
 		FString CommonEncryptionScheme;
@@ -544,6 +580,8 @@ public:
 		TArray<FString> SwitchedFromSetIDs;
 		bool bIsInSwitchGroup = false;
 		bool bIsSwitchGroup = false;
+		// Warnings
+		bool bWarnedAboutInconsistentAvailabilityTimeComplete = false;
 	};
 
 	class FPeriod : public ITimelineMediaAsset
@@ -633,25 +671,25 @@ public:
 		//----------------------------------------------
 		// Methods from ITimelineMediaAsset
 		//
-		virtual FTimeRange GetTimeRange() const override
+		FTimeRange GetTimeRange() const override
 		{
 			// Per convention the time range includes the AST
 			return FTimeRange({StartAST, EndAST});
 		}
-		virtual FTimeValue GetDuration() const override
+		FTimeValue GetDuration() const override
 		{
 			return Duration;
 		}
-		virtual FString GetAssetIdentifier() const override
+		FString GetAssetIdentifier() const override
 		{
 			TSharedPtrTS<FDashMPD_PeriodType> MPDPeriod = Period.Pin();
 			return MPDPeriod.IsValid() && MPDPeriod->GetAssetIdentifier().IsValid() ? MPDPeriod->GetAssetIdentifier()->GetValue() : FString();
 		}
-		virtual FString GetUniqueIdentifier() const override
+		FString GetUniqueIdentifier() const override
 		{
 			return ID;
 		}
-		virtual int32 GetNumberOfAdaptationSets(EStreamType OfStreamType) const override
+		int32 GetNumberOfAdaptationSets(EStreamType OfStreamType) const override
 		{
 			int32 Num=0;
 			for(int32 i=0; i<AdaptationSets.Num(); ++i)
@@ -663,7 +701,7 @@ public:
 			}
 			return Num;
 		}
-		virtual TSharedPtrTS<IPlaybackAssetAdaptationSet> GetAdaptationSetByTypeAndIndex(EStreamType OfStreamType, int32 AdaptationSetIndex) const override
+		TSharedPtrTS<IPlaybackAssetAdaptationSet> GetAdaptationSetByTypeAndIndex(EStreamType OfStreamType, int32 AdaptationSetIndex) const override
 		{
 			int32 Num=0;
 			for(int32 i=0; i<AdaptationSets.Num(); ++i)
@@ -679,7 +717,7 @@ public:
 			return nullptr;
 		}
 
-		virtual void GetMetaData(TArray<FTrackMetadata>& OutMetadata, EStreamType OfStreamType) const override
+		void GetMetaData(TArray<FTrackMetadata>& OutMetadata, EStreamType OfStreamType) const override
 		{
 			for(int32 i=0, iMax=GetNumberOfAdaptationSets(OfStreamType); i<iMax; ++i)
 			{
@@ -798,6 +836,24 @@ public:
 	ELECTRA_IMPL_DEFAULT_ERROR_METHODS(DASHManifest);
 
 
+	TSharedPtrTS<const FLowLatencyDescriptor> GetLowLatencyDescriptor() const
+	{
+		return LowLatencyDescriptor;
+	}
+	TSharedPtrTS<FProducerReferenceTimeInfo> GetProducerReferenceTimeElement(int64 InID) const
+	{
+		return InID >= 0 ? ProducerReferenceTimeElements.FindRef((uint32)InID) : nullptr;
+	}
+
+	FTimeValue GetSegmentFetchDelay() const
+	{
+		return SegmentFetchDelay;
+	}
+	void SetSegmentFetchDelay(const FTimeValue& InNewFetchDelay)
+	{
+		SegmentFetchDelay = InNewFetchDelay;
+	}
+
 	FTimeValue GetAnchorTime() const;
 	FTimeRange GetTotalTimeRange() const;
 	FTimeRange GetSeekableTimeRange() const;
@@ -807,6 +863,7 @@ public:
 	FTimeValue GetDefaultStartTime() const;
 	void ClearDefaultStartTime();
 	FTimeRange GetPlayTimesFromURI() const;
+	FTimeValue GetDesiredLiveLatency() const;
 
 	FTimeValue GetMPDValidityEndTime() const;
 	FTimeValue GetLastPeriodEndTime() const;
@@ -851,6 +908,13 @@ private:
 
 	TArray<TSharedPtrTS<FPeriod>> Periods;
 
+	TMap<uint32, TSharedPtrTS<FProducerReferenceTimeInfo>> ProducerReferenceTimeElements;
+	TSharedPtrTS<FLowLatencyDescriptor> LowLatencyDescriptor;
+
+	FTimeValue SegmentFetchDelay;
+
+	mutable FTimeValue CalculatedLiveDistance;
+
 	mutable FTimeRange TotalTimeRange;
 	mutable FTimeRange SeekableTimeRange;
 	FTimeValue DefaultStartTime;
diff --git a/Engine/Plugins/Media/ElectraPlayer/Source/ElectraPlayerRuntime/Private/Runtime/Player/DASH/StreamReaderFMP4DASH.cpp b/Engine/Plugins/Media/ElectraPlayer/Source/ElectraPlayerRuntime/Private/Runtime/Player/DASH/StreamReaderFMP4DASH.cpp
--- a/Engine/Plugins/Media/ElectraPlayer/Source/ElectraPlayerRuntime/Private/Runtime/Player/DASH/StreamReaderFMP4DASH.cpp
+++ b/Engine/Plugins/Media/ElectraPlayer/Source/ElectraPlayerRuntime/Private/Runtime/Player/DASH/StreamReaderFMP4DASH.cpp
@@ -16,6 +16,7 @@
 #include "Player/DASH/PlayerEventDASH_Internal.h"
 #include "Player/DRM/DRMManager.h"
 #include "Utilities/Utilities.h"
+#include "Utilities/TimeUtilities.h"
 #include "Async/Async.h"
 
 #define INTERNAL_ERROR_INIT_SEGMENT_DOWNLOAD_ERROR					1
@@ -45,8 +46,19 @@ uint32 FStreamSegmentRequestFMP4DASH::GetPlaybackSequenceID() const
 	return CurrentPlaybackSequenceID;
 }
 
-void FStreamSegmentRequestFMP4DASH::SetExecutionDelay(const FTimeValue& ExecutionDelay)
+void FStreamSegmentRequestFMP4DASH::SetExecutionDelay(const FTimeValue& UTCNow, const FTimeValue& ExecutionDelay)
 {
+	// If there is a delay specified and the current time is already past the availability time
+	// then this is an old segment before the Live edge since we had paused or seeked backwards.
+	// In that case, or if there is no availability time due to VoD, set the availability time
+	// as the provided current time to apply the delay to.
+	if (UTCNow.IsValid() && ExecutionDelay > FTimeValue::GetZero())
+	{
+		if (!ASAST.IsValid() || UTCNow > ASAST)
+		{
+			ASAST = UTCNow;
+		}
+	}
 	DownloadDelayTime = ExecutionDelay;
 }
 
@@ -102,7 +114,7 @@ void FStreamSegmentRequestFMP4DASH::GetEndedStreams(TArray<TSharedPtrTS<IStreamS
 
 FTimeValue FStreamSegmentRequestFMP4DASH::GetFirstPTS() const
 {
-	return AST + AdditionalAdjustmentTime + PeriodStart + FTimeValue((Segment.bFrameAccuracyRequired ? Segment.MediaLocalFirstAUTime : Segment.Time) - Segment.PTO, Segment.Timescale);
+	return AST + AdditionalAdjustmentTime + PeriodStart + FTimeValue((Segment.bFrameAccuracyRequired ? Segment.MediaLocalFirstPTS : Segment.Time) - Segment.PTO, Segment.Timescale);
 }
 
 int32 FStreamSegmentRequestFMP4DASH::GetQualityIndex() const
@@ -120,6 +132,20 @@ void FStreamSegmentRequestFMP4DASH::GetDownloadStats(Metrics::FSegmentDownloadSt
 	OutStats = DownloadStats;
 }
 
+
+bool FStreamSegmentRequestFMP4DASH::GetStartupDelay(FTimeValue& OutStartTime, FTimeValue& OutTimeIntoSegment, FTimeValue& OutSegmentDuration) const
+{
+	check(DependentStreams.Num());
+	if (DependentStreams.Num())
+	{
+		OutTimeIntoSegment.SetFromND(DependentStreams[0]->Segment.MediaLocalFirstAUTime - DependentStreams[0]->Segment.Time, DependentStreams[0]->Segment.Timescale, 0);
+		OutSegmentDuration.SetFromND(DependentStreams[0]->Segment.Duration, DependentStreams[0]->Segment.Timescale, 0);
+		OutStartTime = DependentStreams[0]->GetFirstPTS();
+		return true;
+	}
+	return false;
+}
+
 /***************************************************************************************************************************************************/
 /***************************************************************************************************************************************************/
 /***************************************************************************************************************************************************/
@@ -504,7 +530,6 @@ FErrorDetail FStreamReaderFMP4DASH::FStreamHandler::GetInitSegment(TSharedPtrTS<
 	ds.TimeToDownload = (ci->RequestEndTime - ci->RequestStartTime).GetAsSeconds();
 	ds.ByteSize = ci->ContentLength;
 	ds.NumBytesDownloaded = ci->BytesReadSoFar;
-	ds.ThroughputBps = ci->Throughput.GetThroughput() ? ci->Throughput.GetThroughput() : (ds.TimeToDownload > 0.0 ? 8 * ds.NumBytesDownloaded / ds.TimeToDownload : 0);
 	ds.MediaAssetID = Request->Period.IsValid() ? Request->Period->GetUniqueIdentifier() : "";
 	ds.AdaptationSetID = Request->AdaptationSet.IsValid() ? Request->AdaptationSet->GetUniqueIdentifier() : "";
 	ds.RepresentationID = Request->Representation.IsValid() ? Request->Representation->GetUniqueIdentifier() : "";
@@ -637,7 +662,6 @@ FErrorDetail FStreamReaderFMP4DASH::FStreamHandler::RetrieveSideloadedFile(TShar
 	ds.TimeToDownload = (ci->RequestEndTime - ci->RequestStartTime).GetAsSeconds();
 	ds.ByteSize = ci->ContentLength;
 	ds.NumBytesDownloaded = ci->BytesReadSoFar;
-	ds.ThroughputBps = ci->Throughput.GetThroughput() ? ci->Throughput.GetThroughput() : (ds.TimeToDownload > 0.0 ? 8 * ds.NumBytesDownloaded / ds.TimeToDownload : 0);
 	ds.MediaAssetID = Request->Period.IsValid() ? Request->Period->GetUniqueIdentifier() : "";
 	ds.AdaptationSetID = Request->AdaptationSet.IsValid() ? Request->AdaptationSet->GetUniqueIdentifier() : "";
 	ds.RepresentationID = Request->Representation.IsValid() ? Request->Representation->GetUniqueIdentifier() : "";
@@ -795,7 +819,8 @@ void FStreamReaderFMP4DASH::FStreamHandler::HandleRequest()
 	ds.bDidTimeout  	   = false;
 	ds.HTTPStatusCode      = 0;
 	ds.StreamType   	   = Request->GetType();
-	ds.SegmentType  	   = Metrics::ESegmentType::Media;
+	// Assume we need to fetch the init segment first. This gets changed to 'media' when we load the media segment.
+	ds.SegmentType  	   = Metrics::ESegmentType::Init;
 	ds.PresentationTime    = Request->GetFirstPTS().GetAsSeconds();
 	ds.Bitrate  		   = Request->GetBitrate();
 	ds.Duration 		   = SegmentDuration.GetAsSeconds();
@@ -805,7 +830,6 @@ void FStreamReaderFMP4DASH::FStreamHandler::HandleRequest()
 	ds.TimeToDownload      = 0.0;
 	ds.ByteSize 		   = -1;
 	ds.NumBytesDownloaded  = 0;
-	ds.ThroughputBps	   = 0;
 	ds.bInsertedFillerData = false;
 
 	ds.MediaAssetID 	= Request->Period.IsValid() ? Request->Period->GetUniqueIdentifier() : "";
@@ -815,7 +839,6 @@ void FStreamReaderFMP4DASH::FStreamHandler::HandleRequest()
 	ds.Range  			= Request->Segment.MediaURL.Range;
 	ds.CDN  			= Request->Segment.MediaURL.CDN;
 	ds.RetryNumber  	= Request->NumOverallRetries;
-	ds.ABRState.Reset();
 
 	// Clear out the list of events found the last time.
 	SegmentEventsFound.Empty();
@@ -841,6 +864,7 @@ void FStreamReaderFMP4DASH::FStreamHandler::HandleRequest()
 	bAbortedByABR   	   = false;
 	bAllowEarlyEmitting    = false;
 	bFillRemainingDuration = false;
+	ABRAbortReason.Empty();
 	DurationSuccessfullyRead.SetToZero();
 	DurationSuccessfullyDelivered.SetToZero();
 	AccessUnitFIFO.Clear();
@@ -873,6 +897,9 @@ void FStreamReaderFMP4DASH::FStreamHandler::HandleRequest()
 
 	if (InitSegmentError.IsOK() && !HasReadBeenAborted())
 	{
+		// Tending to the media segment now.
+		ds.SegmentType = Metrics::ESegmentType::Media;
+
 		if (Request->Segment.bIsSideload)
 		{
 			FErrorDetail SideloadError;
@@ -943,6 +970,7 @@ void FStreamReaderFMP4DASH::FStreamHandler::HandleRequest()
 			{
 				HTTP->Parameters.RequestHeaders.Emplace(HTTP::FHTTPHeader({DASH::HTTPHeaderOptionName, Request->Segment.MediaURL.CustomHeader}));
 			}
+			HTTP->Parameters.bCollectTimingTraces = Request->Segment.bLowLatencyChunkedEncodingExpected;
 			// Set timeouts for media segment retrieval
 			const FParamDict& Options = PlayerSessionService->GetOptions();
 			HTTP->Parameters.ConnectTimeout = Options.GetValue(DASH::OptionKeyMediaSegmentConnectTimeout).SafeGetTimeValue(FTimeValue().SetFromMilliseconds(1000 * 4));
@@ -975,8 +1003,11 @@ void FStreamReaderFMP4DASH::FStreamHandler::HandleRequest()
 				}
 			}
 
+
 			while(!bDone && !HasErrored() && !HasReadBeenAborted())
 			{
+				Metrics::FSegmentDownloadStats::FMovieChunkInfo MoofInfo;
+				MoofInfo.HeaderOffset = GetCurrentOffset();
 				UEMediaError parseError = MP4Parser->ParseHeader(this, this, PlayerSessionService, MP4InitSegment.Get());
 				if (parseError == UEMEDIA_ERROR_OK)
 				{
@@ -1014,13 +1045,13 @@ void FStreamReaderFMP4DASH::FStreamHandler::HandleRequest()
 								{
 									// Need to rescale the AU times from the MPD timescale to the media timescale.
 									MediaLocalFirstAUTime = FTimeFraction(Request->Segment.MediaLocalFirstAUTime, Request->Segment.Timescale).GetAsTimebase(TrackTimescale);
-									MediaLocalLastAUTime = FTimeFraction(Request->Segment.MediaLocalLastAUTime, Request->Segment.Timescale).GetAsTimebase(TrackTimescale);
+									MediaLocalLastAUTime = MediaLocalLastAUTime == TNumericLimits<int64>::Max() ? MediaLocalLastAUTime : FTimeFraction(Request->Segment.MediaLocalLastAUTime, Request->Segment.Timescale).GetAsTimebase(TrackTimescale);
 									PTO = FTimeFraction(Request->Segment.PTO, Request->Segment.Timescale).GetAsTimebase(TrackTimescale);
 
 									if (!Request->bWarnedAboutTimescale)
 									{
 										Request->bWarnedAboutTimescale = true;
-										LogMessage(IInfoLog::ELevel::Warning, FString::Printf(TEXT("Track timescale %u differs from timescale of %u in MPD or segment index. This may cause playback problems!"), TrackTimescale, Request->Segment.Timescale));
+										LogMessage(IInfoLog::ELevel::Verbose, FString::Printf(TEXT("Track timescale %u differs from timescale of %u in MPD or segment index. This may cause playback problems!"), TrackTimescale, Request->Segment.Timescale));
 									}
 								}
 
@@ -1032,20 +1063,73 @@ void FStreamReaderFMP4DASH::FStreamHandler::HandleRequest()
 									Decrypter->UpdateInitDataFromMultiplePSSH(PSSHBoxes);
 								}
 
+								// Producer reference time?
+								struct FProducerTime
+								{
+									FTimeValue	Base;
+									int64		Media = 0;
+								};
+								FProducerTime ProducerTime;
+								if (Request->Segment.MeasureLatencyViaReferenceTimeInfoID >= 0 && Request->Segment.ProducerReferenceTimeInfos.Num())
+								{
+									// We only look at inband 'prtf' boxes if the <ProducerReferenceTime> element in the MPD tells us to.
+									// This is similar to events that are not to be considered any more if the MPD doesn't specify them.
+									for(auto &MPDPrtf : Request->Segment.ProducerReferenceTimeInfos)
+									{
+										if (MPDPrtf.GetID() == (uint32)Request->Segment.MeasureLatencyViaReferenceTimeInfoID)
+										{
+											// Use the inband 'prft' boxes?
+											if (MPDPrtf.bInband)
+											{
+												TArray<IParserISO14496_12::ITrack::FProducerReferenceTime> PRFTBoxes;
+												Track->GetPRFTBoxes(PRFTBoxes);
+												bool bFound = false;
+												for(auto &MP4Prtf : PRFTBoxes)
+												{
+													if ((MPDPrtf.Type == FManifestDASHInternal::FProducerReferenceTimeInfo::EType::Encoder && MP4Prtf.Reference == IParserISO14496_12::ITrack::FProducerReferenceTime::EReferenceType::Encoder) ||
+														(MPDPrtf.Type == FManifestDASHInternal::FProducerReferenceTimeInfo::EType::Captured && MP4Prtf.Reference == IParserISO14496_12::ITrack::FProducerReferenceTime::EReferenceType::Captured))
+													{
+														RFC5905::ParseNTPTime(ProducerTime.Base, MP4Prtf.NtpTimestamp);
+														ProducerTime.Media = MP4Prtf.MediaTime;
+														bFound = true;
+														break;
+													}
+												}
+												// When the MPD says that there are inband prtf's then this has to be so. If for some reason this is not the case
+												// then what are we to do?
+												if (!bFound)
+												{
+													// We take the values from the MPD here, which may be better than nothing?!
+													ProducerTime.Base = MPDPrtf.WallclockTime;
+													ProducerTime.Media = MPDPrtf.PresentationTime;
+												}
+											}
+											else
+											{
+												// Use values from MPD
+												ProducerTime.Base = MPDPrtf.WallclockTime;
+												ProducerTime.Media = MPDPrtf.PresentationTime;
+											}
+											break;
+										}
+									}
+								}
+
+								// Iterate the moof
 								for(Error = TrackIterator->StartAtFirst(false); Error == UEMEDIA_ERROR_OK; Error = TrackIterator->Next())
 								{
 									if (bIsFirstAU)
 									{
 										FTimeValue BaseMediaDecodeTime;
-										BaseMediaDecodeTime.SetFromND(TrackIterator->GetBaseMediaDecodeTime(), TrackIterator->GetTimescale());
+										BaseMediaDecodeTime.SetFromND(TrackIterator->GetBaseMediaDecodeTime(), TrackTimescale);
 										TimelineOffset = BaseMediaDecodeTime - Request->Segment.PeriodLocalSegmentStartTime;
 										//LogMessage(IInfoLog::ELevel::Info, FString::Printf(TEXT("Segment to timeline delta = %.3fs"), TimelineOffset.GetAsSeconds()));
 									}
 
 									// Get the DTS and PTS. Those are 0-based in a fragment and offset by the base media decode time of the fragment.
-									int64 AUDTS = TrackIterator->GetDTS();
-									int64 AUPTS = TrackIterator->GetPTS();
-									int64 AUDuration = TrackIterator->GetDuration();
+									const int64 AUDTS = TrackIterator->GetDTS();
+									const int64 AUPTS = TrackIterator->GetPTS();
+									const int64 AUDuration = TrackIterator->GetDuration();
 
 									// Create access unit
 									FAccessUnit *AccessUnit = FAccessUnit::Create(Parameters.MemoryProvider);
@@ -1081,7 +1165,7 @@ void FStreamReaderFMP4DASH::FStreamHandler::HandleRequest()
 										}
 										*/
 									}
-									FTimeValue Duration(TrackIterator->GetDuration(), TrackTimescale);
+									FTimeValue Duration(AUDuration, TrackTimescale);
 									AccessUnit->Duration = Duration;
 
 									// Offset the AU's DTS and PTS to the time mapping of the segment.
@@ -1102,12 +1186,21 @@ void FStreamReaderFMP4DASH::FStreamHandler::HandleRequest()
 										AccessUnit->EarliestPTS += TimeOffset;
 										AccessUnit->EarliestPTS.SetSequenceIndex(Request->TimestampSequenceIndex);
 									}
-									AccessUnit->LatestPTS.SetFromND(MediaLocalLastAUTime - PTO, TrackTimescale);
+									if (MediaLocalLastAUTime != TNumericLimits<int64>::Max())
+									{
+										AccessUnit->LatestPTS.SetFromND(MediaLocalLastAUTime - PTO, TrackTimescale);
+									}
+									else
+									{
+										AccessUnit->LatestPTS.SetToPositiveInfinity();
+									}
 									AccessUnit->LatestPTS += TimeOffset;
 									AccessUnit->LatestPTS.SetSequenceIndex(Request->TimestampSequenceIndex);
 
 									AccessUnit->OffsetFromSegmentStart = TimelineOffset;
 
+									AccessUnit->ProducerReferenceTime = ProducerTime.Base + FTimeValue(AUDTS - ProducerTime.Media, TrackTimescale);
+
 									ElectraCDM::FMediaCDMSampleInfo SampleEncryptionInfo;
 									bool bIsSampleEncrypted = TrackIterator->GetEncryptionInfo(SampleEncryptionInfo);
 
@@ -1138,9 +1231,17 @@ void FStreamReaderFMP4DASH::FStreamHandler::HandleRequest()
 											break;
 										}
 									}
+
+									if (MoofInfo.PayloadStartOffset == 0)
+									{
+										MoofInfo.PayloadStartOffset = GetCurrentOffset();
+									}
+
 									int64 NumRead = ReadData(AccessUnit->AUData, AccessUnit->AUSize);
 									if (NumRead == AccessUnit->AUSize)
 									{
+										MoofInfo.NumKeyframeBytes += AccessUnit->bIsSyncSample ? AccessUnit->AUSize : 0;
+										MoofInfo.ContentDuration += Duration;
 										DurationSuccessfullyRead += Duration;
 										NextExpectedDTS = AccessUnit->DTS + Duration;
 										LastKnownAUDuration = Duration;
@@ -1235,6 +1336,12 @@ void FStreamReaderFMP4DASH::FStreamHandler::HandleRequest()
 								}
 								delete TrackIterator;
 
+								MoofInfo.PayloadEndOffset = LastSuccessfulFilePos;
+								if (Request->Segment.bLowLatencyChunkedEncodingExpected)
+								{
+									ds.MovieChunkInfos.Emplace(MoveTemp(MoofInfo));
+								}
+
 								if (Error != UEMEDIA_ERROR_OK && Error != UEMEDIA_ERROR_END_OF_STREAM)
 								{
 									// error iterating
@@ -1301,6 +1408,7 @@ void FStreamReaderFMP4DASH::FStreamHandler::HandleRequest()
 		if (!HasReadBeenAborted())
 		{
 			CurrentRequest->ConnectionInfo.StatusInfo.ErrorDetail = InitSegmentError;
+			ds.bParseFailure = InitSegmentError.GetCode() == INTERNAL_ERROR_INIT_SEGMENT_PARSE_ERROR;
 			bHasErrored = true;
 		}
 	}
@@ -1411,6 +1519,9 @@ void FStreamReaderFMP4DASH::FStreamHandler::HandleRequest()
 			AccessUnit->LatestPTS = DiscardAfter;
 			AccessUnit->LatestPTS.SetSequenceIndex(Request->TimestampSequenceIndex);
 
+			//AccessUnit->OffsetFromSegmentStart = TimelineOffset;
+			//AccessUnit->ProducerReferenceTime = ProducerTime.Base + FTimeValue(AUDTS - ProducerTime.Media, TrackTimescale);
+
 			// Calculate the drop on the fragment local NextExpectedDTS/PTS.
 			AccessUnit->DropState = FAccessUnit::EDropState::None;
 			if (NextExpectedDTS + AccessUnit->Duration < DiscardBefore)
@@ -1466,7 +1577,7 @@ void FStreamReaderFMP4DASH::FStreamHandler::HandleRequest()
 	// If the ABR aborted this takes precedence in the failure message. Overwrite it.
 	if (bAbortedByABR)
 	{
-		ds.FailureReason = ds.ABRState.ProgressDecision.Reason;
+		ds.FailureReason = ABRAbortReason;
 	}
 	// Set up remaining download stat fields.
 	ds.bWasAborted  	  = bAbortedByABR;
@@ -1479,13 +1590,8 @@ void FStreamReaderFMP4DASH::FStreamHandler::HandleRequest()
 	ds.TimeToDownload     = (CurrentRequest->ConnectionInfo.RequestEndTime - CurrentRequest->ConnectionInfo.RequestStartTime).GetAsSeconds();
 	ds.ByteSize 		  = CurrentRequest->ConnectionInfo.ContentLength;
 	ds.NumBytesDownloaded = CurrentRequest->ConnectionInfo.BytesReadSoFar;
-	ds.ThroughputBps	  = CurrentRequest->ConnectionInfo.Throughput.GetThroughput();
-	if (ds.ThroughputBps == 0)
-	{
-		ds.ThroughputBps = ds.TimeToDownload > 0.0 ? 8 * ds.NumBytesDownloaded / ds.TimeToDownload : 0;
-	}
-	ds.bIsCachedResponse = CurrentRequest->ConnectionInfo.bIsCachedResponse;
-
+	ds.bIsCachedResponse  = CurrentRequest->ConnectionInfo.bIsCachedResponse;
+	CurrentRequest->ConnectionInfo.GetTimingTraces(ds.TimingTraces);
 
 	// Was this request for a segment that might potentially be missing and it did?
 	if (Request->Segment.bMayBeMissing && (ds.HTTPStatusCode == 404 || ds.HTTPStatusCode == 416))
@@ -1588,11 +1694,11 @@ int64 FStreamReaderFMP4DASH::FStreamHandler::ReadData(void* IntoBuffer, int64 Nu
 				Metrics::FSegmentDownloadStats currentDownloadStats = CurrentRequest->DownloadStats;
 				currentDownloadStats.DurationDelivered = DurationSuccessfullyDelivered.GetAsSeconds();
 				currentDownloadStats.DurationDownloaded = DurationSuccessfullyRead.GetAsSeconds();
+				currentDownloadStats.TimeToDownload = (MEDIAutcTime::Current() - CurrentRequest->ConnectionInfo.RequestStartTime).GetAsSeconds();
 				MetricUpdateLock.Unlock();
 
 				Metrics::FSegmentDownloadStats& ds = CurrentRequest->DownloadStats;
 				FABRDownloadProgressDecision StreamSelectorDecision = StreamSelector->ReportDownloadProgress(currentDownloadStats);
-				ds.ABRState.ProgressDecision = StreamSelectorDecision;
 				if ((StreamSelectorDecision.Flags & FABRDownloadProgressDecision::EDecisionFlags::eABR_EmitPartialData) != 0)
 				{
 					SCOPE_CYCLE_COUNTER(STAT_ElectraPlayer_DASH_StreamReader);
@@ -1626,6 +1732,7 @@ int64 FStreamReaderFMP4DASH::FStreamHandler::ReadData(void* IntoBuffer, int64 Nu
 					{
 						bFillRemainingDuration = true;
 					}
+					ABRAbortReason = StreamSelectorDecision.Reason;
 					bAbortedByABR = true;
 					return -1;
 				}
diff --git a/Engine/Plugins/Media/ElectraPlayer/Source/ElectraPlayerRuntime/Private/Runtime/Player/DASH/StreamReaderFMP4DASH.h b/Engine/Plugins/Media/ElectraPlayer/Source/ElectraPlayerRuntime/Private/Runtime/Player/DASH/StreamReaderFMP4DASH.h
--- a/Engine/Plugins/Media/ElectraPlayer/Source/ElectraPlayerRuntime/Private/Runtime/Player/DASH/StreamReaderFMP4DASH.h
+++ b/Engine/Plugins/Media/ElectraPlayer/Source/ElectraPlayerRuntime/Private/Runtime/Player/DASH/StreamReaderFMP4DASH.h
@@ -19,25 +19,26 @@ public:
 	FStreamSegmentRequestFMP4DASH();
 	virtual ~FStreamSegmentRequestFMP4DASH();
 
-	virtual void SetPlaybackSequenceID(uint32 PlaybackSequenceID) override;
-	virtual uint32 GetPlaybackSequenceID() const override;
+	void SetPlaybackSequenceID(uint32 PlaybackSequenceID) override;
+	uint32 GetPlaybackSequenceID() const override;
 
-	virtual void SetExecutionDelay(const FTimeValue& ExecutionDelay) override;
-	virtual FTimeValue GetExecuteAtUTCTime() const override;
+	void SetExecutionDelay(const FTimeValue& UTCNow, const FTimeValue& ExecutionDelay) override;
+	FTimeValue GetExecuteAtUTCTime() const override;
 
-	virtual EStreamType GetType() const override;
+	EStreamType GetType() const override;
 
-	virtual void GetDependentStreams(TArray<TSharedPtrTS<IStreamSegment>>& OutDependentStreams) const override;
-	virtual void GetRequestedStreams(TArray<TSharedPtrTS<IStreamSegment>>& OutRequestedStreams) override;
-	virtual void GetEndedStreams(TArray<TSharedPtrTS<IStreamSegment>>& OutAlreadyEndedStreams) override;
+	void GetDependentStreams(TArray<TSharedPtrTS<IStreamSegment>>& OutDependentStreams) const override;
+	void GetRequestedStreams(TArray<TSharedPtrTS<IStreamSegment>>& OutRequestedStreams) override;
+	void GetEndedStreams(TArray<TSharedPtrTS<IStreamSegment>>& OutAlreadyEndedStreams) override;
 
 	//! Returns the first PTS value as indicated by the media timeline. This should correspond to the actual absolute PTS of the sample.
-	virtual FTimeValue GetFirstPTS() const override;
+	FTimeValue GetFirstPTS() const override;
 
-	virtual int32 GetQualityIndex() const override;
-	virtual int32 GetBitrate() const override;
+	int32 GetQualityIndex() const override;
+	int32 GetBitrate() const override;
 
-	virtual void GetDownloadStats(Metrics::FSegmentDownloadStats& OutStats) const override;
+	void GetDownloadStats(Metrics::FSegmentDownloadStats& OutStats) const override;
+	bool GetStartupDelay(FTimeValue& OutStartTime, FTimeValue& OutTimeIntoSegment, FTimeValue& OutSegmentDuration) const override;
 
 	EStreamType												StreamType = EStreamType::Unsupported;				//!< Type of stream (video, audio, etc.)
 	FStreamCodecInformation									CodecInfo;											//!< Partial codec info as can be collected from the MPD.
@@ -151,6 +152,7 @@ private:
 		FMediaCriticalSection									MetricUpdateLock;
 		int32													ProgressReportCount = 0;
 		TSharedPtrTS<IAdaptiveStreamSelector>					StreamSelector;
+		FString													ABRAbortReason;
 
 
 		FStreamHandler();
diff --git a/Engine/Plugins/Media/ElectraPlayer/Source/ElectraPlayerRuntime/Private/Runtime/Player/HLS/ManifestBuilderHLS.cpp b/Engine/Plugins/Media/ElectraPlayer/Source/ElectraPlayerRuntime/Private/Runtime/Player/HLS/ManifestBuilderHLS.cpp
--- a/Engine/Plugins/Media/ElectraPlayer/Source/ElectraPlayerRuntime/Private/Runtime/Player/HLS/ManifestBuilderHLS.cpp
+++ b/Engine/Plugins/Media/ElectraPlayer/Source/ElectraPlayerRuntime/Private/Runtime/Player/HLS/ManifestBuilderHLS.cpp
@@ -950,9 +950,24 @@ FErrorDetail FManifestBuilderHLS::SetupVariants(FManifestHLSInternal* Manifest,
 		}
 	}
 
+
+	// Index the stream quality level map. Lower quality = lower index
+	int32 QualityIndex = 0;
+	Manifest->BandwidthToQualityIndex.KeySort([](int32 A, int32 B){return A<B;});
+	for(TMap<int32, int32>::TIterator It = Manifest->BandwidthToQualityIndex.CreateIterator(); It; ++It)
+	{
+		It.Value() = QualityIndex++;
+	}
+
 	// Set up the highest bandwidth and corresponding codec information in the adaptation sets.
 	if (Asset->VideoAdaptationSet.IsValid())
 	{
+		for(int32 i=0; i<Asset->VideoAdaptationSet->GetNumberOfRepresentations(); ++i)
+		{
+			FPlaybackAssetRepresentationHLS* Repr = static_cast<FPlaybackAssetRepresentationHLS*>(Asset->VideoAdaptationSet->GetRepresentationByIndex(i).Get());
+			Repr->QualityIndex = Manifest->BandwidthToQualityIndex[Repr->Bitrate];
+		}
+
 		FTrackMetadata& Meta = Asset->VideoAdaptationSet->Metadata;
 		for(int32 i=0; i<Meta.StreamDetails.Num(); ++i)
 		{
@@ -978,14 +993,6 @@ FErrorDetail FManifestBuilderHLS::SetupVariants(FManifestHLSInternal* Manifest,
 		Manifest->TrackMetadataAudio.Push(Meta);
 	}
 
-	// Index the stream quality level map. Lower quality = lower index
-	int32 QualityIndex = 0;
-	Manifest->BandwidthToQualityIndex.KeySort([](int32 A, int32 B){return A<B;});
-	for(TMap<int32, int32>::TIterator It = Manifest->BandwidthToQualityIndex.CreateIterator(); It; ++It)
-	{
-		It.Value() = QualityIndex++;
-	}
-
 	return FErrorDetail();
 }
 
diff --git a/Engine/Plugins/Media/ElectraPlayer/Source/ElectraPlayerRuntime/Private/Runtime/Player/HLS/ManifestBuilderHLS.h b/Engine/Plugins/Media/ElectraPlayer/Source/ElectraPlayerRuntime/Private/Runtime/Player/HLS/ManifestBuilderHLS.h
--- a/Engine/Plugins/Media/ElectraPlayer/Source/ElectraPlayerRuntime/Private/Runtime/Player/HLS/ManifestBuilderHLS.h
+++ b/Engine/Plugins/Media/ElectraPlayer/Source/ElectraPlayerRuntime/Private/Runtime/Player/HLS/ManifestBuilderHLS.h
@@ -476,23 +476,23 @@ public:
 	{ }
 	virtual ~FPlaybackAssetRepresentationHLS()
 	{ }
-	virtual FString GetUniqueIdentifier() const override
+	FString GetUniqueIdentifier() const override
 	{
 		return UniqueIdentifier;
 	}
-	virtual const FStreamCodecInformation& GetCodecInformation() const override
+	const FStreamCodecInformation& GetCodecInformation() const override
 	{
 		return CodecInformation;
 	}
-	virtual int32 GetBitrate() const override
+	int32 GetBitrate() const override
 	{
 		return Bitrate;
 	}
-	virtual int32 GetQualityIndex() const override
+	int32 GetQualityIndex() const override
 	{ 
 		return QualityIndex;
 	}
-	virtual bool CanBePlayed() const override
+	bool CanBePlayed() const override
 	{
 		return true;
 	}
@@ -508,23 +508,27 @@ public:
 	virtual ~FPlaybackAssetAdaptationSetHLS()
 	{
 	}
-	virtual FString GetUniqueIdentifier() const override
+	FString GetUniqueIdentifier() const override
 	{
 		return UniqueIdentifier;
 	}
-	virtual FString GetListOfCodecs() const override
+	FString GetListOfCodecs() const override
 	{
 		return Codecs;
 	}
-	virtual FString GetLanguage() const override
+	FString GetLanguage() const override
 	{
 		return Language;
 	}
-	virtual int32 GetNumberOfRepresentations() const override
+	int32 GetNumberOfRepresentations() const override
 	{
 		return Representations.Num();
 	}
-	virtual TSharedPtrTS<IPlaybackAssetRepresentation> GetRepresentationByIndex(int32 RepresentationIndex) const override
+	bool IsLowLatencyEnabled() const override
+	{
+		return false;
+	}
+	TSharedPtrTS<IPlaybackAssetRepresentation> GetRepresentationByIndex(int32 RepresentationIndex) const override
 	{
 		check(RepresentationIndex < Representations.Num());
 		if (RepresentationIndex < Representations.Num())
@@ -533,7 +537,7 @@ public:
 		}
 		return TSharedPtrTS<IPlaybackAssetRepresentation>();
 	}
-	virtual TSharedPtrTS<IPlaybackAssetRepresentation> GetRepresentationByUniqueIdentifier(const FString& InUniqueIdentifier) const override
+	TSharedPtrTS<IPlaybackAssetRepresentation> GetRepresentationByUniqueIdentifier(const FString& InUniqueIdentifier) const override
 	{
 		for(int32 i=0, iMax=Representations.Num(); i<iMax; ++i)
 		{
@@ -559,23 +563,23 @@ public:
 	virtual ~FTimelineMediaAssetHLS()
 	{
 	}
-	virtual FTimeRange GetTimeRange() const override
+	FTimeRange GetTimeRange() const override
 	{
 		return TimeRange;
 	}
-	virtual FTimeValue GetDuration() const override
+	FTimeValue GetDuration() const override
 	{
 		return Duration;
 	}
-	virtual FString GetAssetIdentifier() const override
+	FString GetAssetIdentifier() const override
 	{
 		return AssetIdentifier;
 	}
-	virtual FString GetUniqueIdentifier() const override
+	FString GetUniqueIdentifier() const override
 	{
 		return UniqueIdentifier;
 	}
-	virtual int32 GetNumberOfAdaptationSets(EStreamType OfStreamType) const override
+	int32 GetNumberOfAdaptationSets(EStreamType OfStreamType) const override
 	{
 		switch(OfStreamType)
 		{
@@ -587,7 +591,7 @@ public:
 				return 0;
 		}
 	}
-	virtual TSharedPtrTS<IPlaybackAssetAdaptationSet> GetAdaptationSetByTypeAndIndex(EStreamType OfStreamType, int32 AdaptationSetIndex) const override
+	TSharedPtrTS<IPlaybackAssetAdaptationSet> GetAdaptationSetByTypeAndIndex(EStreamType OfStreamType, int32 AdaptationSetIndex) const override
 	{
 		switch(OfStreamType)
 		{
@@ -608,7 +612,7 @@ public:
 		}
 	}
 
-	virtual void GetMetaData(TArray<FTrackMetadata>& OutMetadata, EStreamType OfStreamType) const override
+	void GetMetaData(TArray<FTrackMetadata>& OutMetadata, EStreamType OfStreamType) const override
 	{
 		switch(OfStreamType)
 		{
diff --git a/Engine/Plugins/Media/ElectraPlayer/Source/ElectraPlayerRuntime/Private/Runtime/Player/HLS/ManifestHLS.cpp b/Engine/Plugins/Media/ElectraPlayer/Source/ElectraPlayerRuntime/Private/Runtime/Player/HLS/ManifestHLS.cpp
--- a/Engine/Plugins/Media/ElectraPlayer/Source/ElectraPlayerRuntime/Private/Runtime/Player/HLS/ManifestHLS.cpp
+++ b/Engine/Plugins/Media/ElectraPlayer/Source/ElectraPlayerRuntime/Private/Runtime/Player/HLS/ManifestHLS.cpp
@@ -30,32 +30,33 @@ public:
 	FPlayPeriodHLS(IPlayerSessionServices* SessionServices, IPlaylistReaderHLS* PlaylistReader, TSharedPtrTS<FManifestHLSInternal> Manifest);
 	virtual ~FPlayPeriodHLS();
 
-	virtual void SetStreamPreferences(EStreamType ForStreamType, const FStreamSelectionAttributes& StreamAttributes) override;
+	void SetStreamPreferences(EStreamType ForStreamType, const FStreamSelectionAttributes& StreamAttributes) override;
 
-	virtual EReadyState GetReadyState() override;
-	virtual void Load() override;
-	virtual void PrepareForPlay() override;
-	virtual int64 GetDefaultStartingBitrate() const override;
+	EReadyState GetReadyState() override;
+	void Load() override;
+	void PrepareForPlay() override;
+	int64 GetDefaultStartingBitrate() const override;
 
-	virtual TSharedPtrTS<FBufferSourceInfo> GetSelectedStreamBufferSourceInfo(EStreamType StreamType) override;
-	virtual FString GetSelectedAdaptationSetID(EStreamType StreamType) override;
-	virtual ETrackChangeResult ChangeTrackStreamPreference(EStreamType ForStreamType, const FStreamSelectionAttributes& StreamAttributes) override;
+	TSharedPtrTS<FBufferSourceInfo> GetSelectedStreamBufferSourceInfo(EStreamType StreamType) override;
+	FString GetSelectedAdaptationSetID(EStreamType StreamType) override;
+	ETrackChangeResult ChangeTrackStreamPreference(EStreamType ForStreamType, const FStreamSelectionAttributes& StreamAttributes) override;
 
 	// TODO: need to provide metadata (duration, streams, languages, etc.)
 
-	virtual IManifest::FResult GetStartingSegment(TSharedPtrTS<IStreamSegment>& OutSegment, const FPlayerSequenceState& InSequenceState, const FPlayStartPosition& StartPosition, IManifest::ESearchType SearchType) override;
-	virtual IManifest::FResult GetContinuationSegment(TSharedPtrTS<IStreamSegment>& OutSegment, EStreamType StreamType, const FPlayerSequenceState& SequenceState, const FPlayStartPosition& StartPosition, IManifest::ESearchType SearchType) override;
-	virtual IManifest::FResult GetNextSegment(TSharedPtrTS<IStreamSegment>& OutSegment, TSharedPtrTS<const IStreamSegment> CurrentSegment, const FPlayStartOptions& Options) override;
-	virtual IManifest::FResult GetRetrySegment(TSharedPtrTS<IStreamSegment>& OutSegment, TSharedPtrTS<const IStreamSegment> CurrentSegment, const FPlayStartOptions& Options, bool bReplaceWithFillerData) override;
-	virtual IManifest::FResult GetLoopingSegment(TSharedPtrTS<IStreamSegment>& OutSegment, const FPlayerSequenceState& SequenceState, const TMultiMap<EStreamType, TSharedPtrTS<IStreamSegment>>& InFinishedSegments, const FPlayStartPosition& StartPosition, IManifest::ESearchType SearchType) override;
-	virtual void IncreaseSegmentFetchDelay(const FTimeValue& IncreaseAmount) override;
+	IManifest::FResult GetStartingSegment(TSharedPtrTS<IStreamSegment>& OutSegment, const FPlayerSequenceState& InSequenceState, const FPlayStartPosition& StartPosition, IManifest::ESearchType SearchType) override;
+	IManifest::FResult GetContinuationSegment(TSharedPtrTS<IStreamSegment>& OutSegment, EStreamType StreamType, const FPlayerSequenceState& SequenceState, const FPlayStartPosition& StartPosition, IManifest::ESearchType SearchType) override;
+	IManifest::FResult GetNextSegment(TSharedPtrTS<IStreamSegment>& OutSegment, TSharedPtrTS<const IStreamSegment> CurrentSegment, const FPlayStartOptions& Options) override;
+	IManifest::FResult GetRetrySegment(TSharedPtrTS<IStreamSegment>& OutSegment, TSharedPtrTS<const IStreamSegment> CurrentSegment, const FPlayStartOptions& Options, bool bReplaceWithFillerData) override;
+	IManifest::FResult GetLoopingSegment(TSharedPtrTS<IStreamSegment>& OutSegment, const FPlayerSequenceState& SequenceState, const TMultiMap<EStreamType, TSharedPtrTS<IStreamSegment>>& InFinishedSegments, const FPlayStartPosition& StartPosition, IManifest::ESearchType SearchType) override;
+	void IncreaseSegmentFetchDelay(const FTimeValue& IncreaseAmount) override;
 
 	// Obtains information on the stream segmentation of a particular stream starting at a given current reference segment (optional, if not given returns suitable default values).
-	virtual void GetSegmentInformation(TArray<FSegmentInformation>& OutSegmentInformation, FTimeValue& OutAverageSegmentDuration, TSharedPtrTS<const IStreamSegment> CurrentSegment, const FTimeValue& LookAheadTime, const FString& AdaptationSetID, const FString& RepresentationID) override;
+	void GetSegmentInformation(TArray<FSegmentInformation>& OutSegmentInformation, FTimeValue& OutAverageSegmentDuration, TSharedPtrTS<const IStreamSegment> CurrentSegment, const FTimeValue& LookAheadTime, const FString& AdaptationSetID, const FString& RepresentationID) override;
 
-	virtual TSharedPtrTS<ITimelineMediaAsset> GetMediaAsset() const override;
+	TSharedPtrTS<ITimelineMediaAsset> GetMediaAsset() const override;
 
-	virtual void SelectStream(const FString& AdaptationSetID, const FString& RepresentationID) override;
+	void SelectStream(const FString& AdaptationSetID, const FString& RepresentationID) override;
+	void TriggerInitSegmentPreload(const TArray<FInitSegmentPreload>& InitSegmentsToPreload) override;
 
 private:
 	struct FSegSearchParam
@@ -126,6 +127,10 @@ IManifest::EType FManifestHLS::GetPresentationType() const
 	return InternalManifest->MasterPlaylistVars.PresentationType;
 }
 
+TSharedPtrTS<const FLowLatencyDescriptor> FManifestHLS::GetLowLatencyDescriptor() const
+{
+	return nullptr; 
+}
 
 FTimeValue FManifestHLS::GetAnchorTime() const
 {
@@ -222,6 +227,26 @@ FTimeValue FManifestHLS::GetMinBufferTime() const
 	return FTimeValue().SetFromSeconds(2.0);
 }
 
+FTimeValue FManifestHLS::GetDesiredLiveLatency() const
+{
+	FTimeValue ll;
+	if (InternalManifest.IsValid() && InternalManifest->CurrentMediaAsset.IsValid() && InternalManifest->MasterPlaylistVars.PresentationType == IManifest::EType::Live)
+	{
+		FTimeRange Full, Seekable;
+		FScopeLock lock(&InternalManifest->CurrentMediaAsset->UpdateLock);
+		Full = InternalManifest->CurrentMediaAsset->TimeRange;
+		Seekable = InternalManifest->CurrentMediaAsset->SeekableTimeRange;
+		ll = Full.End - Seekable.End;
+	}
+	return ll;
+}
+
+
+TSharedPtrTS<IProducerReferenceTimeInfo> FManifestHLS::GetProducerReferenceTimeInfo(int64 ID) const
+{
+	return nullptr;
+}
+
 
 
 
@@ -251,6 +276,11 @@ IManifest::FResult FManifestHLS::FindNextPlayPeriod(TSharedPtrTS<IPlayPeriod>& O
 }
 
 
+void FManifestHLS::TriggerClockSync(IManifest::EClockSyncType InClockSyncType)
+{
+	// No-op.
+}
+
 
 /**
  * Creates a stream reader for the media segments.
@@ -1406,7 +1436,10 @@ void FPlayPeriodHLS::SelectStream(const FString& AdaptationSetID, const FString&
 	}
 }
 
-
-} // namespace Electra
+void FPlayPeriodHLS::TriggerInitSegmentPreload(const TArray<FInitSegmentPreload>& InitSegmentsToPreload)
+{
+	// No-op
+}
 
 
+} // namespace Electra
diff --git a/Engine/Plugins/Media/ElectraPlayer/Source/ElectraPlayerRuntime/Private/Runtime/Player/HLS/ManifestHLS.h b/Engine/Plugins/Media/ElectraPlayer/Source/ElectraPlayerRuntime/Private/Runtime/Player/HLS/ManifestHLS.h
--- a/Engine/Plugins/Media/ElectraPlayer/Source/ElectraPlayerRuntime/Private/Runtime/Player/HLS/ManifestHLS.h
+++ b/Engine/Plugins/Media/ElectraPlayer/Source/ElectraPlayerRuntime/Private/Runtime/Player/HLS/ManifestHLS.h
@@ -16,22 +16,26 @@ public:
 	static TSharedPtrTS<FManifestHLS> Create(IPlayerSessionServices* SessionServices, IPlaylistReaderHLS* PlaylistReader, TSharedPtrTS<FManifestHLSInternal> Manifest);
 
 	virtual ~FManifestHLS();
-	virtual EType GetPresentationType() const override;
-	virtual FTimeValue GetAnchorTime() const override;
-	virtual FTimeRange GetTotalTimeRange() const override;
-	virtual FTimeRange GetSeekableTimeRange() const override;
-	virtual FTimeRange GetPlaybackRange() const override;
-	virtual void GetSeekablePositions(TArray<FTimespan>& OutPositions) const override;
-	virtual FTimeValue GetDuration() const override;
-	virtual FTimeValue GetDefaultStartTime() const override;
-	virtual void ClearDefaultStartTime() override;
-	virtual FTimeValue GetMinBufferTime() const override;
-	virtual void GetTrackMetadata(TArray<FTrackMetadata>& OutMetadata, EStreamType StreamType) const override;
-	virtual void UpdateDynamicRefetchCounter() override;
-	virtual IStreamReader *CreateStreamReaderHandler() override;
-
-	virtual FResult FindPlayPeriod(TSharedPtrTS<IPlayPeriod>& OutPlayPeriod, const FPlayStartPosition& StartPosition, ESearchType SearchType) override;
-	virtual FResult FindNextPlayPeriod(TSharedPtrTS<IPlayPeriod>& OutPlayPeriod, TSharedPtrTS<const IStreamSegment> CurrentSegment) override;
+	EType GetPresentationType() const override;
+	TSharedPtrTS<const FLowLatencyDescriptor> GetLowLatencyDescriptor() const override;
+	FTimeValue GetAnchorTime() const override;
+	FTimeRange GetTotalTimeRange() const override;
+	FTimeRange GetSeekableTimeRange() const override;
+	FTimeRange GetPlaybackRange() const override;
+	void GetSeekablePositions(TArray<FTimespan>& OutPositions) const override;
+	FTimeValue GetDuration() const override;
+	FTimeValue GetDefaultStartTime() const override;
+	void ClearDefaultStartTime() override;
+	FTimeValue GetMinBufferTime() const override;
+	FTimeValue GetDesiredLiveLatency() const override;
+	TSharedPtrTS<IProducerReferenceTimeInfo> GetProducerReferenceTimeInfo(int64 ID) const override;
+	void GetTrackMetadata(TArray<FTrackMetadata>& OutMetadata, EStreamType StreamType) const override;
+	void UpdateDynamicRefetchCounter() override;
+	void TriggerClockSync(EClockSyncType InClockSyncType) override;
+	IStreamReader *CreateStreamReaderHandler() override;
+
+	FResult FindPlayPeriod(TSharedPtrTS<IPlayPeriod>& OutPlayPeriod, const FPlayStartPosition& StartPosition, ESearchType SearchType) override;
+	FResult FindNextPlayPeriod(TSharedPtrTS<IPlayPeriod>& OutPlayPeriod, TSharedPtrTS<const IStreamSegment> CurrentSegment) override;
 
 private:
 	FManifestHLS(IPlayerSessionServices* SessionServices, IPlaylistReaderHLS* PlaylistReader, TSharedPtrTS<FManifestHLSInternal> Manifest);
diff --git a/Engine/Plugins/Media/ElectraPlayer/Source/ElectraPlayerRuntime/Private/Runtime/Player/HLS/StreamReaderHLSfmp4.cpp b/Engine/Plugins/Media/ElectraPlayer/Source/ElectraPlayerRuntime/Private/Runtime/Player/HLS/StreamReaderHLSfmp4.cpp
--- a/Engine/Plugins/Media/ElectraPlayer/Source/ElectraPlayerRuntime/Private/Runtime/Player/HLS/StreamReaderHLSfmp4.cpp
+++ b/Engine/Plugins/Media/ElectraPlayer/Source/ElectraPlayerRuntime/Private/Runtime/Player/HLS/StreamReaderHLSfmp4.cpp
@@ -54,7 +54,7 @@ uint32 FStreamSegmentRequestHLSfmp4::GetPlaybackSequenceID() const
 	return CurrentPlaybackSequenceID;
 }
 
-void FStreamSegmentRequestHLSfmp4::SetExecutionDelay(const FTimeValue& ExecutionDelay)
+void FStreamSegmentRequestHLSfmp4::SetExecutionDelay(const FTimeValue& UTCNow, const FTimeValue& ExecutionDelay)
 {
 	// No-op for HLS.
 }
@@ -635,7 +635,6 @@ void FStreamReaderHLSfmp4::FStreamHandler::HandleRequest()
 	ds.TimeToDownload      = 0.0;
 	ds.ByteSize 		   = -1;
 	ds.NumBytesDownloaded  = 0;
-	ds.ThroughputBps	   = 0;
 	ds.bInsertedFillerData = false;
 
 	ds.MediaAssetID 	= Request->MediaAsset.IsValid() ? Request->MediaAsset->GetUniqueIdentifier() : "";
@@ -644,7 +643,6 @@ void FStreamReaderHLSfmp4::FStreamHandler::HandleRequest()
 	ds.URL  			= Request->URL;
 	ds.CDN  			= Request->CDN;
 	ds.RetryNumber  	= Request->NumOverallRetries;
-	ds.ABRState.Reset();
 
 	bIsEmptyFillerSegment = Request->bInsertFillerData;
 
@@ -657,6 +655,7 @@ void FStreamReaderHLSfmp4::FStreamHandler::HandleRequest()
 	bAbortedByABR   	   = false;
 	bAllowEarlyEmitting    = false;
 	bFillRemainingDuration = false;
+	ABRAbortReason.Empty();
 	DurationSuccessfullyRead.SetToZero();
 	DurationSuccessfullyDelivered.SetToZero();
 	AccessUnitFIFO.Clear();
@@ -797,7 +796,6 @@ void FStreamReaderHLSfmp4::FStreamHandler::HandleRequest()
 				ds.TimeToDownload      = 0.0;
 				ds.ByteSize 		   = -1;
 				ds.NumBytesDownloaded  = 0;
-				ds.ThroughputBps	   = 0;
 
 				// Clear out the current connection info which may now be populated with the init segment fetch results.
 				CurrentRequest->ConnectionInfo = {};
@@ -832,6 +830,9 @@ void FStreamReaderHLSfmp4::FStreamHandler::HandleRequest()
 				bool bIsFirstAU = true;
 				while(!bDone && !HasErrored() && !HasReadBeenAborted())
 				{
+					Metrics::FSegmentDownloadStats::FMovieChunkInfo MoofInfo;
+					MoofInfo.HeaderOffset = GetCurrentOffset();
+
 					UEMediaError parseError = MP4Parser->ParseHeader(this, this, PlayerSessionService, MP4InitSegment.Get());
 					if (parseError == UEMEDIA_ERROR_OK)
 					{
@@ -921,9 +922,16 @@ void FStreamReaderHLSfmp4::FStreamHandler::HandleRequest()
 										// There should not be any gaps!
 								// TODO: what if there are?
 										check(GetCurrentOffset() == TrackIterator->GetSampleFileOffset());
+
+										if (MoofInfo.PayloadStartOffset == 0)
+										{
+											MoofInfo.PayloadStartOffset = GetCurrentOffset();
+										}
+
 										int64 NumRead = ReadData(AccessUnit->AUData, AccessUnit->AUSize);
 										if (NumRead == AccessUnit->AUSize)
 										{
+											MoofInfo.ContentDuration += Duration;
 											DurationSuccessfullyRead += Duration;
 											NextExpectedDTS = AccessUnit->DTS + Duration;
 											LastKnownAUDuration = Duration;
@@ -991,6 +999,14 @@ void FStreamReaderHLSfmp4::FStreamHandler::HandleRequest()
 									}
 									delete TrackIterator;
 
+									MoofInfo.PayloadEndOffset = LastSuccessfulFilePos;
+									/*
+										Not adding this at the moment since it is not needed. We are not gathering download timing traces to
+										make use of the moof chunks. Segments are expected to be downloaded in one fell swoop.
+
+									ds.MovieChunkInfos.Emplace(MoveTemp(MoofInfo));
+									*/
+
 									if (Error != UEMEDIA_ERROR_OK && Error != UEMEDIA_ERROR_END_OF_STREAM)
 									{
 										// error iterating
@@ -1226,7 +1242,7 @@ void FStreamReaderHLSfmp4::FStreamHandler::HandleRequest()
 	if (bAbortedByABR)
 	{
 		// If aborted set the reason as the download failure.
-		ds.FailureReason = ds.ABRState.ProgressDecision.Reason;
+		ds.FailureReason = ABRAbortReason;
 	}
 	ds.bWasAborted  	  = bAbortedByABR;
 	ds.bWasSuccessful     = !bHasErrored && !bAbortedByABR;
@@ -1238,12 +1254,8 @@ void FStreamReaderHLSfmp4::FStreamHandler::HandleRequest()
 	ds.TimeToDownload     = (CurrentRequest->ConnectionInfo.RequestEndTime - CurrentRequest->ConnectionInfo.RequestStartTime).GetAsSeconds();
 	ds.ByteSize 		  = CurrentRequest->ConnectionInfo.ContentLength;
 	ds.NumBytesDownloaded = CurrentRequest->ConnectionInfo.BytesReadSoFar;
-	ds.ThroughputBps	  = CurrentRequest->ConnectionInfo.Throughput.GetThroughput();
-	if (ds.ThroughputBps == 0)
-	{
-		ds.ThroughputBps = ds.TimeToDownload > 0.0 ? 8 * ds.NumBytesDownloaded / ds.TimeToDownload : 0;
-	}
-	ds.bIsCachedResponse = CurrentRequest->ConnectionInfo.bIsCachedResponse;
+	ds.bIsCachedResponse  = CurrentRequest->ConnectionInfo.bIsCachedResponse;
+	CurrentRequest->ConnectionInfo.GetTimingTraces(ds.TimingTraces);
 
 	if (!bSilentCancellation)
 	{
@@ -1304,7 +1316,6 @@ int64 FStreamReaderHLSfmp4::FStreamHandler::ReadData(void* IntoBuffer, int64 Num
 
 				Metrics::FSegmentDownloadStats& ds = CurrentRequest->DownloadStats;
 				FABRDownloadProgressDecision StreamSelectorDecision = StreamSelector->ReportDownloadProgress(currentDownloadStats);
-				ds.ABRState.ProgressDecision = StreamSelectorDecision;
 				if ((StreamSelectorDecision.Flags & FABRDownloadProgressDecision::EDecisionFlags::eABR_EmitPartialData) != 0)
 				{
 					SCOPE_CYCLE_COUNTER(STAT_ElectraPlayer_HLS_StreamReader);
@@ -1339,6 +1350,7 @@ int64 FStreamReaderHLSfmp4::FStreamHandler::ReadData(void* IntoBuffer, int64 Num
 					{
 						bFillRemainingDuration = true;
 					}
+					ABRAbortReason = StreamSelectorDecision.Reason;
 					bAbortedByABR = true;
 					return -1;
 				}
diff --git a/Engine/Plugins/Media/ElectraPlayer/Source/ElectraPlayerRuntime/Private/Runtime/Player/HLS/StreamReaderHLSfmp4.h b/Engine/Plugins/Media/ElectraPlayer/Source/ElectraPlayerRuntime/Private/Runtime/Player/HLS/StreamReaderHLSfmp4.h
--- a/Engine/Plugins/Media/ElectraPlayer/Source/ElectraPlayerRuntime/Private/Runtime/Player/HLS/StreamReaderHLSfmp4.h
+++ b/Engine/Plugins/Media/ElectraPlayer/Source/ElectraPlayerRuntime/Private/Runtime/Player/HLS/StreamReaderHLSfmp4.h
@@ -19,25 +19,27 @@ public:
 	FStreamSegmentRequestHLSfmp4();
 	virtual ~FStreamSegmentRequestHLSfmp4();
 
-	virtual void SetPlaybackSequenceID(uint32 PlaybackSequenceID) override;
-	virtual uint32 GetPlaybackSequenceID() const override;
+	void SetPlaybackSequenceID(uint32 PlaybackSequenceID) override;
+	uint32 GetPlaybackSequenceID() const override;
 
-	virtual void SetExecutionDelay(const FTimeValue& ExecutionDelay) override;
-	virtual FTimeValue GetExecuteAtUTCTime() const override;
+	void SetExecutionDelay(const FTimeValue& UTCNow, const FTimeValue& ExecutionDelay) override;
+	FTimeValue GetExecuteAtUTCTime() const override;
 
-	virtual EStreamType GetType() const override;
+	EStreamType GetType() const override;
 
-	virtual void GetDependentStreams(TArray<TSharedPtrTS<IStreamSegment>>& OutDependentStreams) const override;
-	virtual void GetRequestedStreams(TArray<TSharedPtrTS<IStreamSegment>>& OutRequestedStreams) override;
-	virtual void GetEndedStreams(TArray<TSharedPtrTS<IStreamSegment>>& OutAlreadyEndedStreams) override;
+	void GetDependentStreams(TArray<TSharedPtrTS<IStreamSegment>>& OutDependentStreams) const override;
+	void GetRequestedStreams(TArray<TSharedPtrTS<IStreamSegment>>& OutRequestedStreams) override;
+	void GetEndedStreams(TArray<TSharedPtrTS<IStreamSegment>>& OutAlreadyEndedStreams) override;
 
 	//! Returns the first PTS value as indicated by the media timeline. This should correspond to the actual absolute PTS of the sample.
-	virtual FTimeValue GetFirstPTS() const override;
+	FTimeValue GetFirstPTS() const override;
 
-	virtual int32 GetQualityIndex() const override;
-	virtual int32 GetBitrate() const override;
+	int32 GetQualityIndex() const override;
+	int32 GetBitrate() const override;
 
-	virtual void GetDownloadStats(Metrics::FSegmentDownloadStats& OutStats) const override;
+	void GetDownloadStats(Metrics::FSegmentDownloadStats& OutStats) const override;
+	bool GetStartupDelay(FTimeValue& OutStartTime, FTimeValue& OutTimeIntoSegment, FTimeValue& OutSegmentDuration) const override
+	{ return false; }
 
 	FString																		URL;
 	ElectraHTTPStream::FHttpRange												Range;
@@ -221,7 +223,7 @@ private:
 		FMediaCriticalSection									MetricUpdateLock;
 		int32													ProgressReportCount = 0;
 		TSharedPtrTS<IAdaptiveStreamSelector>					StreamSelector;
-
+		FString													ABRAbortReason;
 
 		FStreamHandler();
 		virtual ~FStreamHandler();
diff --git a/Engine/Plugins/Media/ElectraPlayer/Source/ElectraPlayerRuntime/Private/Runtime/Player/Manifest.h b/Engine/Plugins/Media/ElectraPlayer/Source/ElectraPlayerRuntime/Private/Runtime/Player/Manifest.h
--- a/Engine/Plugins/Media/ElectraPlayer/Source/ElectraPlayerRuntime/Private/Runtime/Player/Manifest.h
+++ b/Engine/Plugins/Media/ElectraPlayer/Source/ElectraPlayerRuntime/Private/Runtime/Player/Manifest.h
@@ -43,6 +43,50 @@ namespace Electra
 		FPlayStartOptions	Options;
 	};
 
+	struct FLowLatencyDescriptor
+	{
+		struct FLatency
+		{
+			int64 ReferenceID = -1;
+			FTimeValue Target;
+			FTimeValue Min;
+			FTimeValue Max;
+		};
+		struct FPlayRate
+		{
+			FTimeValue Min;
+			FTimeValue Max;
+		};
+		FLatency Latency;
+		FPlayRate PlayRate;
+
+		FTimeValue GetLatencyMin() const
+		{ return Latency.Min; }
+		FTimeValue GetLatencyMax() const
+		{ return Latency.Max; }
+		FTimeValue GetLatencyTarget() const
+		{ return Latency.Target; }
+		FTimeValue GetPlayrateMin() const
+		{ return PlayRate.Min; }
+		FTimeValue GetPlayrateMax() const
+		{ return PlayRate.Max; }
+	};
+
+	struct IProducerReferenceTimeInfo
+	{
+		enum class EType
+		{
+			Encoder,
+			Captured
+		};
+		virtual FTimeValue GetWallclockTime() const = 0;
+		virtual uint64 GetPresentationTime() const = 0;
+		virtual uint32 GetID() const = 0;
+		virtual EType GetType() const = 0;
+		virtual bool GetIsInband() const = 0;
+	};
+
+
 
 	class IStreamSegment : public TSharedFromThis<IStreamSegment, ESPMode::ThreadSafe>
 	{
@@ -59,7 +103,7 @@ namespace Electra
 		 */
 		virtual uint32 GetPlaybackSequenceID() const = 0;
 
-		virtual void SetExecutionDelay(const FTimeValue& ExecutionDelay) = 0;
+		virtual void SetExecutionDelay(const FTimeValue& UTCNow, const FTimeValue& ExecutionDelay) = 0;
 
 		virtual FTimeValue GetExecuteAtUTCTime() const = 0;
 
@@ -88,6 +132,8 @@ namespace Electra
 		virtual int32 GetBitrate() const = 0;
 
 		virtual void GetDownloadStats(Metrics::FSegmentDownloadStats& OutStats) const = 0;
+
+		virtual bool GetStartupDelay(FTimeValue& OutStartTime, FTimeValue& OutTimeIntoSegment, FTimeValue& OutSegmentDuration) const = 0;
 	};
 
 
@@ -156,7 +202,7 @@ namespace Electra
 			}
 			static const TCHAR* GetTypeName(EType s)
 			{
-				switch (s)
+				switch(s)
 				{
 					case EType::Found:
 						return TEXT("Found");
@@ -180,10 +226,9 @@ namespace Electra
 			FErrorDetail	ErrorDetail;
 		};
 
-		//!
-		virtual ~IManifest() = default;
 
 
+		virtual ~IManifest() = default;
 
 		//-------------------------------------------------------------------------
 		// Presentation related functions
@@ -191,6 +236,8 @@ namespace Electra
 		//! Returns the type of this presentation, either on-demand or live.
 		virtual EType GetPresentationType() const = 0;
 
+		//! Returns the low-latency descriptor, if any. May return nullptr if there is none.
+		virtual TSharedPtrTS<const FLowLatencyDescriptor> GetLowLatencyDescriptor() const = 0;
 
 
 		//-------------------------------------------------------------------------
@@ -266,6 +313,10 @@ namespace Electra
 		//
 		virtual FTimeValue GetMinBufferTime() const = 0;
 
+		//
+		virtual TSharedPtrTS<IProducerReferenceTimeInfo> GetProducerReferenceTimeInfo(int64 ID) const = 0;
+
+		virtual FTimeValue GetDesiredLiveLatency() const = 0;
 
 		//! Needs to be called when the user has explicitly triggered a seek, including a programmatic loop back to the beginning.
 		//! For presentations with dynamic content changes (eg. DASH xlink:onRequest Periods) the content may need to be updated
@@ -273,6 +324,14 @@ namespace Electra
 		virtual void UpdateDynamicRefetchCounter() = 0;
 
 
+		enum class EClockSyncType
+		{
+			Recommended,
+			Required
+		};
+		virtual void TriggerClockSync(EClockSyncType InClockSyncType) = 0;
+
+
 		//-------------------------------------------------------------------------
 		// Stream fragment reader
 		//
@@ -348,6 +407,18 @@ namespace Electra
 			 */
 			virtual void SelectStream(const FString& AdaptationSetID, const FString& RepresentationID) = 0;
 
+			struct FInitSegmentPreload
+			{
+				FString AdaptationSetID;
+				FString RepresentationID;
+			};
+			/**
+			 * Triggers pre-loading of initialization segments.
+			 * This may be called multiple times with different streams. The implementation needs to keep track
+			 * of which init segments have already been loaded.
+			 */
+			virtual void TriggerInitSegmentPreload(const TArray<FInitSegmentPreload>& InitSegmentsToPreload) = 0;
+
 			/**
 			 * Sets up a starting segment request to begin playback at the specified time.
 			 * The streams selected through SelectStream() will be used.
diff --git a/Engine/Plugins/Media/ElectraPlayer/Source/ElectraPlayerRuntime/Private/Runtime/Player/PlaybackTimeline.h b/Engine/Plugins/Media/ElectraPlayer/Source/ElectraPlayerRuntime/Private/Runtime/Player/PlaybackTimeline.h
--- a/Engine/Plugins/Media/ElectraPlayer/Source/ElectraPlayerRuntime/Private/Runtime/Player/PlaybackTimeline.h
+++ b/Engine/Plugins/Media/ElectraPlayer/Source/ElectraPlayerRuntime/Private/Runtime/Player/PlaybackTimeline.h
@@ -44,6 +44,8 @@ namespace Electra
 
 		virtual int32 GetNumberOfRepresentations() const = 0;
 
+		virtual bool IsLowLatencyEnabled() const = 0;
+
 		virtual TSharedPtrTS<IPlaybackAssetRepresentation> GetRepresentationByIndex(int32 RepresentationIndex) const = 0;
 
 		virtual TSharedPtrTS<IPlaybackAssetRepresentation> GetRepresentationByUniqueIdentifier(const FString& UniqueIdentifier) const = 0;
diff --git a/Engine/Plugins/Media/ElectraPlayer/Source/ElectraPlayerRuntime/Private/Runtime/Player/PlayerEntityCache.cpp b/Engine/Plugins/Media/ElectraPlayer/Source/ElectraPlayerRuntime/Private/Runtime/Player/PlayerEntityCache.cpp
--- a/Engine/Plugins/Media/ElectraPlayer/Source/ElectraPlayerRuntime/Private/Runtime/Player/PlayerEntityCache.cpp
+++ b/Engine/Plugins/Media/ElectraPlayer/Source/ElectraPlayerRuntime/Private/Runtime/Player/PlayerEntityCache.cpp
@@ -93,6 +93,7 @@ void FPlayerEntityCache::CacheEntity(const FCacheItem& EntityToAdd)
 {
 	FString CachedURL = EntityToAdd.Range + EntityToAdd.URL;
 	FScopeLock ScopeLock(&Lock);
+	Cache.Remove(CachedURL);
 	Cache.Add(CachedURL, EntityToAdd);
 }
 
diff --git a/Engine/Plugins/Media/ElectraPlayer/Source/ElectraPlayerRuntime/Private/Runtime/Player/PlayerLicenseKey.h b/Engine/Plugins/Media/ElectraPlayer/Source/ElectraPlayerRuntime/Private/Runtime/Player/PlayerLicenseKey.h
--- a/Engine/Plugins/Media/ElectraPlayer/Source/ElectraPlayerRuntime/Private/Runtime/Player/PlayerLicenseKey.h
+++ b/Engine/Plugins/Media/ElectraPlayer/Source/ElectraPlayerRuntime/Private/Runtime/Player/PlayerLicenseKey.h
@@ -60,7 +60,7 @@ namespace Electra
 			if (InConnectionInfo)
 			{
 				// Have to make a dedicated copy of the connection info in order to get a copy of the retry info at this point in time.
-				ConnectionInfo.CopyFrom(*InConnectionInfo);
+				ConnectionInfo = *InConnectionInfo;
 			}
 		}
 		HTTP::FConnectionInfo	ConnectionInfo;
diff --git a/Engine/Plugins/Media/ElectraPlayer/Source/ElectraPlayerRuntime/Private/Runtime/Player/PlayerSessionServices.h b/Engine/Plugins/Media/ElectraPlayer/Source/ElectraPlayerRuntime/Private/Runtime/Player/PlayerSessionServices.h
--- a/Engine/Plugins/Media/ElectraPlayer/Source/ElectraPlayerRuntime/Private/Runtime/Player/PlayerSessionServices.h
+++ b/Engine/Plugins/Media/ElectraPlayer/Source/ElectraPlayerRuntime/Private/Runtime/Player/PlayerSessionServices.h
@@ -84,7 +84,7 @@ namespace Electra
 		 *
 		 * @return Pointer to the HTTP manager.
 		 */
-		virtual IElectraHttpManager* GetHTTPManager() = 0;
+		virtual TSharedPtrTS<IElectraHttpManager> GetHTTPManager() = 0;
 
 		/**
 		 * Returns the ABR stream selector instance.
@@ -93,11 +93,6 @@ namespace Electra
 		 */
 		virtual TSharedPtrTS<IAdaptiveStreamSelector> GetStreamSelector() = 0;
 
-		/**
-		 * Returns the current stream access unit buffer stats.
-		 */
-		virtual void GetStreamBufferStats(FAccessUnitBufferInfo& OutBufferStats, EStreamType ForStream) = 0;
-
 		/**
 		 * Returns the stream filter interface used by playlist readers to determine whether or not a stream
 		 * can be used on the platform.
diff --git a/Engine/Plugins/Media/ElectraPlayer/Source/ElectraPlayerRuntime/Private/Runtime/Player/PlaylistReader.h b/Engine/Plugins/Media/ElectraPlayer/Source/ElectraPlayerRuntime/Private/Runtime/Player/PlaylistReader.h
--- a/Engine/Plugins/Media/ElectraPlayer/Source/ElectraPlayerRuntime/Private/Runtime/Player/PlaylistReader.h
+++ b/Engine/Plugins/Media/ElectraPlayer/Source/ElectraPlayerRuntime/Private/Runtime/Player/PlaylistReader.h
@@ -102,8 +102,7 @@ namespace Electra
 				if (InConnectionInfo)
 				{
 					// Have to make a dedicated copy of the connection info in order to get a copy of the retry info at this point in time.
-					ConnectionInfo.CopyFrom(*InConnectionInfo);
-					int32 y = 0;
+					ConnectionInfo = *InConnectionInfo;
 				}
 			}
 			HTTP::FConnectionInfo	ConnectionInfo;
@@ -161,7 +160,7 @@ namespace Electra
 				if (InConnectionInfo)
 				{
 					// Have to make a dedicated copy of the connection info in order to get a copy of the retry info at this point in time.
-					ConnectionInfo.CopyFrom(*InConnectionInfo);
+					ConnectionInfo = *InConnectionInfo;
 				}
 			}
 			HTTP::FConnectionInfo	ConnectionInfo;
