ElectraPlayer: Checking for the player render clock to be running in ABR handling.

#rnx

diff --git a/Engine/Plugins/Media/ElectraPlayer/Source/ElectraPlayerRuntime/Private/Runtime/Player/ABRRules/ABRLiveStream.cpp b/Engine/Plugins/Media/ElectraPlayer/Source/ElectraPlayerRuntime/Private/Runtime/Player/ABRRules/ABRLiveStream.cpp
--- a/Engine/Plugins/Media/ElectraPlayer/Source/ElectraPlayerRuntime/Private/Runtime/Player/ABRRules/ABRLiveStream.cpp
+++ b/Engine/Plugins/Media/ElectraPlayer/Source/ElectraPlayerRuntime/Private/Runtime/Player/ABRRules/ABRLiveStream.cpp
@@ -1567,7 +1567,7 @@ IAdaptiveStreamSelector::EHandlingAction FABRLiveStream::PeriodicHandle()
 	Lock.Unlock();
 
 	// Perform catch up / slow down of playback when playing (not paused and not buffering).
-	if (!bIsPaused && !bIsBuffering)
+	if (!bIsPaused && !bIsBuffering && Info->ABRGetPlaySpeed() != FTimeValue::GetZero())
 	{
 		const EStreamType StreamType = GetPrimaryStreamType();
 		const FStreamWorkVars* WorkVars = GetWorkVars(StreamType);
diff --git a/Engine/Plugins/Media/ElectraPlayer/Source/ElectraPlayerRuntime/Private/Runtime/Player/ABRRules/ABRVoDPlus.cpp b/Engine/Plugins/Media/ElectraPlayer/Source/ElectraPlayerRuntime/Private/Runtime/Player/ABRRules/ABRVoDPlus.cpp
--- a/Engine/Plugins/Media/ElectraPlayer/Source/ElectraPlayerRuntime/Private/Runtime/Player/ABRRules/ABRVoDPlus.cpp
+++ b/Engine/Plugins/Media/ElectraPlayer/Source/ElectraPlayerRuntime/Private/Runtime/Player/ABRRules/ABRVoDPlus.cpp
@@ -861,7 +861,7 @@ IAdaptiveStreamSelector::EHandlingAction FABROnDemandPlus::PeriodicHandle()
 	IAdaptiveStreamSelector::EHandlingAction NextAction = IAdaptiveStreamSelector::EHandlingAction::None;
 
 	// Perform slow down when buffer is about to run dry?
-	if (!bIsPaused && !bIsBuffering)
+	if (!bIsPaused && !bIsBuffering && Info->ABRGetPlaySpeed() != FTimeValue::GetZero())
 	{
 		const EStreamType StreamType = GetPrimaryStreamType();
 		const FStreamWorkVars* WorkVars = GetWorkVars(StreamType);
