ElectraPlayer: Fixed static analysis issues

#rnx

diff --git a/Engine/Plugins/Media/ElectraPlayer/Source/ElectraPlayerRuntime/Private/Runtime/Player/ABRRules/ABRLiveStream.cpp b/Engine/Plugins/Media/ElectraPlayer/Source/ElectraPlayerRuntime/Private/Runtime/Player/ABRRules/ABRLiveStream.cpp
--- a/Engine/Plugins/Media/ElectraPlayer/Source/ElectraPlayerRuntime/Private/Runtime/Player/ABRRules/ABRLiveStream.cpp
+++ b/Engine/Plugins/Media/ElectraPlayer/Source/ElectraPlayerRuntime/Private/Runtime/Player/ABRRules/ABRLiveStream.cpp
@@ -1571,59 +1571,61 @@ IAdaptiveStreamSelector::EHandlingAction FABRLiveStream::PeriodicHandle()
 	{
 		const EStreamType StreamType = GetPrimaryStreamType();
 		const FStreamWorkVars* WorkVars = GetWorkVars(StreamType);
-
-		const double AvailableBufferedDuration = GetPlayablePlayerDuration(StreamType);
-		const double CurrentLatency = Info->ABRGetLatency().GetAsSeconds();
-		const double TargetLatency = lc.TargetLatency;
-		const double NetworkLatency = WorkVars->AverageLatency.GetWeightedMax(DefaultNetworkLatency);
-		const double Distance = CurrentLatency - TargetLatency;
-		const double AbsDistance = Utils::AbsoluteValue(Distance);
-		const double CurrentPlayRate = Info->ABRGetRenderRateScale();
-		const bool bOvertime = WorkVars ? WorkVars->bWentIntoOvertime : false;
-
-		int32 Gain, Trend;
-		const bool bStable = HasStableBuffer(Gain, Trend, StreamType, 0.0);
-		const bool bMaybeSpeedUp = bStable || Trend>=0;
-
-		// Slow down because of insufficient buffered data?
-		bool bBufferSlowDown = bOvertime || AvailableBufferedDuration < Utils::Max(lc.LowBufferContentBackoff, NetworkLatency);
-		// Do not slow down if we are supposed to play on the Live edge and are already too far behind.
-		if (Info->ABRShouldPlayOnLiveEdge() && Distance > lc.LowBufferMaxTargetLatencyDistance)
-		{
-			bBufferSlowDown = false;
-		}
-		if (bBufferSlowDown)
-		{
-			const double NewRate = lc.LowBufferMinMinPlayRate;
-			SetRenderRateScale(NewRate);
-		}
-		// Should we be playing on the Live edge?
-		else if (Info->ABRShouldPlayOnLiveEdge() && AbsDistance > lc.ActivationThreshold * TargetLatency)
+		if (WorkVars)
 		{
-			double NewRate = CalculateCatchupPlayRate(lc, Distance);
-			if (Utils::AbsoluteValue(NewRate - CurrentPlayRate) > lc.MinRateChangeUseThreshold)
+			const double AvailableBufferedDuration = GetPlayablePlayerDuration(StreamType);
+			const double CurrentLatency = Info->ABRGetLatency().GetAsSeconds();
+			const double TargetLatency = lc.TargetLatency;
+			const double NetworkLatency = WorkVars->AverageLatency.GetWeightedMax(DefaultNetworkLatency);
+			const double Distance = CurrentLatency - TargetLatency;
+			const double AbsDistance = Utils::AbsoluteValue(Distance);
+			const double CurrentPlayRate = Info->ABRGetRenderRateScale();
+			const bool bOvertime = WorkVars->bWentIntoOvertime;
+
+			int32 Gain, Trend;
+			const bool bStable = HasStableBuffer(Gain, Trend, StreamType, 0.0);
+			const bool bMaybeSpeedUp = bStable || Trend>=0;
+
+			// Slow down because of insufficient buffered data?
+			bool bBufferSlowDown = bOvertime || AvailableBufferedDuration < Utils::Max(lc.LowBufferContentBackoff, NetworkLatency);
+			// Do not slow down if we are supposed to play on the Live edge and are already too far behind.
+			if (Info->ABRShouldPlayOnLiveEdge() && Distance > lc.LowBufferMaxTargetLatencyDistance)
+			{
+				bBufferSlowDown = false;
+			}
+			if (bBufferSlowDown)
+			{
+				const double NewRate = lc.LowBufferMinMinPlayRate;
+				SetRenderRateScale(NewRate);
+			}
+			// Should we be playing on the Live edge?
+			else if (Info->ABRShouldPlayOnLiveEdge() && AbsDistance > lc.ActivationThreshold * TargetLatency)
 			{
-				// If we are to speed up to catch up with the Live edge we do so only if the buffer is stable.
-				if (NewRate > 1.0 && !bMaybeSpeedUp)
+				double NewRate = CalculateCatchupPlayRate(lc, Distance);
+				if (Utils::AbsoluteValue(NewRate - CurrentPlayRate) > lc.MinRateChangeUseThreshold)
 				{
-					NewRate = 1.0;
+					// If we are to speed up to catch up with the Live edge we do so only if the buffer is stable.
+					if (NewRate > 1.0 && !bMaybeSpeedUp)
+					{
+						NewRate = 1.0;
+					}
+					SetRenderRateScale(NewRate);
 				}
-				SetRenderRateScale(NewRate);
 			}
-		}
-		else if (CurrentPlayRate != 1.0)
-		{
-			SetRenderRateScale(1.0);
-		}
+			else if (CurrentPlayRate != 1.0)
+			{
+				SetRenderRateScale(1.0);
+			}
 
-		// Is the maximum latency exceeded and a jump to the Live edge required?
-		if (Info->ABRShouldPlayOnLiveEdge() && CurrentLatency > lc.MaxLatency)
-		{
-			if (!LatencyExceededVars.bLiveSeekRequestIssued)
+			// Is the maximum latency exceeded and a jump to the Live edge required?
+			if (Info->ABRShouldPlayOnLiveEdge() && CurrentLatency > lc.MaxLatency)
 			{
-				LatencyExceededVars.bLiveSeekRequestIssued = true;
-				NextAction = IAdaptiveStreamSelector::EHandlingAction::SeekToLive;
-				Info->LogMessage(IInfoLog::ELevel::Info, FString::Printf(TEXT("Maximum latency exceeded, requesting seek to Live edge")));
+				if (!LatencyExceededVars.bLiveSeekRequestIssued)
+				{
+					LatencyExceededVars.bLiveSeekRequestIssued = true;
+					NextAction = IAdaptiveStreamSelector::EHandlingAction::SeekToLive;
+					Info->LogMessage(IInfoLog::ELevel::Info, FString::Printf(TEXT("Maximum latency exceeded, requesting seek to Live edge")));
+				}
 			}
 		}
 	}
diff --git a/Engine/Plugins/Media/ElectraPlayer/Source/ElectraPlayerRuntime/Private/Runtime/Player/ABRRules/ABRVoDPlus.cpp b/Engine/Plugins/Media/ElectraPlayer/Source/ElectraPlayerRuntime/Private/Runtime/Player/ABRRules/ABRVoDPlus.cpp
--- a/Engine/Plugins/Media/ElectraPlayer/Source/ElectraPlayerRuntime/Private/Runtime/Player/ABRRules/ABRVoDPlus.cpp
+++ b/Engine/Plugins/Media/ElectraPlayer/Source/ElectraPlayerRuntime/Private/Runtime/Player/ABRRules/ABRVoDPlus.cpp
@@ -865,21 +865,23 @@ IAdaptiveStreamSelector::EHandlingAction FABROnDemandPlus::PeriodicHandle()
 	{
 		const EStreamType StreamType = GetPrimaryStreamType();
 		const FStreamWorkVars* WorkVars = GetWorkVars(StreamType);
-
-		const double AvailableBufferedDuration = GetPlayablePlayerDuration(StreamType);
-		const double NetworkLatency = WorkVars->AverageLatency.GetWeightedMax(DefaultNetworkLatency);
-		const double CurrentPlayRate = Info->ABRGetRenderRateScale();
-		const bool bOvertime = WorkVars ? WorkVars->bWentIntoOvertime : false;
-
-		// Slow down because of insufficient buffered data?
-		bool bBufferSlowDown = bOvertime || AvailableBufferedDuration < Utils::Max(SlowDownWhenBufferDurationUnder, NetworkLatency);
-		if (bBufferSlowDown)
-		{
-			SetRenderRateScale(SlowDownRateWhenBufferDurationUnder);
-		}
-		else if (CurrentPlayRate != 1.0)
+		if (WorkVars)
 		{
-			SetRenderRateScale(1.0);
+			const double AvailableBufferedDuration = GetPlayablePlayerDuration(StreamType);
+			const double NetworkLatency = WorkVars->AverageLatency.GetWeightedMax(DefaultNetworkLatency);
+			const double CurrentPlayRate = Info->ABRGetRenderRateScale();
+			const bool bOvertime = WorkVars->bWentIntoOvertime;
+
+			// Slow down because of insufficient buffered data?
+			bool bBufferSlowDown = bOvertime || AvailableBufferedDuration < Utils::Max(SlowDownWhenBufferDurationUnder, NetworkLatency);
+			if (bBufferSlowDown)
+			{
+				SetRenderRateScale(SlowDownRateWhenBufferDurationUnder);
+			}
+			else if (CurrentPlayRate != 1.0)
+			{
+				SetRenderRateScale(1.0);
+			}
 		}
 	}
 	return NextAction;
