ElectraPlayer: Added handling of short term reference picture parsing in H.265 SPS to parse files using this

diff --git a/Engine/Plugins/Media/ElectraPlayer/Source/ElectraPlayerRuntime/Private/Runtime/Utilities/UtilsMPEGVideo.cpp b/Engine/Plugins/Media/ElectraPlayer/Source/ElectraPlayerRuntime/Private/Runtime/Utilities/UtilsMPEGVideo.cpp
--- a/Engine/Plugins/Media/ElectraPlayer/Source/ElectraPlayerRuntime/Private/Runtime/Utilities/UtilsMPEGVideo.cpp
+++ b/Engine/Plugins/Media/ElectraPlayer/Source/ElectraPlayerRuntime/Private/Runtime/Utilities/UtilsMPEGVideo.cpp
@@ -866,25 +866,14 @@ namespace Electra
 			}
 
 			OutSPS.num_short_term_ref_pic_sets = FSyntaxElement::ue_v(BitReader);
-			// For now there must not be any as it it not exactly clear how to generate NumDeltaPocs[] that are needed
-			check(OutSPS.num_short_term_ref_pic_sets == 0);
-			if (OutSPS.num_short_term_ref_pic_sets != 0)
-			{
-				return false;
-			}
-		#if 0
+			uint32 unused_num_delta_pocs[64] = {0};
 			for(uint32 i=0; i<OutSPS.num_short_term_ref_pic_sets; ++i)
 			{
-				//short_term_ref_pic_set(i);
 				struct short_term_ref_pic_set
 				{
-					static void parse(uint32 stRpsIdx, uint32 num_short_term_ref_pic_sets, FBitDataStream& BitReader)
+					static bool parse(uint32 stRpsIdx, uint32 num_short_term_ref_pic_sets, uint32* num_delta_pocs, FBitDataStream& BitReader)
 					{
-						uint8 inter_ref_pic_set_prediction_flag = 0;
-						if (stRpsIdx != 0)
-						{
-							inter_ref_pic_set_prediction_flag = BitReader.GetBits(1);
-						}
+						uint32 inter_ref_pic_set_prediction_flag = stRpsIdx != 0 ? BitReader.GetBits(1) : 0;
 						if (inter_ref_pic_set_prediction_flag)
 						{
 							uint32 delta_idx_minus1 = 0;
@@ -895,13 +884,18 @@ namespace Electra
 							BitReader.SkipBits(1);	// delta_rps_sign
 							uint32 abs_delta_rps_minus1 = FSyntaxElement::ue_v(BitReader);
 							uint32 RefRpsIdx = stRpsIdx - (delta_idx_minus1 + 1);
-							uint32 NumDeltaPocs[64] = {0};
-							for(uint32 j=0; j<=NumDeltaPocs[RefRpsIdx]; ++j)
+							num_delta_pocs[stRpsIdx] = 0;
+							for(uint32 j=0; j<=num_delta_pocs[RefRpsIdx]; ++j)
 							{
 								uint8 used_by_curr_pic_flag/*[j]*/ = BitReader.GetBits(1);
+								uint8 use_delta_flag = 0;
 								if (!used_by_curr_pic_flag/*[j]*/)
 								{
-									BitReader.SkipBits(1);		// use_delta_flag[ j ]
+									use_delta_flag = BitReader.GetBits(1);
+								}
+								if (used_by_curr_pic_flag || use_delta_flag)
+								{
+									++num_delta_pocs[stRpsIdx];
 								}
 							}
 						}
@@ -909,22 +903,35 @@ namespace Electra
 						{
 							uint32 num_negative_pics = FSyntaxElement::ue_v(BitReader);
 							uint32 num_positive_pics = FSyntaxElement::ue_v(BitReader);
+
+							if ((uint64)(num_positive_pics + num_negative_pics) * 2 > BitReader.GetRemainingBits())
+							{
+								return false;
+							}
+
+							num_delta_pocs[stRpsIdx] = num_negative_pics + num_positive_pics;
+
 							for(uint32 j=0; j<num_negative_pics; ++j)
 							{
 								uint32 delta_poc_s0_minus1/*[j]*/ = FSyntaxElement::ue_v(BitReader);
+								(void)delta_poc_s0_minus1;
 								BitReader.SkipBits(1);				// used_by_curr_pic_s0_flag[ j ]
 							}
 							for(uint32 j=0; j<num_positive_pics; ++j)
 							{
 								uint32 delta_poc_s1_minus1/*[j]*/ = FSyntaxElement::ue_v(BitReader);
+								(void)delta_poc_s1_minus1;
 								BitReader.SkipBits(1);				// used_by_curr_pic_s1_flag[ j ]
 							}
 						}
+						return true;
 					}
 				};
-				short_term_ref_pic_set::parse(i, OutSPS.num_short_term_ref_pic_sets, BitReader);
+				if (!short_term_ref_pic_set::parse(i, OutSPS.num_short_term_ref_pic_sets, unused_num_delta_pocs, BitReader))
+				{
+					return false;
+				}
 			}
-		#endif
 			
 			uint8 long_term_ref_pics_present_flag = BitReader.GetBits(1);
 			if (long_term_ref_pics_present_flag)
diff --git a/Engine/Plugins/Media/ElectraPlayer/Source/ElectraPlayerRuntime/Private/Runtime/Utilities/UtilsMPEGVideo.h b/Engine/Plugins/Media/ElectraPlayer/Source/ElectraPlayerRuntime/Private/Runtime/Utilities/UtilsMPEGVideo.h
--- a/Engine/Plugins/Media/ElectraPlayer/Source/ElectraPlayerRuntime/Private/Runtime/Utilities/UtilsMPEGVideo.h
+++ b/Engine/Plugins/Media/ElectraPlayer/Source/ElectraPlayerRuntime/Private/Runtime/Utilities/UtilsMPEGVideo.h
@@ -362,7 +362,7 @@ namespace Electra
 			{
 				if (vui_parameters_present_flag && vui_timing_info_present_flag && vui_time_scale)
 				{
-					return FTimeFraction(vui_num_units_in_tick, vui_time_scale);
+					return FTimeFraction(vui_time_scale, vui_num_units_in_tick);
 				}
 				return FTimeFraction();
 			}
