ElectraPlayer: Code cleanup, replacing fastdelegate with proper UE delegates

#rnx

diff --git a/Engine/Plugins/Media/ElectraPlayer/Source/ElectraPlayerRuntime/Private/Runtime/Core/MediaThreads.cpp b/Engine/Plugins/Media/ElectraPlayer/Source/ElectraPlayerRuntime/Private/Runtime/Core/MediaThreads.cpp
--- a/Engine/Plugins/Media/ElectraPlayer/Source/ElectraPlayerRuntime/Private/Runtime/Core/MediaThreads.cpp
+++ b/Engine/Plugins/Media/ElectraPlayer/Source/ElectraPlayerRuntime/Private/Runtime/Core/MediaThreads.cpp
@@ -11,17 +11,19 @@
  * Actually only the parameters are set, but no actual thread is created.
  * We leave this to the Start() function so it becomes possible to change parameters after construction.
  *
- * @param coreAffinityMask
- * @param priority
- * @param stackSize
- * @param pName
+ * @param InCoreAffinityMask
+ * @param InPriority
+ * @param InStackSize
+ * @param InThreadName
  *
- * @return none
+ * @return New runnable object
  */
-FMediaRunnable* FMediaRunnable::Create(int32 InCoreAffinityMask, EThreadPriority InPriority, SIZE_T InStackSize, const FString& InThreadName)
+FMediaRunnable* FMediaRunnable::Create(int32 InCoreAffinityMask, EThreadPriority InPriority, uint32 InStackSize, const FString& InThreadName)
 {
 	if (InCoreAffinityMask == 0 || InCoreAffinityMask == -1)
+	{
 		InCoreAffinityMask = FPlatformAffinity::GetNoAffinityMask();
+	}
 
 	FMediaRunnable* t = new FMediaRunnable;
 	if (t)
@@ -32,7 +34,7 @@ FMediaRunnable* FMediaRunnable::Create(int32 InCoreAffinityMask, EThreadPriority
 		t->ThreadName = InThreadName;
 		t->bIsStarted = false;
 	}
-	return(t);
+	return t;
 }
 
 //-----------------------------------------------------------------------------
@@ -48,19 +50,21 @@ void FMediaRunnable::Destroy(FMediaRunnable* InMediaRunnable)
 /**
  * Creates and starts the thread with the current settings.
  *
- * @param entry
+ * @param Entry
  * @param bWaitRunning
  */
-void FMediaRunnable::Start(FMediaRunnable::StartDelegate0 Entry, bool bWaitRunning)
+void FMediaRunnable::Start(FMediaRunnable::FStartDelegate Entry, bool bWaitRunning)
 {
 	// Remember entry function
-	EntryFunction0 = Entry;
+	EntryFunction = Entry;
 
 	StartInternal();
 
 	// Optionally wait for thread startup
 	if (bWaitRunning)
+	{
 		SignalRunning->Wait();
+	}
 }
 
 
@@ -103,7 +107,7 @@ void FMediaRunnable::StartInternal()
 uint32 FMediaRunnable::Run()
 {
 	SignalRunning->Trigger();
-	EntryFunction0();
+	EntryFunction.Execute();
 	return 0;
 }
 
@@ -122,7 +126,9 @@ void FMediaRunnable::Exit()
 	DoneSignal = nullptr;
 	StateAccessMutex.Unlock();
 	if (Sig)
+	{
 		Sig->Signal();
+	}
 }
 
 
@@ -218,9 +224,13 @@ FMediaThread::FMediaThread(const char* AnsiName)
 FMediaThread::~FMediaThread()
 {
 	if (bWaitDoneOnDelete && bIsStarted)
+	{
 		SigDone.Wait();
+	}
 	if (MediaRunnable)
+	{
 		FMediaRunnable::Destroy(MediaRunnable);
+	}
 }
 
 
@@ -234,7 +244,9 @@ void FMediaThread::ThreadSetPriority(EThreadPriority InPriority)
 {
 	Priority = InPriority;
 	if (MediaRunnable)
+	{
 		MediaRunnable->ChangePriority(Priority);
+	}
 }
 
 
@@ -259,7 +271,7 @@ void FMediaThread::ThreadSetCoreAffinity(int32 coreAffinity)
  *
  * @param stackSize
  */
-void FMediaThread::ThreadSetStackSize(SIZE_T stackSize)
+void FMediaThread::ThreadSetStackSize(uint32 stackSize)
 {
 	checkf(!bIsStarted, TEXT("not while the thread is already running!"));
 	StackSize = stackSize;
@@ -327,7 +339,7 @@ void FMediaThread::ThreadWaitDone(void)
  *
  * @param entryFunction
  */
-void FMediaThread::ThreadStart(FMediaRunnable::StartDelegate0 EntryFunction)
+void FMediaThread::ThreadStart(FMediaRunnable::FStartDelegate EntryFunction)
 {
 	MediaRunnable = FMediaRunnable::Create(CoreAffinity, Priority, StackSize, ThreadName);
 	MediaRunnable->SetDoneSignal(&SigDone);
@@ -345,7 +357,9 @@ void FMediaThread::ThreadReset()
 {
 	check(!bIsStarted || (bIsStarted && SigDone.IsSignaled()));
 	if (MediaRunnable)
+	{
 		FMediaRunnable::Destroy(MediaRunnable);
+	}
 	MediaRunnable = nullptr;
 	bIsStarted = false;
 	SigDone.Reset();
diff --git a/Engine/Plugins/Media/ElectraPlayer/Source/ElectraPlayerRuntime/Private/Runtime/Core/MediaThreads.h b/Engine/Plugins/Media/ElectraPlayer/Source/ElectraPlayerRuntime/Private/Runtime/Core/MediaThreads.h
--- a/Engine/Plugins/Media/ElectraPlayer/Source/ElectraPlayerRuntime/Private/Runtime/Core/MediaThreads.h
+++ b/Engine/Plugins/Media/ElectraPlayer/Source/ElectraPlayerRuntime/Private/Runtime/Core/MediaThreads.h
@@ -13,6 +13,7 @@
 #include "HAL/Runnable.h"
 #include "HAL/PlatformProcess.h"
 #include "HAL/CriticalSection.h"
+#include "Delegates/Delegate.h"
 
 #include "Containers/UnrealString.h"
 #include "Containers/StringConv.h"
@@ -25,6 +26,7 @@
 class FMediaRunnable : private TMediaNoncopyable<FMediaRunnable>, public FRunnable
 {
 public:
+	DECLARE_DELEGATE(FStartDelegate);
 
 	//! Common thread configuration parameters.
 	struct Param
@@ -37,16 +39,14 @@ public:
 			CoreAffinity = -1;
 		}
 		EThreadPriority		Priority;
-		SIZE_T				StackSize;
+		uint32				StackSize;
 		int32				CoreAffinity;
 	};
 
-	static FMediaRunnable* Create(int32 CoreAffinityMask, EThreadPriority Priority, SIZE_T StackSize, const FString& InThreadName);
+	static FMediaRunnable* Create(int32 CoreAffinityMask, EThreadPriority Priority, uint32 StackSize, const FString& InThreadName);
 	static void Destroy(FMediaRunnable* Thread);
 
-	typedef Electra::FastDelegate0<> StartDelegate0;
-
-	void Start(StartDelegate0 Entry, bool bWaitRunning = false);
+	void Start(FStartDelegate Entry, bool bWaitRunning = false);
 
 	void SetDoneSignal(FMediaEvent* DoneSignal);
 
@@ -60,13 +60,13 @@ public:
 	}
 
 	//! Returns the size of the stack.
-	SIZE_T StackSizeGet() const
+	uint32 StackSizeGet() const
 	{
 		return StackSize;
 	}
 
 	//! Returns the default stack size passed to Startup()
-	static SIZE_T StackSizeGetDefault()
+	static uint32 StackSizeGetDefault()
 	{
 		return 0;
 	}
@@ -95,7 +95,7 @@ private:
 	FString					ThreadName;
 	FCriticalSection		StateAccessMutex;
 	FEvent*					SignalRunning;
-	StartDelegate0			EntryFunction0;
+	FStartDelegate			EntryFunction;
 	FMediaEvent*			DoneSignal;
 	bool					bIsStarted;
 
@@ -141,7 +141,7 @@ public:
 	void ThreadSetCoreAffinity(int32 CoreAffinity);
 
 	//! Set a thread stack size other than the one given to the constructor before starting the thread.
-	void ThreadSetStackSize(SIZE_T StackSize);
+	void ThreadSetStackSize(uint32 StackSize);
 
 	//! Set a thread name other than the one given to the constructor before starting the thread.
 	void ThreadSetName(const char* InAnsiThreadName);
@@ -156,7 +156,7 @@ public:
 	void ThreadWaitDone();
 
 	//! Starts the thread at the given function void(*)(void)
-	void ThreadStart(FMediaRunnable::StartDelegate0 EntryFunction);
+	void ThreadStart(FMediaRunnable::FStartDelegate EntryFunction);
 
 	//! Resets the thread to be started again. Must have waited for thread termination using ThreadWaitDone() first!
 	void ThreadReset();
@@ -166,13 +166,9 @@ private:
 	FMediaRunnable*					MediaRunnable;
 	EThreadPriority					Priority;
 	int32							CoreAffinity;
-	SIZE_T							StackSize;
+	uint32							StackSize;
 	FString 						ThreadName;
 	bool							bIsStarted;
 	bool							bWaitDoneOnDelete;
 
 };
-
-
-
-
diff --git a/Engine/Plugins/Media/ElectraPlayer/Source/ElectraPlayerRuntime/Private/Runtime/Decoder/Android/AudioDecoderAAC_Android.cpp b/Engine/Plugins/Media/ElectraPlayer/Source/ElectraPlayerRuntime/Private/Runtime/Decoder/Android/AudioDecoderAAC_Android.cpp
--- a/Engine/Plugins/Media/ElectraPlayer/Source/ElectraPlayerRuntime/Private/Runtime/Decoder/Android/AudioDecoderAAC_Android.cpp
+++ b/Engine/Plugins/Media/ElectraPlayer/Source/ElectraPlayerRuntime/Private/Runtime/Decoder/Android/AudioDecoderAAC_Android.cpp
@@ -423,7 +423,7 @@ void FAudioDecoderAAC::StartThread()
 	ThreadSetPriority(Config.ThreadConfig.Decoder.Priority);
 	ThreadSetStackSize(Config.ThreadConfig.Decoder.StackSize);
 	ThreadSetCoreAffinity(Config.ThreadConfig.Decoder.CoreAffinity);
-	ThreadStart(Electra::MakeDelegate(this, &FAudioDecoderAAC::WorkerThread));
+	ThreadStart(FMediaRunnable::FStartDelegate::CreateRaw(this, &FAudioDecoderAAC::WorkerThread));
 	bThreadStarted = true;
 }
 
diff --git a/Engine/Plugins/Media/ElectraPlayer/Source/ElectraPlayerRuntime/Private/Runtime/Decoder/Android/VideoDecoderH264_Android.cpp b/Engine/Plugins/Media/ElectraPlayer/Source/ElectraPlayerRuntime/Private/Runtime/Decoder/Android/VideoDecoderH264_Android.cpp
--- a/Engine/Plugins/Media/ElectraPlayer/Source/ElectraPlayerRuntime/Private/Runtime/Decoder/Android/VideoDecoderH264_Android.cpp
+++ b/Engine/Plugins/Media/ElectraPlayer/Source/ElectraPlayerRuntime/Private/Runtime/Decoder/Android/VideoDecoderH264_Android.cpp
@@ -504,7 +504,7 @@ void FVideoDecoderH264::StartThread()
 	ThreadSetPriority(Config.ThreadConfig.Decoder.Priority);
 	ThreadSetCoreAffinity(Config.ThreadConfig.Decoder.CoreAffinity);
 	ThreadSetStackSize(Config.ThreadConfig.Decoder.StackSize);
-	ThreadStart(Electra::MakeDelegate(this, &FVideoDecoderH264::WorkerThread));
+	ThreadStart(FMediaRunnable::FStartDelegate::CreateRaw(this, &FVideoDecoderH264::WorkerThread));
 	bThreadStarted = true;
 }
 
diff --git a/Engine/Plugins/Media/ElectraPlayer/Source/ElectraPlayerRuntime/Private/Runtime/Decoder/Android/VideoDecoderH265_Android.cpp b/Engine/Plugins/Media/ElectraPlayer/Source/ElectraPlayerRuntime/Private/Runtime/Decoder/Android/VideoDecoderH265_Android.cpp
--- a/Engine/Plugins/Media/ElectraPlayer/Source/ElectraPlayerRuntime/Private/Runtime/Decoder/Android/VideoDecoderH265_Android.cpp
+++ b/Engine/Plugins/Media/ElectraPlayer/Source/ElectraPlayerRuntime/Private/Runtime/Decoder/Android/VideoDecoderH265_Android.cpp
@@ -487,7 +487,7 @@ void FVideoDecoderH265::SetResourceDelegate(const TSharedPtr<IVideoDecoderResour
  */
 void FVideoDecoderH265::StartThread()
 {
-	ThreadStart(Electra::MakeDelegate(this, &FVideoDecoderH265::WorkerThread));
+	ThreadStart(FMediaRunnable::FStartDelegate::CreateRaw(this, &FVideoDecoderH265::WorkerThread));
 	bThreadStarted = true;
 }
 
diff --git a/Engine/Plugins/Media/ElectraPlayer/Source/ElectraPlayerRuntime/Private/Runtime/Decoder/Apple/AudioDecoderAAC_Apple.cpp b/Engine/Plugins/Media/ElectraPlayer/Source/ElectraPlayerRuntime/Private/Runtime/Decoder/Apple/AudioDecoderAAC_Apple.cpp
--- a/Engine/Plugins/Media/ElectraPlayer/Source/ElectraPlayerRuntime/Private/Runtime/Decoder/Apple/AudioDecoderAAC_Apple.cpp
+++ b/Engine/Plugins/Media/ElectraPlayer/Source/ElectraPlayerRuntime/Private/Runtime/Decoder/Apple/AudioDecoderAAC_Apple.cpp
@@ -447,7 +447,7 @@ void FAudioDecoderAAC::StartThread()
 	ThreadSetPriority(Config.ThreadConfig.Decoder.Priority);
 	ThreadSetStackSize(Config.ThreadConfig.Decoder.StackSize);
 	ThreadSetCoreAffinity(Config.ThreadConfig.Decoder.CoreAffinity);
-	ThreadStart(Electra::MakeDelegate(this, &FAudioDecoderAAC::WorkerThread));
+	ThreadStart(FMediaRunnable::FStartDelegate::CreateRaw(this, &FAudioDecoderAAC::WorkerThread));
 	bThreadStarted = true;
 }
 
diff --git a/Engine/Plugins/Media/ElectraPlayer/Source/ElectraPlayerRuntime/Private/Runtime/Decoder/Apple/VideoDecoderH264_Apple.cpp b/Engine/Plugins/Media/ElectraPlayer/Source/ElectraPlayerRuntime/Private/Runtime/Decoder/Apple/VideoDecoderH264_Apple.cpp
--- a/Engine/Plugins/Media/ElectraPlayer/Source/ElectraPlayerRuntime/Private/Runtime/Decoder/Apple/VideoDecoderH264_Apple.cpp
+++ b/Engine/Plugins/Media/ElectraPlayer/Source/ElectraPlayerRuntime/Private/Runtime/Decoder/Apple/VideoDecoderH264_Apple.cpp
@@ -533,7 +533,7 @@ void FVideoDecoderH264::StartThread()
 	ThreadSetPriority(Config.ThreadConfig.Decoder.Priority);
 	ThreadSetCoreAffinity(Config.ThreadConfig.Decoder.CoreAffinity);
 	ThreadSetStackSize(Config.ThreadConfig.Decoder.StackSize);
-	ThreadStart(Electra::MakeDelegate(this, &FVideoDecoderH264::WorkerThread));
+	ThreadStart(FMediaRunnable::FStartDelegate::CreateRaw(this, &FVideoDecoderH264::WorkerThread));
 	bThreadStarted = true;
 }
 
diff --git a/Engine/Plugins/Media/ElectraPlayer/Source/ElectraPlayerRuntime/Private/Runtime/Decoder/Apple/VideoDecoderH265_Apple.cpp b/Engine/Plugins/Media/ElectraPlayer/Source/ElectraPlayerRuntime/Private/Runtime/Decoder/Apple/VideoDecoderH265_Apple.cpp
--- a/Engine/Plugins/Media/ElectraPlayer/Source/ElectraPlayerRuntime/Private/Runtime/Decoder/Apple/VideoDecoderH265_Apple.cpp
+++ b/Engine/Plugins/Media/ElectraPlayer/Source/ElectraPlayerRuntime/Private/Runtime/Decoder/Apple/VideoDecoderH265_Apple.cpp
@@ -520,7 +520,7 @@ void FVideoDecoderH265::SetResourceDelegate(const TSharedPtr<IVideoDecoderResour
  */
 void FVideoDecoderH265::StartThread()
 {
-	ThreadStart(Electra::MakeDelegate(this, &FVideoDecoderH265::WorkerThread));
+	ThreadStart(FMediaRunnable::FStartDelegate::CreateRaw(this, &FVideoDecoderH265::WorkerThread));
 	bThreadStarted = true;
 }
 
diff --git a/Engine/Plugins/Media/ElectraPlayer/Source/ElectraPlayerRuntime/Private/Runtime/Decoder/DX/AudioDecoderAAC_DX.cpp b/Engine/Plugins/Media/ElectraPlayer/Source/ElectraPlayerRuntime/Private/Runtime/Decoder/DX/AudioDecoderAAC_DX.cpp
--- a/Engine/Plugins/Media/ElectraPlayer/Source/ElectraPlayerRuntime/Private/Runtime/Decoder/DX/AudioDecoderAAC_DX.cpp
+++ b/Engine/Plugins/Media/ElectraPlayer/Source/ElectraPlayerRuntime/Private/Runtime/Decoder/DX/AudioDecoderAAC_DX.cpp
@@ -453,7 +453,7 @@ namespace Electra
 		ThreadSetPriority(Config.ThreadConfig.Decoder.Priority);
 		ThreadSetStackSize(Config.ThreadConfig.Decoder.StackSize);
 		ThreadSetCoreAffinity(Config.ThreadConfig.Decoder.CoreAffinity);
-		ThreadStart(Electra::MakeDelegate(this, &FAudioDecoderAAC::WorkerThread));
+		ThreadStart(FMediaRunnable::FStartDelegate::CreateRaw(this, &FAudioDecoderAAC::WorkerThread));
 		bThreadStarted = true;
 	}
 
diff --git a/Engine/Plugins/Media/ElectraPlayer/Source/ElectraPlayerRuntime/Private/Runtime/Decoder/DX/VideoDecoderH264_DX.cpp b/Engine/Plugins/Media/ElectraPlayer/Source/ElectraPlayerRuntime/Private/Runtime/Decoder/DX/VideoDecoderH264_DX.cpp
--- a/Engine/Plugins/Media/ElectraPlayer/Source/ElectraPlayerRuntime/Private/Runtime/Decoder/DX/VideoDecoderH264_DX.cpp
+++ b/Engine/Plugins/Media/ElectraPlayer/Source/ElectraPlayerRuntime/Private/Runtime/Decoder/DX/VideoDecoderH264_DX.cpp
@@ -284,7 +284,7 @@ void FVideoDecoderH264::StartThread()
 	ThreadSetPriority(Config.ThreadConfig.Decoder.Priority);
 	ThreadSetCoreAffinity(Config.ThreadConfig.Decoder.CoreAffinity);
 	ThreadSetStackSize(Config.ThreadConfig.Decoder.StackSize);
-	ThreadStart(Electra::MakeDelegate(this, &FVideoDecoderH264::WorkerThread));
+	ThreadStart(FMediaRunnable::FStartDelegate::CreateRaw(this, &FVideoDecoderH264::WorkerThread));
 	bThreadStarted = true;
 }
 
diff --git a/Engine/Plugins/Media/ElectraPlayer/Source/ElectraPlayerRuntime/Private/Runtime/Decoder/DX/VideoDecoderH265_DX.cpp b/Engine/Plugins/Media/ElectraPlayer/Source/ElectraPlayerRuntime/Private/Runtime/Decoder/DX/VideoDecoderH265_DX.cpp
--- a/Engine/Plugins/Media/ElectraPlayer/Source/ElectraPlayerRuntime/Private/Runtime/Decoder/DX/VideoDecoderH265_DX.cpp
+++ b/Engine/Plugins/Media/ElectraPlayer/Source/ElectraPlayerRuntime/Private/Runtime/Decoder/DX/VideoDecoderH265_DX.cpp
@@ -261,7 +261,7 @@ void FVideoDecoderH265::SetResourceDelegate(const TSharedPtr<IVideoDecoderResour
  */
 void FVideoDecoderH265::StartThread()
 {
-	ThreadStart(Electra::MakeDelegate(this, &FVideoDecoderH265::WorkerThread));
+	ThreadStart(FMediaRunnable::FStartDelegate::CreateRaw(this, &FVideoDecoderH265::WorkerThread));
 	bThreadStarted = true;
 }
 
diff --git a/Engine/Plugins/Media/ElectraPlayer/Source/ElectraPlayerRuntime/Private/Runtime/HTTP/HTTPManager.cpp b/Engine/Plugins/Media/ElectraPlayer/Source/ElectraPlayerRuntime/Private/Runtime/HTTP/HTTPManager.cpp
--- a/Engine/Plugins/Media/ElectraPlayer/Source/ElectraPlayerRuntime/Private/Runtime/HTTP/HTTPManager.cpp
+++ b/Engine/Plugins/Media/ElectraPlayer/Source/ElectraPlayerRuntime/Private/Runtime/HTTP/HTTPManager.cpp
@@ -764,7 +764,7 @@ namespace Electra
 			TSharedPtrTS<FProgressListener> ProgressListener = Request->ProgressListener.Pin();
 			if (ProgressListener.IsValid())
 			{
-				if (!ProgressListener->CompletionDelegate.empty())
+				if (ProgressListener->CompletionDelegate.IsBound())
 				{
 					ProgressListener->CallCompletionDelegate(Request.Get());
 				}
@@ -788,7 +788,7 @@ namespace Electra
 				TSharedPtrTS<FProgressListener> ProgressListener = It.Value()->ProgressListener.Pin();
 				if (ProgressListener.IsValid())
 				{
-					if (!ProgressListener->ProgressDelegate.empty())
+					if (ProgressListener->ProgressDelegate.IsBound())
 					{
 						int32 Result = ProgressListener->CallProgressDelegate(It.Value().Get());
 						// Did the progress callback ask to abort the download?
diff --git a/Engine/Plugins/Media/ElectraPlayer/Source/ElectraPlayerRuntime/Private/Runtime/HTTP/HTTPManager.h b/Engine/Plugins/Media/ElectraPlayer/Source/ElectraPlayerRuntime/Private/Runtime/HTTP/HTTPManager.h
--- a/Engine/Plugins/Media/ElectraPlayer/Source/ElectraPlayerRuntime/Private/Runtime/HTTP/HTTPManager.h
+++ b/Engine/Plugins/Media/ElectraPlayer/Source/ElectraPlayerRuntime/Private/Runtime/HTTP/HTTPManager.h
@@ -3,6 +3,7 @@
 #pragma once
 
 #include "PlayerCore.h"
+#include "Delegates/Delegate.h"
 
 #include "StreamDataBuffer.h"
 #include "OptionalValue.h"
@@ -180,10 +181,11 @@ namespace Electra
 
 		struct FRequest;
 
-		typedef Electra::FastDelegate1<const FRequest*, int32> FProgressDelegate;
-		typedef Electra::FastDelegate1<const FRequest*> FCompletionDelegate;
 		struct FProgressListener
 		{
+			DECLARE_DELEGATE_RetVal_OneParam(int32, FProgressDelegate, const FRequest*);
+			DECLARE_DELEGATE_OneParam(FCompletionDelegate, const FRequest*);
+
 			~FProgressListener()
 			{
 				Clear();
@@ -191,24 +193,24 @@ namespace Electra
 			void Clear()
 			{
 				FMediaCriticalSection::ScopedLock lock(Lock);
-				ProgressDelegate.clear();
-				CompletionDelegate.clear();
+				ProgressDelegate.Unbind();
+				CompletionDelegate.Unbind();
 			}
 			int32 CallProgressDelegate(const FRequest* Request)
 			{
 				FMediaCriticalSection::ScopedLock lock(Lock);
-				if (!ProgressDelegate.empty())
+				if (ProgressDelegate.IsBound())
 				{
-					return ProgressDelegate(Request);
+					return ProgressDelegate.Execute(Request);
 				}
 				return 0;
 			}
 			void CallCompletionDelegate(const FRequest* Request)
 			{
 				FMediaCriticalSection::ScopedLock lock(Lock);
-				if (!CompletionDelegate.empty())
+				if (CompletionDelegate.IsBound())
 				{
-					CompletionDelegate(Request);
+					CompletionDelegate.Execute(Request);
 				}
 			}
 			FMediaCriticalSection	Lock;
diff --git a/Engine/Plugins/Media/ElectraPlayer/Source/ElectraPlayerRuntime/Private/Runtime/Player/AdaptiveStreamingPlayer.cpp b/Engine/Plugins/Media/ElectraPlayer/Source/ElectraPlayerRuntime/Private/Runtime/Player/AdaptiveStreamingPlayer.cpp
--- a/Engine/Plugins/Media/ElectraPlayer/Source/ElectraPlayerRuntime/Private/Runtime/Player/AdaptiveStreamingPlayer.cpp
+++ b/Engine/Plugins/Media/ElectraPlayer/Source/ElectraPlayerRuntime/Private/Runtime/Player/AdaptiveStreamingPlayer.cpp
@@ -4930,7 +4930,7 @@ void FAdaptiveStreamingPlayerWorkerThread::StartWorkerThread()
 	{
 		WorkerThread.ThreadSetName("ElectraPlayer::SharedWorker");
 	}
-	WorkerThread.ThreadStart(Electra::MakeDelegate(this, &FAdaptiveStreamingPlayerWorkerThread::WorkerThreadFN));
+	WorkerThread.ThreadStart(FMediaRunnable::FStartDelegate::CreateRaw(this, &FAdaptiveStreamingPlayerWorkerThread::WorkerThreadFN));
 }
 
 void FAdaptiveStreamingPlayerWorkerThread::StopWorkerThread()
@@ -5013,7 +5013,7 @@ FAdaptiveStreamingPlayerEventHandler::~FAdaptiveStreamingPlayerEventHandler()
 void FAdaptiveStreamingPlayerEventHandler::StartWorkerThread()
 {
 	WorkerThread.ThreadSetName("ElectraPlayer::EventDispatch");
-	WorkerThread.ThreadStart(Electra::MakeDelegate(this, &FAdaptiveStreamingPlayerEventHandler::WorkerThreadFN));
+	WorkerThread.ThreadStart(FMediaRunnable::FStartDelegate::CreateRaw(this, &FAdaptiveStreamingPlayerEventHandler::WorkerThreadFN));
 }
 
 void FAdaptiveStreamingPlayerEventHandler::StopWorkerThread()
diff --git a/Engine/Plugins/Media/ElectraPlayer/Source/ElectraPlayerRuntime/Private/Runtime/Player/AdaptiveStreamingPlayerResourceRequest.cpp b/Engine/Plugins/Media/ElectraPlayer/Source/ElectraPlayerRuntime/Private/Runtime/Player/AdaptiveStreamingPlayerResourceRequest.cpp
--- a/Engine/Plugins/Media/ElectraPlayer/Source/ElectraPlayerRuntime/Private/Runtime/Player/AdaptiveStreamingPlayerResourceRequest.cpp
+++ b/Engine/Plugins/Media/ElectraPlayer/Source/ElectraPlayerRuntime/Private/Runtime/Player/AdaptiveStreamingPlayerResourceRequest.cpp
@@ -14,8 +14,8 @@ FHTTPResourceRequest::FHTTPResourceRequest()
 	Request = MakeSharedTS<IElectraHttpManager::FRequest>();
 	ReceiveBuffer = MakeSharedTS<IElectraHttpManager::FReceiveBuffer>();
 	ProgressListener = MakeSharedTS<IElectraHttpManager::FProgressListener>();
-	ProgressListener->ProgressDelegate   = Electra::MakeDelegate(this, &FHTTPResourceRequest::HTTPProgressCallback);
-	ProgressListener->CompletionDelegate = Electra::MakeDelegate(this, &FHTTPResourceRequest::HTTPCompletionCallback);
+	ProgressListener->ProgressDelegate   = IElectraHttpManager::FProgressListener::FProgressDelegate::CreateRaw(this, &FHTTPResourceRequest::HTTPProgressCallback);
+	ProgressListener->CompletionDelegate = IElectraHttpManager::FProgressListener::FCompletionDelegate::CreateRaw(this, &FHTTPResourceRequest::HTTPCompletionCallback);
 	Request->ProgressListener = ProgressListener;
 	Request->ReceiveBuffer = ReceiveBuffer;
 }
diff --git a/Engine/Plugins/Media/ElectraPlayer/Source/ElectraPlayerRuntime/Private/Runtime/Player/DASH/StreamReaderFMP4DASH.cpp b/Engine/Plugins/Media/ElectraPlayer/Source/ElectraPlayerRuntime/Private/Runtime/Player/DASH/StreamReaderFMP4DASH.cpp
--- a/Engine/Plugins/Media/ElectraPlayer/Source/ElectraPlayerRuntime/Private/Runtime/Player/DASH/StreamReaderFMP4DASH.cpp
+++ b/Engine/Plugins/Media/ElectraPlayer/Source/ElectraPlayerRuntime/Private/Runtime/Player/DASH/StreamReaderFMP4DASH.cpp
@@ -278,7 +278,7 @@ IStreamReader::EAddResult FStreamReaderFMP4DASH::AddRequest(uint32 CurrentPlayba
 		{
 			if (!Handler->bWasStarted)
 			{
-				Handler->ThreadStart(Electra::MakeDelegate(Handler, &FStreamHandler::WorkerThread));
+				Handler->ThreadStart(FMediaRunnable::FStartDelegate::CreateRaw(Handler, &FStreamHandler::WorkerThread));
 				Handler->bWasStarted = true;
 			}
 		}
@@ -957,8 +957,8 @@ void FStreamReaderFMP4DASH::FStreamHandler::HandleRequest()
 
 			// Start downloading the segment.
 			TSharedPtrTS<IElectraHttpManager::FProgressListener>	ProgressListener(new IElectraHttpManager::FProgressListener);
-			ProgressListener->ProgressDelegate   = Electra::MakeDelegate(this, &FStreamHandler::HTTPProgressCallback);
-			ProgressListener->CompletionDelegate = Electra::MakeDelegate(this, &FStreamHandler::HTTPCompletionCallback);
+			ProgressListener->ProgressDelegate   = IElectraHttpManager::FProgressListener::FProgressDelegate::CreateRaw(this, &FStreamHandler::HTTPProgressCallback);
+			ProgressListener->CompletionDelegate = IElectraHttpManager::FProgressListener::FCompletionDelegate::CreateRaw(this, &FStreamHandler::HTTPCompletionCallback);
 			TSharedPtrTS<IElectraHttpManager::FRequest> HTTP(new IElectraHttpManager::FRequest);
 			HTTP->Parameters.URL   = RequestURL;
 			HTTP->ReceiveBuffer    = ReadBuffer.ReceiveBuffer;
diff --git a/Engine/Plugins/Media/ElectraPlayer/Source/ElectraPlayerRuntime/Private/Runtime/Player/HLS/PlaylistReaderHLS.cpp b/Engine/Plugins/Media/ElectraPlayer/Source/ElectraPlayerRuntime/Private/Runtime/Player/HLS/PlaylistReaderHLS.cpp
--- a/Engine/Plugins/Media/ElectraPlayer/Source/ElectraPlayerRuntime/Private/Runtime/Player/HLS/PlaylistReaderHLS.cpp
+++ b/Engine/Plugins/Media/ElectraPlayer/Source/ElectraPlayerRuntime/Private/Runtime/Player/HLS/PlaylistReaderHLS.cpp
@@ -323,7 +323,7 @@ void FPlaylistReaderHLS::Initialize(IPlayerSessionServices* InPlayerSessionServi
 {
 	PlayerSessionServices = InPlayerSessionServices;
 	ProgressListener = MakeSharedTS<IElectraHttpManager::FProgressListener>();
-	ProgressListener->CompletionDelegate = Electra::MakeDelegate(this, &FPlaylistReaderHLS::HTTPCompletionCallback);
+	ProgressListener->CompletionDelegate = IElectraHttpManager::FProgressListener::FCompletionDelegate::CreateRaw(this, &FPlaylistReaderHLS::HTTPCompletionCallback);
 }
 
 void FPlaylistReaderHLS::Close()
diff --git a/Engine/Plugins/Media/ElectraPlayer/Source/ElectraPlayerRuntime/Private/Runtime/Player/HLS/StreamReaderHLSfmp4.cpp b/Engine/Plugins/Media/ElectraPlayer/Source/ElectraPlayerRuntime/Private/Runtime/Player/HLS/StreamReaderHLSfmp4.cpp
--- a/Engine/Plugins/Media/ElectraPlayer/Source/ElectraPlayerRuntime/Private/Runtime/Player/HLS/StreamReaderHLSfmp4.cpp
+++ b/Engine/Plugins/Media/ElectraPlayer/Source/ElectraPlayerRuntime/Private/Runtime/Player/HLS/StreamReaderHLSfmp4.cpp
@@ -242,7 +242,7 @@ IStreamReader::EAddResult FStreamReaderHLSfmp4::AddRequest(uint32 CurrentPlaybac
 	{
 		if (!Handler->bWasStarted)
 		{
-			Handler->ThreadStart(Electra::MakeDelegate(Handler, &FStreamHandler::WorkerThread));
+			Handler->ThreadStart(FMediaRunnable::FStartDelegate::CreateRaw(Handler, &FStreamHandler::WorkerThread));
 			Handler->bWasStarted = true;
 		}
 		Handler->bRequestCanceled = false;
@@ -439,8 +439,8 @@ FStreamReaderHLSfmp4::FStreamHandler::ELicenseKeyResult FStreamReaderHLSfmp4::FS
 			if (!bHaveStaticResponse)
 			{
 				TSharedPtrTS<IElectraHttpManager::FProgressListener>	ProgressListener(new IElectraHttpManager::FProgressListener);
-				ProgressListener->ProgressDelegate   = Electra::MakeDelegate(this, &FStreamHandler::HTTPProgressCallback);
-				ProgressListener->CompletionDelegate = Electra::MakeDelegate(this, &FStreamHandler::HTTPCompletionCallback);
+				ProgressListener->ProgressDelegate   = IElectraHttpManager::FProgressListener::FProgressDelegate::CreateRaw(this, &FStreamHandler::HTTPProgressCallback);
+				ProgressListener->CompletionDelegate = IElectraHttpManager::FProgressListener::FCompletionDelegate::CreateRaw(this, &FStreamHandler::HTTPCompletionCallback);
 				HTTP->ProgressListener = ProgressListener;
 				PlayerSessionService->GetHTTPManager()->AddRequest(HTTP, false);
 				while(!HasReadBeenAborted())
@@ -521,8 +521,8 @@ FStreamReaderHLSfmp4::FStreamHandler::EInitSegmentResult FStreamReaderHLSfmp4::F
 
 			TSharedPtrTS<IElectraHttpManager::FRequest> 			HTTP(new IElectraHttpManager::FRequest);
 			TSharedPtrTS<IElectraHttpManager::FProgressListener>	ProgressListener(new IElectraHttpManager::FProgressListener);
-			ProgressListener->ProgressDelegate   = Electra::MakeDelegate(this, &FStreamHandler::HTTPProgressCallback);
-			ProgressListener->CompletionDelegate = Electra::MakeDelegate(this, &FStreamHandler::HTTPCompletionCallback);
+			ProgressListener->ProgressDelegate   = IElectraHttpManager::FProgressListener::FProgressDelegate::CreateRaw(this, &FStreamHandler::HTTPProgressCallback);
+			ProgressListener->CompletionDelegate = IElectraHttpManager::FProgressListener::FCompletionDelegate::CreateRaw(this, &FStreamHandler::HTTPCompletionCallback);
 			ReadBuffer.Reset();
 			ReadBuffer.ReceiveBuffer = MakeSharedTS<IElectraHttpManager::FReceiveBuffer>();
 
@@ -801,8 +801,8 @@ void FStreamReaderHLSfmp4::FStreamHandler::HandleRequest()
 				CurrentRequest->ConnectionInfo = {};
 
 				TSharedPtrTS<IElectraHttpManager::FProgressListener>	ProgressListener(new IElectraHttpManager::FProgressListener);
-				ProgressListener->ProgressDelegate   = Electra::MakeDelegate(this, &FStreamHandler::HTTPProgressCallback);
-				ProgressListener->CompletionDelegate = Electra::MakeDelegate(this, &FStreamHandler::HTTPCompletionCallback);
+				ProgressListener->ProgressDelegate   = IElectraHttpManager::FProgressListener::FProgressDelegate::CreateRaw(this, &FStreamHandler::HTTPProgressCallback);
+				ProgressListener->CompletionDelegate = IElectraHttpManager::FProgressListener::FCompletionDelegate::CreateRaw(this, &FStreamHandler::HTTPCompletionCallback);
 				TSharedPtrTS<IElectraHttpManager::FRequest> HTTP(new IElectraHttpManager::FRequest);
 				HTTP->Parameters.URL   = RequestURL;
 				HTTP->ReceiveBuffer    = ReadBuffer.ReceiveBuffer;
diff --git a/Engine/Plugins/Media/ElectraPlayer/Source/ElectraPlayerRuntime/Private/Runtime/Player/mp4/PlaylistReaderMP4.cpp b/Engine/Plugins/Media/ElectraPlayer/Source/ElectraPlayerRuntime/Private/Runtime/Player/mp4/PlaylistReaderMP4.cpp
--- a/Engine/Plugins/Media/ElectraPlayer/Source/ElectraPlayerRuntime/Private/Runtime/Player/mp4/PlaylistReaderMP4.cpp
+++ b/Engine/Plugins/Media/ElectraPlayer/Source/ElectraPlayerRuntime/Private/Runtime/Player/mp4/PlaylistReaderMP4.cpp
@@ -199,7 +199,7 @@ void FPlaylistReaderMP4::HandleOnce()
 void FPlaylistReaderMP4::StartWorkerThread()
 {
 	check(!bIsWorkerThreadStarted);
-	ThreadStart(Electra::MakeDelegate(this, &FPlaylistReaderMP4::WorkerThread));
+	ThreadStart(FMediaRunnable::FStartDelegate::CreateRaw(this, &FPlaylistReaderMP4::WorkerThread));
 	bIsWorkerThreadStarted = true;
 }
 
@@ -325,8 +325,8 @@ void FPlaylistReaderMP4::ReadNextChunk(int64 InFromOffset, int64 ChunkSize)
 		return;
 	}
 	ProgressListener = MakeSharedTS<IElectraHttpManager::FProgressListener>();
-	ProgressListener->CompletionDelegate = Electra::MakeDelegate(this, &FPlaylistReaderMP4::HTTPCompletionCallback);
-	ProgressListener->ProgressDelegate   = Electra::MakeDelegate(this, &FPlaylistReaderMP4::HTTPProgressCallback);
+	ProgressListener->CompletionDelegate = IElectraHttpManager::FProgressListener::FCompletionDelegate::CreateRaw(this, &FPlaylistReaderMP4::HTTPCompletionCallback);
+	ProgressListener->ProgressDelegate   = IElectraHttpManager::FProgressListener::FProgressDelegate::CreateRaw(this, &FPlaylistReaderMP4::HTTPProgressCallback);
 
 	ReceiveBuffer = MakeSharedTS<IElectraHttpManager::FReceiveBuffer>();
 	ReceiveBuffer->Buffer.Reserve(ChunkSize);
diff --git a/Engine/Plugins/Media/ElectraPlayer/Source/ElectraPlayerRuntime/Private/Runtime/Player/mp4/StreamReaderMP4.cpp b/Engine/Plugins/Media/ElectraPlayer/Source/ElectraPlayerRuntime/Private/Runtime/Player/mp4/StreamReaderMP4.cpp
--- a/Engine/Plugins/Media/ElectraPlayer/Source/ElectraPlayerRuntime/Private/Runtime/Player/mp4/StreamReaderMP4.cpp
+++ b/Engine/Plugins/Media/ElectraPlayer/Source/ElectraPlayerRuntime/Private/Runtime/Player/mp4/StreamReaderMP4.cpp
@@ -164,7 +164,7 @@ UEMediaError FStreamReaderMP4::Create(IPlayerSessionServices* InPlayerSessionSer
 	bIsStarted = true;
 
 	ThreadSetName("ElectraPlayer::MP4 streamer");
-	ThreadStart(Electra::MakeDelegate(this, &FStreamReaderMP4::WorkerThread));
+	ThreadStart(FMediaRunnable::FStartDelegate::CreateRaw(this, &FStreamReaderMP4::WorkerThread));
 
 	return UEMEDIA_ERROR_OK;
 }
@@ -389,8 +389,8 @@ void FStreamReaderMP4::HandleRequest()
 
 	TSharedPtrTS<IElectraHttpManager::FProgressListener>	ProgressListener;
 	ProgressListener = MakeSharedTS<IElectraHttpManager::FProgressListener>();
-	ProgressListener->CompletionDelegate = Electra::MakeDelegate(this, &FStreamReaderMP4::HTTPCompletionCallback);
-	ProgressListener->ProgressDelegate   = Electra::MakeDelegate(this, &FStreamReaderMP4::HTTPProgressCallback);
+	ProgressListener->CompletionDelegate = IElectraHttpManager::FProgressListener::FCompletionDelegate::CreateRaw(this, &FStreamReaderMP4::HTTPCompletionCallback);
+	ProgressListener->ProgressDelegate   = IElectraHttpManager::FProgressListener::FProgressDelegate::CreateRaw(this, &FStreamReaderMP4::HTTPProgressCallback);
 
 	ReadBuffer.Reset();
 	ReadBuffer.ReceiveBuffer = MakeSharedTS<IElectraHttpManager::FReceiveBuffer>();
diff --git a/Engine/Plugins/Media/ElectraPlayer/Source/ElectraPlayerRuntime/Private/Runtime/PlayerCore.h b/Engine/Plugins/Media/ElectraPlayer/Source/ElectraPlayerRuntime/Private/Runtime/PlayerCore.h
--- a/Engine/Plugins/Media/ElectraPlayer/Source/ElectraPlayerRuntime/Private/Runtime/PlayerCore.h
+++ b/Engine/Plugins/Media/ElectraPlayer/Source/ElectraPlayerRuntime/Private/Runtime/PlayerCore.h
@@ -6,7 +6,6 @@
 #include "Core/MediaMacros.h"
 #include "Core/MediaNoncopyable.h"
 
-#include "Core/MediaFastDelegate.h"
 #include "Core/MediaThreads.h"
 #include "Core/MediaQueue.h"
 #include "Core/MediaSemaphore.h"
