ElectraPlayer: Fixed a static analysis warning

#rnx

diff --git a/Engine/Plugins/Media/ElectraPlayer/Source/ElectraPlayerRuntime/Private/Runtime/Utilities/UtilsMPEGVideo.cpp b/Engine/Plugins/Media/ElectraPlayer/Source/ElectraPlayerRuntime/Private/Runtime/Utilities/UtilsMPEGVideo.cpp
--- a/Engine/Plugins/Media/ElectraPlayer/Source/ElectraPlayerRuntime/Private/Runtime/Utilities/UtilsMPEGVideo.cpp
+++ b/Engine/Plugins/Media/ElectraPlayer/Source/ElectraPlayerRuntime/Private/Runtime/Utilities/UtilsMPEGVideo.cpp
@@ -904,7 +904,7 @@ namespace Electra
 							uint32 num_negative_pics = FSyntaxElement::ue_v(BitReader);
 							uint32 num_positive_pics = FSyntaxElement::ue_v(BitReader);
 
-							if ((uint64)(num_positive_pics + num_negative_pics) * 2 > BitReader.GetRemainingBits())
+							if (((uint64)num_positive_pics + (uint64)num_negative_pics) * 2 > BitReader.GetRemainingBits())
 							{
 								return false;
 							}
