ElectraPlayer: Register application termination callback and close still active players

#rnx

diff --git a/Engine/Plugins/Media/ElectraPlayer/Source/ElectraPlayerRuntime/Private/ElectraPlayer.cpp b/Engine/Plugins/Media/ElectraPlayer/Source/ElectraPlayerRuntime/Private/ElectraPlayer.cpp
--- a/Engine/Plugins/Media/ElectraPlayer/Source/ElectraPlayerRuntime/Private/ElectraPlayer.cpp
+++ b/Engine/Plugins/Media/ElectraPlayer/Source/ElectraPlayerRuntime/Private/ElectraPlayer.cpp
@@ -130,6 +130,10 @@ FElectraPlayer::FElectraPlayer(const TSharedPtr<IElectraPlayerAdapterDelegate, E
 	WaitForPlayerDestroyedEvent = FPlatformProcess::GetSynchEventFromPool(true);
 	WaitForPlayerDestroyedEvent->Trigger();
 
+	AppTerminationHandler = MakeSharedTS<Electra::FApplicationTerminationHandler>();
+	AppTerminationHandler->Terminate = [this]() { CloseInternal(false); };
+	Electra::AddTerminationNotificationHandler(AppTerminationHandler);
+
 	FString	OSMinor;
 	AnalyticsGPUType = InAdapterDelegate->GetVideoAdapterName().TrimStartAndEnd();
 	FPlatformMisc::GetOSVersions(AnalyticsOSVersion, OSMinor);
@@ -163,6 +167,9 @@ FElectraPlayer::FElectraPlayer(const TSharedPtr<IElectraPlayerAdapterDelegate, E
  */
 FElectraPlayer::~FElectraPlayer()
 {
+	Electra::RemoveTerminationNotificationHandler(AppTerminationHandler);
+	AppTerminationHandler.Reset();
+
 	CloseInternal(false);
 	WaitForPlayerDestroyedEvent->Wait();
 
diff --git a/Engine/Plugins/Media/ElectraPlayer/Source/ElectraPlayerRuntime/Private/ElectraPlayer.h b/Engine/Plugins/Media/ElectraPlayer/Source/ElectraPlayerRuntime/Private/ElectraPlayer.h
--- a/Engine/Plugins/Media/ElectraPlayer/Source/ElectraPlayerRuntime/Private/ElectraPlayer.h
+++ b/Engine/Plugins/Media/ElectraPlayer/Source/ElectraPlayerRuntime/Private/ElectraPlayer.h
@@ -12,6 +12,7 @@
 #include "IElectraPlayerInterface.h"
 
 #include "Player/AdaptiveStreamingPlayer.h"
+#include "PlayerRuntimeGlobal.h"
 
 class FVideoDecoderOutput;
 using FVideoDecoderOutputPtr = TSharedPtr<FVideoDecoderOutput, ESPMode::ThreadSafe>;
@@ -471,6 +472,7 @@ private:
 	mutable FCriticalSection										PlayerLock;
 	TSharedPtr<FInternalPlayerImpl, ESPMode::ThreadSafe>			CurrentPlayer;
 	FEvent*															WaitForPlayerDestroyedEvent;
+	TSharedPtr<Electra::FApplicationTerminationHandler, ESPMode::ThreadSafe> AppTerminationHandler;
 
 	TSharedPtr<IAsyncResourceReleaseNotifyContainer, ESPMode::ThreadSafe> AsyncResourceReleaseNotification;
 
diff --git a/Engine/Plugins/Media/ElectraPlayer/Source/ElectraPlayerRuntime/Private/ElectraPlayerRuntimeModule.cpp b/Engine/Plugins/Media/ElectraPlayer/Source/ElectraPlayerRuntime/Private/ElectraPlayerRuntimeModule.cpp
--- a/Engine/Plugins/Media/ElectraPlayer/Source/ElectraPlayerRuntime/Private/ElectraPlayerRuntimeModule.cpp
+++ b/Engine/Plugins/Media/ElectraPlayer/Source/ElectraPlayerRuntime/Private/ElectraPlayerRuntimeModule.cpp
@@ -101,14 +101,25 @@ public:
 	{
 		if (bInitialized)
 		{
-			Electra::WaitForAllPlayersToHaveTerminated();
-			IAudioDecoderAAC::Shutdown();
-			IVideoDecoderH265::Shutdown();
-			IVideoDecoderH264::Shutdown();
-			FElectraRendererAudio::Shutdown();
-			FElectraRendererVideo::Shutdown();
-			Electra::Shutdown();
-			Electra::PlatformShutdown();
+			// Wait for players to have terminated. If this fails then do not shutdown the sub components
+			// to avoid potential hangs in there.
+			if (Electra::WaitForAllPlayersToHaveTerminated())
+			{
+				IAudioDecoderAAC::Shutdown();
+				IVideoDecoderH265::Shutdown();
+				IVideoDecoderH264::Shutdown();
+				FElectraRendererAudio::Shutdown();
+				FElectraRendererVideo::Shutdown();
+				Electra::Shutdown();
+				Electra::PlatformShutdown();
+			}
+			else
+			{
+				UE_LOG(LogElectraPlayer, Warning, TEXT("Shutting down with active player instances. This could lead to problems."));
+
+				// At least unbind all the application notification handlers.
+				Electra::Shutdown();
+			}
 			bInitialized = false;
 		}
 	}
diff --git a/Engine/Plugins/Media/ElectraPlayer/Source/ElectraPlayerRuntime/Private/Runtime/PlayerRuntimeGlobal.cpp b/Engine/Plugins/Media/ElectraPlayer/Source/ElectraPlayerRuntime/Private/Runtime/PlayerRuntimeGlobal.cpp
--- a/Engine/Plugins/Media/ElectraPlayer/Source/ElectraPlayerRuntime/Private/Runtime/PlayerRuntimeGlobal.cpp
+++ b/Engine/Plugins/Media/ElectraPlayer/Source/ElectraPlayerRuntime/Private/Runtime/PlayerRuntimeGlobal.cpp
@@ -8,22 +8,41 @@ namespace Electra
 {
 	namespace Global
 	{
-		TMap<FString, bool>								EnabledAnalyticsEvents;
-		FDelegateHandle									ApplicationSuspendedDelegate;
-		FDelegateHandle									ApplicationResumeDelegate;
-		FCriticalSection								ApplicationBGFGLock;
-		TArray<TWeakPtrTS<FFGBGNotificationHandlers>>	ApplicationBGFGHandlers;
-		volatile bool									bIsInBackground = false;
+		TMap<FString, bool>									EnabledAnalyticsEvents;
+		FDelegateHandle										ApplicationSuspendedDelegate;
+		FDelegateHandle										ApplicationResumeDelegate;
+		FDelegateHandle										ApplicationTerminatingDelegate;
+		FCriticalSection									ApplicationHandlerLock;
+		TArray<TWeakPtrTS<FApplicationTerminationHandler>>	ApplicationTerminationHandlers;
+		TArray<TWeakPtrTS<FFGBGNotificationHandlers>>		ApplicationBGFGHandlers;
+		volatile bool										bIsInBackground = false;
+		volatile bool										bAppIsTerminating = false;
 	}
 	using namespace Global;
 
 
+	static void HandleApplicationWillTerminate()
+	{
+		ApplicationHandlerLock.Lock();
+		TArray<TWeakPtrTS<FApplicationTerminationHandler>>	CurrentHandlers(ApplicationTerminationHandlers);
+		bAppIsTerminating = true;
+		ApplicationHandlerLock.Unlock();
+		for(auto &Handler : CurrentHandlers)
+		{
+			TSharedPtrTS<FApplicationTerminationHandler> Hdlr = Handler.Pin();
+			if (Hdlr.IsValid())
+			{
+				Hdlr->Terminate();
+			}
+		}
+	}
+
 	static void HandleApplicationWillEnterBackground()
 	{
-		ApplicationBGFGLock.Lock();
+		ApplicationHandlerLock.Lock();
 		TArray<TWeakPtrTS<FFGBGNotificationHandlers>>	CurrentHandlers(ApplicationBGFGHandlers);
 		bIsInBackground = true;
-		ApplicationBGFGLock.Unlock();
+		ApplicationHandlerLock.Unlock();
 		for(auto &Handler : CurrentHandlers)
 		{
 			TSharedPtrTS<FFGBGNotificationHandlers> Hdlr = Handler.Pin();
@@ -36,10 +55,10 @@ namespace Electra
 
 	static void HandleApplicationHasEnteredForeground()
 	{
-		ApplicationBGFGLock.Lock();
+		ApplicationHandlerLock.Lock();
 		TArray<TWeakPtrTS<FFGBGNotificationHandlers>>	CurrentHandlers(ApplicationBGFGHandlers);
 		bIsInBackground = false;
-		ApplicationBGFGLock.Unlock();
+		ApplicationHandlerLock.Unlock();
 		for(auto &Handler : CurrentHandlers)
 		{
 			TSharedPtrTS<FFGBGNotificationHandlers> Hdlr = Handler.Pin();
@@ -61,6 +80,10 @@ namespace Electra
 	 */
 	bool Startup(const Configuration& InConfiguration)
 	{
+		if (!ApplicationTerminatingDelegate.IsValid())
+		{
+			ApplicationTerminatingDelegate = FCoreDelegates::ApplicationWillTerminateDelegate.AddStatic(&HandleApplicationWillTerminate);
+		}
 		if (!ApplicationSuspendedDelegate.IsValid())
 		{
 			ApplicationSuspendedDelegate = FCoreDelegates::ApplicationWillEnterBackgroundDelegate.AddStatic(&HandleApplicationWillEnterBackground);
@@ -91,18 +114,35 @@ namespace Electra
 			FCoreDelegates::ApplicationHasEnteredForegroundDelegate.Remove(ApplicationResumeDelegate);
 			ApplicationResumeDelegate.Reset();
 		}
+		if (ApplicationTerminatingDelegate.IsValid())
+		{
+			FCoreDelegates::ApplicationWillTerminateDelegate.Remove(ApplicationTerminatingDelegate);
+		}
+	}
+
+
+	void AddTerminationNotificationHandler(TSharedPtrTS<FApplicationTerminationHandler> InHandler)
+	{
+		FScopeLock lock(&ApplicationHandlerLock);
+		ApplicationTerminationHandlers.Add(InHandler);
+	}
+
+	void RemoveTerminationNotificationHandler(TSharedPtrTS<FApplicationTerminationHandler> InHandler)
+	{
+		FScopeLock lock(&ApplicationHandlerLock);
+		ApplicationTerminationHandlers.Remove(InHandler);
 	}
 
 	bool AddBGFGNotificationHandler(TSharedPtrTS<FFGBGNotificationHandlers> InHandlers)
 	{
-		FScopeLock lock(&ApplicationBGFGLock);
+		FScopeLock lock(&ApplicationHandlerLock);
 		ApplicationBGFGHandlers.Add(InHandlers);
 		return bIsInBackground;
 	}
 
 	void RemoveBGFGNotificationHandler(TSharedPtrTS<FFGBGNotificationHandlers> InHandlers)
 	{
-		FScopeLock lock(&ApplicationBGFGLock);
+		FScopeLock lock(&ApplicationHandlerLock);
 		ApplicationBGFGHandlers.Remove(InHandlers);
 	}
 
@@ -159,10 +199,15 @@ namespace Electra
 		}
 
 		//! Waits for all pending tasks to have finished. Once all are done new tasks cannot be added.
-		void WaitAllFinished()
+		bool WaitAllFinished(int32 TimeoutMs)
 		{
 			Init();
-			AllDoneSignal->Wait();
+			if (TimeoutMs <= 0)
+			{
+				AllDoneSignal->Wait();
+				return true;
+			}
+			return AllDoneSignal->WaitTimeout(TimeoutMs * 1000);
 		}
 
 		void Reset()
@@ -198,11 +243,15 @@ namespace Electra
 	static PendingTaskCounter NumActivePlayers;
 
 
-	void WaitForAllPlayersToHaveTerminated()
+	bool WaitForAllPlayersToHaveTerminated()
 	{
-		NumActivePlayers.WaitAllFinished();
-		// Explicitly shutdown anything in the counter class that may use the engine (as it might shutdown after this)
-		NumActivePlayers.Reset();
+		bool bOk = NumActivePlayers.WaitAllFinished(1000);	// bAppIsTerminating ? 1000 : 0);
+		if (bOk)
+		{
+			// Explicitly shutdown anything in the counter class that may use the engine (as it might shutdown after this)
+			NumActivePlayers.Reset();
+		}
+		return bOk;
 	}
 
 	void AddActivePlayerInstance()
diff --git a/Engine/Plugins/Media/ElectraPlayer/Source/ElectraPlayerRuntime/Private/Runtime/PlayerRuntimeGlobal.h b/Engine/Plugins/Media/ElectraPlayer/Source/ElectraPlayerRuntime/Private/Runtime/PlayerRuntimeGlobal.h
--- a/Engine/Plugins/Media/ElectraPlayer/Source/ElectraPlayerRuntime/Private/Runtime/PlayerRuntimeGlobal.h
+++ b/Engine/Plugins/Media/ElectraPlayer/Source/ElectraPlayerRuntime/Private/Runtime/PlayerRuntimeGlobal.h
@@ -19,7 +19,7 @@ namespace Electra
 	void Shutdown();
 
 	//! Waits until all player instances have terminated, which may happen asynchronously.
-	void WaitForAllPlayersToHaveTerminated();
+	bool WaitForAllPlayersToHaveTerminated();
 
 
 	void AddActivePlayerInstance();
@@ -29,6 +29,14 @@ namespace Electra
 	bool IsAnalyticsEventEnabled(const FString& AnalyticsEventName);
 
 
+	//! Application termination
+	struct FApplicationTerminationHandler
+	{
+		TFunction<void()> Terminate;
+	};
+	void AddTerminationNotificationHandler(TSharedPtrTS<FApplicationTerminationHandler> InHandler);
+	void RemoveTerminationNotificationHandler(TSharedPtrTS<FApplicationTerminationHandler> InHandler);
+
 	//! Background / foreground handling
 	struct FFGBGNotificationHandlers
 	{
