ElectraPlayer: Revised buffer handling to support track switching; enabled switching of video tracks for mp4 sources

#rnx

diff --git a/Engine/Plugins/Media/ElectraPlayer/Source/ElectraPlayerPlugin/Private/ElectraPlayerPlugin.cpp b/Engine/Plugins/Media/ElectraPlayer/Source/ElectraPlayerPlugin/Private/ElectraPlayerPlugin.cpp
--- a/Engine/Plugins/Media/ElectraPlayer/Source/ElectraPlayerPlugin/Private/ElectraPlayerPlugin.cpp
+++ b/Engine/Plugins/Media/ElectraPlayer/Source/ElectraPlayerPlugin/Private/ElectraPlayerPlugin.cpp
@@ -665,6 +665,8 @@ bool FElectraPlayerPlugin::GetPlayerFeatureFlag(EFeatureFlag flag) const
 			return true;
 		case EFeatureFlag::PlayerUsesInternalFlushOnSeek:
 			return true;
+		case EFeatureFlag::IsTrackSwitchSeamless:
+			return true;
 		default:
 			break;
 	}
diff --git a/Engine/Plugins/Media/ElectraPlayer/Source/ElectraPlayerRuntime/Private/ElectraPlayer.cpp b/Engine/Plugins/Media/ElectraPlayer/Source/ElectraPlayerRuntime/Private/ElectraPlayer.cpp
--- a/Engine/Plugins/Media/ElectraPlayer/Source/ElectraPlayerRuntime/Private/ElectraPlayer.cpp
+++ b/Engine/Plugins/Media/ElectraPlayer/Source/ElectraPlayerRuntime/Private/ElectraPlayer.cpp
@@ -1228,7 +1228,7 @@ bool FElectraPlayer::GetAudioTrackFormat(int32 TrackIndex, int32 FormatIndex, FA
 
 bool FElectraPlayer::GetVideoTrackFormat(int32 TrackIndex, int32 FormatIndex, FVideoTrackFormat& OutFormat) const
 {
-	if (TrackIndex == 0 && FormatIndex == 0)
+	if (TrackIndex >= 0 && TrackIndex < NumTracksVideo && FormatIndex == 0)
 	{
 		TSharedPtr<Electra::FTrackMetadata, ESPMode::ThreadSafe> Meta = GetTrackStreamMetadata(EPlayerTrackType::Video, TrackIndex);
 		if (Meta.IsValid())
@@ -1404,7 +1404,7 @@ FText FElectraPlayer::GetTrackDisplayName(EPlayerTrackType TrackType, int32 Trac
 		{
 			if (!Meta->Label.IsEmpty())
 			{
-				return FText::FromString(FString::Printf(TEXT("%s, ID=%s"), *Meta->Label, *Meta->ID));
+				return FText::FromString(Meta->Label);
 			}
 			return FText::FromString(FString::Printf(TEXT("Video Track ID %s"), *Meta->ID));
 		}
@@ -1412,7 +1412,7 @@ FText FElectraPlayer::GetTrackDisplayName(EPlayerTrackType TrackType, int32 Trac
 		{
 			if (!Meta->Label.IsEmpty())
 			{
-				return FText::FromString(FString::Printf(TEXT("%s, ID=%s"), *Meta->Label, *Meta->ID));
+				return FText::FromString(Meta->Label);
 			}
 			return FText::FromString(FString::Printf(TEXT("Audio Track ID %s"), *Meta->ID));
 		}
@@ -1421,7 +1421,7 @@ FText FElectraPlayer::GetTrackDisplayName(EPlayerTrackType TrackType, int32 Trac
 			FString Name;
 			if (!Meta->Label.IsEmpty())
 			{
-				Name = FString::Printf(TEXT("%s (%s), ID=%s"), *Meta->Label, *Meta->HighestBandwidthCodec.GetCodecSpecifierRFC6381(), *Meta->ID);
+				Name = FString::Printf(TEXT("%s (%s)"), *Meta->Label, *Meta->HighestBandwidthCodec.GetCodecSpecifierRFC6381());
 			}
 			else
 			{
@@ -1482,115 +1482,78 @@ FString FElectraPlayer::GetTrackName(EPlayerTrackType TrackType, int32 TrackInde
  */
 bool FElectraPlayer::SelectTrack(EPlayerTrackType TrackType, int32 TrackIndex)
 {
-	if (TrackType == EPlayerTrackType::Audio)
+	auto PerformSelection = [this, TrackType, TrackIndex](int32& OutSelectedTrackIndex, IElectraPlayerInterface::FStreamSelectionAttributes& OutSelectionAttributes) -> bool
 	{
+		Electra::EStreamType StreamType = TrackType == IElectraPlayerInterface::EPlayerTrackType::Video ? Electra::EStreamType::Video :
+										  TrackType == IElectraPlayerInterface::EPlayerTrackType::Audio ? Electra::EStreamType::Audio :
+										  TrackType == IElectraPlayerInterface::EPlayerTrackType::Subtitle ? Electra::EStreamType::Subtitle :
+										  Electra::EStreamType::Unsupported;
 		// Select a track or deselect?
 		if (TrackIndex >= 0)
 		{
 			// Check if the track index exists by checking the presence of the track metadata.
 			// If for some reason the index is not valid the selection will not be changed.
-			TSharedPtr<Electra::FTrackMetadata, ESPMode::ThreadSafe> Meta = GetTrackStreamMetadata(EPlayerTrackType::Audio, TrackIndex);
+			TSharedPtr<Electra::FTrackMetadata, ESPMode::ThreadSafe> Meta = GetTrackStreamMetadata(TrackType, TrackIndex);
 			if (Meta.IsValid())
 			{
 				// Switch only when the track index has changed.
 				if (GetSelectedTrack(TrackType) != TrackIndex)
 				{
-					Electra::FStreamSelectionAttributes AudioAttributes;
-					AudioAttributes.OverrideIndex = TrackIndex;
+					Electra::FStreamSelectionAttributes TrackAttributes;
+					TrackAttributes.OverrideIndex = TrackIndex;
 
-					PlaystartOptions.InitialAudioTrackAttributes.TrackIndexOverride = TrackIndex;
+					OutSelectionAttributes.TrackIndexOverride = TrackIndex;
 					if (!Meta->Kind.IsEmpty())
 					{
-						AudioAttributes.Kind = Meta->Kind;
-						PlaystartOptions.InitialAudioTrackAttributes.Kind = Meta->Kind;
+						TrackAttributes.Kind = Meta->Kind;
+						OutSelectionAttributes.Kind = Meta->Kind;
 					}
 					if (!Meta->Language.IsEmpty())
 					{
-						AudioAttributes.Language_ISO639 = Meta->Language;
-						PlaystartOptions.InitialAudioTrackAttributes.Language_ISO639 = Meta->Language;
+						TrackAttributes.Language_ISO639 = Meta->Language;
+						OutSelectionAttributes.Language_ISO639 = Meta->Language;
 					}
-					AudioAttributes.Codec = Meta->HighestBandwidthCodec.GetCodecName();
-					PlaystartOptions.InitialAudioTrackAttributes.Codec = Meta->HighestBandwidthCodec.GetCodecName();
+					TrackAttributes.Codec = Meta->HighestBandwidthCodec.GetCodecName();
+					OutSelectionAttributes.Codec = Meta->HighestBandwidthCodec.GetCodecName();
 
 					TSharedPtr<FInternalPlayerImpl, ESPMode::ThreadSafe> LockedPlayer = CurrentPlayer;
 					if (LockedPlayer.IsValid() && LockedPlayer->AdaptivePlayer.IsValid())
 					{
-						LockedPlayer->AdaptivePlayer->SelectTrackByAttributes(Electra::EStreamType::Audio, AudioAttributes);
+						LockedPlayer->AdaptivePlayer->SelectTrackByAttributes(StreamType, TrackAttributes);
 					}
 
-					SelectedAudioTrackIndex = TrackIndex;
+					OutSelectedTrackIndex = TrackIndex;
 				}
 				return true;
 			}
 		}
 		else
 		{
-			// Deselect audio track.
-			PlaystartOptions.InitialAudioTrackAttributes.TrackIndexOverride = -1;
-			SelectedAudioTrackIndex = -1;
+			// Deselect track.
+			OutSelectionAttributes.TrackIndexOverride = -1;
+			OutSelectedTrackIndex = -1;
 			TSharedPtr<FInternalPlayerImpl, ESPMode::ThreadSafe> LockedPlayer = CurrentPlayer;
 			if (LockedPlayer.IsValid() && LockedPlayer->AdaptivePlayer.IsValid())
 			{
-				LockedPlayer->AdaptivePlayer->DeselectTrack(Electra::EStreamType::Audio);
+				LockedPlayer->AdaptivePlayer->DeselectTrack(StreamType);
 			}
 			return true;
 		}
+		return false;
+	};
+
+	if (TrackType == EPlayerTrackType::Video)
+	{
+		return PerformSelection(SelectedVideoTrackIndex, PlaystartOptions.InitialVideoTrackAttributes);
+	}
+	else if (TrackType == EPlayerTrackType::Audio)
+	{
+		return PerformSelection(SelectedAudioTrackIndex, PlaystartOptions.InitialAudioTrackAttributes);
 	}
 	else if (TrackType == EPlayerTrackType::Subtitle)
 	{
-		// Select a track or deselect?
-		if (TrackIndex >= 0)
-		{
-			// Check if the track index exists by checking the presence of the track metadata.
-			// If for some reason the index is not valid the selection will not be changed.
-			TSharedPtr<Electra::FTrackMetadata, ESPMode::ThreadSafe> Meta = GetTrackStreamMetadata(EPlayerTrackType::Subtitle, TrackIndex);
-			if (Meta.IsValid())
-			{
-				// Switch only when the track index has changed.
-				if (GetSelectedTrack(TrackType) != TrackIndex)
-				{
-					Electra::FStreamSelectionAttributes SubtitleAttributes;
-					SubtitleAttributes.OverrideIndex = TrackIndex;
-
-					PlaystartOptions.InitialSubtitleTrackAttributes.TrackIndexOverride = TrackIndex;
-					if (!Meta->Kind.IsEmpty())
-					{
-						SubtitleAttributes.Kind = Meta->Kind;
-						PlaystartOptions.InitialSubtitleTrackAttributes.Kind = Meta->Kind;
-					}
-					if (!Meta->Language.IsEmpty())
-					{
-						SubtitleAttributes.Language_ISO639 = Meta->Language;
-						PlaystartOptions.InitialSubtitleTrackAttributes.Language_ISO639 = Meta->Language;
-					}
-					SubtitleAttributes.Codec = Meta->HighestBandwidthCodec.GetCodecName();
-					PlaystartOptions.InitialSubtitleTrackAttributes.Codec = Meta->HighestBandwidthCodec.GetCodecName();
-
-					TSharedPtr<FInternalPlayerImpl, ESPMode::ThreadSafe> LockedPlayer = CurrentPlayer;
-					if (LockedPlayer.IsValid() && LockedPlayer->AdaptivePlayer.IsValid())
-					{
-						LockedPlayer->AdaptivePlayer->SelectTrackByAttributes(Electra::EStreamType::Subtitle, SubtitleAttributes);
-					}
-
-					SelectedSubtitleTrackIndex = TrackIndex;
-				}
-				return true;
-			}
-		}
-		else
-		{
-			// Deselect subtitle track.
-			PlaystartOptions.InitialSubtitleTrackAttributes.TrackIndexOverride = -1;
-			SelectedSubtitleTrackIndex = -1;
-			TSharedPtr<FInternalPlayerImpl, ESPMode::ThreadSafe> LockedPlayer = CurrentPlayer;
-			if (LockedPlayer.IsValid() && LockedPlayer->AdaptivePlayer.IsValid())
-			{
-				LockedPlayer->AdaptivePlayer->DeselectTrack(Electra::EStreamType::Subtitle);
-			}
-			return true;
-		}
+		return PerformSelection(SelectedSubtitleTrackIndex, PlaystartOptions.InitialSubtitleTrackAttributes);
 	}
-	// The player is selecting tracks for now.
 	return false;
 }
 
@@ -2171,7 +2134,7 @@ void FElectraPlayer::HandlePlayerEventBufferingEnd(Electra::Metrics::EBufferingR
 void FElectraPlayer::HandlePlayerEventBandwidth(int64 EffectiveBps, int64 ThroughputBps, double LatencyInSeconds)
 {
 //	FScopeLock Lock(&StatisticsLock);
-	UE_LOG(LogElectraPlayer, Verbose, TEXT("[%p][%p] Observed bandwidth of %d bps; throughput = %d; latency = %.3fs"), this, CurrentPlayer.Get(), (int32)EffectiveBps, (int32)ThroughputBps, LatencyInSeconds);
+	UE_LOG(LogElectraPlayer, VeryVerbose, TEXT("[%p][%p] Observed bandwidth of %d bps; throughput = %d; latency = %.3fs"), this, CurrentPlayer.Get(), (int32)EffectiveBps, (int32)ThroughputBps, LatencyInSeconds);
 }
 
 void FElectraPlayer::HandlePlayerEventBufferUtilization(const Electra::Metrics::FBufferStats& BufferStats)
@@ -2212,7 +2175,7 @@ void FElectraPlayer::HandlePlayerEventSegmentDownload(const Electra::Metrics::FS
 	}
 	if (SegmentDownloadStats.bWasSuccessful)
 	{
-		UE_LOG(LogElectraPlayer, Verbose, TEXT("[%p][%p] Downloaded %s segment at bitrate %d: Playback time = %.3fs, duration = %.3fs, download time = %.3fs, URL=%s \"%s\""), this, CurrentPlayer.Get(), Electra::GetStreamTypeName(SegmentDownloadStats.StreamType), SegmentDownloadStats.Bitrate, SegmentDownloadStats.PresentationTime, SegmentDownloadStats.Duration, SegmentDownloadStats.TimeToDownload, *SegmentDownloadStats.Range, *SanitizeMessage(SegmentDownloadStats.URL));
+		UE_LOG(LogElectraPlayer, VeryVerbose, TEXT("[%p][%p] Downloaded %s segment at bitrate %d: Playback time = %.3fs, duration = %.3fs, download time = %.3fs, URL=%s \"%s\""), this, CurrentPlayer.Get(), Electra::GetStreamTypeName(SegmentDownloadStats.StreamType), SegmentDownloadStats.Bitrate, SegmentDownloadStats.PresentationTime, SegmentDownloadStats.Duration, SegmentDownloadStats.TimeToDownload, *SegmentDownloadStats.Range, *SanitizeMessage(SegmentDownloadStats.URL));
 	}
 	else if (SegmentDownloadStats.bWasAborted)
 	{
diff --git a/Engine/Plugins/Media/ElectraPlayer/Source/ElectraPlayerRuntime/Private/ElectraPlayer.h b/Engine/Plugins/Media/ElectraPlayer/Source/ElectraPlayerRuntime/Private/ElectraPlayer.h
--- a/Engine/Plugins/Media/ElectraPlayer/Source/ElectraPlayerRuntime/Private/ElectraPlayer.h
+++ b/Engine/Plugins/Media/ElectraPlayer/Source/ElectraPlayerRuntime/Private/ElectraPlayer.h
@@ -383,7 +383,7 @@ private:
 	int32											NumTracksVideo;
 	int32											NumTracksSubtitle;
 	int32											SelectedQuality;
-	int32											SelectedVideoTrackIndex;
+	mutable int32									SelectedVideoTrackIndex;
 	mutable int32									SelectedAudioTrackIndex;
 	mutable int32									SelectedSubtitleTrackIndex;
 	mutable bool									bAudioTrackIndexDirty;
diff --git a/Engine/Plugins/Media/ElectraPlayer/Source/ElectraPlayerRuntime/Private/Runtime/Demuxer/ParserISO14496-12.cpp b/Engine/Plugins/Media/ElectraPlayer/Source/ElectraPlayerRuntime/Private/Runtime/Demuxer/ParserISO14496-12.cpp
--- a/Engine/Plugins/Media/ElectraPlayer/Source/ElectraPlayerRuntime/Private/Runtime/Demuxer/ParserISO14496-12.cpp
+++ b/Engine/Plugins/Media/ElectraPlayer/Source/ElectraPlayerRuntime/Private/Runtime/Demuxer/ParserISO14496-12.cpp
@@ -250,7 +250,7 @@ public:
 	FDataBufferReader(const TArray<uint8>& InDataBufferToReadFrom) : DataBufferRef(InDataBufferToReadFrom)
 	{}
 	virtual ~FDataBufferReader() = default;
-	virtual int64 ReadData(void* IntoBuffer, int64 NumBytesToRead) override
+	int64 ReadData(void* IntoBuffer, int64 NumBytesToRead) override
 	{
 		int64 NumAvail = DataBufferRef.Num() - CurrentOffset;
 		if (NumAvail >= NumBytesToRead)
@@ -264,11 +264,11 @@ public:
 		}
 		return -1;
 	}
-	virtual bool HasReachedEOF() const override
+	bool HasReachedEOF() const override
 	{ return CurrentOffset >= DataBufferRef.Num(); }
-	virtual bool HasReadBeenAborted() const override
+	bool HasReadBeenAborted() const override
 	{ return false;	}
-	virtual int64 GetCurrentOffset() const override
+	int64 GetCurrentOffset() const override
 	{ return CurrentOffset;	}
 private:
 	const TArray<uint8>& DataBufferRef;
@@ -601,6 +601,17 @@ private:
 
 
 
+		struct FFillerData
+		{
+			~FFillerData()
+			{
+				FMemory::Free(Data);
+			}
+			void* Data = nullptr;
+			int64	Size = 0;
+			int64	StartOffset = 0;
+		};
+
 		UEMediaError ReadFillerData(FMP4ParseInfo* ParseInfo, int64 Size)
 		{
 			if (Size == 0)
@@ -626,6 +637,11 @@ private:
 			return ParseInfo->Reader()->ReadBytes(FillerData->Data, FillerData->Size);
 		}
 
+		const FFillerData* GetFillerData() const
+		{
+			return FillerData;
+		}
+
 		virtual UEMediaError ReadAndParseAttributes(FMP4ParseInfo* ParseInfo) = 0;
 
 		// constexpr doesn't work here
@@ -764,21 +780,6 @@ private:
 		FMP4Box() = delete;
 		FMP4Box(const FMP4Box&) = delete;
 	protected:
-		struct FFillerData
-		{
-			FFillerData()
-				: Data(nullptr), Size(0), StartOffset(0)
-			{
-			}
-			~FFillerData()
-			{
-				FMemory::Free(Data);
-			}
-			void* Data;
-			int64	Size;
-			int64	StartOffset;
-		};
-
 		FMP4Box*						ParentBox;
 		TArray<FMP4Box*>				ChildBoxes;
 		IParserISO14496_12::FBoxType	BoxType;
@@ -820,7 +821,7 @@ private:
 		FMP4BoxBasic(const FMP4BoxBasic&) = delete;
 
 	protected:
-		virtual UEMediaError ReadAndParseAttributes(FMP4ParseInfo* ParseInfo) override
+		UEMediaError ReadAndParseAttributes(FMP4ParseInfo* ParseInfo) override
 		{
 			// The basic box is merely a container with no attributes.
 			return UEMEDIA_ERROR_OK;
@@ -854,7 +855,7 @@ private:
 		FMP4BoxFull(const FMP4BoxFull&) = delete;
 
 	protected:
-		virtual UEMediaError ReadAndParseAttributes(FMP4ParseInfo* ParseInfo) override
+		UEMediaError ReadAndParseAttributes(FMP4ParseInfo* ParseInfo) override
 		{
 			uint32 VersionAndFlags = 0;
 			UEMediaError Error = UEMEDIA_ERROR_OK;
@@ -899,7 +900,7 @@ private:
 		FMP4BoxIgnored(const FMP4BoxIgnored&) = delete;
 
 	protected:
-		virtual UEMediaError ReadAndParseAttributes(FMP4ParseInfo* ParseInfo) override
+		UEMediaError ReadAndParseAttributes(FMP4ParseInfo* ParseInfo) override
 		{
 			// Read the remaining bytes from this box.
 			return ReadFillerData(ParseInfo, BoxSize - (ParseInfo->Reader()->GetCurrentReadOffset() - StartOffset));
@@ -929,7 +930,7 @@ private:
 		FMP4BoxMDAT(const FMP4BoxMDAT&) = delete;
 
 	protected:
-		virtual UEMediaError ReadAndParseAttributes(FMP4ParseInfo* ParseInfo) override
+		UEMediaError ReadAndParseAttributes(FMP4ParseInfo* ParseInfo) override
 		{
 			// Read the data from this box to NULL.
 			return ParseInfo->Reader()->ReadBytes(nullptr, BoxSize - (ParseInfo->Reader()->GetCurrentReadOffset() - StartOffset));
@@ -972,7 +973,7 @@ private:
 		FMP4BoxFTYP(const FMP4BoxFTYP&) = delete;
 
 	protected:
-		virtual UEMediaError ReadAndParseAttributes(FMP4ParseInfo* ParseInfo) override
+		UEMediaError ReadAndParseAttributes(FMP4ParseInfo* ParseInfo) override
 		{
 			UEMediaError Error = UEMEDIA_ERROR_OK;
 			uint32 Major;
@@ -1072,7 +1073,7 @@ private:
 		FMP4BoxMVHD(const FMP4BoxMVHD&) = delete;
 
 	protected:
-		virtual UEMediaError ReadAndParseAttributes(FMP4ParseInfo* ParseInfo) override
+		UEMediaError ReadAndParseAttributes(FMP4ParseInfo* ParseInfo) override
 		{
 			uint64  		Value64 = 0;
 			uint32  		Value32 = 0;
@@ -1132,7 +1133,7 @@ private:
 		FMP4BoxMEHD(const FMP4BoxMEHD&) = delete;
 
 	protected:
-		virtual UEMediaError ReadAndParseAttributes(FMP4ParseInfo* ParseInfo) override
+		UEMediaError ReadAndParseAttributes(FMP4ParseInfo* ParseInfo) override
 		{
 			uint32  		Value32 = 0;
 			UEMediaError	Error = UEMEDIA_ERROR_OK;
@@ -1209,7 +1210,7 @@ private:
 		FMP4BoxTKHD(const FMP4BoxTKHD&) = delete;
 
 	protected:
-		virtual UEMediaError ReadAndParseAttributes(FMP4ParseInfo* ParseInfo) override
+		UEMediaError ReadAndParseAttributes(FMP4ParseInfo* ParseInfo) override
 		{
 			uint64 Value64 = 0;
 			uint32 Value32 = 0;
@@ -1328,7 +1329,7 @@ private:
 		FMP4BoxTREX(const FMP4BoxTREX&) = delete;
 
 	protected:
-		virtual UEMediaError ReadAndParseAttributes(FMP4ParseInfo* ParseInfo) override
+		UEMediaError ReadAndParseAttributes(FMP4ParseInfo* ParseInfo) override
 		{
 			UEMediaError Error = UEMEDIA_ERROR_OK;
 
@@ -1394,7 +1395,7 @@ private:
 		FMP4BoxMVEX(const FMP4BoxMVEX&) = delete;
 
 	protected:
-		virtual UEMediaError ReadAndParseAttributes(FMP4ParseInfo* ParseInfo) override
+		UEMediaError ReadAndParseAttributes(FMP4ParseInfo* ParseInfo) override
 		{
 			// The mvex box is merely a container with no attributes.
 			return UEMEDIA_ERROR_OK;
@@ -1442,7 +1443,7 @@ private:
 		FMP4BoxMDHD(const FMP4BoxMDHD&) = delete;
 
 	protected:
-		virtual UEMediaError ReadAndParseAttributes(FMP4ParseInfo* ParseInfo) override
+		UEMediaError ReadAndParseAttributes(FMP4ParseInfo* ParseInfo) override
 		{
 			uint64 Value64 = 0;
 			uint32 Value32 = 0;
@@ -1513,7 +1514,7 @@ private:
 		FMP4BoxHDLR(const FMP4BoxHDLR&) = delete;
 
 	protected:
-		virtual UEMediaError ReadAndParseAttributes(FMP4ParseInfo* ParseInfo) override
+		UEMediaError ReadAndParseAttributes(FMP4ParseInfo* ParseInfo) override
 		{
 			uint32 Value32 = 0;
 			UEMediaError Error = UEMEDIA_ERROR_OK;
@@ -1581,7 +1582,7 @@ private:
 		FMP4BoxELST(const FMP4BoxELST&) = delete;
 
 	protected:
-		virtual UEMediaError ReadAndParseAttributes(FMP4ParseInfo* ParseInfo) override
+		UEMediaError ReadAndParseAttributes(FMP4ParseInfo* ParseInfo) override
 		{
 			uint32	NumEntries = 0;
 			UEMediaError Error = UEMEDIA_ERROR_OK;
@@ -1656,7 +1657,7 @@ private:
 		FMP4BoxNMHD(const FMP4BoxNMHD&) = delete;
 
 	protected:
-		virtual UEMediaError ReadAndParseAttributes(FMP4ParseInfo* ParseInfo) override
+		UEMediaError ReadAndParseAttributes(FMP4ParseInfo* ParseInfo) override
 		{
 			UEMediaError Error = UEMEDIA_ERROR_OK;
 			RETURN_IF_ERROR(FMP4BoxFull::ReadAndParseAttributes(ParseInfo));
@@ -1688,7 +1689,7 @@ private:
 		FMP4BoxVMHD(const FMP4BoxVMHD&) = delete;
 
 	protected:
-		virtual UEMediaError ReadAndParseAttributes(FMP4ParseInfo* ParseInfo) override
+		UEMediaError ReadAndParseAttributes(FMP4ParseInfo* ParseInfo) override
 		{
 			uint16 Value16 = 0;
 			UEMediaError Error = UEMEDIA_ERROR_OK;
@@ -1727,7 +1728,7 @@ private:
 		FMP4BoxSMHD(const FMP4BoxSMHD&) = delete;
 
 	protected:
-		virtual UEMediaError ReadAndParseAttributes(FMP4ParseInfo* ParseInfo) override
+		UEMediaError ReadAndParseAttributes(FMP4ParseInfo* ParseInfo) override
 		{
 			uint16 Value16 = 0;
 			UEMediaError Error = UEMEDIA_ERROR_OK;
@@ -1764,7 +1765,7 @@ private:
 		FMP4BoxSTHD(const FMP4BoxSTHD&) = delete;
 
 	protected:
-		virtual UEMediaError ReadAndParseAttributes(FMP4ParseInfo* ParseInfo) override
+		UEMediaError ReadAndParseAttributes(FMP4ParseInfo* ParseInfo) override
 		{
 			UEMediaError Error = UEMEDIA_ERROR_OK;
 			RETURN_IF_ERROR(FMP4BoxFull::ReadAndParseAttributes(ParseInfo));
@@ -1797,7 +1798,7 @@ private:
 		FMP4BoxDREF(const FMP4BoxDREF&) = delete;
 
 	protected:
-		virtual UEMediaError ReadAndParseAttributes(FMP4ParseInfo* ParseInfo) override
+		UEMediaError ReadAndParseAttributes(FMP4ParseInfo* ParseInfo) override
 		{
 			// Now that we parse this box we clear its leaf status. We call back into the generic box reader
 			// for which this flag must be clear to continue with the next boxes.
@@ -1838,7 +1839,7 @@ private:
 		FMP4BoxURL(const FMP4BoxURL&) = delete;
 
 	protected:
-		virtual UEMediaError ReadAndParseAttributes(FMP4ParseInfo* ParseInfo) override
+		UEMediaError ReadAndParseAttributes(FMP4ParseInfo* ParseInfo) override
 		{
 			UEMediaError Error = UEMEDIA_ERROR_OK;
 			RETURN_IF_ERROR(FMP4BoxFull::ReadAndParseAttributes(ParseInfo));
@@ -1878,7 +1879,7 @@ private:
 		FMP4BoxURN(const FMP4BoxURN&) = delete;
 
 	protected:
-		virtual UEMediaError ReadAndParseAttributes(FMP4ParseInfo* ParseInfo) override
+		UEMediaError ReadAndParseAttributes(FMP4ParseInfo* ParseInfo) override
 		{
 			UEMediaError Error = UEMEDIA_ERROR_OK;
 			RETURN_IF_ERROR(FMP4BoxFull::ReadAndParseAttributes(ParseInfo));
@@ -1920,7 +1921,7 @@ private:
 		FMP4BoxSampleEntry(const FMP4BoxSampleEntry&) = delete;
 
 	protected:
-		virtual UEMediaError ReadAndParseAttributes(FMP4ParseInfo* ParseInfo) override
+		UEMediaError ReadAndParseAttributes(FMP4ParseInfo* ParseInfo) override
 		{
 			UEMediaError Error = UEMEDIA_ERROR_OK;
 			RETURN_IF_ERROR(ParseInfo->Reader()->ReadBytes(nullptr, 6));			// reserved
@@ -1953,7 +1954,7 @@ private:
 		FMP4BoxVisualSampleEntry(const FMP4BoxVisualSampleEntry&) = delete;
 
 	protected:
-		virtual UEMediaError ReadAndParseAttributes(FMP4ParseInfo* ParseInfo) override
+		UEMediaError ReadAndParseAttributes(FMP4ParseInfo* ParseInfo) override
 		{
 			uint32 Value32 = 0;
 			uint16 Value16 = 0;
@@ -2025,7 +2026,7 @@ private:
 		FMP4BoxAudioSampleEntry(const FMP4BoxAudioSampleEntry&) = delete;
 
 	protected:
-		virtual UEMediaError ReadAndParseAttributes(FMP4ParseInfo* ParseInfo) override
+		UEMediaError ReadAndParseAttributes(FMP4ParseInfo* ParseInfo) override
 		{
 			// The audio sample entry can be version 0 or version 1.
 			// In ISO/IEC 14496-12 a version 1 sample is required to be inside a version 1 'stsd' box while in QuickTime a
@@ -2098,7 +2099,7 @@ private:
 		FMP4BoxPlainTextSampleEntry(const FMP4BoxPlainTextSampleEntry&) = delete;
 
 	protected:
-		virtual UEMediaError ReadAndParseAttributes(FMP4ParseInfo* ParseInfo) override
+		UEMediaError ReadAndParseAttributes(FMP4ParseInfo* ParseInfo) override
 		{
 			UEMediaError Error = UEMEDIA_ERROR_OK;
 			RETURN_IF_ERROR(FMP4BoxSampleEntry::ReadAndParseAttributes(ParseInfo));
@@ -2195,7 +2196,7 @@ private:
 		FMP4BoxTX3GSampleEntry(const FMP4BoxTX3GSampleEntry&) = delete;
 
 	protected:
-		virtual UEMediaError ReadAndParseAttributes(FMP4ParseInfo* ParseInfo) override
+		UEMediaError ReadAndParseAttributes(FMP4ParseInfo* ParseInfo) override
 		{
 			UEMediaError Error = UEMEDIA_ERROR_OK;
 
@@ -2355,7 +2356,7 @@ private:
 		FMP4BoxWVTTSampleEntry(const FMP4BoxWVTTSampleEntry&) = delete;
 
 	protected:
-		virtual UEMediaError ReadAndParseAttributes(FMP4ParseInfo* ParseInfo) override
+		UEMediaError ReadAndParseAttributes(FMP4ParseInfo* ParseInfo) override
 		{
 			UEMediaError Error = UEMEDIA_ERROR_OK;
 
@@ -2490,7 +2491,7 @@ private:
 		FMP4BoxXMLSubtitleSampleEntry(const FMP4BoxXMLSubtitleSampleEntry&) = delete;
 
 	protected:
-		virtual UEMediaError ReadAndParseAttributes(FMP4ParseInfo* ParseInfo) override
+		UEMediaError ReadAndParseAttributes(FMP4ParseInfo* ParseInfo) override
 		{
 			UEMediaError Error = UEMEDIA_ERROR_OK;
 
@@ -2591,7 +2592,7 @@ private:
 		FMP4BoxAVCC(const FMP4BoxAVCC&) = delete;
 
 	protected:
-		virtual UEMediaError ReadAndParseAttributes(FMP4ParseInfo* ParseInfo) override
+		UEMediaError ReadAndParseAttributes(FMP4ParseInfo* ParseInfo) override
 		{
 			UEMediaError Error = UEMEDIA_ERROR_OK;
 			uint32 BytesRemaining = BoxSize - (ParseInfo->Reader()->GetCurrentReadOffset() - StartOffset);
@@ -2640,7 +2641,7 @@ private:
 		FMP4BoxHVCC(const FMP4BoxHVCC&) = delete;
 
 	protected:
-		virtual UEMediaError ReadAndParseAttributes(FMP4ParseInfo* ParseInfo) override
+		UEMediaError ReadAndParseAttributes(FMP4ParseInfo* ParseInfo) override
 		{
 			UEMediaError Error = UEMEDIA_ERROR_OK;
 			uint32 BytesRemaining = BoxSize - (ParseInfo->Reader()->GetCurrentReadOffset() - StartOffset);
@@ -2691,7 +2692,7 @@ private:
 		FMP4BoxESDS(const FMP4BoxESDS&) = delete;
 
 	protected:
-		virtual UEMediaError ReadAndParseAttributes(FMP4ParseInfo* ParseInfo) override
+		UEMediaError ReadAndParseAttributes(FMP4ParseInfo* ParseInfo) override
 		{
 			UEMediaError Error = UEMEDIA_ERROR_OK;
 			RETURN_IF_ERROR(FMP4BoxFull::ReadAndParseAttributes(ParseInfo));
@@ -2738,7 +2739,7 @@ private:
 		FMP4BoxDAC3(const FMP4BoxDAC3&) = delete;
 
 	protected:
-		virtual UEMediaError ReadAndParseAttributes(FMP4ParseInfo* ParseInfo) override
+		UEMediaError ReadAndParseAttributes(FMP4ParseInfo* ParseInfo) override
 		{
 			UEMediaError Error = UEMEDIA_ERROR_OK;
 
@@ -2775,7 +2776,7 @@ private:
 		FMP4BoxDEC3(const FMP4BoxDEC3&) = delete;
 
 	protected:
-		virtual UEMediaError ReadAndParseAttributes(FMP4ParseInfo* ParseInfo) override
+		UEMediaError ReadAndParseAttributes(FMP4ParseInfo* ParseInfo) override
 		{
 			UEMediaError Error = UEMEDIA_ERROR_OK;
 
@@ -2827,7 +2828,7 @@ private:
 		FMP4BoxBTRT(const FMP4BoxBTRT&) = delete;
 
 	protected:
-		virtual UEMediaError ReadAndParseAttributes(FMP4ParseInfo* ParseInfo) override
+		UEMediaError ReadAndParseAttributes(FMP4ParseInfo* ParseInfo) override
 		{
 			UEMediaError Error = UEMEDIA_ERROR_OK;
 			RETURN_IF_ERROR(ParseInfo->Reader()->Read(BufferSizeDB));				// bufferSizeDB
@@ -2865,7 +2866,7 @@ private:
 		FMP4BoxPASP(const FMP4BoxPASP&) = delete;
 
 	protected:
-		virtual UEMediaError ReadAndParseAttributes(FMP4ParseInfo* ParseInfo) override
+		UEMediaError ReadAndParseAttributes(FMP4ParseInfo* ParseInfo) override
 		{
 			UEMediaError Error = UEMEDIA_ERROR_OK;
 			RETURN_IF_ERROR(ParseInfo->Reader()->Read(HSpacing));				// hSpacing
@@ -2900,7 +2901,7 @@ private:
 		FMP4BoxELNG(const FMP4BoxELNG&) = delete;
 
 	protected:
-		virtual UEMediaError ReadAndParseAttributes(FMP4ParseInfo* ParseInfo) override
+		UEMediaError ReadAndParseAttributes(FMP4ParseInfo* ParseInfo) override
 		{
 			UEMediaError Error = UEMEDIA_ERROR_OK;
 			RETURN_IF_ERROR(FMP4BoxFull::ReadAndParseAttributes(ParseInfo));
@@ -2951,7 +2952,7 @@ private:
 		FMP4BoxPRFT(const FMP4BoxPRFT&) = delete;
 
 	protected:
-		virtual UEMediaError ReadAndParseAttributes(FMP4ParseInfo* ParseInfo) override
+		UEMediaError ReadAndParseAttributes(FMP4ParseInfo* ParseInfo) override
 		{
 			UEMediaError Error = UEMEDIA_ERROR_OK;
 			RETURN_IF_ERROR(FMP4BoxFull::ReadAndParseAttributes(ParseInfo));
@@ -2995,39 +2996,39 @@ private:
 		{
 		}
 
-		virtual int32 GetVersion() const override
-		{ 
-			return Version; 
+		int32 GetVersion() const override
+		{
+			return Version;
 		}
-		virtual const FString& GetSchemeIdUri() const override
-		{ 
-			return SchemeIdUri; 
+		const FString& GetSchemeIdUri() const override
+		{
+			return SchemeIdUri;
 		}
-		virtual const FString& GetValue() const override
-		{ 
-			return Value; 
+		const FString& GetValue() const override
+		{
+			return Value;
 		}
-		virtual uint32 GetTimescale() const override
+		uint32 GetTimescale() const override
 		{
 			return Timescale;
 		}
-		virtual uint32 GetPresentationTimeDelta() const override
+		uint32 GetPresentationTimeDelta() const override
 		{
 			return PresentationTimeDelta;
 		}
-		virtual uint64 GetPresentationTime() const override
+		uint64 GetPresentationTime() const override
 		{
 			return PresentationTime;
 		}
-		virtual uint32 GetEventDuration() const override
+		uint32 GetEventDuration() const override
 		{
 			return EventDuration;
 		}
-		virtual uint32 GetID() const override
+		uint32 GetID() const override
 		{
 			return ID;
 		}
-		virtual const TArray<uint8>& GetMessageData() const override
+		const TArray<uint8>& GetMessageData() const override
 		{
 			return MessageData;
 		}
@@ -3037,7 +3038,7 @@ private:
 		FMP4BoxEMSG(const FMP4BoxEMSG&) = delete;
 
 	protected:
-		virtual UEMediaError ReadAndParseAttributes(FMP4ParseInfo* ParseInfo) override
+		UEMediaError ReadAndParseAttributes(FMP4ParseInfo* ParseInfo) override
 		{
 			UEMediaError Error = UEMEDIA_ERROR_OK;
 			RETURN_IF_ERROR(FMP4BoxFull::ReadAndParseAttributes(ParseInfo));
@@ -3115,7 +3116,7 @@ private:
 		FMP4BoxSTSD(const FMP4BoxSTSD&) = delete;
 
 	protected:
-		virtual UEMediaError ReadAndParseAttributes(FMP4ParseInfo* ParseInfo) override
+		UEMediaError ReadAndParseAttributes(FMP4ParseInfo* ParseInfo) override
 		{
 			// Now that we parse this box we clear its leaf status. We call back into the generic box reader
 			// for which this flag must be clear to continue with the next boxes.
@@ -3312,7 +3313,7 @@ private:
 		FMP4BoxSTTS(const FMP4BoxSTTS&) = delete;
 
 	protected:
-		virtual UEMediaError ReadAndParseAttributes(FMP4ParseInfo* ParseInfo) override
+		UEMediaError ReadAndParseAttributes(FMP4ParseInfo* ParseInfo) override
 		{
 			UEMediaError Error = UEMEDIA_ERROR_OK;
 			RETURN_IF_ERROR(FMP4BoxFull::ReadAndParseAttributes(ParseInfo));
@@ -3394,7 +3395,7 @@ private:
 			uint32		SampleCount;
 			uint32		SampleOffset;
 		};
-		virtual UEMediaError ReadAndParseAttributes(FMP4ParseInfo* ParseInfo) override
+		UEMediaError ReadAndParseAttributes(FMP4ParseInfo* ParseInfo) override
 		{
 			UEMediaError Error = UEMEDIA_ERROR_OK;
 			RETURN_IF_ERROR(FMP4BoxFull::ReadAndParseAttributes(ParseInfo));
@@ -3461,7 +3462,7 @@ private:
 		FMP4BoxSTSS(const FMP4BoxSTSS&) = delete;
 
 	protected:
-		virtual UEMediaError ReadAndParseAttributes(FMP4ParseInfo* ParseInfo) override
+		UEMediaError ReadAndParseAttributes(FMP4ParseInfo* ParseInfo) override
 		{
 			UEMediaError Error = UEMEDIA_ERROR_OK;
 			RETURN_IF_ERROR(FMP4BoxFull::ReadAndParseAttributes(ParseInfo));
@@ -3523,7 +3524,7 @@ private:
 		FMP4BoxSTSC(const FMP4BoxSTSC&) = delete;
 
 	protected:
-		virtual UEMediaError ReadAndParseAttributes(FMP4ParseInfo* ParseInfo) override
+		UEMediaError ReadAndParseAttributes(FMP4ParseInfo* ParseInfo) override
 		{
 			UEMediaError Error = UEMEDIA_ERROR_OK;
 			RETURN_IF_ERROR(FMP4BoxFull::ReadAndParseAttributes(ParseInfo));
@@ -3593,7 +3594,7 @@ private:
 		FMP4BoxSTSZ(const FMP4BoxSTSZ&) = delete;
 
 	protected:
-		virtual UEMediaError ReadAndParseAttributes(FMP4ParseInfo* ParseInfo) override
+		UEMediaError ReadAndParseAttributes(FMP4ParseInfo* ParseInfo) override
 		{
 			UEMediaError Error = UEMEDIA_ERROR_OK;
 			RETURN_IF_ERROR(FMP4BoxFull::ReadAndParseAttributes(ParseInfo));
@@ -3651,7 +3652,7 @@ private:
 		FMP4BoxSTCO(const FMP4BoxSTCO&) = delete;
 
 	protected:
-		virtual UEMediaError ReadAndParseAttributes(FMP4ParseInfo* ParseInfo) override
+		UEMediaError ReadAndParseAttributes(FMP4ParseInfo* ParseInfo) override
 		{
 			UEMediaError Error = UEMEDIA_ERROR_OK;
 			RETURN_IF_ERROR(FMP4BoxFull::ReadAndParseAttributes(ParseInfo));
@@ -3713,15 +3714,15 @@ private:
 		// ----------------------------------------------------------------
 		// Methods from IParserISO14496_12::ISegmentIndex
 		//
-		virtual uint64 GetEarliestPresentationTime() const override
+		uint64 GetEarliestPresentationTime() const override
 		{ return EarliestPresentationTime; }
-		virtual uint64 GetFirstOffset() const override
+		uint64 GetFirstOffset() const override
 		{ return FirstOffset; }
-		virtual uint32 GetReferenceID() const override
+		uint32 GetReferenceID() const override
 		{ return ReferenceID; }
-		virtual uint32 GetTimescale() const override
+		uint32 GetTimescale() const override
 		{ return Timescale; }
-		virtual int32 GetNumEntries() const override
+		int32 GetNumEntries() const override
 		{ return Entries.Num(); }
 		virtual const FEntry& GetEntry(int32 Index) const
 		{
@@ -3734,7 +3735,7 @@ private:
 		FMP4BoxSIDX(const FMP4BoxSIDX&) = delete;
 
 	protected:
-		virtual UEMediaError ReadAndParseAttributes(FMP4ParseInfo* ParseInfo) override
+		UEMediaError ReadAndParseAttributes(FMP4ParseInfo* ParseInfo) override
 		{
 			UEMediaError Error = UEMEDIA_ERROR_OK;
 			RETURN_IF_ERROR(FMP4BoxFull::ReadAndParseAttributes(ParseInfo));
@@ -3815,7 +3816,7 @@ private:
 		FMP4BoxMFHD(const FMP4BoxMFHD&) = delete;
 
 	protected:
-		virtual UEMediaError ReadAndParseAttributes(FMP4ParseInfo* ParseInfo) override
+		UEMediaError ReadAndParseAttributes(FMP4ParseInfo* ParseInfo) override
 		{
 			UEMediaError Error = UEMEDIA_ERROR_OK;
 			RETURN_IF_ERROR(FMP4BoxFull::ReadAndParseAttributes(ParseInfo));
@@ -3909,7 +3910,7 @@ private:
 		FMP4BoxTFHD(const FMP4BoxTFHD&) = delete;
 
 	protected:
-		virtual UEMediaError ReadAndParseAttributes(FMP4ParseInfo* ParseInfo) override
+		UEMediaError ReadAndParseAttributes(FMP4ParseInfo* ParseInfo) override
 		{
 			UEMediaError Error = UEMEDIA_ERROR_OK;
 			RETURN_IF_ERROR(FMP4BoxFull::ReadAndParseAttributes(ParseInfo));
@@ -3973,7 +3974,7 @@ private:
 		FMP4BoxTFDT(const FMP4BoxTFDT&) = delete;
 
 	protected:
-		virtual UEMediaError ReadAndParseAttributes(FMP4ParseInfo* ParseInfo) override
+		UEMediaError ReadAndParseAttributes(FMP4ParseInfo* ParseInfo) override
 		{
 			UEMediaError Error = UEMEDIA_ERROR_OK;
 			RETURN_IF_ERROR(FMP4BoxFull::ReadAndParseAttributes(ParseInfo));
@@ -4075,7 +4076,7 @@ private:
 		FMP4BoxTRUN(const FMP4BoxTRUN&) = delete;
 
 	protected:
-		virtual UEMediaError ReadAndParseAttributes(FMP4ParseInfo* ParseInfo) override
+		UEMediaError ReadAndParseAttributes(FMP4ParseInfo* ParseInfo) override
 		{
 			UEMediaError Error = UEMEDIA_ERROR_OK;
 			RETURN_IF_ERROR(FMP4BoxFull::ReadAndParseAttributes(ParseInfo));
@@ -4197,7 +4198,7 @@ private:
 		FMP4BoxPSSH(const FMP4BoxPSSH&) = delete;
 
 	protected:
-		virtual UEMediaError ReadAndParseAttributes(FMP4ParseInfo* ParseInfo) override
+		UEMediaError ReadAndParseAttributes(FMP4ParseInfo* ParseInfo) override
 		{
 			UEMediaError Error = UEMEDIA_ERROR_OK;
 			// We need to retain the box as a whole since some DRM systems need it to be passed in verbatim.
@@ -4266,7 +4267,7 @@ private:
 		FMP4BoxSINF(const FMP4BoxSINF&) = delete;
 
 	protected:
-		virtual UEMediaError ReadAndParseAttributes(FMP4ParseInfo* ParseInfo) override
+		UEMediaError ReadAndParseAttributes(FMP4ParseInfo* ParseInfo) override
 		{
 			// Now that we parse this box we clear its leaf status. We call back into the generic box reader
 			// for which this flag must be clear to continue with the next boxes.
@@ -4304,7 +4305,7 @@ private:
 		FMP4BoxFRMA(const FMP4BoxFRMA&) = delete;
 
 	protected:
-		virtual UEMediaError ReadAndParseAttributes(FMP4ParseInfo* ParseInfo) override
+		UEMediaError ReadAndParseAttributes(FMP4ParseInfo* ParseInfo) override
 		{
 			UEMediaError Error = UEMEDIA_ERROR_OK;
 			RETURN_IF_ERROR(ParseInfo->Reader()->Read(DataFormat));						// data_format
@@ -4342,7 +4343,7 @@ private:
 		FMP4BoxSCHM(const FMP4BoxSCHM&) = delete;
 
 	protected:
-		virtual UEMediaError ReadAndParseAttributes(FMP4ParseInfo* ParseInfo) override
+		UEMediaError ReadAndParseAttributes(FMP4ParseInfo* ParseInfo) override
 		{
 			UEMediaError Error = UEMEDIA_ERROR_OK;
 			RETURN_IF_ERROR(FMP4BoxFull::ReadAndParseAttributes(ParseInfo));
@@ -4417,7 +4418,7 @@ private:
 		FMP4BoxTENC(const FMP4BoxTENC&) = delete;
 
 	protected:
-		virtual UEMediaError ReadAndParseAttributes(FMP4ParseInfo* ParseInfo) override
+		UEMediaError ReadAndParseAttributes(FMP4ParseInfo* ParseInfo) override
 		{
 			UEMediaError Error = UEMEDIA_ERROR_OK;
 			RETURN_IF_ERROR(FMP4BoxFull::ReadAndParseAttributes(ParseInfo));
@@ -4478,7 +4479,7 @@ private:
 		FMP4BoxSAIZ(const FMP4BoxSAIZ&) = delete;
 
 	protected:
-		virtual UEMediaError ReadAndParseAttributes(FMP4ParseInfo* ParseInfo) override
+		UEMediaError ReadAndParseAttributes(FMP4ParseInfo* ParseInfo) override
 		{
 			UEMediaError Error = UEMEDIA_ERROR_OK;
 			RETURN_IF_ERROR(FMP4BoxFull::ReadAndParseAttributes(ParseInfo));
@@ -4529,7 +4530,7 @@ private:
 		FMP4BoxSAIO(const FMP4BoxSAIO&) = delete;
 
 	protected:
-		virtual UEMediaError ReadAndParseAttributes(FMP4ParseInfo* ParseInfo) override
+		UEMediaError ReadAndParseAttributes(FMP4ParseInfo* ParseInfo) override
 		{
 			UEMediaError Error = UEMEDIA_ERROR_OK;
 			RETURN_IF_ERROR(FMP4BoxFull::ReadAndParseAttributes(ParseInfo));
@@ -4609,7 +4610,7 @@ private:
 			{
 				return UEMEDIA_ERROR_FORMAT_ERROR;
 			}
-			
+
 			int32 PerSampleIVSize = TENCBox->GetDefaultPerSampleIVSize();
 
 			FDataBufferReader BufferReader(BoxData);
@@ -4647,7 +4648,7 @@ private:
 		FMP4BoxSENC(const FMP4BoxSENC&) = delete;
 
 	protected:
-		virtual UEMediaError ReadAndParseAttributes(FMP4ParseInfo* ParseInfo) override
+		UEMediaError ReadAndParseAttributes(FMP4ParseInfo* ParseInfo) override
 		{
 			UEMediaError Error = UEMEDIA_ERROR_OK;
 			RETURN_IF_ERROR(FMP4BoxFull::ReadAndParseAttributes(ParseInfo));
@@ -4657,7 +4658,7 @@ private:
 			{
 				return UEMEDIA_ERROR_FORMAT_ERROR;
 			}
-			
+
 			RETURN_IF_ERROR(ParseInfo->Reader()->Read(SampleCount));
 			// The contents of this box cannot be parsed without information from another box
 			// containing the size of the IV.
@@ -4708,7 +4709,7 @@ private:
 
 		int32 GetNumberOfEntries() const
 		{ return Entries.Num(); }
-		
+
 		const FEntry& GetEntry(int32 Index) const
 		{ return Entries[Index]; }
 
@@ -4717,7 +4718,7 @@ private:
 		FMP4BoxSBGP(const FMP4BoxSBGP&) = delete;
 
 	protected:
-		virtual UEMediaError ReadAndParseAttributes(FMP4ParseInfo* ParseInfo) override
+		UEMediaError ReadAndParseAttributes(FMP4ParseInfo* ParseInfo) override
 		{
 			UEMediaError Error = UEMEDIA_ERROR_OK;
 			RETURN_IF_ERROR(FMP4BoxFull::ReadAndParseAttributes(ParseInfo));
@@ -4792,7 +4793,7 @@ private:
 
 		int32 GetNumberOfEntries() const
 		{ return Entries.Num(); }
-		
+
 		const FSampleGroupDescriptionEntry& GetEntry(int32 Index) const
 		{ return Entries[Index]; }
 
@@ -4801,7 +4802,7 @@ private:
 		FMP4BoxSGPD(const FMP4BoxSGPD&) = delete;
 
 	protected:
-		virtual UEMediaError ReadAndParseAttributes(FMP4ParseInfo* ParseInfo) override
+		UEMediaError ReadAndParseAttributes(FMP4ParseInfo* ParseInfo) override
 		{
 			UEMediaError Error = UEMEDIA_ERROR_OK;
 			RETURN_IF_ERROR(FMP4BoxFull::ReadAndParseAttributes(ParseInfo));
@@ -4819,7 +4820,7 @@ private:
 			}
 
 			RETURN_IF_ERROR(ParseInfo->Reader()->Read(GroupingType));							// grouping_type
-			
+
 			if (Version == 1)
 			{
 				RETURN_IF_ERROR(ParseInfo->Reader()->Read(DefaultLength));						// default_length
@@ -4855,6 +4856,81 @@ private:
 	};
 
 
+
+	/**
+	 * 'udta' box. ISO/IEC 14496-12:2015 - 8.10.1 User Data Box
+	 */
+	class FMP4BoxUDTA : public FMP4BoxBasic
+	{
+	public:
+		FMP4BoxUDTA(IParserISO14496_12::FBoxType InBoxType, int64 InBoxSize, int64 InStartOffset, int64 InDataOffset, bool bInIsLeafBox)
+			: FMP4BoxBasic(InBoxType, InBoxSize, InStartOffset, InDataOffset, bInIsLeafBox)
+		{
+		}
+
+		virtual ~FMP4BoxUDTA() = default;
+
+
+		TMediaOptionalValue<FString> GetName() const
+		{
+			static const IParserISO14496_12::FBoxType kUDTA_name = MAKE_BOX_ATOM('n', 'a', 'm', 'e');
+			const FMP4BoxIgnored* NameBox = static_cast<const FMP4BoxIgnored*>(FindBox(kUDTA_name, 1));
+			const FFillerData* Filler = NameBox ? NameBox->GetFillerData() : nullptr;
+			if (Filler)
+			{
+				FString name(Filler->Size, static_cast<const ANSICHAR*>(Filler->Data));
+				return TMediaOptionalValue<FString>(name);
+			}
+			return TMediaOptionalValue<FString>();
+		}
+
+	private:
+		FMP4BoxUDTA() = delete;
+		FMP4BoxUDTA(const FMP4BoxUDTA&) = delete;
+
+	protected:
+		UEMediaError ReadAndParseAttributes(FMP4ParseInfo* ParseInfo) override
+		{
+			UEMediaError Error = UEMEDIA_ERROR_OK;
+
+			int32 BytesRemaining = (int32) (BoxSize - (ParseInfo->Reader()->GetCurrentReadOffset() - StartOffset));
+			if (BytesRemaining > 8)
+			{
+				// Now that there are additional boxes we ourselves are no longer a leaf box.
+				bIsLeafBox = false;
+				// This will read all the boxes in here and add them as children.
+				for(; BytesRemaining>8; BytesRemaining = (int32) (BoxSize - (ParseInfo->Reader()->GetCurrentReadOffset() - StartOffset)))
+				{
+					IParserISO14496_12::FBoxType	ChildBoxType;
+					int64							ChildBoxSize;
+					uint8							ChildBoxUUID[16];
+
+					int64 BoxStartOffset = ParseInfo->Reader()->GetCurrentReadOffset();
+					RETURN_IF_ERROR(ParseInfo->ReadBoxTypeAndSize(ChildBoxType, ChildBoxSize, ChildBoxUUID));
+#if MEDIA_DEBUG_HAS_BOX_NAMES
+					char boxName[4];
+					*((uint32*)&boxName) = MEDIA_TO_BIG_ENDIAN(ChildBoxType);
+#endif
+					int64 BoxDataOffset = ParseInfo->Reader()->GetCurrentReadOffset();
+					FMP4Box* NextBox = new FMP4BoxIgnored(ChildBoxType, ChildBoxSize, BoxStartOffset, BoxDataOffset, true);
+
+					AddChildBox(NextBox);
+					RETURN_IF_ERROR(ParseInfo->ReadAndParseNextBox(NextBox));
+				}
+			}
+			// If there is still something left that is not a complete box, skip over it.
+			// There may be a 32 bit integer zero value to terminate the list.
+			// See: "User Data Atoms" in https://developer.apple.com/library/archive/documentation/QuickTime/QTFF/QTFFChap2/qtff2.html
+			BytesRemaining = (int32) (BoxSize - (ParseInfo->Reader()->GetCurrentReadOffset() - StartOffset));
+			if (BytesRemaining)
+			{
+				RETURN_IF_ERROR(ParseInfo->Reader()->ReadBytes(nullptr, BytesRemaining));
+			}
+			return Error;
+		}
+	};
+
+
 	FMP4ParseInfo::~FMP4ParseInfo()
 	{
 		delete RootBox;
@@ -5169,8 +5245,10 @@ private:
 						case FMP4Box::kBox_gmhd:
 							NextBox = new FMP4BoxIgnored(BoxType, BoxSize, BoxStartOffset, BoxDataOffset, true);
 							break;
-						case FMP4Box::kBox_skip:
 						case FMP4Box::kBox_udta:
+							NextBox = new FMP4BoxUDTA(BoxType, BoxSize, BoxStartOffset, BoxDataOffset, true);
+							break;
+						case FMP4Box::kBox_skip:
 						case FMP4Box::kBox_VOID:
 							NextBox = new FMP4BoxIgnored(BoxType, BoxSize, BoxStartOffset, BoxDataOffset, true);
 							break;
@@ -5347,23 +5425,23 @@ private:
 	public:
 		FParserISO14496_12();
 		virtual ~FParserISO14496_12();
-		virtual UEMediaError ParseHeader(IReader* DataReader, IBoxCallback* BoxParseCallback, IPlayerSessionServices* PlayerSession, const IParserISO14496_12* OptionalInitSegment) override;
+		UEMediaError ParseHeader(IReader* DataReader, IBoxCallback* BoxParseCallback, IPlayerSessionServices* PlayerSession, const IParserISO14496_12* OptionalInitSegment) override;
 
-		virtual UEMediaError PrepareTracks(IPlayerSessionServices* PlayerSession, TSharedPtrTS<const IParserISO14496_12> OptionalMP4InitSegment) override;
+		UEMediaError PrepareTracks(IPlayerSessionServices* PlayerSession, TSharedPtrTS<const IParserISO14496_12> OptionalMP4InitSegment) override;
 
-		virtual TMediaOptionalValue<FTimeFraction> GetMovieDuration() const override;
-		virtual int32 GetNumberOfTracks() const override;
-		virtual int32 GetNumberOfSegmentIndices() const override;
-		virtual int32 GetNumberOfEventMessages() const override;
-		virtual int32 GetNumberOfBrands() const override;
-		virtual FBrandType GetBrandByIndex(int32 Index) const override;
-		virtual bool HasBrand(const FBrandType InBrand) const override;
-		virtual const ITrack* GetTrackByIndex(int32 Index) const override;
-		virtual const ITrack* GetTrackByTrackID(int32 TrackID) const override;
-		virtual const ISegmentIndex* GetSegmentIndexByIndex(int32 Index) const override;
-		virtual const IEventMessage* GetEventMessageByIndex(int32 Index) const override;
+		TMediaOptionalValue<FTimeFraction> GetMovieDuration() const override;
+		int32 GetNumberOfTracks() const override;
+		int32 GetNumberOfSegmentIndices() const override;
+		int32 GetNumberOfEventMessages() const override;
+		int32 GetNumberOfBrands() const override;
+		FBrandType GetBrandByIndex(int32 Index) const override;
+		bool HasBrand(const FBrandType InBrand) const override;
+		const ITrack* GetTrackByIndex(int32 Index) const override;
+		const ITrack* GetTrackByTrackID(int32 TrackID) const override;
+		const ISegmentIndex* GetSegmentIndexByIndex(int32 Index) const override;
+		const IEventMessage* GetEventMessageByIndex(int32 Index) const override;
 
-		virtual TSharedPtrTS<IAllTrackIterator> CreateAllTrackIteratorByFilePos(int64 InFromFilePos) const override;
+		TSharedPtrTS<IAllTrackIterator> CreateAllTrackIteratorByFilePos(int64 InFromFilePos) const override;
 
 	private:
 		class FTrack;
@@ -5374,27 +5452,27 @@ private:
 			FTrackIterator();
 			virtual ~FTrackIterator();
 
-			virtual const ITrack* GetTrack() const override;
-			virtual int64 GetBaseMediaDecodeTime() const override;
-			virtual bool IsAtEOS() const override;
-
-			virtual UEMediaError StartAtTime(const FTimeValue& AtTime, ESearchMode SearchMode, bool bNeedSyncSample) override;
-			virtual UEMediaError StartAtFirst(bool bNeedSyncSample) override;
-			virtual UEMediaError Next() override;
-
-			virtual uint32 GetSampleNumber() const override;
-			virtual int64 GetDTS() const override;
-			virtual int64 GetPTS() const override;
-			virtual int64 GetDuration() const override;
-			virtual uint32 GetTimescale() const override;
-			virtual bool IsSyncSample() const override;
-			virtual int64 GetSampleSize() const override;
-			virtual int64 GetSampleFileOffset() const override;
-			virtual int64 GetRawDTS() const override;
-			virtual int64 GetRawPTS() const override;
-			virtual int64 GetCompositionTimeEdit() const override;
-			virtual int64 GetEmptyEditOffset() const override;
-			virtual bool GetEncryptionInfo(ElectraCDM::FMediaCDMSampleInfo& OutSampleEncryptionInfo) const override;
+			const ITrack* GetTrack() const override;
+			int64 GetBaseMediaDecodeTime() const override;
+			bool IsAtEOS() const override;
+
+			UEMediaError StartAtTime(const FTimeValue& AtTime, ESearchMode SearchMode, bool bNeedSyncSample) override;
+			UEMediaError StartAtFirst(bool bNeedSyncSample) override;
+			UEMediaError Next() override;
+
+			uint32 GetSampleNumber() const override;
+			int64 GetDTS() const override;
+			int64 GetPTS() const override;
+			int64 GetDuration() const override;
+			uint32 GetTimescale() const override;
+			bool IsSyncSample() const override;
+			int64 GetSampleSize() const override;
+			int64 GetSampleFileOffset() const override;
+			int64 GetRawDTS() const override;
+			int64 GetRawPTS() const override;
+			int64 GetCompositionTimeEdit() const override;
+			int64 GetEmptyEditOffset() const override;
+			bool GetEncryptionInfo(ElectraCDM::FMediaCDMSampleInfo& OutSampleEncryptionInfo) const override;
 
 			UEMediaError StartAtFirstInteral();
 			void SetTrack(const FTrack* InTrack);
@@ -5491,17 +5569,18 @@ private:
 			virtual ~FTrack()
 			{
 			}
-			virtual uint32 GetID() const override;
-			virtual FString GetNameFromHandler() const override;
-			virtual FTimeFraction GetDuration() const override;
-			virtual ITrackIterator* CreateIterator() const override;
-			virtual const TArray<uint8>& GetCodecSpecificData() const override;
-			virtual const TArray<uint8>& GetCodecSpecificDataRAW() const override;
-			virtual const FStreamCodecInformation& GetCodecInformation() const override;
-			virtual const FBitrateInfo& GetBitrateInfo() const override;
-			virtual const FString GetLanguage() const override;
-			virtual void GetPSSHBoxes(TArray<TArray<uint8>>& OutBoxes, bool bFromMOOV, bool bFromMOOF) const override;
-			virtual void GetPRFTBoxes(TArray<ITrack::FProducerReferenceTime>& OutBoxes) const override;
+			uint32 GetID() const override;
+			FString GetName() const override;
+			FString GetNameFromHandler() const override;
+			FTimeFraction GetDuration() const override;
+			ITrackIterator* CreateIterator() const override;
+			const TArray<uint8>& GetCodecSpecificData() const override;
+			const TArray<uint8>& GetCodecSpecificDataRAW() const override;
+			const FStreamCodecInformation& GetCodecInformation() const override;
+			const FBitrateInfo& GetBitrateInfo() const override;
+			const FString GetLanguage() const override;
+			void GetPSSHBoxes(TArray<TArray<uint8>>& OutBoxes, bool bFromMOOV, bool bFromMOOF) const override;
+			void GetPRFTBoxes(TArray<ITrack::FProducerReferenceTime>& OutBoxes) const override;
 
 			//private:
 			const FMP4BoxMVHD* MVHDBox = nullptr;
@@ -5518,6 +5597,7 @@ private:
 			const FMP4BoxSTSZ* STSZBox = nullptr;
 			const FMP4BoxSTCO* STCOBox = nullptr;
 			const FMP4BoxSTSS* STSSBox = nullptr;
+			const FMP4BoxUDTA* UDTABox = nullptr;
 			// Encryption
 			TArray<const FMP4BoxPSSH*>	PSSHBoxesFromInit;
 			TArray<const FMP4BoxPSSH*>	PSSHBoxes;
@@ -5553,15 +5633,15 @@ private:
 		public:
 			virtual ~FAllTrackIterator();
 			//! Returns the iterator at the current file position.
-			virtual const ITrackIterator* Current() const override;
+			const ITrackIterator* Current() const override;
 			//! Advance iterator to point to the next sample in sequence. Returns false if there are no more samples.
-			virtual bool Next() override;
+			bool Next() override;
 			//! Returns a list of all tracks iterators that reached EOS while iterating since the most recent call to ClearNewEOSTracks().
-			virtual void GetNewEOSTracks(TArray<const ITrackIterator*>& OutTracksThatNewlyReachedEOS) const override;
+			void GetNewEOSTracks(TArray<const ITrackIterator*>& OutTracksThatNewlyReachedEOS) const override;
 			//! Clears the list of track iterators that have reached EOS.
-			virtual void ClearNewEOSTracks() override;
+			void ClearNewEOSTracks() override;
 			//! Returns list of all iterators.
-			virtual void GetAllIterators(TArray<const ITrackIterator*>& OutIterators) const override;
+			void GetAllIterators(TArray<const ITrackIterator*>& OutIterators) const override;
 
 			TArray<TSharedPtrTS<ITrackIterator>>		TrackIterators;
 			TSharedPtrTS<ITrackIterator>				CurrentIterator;
@@ -5666,7 +5746,20 @@ private:
 	{
 		return TKHDBox ? TKHDBox->GetTrackID() : 0;
 	}
-	
+
+	FString FParserISO14496_12::FTrack::GetName() const
+	{
+		if (UDTABox)
+		{
+			TMediaOptionalValue<FString> Name = UDTABox->GetName();
+			if (Name.IsSet())
+			{
+				return Name.Value();
+			}
+		}
+		return FString();
+	}
+
 	FString FParserISO14496_12::FTrack::GetNameFromHandler() const
 	{
 		return HDLRBox ? HDLRBox->GetHandlerName() : FString();
@@ -6449,7 +6542,7 @@ private:
 		{
 			OutSampleEncryptionInfo.DefaultKID = Track->TENCBox->GetDefaultKID();
 		}
-	
+
 		const TArray<FMP4BoxSENC::FEntry>& SencEntries = Track->SENCBox->GetEntries();
 		check((int32) SampleNumberInTRUN < SencEntries.Num());
 		if ((int32) SampleNumberInTRUN >= SencEntries.Num())
@@ -6709,7 +6802,7 @@ private:
 						{
 							// Is this AAC audio?
 							if (Track->CodecSpecificDataMP4A.GetObjectTypeID() == MPEG::FESDescriptor::FObjectTypeID::MPEG4_Audio &&
-								Track->CodecSpecificDataMP4A.GetStreamType() == MPEG::FESDescriptor::FStreamType::Audio)
+								Track->CodecSpecificDataMP4A.GetStreamType() == MPEG::FESDescriptor::FStreamType::AudioStream)
 							{
 								Track->CodecInformation.SetStreamType(EStreamType::Audio);
 								Track->CodecInformation.SetCodec(FStreamCodecInformation::ECodec::AAC);
@@ -6926,6 +7019,7 @@ private:
 						{
 							return UEMEDIA_ERROR_FORMAT_ERROR;
 						}
+						Track->UDTABox = static_cast<const FMP4BoxUDTA*>(Box->GetBoxPath(FMP4Box::kBox_udta));
 						// TODO: Check the dinf->dref->url box to not reference any external media!
 
 						// Get the composition time offset and the empty edit from the edit list. Those are necessary to correct the PTS values from
@@ -7413,7 +7507,7 @@ private:
 		}
 		return 0;
 	}
-	
+
 	IParserISO14496_12::FBrandType FParserISO14496_12::GetBrandByIndex(int32 Index) const
 	{
 		const FMP4BoxFTYP* Box = static_cast<const FMP4BoxFTYP*>(ParsedData->GetCurrentStypBox() ? ParsedData->GetCurrentStypBox() : ParsedData->GetCurrentFtypBox());
diff --git a/Engine/Plugins/Media/ElectraPlayer/Source/ElectraPlayerRuntime/Private/Runtime/Demuxer/ParserISO14496-12.h b/Engine/Plugins/Media/ElectraPlayer/Source/ElectraPlayerRuntime/Private/Runtime/Demuxer/ParserISO14496-12.h
--- a/Engine/Plugins/Media/ElectraPlayer/Source/ElectraPlayerRuntime/Private/Runtime/Demuxer/ParserISO14496-12.h
+++ b/Engine/Plugins/Media/ElectraPlayer/Source/ElectraPlayerRuntime/Private/Runtime/Demuxer/ParserISO14496-12.h
@@ -228,6 +228,7 @@ namespace Electra
 			virtual ~ITrack() = default;
 
 			virtual uint32 GetID() const = 0;
+			virtual FString GetName() const = 0;
 			virtual FString GetNameFromHandler() const = 0;
 			virtual FTimeFraction GetDuration() const = 0;
 
diff --git a/Engine/Plugins/Media/ElectraPlayer/Source/ElectraPlayerRuntime/Private/Runtime/Player/ABRRules/ABRFixedStream.cpp b/Engine/Plugins/Media/ElectraPlayer/Source/ElectraPlayerRuntime/Private/Runtime/Player/ABRRules/ABRFixedStream.cpp
--- a/Engine/Plugins/Media/ElectraPlayer/Source/ElectraPlayerRuntime/Private/Runtime/Player/ABRRules/ABRFixedStream.cpp
+++ b/Engine/Plugins/Media/ElectraPlayer/Source/ElectraPlayerRuntime/Private/Runtime/Player/ABRRules/ABRFixedStream.cpp
@@ -34,7 +34,6 @@ public:
 	void RepresentationsChanged(EStreamType InStreamType, TSharedPtrTS<IManifest::IPlayPeriod> InCurrentPlayPeriod) override {}
 	void SetBandwidth(int64 bitsPerSecond) override {}
 	void SetForcedNextBandwidth(int64 bitsPerSecond, double minBufferTimeBeforePlayback) override {}
-	void SwitchBufferingQuality(const IAdaptiveStreamSelector::FBufferingQuality& InQualityChange) override {}
 	int64 GetLastBandwidth() override;
 	int64 GetAverageBandwidth() override;
 	int64 GetAverageThroughput() override;
diff --git a/Engine/Plugins/Media/ElectraPlayer/Source/ElectraPlayerRuntime/Private/Runtime/Player/ABRRules/ABRInternal.h b/Engine/Plugins/Media/ElectraPlayer/Source/ElectraPlayerRuntime/Private/Runtime/Player/ABRRules/ABRInternal.h
--- a/Engine/Plugins/Media/ElectraPlayer/Source/ElectraPlayerRuntime/Private/Runtime/Player/ABRRules/ABRInternal.h
+++ b/Engine/Plugins/Media/ElectraPlayer/Source/ElectraPlayerRuntime/Private/Runtime/Player/ABRRules/ABRInternal.h
@@ -78,7 +78,6 @@ namespace Electra
 		virtual void RepresentationsChanged(EStreamType InStreamType, TSharedPtrTS<IManifest::IPlayPeriod> InCurrentPlayPeriod) = 0;
 		virtual void SetBandwidth(int64 bitsPerSecond) = 0;
 		virtual void SetForcedNextBandwidth(int64 bitsPerSecond, double minBufferTimeBeforePlayback) = 0;
-		virtual void SwitchBufferingQuality(const IAdaptiveStreamSelector::FBufferingQuality& InQualityChange) = 0;
 		virtual int64 GetLastBandwidth() = 0;
 		virtual int64 GetAverageBandwidth() = 0;
 		virtual int64 GetAverageThroughput() = 0;
diff --git a/Engine/Plugins/Media/ElectraPlayer/Source/ElectraPlayerRuntime/Private/Runtime/Player/ABRRules/ABRLiveStream.cpp b/Engine/Plugins/Media/ElectraPlayer/Source/ElectraPlayerRuntime/Private/Runtime/Player/ABRRules/ABRLiveStream.cpp
--- a/Engine/Plugins/Media/ElectraPlayer/Source/ElectraPlayerRuntime/Private/Runtime/Player/ABRRules/ABRLiveStream.cpp
+++ b/Engine/Plugins/Media/ElectraPlayer/Source/ElectraPlayerRuntime/Private/Runtime/Player/ABRRules/ABRLiveStream.cpp
@@ -46,7 +46,6 @@ public:
 	void RepresentationsChanged(EStreamType InStreamType, TSharedPtrTS<IManifest::IPlayPeriod> InCurrentPlayPeriod) override;
 	void SetBandwidth(int64 bitsPerSecond) override;
 	void SetForcedNextBandwidth(int64 bitsPerSecond, double minBufferTimeBeforePlayback) override;
-	void SwitchBufferingQuality(const IAdaptiveStreamSelector::FBufferingQuality& InQualityChange) override;
 	int64 GetLastBandwidth() override;
 	int64 GetAverageBandwidth() override;
 	int64 GetAverageThroughput() override;
@@ -272,10 +271,11 @@ private:
 		return Stream ? (*Stream) : nullptr;
 	}
 	
-	double GetPlayablePlayerDuration(EStreamType InStreamType)
+	double GetPlayablePlayerDuration(bool& bEOS, EStreamType InStreamType)
 	{
 		IAdaptiveStreamSelector::IPlayerLiveControl::FABRBufferStats bs;
 		Info->ABRGetStreamBufferStats(bs, InStreamType);
+		bEOS = bs.bReachedEnd;
 		return bs.PlayableContentDuration.GetAsSeconds();
 	}
 
@@ -444,7 +444,8 @@ FABRDownloadProgressDecision FABRLiveStream::ReportDownloadProgress(const Metric
 
 		if (SegmentDownloadStats.SegmentType == Metrics::ESegmentType::Media)
 		{
-			const double playableSeconds = GetPlayablePlayerDuration(SegmentDownloadStats.StreamType);
+			bool bEOS = false;
+			const double playableSeconds = GetPlayablePlayerDuration(bEOS, SegmentDownloadStats.StreamType);
 
 			FScopeLock lock(&WorkVars->Lock);
 
@@ -725,8 +726,9 @@ void FABRLiveStream::ReportDownloadEnd(const Metrics::FSegmentDownloadStats& Seg
 		}
 
 		FStreamWorkVars::FSegmentInfo si;
+		bool bEOS = false;
 		si.BufferDurationAtStart = WorkVars->BufferContentDurationAtSegmentStart;
-		si.BufferDurationAtEnd = GetPlayablePlayerDuration(SegmentDownloadStats.StreamType);
+		si.BufferDurationAtEnd = GetPlayablePlayerDuration(bEOS, SegmentDownloadStats.StreamType);
 		si.ExpectedNetworkLatency = WorkVars->ExpectedNetworkLatency;
 		si.SegmentDuration = SegmentDownloadStats.Duration;
 		si.DownloadTime = SegmentDownloadStats.TimeToDownload;
@@ -831,11 +833,6 @@ void FABRLiveStream::SetForcedNextBandwidth(int64 bitsPerSecond, double /*minBuf
 	ForcedInitialBandwidth = (int32)bitsPerSecond;
 }
 
-void FABRLiveStream::SwitchBufferingQuality(const IAdaptiveStreamSelector::FBufferingQuality& InQualityChange)
-{
-	// No-op for Live ABR
-}
-
 int64 FABRLiveStream::GetLastBandwidth()
 {
 	FStreamWorkVars* WorkVars = GetWorkVars(GetPrimaryStreamType());
@@ -1047,7 +1044,8 @@ IAdaptiveStreamSelector::ESegmentAction FABRLiveStream::PerformSelection(const T
 			ConsumptionRateScale = 1.0;
 		}
 
-		const double AvailableDuration = GetPlayablePlayerDuration(StreamType);
+		bool bEOS = false;
+		const double AvailableDuration = GetPlayablePlayerDuration(bEOS, StreamType);
 		const double DownloadLatency = WorkVars ? WorkVars->AverageLatency.GetWeightedMax(DefaultNetworkLatency) : DefaultNetworkLatency;
 
 		if (WorkVars && StreamType == EStreamType::Video)
@@ -1385,7 +1383,8 @@ bool FABRLiveStream::HasStableBuffer(int32& OutGain, int32& OutTrend, EStreamTyp
 	}
 	FScopeLock lock(&WorkVars->Lock);
 
-	const double AvailableDuration = GetPlayablePlayerDuration(InStreamType) + ExtraContentDuration;
+	bool bEOS = false;
+	const double AvailableDuration = GetPlayablePlayerDuration(bEOS, InStreamType) + ExtraContentDuration;
 	const double DesiredLatency = LatencyConfig.TargetLatency;
 
 	/*
@@ -1573,7 +1572,8 @@ IAdaptiveStreamSelector::EHandlingAction FABRLiveStream::PeriodicHandle()
 		const FStreamWorkVars* WorkVars = GetWorkVars(StreamType);
 		if (WorkVars)
 		{
-			const double AvailableBufferedDuration = GetPlayablePlayerDuration(StreamType);
+			bool bEOS = false;
+			const double AvailableBufferedDuration = GetPlayablePlayerDuration(bEOS, StreamType);
 			const double CurrentLatency = Info->ABRGetLatency().GetAsSeconds();
 			const double TargetLatency = lc.TargetLatency;
 			const double NetworkLatency = WorkVars->AverageLatency.GetWeightedMax(DefaultNetworkLatency);
@@ -1587,7 +1587,7 @@ IAdaptiveStreamSelector::EHandlingAction FABRLiveStream::PeriodicHandle()
 			const bool bMaybeSpeedUp = bStable || Trend>=0;
 
 			// Slow down because of insufficient buffered data?
-			bool bBufferSlowDown = bOvertime || AvailableBufferedDuration < Utils::Max(lc.LowBufferContentBackoff, NetworkLatency);
+			bool bBufferSlowDown = bOvertime || (!bEOS && AvailableBufferedDuration < Utils::Max(lc.LowBufferContentBackoff, NetworkLatency));
 			// Do not slow down if we are supposed to play on the Live edge and are already too far behind.
 			if (Info->ABRShouldPlayOnLiveEdge() && Distance > lc.LowBufferMaxTargetLatencyDistance)
 			{
@@ -1678,7 +1678,9 @@ void FABRLiveStream::DebugPrint(void* pThat, void (*pPrintFN)(void* pThat, const
 	pPrintFN(pThat, "  Catchup rate: %.3f", Info->ABRGetRenderRateScale());
 	pPrintFN(pThat, "Network latency: %.3f", WorkVars->AverageLatency.GetWeightedMax(DefaultNetworkLatency));
 	pPrintFN(pThat, " Num rebuffers: %d", WorkVars->NumSegmentsRebuffered);
-	pPrintFN(pThat, "Duration avail: %.3f", GetPlayablePlayerDuration(StreamType));
+	bool bEOS = false;
+	double Dur = GetPlayablePlayerDuration(bEOS, GetPrimaryStreamType());
+	pPrintFN(pThat, "Duration avail: %.3f, EOS=%d", Dur, bEOS);
 #if 0
 	int32 Gain, Trend;
 	bool bStable = HasStableBuffer(Gain, Trend, StreamType, 0.0);
diff --git a/Engine/Plugins/Media/ElectraPlayer/Source/ElectraPlayerRuntime/Private/Runtime/Player/ABRRules/ABRVoDPlus.cpp b/Engine/Plugins/Media/ElectraPlayer/Source/ElectraPlayerRuntime/Private/Runtime/Player/ABRRules/ABRVoDPlus.cpp
--- a/Engine/Plugins/Media/ElectraPlayer/Source/ElectraPlayerRuntime/Private/Runtime/Player/ABRRules/ABRVoDPlus.cpp
+++ b/Engine/Plugins/Media/ElectraPlayer/Source/ElectraPlayerRuntime/Private/Runtime/Player/ABRRules/ABRVoDPlus.cpp
@@ -35,7 +35,6 @@ public:
 	void RepresentationsChanged(EStreamType InStreamType, TSharedPtrTS<IManifest::IPlayPeriod> InCurrentPlayPeriod) override;
 	void SetBandwidth(int64 bitsPerSecond) override;
 	void SetForcedNextBandwidth(int64 bitsPerSecond, double minBufferTimeBeforePlayback) override;
-	void SwitchBufferingQuality(const IAdaptiveStreamSelector::FBufferingQuality& InQualityChange) override;
 	int64 GetLastBandwidth() override;
 	int64 GetAverageBandwidth() override;
 	int64 GetAverageThroughput() override;
@@ -138,10 +137,11 @@ private:
 		return Stream ? (*Stream) : nullptr;
 	}
 	
-	double GetPlayablePlayerDuration(EStreamType InStreamType)
+	double GetPlayablePlayerDuration(bool& bEOS, EStreamType InStreamType)
 	{
 		IAdaptiveStreamSelector::IPlayerLiveControl::FABRBufferStats bs;
 		Info->ABRGetStreamBufferStats(bs, InStreamType);
+		bEOS = bs.bReachedEnd;
 		return bs.PlayableContentDuration.GetAsSeconds();
 	}
 
@@ -276,7 +276,8 @@ FABRDownloadProgressDecision FABROnDemandPlus::ReportDownloadProgress(const Metr
 			double EstimatedTotalDownloadTime = SegmentDownloadStats.Duration;
 
 			const int32 QualityIndex = WorkVars->QualityIndexDownloading;
-			const double playableSeconds = GetPlayablePlayerDuration(SegmentDownloadStats.StreamType);
+			bool bEOS = false;
+			const double playableSeconds = GetPlayablePlayerDuration(bEOS, SegmentDownloadStats.StreamType);
 			const bool bIsEmittingPartial = (Flags & FABRDownloadProgressDecision::EDecisionFlags::eABR_EmitPartialData) != 0;
 
 			// Estimate how long we will take to complete the download.
@@ -544,11 +545,6 @@ void FABROnDemandPlus::SetForcedNextBandwidth(int64 bitsPerSecond, double minBuf
 	ForcedInitialBandwidthUntilSecondsBuffered = minBufferTimeBeforePlayback;
 }
 
-void FABROnDemandPlus::SwitchBufferingQuality(const IAdaptiveStreamSelector::FBufferingQuality& InQualityChange)
-{
-	// No-op
-}
-
 int64 FABROnDemandPlus::GetLastBandwidth()
 {
 	FStreamWorkVars* WorkVars = GetWorkVars(GetPrimaryStreamType());
@@ -733,7 +729,8 @@ IAdaptiveStreamSelector::ESegmentAction FABROnDemandPlus::PerformSelection(const
 
 		int32 NewQualityIndex = 0;
 
-		const double AvailableDuration = GetPlayablePlayerDuration(StreamType);
+		bool bEOS = false;
+		const double AvailableDuration = GetPlayablePlayerDuration(bEOS, StreamType);
 		const double DownloadLatency = WorkVars ? WorkVars->AverageLatency.GetWeightedMax(DefaultNetworkLatency) : DefaultNetworkLatency;
 
 		if (WorkVars && StreamType == EStreamType::Video)
@@ -867,13 +864,14 @@ IAdaptiveStreamSelector::EHandlingAction FABROnDemandPlus::PeriodicHandle()
 		const FStreamWorkVars* WorkVars = GetWorkVars(StreamType);
 		if (WorkVars)
 		{
-			const double AvailableBufferedDuration = GetPlayablePlayerDuration(StreamType);
+			bool bEOS = false;
+			const double AvailableBufferedDuration = GetPlayablePlayerDuration(bEOS, StreamType);
 			const double NetworkLatency = WorkVars->AverageLatency.GetWeightedMax(DefaultNetworkLatency);
 			const double CurrentPlayRate = Info->ABRGetRenderRateScale();
 			const bool bOvertime = WorkVars->bWentIntoOvertime;
 
 			// Slow down because of insufficient buffered data?
-			bool bBufferSlowDown = bOvertime || AvailableBufferedDuration < Utils::Max(SlowDownWhenBufferDurationUnder, NetworkLatency);
+			bool bBufferSlowDown = bOvertime || (!bEOS && AvailableBufferedDuration < Utils::Max(SlowDownWhenBufferDurationUnder, NetworkLatency));
 			if (bBufferSlowDown)
 			{
 				SetRenderRateScale(SlowDownRateWhenBufferDurationUnder);
@@ -911,7 +909,9 @@ void FABROnDemandPlus::DebugPrint(void* pThat, void (*pPrintFN)(void* pThat, con
 
 	pPrintFN(pThat, "  Catchup rate: %.3f", Info->ABRGetRenderRateScale());
 	pPrintFN(pThat, "Network latency: %.3f", VideoWorkVars.AverageLatency.GetWeightedMax(DefaultNetworkLatency));
-	pPrintFN(pThat, "Duration avail: %.3f", GetPlayablePlayerDuration(GetPrimaryStreamType()));
+	bool bEOS = false;
+	double Dur = GetPlayablePlayerDuration(bEOS, GetPrimaryStreamType());
+	pPrintFN(pThat, "Duration avail: %.3f, EOS=%d", Dur, bEOS);
 }
 
 
diff --git a/Engine/Plugins/Media/ElectraPlayer/Source/ElectraPlayerRuntime/Private/Runtime/Player/AdaptiveStreamingPlayer.cpp b/Engine/Plugins/Media/ElectraPlayer/Source/ElectraPlayerRuntime/Private/Runtime/Player/AdaptiveStreamingPlayer.cpp
--- a/Engine/Plugins/Media/ElectraPlayer/Source/ElectraPlayerRuntime/Private/Runtime/Player/AdaptiveStreamingPlayer.cpp
+++ b/Engine/Plugins/Media/ElectraPlayer/Source/ElectraPlayerRuntime/Private/Runtime/Player/AdaptiveStreamingPlayer.cpp
@@ -88,7 +88,7 @@ FAdaptiveStreamingPlayer::FAdaptiveStreamingPlayer(const IAdaptiveStreamingPlaye
 	PlaybackRate   		    = 0.0;
 	RenderRateScale			= 1.0;
 	StreamReaderHandler     = nullptr;
-	bRebufferPending   	    = false;
+	RebufferCause   	    = ERebufferCause::None;
 	Manifest  				= nullptr;
 	LastBufferingState 		= EPlayerState::eState_Buffering;
 	bIsClosing 				= false;
@@ -1944,97 +1944,6 @@ double FAdaptiveStreamingPlayer::GetMinBufferTimeBeforePlayback()
 	return kMinBufferBeforePlayback;
 }
 
-
-void FAdaptiveStreamingPlayer::InternalHandleLiveSync()
-{
-	// If we have to start playback at the current time NOW we need to tag all buffered access units
-	// that are older than NOW so that they will be skipped over in decoding or rendering.
-	// They may have been tagged already if frame accurate seeking is enabled, but the real time NOW
-	// between the time the segment download was started until, well, NOW, has continued to elapse
-	// so any access units in between need to be tagged as well.
-	if (LiveSyncVars.bSyncLiveToNow && Manifest.IsValid())
-	{
-		FTimeRange Seekable = Manifest->GetSeekableTimeRange();
-		FTimeRange TotalRange = Manifest->GetTotalTimeRange();
-		double minBufTime = Manifest->GetMinBufferTime().GetAsSeconds();
-
-		// For now attempting to sync to Live works only when there is a sufficiently large distance
-		// between the Live edge and the playback position. The difference can be found by examining
-		// the end value of the total vs. seekable range.
-		FTimeValue DistanceToLiveEdge = TotalRange.End - Seekable.End;
-		if (!DistanceToLiveEdge.IsValid() || DistanceToLiveEdge.GetAsSeconds() < minBufTime)
-		{
-			PostLog(Facility::EFacility::Player, IInfoLog::ELevel::Warning, FString::Printf(TEXT("Live sync not possible due to insufficient distance to Live edge (%.3fs) compared to minimum buffer time (%.3fs)"), DistanceToLiveEdge.GetAsSeconds(), minBufTime));
-			LiveSyncVars.Clear();
-			return;
-		}
-
-		TSharedPtrTS<FMultiTrackAccessUnitBuffer> VidBuffer = GetCurrentReceiveStreamBuffer(EStreamType::Video);
-		TSharedPtrTS<FMultiTrackAccessUnitBuffer> AudBuffer = GetCurrentReceiveStreamBuffer(EStreamType::Audio);
-		TSharedPtrTS<FMultiTrackAccessUnitBuffer> TxtBuffer = GetCurrentReceiveStreamBuffer(EStreamType::Subtitle);
-
-		// See if we should perform a sync. For now we do this only if there is a video track.
-		bool bPerformSync = LiveSyncVars.StartTime.IsValid();
-		if (!bPerformSync && VidBuffer.IsValid())
-		{
-			FAccessUnitBufferInfo bi;
-			VidBuffer->GetStats(bi);
-			//bool bHaveEnough = bi.bLastPushWasBlocked || bi.bEndOfData || bi.PlayableDuration.GetAsSeconds() >= minBufTime;
-			bool bHaveEnough = bi.bLastPushWasBlocked || bi.bEndOfData || bi.PlayableDuration.GetAsSeconds() > 0.0;
-			if (bHaveEnough)
-			{
-				bPerformSync = true;
-				LiveSyncVars.StartTime = MEDIAutcTime::Current();
-			}
-		}
-
-		if (bPerformSync)
-		{
-			if (VidBuffer.IsValid())
-			{
-				LiveSyncVars.DroppedDurationVid += VidBuffer->PrepareForDecodeStartingAt(Seekable.End);
-			}
-			if (AudBuffer.IsValid())
-			{
-				LiveSyncVars.DroppedDurationAud += AudBuffer->PrepareForDecodeStartingAt(Seekable.End);
-			}
-			if (TxtBuffer.IsValid())
-			{
-				// Subtitles we discard.
-				FMultiTrackAccessUnitBuffer::FScopedLock lock(TxtBuffer);
-				TxtBuffer->PopDiscardUntil(Seekable.End);
-			}
-			UpdateDiagnostics();
-
-			// Check if we are spending too much time trying to perform the sync.
-			FTimeValue ElapsedTime = MEDIAutcTime::Current() - LiveSyncVars.StartTime;
-			if (ElapsedTime > DistanceToLiveEdge)
-			{
-				PostLog(Facility::EFacility::Player, IInfoLog::ELevel::Warning, FString::Printf(TEXT("Live sync is taking too long. Bandwidth to CDN is most likey insufficient. Starting now.")));
-				LiveSyncVars.Clear();
-			}
-
-			FTimeValue VidDropped = LiveSyncVars.DroppedDurationVid;
-			FTimeValue VidTotal = VideoBufferStats.StreamBuffer.PlayableDuration + VidDropped;
-			double DropRatio = VidDropped.GetAsSeconds() / VidTotal.GetAsSeconds();
-			if (LiveSyncVars.LastDropRatio < DropRatio && DropRatio >= 1.0)
-			{
-				// Force a rate scale in the ABR in the hope to get new data faster on a lower quality level.
-				if (StreamSelector.IsValid())
-				{
-					IAdaptiveStreamSelector::FBufferingQuality RateScale;
-					RateScale.BandwidthScaleFactor = 0.45;
-					StreamSelector->SwitchBufferingQuality(RateScale);
-				}
-			}
-			LiveSyncVars.LastDropRatio = DropRatio;
-
-			//PostLog(Facility::EFacility::Player, IInfoLog::ELevel::Info, FString::Printf(TEXT("%3.3fs: %.3f drop, %.3f total, %.3f ratio"), ElapsedTime.GetAsSeconds(), VidDropped.GetAsSeconds(), VidTotal.GetAsSeconds(), DropRatio));
-		}
-	}
-}
-
-
 bool FAdaptiveStreamingPlayer::HaveEnoughBufferedDataToStartPlayback()
 {
 	double kMinBufferBeforePlayback = GetMinBufferTimeBeforePlayback();
@@ -2176,9 +2085,6 @@ void FAdaptiveStreamingPlayer::HandleNewBufferedData()
 	{
 		if (!SeekVars.bForScrubbing)
 		{
-			// Synchronize buffers to wallclock time if necessary.
-			InternalHandleLiveSync();
-
 			check(LastBufferingState == EPlayerState::eState_Buffering || LastBufferingState == EPlayerState::eState_Rebuffering || LastBufferingState == EPlayerState::eState_Seeking);
 			// Check if we have enough data buffered up to begin handing off data to the decoders.
 			bool bHaveEnoughData = HaveEnoughBufferedDataToStartPlayback();
@@ -2207,8 +2113,6 @@ void FAdaptiveStreamingPlayer::HandleNewBufferedData()
 				DecoderState = EDecoderState::eDecoder_Running;
 
 				PostrollVars.Clear();
-				// We are synced up now and can clear the state.
-				LiveSyncVars.Clear();
 
 				// Send buffering end event
 				DispatchBufferingEvent(false, LastBufferingState);
@@ -2240,7 +2144,6 @@ void FAdaptiveStreamingPlayer::HandleNewBufferedData()
 					PipelineState = EPipelineState::ePipeline_Prerolling;
 					DecoderState = EDecoderState::eDecoder_Running;
 					PostrollVars.Clear();
-					LiveSyncVars.Clear();
 					DispatchEvent(FMetricEvent::ReportPrerollStart());
 				}
 			}
@@ -2525,9 +2428,6 @@ void FAdaptiveStreamingPlayer::InternalHandlePendingStartRequest(const FTimeValu
 						// DASH Live streams provide the means to join at an exact wallclock time, so enable this.
 						if (ManifestType == EMediaFormatType::DASH)
 						{
-							LiveSyncVars.Clear();
-							// Live sync is no longer required since DASH Live maintains sync with the Live edge now.
-							//LiveSyncVars.bSyncLiveToNow = true;
 							PlaybackState.SetShouldPlayOnLiveEdge(true);
 						}
 					}
@@ -2894,9 +2794,9 @@ void FAdaptiveStreamingPlayer::InternalHandlePendingFirstSegmentRequest(const FT
 			TSharedPtrTS<FBufferSourceInfo> BufferSourceInfoTxt = CurrentPlayPeriodText.IsValid() ? CurrentPlayPeriodText->GetSelectedStreamBufferSourceInfo(EStreamType::Subtitle) : nullptr;
 
 			TSharedPtrTS<FStreamDataBuffers> NewReceiveBuffers = MakeSharedTS<FStreamDataBuffers>();
-			NewReceiveBuffers->VidBuffer = MakeSharedTS<FMultiTrackAccessUnitBuffer>();
-			NewReceiveBuffers->AudBuffer = MakeSharedTS<FMultiTrackAccessUnitBuffer>();
-			NewReceiveBuffers->TxtBuffer = MakeSharedTS<FMultiTrackAccessUnitBuffer>();
+			NewReceiveBuffers->VidBuffer = MakeSharedTS<FMultiTrackAccessUnitBuffer>(EStreamType::Video);
+			NewReceiveBuffers->AudBuffer = MakeSharedTS<FMultiTrackAccessUnitBuffer>(EStreamType::Audio);
+			NewReceiveBuffers->TxtBuffer = MakeSharedTS<FMultiTrackAccessUnitBuffer>(EStreamType::Subtitle);
 
 			if (ManifestType == EMediaFormatType::ISOBMFF)
 			{
@@ -2905,9 +2805,9 @@ void FAdaptiveStreamingPlayer::InternalHandlePendingFirstSegmentRequest(const FT
 				NewReceiveBuffers->AudBuffer->SetParallelTrackMode();
 				NewReceiveBuffers->TxtBuffer->SetParallelTrackMode();
 			}
-			NewReceiveBuffers->VidBuffer->SelectTrackWhenAvailable(BufferSourceInfoVid);
-			NewReceiveBuffers->AudBuffer->SelectTrackWhenAvailable(BufferSourceInfoAud);
-			NewReceiveBuffers->TxtBuffer->SelectTrackWhenAvailable(BufferSourceInfoTxt);
+			NewReceiveBuffers->VidBuffer->SelectTrackWhenAvailable(CurrentPlaybackSequenceID[StreamTypeToArrayIndex(EStreamType::Video)], BufferSourceInfoVid);
+			NewReceiveBuffers->AudBuffer->SelectTrackWhenAvailable(CurrentPlaybackSequenceID[StreamTypeToArrayIndex(EStreamType::Audio)], BufferSourceInfoAud);
+			NewReceiveBuffers->TxtBuffer->SelectTrackWhenAvailable(CurrentPlaybackSequenceID[StreamTypeToArrayIndex(EStreamType::Subtitle)], BufferSourceInfoTxt);
 
 			// Add the new receive buffer to the list.
 			DataBuffersCriticalSection.Lock();
@@ -3233,30 +3133,18 @@ void FAdaptiveStreamingPlayer::HandlePendingMediaSegmentRequests()
 						CurrentPlayPeriodVideo = FinishedReq.Period;
 						StreamSelector->SetCurrentPlaybackPeriod(EStreamType::Video, CurrentPlayPeriodVideo);
 						UpdateStreamResolutionLimit();
-						if ((!FinishedReq.bStartOver || FinishedReq.bPlayPosAutoReselect) && StreamRcvBuffer.IsValid())
-						{
-							StreamRcvBuffer->AddUpcomingBuffer(SegmentPlayPeriod->GetSelectedStreamBufferSourceInfo(StreamType));
-						}
 					}
 					else if (StreamType == EStreamType::Audio)
 					{
 						SegmentPlayPeriod = FinishedReq.Period;
 						CurrentPlayPeriodAudio = FinishedReq.Period;
 						StreamSelector->SetCurrentPlaybackPeriod(EStreamType::Audio, CurrentPlayPeriodAudio);
-						if ((!FinishedReq.bStartOver || FinishedReq.bPlayPosAutoReselect) && StreamRcvBuffer.IsValid())
-						{
-							StreamRcvBuffer->AddUpcomingBuffer(SegmentPlayPeriod->GetSelectedStreamBufferSourceInfo(StreamType));
-						}
 					}
 					else if (StreamType == EStreamType::Subtitle)
 					{
 						SegmentPlayPeriod = FinishedReq.Period;
 						CurrentPlayPeriodText = FinishedReq.Period;
 						//StreamSelector->SetCurrentPlaybackPeriod(EStreamType::Subtitle, CurrentPlayPeriodText);
-						if ((!FinishedReq.bStartOver || FinishedReq.bPlayPosAutoReselect) && StreamRcvBuffer.IsValid())
-						{
-							StreamRcvBuffer->AddUpcomingBuffer(SegmentPlayPeriod->GetSelectedStreamBufferSourceInfo(StreamType));
-						}
 					}
 					// Add this to the upcoming periods the play position will move into and metadata will need to be updated then
 					// unless this is a startover request.
@@ -3320,7 +3208,7 @@ void FAdaptiveStreamingPlayer::HandlePendingMediaSegmentRequests()
 							TSharedPtrTS<FMultiTrackAccessUnitBuffer> StreamRcvBuffer = GetCurrentReceiveStreamBuffer(StreamType);
 							if (StreamRcvBuffer.IsValid())
 							{
-								StreamRcvBuffer->SelectTrackWhenAvailable(BufferSourceInfo);
+								StreamRcvBuffer->SelectTrackWhenAvailable(CurrentPlaybackSequenceID[StreamTypeToArrayIndex(StreamType)], BufferSourceInfo);
 							}
 							switch(StreamType)
 							{
@@ -3556,9 +3444,13 @@ void FAdaptiveStreamingPlayer::InternalDeselectStream(EStreamType StreamType)
 	{
 		case EStreamType::Video:
 			bIsVideoDeselected = true;
+			SelectedStreamAttributesVid.Deselect();
+			VideoDecoder.Flush();
+			VideoRender.Flush(false);
 			break;
 		case EStreamType::Audio:
 			bIsAudioDeselected = true;
+			SelectedStreamAttributesAud.Deselect();
 			AudioDecoder.Flush();
 			AudioRender.Flush();
 			break;
@@ -3611,45 +3503,32 @@ void FAdaptiveStreamingPlayer::InternalHandleSegmentTrackChanges(const FTimeValu
 
 	if (ManifestType == EMediaFormatType::ISOBMFF)
 	{
-		if (PendingTrackSelectionAud.IsValid() && CurrentPlayPeriodAudio.IsValid())
+		auto HandleISOBMFFTrackSelection = [&bChangeMade, this](FStreamSelectionAttributes& OutSelectionAttributes, FInternalStreamSelectionAttributes& InOutSelectedAttributes, EStreamType InType,
+																TSharedPtrTS<FStreamSelectionAttributes>& InPendingSelection, TSharedPtrTS<IManifest::IPlayPeriod> InCurrentPeriod) -> void
 		{
-			IManifest::IPlayPeriod::ETrackChangeResult Result = CurrentPlayPeriodAudio->ChangeTrackStreamPreference(EStreamType::Audio, *PendingTrackSelectionAud);
-			TSharedPtrTS<FBufferSourceInfo> BufferSourceInfo = CurrentPlayPeriodAudio->GetSelectedStreamBufferSourceInfo(EStreamType::Audio);
-			if (Result == IManifest::IPlayPeriod::ETrackChangeResult::Changed)
+			if (InPendingSelection.IsValid() && InCurrentPeriod.IsValid())
 			{
-				TSharedPtrTS<FMultiTrackAccessUnitBuffer> StreamRcvBuf = GetCurrentReceiveStreamBuffer(EStreamType::Audio);
-				if (StreamRcvBuf.IsValid())
+				IManifest::IPlayPeriod::ETrackChangeResult Result = InCurrentPeriod->ChangeTrackStreamPreference(InType, *InPendingSelection);
+				TSharedPtrTS<FBufferSourceInfo> BufferSourceInfo = InCurrentPeriod->GetSelectedStreamBufferSourceInfo(InType);
+				if (Result == IManifest::IPlayPeriod::ETrackChangeResult::Changed)
 				{
-					StreamRcvBuf->SelectTrackWhenAvailable(BufferSourceInfo);
-				}
-				StreamSelectionAttributesAud = *PendingTrackSelectionAud;
-				bChangeMade = true;
-			}
-			SelectedStreamAttributesAud.UpdateWith(BufferSourceInfo->Kind, BufferSourceInfo->Language, BufferSourceInfo->Codec, BufferSourceInfo->HardIndex);
-			PendingTrackSelectionAud.Reset();
-		}
-
-		if (PendingTrackSelectionTxt.IsValid() && CurrentPlayPeriodText.IsValid())
-		{
-			IManifest::IPlayPeriod::ETrackChangeResult Result = CurrentPlayPeriodText->ChangeTrackStreamPreference(EStreamType::Subtitle, *PendingTrackSelectionTxt);
-			TSharedPtrTS<FBufferSourceInfo> BufferSourceInfo = CurrentPlayPeriodText->GetSelectedStreamBufferSourceInfo(EStreamType::Subtitle);
-			if (Result == IManifest::IPlayPeriod::ETrackChangeResult::Changed)
-			{
-				TSharedPtrTS<FMultiTrackAccessUnitBuffer> StreamRcvBuf = GetCurrentReceiveStreamBuffer(EStreamType::Subtitle);
-				if (StreamRcvBuf.IsValid())
-				{
-					StreamRcvBuf->SelectTrackWhenAvailable(BufferSourceInfo);
+					TSharedPtrTS<FMultiTrackAccessUnitBuffer> StreamRcvBuf = GetCurrentReceiveStreamBuffer(InType);
+					if (StreamRcvBuf.IsValid())
+					{
+						StreamRcvBuf->SelectTrackWhenAvailable(CurrentPlaybackSequenceID[StreamTypeToArrayIndex(InType)], BufferSourceInfo);
+					}
+					OutSelectionAttributes = *InPendingSelection;
+					bChangeMade = true;
 				}
-				StreamSelectionAttributesTxt = *PendingTrackSelectionTxt;
-				bChangeMade = true;
+				InOutSelectedAttributes.UpdateWith(BufferSourceInfo->Kind, BufferSourceInfo->Language, BufferSourceInfo->Codec, BufferSourceInfo->HardIndex);
+				InOutSelectedAttributes.Select();
+				InPendingSelection.Reset();
 			}
-			SelectedStreamAttributesTxt.UpdateWith(BufferSourceInfo->Kind, BufferSourceInfo->Language, BufferSourceInfo->Codec, BufferSourceInfo->HardIndex);
-			SelectedStreamAttributesTxt.Select();
-			PendingTrackSelectionTxt.Reset();
-		}
+		};
 
-		// Ignore video changes for now.
-		PendingTrackSelectionVid.Reset();
+		HandleISOBMFFTrackSelection(StreamSelectionAttributesVid, SelectedStreamAttributesVid, EStreamType::Video, PendingTrackSelectionVid, CurrentPlayPeriodVideo);
+		HandleISOBMFFTrackSelection(StreamSelectionAttributesAud, SelectedStreamAttributesAud, EStreamType::Audio, PendingTrackSelectionAud, CurrentPlayPeriodAudio);
+		HandleISOBMFFTrackSelection(StreamSelectionAttributesTxt, SelectedStreamAttributesTxt, EStreamType::Subtitle, PendingTrackSelectionTxt, CurrentPlayPeriodText);
 	}
 	else if (ManifestType == EMediaFormatType::HLS)
 	{
@@ -3659,7 +3538,7 @@ void FAdaptiveStreamingPlayer::InternalHandleSegmentTrackChanges(const FTimeValu
 			TSharedPtrTS<FMultiTrackAccessUnitBuffer> StreamRcvBuf = GetCurrentReceiveStreamBuffer(EStreamType::Audio);
 			if (StreamRcvBuf.IsValid())
 			{
-				StreamRcvBuf->SelectTrackWhenAvailable(BufferSourceInfo);
+				StreamRcvBuf->SelectTrackWhenAvailable(CurrentPlaybackSequenceID[StreamTypeToArrayIndex(EStreamType::Audio)], BufferSourceInfo);
 			}
 			StreamSelectionAttributesAud = *PendingTrackSelectionAud;
 			SelectedStreamAttributesAud.UpdateWith(BufferSourceInfo->Kind, BufferSourceInfo->Language, BufferSourceInfo->Codec, BufferSourceInfo->HardIndex);
@@ -3677,63 +3556,42 @@ void FAdaptiveStreamingPlayer::InternalHandleSegmentTrackChanges(const FTimeValu
 			return;
 		}
 
-		if (PendingTrackSelectionAud.IsValid() && CurrentPlayPeriodAudio.IsValid())
-		{
-			if (!SelectedStreamAttributesAud.IsCompatibleWith(*PendingTrackSelectionAud))
-			{
-				IManifest::IPlayPeriod::ETrackChangeResult TrackChangeResult = CurrentPlayPeriodAudio->ChangeTrackStreamPreference(EStreamType::Audio, *PendingTrackSelectionAud);
-				if (TrackChangeResult == IManifest::IPlayPeriod::ETrackChangeResult::NewPeriodNeeded)
-				{
-					// Make sure any finishing current requests will be ignored by increasing the sequence ID.
-					++CurrentPlaybackSequenceID[StreamTypeToArrayIndex(EStreamType::Audio)];
-					// Cancel any ongoing segment download now.
-					CancelPendingMediaSegmentRequests(EStreamType::Audio);
-
-					StreamSelectionAttributesAud = *PendingTrackSelectionAud;
-
-					FPendingSegmentRequest NextReq;
-					NextReq.bStartOver = true;
-					NextReq.StartoverPosition.Time = PlaybackState.GetPlayPosition();
-					SetPlaystartOptions(NextReq.StartoverPosition.Options);
-					NextReq.StreamType = EStreamType::Audio;
-					NextPendingSegmentRequests.Enqueue(MoveTemp(NextReq));
-
-					bChangeMade = true;
-				}
-			}
-			PendingTrackSelectionAud.Reset();
-		}
-
-		if (PendingTrackSelectionTxt.IsValid() && CurrentPlayPeriodText.IsValid())
+		auto HandleDASHTrackSelection = [&bChangeMade, this](FStreamSelectionAttributes& OutSelectionAttributes, FInternalStreamSelectionAttributes& InOutSelectedAttributes, EStreamType InType,
+															 TSharedPtrTS<FStreamSelectionAttributes>& InPendingSelection, TSharedPtrTS<IManifest::IPlayPeriod> InCurrentPeriod) -> void
 		{
-			if (SelectedStreamAttributesTxt.IsDeselected() || !SelectedStreamAttributesTxt.IsCompatibleWith(*PendingTrackSelectionTxt))
+			if (InPendingSelection.IsValid() && InCurrentPeriod.IsValid())
 			{
-				IManifest::IPlayPeriod::ETrackChangeResult TrackChangeResult = CurrentPlayPeriodText->ChangeTrackStreamPreference(EStreamType::Subtitle, *PendingTrackSelectionTxt);
-				if (TrackChangeResult == IManifest::IPlayPeriod::ETrackChangeResult::NewPeriodNeeded)
+				if (InOutSelectedAttributes.IsDeselected() || !InOutSelectedAttributes.IsCompatibleWith(*InPendingSelection))
 				{
-					// Make sure any finishing current requests will be ignored by increasing the sequence ID.
-					++CurrentPlaybackSequenceID[StreamTypeToArrayIndex(EStreamType::Subtitle)];
-					// Cancel any ongoing segment download now.
-					CancelPendingMediaSegmentRequests(EStreamType::Subtitle);
-
-					StreamSelectionAttributesTxt = *PendingTrackSelectionTxt;
-					SelectedStreamAttributesTxt.Select();
-
-					FPendingSegmentRequest NextReq;
-					NextReq.bStartOver = true;
-					NextReq.StartoverPosition.Time = PlaybackState.GetPlayPosition();
-					SetPlaystartOptions(NextReq.StartoverPosition.Options);
-					NextReq.StreamType = EStreamType::Subtitle;
-					NextPendingSegmentRequests.Enqueue(MoveTemp(NextReq));
-
-					bChangeMade = true;
+					IManifest::IPlayPeriod::ETrackChangeResult TrackChangeResult = InCurrentPeriod->ChangeTrackStreamPreference(InType, *InPendingSelection);
+					if (TrackChangeResult == IManifest::IPlayPeriod::ETrackChangeResult::NewPeriodNeeded)
+					{
+						// Make sure any finishing current requests will be ignored by increasing the sequence ID.
+						++CurrentPlaybackSequenceID[StreamTypeToArrayIndex(InType)];
+						// Cancel any ongoing segment download now.
+						CancelPendingMediaSegmentRequests(InType);
+
+						OutSelectionAttributes = *InPendingSelection;
+						InOutSelectedAttributes.Select();
+
+						FPendingSegmentRequest NextReq;
+						NextReq.bStartOver = true;
+						NextReq.StartoverPosition.Time = PlaybackState.GetPlayPosition();
+						SetPlaystartOptions(NextReq.StartoverPosition.Options);
+						NextReq.StreamType = InType;
+						NextPendingSegmentRequests.Enqueue(MoveTemp(NextReq));
+
+						bChangeMade = true;
+					}
 				}
+				InOutSelectedAttributes.Select();
+				InPendingSelection.Reset();
 			}
-			PendingTrackSelectionTxt.Reset();
-		}
+		};
 
-		// Ignore video changes for now.
-		PendingTrackSelectionVid.Reset();
+		PendingTrackSelectionVid.Reset();	// Ignore video changes for now.
+		HandleDASHTrackSelection(StreamSelectionAttributesAud, SelectedStreamAttributesAud, EStreamType::Audio, PendingTrackSelectionAud, CurrentPlayPeriodAudio);
+		HandleDASHTrackSelection(StreamSelectionAttributesTxt, SelectedStreamAttributesTxt, EStreamType::Subtitle, PendingTrackSelectionTxt, CurrentPlayPeriodText);
 	}
 	else
 	{
@@ -4483,10 +4341,9 @@ void FAdaptiveStreamingPlayer::InternalClose()
 	CurrentState   	= EPlayerState::eState_Idle;
 	PlaybackRate   	= 0.0;
 	StreamState		= EStreamState::eStream_Running;
-	bRebufferPending   = false;
+	RebufferCause = ERebufferCause::None;
 	LastBufferingState = EPlayerState::eState_Buffering;
 	RebufferDetectedAtPlayPos.SetToInvalid();
-	LiveSyncVars.Clear();
 	PrerollVars.Clear();
 	PostrollVars.Clear();
 
@@ -4565,7 +4422,6 @@ void FAdaptiveStreamingPlayer::InternalSetLoop(const FLoopParam& LoopParam)
  */
 void FAdaptiveStreamingPlayer::InternalRebuffer()
 {
-
 	IAdaptiveStreamSelector::FRebufferAction Action;
 	check(StreamSelector.IsValid());
 	if (StreamSelector.IsValid())
@@ -4597,7 +4453,14 @@ void FAdaptiveStreamingPlayer::InternalRebuffer()
 		CurrentState = EPlayerState::eState_Rebuffering;
 		LastBufferingState = EPlayerState::eState_Rebuffering;
 
-		bRebufferPending = false;
+		// If rebuffering triggered while switching tracks then the spacing of sync frames may be too big and we
+		// have none to switch over to the new track yet. To overcome this we restart at the current position.
+		if (RebufferCause == ERebufferCause::TrackswitchUnderrun && Action.Action == IAdaptiveStreamSelector::FRebufferAction::EAction::ContinueLoading)
+		{
+			Action.Action = IAdaptiveStreamSelector::FRebufferAction::EAction::Restart;
+			PostLog(Facility::EFacility::Player, IInfoLog::ELevel::Info, FString::Printf(TEXT("Rebuffering during track change switches rebuffering policy from continue to restart.")));
+		}
+		RebufferCause = ERebufferCause::None;
 
 		// Check if we should just continue loading data
 		if (Action.Action == IAdaptiveStreamSelector::FRebufferAction::EAction::ContinueLoading)
@@ -4686,32 +4549,28 @@ FTimeValue FAdaptiveStreamingPlayer::ABRGetPlaySpeed() const
 {
 	return FTimeValue(RenderClock->IsRunning() ? 1.0 : 0.0);
 }
-void FAdaptiveStreamingPlayer::ABRGetStreamBufferStats(FAccessUnitBufferInfo& OutBufferStats, EStreamType ForStream)
+void FAdaptiveStreamingPlayer::ABRGetStreamBufferStats(IAdaptiveStreamSelector::IPlayerLiveControl::FABRBufferStats& OutBufferStats, EStreamType ForStream)
 {
 	UpdateDiagnostics();
-	FMediaCriticalSection::ScopedLock lock(DiagnosticsCriticalSection);
+	DiagnosticsCriticalSection.Lock();
 	switch(ForStream)
 	{
 		case EStreamType::Video:
-			OutBufferStats = VideoBufferStats.StreamBuffer;
+			OutBufferStats.PlayableContentDuration = VideoBufferStats.StreamBuffer.PlayableDuration;
+			OutBufferStats.bReachedEnd = VideoBufferStats.StreamBuffer.bEndOfData;
 			break;
 		case EStreamType::Audio:
-			OutBufferStats = AudioBufferStats.StreamBuffer;
+			OutBufferStats.PlayableContentDuration = AudioBufferStats.StreamBuffer.PlayableDuration;
+			OutBufferStats.bReachedEnd = AudioBufferStats.StreamBuffer.bEndOfData;
 			break;
 		case EStreamType::Subtitle:
-			OutBufferStats = TextBufferStats.StreamBuffer;
+			OutBufferStats.PlayableContentDuration = TextBufferStats.StreamBuffer.PlayableDuration;
+			OutBufferStats.bReachedEnd = TextBufferStats.StreamBuffer.bEndOfData;
 			break;
 		default:
-			OutBufferStats.Clear();
 			break;
 	}
-}
-
-void FAdaptiveStreamingPlayer::ABRGetStreamBufferStats(IAdaptiveStreamSelector::IPlayerLiveControl::FABRBufferStats& OutBufferStats, EStreamType ForStream)
-{
-	FAccessUnitBufferInfo bi;
-	ABRGetStreamBufferStats(bi, ForStream);
-	OutBufferStats.PlayableContentDuration = bi.PlayableDuration;
+	DiagnosticsCriticalSection.Unlock();
 	if (ForStream == EStreamType::Video && VideoRender.Renderer.IsValid())
 	{
 		OutBufferStats.PlayableContentDuration += VideoRender.Renderer->GetEnqueuedSampleDuration();
@@ -4721,7 +4580,6 @@ void FAdaptiveStreamingPlayer::ABRGetStreamBufferStats(IAdaptiveStreamSelector::
 		OutBufferStats.PlayableContentDuration += AudioRender.Renderer->GetEnqueuedSampleDuration();
 	}
 }
-
 void FAdaptiveStreamingPlayer::ABRSetRenderRateScale(double InRenderRateScale)
 {
 	RenderRateScale = InRenderRateScale;
diff --git a/Engine/Plugins/Media/ElectraPlayer/Source/ElectraPlayerRuntime/Private/Runtime/Player/AdaptiveStreamingPlayerABR.cpp b/Engine/Plugins/Media/ElectraPlayer/Source/ElectraPlayerRuntime/Private/Runtime/Player/AdaptiveStreamingPlayerABR.cpp
--- a/Engine/Plugins/Media/ElectraPlayer/Source/ElectraPlayerRuntime/Private/Runtime/Player/AdaptiveStreamingPlayerABR.cpp
+++ b/Engine/Plugins/Media/ElectraPlayer/Source/ElectraPlayerRuntime/Private/Runtime/Player/AdaptiveStreamingPlayerABR.cpp
@@ -33,7 +33,6 @@ namespace Electra
 		FTimeValue GetMinBufferTimeForPlayback(EMinBufferType InBufferingType, FTimeValue InDefaultMBT) override;
 		FRebufferAction GetRebufferAction(const FParamDict& CurrentPlayerOptions) override;
 		EHandlingAction PeriodicHandle() override;
-		void SwitchBufferingQuality(const FBufferingQuality& InQualityChange) override;
 		void MarkStreamAsUnavailable(const FBlacklistedStream& BlacklistedStream) override;
 		void MarkStreamAsAvailable(const FBlacklistedStream& NoLongerBlacklistedStream) override;
 		int64 GetLastBandwidth() override;
@@ -105,8 +104,6 @@ namespace Electra
 		{ return PlayerLiveControl->ABRGetLatency(); }
 		FTimeValue ABRGetPlaySpeed() const override
 		{ return PlayerLiveControl->ABRGetPlaySpeed(); }
-		void ABRGetStreamBufferStats(FAccessUnitBufferInfo& OutBufferStats, EStreamType ForStream) override
-		{ PlayerLiveControl->ABRGetStreamBufferStats(OutBufferStats, ForStream); }
 		void ABRGetStreamBufferStats(IAdaptiveStreamSelector::IPlayerLiveControl::FABRBufferStats& OutBufferStats, EStreamType ForStream) override
 		{ PlayerLiveControl->ABRGetStreamBufferStats(OutBufferStats, ForStream); }
 		FTimeRange ABRGetSupportedRenderRateScale() override
@@ -513,12 +510,6 @@ namespace Electra
 		ABRMethod->ReportBufferingEnd(BufferingReason);
 	}
 
-	void FAdaptiveStreamSelector::SwitchBufferingQuality(const FBufferingQuality& InQualityChange)
-	{
-		check(ABRMethod.IsValid());
-		ABRMethod->SwitchBufferingQuality(InQualityChange);
-	}
-
 	const TArray<TSharedPtrTS<FABRStreamInformation>>& FAdaptiveStreamSelector::GetStreamInformations(EStreamType InForStreamType)
 	{
 		return InForStreamType == EStreamType::Video ? StreamInformationVideo : StreamInformationAudio;
diff --git a/Engine/Plugins/Media/ElectraPlayer/Source/ElectraPlayerRuntime/Private/Runtime/Player/AdaptiveStreamingPlayerABR.h b/Engine/Plugins/Media/ElectraPlayer/Source/ElectraPlayerRuntime/Private/Runtime/Player/AdaptiveStreamingPlayerABR.h
--- a/Engine/Plugins/Media/ElectraPlayer/Source/ElectraPlayerRuntime/Private/Runtime/Player/AdaptiveStreamingPlayerABR.h
+++ b/Engine/Plugins/Media/ElectraPlayer/Source/ElectraPlayerRuntime/Private/Runtime/Player/AdaptiveStreamingPlayerABR.h
@@ -51,11 +51,10 @@ public:
 		virtual FTimeValue ABRGetPlaySpeed() const = 0;
 
 		// Returns the current stream access unit buffer stats.
-		virtual void ABRGetStreamBufferStats(FAccessUnitBufferInfo& OutBufferStats, EStreamType ForStream) = 0;
-
 		struct FABRBufferStats
 		{
 			FTimeValue PlayableContentDuration;
+			bool bReachedEnd = false;
 		};
 		virtual void ABRGetStreamBufferStats(FABRBufferStats& OutBufferStats, EStreamType ForStream) = 0;
 
@@ -120,13 +119,6 @@ public:
 	// The ABR may take internal action to switch quality when this is called!
 	virtual FRebufferAction GetRebufferAction(const FParamDict& CurrentPlayerOptions) = 0;
 
-	struct FBufferingQuality
-	{
-		TOptional<double> BandwidthScaleFactor;
-		TOptional<int64> AbsoluteBandwidth;
-	};
-	virtual void SwitchBufferingQuality(const FBufferingQuality& InQualityChange) = 0;
-
 	struct FBlacklistedStream
 	{
 		FString		AssetUniqueID;
diff --git a/Engine/Plugins/Media/ElectraPlayer/Source/ElectraPlayerRuntime/Private/Runtime/Player/AdaptiveStreamingPlayerDecoders.cpp b/Engine/Plugins/Media/ElectraPlayer/Source/ElectraPlayerRuntime/Private/Runtime/Player/AdaptiveStreamingPlayerDecoders.cpp
--- a/Engine/Plugins/Media/ElectraPlayer/Source/ElectraPlayerRuntime/Private/Runtime/Player/AdaptiveStreamingPlayerDecoders.cpp
+++ b/Engine/Plugins/Media/ElectraPlayer/Source/ElectraPlayerRuntime/Private/Runtime/Player/AdaptiveStreamingPlayerDecoders.cpp
@@ -747,7 +747,7 @@ void FAdaptiveStreamingPlayer::FeedDecoder(EStreamType Type, IAccessUnitBufferIn
 	if (!bIsDeselected)
 	{
 		// Check for buffer underrun.
-		if (bHandleUnderrun && !bRebufferPending && CurrentState == EPlayerState::eState_Playing && StreamState == EStreamState::eStream_Running && PipelineState == EPipelineState::ePipeline_Running)
+		if (bHandleUnderrun && RebufferCause == ERebufferCause::None && CurrentState == EPlayerState::eState_Playing && StreamState == EStreamState::eStream_Running && PipelineState == EPipelineState::ePipeline_Running)
 		{
 			bool bEODSet = FromMultistreamBuffer->IsEODFlagSet();
 			if (!bEODSet && FromMultistreamBuffer->Num() == 0)
@@ -765,7 +765,7 @@ void FAdaptiveStreamingPlayer::FeedDecoder(EStreamType Type, IAccessUnitBufferIn
 				if (NumEnqueuedSamples <= 1)
 				{
 					// Buffer underrun.
-					bRebufferPending = true;
+					RebufferCause = FromMultistreamBuffer->HasPendingTrackSwitch() ? ERebufferCause::TrackswitchUnderrun : ERebufferCause::Underrun;
 					FTimeValue LastKnownPTS = FromMultistreamBuffer->GetLastPoppedPTS();
 					// Only set the 'rebuffer at' time if we have a valid last known PTS. If we don't
 					// then maybe this is a cascade failure from a previous rebuffer attempt for which
diff --git a/Engine/Plugins/Media/ElectraPlayer/Source/ElectraPlayerRuntime/Private/Runtime/Player/AdaptiveStreamingPlayerInternal.h b/Engine/Plugins/Media/ElectraPlayer/Source/ElectraPlayerRuntime/Private/Runtime/Player/AdaptiveStreamingPlayerInternal.h
--- a/Engine/Plugins/Media/ElectraPlayer/Source/ElectraPlayerRuntime/Private/Runtime/Player/AdaptiveStreamingPlayerInternal.h
+++ b/Engine/Plugins/Media/ElectraPlayer/Source/ElectraPlayerRuntime/Private/Runtime/Player/AdaptiveStreamingPlayerInternal.h
@@ -923,7 +923,6 @@ private:
 	FTimeValue ABRGetDesiredLiveEdgeLatency() const override;
 	FTimeValue ABRGetLatency() const override;
 	FTimeValue ABRGetPlaySpeed() const override;
-	void ABRGetStreamBufferStats(FAccessUnitBufferInfo& OutBufferStats, EStreamType ForStream) override;
 	void ABRGetStreamBufferStats(IAdaptiveStreamSelector::IPlayerLiveControl::FABRBufferStats& OutBufferStats, EStreamType ForStream) override;
 	FTimeRange ABRGetSupportedRenderRateScale() override;
 	void ABRSetRenderRateScale(double InRenderRateScale) override;
@@ -1009,6 +1008,13 @@ private:
 		}
 	}
 
+	enum class ERebufferCause
+	{
+		None,
+		Underrun,
+		TrackswitchUnderrun
+	};
+
 	struct FVideoRenderer
 	{
 		void Close()
@@ -1567,27 +1573,6 @@ private:
 		FStallMonitor										DecoderOutputStalledMonitor;
 	};
 
-	struct FLiveSyncVars
-	{
-		FLiveSyncVars()
-		{
-			Clear();
-		}
-		void Clear()
-		{
-			bSyncLiveToNow = false;
-			StartTime.SetToInvalid();
-			DroppedDurationVid.SetToZero();
-			DroppedDurationAud.SetToZero();
-			LastDropRatio = 0.0;
-		}
-		bool bSyncLiveToNow;
-		FTimeValue StartTime;
-		FTimeValue DroppedDurationVid;
-		FTimeValue DroppedDurationAud;
-		double LastDropRatio;
-	};
-
 	struct FPrerollVars
 	{
 		FPrerollVars()
@@ -1828,7 +1813,6 @@ private:
 	void InternalHandleSegmentTrackChanges(const FTimeValue& CurrentTime);
 	void InternalStartoverAtCurrentPosition();
 	void InternalStartoverAtLiveEdge();
-	void InternalHandleLiveSync();
 
 	void UpdateDiagnostics();
 	void HandleNewBufferedData();
@@ -1993,7 +1977,6 @@ private:
 	EDecoderState														DecoderState;
 	EStreamState														StreamState;
 
-	FLiveSyncVars														LiveSyncVars;
 	FPrerollVars														PrerollVars;
 	FPostrollVars														PostrollVars;
 	FSeekVars															SeekVars;
@@ -2001,7 +1984,7 @@ private:
 	double																PlaybackRate;
 	double																RenderRateScale;
 	FTimeValue															RebufferDetectedAtPlayPos;
-	bool																bRebufferPending;
+	ERebufferCause														RebufferCause;
 	bool																bIsClosing;
 
 	TSharedPtrTS<IAdaptiveStreamSelector>								StreamSelector;
diff --git a/Engine/Plugins/Media/ElectraPlayer/Source/ElectraPlayerRuntime/Private/Runtime/Player/DASH/StreamReaderFMP4DASH.cpp b/Engine/Plugins/Media/ElectraPlayer/Source/ElectraPlayerRuntime/Private/Runtime/Player/DASH/StreamReaderFMP4DASH.cpp
--- a/Engine/Plugins/Media/ElectraPlayer/Source/ElectraPlayerRuntime/Private/Runtime/Player/DASH/StreamReaderFMP4DASH.cpp
+++ b/Engine/Plugins/Media/ElectraPlayer/Source/ElectraPlayerRuntime/Private/Runtime/Player/DASH/StreamReaderFMP4DASH.cpp
@@ -811,6 +811,10 @@ void FStreamReaderFMP4DASH::FStreamHandler::HandleRequest()
 	bool bIsEmptyFillerSegment = Request->bInsertFillerData;
 	bool bIsLastSegment = Request->Segment.bIsLastInPeriod;
 
+	// Copy the source buffer info into a new instance and set the playback sequence ID in it.
+	TSharedPtrTS<FBufferSourceInfo> SourceBufferInfo = MakeSharedTS<FBufferSourceInfo>(*Request->SourceBufferInfo);
+	SourceBufferInfo->PlaybackSequenceID = Request->GetPlaybackSequenceID();
+
 	Metrics::FSegmentDownloadStats& ds = CurrentRequest->DownloadStats;
 	ds.StatsID = FMediaInterlockedIncrement(UniqueDownloadID);
 	ds.FailureReason.Empty();
@@ -923,7 +927,7 @@ void FStreamReaderFMP4DASH::FStreamHandler::HandleRequest()
 				AccessUnit->bIsSyncSample = true;
 				AccessUnit->bIsDummyData = false;
 				AccessUnit->bIsSideloaded = true;
-				AccessUnit->BufferSourceInfo = Request->SourceBufferInfo;
+				AccessUnit->BufferSourceInfo = SourceBufferInfo;
 				AccessUnit->SequenceIndex = Request->TimestampSequenceIndex;
 				AccessUnit->DropState = FAccessUnit::EDropState::None;
 
@@ -1140,7 +1144,7 @@ void FStreamReaderFMP4DASH::FStreamHandler::HandleRequest()
 									AccessUnit->bIsSyncSample = TrackIterator->IsSyncSample();
 									AccessUnit->bIsDummyData = false;
 									AccessUnit->AUCodecData = CSD;
-									AccessUnit->BufferSourceInfo = Request->SourceBufferInfo;
+									AccessUnit->BufferSourceInfo = SourceBufferInfo;
 
 									AccessUnit->DropState = FAccessUnit::EDropState::None;
 									// Is this AU entirely before the time we want?
@@ -1493,7 +1497,7 @@ void FStreamReaderFMP4DASH::FStreamHandler::HandleRequest()
 			}
 
 			AccessUnit->ESType = Request->GetType();
-			AccessUnit->BufferSourceInfo = Request->SourceBufferInfo;
+			AccessUnit->BufferSourceInfo = SourceBufferInfo;
 			AccessUnit->Duration = DefaultDuration;
 			AccessUnit->AUSize = 0;
 			AccessUnit->AUData = nullptr;
diff --git a/Engine/Plugins/Media/ElectraPlayer/Source/ElectraPlayerRuntime/Private/Runtime/Player/HLS/StreamReaderHLSfmp4.cpp b/Engine/Plugins/Media/ElectraPlayer/Source/ElectraPlayerRuntime/Private/Runtime/Player/HLS/StreamReaderHLSfmp4.cpp
--- a/Engine/Plugins/Media/ElectraPlayer/Source/ElectraPlayerRuntime/Private/Runtime/Player/HLS/StreamReaderHLSfmp4.cpp
+++ b/Engine/Plugins/Media/ElectraPlayer/Source/ElectraPlayerRuntime/Private/Runtime/Player/HLS/StreamReaderHLSfmp4.cpp
@@ -616,6 +616,9 @@ void FStreamReaderHLSfmp4::FStreamHandler::HandleRequest()
 	//        to make sure the retry info is not modified.
 	TSharedPtrTS<HTTP::FRetryInfo> CurrentRetryInfo = CurrentRequest->ConnectionInfo.RetryInfo;
 
+	TSharedPtrTS<FBufferSourceInfo> SourceBufferInfo = MakeSharedTS<FBufferSourceInfo>(*Request->SourceBufferInfo);
+	SourceBufferInfo->PlaybackSequenceID = Request->GetPlaybackSequenceID();
+
 	Metrics::FSegmentDownloadStats& ds = CurrentRequest->DownloadStats;
 	ds.StatsID = FMediaInterlockedIncrement(UniqueDownloadID);
 
@@ -917,7 +920,7 @@ void FStreamReaderHLSfmp4::FStreamHandler::HandleRequest()
 										AccessUnit->LatestPTS = Request->LastPTS;
 										AccessUnit->LatestPTS.SetSequenceIndex(Request->TimestampSequenceIndex);
 
-										AccessUnit->BufferSourceInfo = Request->SourceBufferInfo;
+										AccessUnit->BufferSourceInfo = SourceBufferInfo;
 
 										// There should not be any gaps!
 								// TODO: what if there are?
@@ -1169,7 +1172,7 @@ void FStreamReaderHLSfmp4::FStreamHandler::HandleRequest()
 				AccessUnit->AUSize = 0;
 				AccessUnit->AUData = nullptr;
 				AccessUnit->bIsDummyData = true;
-				AccessUnit->BufferSourceInfo = Request->SourceBufferInfo;
+				AccessUnit->BufferSourceInfo = SourceBufferInfo;
 				if (CSD.IsValid() && CSD->CodecSpecificData.Num())
 				{
 					AccessUnit->AUCodecData = CSD;
diff --git a/Engine/Plugins/Media/ElectraPlayer/Source/ElectraPlayerRuntime/Private/Runtime/Player/mp4/ManifestMP4.cpp b/Engine/Plugins/Media/ElectraPlayer/Source/ElectraPlayerRuntime/Private/Runtime/Player/mp4/ManifestMP4.cpp
--- a/Engine/Plugins/Media/ElectraPlayer/Source/ElectraPlayerRuntime/Private/Runtime/Player/mp4/ManifestMP4.cpp
+++ b/Engine/Plugins/Media/ElectraPlayer/Source/ElectraPlayerRuntime/Private/Runtime/Player/mp4/ManifestMP4.cpp
@@ -116,13 +116,13 @@ void FManifestMP4Internal::GetTrackMetadata(TArray<FTrackMetadata>& OutMetadata,
  * time range. This may be set through manifest internal means or by URL fragment
  * parameters where permissable (eg. example.mp4#t=22,50).
  * If start or end are not specified they will be set to invalid.
- * 
+ *
  * @return Optionally set time range to which playback is restricted.
  */
 FTimeRange FManifestMP4Internal::GetPlaybackRange() const
 {
 	FTimeRange FromTo;
-	
+
 	// We are interested in the 't' fragment value here.
 	FString Time;
 	for(auto& Fragment : URLFragmentComponents)
@@ -260,7 +260,11 @@ FManifestMP4Internal::FPlayPeriodMP4::~FPlayPeriodMP4()
  */
 void FManifestMP4Internal::FPlayPeriodMP4::SetStreamPreferences(EStreamType ForStreamType, const FStreamSelectionAttributes& StreamAttributes)
 {
-	if (ForStreamType == EStreamType::Audio)
+	if (ForStreamType == EStreamType::Video)
+	{
+		VideoPreferences = StreamAttributes;
+	}
+	else if (ForStreamType == EStreamType::Audio)
 	{
 		AudioPreferences = StreamAttributes;
 	}
@@ -728,7 +732,7 @@ void FManifestMP4Internal::FTimelineAssetMP4::LimitSegmentDownloadSize(TSharedPt
 					TrackDur += CurrentTrackIt->GetDuration();
 				}
 				int64 CurrentTrackOffset = LastTrackOffset;
-				if (CurrentTrackOffset >= EndOffset || 
+				if (CurrentTrackOffset >= EndOffset ||
 					CurrentTrackOffset - StartOffset >= MaxSegmentSize ||
 					(TrackDurationLimit > 0 && TrackDur > TrackDurationLimit))
 				{
@@ -801,7 +805,7 @@ IManifest::FResult FManifestMP4Internal::FTimelineAssetMP4::GetStartingSegment(T
 	SCOPE_CYCLE_COUNTER(STAT_ElectraPlayer_MP4_FindSegment);
 	CSV_SCOPED_TIMING_STAT(ElectraPlayer, MP4_FindSegment);
 
-// TODO: If there is a SIDX box we will look in there.
+// TODO: If there is a SIDX box we should look in there.
 
 	// Frame accurate seek required?
 	bool bFrameAccurateSearch = AtAbsoluteFilePos < 0 ? StartPosition.Options.bFrameAccuracy : false;
@@ -812,7 +816,8 @@ IManifest::FResult FManifestMP4Internal::FTimelineAssetMP4::GetStartingSegment(T
 	// Look at the actual tracks. If there is video search there first for a keyframe/IDR frame.
 	if (VideoAdaptationSets.Num())
 	{
-// TODO: use the selected track index if there are several tracks and we have a selected one.
+		// Start at the first track. The *assumption* is that the file is interleaved such that the video track with the lowest
+		// track ID comes before the other tracks for every sync sample. This may not always be the case though.
 		TSharedPtrTS<FRepresentationMP4> Repr = StaticCastSharedPtr<FRepresentationMP4>(VideoAdaptationSets[0]->GetRepresentationByIndex(0));
 		if (Repr.IsValid())
 		{
@@ -893,7 +898,7 @@ IManifest::FResult FManifestMP4Internal::FTimelineAssetMP4::GetStartingSegment(T
 					req->LastPTS = PlayRangeEnd;
 					req->TimestampSequenceIndex = InSequenceState.SequenceIndex;
 
-					// FIXME: this may need to add all additional tracks at some point if their individual IDs matter
+					// Add dependent stream types
 					if (AudioAdaptationSets.Num())
 					{
 						req->DependentStreamTypes.Add(EStreamType::Audio);
@@ -956,7 +961,6 @@ IManifest::FResult FManifestMP4Internal::FTimelineAssetMP4::GetStartingSegment(T
 	// No video track(s). Are there audio tracks?
 	if (AudioAdaptationSets.Num())
 	{
-// TODO: use the selected track index if there are several tracks and we have a selected one.
 		TSharedPtrTS<FRepresentationMP4> Repr = StaticCastSharedPtr<FRepresentationMP4>(AudioAdaptationSets[0]->GetRepresentationByIndex(0));
 		if (Repr.IsValid())
 		{
@@ -965,7 +969,6 @@ IManifest::FResult FManifestMP4Internal::FTimelineAssetMP4::GetStartingSegment(T
 			const IParserISO14496_12::ITrack* Track = MoovBoxParser->GetTrackByTrackID(TrackID);
 			if (Track)
 			{
-				//
 				TSharedPtrTS<IParserISO14496_12::ITrackIterator> TrackIt(Track->CreateIterator());
 				IParserISO14496_12::ITrackIterator::ESearchMode SearchMode =
 					SearchType == ESearchType::After  || SearchType == ESearchType::StrictlyAfter  ? IParserISO14496_12::ITrackIterator::ESearchMode::After  :
@@ -1038,6 +1041,7 @@ IManifest::FResult FManifestMP4Internal::FTimelineAssetMP4::GetStartingSegment(T
 					req->LastPTS = PlayRangeEnd;
 					req->TimestampSequenceIndex = InSequenceState.SequenceIndex;
 
+					// Add dependent stream types.
 					// In case the video stream is shorter than audio we still need to add it as a dependent stream
 					// (if it exists) in case the video will loop back to a point where there is video.
 					if (VideoAdaptationSets.Num())
@@ -1049,8 +1053,6 @@ IManifest::FResult FManifestMP4Internal::FTimelineAssetMP4::GetStartingSegment(T
 						req->DependentStreamTypes.Add(EStreamType::Subtitle);
 					}
 
-					// FIXME: there may be subtitle tracks here we need to add as dependent streams.
-
 					LimitSegmentDownloadSize(OutSegment);
 					if (req->SegmentInternalSize > 0)
 					{
@@ -1229,7 +1231,11 @@ FErrorDetail FManifestMP4Internal::FRepresentationMP4::CreateFrom(const IParserI
 	// NOTE: This *MUST* be just a number since it gets parsed back out from a string into a number later! Do *NOT* prepend/append any string literals!!
 	UniqueIdentifier = LexToString(InTrack->GetID());
 
-	Name = InTrack->GetNameFromHandler();
+	Name = InTrack->GetName();
+	if (Name.Len() == 0)
+	{
+		Name = FString::Printf(TEXT("%s (ID=%u)"), *InTrack->GetNameFromHandler(), InTrack->GetID());
+	}
 
 	// Get bitrate from the average or max bitrate as stored in the track. If not stored it will be 0.
 	Bitrate = InTrack->GetBitrateInfo().AvgBitrate ? InTrack->GetBitrateInfo().AvgBitrate : InTrack->GetBitrateInfo().MaxBitrate;
diff --git a/Engine/Plugins/Media/ElectraPlayer/Source/ElectraPlayerRuntime/Private/Runtime/Player/mp4/StreamReaderMP4.cpp b/Engine/Plugins/Media/ElectraPlayer/Source/ElectraPlayerRuntime/Private/Runtime/Player/mp4/StreamReaderMP4.cpp
--- a/Engine/Plugins/Media/ElectraPlayer/Source/ElectraPlayerRuntime/Private/Runtime/Player/mp4/StreamReaderMP4.cpp
+++ b/Engine/Plugins/Media/ElectraPlayer/Source/ElectraPlayerRuntime/Private/Runtime/Player/mp4/StreamReaderMP4.cpp
@@ -438,15 +438,26 @@ void FStreamReaderMP4::HandleRequest()
 		bHasErrored = true;
 	}
 
+	uint32 PlaybackSequenceID = Request->GetPlaybackSequenceID();
 	while(!bDone && !HasErrored() && !HasBeenAborted() && !bTerminate)
 	{
-		auto UpdateSelectedTrack = [&SelectedTrackMap](const IParserISO14496_12::ITrackIterator* trkIt, TMap<uint32, FSelectedTrackData>& ActiveTrks) -> FSelectedTrackData&
+		auto UpdateSelectedTrack = [&SelectedTrackMap, &PlaybackSequenceID](const IParserISO14496_12::ITrackIterator* trkIt, TMap<uint32, FSelectedTrackData>& ActiveTrks) -> FSelectedTrackData&
 		{
 			const IParserISO14496_12::ITrack* Track = trkIt->GetTrack();
 			uint32 tkid = Track->GetID();
 
 			// Check if this track ID is already in our map of active tracks.
 			FSelectedTrackData& st = ActiveTrks.FindOrAdd(tkid);
+			if (!st.CSD.IsValid())
+			{
+				TSharedPtrTS<FAccessUnit::CodecData> CSD(new FAccessUnit::CodecData);
+				CSD->CodecSpecificData = Track->GetCodecSpecificData();
+				CSD->RawCSD			   = Track->GetCodecSpecificDataRAW();
+				CSD->ParsedInfo		   = Track->GetCodecInformation();
+				// Set information not necessarily available on the CSD.
+				CSD->ParsedInfo.SetBitrate(st.Bitrate);
+				st.CSD = MoveTemp(CSD);
+			}
 			if (!st.BufferSourceInfo.IsValid())
 			{
 				auto meta = MakeShared<FBufferSourceInfo, ESPMode::ThreadSafe>();
@@ -460,21 +471,13 @@ void FStreamReaderMP4::HandleRequest()
 					st.Bitrate = SelectedTrackMetadata->Bitrate;
 					meta->Kind = SelectedTrackMetadata->Kind;
 					meta->Language = SelectedTrackMetadata->Language;
+					meta->Codec = st.CSD.IsValid() ? st.CSD->ParsedInfo.GetCodecName() : FString();
 					meta->PeriodAdaptationSetID = SelectedTrackMetadata->PeriodID + TEXT(".") + SelectedTrackMetadata->AdaptationSetID;
 					meta->HardIndex = SelectedTrackMetadata->Index;
+					meta->PlaybackSequenceID = PlaybackSequenceID;
 				}
 				st.BufferSourceInfo = MoveTemp(meta);
 			}
-			if (!st.CSD.IsValid())
-			{
-				TSharedPtrTS<FAccessUnit::CodecData> CSD(new FAccessUnit::CodecData);
-				CSD->CodecSpecificData = Track->GetCodecSpecificData();
-				CSD->RawCSD			   = Track->GetCodecSpecificDataRAW();
-				CSD->ParsedInfo		   = Track->GetCodecInformation();
-				// Set information not necessarily available on the CSD.
-				CSD->ParsedInfo.SetBitrate(st.Bitrate);
-				st.CSD = MoveTemp(CSD);
-			}
 			return st;
 		};
 
diff --git a/Engine/Plugins/Media/ElectraPlayer/Source/ElectraPlayerRuntime/Private/Runtime/StreamAccessUnitBuffer.cpp b/Engine/Plugins/Media/ElectraPlayer/Source/ElectraPlayerRuntime/Private/Runtime/StreamAccessUnitBuffer.cpp
--- a/Engine/Plugins/Media/ElectraPlayer/Source/ElectraPlayerRuntime/Private/Runtime/StreamAccessUnitBuffer.cpp
+++ b/Engine/Plugins/Media/ElectraPlayer/Source/ElectraPlayerRuntime/Private/Runtime/StreamAccessUnitBuffer.cpp
@@ -4,14 +4,16 @@
 
 namespace Electra
 {
-	FMultiTrackAccessUnitBuffer::FMultiTrackAccessUnitBuffer()
+	FMultiTrackAccessUnitBuffer::FMultiTrackAccessUnitBuffer(EStreamType InForType)
 	{
+		Type = InForType;
 		EmptyBuffer = CreateNewBuffer();
 		LastPoppedDTS.SetToInvalid();
 		LastPoppedPTS.SetToInvalid();
 		bEndOfData = false;
 		bLastPushWasBlocked = false;
 		bIsParallelTrackMode = false;
+		bPopAsDummyUntilSyncFrame = true;
 	}
 
 	TSharedPtrTS<FAccessUnitBuffer> FMultiTrackAccessUnitBuffer::CreateNewBuffer()
@@ -22,7 +24,7 @@ namespace Electra
 
 	FMultiTrackAccessUnitBuffer::~FMultiTrackAccessUnitBuffer()
 	{
-		PurgeAll();
+		Clear();
 	}
 
 	void FMultiTrackAccessUnitBuffer::SetParallelTrackMode()
@@ -30,71 +32,77 @@ namespace Electra
 		bIsParallelTrackMode = true;
 	}
 
+
+	TSharedPtrTS<FAccessUnitBuffer> FMultiTrackAccessUnitBuffer::FindOrCreateBufferFor(TSharedPtrTS<const FBufferSourceInfo>& OutBufferSourceInfo, const TSharedPtrTS<const FBufferSourceInfo>& InBufferInfo, bool bCreateIfNotExist)
+	{
+		// Note: The access mutex must have been locked already!
+		for(int32 i=0; i<BufferList.Num(); ++i)
+		{
+			if (BufferList[i].Info->PlaybackSequenceID == InBufferInfo->PlaybackSequenceID)
+			{
+				if (bIsParallelTrackMode)
+				{
+					if (BufferList[i].Info->HardIndex == InBufferInfo->HardIndex)
+					{
+						OutBufferSourceInfo = BufferList[i].Info;
+						return BufferList[i].Buffer;
+					}
+				}
+				else
+				{
+					OutBufferSourceInfo = BufferList[i].Info;
+					return BufferList[i].Buffer;
+				}
+			}
+		}
+		TSharedPtrTS<FAccessUnitBuffer> Buffer;
+		if (bCreateIfNotExist)
+		{
+			FBufferByInfoType bit;
+			OutBufferSourceInfo = bit.Info = InBufferInfo;
+			Buffer = bit.Buffer = CreateNewBuffer();
+			BufferList.Emplace(MoveTemp(bit));
+		}
+		return Buffer;
+	}
+
 	bool FMultiTrackAccessUnitBuffer::Push(FAccessUnit*& AU, const FAccessUnitBuffer::FConfiguration* BufferConfiguration, const FAccessUnitBuffer::FExternalBufferInfo* InCurrentTotalBufferUtilization)
 	{
 		check(AU->BufferSourceInfo.IsValid());
 
-		FString TrackID = AU->BufferSourceInfo->PeriodAdaptationSetID;
-		// Do we have a buffer for this track ID yet?
 		AccessLock.Lock();
-		if (!TrackBuffers.Contains(TrackID))
-		{
-			TrackBuffers.Emplace(TrackID, CreateNewBuffer());
-		}
-		TSharedPtrTS<FAccessUnitBuffer> TrackBuffer = TrackBuffers[TrackID];
 
-		ActivateBuffer(true);
+		TSharedPtrTS<const FBufferSourceInfo> BufferSourceInfo;
+		TSharedPtrTS<FAccessUnitBuffer> Buffer = FindOrCreateBufferFor(BufferSourceInfo, AU->BufferSourceInfo, true);
+
+		ActivateInitialBuffer();
 
 		// Pushing data to either buffer means that there is data and we are not at EOD here any more.
 		// This does not necessarily mean that the selected track is not at EOD. Just that some track is not.
 		bEndOfData = false;
 
-		bool bIsSwitchOverAU = false;
-		for(auto &SwitchBuf : SwitchOverBufferChain)
-		{
-			if (SwitchBuf.Info->PeriodAdaptationSetID.Equals(TrackID))
-			{
-				bIsSwitchOverAU = true;
-				break;
-			}
-		}
+		AccessLock.Unlock();
 
-		// Get the combined buffer utilization of enqueued buffers.
-		FAccessUnitBuffer::FExternalBufferInfo EnqueuedBufferInfo;
-		if (InCurrentTotalBufferUtilization)
+		if (!Buffer.IsValid())
 		{
-			EnqueuedBufferInfo = *InCurrentTotalBufferUtilization;
-		}
-		else
-		{
-			GetEnqueuedBufferInfo(EnqueuedBufferInfo, bIsSwitchOverAU);
+			return false;
 		}
-
-		AccessLock.Unlock();
-
-		bool bWasPushed = TrackBuffer->Push(AU, BufferConfiguration, &EnqueuedBufferInfo);
-		bLastPushWasBlocked = TrackBuffer->WasLastPushBlocked();
-
+		bool bWasPushed = Buffer->Push(AU, BufferConfiguration, InCurrentTotalBufferUtilization);
+		bLastPushWasBlocked = Buffer->WasLastPushBlocked();
 		return bWasPushed;
 	}
 
 	void FMultiTrackAccessUnitBuffer::PushEndOfDataFor(TSharedPtrTS<const FBufferSourceInfo> InStreamSourceInfo)
 	{
-		check(InStreamSourceInfo.IsValid());
-		
-		FString TrackID = InStreamSourceInfo->PeriodAdaptationSetID;
-		// Do we have a buffer for this track ID yet?
 		AccessLock.Lock();
-		if (!TrackBuffers.Contains(TrackID))
+		TSharedPtrTS<const FBufferSourceInfo> BufferSourceInfo;
+		TSharedPtrTS<FAccessUnitBuffer> Buffer = FindOrCreateBufferFor(BufferSourceInfo, InStreamSourceInfo, true);
+		ActivateInitialBuffer();
+		AccessLock.Unlock();
+		if (Buffer.IsValid())
 		{
-			TrackBuffers.Emplace(TrackID, CreateNewBuffer());
+			Buffer->PushEndOfData();
 		}
-		TSharedPtrTS<FAccessUnitBuffer> TrackBuffer = TrackBuffers[TrackID];
-
-		ActivateBuffer(true);
-
-		AccessLock.Unlock();
-		TrackBuffer->PushEndOfData();
 	}
 
 	void FMultiTrackAccessUnitBuffer::PushEndOfDataAll()
@@ -102,141 +110,115 @@ namespace Electra
 		FScopeLock lock(&AccessLock);
 		bEndOfData = true;
 		// Push an end-of-data into all tracks.
-		for(auto& It : TrackBuffers)
+		for(auto& It : BufferList)
 		{
-			It.Value->PushEndOfData();
+			It.Buffer->PushEndOfData();
 		}
 	}
 
 	void FMultiTrackAccessUnitBuffer::Clear()
 	{
+		BufferList.Empty();
+		PendingBufferSwitch.Reset();
+		ActiveBuffer.Reset();
+		ActiveOutputBufferInfo.Reset();
+		LastPoppedBufferInfo.Reset();
 		LastPoppedDTS.SetToInvalid();
 		LastPoppedPTS.SetToInvalid();
 		bEndOfData = false;
 		bLastPushWasBlocked = false;
-		ActiveOutputBufferInfo.Reset();
-		LastPoppedBufferInfo.Reset();
-		UpcomingBufferChain.Empty();
-		SwitchOverBufferChain.Empty();
+		bPopAsDummyUntilSyncFrame = true;
 	}
 
 	void FMultiTrackAccessUnitBuffer::Flush()
 	{
 		FScopeLock lock(&AccessLock);
-		if (bIsParallelTrackMode)
+		Clear();
+	}
+
+	void FMultiTrackAccessUnitBuffer::SelectTrackWhenAvailable(uint32 PlaybackSequenceID, TSharedPtrTS<FBufferSourceInfo> InBufferSourceInfo)
+	{
+		FScopeLock lock(&AccessLock);
+		if (InBufferSourceInfo.IsValid())
 		{
-			// Flush all existing buffers but keep them in the track map.
-			for(auto& It : TrackBuffers)
-			{
-				It.Value->Flush();
-			}
+			PendingBufferSwitch.BufferInfo = MakeSharedTS<FBufferSourceInfo>(*InBufferSourceInfo);
+			PendingBufferSwitch.BufferInfo->PlaybackSequenceID = PlaybackSequenceID;
 		}
 		else
 		{
-			TrackBuffers.Empty();
+			PendingBufferSwitch.Reset();
 		}
-		Clear();
 	}
 
-	void FMultiTrackAccessUnitBuffer::PurgeAll()
+	void FMultiTrackAccessUnitBuffer::ActivateInitialBuffer()
 	{
-		FScopeLock lock(&AccessLock);
-		// Get rid of all buffers
-		TrackBuffers.Empty();
-		Clear();
+		// Note: The access mutex must have been locked already!
+		if (!ActiveBuffer.IsValid() && PendingBufferSwitch.IsSet())
+		{
+			TSharedPtrTS<const FBufferSourceInfo> BufferSourceInfo;
+			TSharedPtrTS<FAccessUnitBuffer> Buffer = FindOrCreateBufferFor(BufferSourceInfo, PendingBufferSwitch.BufferInfo, false);
+			if (Buffer.IsValid())
+			{
+				ActiveBuffer = MoveTemp(Buffer);
+				ActiveOutputBufferInfo = MoveTemp(BufferSourceInfo);
+				PendingBufferSwitch.BufferInfo.Reset();
+			}
+		}
 	}
 
-
-	void FMultiTrackAccessUnitBuffer::SelectTrackWhenAvailable(TSharedPtrTS<FBufferSourceInfo> InBufferSourceInfo)
+	void FMultiTrackAccessUnitBuffer::HandlePendingSwitch()
 	{
-		if (InBufferSourceInfo.IsValid())
+		// Note: The access mutex must have been locked already!
+		if (PendingBufferSwitch.IsSet())
 		{
-			FScopeLock lock(&AccessLock);
-			if (SwitchOverBufferChain.Num() == 0)
+			TSharedPtrTS<const FBufferSourceInfo> BufferSourceInfo;
+			TSharedPtrTS<FAccessUnitBuffer> Buffer = FindOrCreateBufferFor(BufferSourceInfo, PendingBufferSwitch.BufferInfo, false);
+			// Switching to the same buffer we are already on means we're already done.
+			if (ActiveBuffer.IsValid() && ActiveBuffer == Buffer)
 			{
-				FQueuedBuffer Next;
-				Next.Info = MoveTemp(InBufferSourceInfo);
-				SwitchOverBufferChain.Empty();
-				SwitchOverBufferChain.Emplace(MoveTemp(Next));
+				PendingBufferSwitch.BufferInfo.Reset();
+				return;
 			}
-			else if (!SwitchOverBufferChain.Last().Info->PeriodAdaptationSetID.Equals(InBufferSourceInfo->PeriodAdaptationSetID))
-			{
-				/*
-					We have to check if what we are to switch to is what is already active.
-					This can happen on quick track changes where we are currently on track1, then a change to track2 was issued
-					but immediately chaged back to track1. So: track1 -> track2 -> track1
-					In this case we would remove the request for track2 and see track1 -> track1. Since the ID of this track is
-					still the same we would append the data coming in from the new track1 request to the already existing track1
-					buffer. This is bad because segments are streamed in starting at their beginning and we would thus put data
-					into the buffer that is already there.
-					To not make things overly complex we will do the following if the switch-to track is the currently active one:
-					We rename the active one by appending some suffix. This works because the current stream download was already
-					canceled and no new data is coming in from that stream. We need to make sure we change the ID in the current
-					track map (the key) as well as the source info. The source info is maintained as a shared pointer, so changing
-					it there means it is automatically changed everywhere that pointer is shared!
-				*/
-
-				if (ActiveOutputBufferInfo.IsValid() && ActiveOutputBufferInfo->PeriodAdaptationSetID.Equals(InBufferSourceInfo->PeriodAdaptationSetID))
-				{
-					// The ID may be the same but the TSharedPtr MUST NOT be.
-					// If it were the same pointer we would be renaming the new segment buffer and info as well!
-					check(ActiveOutputBufferInfo != InBufferSourceInfo);
 
-					// First let's rename the buffer in the track buffer map. Double check that it is actually there.
-					FString NewName = InBufferSourceInfo->PeriodAdaptationSetID + TEXT("$1");
-					if (TrackBuffers.Contains(InBufferSourceInfo->PeriodAdaptationSetID))
-					{
-						TSharedPtrTS<FAccessUnitBuffer> TrackBuffer = TrackBuffers[ActiveOutputBufferInfo->PeriodAdaptationSetID];
-						TrackBuffers.Remove(ActiveOutputBufferInfo->PeriodAdaptationSetID);
-						TrackBuffers.Emplace(NewName, TrackBuffer);
-					}
-					// Then change in the name in the active output buffer info.
-					ActiveOutputBufferInfo->PeriodAdaptationSetID = NewName;
+			if (Buffer.IsValid() && LastPoppedDTS.IsValid())
+			{
+				FTimeValue dd, dp;
+				Buffer->DiscardUntil(LastPoppedDTS, FTimeValue(), dd, dp);
 
-					// Finally, if the LastPoppedBufferInfo is set the it SHOULD be the same as the ActiveOutputBufferInfo,
-					// ie: check(LastPoppedBufferInfo == ActiveOutputBufferInfo);
-					// But in case it is not (due to future code changes) let's update it as well.
-					if (LastPoppedBufferInfo.IsValid() && LastPoppedBufferInfo->PeriodAdaptationSetID.Equals(InBufferSourceInfo->PeriodAdaptationSetID))
+				// We switch on keyframes only. If there is none we keep the current buffer even if this may mean that
+				// we run out of data altogether.
+				FAccessUnit* NextAU = nullptr;
+				if (Buffer->PeekAndAddRef(NextAU))
+				{
+					if (NextAU->bIsSyncSample)
 					{
-						LastPoppedBufferInfo->PeriodAdaptationSetID = NewName;
+						ActiveBuffer = MoveTemp(Buffer);
+						ActiveOutputBufferInfo = MoveTemp(BufferSourceInfo);
+						PendingBufferSwitch.BufferInfo.Reset();
 					}
 				}
-
-				FQueuedBuffer Next;
-				Next.Info = MoveTemp(InBufferSourceInfo);
-				SwitchOverBufferChain.Empty();
-				SwitchOverBufferChain.Emplace(MoveTemp(Next));
 			}
 		}
+		RemoveOutdatedBuffers();
 	}
 
-	void FMultiTrackAccessUnitBuffer::AddUpcomingBuffer(TSharedPtrTS<FBufferSourceInfo> InBufferSourceInfo)
+	void FMultiTrackAccessUnitBuffer::RemoveOutdatedBuffers()
 	{
-		if (InBufferSourceInfo.IsValid())
+		if (ActiveOutputBufferInfo.IsValid())
 		{
-			FScopeLock lock(&AccessLock);
-			// Add to the current or the switch-over chain?
-			if (SwitchOverBufferChain.Num())
-			{
-				if (!SwitchOverBufferChain.Last().Info->PeriodAdaptationSetID.Equals(InBufferSourceInfo->PeriodAdaptationSetID))
-				{
-					FQueuedBuffer Next;
-					Next.Info = MoveTemp(InBufferSourceInfo);
-					SwitchOverBufferChain.Emplace(MoveTemp(Next));
-				}
-			}
-			else
+			uint32 PlaybackSequenceID = ActiveOutputBufferInfo->PlaybackSequenceID;
+			for(int32 i=0; i<BufferList.Num(); ++i)
 			{
-				if (UpcomingBufferChain.Num() == 0 || !UpcomingBufferChain.Last().Info->PeriodAdaptationSetID.Equals(InBufferSourceInfo->PeriodAdaptationSetID))
+				if (BufferList[i].Info->PlaybackSequenceID < PlaybackSequenceID)
 				{
-					FQueuedBuffer Next;
-					Next.Info = MoveTemp(InBufferSourceInfo);
-					UpcomingBufferChain.Emplace(MoveTemp(Next));
+					BufferList.RemoveAt(i);
+					--i;
 				}
 			}
 		}
 	}
 
+
 	FTimeValue FMultiTrackAccessUnitBuffer::GetLastPoppedPTS()
 	{
 		FScopeLock lock(&AccessLock);
@@ -251,208 +233,27 @@ namespace Electra
 
 	TSharedPtrTS<FAccessUnitBuffer> FMultiTrackAccessUnitBuffer::GetSelectedTrackBuffer()
 	{
-		TSharedPtrTS<FAccessUnitBuffer> TrackBuffer;
 		FScopeLock lock(&AccessLock);
-		if (ActiveOutputBufferInfo.IsValid() && TrackBuffers.Contains(ActiveOutputBufferInfo->PeriodAdaptationSetID))
-		{
-			TrackBuffer = TrackBuffers[ActiveOutputBufferInfo->PeriodAdaptationSetID];
-		}
-		else
-		{
-			TrackBuffer = EmptyBuffer;
-		}
-		return TrackBuffer;
+		return ActiveBuffer.IsValid() ? ActiveBuffer : EmptyBuffer;
 	}
 
 	void FMultiTrackAccessUnitBuffer::GetStats(FAccessUnitBufferInfo& OutStats)
 	{
 		FScopeLock lock(&AccessLock);
+		ActivateInitialBuffer();
 		TSharedPtrTS<const FAccessUnitBuffer> Buf = GetSelectedTrackBuffer();
 		Buf->GetStats(OutStats);
 		if (Buf == EmptyBuffer)
 		{
 			OutStats.bEndOfData = bEndOfData;
 		}
-		else
-		{
-			// Add the stats of any enqueued buffers
-			for(int32 i=0; i<UpcomingBufferChain.Num(); ++i)
-			{
-				if (UpcomingBufferChain[i].Buffer.IsValid())
-				{
-					FAccessUnitBufferInfo NextStats;
-					UpcomingBufferChain[i].Buffer->GetStats(NextStats);
-					OutStats.PushedDuration += NextStats.PushedDuration;
-					OutStats.PlayableDuration += NextStats.PlayableDuration;
-					OutStats.CurrentMemInUse += NextStats.CurrentMemInUse;
-					OutStats.NumCurrentAccessUnits += NextStats.NumCurrentAccessUnits;
-					// With chained buffers the end of data and last push blocked flags are only
-					// relevant on the last buffer in the chain.
-					OutStats.bEndOfData = NextStats.bEndOfData;
-					OutStats.bLastPushWasBlocked = NextStats.bLastPushWasBlocked;
-				}
-			}
-		}
-	}
-
-	void FMultiTrackAccessUnitBuffer::ChangeOver()
-	{
-		ActiveOutputBufferInfo = SwitchOverBufferChain[0].Info;
-		SwitchOverBufferChain.RemoveAt(0);
-		UpcomingBufferChain = MoveTemp(SwitchOverBufferChain);
-		RemoveUnusedBuffers();
-	}
-
-	void FMultiTrackAccessUnitBuffer::RemoveUnusedBuffers()
-	{
-		if (!bIsParallelTrackMode)
-		{
-			auto IsReferencesInChain = [&](const FString& ID, TArray<FQueuedBuffer>& BufferChain) -> bool
-			{
-				for(auto &Next : BufferChain)
-				{
-					if (Next.Info.IsValid() && Next.Info->PeriodAdaptationSetID.Equals(ID))
-					{
-						return true;
-					}
-				}
-				return false;
-			};
-
-			TArray<FString> UnusedBuffers;
-			for(auto &Buffer : TrackBuffers)
-			{
-				const FString& ID = Buffer.Key;
-				bool bIsReferenced = (ActiveOutputBufferInfo.IsValid() && ActiveOutputBufferInfo->PeriodAdaptationSetID.Equals(ID)) || IsReferencesInChain(ID, UpcomingBufferChain) || IsReferencesInChain(ID, SwitchOverBufferChain);
-				if (!bIsReferenced)
-				{
-					UnusedBuffers.Emplace(ID);
-				}
-			}
-			for(auto &Buffer : UnusedBuffers)
-			{
-				TrackBuffers.Remove(Buffer);
-			}
-		}
-	}
-
-	void FMultiTrackAccessUnitBuffer::ActivateBuffer(bool bIsPushing)
-	{
-		auto UpdateBufferChainBuffers = [&](TArray<FQueuedBuffer>& BufferChain) -> void
-		{
-			for(auto &Next : BufferChain)
-			{
-				if (!Next.Buffer.IsValid() && TrackBuffers.Contains(Next.Info->PeriodAdaptationSetID))
-				{
-					Next.Buffer = TrackBuffers[Next.Info->PeriodAdaptationSetID];
-				}
-			}
-		};
-
-
-		// Update the upcoming buffer chains for GetStats().
-		UpdateBufferChainBuffers(UpcomingBufferChain);
-		UpdateBufferChainBuffers(SwitchOverBufferChain);
-
-		if (!ActiveOutputBufferInfo.IsValid())
-		{
-			check(UpcomingBufferChain.Num() == 0);
-			if (SwitchOverBufferChain.Num())
-			{
-				ChangeOver();
-			}
-		}
-		// Upon pushing data to the buffer we only need to activate the track and update the buffer chains
-		// for statistics. Buffer switching is handled when popping data.
-		if (bIsPushing)
-		{
-			return;
-		}
-
-		if (SwitchOverBufferChain.Num())
-		{
-			// In parallel track mode we switch immediately as the individual buffers are
-			// time synchronized already.
-			if (bIsParallelTrackMode)
-			{
-				ChangeOver();
-			}
-			else
-			{
-				FTimeValue PoppedDTS, PoppedPTS;
-				FTimeValue PopUntilDTS, PopUntilPTS;
-				PopUntilDTS = LastPoppedDTS;
-				PopUntilPTS = LastPoppedPTS;
-				bool bSwitch = false;
-				if (PopUntilDTS.IsValid() && PopUntilPTS.IsValid())
-				{
-					TSharedPtrTS<FAccessUnitBuffer> ActiveBuffer = GetSelectedTrackBuffer();
-					bool bActiveIsEmpty = !ActiveBuffer.IsValid() || ActiveBuffer->Num() == 0;
-					for(int32 nBuf=0; nBuf<SwitchOverBufferChain.Num(); ++nBuf)
-					{
-						if (SwitchOverBufferChain[nBuf].Buffer.IsValid())
-						{
-							// Remove AUs from the buffer that are outdated.
-							SwitchOverBufferChain[nBuf].Buffer->DiscardUntil(PopUntilDTS, PopUntilPTS, PoppedDTS, PoppedPTS);
-
-							// Does it contain the time we want to switch at?
-							if (SwitchOverBufferChain[nBuf].Buffer->ContainsPTS(PopUntilPTS))
-							{
-								bSwitch = true;
-								break;
-							}
-							// If this buffer is empty and the switchover one is not we switch over now.
-							// This is for the cases where we are not getting any more data from the current buffer
-							// and the LastPoppedDTS/PTS are in the past and will not be contained in the new buffer,
-							// which happens mostly with sparse data like infrequent subtitle AUs.
-							if (bActiveIsEmpty && SwitchOverBufferChain[nBuf].Buffer->Num())
-							{
-								bSwitch = true;
-								break;
-							}
-							// If the buffer is now empty with another non-empty one following we can remove this one.
-							if (SwitchOverBufferChain[nBuf].Buffer->Num() == 0 && nBuf+1 < SwitchOverBufferChain.Num() && SwitchOverBufferChain[nBuf+1].Buffer.IsValid() && SwitchOverBufferChain[nBuf+1].Buffer->Num())
-							{
-								SwitchOverBufferChain.RemoveAt(nBuf);
-								--nBuf;
-							}
-						}
-						else
-						{
-							// If this buffer is not valid (yet) then the following ones are not either.
-							break;
-						}
-					}
-				}
-				else
-				{
-					// When no data was popped from the buffer yet do an immediate switch.
-					bSwitch = true;
-				}
-				if (bSwitch)
-				{
-					ChangeOver();
-				}
-			}
-		}
-
-		// Moving into the next enqueued buffer?
-		if (UpcomingBufferChain.Num())
-		{
-			TSharedPtrTS<FAccessUnitBuffer> ActiveBuffer = GetSelectedTrackBuffer();
-			if (ActiveBuffer->Num() == 0 && UpcomingBufferChain[0].Buffer.IsValid() && UpcomingBufferChain[0].Buffer->Num())
-			{
-				ActiveOutputBufferInfo = UpcomingBufferChain[0].Info;
-				UpcomingBufferChain.RemoveAt(0);
-				RemoveUnusedBuffers();
-			}
-		}
 	}
 
 	bool FMultiTrackAccessUnitBuffer::PeekAndAddRef(FAccessUnit*& OutAU)
 	{
 		FScopedLock lock(AsShared());
-		ActivateBuffer(false);
+		ActivateInitialBuffer();
+		HandlePendingSwitch();
 		TSharedPtrTS<FAccessUnitBuffer> Buf = FMultiTrackAccessUnitBuffer::GetSelectedTrackBuffer();
 		if (Buf->Num())
 		{
@@ -468,10 +269,10 @@ namespace Electra
 	bool FMultiTrackAccessUnitBuffer::Pop(FAccessUnit*& OutAU)
 	{
 		// Note: We assume the access lock is held by the caller!
-		ActivateBuffer(false);
-
+		ActivateInitialBuffer();
+		HandlePendingSwitch();
 		TSharedPtrTS<FAccessUnitBuffer> Buf = FMultiTrackAccessUnitBuffer::GetSelectedTrackBuffer();
-		if (Buf->Num())
+		if (Buf.IsValid() && Buf->Num())
 		{
 			bool bDidPop = Buf->Pop(OutAU);
 			if (bDidPop && OutAU)
@@ -481,22 +282,30 @@ namespace Electra
 				{
 					if (LastPoppedBufferInfo)
 					{
-						// FIXME: We may need to discard everything that is not tagged as a sync sample to ensure proper stream switching.
 						OutAU->bTrackChangeDiscontinuity = true;
 					}
 
 					// Remember from which buffer we popped.
 					LastPoppedBufferInfo = ActiveOutputBufferInfo;
 				}
+
+				// Do we need to wait for a sync frame and tag everything else as a dummy until then?
+				if (bPopAsDummyUntilSyncFrame)
+				{
+					bPopAsDummyUntilSyncFrame = !OutAU->bIsSyncSample;
+					OutAU->bIsDummyData = !OutAU->bIsSyncSample || OutAU->bIsDummyData;
+					OutAU->bTrackChangeDiscontinuity = OutAU->bIsSyncSample || OutAU->bTrackChangeDiscontinuity;
+				}
+
 				LastPoppedDTS = OutAU->DTS;
 				LastPoppedPTS = OutAU->PTS;
 				// With this AU being popped off now we will also pop off all AUs that are now obsolete in the other tracks.
-				for(auto& It : TrackBuffers)
+				for(auto& It : BufferList)
 				{
-					if (It.Value != Buf)
+					if (It.Buffer != Buf)
 					{
 						FTimeValue PoppedDTS, PoppedPTS;
-						It.Value->DiscardUntil(LastPoppedDTS, LastPoppedPTS, PoppedDTS, PoppedPTS);
+						It.Buffer->DiscardUntil(LastPoppedDTS, LastPoppedPTS, PoppedDTS, PoppedPTS);
 					}
 				}
 			}
@@ -512,11 +321,10 @@ namespace Electra
 	void FMultiTrackAccessUnitBuffer::PopDiscardUntil(FTimeValue UntilTime)
 	{
 		// Note: We assume the access lock is held by the caller!
-		for(auto& It : TrackBuffers)
+		for(auto& It : BufferList)
 		{
-			TSharedPtrTS<FAccessUnitBuffer> Buf = It.Value;
 			FTimeValue PoppedDTS, PoppedPTS;
-			Buf->DiscardUntil(UntilTime, UntilTime, PoppedDTS, PoppedPTS);
+			It.Buffer->DiscardUntil(UntilTime, UntilTime, PoppedDTS, PoppedPTS);
 			if (PoppedDTS.IsValid() && LastPoppedDTS.IsValid() && PoppedDTS > LastPoppedDTS)
 			{
 				LastPoppedDTS = PoppedDTS;
@@ -526,19 +334,9 @@ namespace Electra
 				LastPoppedPTS = PoppedPTS;
 			}
 		}
-	}
-
-	FTimeValue FMultiTrackAccessUnitBuffer::PrepareForDecodeStartingAt(FTimeValue DecodeStartTime)
-	{
-		FTimeValue TaggedDuration(FTimeValue::GetZero());
-		// We lock the access mutex here.
-		FScopedLock lock(AsShared());
-		for(auto& It : TrackBuffers)
-		{
-			TSharedPtrTS<FAccessUnitBuffer> Buf = It.Value;
-			TaggedDuration += Buf->PrepareForDecodeStartingAt(DecodeStartTime);
-		}
-		return TaggedDuration;
+		// Discarding from the buffer means that the next popping of an access unit needs to return
+		// a sync sample. Anything that is not will instead be returned as a dummy sample.
+		bPopAsDummyUntilSyncFrame  = true;
 	}
 
 	bool FMultiTrackAccessUnitBuffer::IsEODFlagSet()
@@ -551,50 +349,19 @@ namespace Electra
 	int32 FMultiTrackAccessUnitBuffer::Num()
 	{
 		// Note: We assume the access lock is held by the caller!
-		int64 TotalNum = GetSelectedTrackBuffer()->Num();
-		// Add the enqueued buffers
-		for(int32 i=0; i<UpcomingBufferChain.Num(); ++i)
-		{
-			if (UpcomingBufferChain[i].Buffer.IsValid())
-			{
-				TotalNum += UpcomingBufferChain[i].Buffer->Num();
-			}
-		}
-		return (int32) TotalNum;
+		ActivateInitialBuffer();
+		return GetSelectedTrackBuffer()->Num();
 	}
 
-	void FMultiTrackAccessUnitBuffer::GetEnqueuedBufferInfo(FAccessUnitBuffer::FExternalBufferInfo& OutInfo, bool bForSwitchOverChain)
+	bool FMultiTrackAccessUnitBuffer::WasLastPushBlocked()
 	{
-		if (!bForSwitchOverChain)
-		{
-			TSharedPtrTS<FAccessUnitBuffer> Buf = GetSelectedTrackBuffer();
-			OutInfo.DataSize = Buf->AllocatedSize();
-			OutInfo.Duration = Buf->GetPlayableDuration();
-			for(int32 i=0; i<UpcomingBufferChain.Num(); ++i)
-			{
-				if (UpcomingBufferChain[i].Buffer.IsValid())
-				{
-					OutInfo.DataSize += UpcomingBufferChain[i].Buffer->AllocatedSize();
-					OutInfo.Duration += UpcomingBufferChain[i].Buffer->GetPlayableDuration();
-				}
-			}
-		}
-		else
-		{
-			for(auto &SwitchBuf : SwitchOverBufferChain)
-			{
-				if (SwitchBuf.Buffer.IsValid())
-				{
-					OutInfo.DataSize += SwitchBuf.Buffer->AllocatedSize();
-					OutInfo.Duration += SwitchBuf.Buffer->GetPlayableDuration();
-				}
-			}
-		}
+		return bLastPushWasBlocked;
 	}
 
-	bool FMultiTrackAccessUnitBuffer::WasLastPushBlocked()
+	bool FMultiTrackAccessUnitBuffer::HasPendingTrackSwitch()
 	{
-		return bLastPushWasBlocked;
+		FScopeLock lock(&AccessLock);
+		return PendingBufferSwitch.IsSet();
 	}
 
 
diff --git a/Engine/Plugins/Media/ElectraPlayer/Source/ElectraPlayerRuntime/Private/Runtime/StreamAccessUnitBuffer.h b/Engine/Plugins/Media/ElectraPlayer/Source/ElectraPlayerRuntime/Private/Runtime/StreamAccessUnitBuffer.h
--- a/Engine/Plugins/Media/ElectraPlayer/Source/ElectraPlayerRuntime/Private/Runtime/StreamAccessUnitBuffer.h
+++ b/Engine/Plugins/Media/ElectraPlayer/Source/ElectraPlayerRuntime/Private/Runtime/StreamAccessUnitBuffer.h
@@ -44,6 +44,9 @@ namespace Electra
 
 		// Internal hard index, used for multiplexed streams.
 		int32		HardIndex = -1;
+
+		// To which playback sequence this belongs.
+		uint32		PlaybackSequenceID = 0;
 	};
 
 	struct FAccessUnit
@@ -641,11 +644,10 @@ namespace Electra
 	class FMultiTrackAccessUnitBuffer : public TSharedFromThis<FMultiTrackAccessUnitBuffer, ESPMode::ThreadSafe>
 	{
 	public:
-		FMultiTrackAccessUnitBuffer();
+		FMultiTrackAccessUnitBuffer(EStreamType InForType);
 		~FMultiTrackAccessUnitBuffer();
 		void SetParallelTrackMode();
-		void SelectTrackWhenAvailable(TSharedPtrTS<FBufferSourceInfo> InBufferSourceInfo);
-		void AddUpcomingBuffer(TSharedPtrTS<FBufferSourceInfo> InBufferSourceInfo);
+		void SelectTrackWhenAvailable(uint32 PlaybackSequenceID, TSharedPtrTS<FBufferSourceInfo> InBufferSourceInfo);
 		bool Push(FAccessUnit*& AU, const FAccessUnitBuffer::FConfiguration* BufferConfiguration, const FAccessUnitBuffer::FExternalBufferInfo* InCurrentTotalBufferUtilization);
 		void PushEndOfDataFor(TSharedPtrTS<const FBufferSourceInfo> InStreamSourceInfo);
 		void PushEndOfDataAll();
@@ -653,8 +655,6 @@ namespace Electra
 		void GetStats(FAccessUnitBufferInfo& OutStats);
 		FTimeValue GetLastPoppedDTS();
 		FTimeValue GetLastPoppedPTS();
-		TSharedPtrTS<FBufferSourceInfo> GetActiveOutputBufferInfo() const
-		{ return ActiveOutputBufferInfo; }
 
 		// Helper class to lock the AU buffer
 		class FScopedLock
@@ -680,40 +680,55 @@ namespace Electra
 		bool PeekAndAddRef(FAccessUnit*& OutAU);
 		bool Pop(FAccessUnit*& OutAU);
 		void PopDiscardUntil(FTimeValue UntilTime);
-		FTimeValue PrepareForDecodeStartingAt(FTimeValue DecodeStartTime);
 		bool IsEODFlagSet();
 		int32 Num();
 		bool WasLastPushBlocked();
+		bool HasPendingTrackSwitch();
 
 	private:
 
-		struct FQueuedBuffer
+		struct FSwitchToBuffer
 		{
-			TSharedPtrTS<FBufferSourceInfo> Info;
+			void Reset()
+			{
+				BufferInfo.Reset();
+			}
+			bool IsSet() const
+			{
+				return BufferInfo.IsValid();
+			}
+			TSharedPtrTS<FBufferSourceInfo> BufferInfo;
+		};
+
+		struct FBufferByInfoType
+		{
+			TSharedPtrTS<const FBufferSourceInfo> Info;
 			TSharedPtrTS<FAccessUnitBuffer> Buffer;
 		};
 
-		void PurgeAll();
 		void Clear();
 
 		TSharedPtrTS<FAccessUnitBuffer> CreateNewBuffer();
+		TSharedPtrTS<FAccessUnitBuffer> FindOrCreateBufferFor(TSharedPtrTS<const FBufferSourceInfo>& OutBufferSourceInfo, const TSharedPtrTS<const FBufferSourceInfo>& InBufferInfo, bool bCreateIfNotExist);
+		void ActivateInitialBuffer();
+		void HandlePendingSwitch();
+		void RemoveOutdatedBuffers();
+
 		TSharedPtrTS<FAccessUnitBuffer> GetSelectedTrackBuffer();
-		void ActivateBuffer(bool bIsPushing);
-		void ChangeOver();
-		void RemoveUnusedBuffers();
-		void GetEnqueuedBufferInfo(FAccessUnitBuffer::FExternalBufferInfo& OutInfo, bool bForSwitchOverChain);
 
 		FCriticalSection								AccessLock;
-		TMap<FString, TSharedPtrTS<FAccessUnitBuffer>>	TrackBuffers;						//!< Map of track buffers. One per track ID.
-		TArray<FQueuedBuffer>							UpcomingBufferChain;
-		TArray<FQueuedBuffer>							SwitchOverBufferChain;
-		TSharedPtrTS<FAccessUnitBuffer>					EmptyBuffer;						//!< An empty buffer
-		TSharedPtrTS<FBufferSourceInfo>					ActiveOutputBufferInfo;
-		TSharedPtrTS<FBufferSourceInfo>					LastPoppedBufferInfo;
+		EStreamType										Type;
+		TArray<FBufferByInfoType>						BufferList;
+		FSwitchToBuffer									PendingBufferSwitch;
+		TSharedPtrTS<FAccessUnitBuffer>					EmptyBuffer;
+		TSharedPtrTS<FAccessUnitBuffer>					ActiveBuffer;
+		TSharedPtrTS<const FBufferSourceInfo>			ActiveOutputBufferInfo;
+		TSharedPtrTS<const FBufferSourceInfo>			LastPoppedBufferInfo;
 		FTimeValue										LastPoppedDTS;
 		FTimeValue										LastPoppedPTS;
 		bool											bEndOfData;
 		bool											bLastPushWasBlocked;
+		bool											bPopAsDummyUntilSyncFrame;
 		bool											bIsParallelTrackMode;
 	};
 
diff --git a/Engine/Plugins/Media/ElectraPlayer/Source/ElectraPlayerRuntime/Private/Runtime/Utilities/UtilsMPEG.h b/Engine/Plugins/Media/ElectraPlayer/Source/ElectraPlayerRuntime/Private/Runtime/Utilities/UtilsMPEG.h
--- a/Engine/Plugins/Media/ElectraPlayer/Source/ElectraPlayerRuntime/Private/Runtime/Utilities/UtilsMPEG.h
+++ b/Engine/Plugins/Media/ElectraPlayer/Source/ElectraPlayerRuntime/Private/Runtime/Utilities/UtilsMPEG.h
@@ -41,51 +41,18 @@ namespace Electra
 			enum class FObjectTypeID
 			{
 				Unknown = 0,
-				System_V1 = 1,
-				System_V2 = 2,
-				MPEG4_Video = 32,
-				MPEG4_AVC_SPS = 33,
-				MPEG4_AVC_PPS = 34,
-				MPEG4_Audio = 64,
-				MPEG2_SimpleVideo = 96,
-				MPEG2_MainVideo = 97,
-				MPEG2_SNRVideo = 98,
-				MPEG2_SpatialVideo = 99,
-				MPEG2_HighVideo = 100,
-				MPEG2_422Video = 101,
-				MPEG4_ADTS_Main = 102,
-				MPEG4_ADTS_LC = 103,
-				MPEG4_ADTS_SSR = 104,
-				MPEG2_ADTS = 105,
-				MPEG1_Video = 106,
-				MPEG1_ADTS = 107,
-				JPEG_Video = 108,
-				Private_Audio = 192,
-				Private_Video = 208,
-				PCM16LE_Audio = 224,
-				Vorbis_Audio = 225,
-				AC3_Audio = 226,
-				ALaw_Audio = 227,
-				MuLaw_Audio = 228,
-				G723_ADPCM_Audio = 229,
-				PCM16BE = 230,
-				YV12_Video = 240,
-				H264_Video = 241,
-				H263_Video = 242,
-				H261_Video = 243,
+				Text_Stream = 8,
+				MPEG4_Video = 0x20,
+				H264 = 0x21,
+				H264_ParameterSets = 0x22,
+				H265 = 0x23,
+				MPEG4_Audio = 0x40
 			};
 			enum class FStreamType
 			{
 				Unknown = 0,
-				ObjectDescription = 1,
-				ClockReference = 2,
-				SceneDescription = 4,  	 //!< aka. visual description
-				Audio = 5,
-				MPEG7 = 6,
-				IPMP = 7,
-				OCI = 8,
-				MPEG_Java = 9,
-				User_Private = 32
+				VisualStream = 4,
+				AudioStream = 5,
 			};
 
 			FObjectTypeID GetObjectTypeID() const
diff --git a/Engine/Plugins/Media/ElectraPlayer/Source/ElectraPlayerRuntime/Public/IElectraPlayerInterface.h b/Engine/Plugins/Media/ElectraPlayer/Source/ElectraPlayerRuntime/Public/IElectraPlayerInterface.h
--- a/Engine/Plugins/Media/ElectraPlayer/Source/ElectraPlayerRuntime/Public/IElectraPlayerInterface.h
+++ b/Engine/Plugins/Media/ElectraPlayer/Source/ElectraPlayerRuntime/Public/IElectraPlayerInterface.h
@@ -118,6 +118,7 @@ public:
 	struct FPlaystartOptions
 	{
 		TOptional<FTimespan>		TimeOffset;
+		FStreamSelectionAttributes	InitialVideoTrackAttributes;
 		FStreamSelectionAttributes	InitialAudioTrackAttributes;
 		FStreamSelectionAttributes	InitialSubtitleTrackAttributes;
 		TOptional<int32>			MaxVerticalStreamResolution;
diff --git a/Engine/Source/Runtime/Media/Public/IMediaPlayer.h b/Engine/Source/Runtime/Media/Public/IMediaPlayer.h
--- a/Engine/Source/Runtime/Media/Public/IMediaPlayer.h
+++ b/Engine/Source/Runtime/Media/Public/IMediaPlayer.h
@@ -305,12 +305,14 @@ public:
 		// Override in child class if needed.
 	}
 
-	enum class EFeatureFlag {
+	enum class EFeatureFlag
+	{
 		AllowShutdownOnClose = 0,		//!< Allow player to be shutdown right after 'close' event is received from it
 		UsePlaybackTimingV2,			//!< Use v2 playback timing and AV sync
 		UseRealtimeWithVideoOnly,		//!< Use realtime rather then game deltatime to control video playback if no audio is present
 		AlwaysPullNewestVideoFrame,		//!< Mediaframework will not gate video frame output with its own timing, but assumes "ASAP" as output time for every sample
 		PlayerUsesInternalFlushOnSeek,	//!< The player implements an internal flush logic on seeks and Mediaframework will not issue an explicit Flush() call to it on seeks
+		IsTrackSwitchSeamless,			//!< If track switching is seamless then a flush of sinks is not necessary.
 	};
 	
 	virtual bool GetPlayerFeatureFlag(EFeatureFlag /*flag*/) const
diff --git a/Engine/Source/Runtime/MediaUtils/Private/MediaPlayerFacade.cpp b/Engine/Source/Runtime/MediaUtils/Private/MediaPlayerFacade.cpp
--- a/Engine/Source/Runtime/MediaUtils/Private/MediaPlayerFacade.cpp
+++ b/Engine/Source/Runtime/MediaUtils/Private/MediaPlayerFacade.cpp
@@ -1133,8 +1133,10 @@ bool FMediaPlayerFacade::SelectTrack(EMediaTrackType TrackType, int32 TrackIndex
 		return false;
 	}
 
-	Flush();
-
+	if (!Player->GetPlayerFeatureFlag(IMediaPlayer::EFeatureFlag::IsTrackSwitchSeamless))
+	{
+		Flush();
+	}
 	return true;
 }
 
@@ -2044,8 +2046,8 @@ void FMediaPlayerFacade::TickFetch(FTimespan DeltaTime, FTimespan Timecode)
 
 		while (true)
 		{
-			ProcessCaptionSamplesV1(Samples, TimeRange);
-			ProcessSubtitleSamplesV1(Samples, TimeRange);
+			ProcessCaptionSamples(Samples, TimeRange);
+			ProcessSubtitleSamples(Samples, TimeRange);
 			ProcessVideoSamplesV1(Samples, TimeRange);
 
 			if (!BlockOnFetch())
@@ -2150,6 +2152,11 @@ void FMediaPlayerFacade::TickTickable()
 
 	ProcessAudioSamples(Samples, AudioTimeRange);
 	ProcessMetadataSamples(Samples, MetadataTimeRange);
+	if (bUseV2Timing)
+	{
+		ProcessCaptionSamples(Samples, MetadataTimeRange);
+		ProcessSubtitleSamples(Samples, MetadataTimeRange);
+	}
 
 	SET_DWORD_STAT(STAT_MediaUtils_FacadeNumAudioSamples, Samples.NumAudio());
 }
@@ -2478,7 +2485,7 @@ void FMediaPlayerFacade::ProcessAudioSamples(IMediaSamples& Samples, TRange<FTim
 }
 
 
-void FMediaPlayerFacade::ProcessCaptionSamplesV1(IMediaSamples& Samples, TRange<FTimespan> TimeRange)
+void FMediaPlayerFacade::ProcessCaptionSamples(IMediaSamples& Samples, TRange<FTimespan> TimeRange)
 {
 	TSharedPtr<IMediaOverlaySample, ESPMode::ThreadSafe> Sample;
 
@@ -2510,12 +2517,13 @@ void FMediaPlayerFacade::ProcessMetadataSamples(IMediaSamples& Samples, TRange<F
 }
 
 
-void FMediaPlayerFacade::ProcessSubtitleSamplesV1(IMediaSamples& Samples, TRange<FTimespan> TimeRange)
+void FMediaPlayerFacade::ProcessSubtitleSamples(IMediaSamples& Samples, TRange<FTimespan> TimeRange)
 {
 	TSharedPtr<IMediaOverlaySample, ESPMode::ThreadSafe> Sample;
 
 	while (Samples.FetchSubtitle(TimeRange, Sample))
 	{
+		//UE_LOG(LogMediaUtils, Display, TEXT("Subtitle @%.3f: %s"), Sample->GetTime().Time.GetTotalSeconds(), *Sample->GetText().ToString());
 		if (Sample.IsValid() && !SubtitleSampleSinks.Enqueue(Sample.ToSharedRef(), FMediaPlayerQueueDepths::MaxSubtitleSinkDepth))
 		{
 #if MEDIAPLAYERFACADE_TRACE_SINKOVERFLOWS
diff --git a/Engine/Source/Runtime/MediaUtils/Public/MediaPlayerFacade.h b/Engine/Source/Runtime/MediaUtils/Public/MediaPlayerFacade.h
--- a/Engine/Source/Runtime/MediaUtils/Public/MediaPlayerFacade.h
+++ b/Engine/Source/Runtime/MediaUtils/Public/MediaPlayerFacade.h
@@ -782,10 +782,10 @@ protected:
 	void ProcessMetadataSamples(IMediaSamples& Samples, TRange<FTimespan> TimeRange);
 
 	/** Fetch audio samples from the player and forward them to the registered sinks. */
-	void ProcessCaptionSamplesV1(IMediaSamples& Samples, TRange<FTimespan> TimeRange);
+	void ProcessCaptionSamples(IMediaSamples& Samples, TRange<FTimespan> TimeRange);
 
 	/** Fetch subtitle samples from the player and forward them to the registered sinks. */
-	void ProcessSubtitleSamplesV1(IMediaSamples& Samples, TRange<FTimespan> TimeRange);
+	void ProcessSubtitleSamples(IMediaSamples& Samples, TRange<FTimespan> TimeRange);
 
 	/** Fetch video samples from the player and forward them to the registered sinks. */
 	void ProcessVideoSamplesV1(IMediaSamples& Samples, TRange<FTimespan> TimeRange);
