ElectraPlayer: Removal of ringbuffer in HTTP transfers; added better error reporting for failed mp4 downloads

#rnx

diff --git a/Engine/Plugins/Media/ElectraPlayer/Source/ElectraPlayerRuntime/Private/Runtime/Demuxer/ParserISO14496-12.h b/Engine/Plugins/Media/ElectraPlayer/Source/ElectraPlayerRuntime/Private/Runtime/Demuxer/ParserISO14496-12.h
--- a/Engine/Plugins/Media/ElectraPlayer/Source/ElectraPlayerRuntime/Private/Runtime/Demuxer/ParserISO14496-12.h
+++ b/Engine/Plugins/Media/ElectraPlayer/Source/ElectraPlayerRuntime/Private/Runtime/Demuxer/ParserISO14496-12.h
@@ -93,6 +93,7 @@ namespace Electra
 		static const FBoxType BoxType_prft = MAKE_BOX_ATOM('p', 'r', 'f', 't');
 		static const FBoxType BoxType_moof = MAKE_BOX_ATOM('m', 'o', 'o', 'f');
 		static const FBoxType BoxType_mdat = MAKE_BOX_ATOM('m', 'd', 'a', 't');
+		static const FBoxType BoxType_free = MAKE_BOX_ATOM('f', 'r', 'e', 'e');
 
 
 
diff --git a/Engine/Plugins/Media/ElectraPlayer/Source/ElectraPlayerRuntime/Private/Runtime/HTTP/HTTPManager.cpp b/Engine/Plugins/Media/ElectraPlayer/Source/ElectraPlayerRuntime/Private/Runtime/HTTP/HTTPManager.cpp
--- a/Engine/Plugins/Media/ElectraPlayer/Source/ElectraPlayerRuntime/Private/Runtime/HTTP/HTTPManager.cpp
+++ b/Engine/Plugins/Media/ElectraPlayer/Source/ElectraPlayerRuntime/Private/Runtime/HTTP/HTTPManager.cpp
@@ -572,16 +572,6 @@ namespace Electra
 		// Response not yet received.
 		Handle->bResponseReceived = false;
 
-		// Do a quick check that if the request is reading into a ring buffer that there is also a sub-range request size set up.
-		TSharedPtrTS<FReceiveBuffer> ReceiveBuffer = Request->ReceiveBuffer.Pin();
-		// Give a warning to the console about this circumstance if the original request is not ranged.
-		// If it is a range request that goes up to the end of the file we cannot say with certainty if this is a problem
-		// since we do not know how large the response is going to be. It could be a few bytes, it could be gigabytes...
-		if (ReceiveBuffer.IsValid() && ReceiveBuffer->bEnableRingbuffer && !Handle->ActiveResponse.Range.IsSet())
-		{
-			UE_LOG(LogElectraHTTPManager, Warning, TEXT("Receive buffer set to ring buffer mode but no sub range request size specified!"));
-		}
-
 		// This could be for the next sub-range request that we also want to add to the cache!
 		Handle->ActiveResponse.bWasAddedToCache = false;
 
@@ -1261,27 +1251,14 @@ namespace Electra
 						TSharedPtrTS<FReceiveBuffer> ReceiveBuffer = Request->ReceiveBuffer.Pin();
 						if (ReceiveBuffer.IsValid())
 						{
-							int64 BufferPushableSize = 0;
-							bool bBufferUsable = false;
-							// Is it to be used as a linear buffer?
-							if (!ReceiveBuffer->bEnableRingbuffer)
+							int64 RequiredBufferSize = ci.ContentLength > 0 ? ci.ContentLength : 0;
+							int64 BufferSizeAfterPush = ReceiveBuffer->Buffer.Num() + NumDataAvailable;
+							if (BufferSizeAfterPush > RequiredBufferSize)
 							{
-								int64 RequiredBufferSize = ci.ContentLength > 0 ? ci.ContentLength : 0;
-								int64 BufferSizeAfterPush = ReceiveBuffer->Buffer.Num() + NumDataAvailable;
-								if (BufferSizeAfterPush > RequiredBufferSize)
-								{
-									RequiredBufferSize = BufferSizeAfterPush;
-								}
-								bBufferUsable = ReceiveBuffer->Buffer.EnlargeTo(RequiredBufferSize);
-								BufferPushableSize = bBufferUsable ? NumDataAvailable : 0;
-							}
-							else
-							{
-								// The ring buffer must have been set up to a size the client wants it to have.
-								bBufferUsable = ReceiveBuffer->Buffer.Capacity() > 0;
-								check(bBufferUsable);
-								BufferPushableSize = bBufferUsable ? ReceiveBuffer->Buffer.Avail() : 0;
+								RequiredBufferSize = BufferSizeAfterPush;
 							}
+							bool bBufferUsable = ReceiveBuffer->Buffer.EnlargeTo(RequiredBufferSize);
+							int64 BufferPushableSize = bBufferUsable ? NumDataAvailable : 0;
 							if (bBufferUsable)
 							{
 								if (BufferPushableSize)
@@ -1290,7 +1267,7 @@ namespace Electra
 									int64 NewDataSize = 0;
 									Response->GetResponseData().LockBuffer(NewDataPtr, NewDataSize);
 									int64 NumToCopy = BufferPushableSize < NewDataSize ? BufferPushableSize : NewDataSize;
-									bBufferUsable = ReceiveBuffer->Buffer.PushData(NewDataPtr, NumToCopy, ReceiveBuffer->bEnableRingbuffer);
+									bBufferUsable = ReceiveBuffer->Buffer.PushData(NewDataPtr, NumToCopy);
 									Response->GetResponseData().UnlockBuffer(bBufferUsable ? NumToCopy : 0);
 									if (bBufferUsable)
 									{
@@ -1477,28 +1454,22 @@ namespace Electra
 
 	int32 FElectraHttpManager::FFileStream::Read(TSharedPtrTS<FReceiveBuffer> RcvBuffer, TSharedPtrTS<FRequest> Request)
 	{
-		// How much room do we have in the buffer?
-		int32 blkAv1, blkAv2;
-		uint8* blkData1;
-		uint8* blkData2;
-		int32 MaxAvail = RcvBuffer->Buffer.PushBlockOpen(blkData1, blkAv1, blkData2, blkAv2);
-		// How much can we read from the file?
-		int32 NumToRead = FileSizeToGo < MaxAvail ? FileSizeToGo : MaxAvail;
+		int64 NumToRead = FileSizeToGo;
 		if (NumToRead)
 		{
-			if (NumToRead <= blkAv1)
+			void* Dst = (void*) RcvBuffer->Buffer.GetLinearWriteData(NumToRead);
+			if (Dst)
 			{
-				Archive->Serialize(blkData1, NumToRead);
+				Archive->Serialize(Dst, NumToRead);
+				RcvBuffer->Buffer.AppendedNewData(NumToRead);
+				Request->ConnectionInfo.BytesReadSoFar += NumToRead;
+				FileSizeToGo -= NumToRead;
 			}
 			else
 			{
-				Archive->Serialize(blkData1, blkAv1);
-				Archive->Serialize(blkData2, NumToRead - blkAv1);
+				NumToRead = 0;
 			}
-			Request->ConnectionInfo.BytesReadSoFar += NumToRead;
 		}
-		RcvBuffer->Buffer.PushBlockClose(NumToRead);
-		FileSizeToGo -= NumToRead;
 		return NumToRead;
 	}
 
@@ -1598,26 +1569,16 @@ namespace Electra
 
 	int32 FElectraHttpManager::FDataUrl::Read(TSharedPtrTS<FReceiveBuffer> RcvBuffer, TSharedPtrTS<FRequest> Request)
 	{
-		// Not to be used with ring buffers yet.
-		check(!RcvBuffer->bEnableRingbuffer);
-		if (!RcvBuffer->bEnableRingbuffer)
+		if (RcvBuffer->Buffer.EnlargeTo(FileSizeToGo))
 		{
-			if (RcvBuffer->Buffer.EnlargeTo(FileSizeToGo))
+			if (RcvBuffer->Buffer.PushData(Data.GetData() + FileStartOffset, FileSizeToGo))
 			{
-				if (RcvBuffer->Buffer.PushData(Data.GetData() + FileStartOffset, FileSizeToGo))
-				{
-					Request->ConnectionInfo.BytesReadSoFar += FileSizeToGo;
-				}
+				Request->ConnectionInfo.BytesReadSoFar += FileSizeToGo;
 			}
-			int32 NumRead = (int32) FileSizeToGo;
-			FileSizeToGo = 0;
-			return NumRead;
-		}
-		else
-		{
-			// Pretend we read everything so this transfer will end.
-			return FileSizeToGo;
 		}
+		int32 NumRead = (int32) FileSizeToGo;
+		FileSizeToGo = 0;
+		return NumRead;
 	}
 
 
diff --git a/Engine/Plugins/Media/ElectraPlayer/Source/ElectraPlayerRuntime/Private/Runtime/HTTP/HTTPManager.h b/Engine/Plugins/Media/ElectraPlayer/Source/ElectraPlayerRuntime/Private/Runtime/HTTP/HTTPManager.h
--- a/Engine/Plugins/Media/ElectraPlayer/Source/ElectraPlayerRuntime/Private/Runtime/HTTP/HTTPManager.h
+++ b/Engine/Plugins/Media/ElectraPlayer/Source/ElectraPlayerRuntime/Private/Runtime/HTTP/HTTPManager.h
@@ -220,8 +220,7 @@ namespace Electra
 
 		struct FReceiveBuffer
 		{
-			FPODRingbuffer		Buffer;
-			bool				bEnableRingbuffer = false;
+			FWaitableBuffer		Buffer;
 		};
 
 		struct FParams
diff --git a/Engine/Plugins/Media/ElectraPlayer/Source/ElectraPlayerRuntime/Private/Runtime/Player/DASH/StreamReaderFMP4DASH.cpp b/Engine/Plugins/Media/ElectraPlayer/Source/ElectraPlayerRuntime/Private/Runtime/Player/DASH/StreamReaderFMP4DASH.cpp
--- a/Engine/Plugins/Media/ElectraPlayer/Source/ElectraPlayerRuntime/Private/Runtime/Player/DASH/StreamReaderFMP4DASH.cpp
+++ b/Engine/Plugins/Media/ElectraPlayer/Source/ElectraPlayerRuntime/Private/Runtime/Player/DASH/StreamReaderFMP4DASH.cpp
@@ -1684,7 +1684,7 @@ bool FStreamReaderFMP4DASH::FStreamHandler::HasErrored() const
  */
 int64 FStreamReaderFMP4DASH::FStreamHandler::ReadData(void* IntoBuffer, int64 NumBytesToRead)
 {
-	FPODRingbuffer& SourceBuffer = ReadBuffer.ReceiveBuffer->Buffer;
+	FWaitableBuffer& SourceBuffer = ReadBuffer.ReceiveBuffer->Buffer;
 	// Make sure the buffer will have the amount of data we need.
 	while(1)
 	{
@@ -1783,7 +1783,7 @@ int64 FStreamReaderFMP4DASH::FStreamHandler::ReadData(void* IntoBuffer, int64 Nu
  */
 bool FStreamReaderFMP4DASH::FStreamHandler::HasReachedEOF() const
 {
-	const FPODRingbuffer& SourceBuffer = ReadBuffer.ReceiveBuffer->Buffer;
+	const FWaitableBuffer& SourceBuffer = ReadBuffer.ReceiveBuffer->Buffer;
 	return !HasErrored() && SourceBuffer.GetEOD() && (ReadBuffer.ParsePos >= SourceBuffer.Num() || ReadBuffer.ParsePos >= ReadBuffer.MaxParsePos);
 }
 
diff --git a/Engine/Plugins/Media/ElectraPlayer/Source/ElectraPlayerRuntime/Private/Runtime/Player/HLS/PlaylistReaderHLS.cpp b/Engine/Plugins/Media/ElectraPlayer/Source/ElectraPlayerRuntime/Private/Runtime/Player/HLS/PlaylistReaderHLS.cpp
--- a/Engine/Plugins/Media/ElectraPlayer/Source/ElectraPlayerRuntime/Private/Runtime/Player/HLS/PlaylistReaderHLS.cpp
+++ b/Engine/Plugins/Media/ElectraPlayer/Source/ElectraPlayerRuntime/Private/Runtime/Player/HLS/PlaylistReaderHLS.cpp
@@ -236,7 +236,7 @@ private:
 	void InternalHandleOnce();
 
 	void HandleEnqueuedPlaylistDownloads(const FTimeValue& TimeNow);
-	void HandleCompletedPlaylistDownloads(const FTimeValue& TimeNow);
+	void HandleCompletedPlaylistDownloads(FTimeValue& TimeNow);
 	void HandleStaticRequestCompletions(const FTimeValue& TimeNow);
 	void CheckForPlaylistUpdate(const FTimeValue& TimeNow);
 
@@ -766,7 +766,7 @@ void FPlaylistReaderHLS::HandleEnqueuedPlaylistDownloads(const FTimeValue& TimeN
 	EnqueuedPlaylistRequestsLock.Unlock();
 }
 
-void FPlaylistReaderHLS::HandleCompletedPlaylistDownloads(const FTimeValue& TimeNow)
+void FPlaylistReaderHLS::HandleCompletedPlaylistDownloads(FTimeValue& TimeNow)
 {
 	FErrorDetail ParseError;
 	while(CompletedPlaylistRequests.Num())
@@ -783,6 +783,8 @@ void FPlaylistReaderHLS::HandleCompletedPlaylistDownloads(const FTimeValue& Time
 			if (!ConnInfo->StatusInfo.ErrorDetail.IsSet())
 			{
 				ParseError = ParsePlaylist(Request);
+				// Get the current time again which may now be synchronized with the server time.
+				TimeNow = PlayerSessionServices->GetSynchronizedUTCTime()->GetTime();
 			}
 
 			// Was this an initially required playlist?
diff --git a/Engine/Plugins/Media/ElectraPlayer/Source/ElectraPlayerRuntime/Private/Runtime/Player/HLS/StreamReaderHLSfmp4.cpp b/Engine/Plugins/Media/ElectraPlayer/Source/ElectraPlayerRuntime/Private/Runtime/Player/HLS/StreamReaderHLSfmp4.cpp
--- a/Engine/Plugins/Media/ElectraPlayer/Source/ElectraPlayerRuntime/Private/Runtime/Player/HLS/StreamReaderHLSfmp4.cpp
+++ b/Engine/Plugins/Media/ElectraPlayer/Source/ElectraPlayerRuntime/Private/Runtime/Player/HLS/StreamReaderHLSfmp4.cpp
@@ -10,6 +10,7 @@
 
 #define INTERNAL_ERROR_INIT_SEGMENT_DOWNLOAD_ERROR					1
 #define INTERNAL_ERROR_INIT_SEGMENT_PARSE_ERROR						2
+#define INTERNAL_ERROR_INIT_SEGMENT_UNSUPPORTED_FORMAT				3
 #define INTERNAL_ERROR_INIT_SEGMENT_LICENSEKEY_ERROR				10
 
 
@@ -499,6 +500,8 @@ FStreamReaderHLSfmp4::FStreamHandler::ELicenseKeyResult FStreamReaderHLSfmp4::FS
 
 FStreamReaderHLSfmp4::FStreamHandler::EInitSegmentResult FStreamReaderHLSfmp4::FStreamHandler::GetInitSegment(FErrorDetail& OutErrorDetail, TSharedPtrTS<const IParserISO14496_12>& OutMP4InitSegment, const TSharedPtrTS<FStreamSegmentRequestHLSfmp4>& Request)
 {
+	bParsingInitSegment = true;
+	bInvalidMP4 = false;
 	if (Request->InitSegmentInfo.IsValid() && Request->InitSegmentInfo->URI.Len())
 	{
 		TSharedPtrTS<const IParserISO14496_12> MP4InitSegment = Request->InitSegmentCache->GetInitSegmentFor(Request->InitSegmentInfo);
@@ -569,15 +572,22 @@ FStreamReaderHLSfmp4::FStreamHandler::EInitSegmentResult FStreamReaderHLSfmp4::F
 
 				TSharedPtrTS<IParserISO14496_12> InitSegmentParser = IParserISO14496_12::CreateParser();
 				UEMediaError parseError = InitSegmentParser->ParseHeader(this, this, PlayerSessionService, nullptr);
-				if (parseError == UEMEDIA_ERROR_OK || parseError == UEMEDIA_ERROR_END_OF_STREAM)
+				if (!bInvalidMP4)
 				{
-					// Parse the tracks of the init segment. We do this mainly to get to the CSD we might need should we have to insert filler data later.
-					parseError = InitSegmentParser->PrepareTracks(PlayerSessionService, TSharedPtrTS<const IParserISO14496_12>());
-					if (parseError == UEMEDIA_ERROR_OK)
+					if (parseError == UEMEDIA_ERROR_OK || parseError == UEMEDIA_ERROR_END_OF_STREAM)
 					{
-						Request->InitSegmentCache->AddInitSegment(InitSegmentParser, Request->InitSegmentInfo, FTimeValue::GetPositiveInfinity());
-						OutMP4InitSegment = InitSegmentParser;
-						return FStreamReaderHLSfmp4::FStreamHandler::EInitSegmentResult::Ok;
+						// Parse the tracks of the init segment. We do this mainly to get to the CSD we might need should we have to insert filler data later.
+						parseError = InitSegmentParser->PrepareTracks(PlayerSessionService, TSharedPtrTS<const IParserISO14496_12>());
+						if (parseError == UEMEDIA_ERROR_OK)
+						{
+							Request->InitSegmentCache->AddInitSegment(InitSegmentParser, Request->InitSegmentInfo, FTimeValue::GetPositiveInfinity());
+							OutMP4InitSegment = InitSegmentParser;
+							return FStreamReaderHLSfmp4::FStreamHandler::EInitSegmentResult::Ok;
+						}
+						else
+						{
+							return FStreamReaderHLSfmp4::FStreamHandler::EInitSegmentResult::ParseError;
+						}
 					}
 					else
 					{
@@ -586,7 +596,7 @@ FStreamReaderHLSfmp4::FStreamHandler::EInitSegmentResult FStreamReaderHLSfmp4::F
 				}
 				else
 				{
-					return FStreamReaderHLSfmp4::FStreamHandler::EInitSegmentResult::ParseError;
+					return FStreamReaderHLSfmp4::FStreamHandler::EInitSegmentResult::InvalidFormat;
 				}
 			}
 			else
@@ -702,6 +712,8 @@ void FStreamReaderHLSfmp4::FStreamHandler::HandleRequest()
 		}
 	}
 
+	bParsingInitSegment = false;
+	bInvalidMP4 = false;
 	FTimeValue NextExpectedDTS;
 	FTimeValue LastKnownAUDuration;
 	if (!bIsEmptyFillerSegment)
@@ -837,7 +849,7 @@ void FStreamReaderHLSfmp4::FStreamHandler::HandleRequest()
 					MoofInfo.HeaderOffset = GetCurrentOffset();
 
 					UEMediaError parseError = MP4Parser->ParseHeader(this, this, PlayerSessionService, MP4InitSegment.Get());
-					if (parseError == UEMEDIA_ERROR_OK)
+					if (parseError == UEMEDIA_ERROR_OK && !bInvalidMP4)
 					{
 						{
 							SCOPE_CYCLE_COUNTER(STAT_ElectraPlayer_HLS_StreamReader);
@@ -1053,7 +1065,14 @@ void FStreamReaderHLSfmp4::FStreamHandler::HandleRequest()
 						// failed to parse the segment (in general)
 						if (!HasReadBeenAborted())
 						{
-							LogMessage(IInfoLog::ELevel::Error, FString::Printf(TEXT("Failed to download segment \"%s\""), *Request->URL));
+							if (!bInvalidMP4)
+							{
+								LogMessage(IInfoLog::ELevel::Error, FString::Printf(TEXT("Failed to download segment \"%s\""), *Request->URL));
+							}
+							else
+							{
+								LogMessage(IInfoLog::ELevel::Error, FString::Printf(TEXT("Media segment \"%s\" is not a valid mp4"), *Request->URL));
+							}
 							bHasErrored = true;
 						}
 					}
@@ -1071,7 +1090,12 @@ void FStreamReaderHLSfmp4::FStreamHandler::HandleRequest()
 			{
 				// Init segment failed to download or parse.
 				CurrentRequest->ConnectionInfo.StatusInfo.ErrorDetail.SetFacility(Facility::EFacility::HLSFMP4Reader);
-				if (InitSegmentResult == EInitSegmentResult::ParseError)
+				if (InitSegmentResult == EInitSegmentResult::InvalidFormat)
+				{
+					CurrentRequest->ConnectionInfo.StatusInfo.ErrorDetail.SetMessage(TEXT("Init segment is not a valid mp4.")).SetCode(INTERNAL_ERROR_INIT_SEGMENT_UNSUPPORTED_FORMAT);
+					ds.bParseFailure  = true;
+				}
+				else if (InitSegmentResult == EInitSegmentResult::ParseError)
 				{
 					CurrentRequest->ConnectionInfo.StatusInfo.ErrorDetail.SetMessage(TEXT("Init segment parse error")).SetCode(INTERNAL_ERROR_INIT_SEGMENT_PARSE_ERROR);
 					ds.bParseFailure  = true;
@@ -1301,7 +1325,7 @@ bool FStreamReaderHLSfmp4::FStreamHandler::HasErrored() const
  */
 int64 FStreamReaderHLSfmp4::FStreamHandler::ReadData(void* IntoBuffer, int64 NumBytesToRead)
 {
-	FPODRingbuffer& SourceBuffer = ReadBuffer.ReceiveBuffer->Buffer;
+	FWaitableBuffer& SourceBuffer = ReadBuffer.ReceiveBuffer->Buffer;
 	// Make sure the buffer will have the amount of data we need.
 	while(1)
 	{
@@ -1515,7 +1539,7 @@ int64 FStreamReaderHLSfmp4::FStreamHandler::ReadData(void* IntoBuffer, int64 Num
  */
 bool FStreamReaderHLSfmp4::FStreamHandler::HasReachedEOF() const
 {
-	const FPODRingbuffer& SourceBuffer = ReadBuffer.ReceiveBuffer->Buffer;
+	const FWaitableBuffer& SourceBuffer = ReadBuffer.ReceiveBuffer->Buffer;
 	return !HasErrored() && SourceBuffer.GetEOD() && (ReadBuffer.ParsePos >= SourceBuffer.Num() || ReadBuffer.ParsePos >= ReadBuffer.MaxParsePos);
 }
 
@@ -1544,6 +1568,34 @@ int64 FStreamReaderHLSfmp4::FStreamHandler::GetCurrentOffset() const
 
 IParserISO14496_12::IBoxCallback::EParseContinuation FStreamReaderHLSfmp4::FStreamHandler::OnFoundBox(IParserISO14496_12::FBoxType Box, int64 BoxSizeInBytes, int64 FileDataOffset, int64 BoxDataOffset)
 {
+	if (bParsingInitSegment)
+	{
+		// We require the very first box to be an 'ftyp' box.
+		if (FileDataOffset == 0 && Box != IParserISO14496_12::BoxType_ftyp)
+		{
+			bInvalidMP4 = true;
+			return IParserISO14496_12::IBoxCallback::EParseContinuation::Stop;
+		}
+	}
+	else
+	{
+		static const TArray<IParserISO14496_12::FBoxType> PermittedFirstBoxes =
+		{
+			IParserISO14496_12::BoxType_ftyp,
+			IParserISO14496_12::BoxType_styp,
+			IParserISO14496_12::BoxType_sidx,
+			IParserISO14496_12::BoxType_moov,
+			IParserISO14496_12::BoxType_moof,
+			IParserISO14496_12::BoxType_prft,
+			IParserISO14496_12::BoxType_free
+		};
+		if (FileDataOffset == 0 && !PermittedFirstBoxes.Contains(Box))
+		{
+			bInvalidMP4 = true;
+			return IParserISO14496_12::IBoxCallback::EParseContinuation::Stop;
+		}
+	}
+
 	// Check which box is being parsed next.
 	switch(Box)
 	{
diff --git a/Engine/Plugins/Media/ElectraPlayer/Source/ElectraPlayerRuntime/Private/Runtime/Player/HLS/StreamReaderHLSfmp4.h b/Engine/Plugins/Media/ElectraPlayer/Source/ElectraPlayerRuntime/Private/Runtime/Player/HLS/StreamReaderHLSfmp4.h
--- a/Engine/Plugins/Media/ElectraPlayer/Source/ElectraPlayerRuntime/Private/Runtime/Player/HLS/StreamReaderHLSfmp4.h
+++ b/Engine/Plugins/Media/ElectraPlayer/Source/ElectraPlayerRuntime/Private/Runtime/Player/HLS/StreamReaderHLSfmp4.h
@@ -103,17 +103,17 @@ public:
 	FStreamReaderHLSfmp4();
 	virtual ~FStreamReaderHLSfmp4();
 
-	virtual UEMediaError Create(IPlayerSessionServices* PlayerSessionService, const CreateParam& InCreateParam) override;
-	virtual void Close() override;
+	UEMediaError Create(IPlayerSessionServices* PlayerSessionService, const CreateParam& InCreateParam) override;
+	void Close() override;
 
 	//! Adds a request to read from a stream
-	virtual EAddResult AddRequest(uint32 CurrentPlaybackSequenceID, TSharedPtrTS<IStreamSegment> Request) override;
+	EAddResult AddRequest(uint32 CurrentPlaybackSequenceID, TSharedPtrTS<IStreamSegment> Request) override;
 
-	//! Cancels any ongoing requests of the given stream type. Silent cancellation will not notify OnFragmentClose() or OnFragmentReachedEOS(). 
-	virtual void CancelRequest(EStreamType StreamType, bool bSilent) override;
+	//! Cancels any ongoing requests of the given stream type. Silent cancellation will not notify OnFragmentClose() or OnFragmentReachedEOS().
+	void CancelRequest(EStreamType StreamType, bool bSilent) override;
 
 	//! Cancels all pending requests.
-	virtual void CancelRequests() override;
+	void CancelRequests() override;
 
 private:
 
@@ -150,21 +150,21 @@ private:
 			virtual ~FStaticResourceRequest()
 			{ }
 
-			virtual EPlaybackResourceType GetResourceType() const override
+			EPlaybackResourceType GetResourceType() const override
 			{ return Type; }
 
-			virtual FString GetResourceURL() const override
+			FString GetResourceURL() const override
 			{ return URL; }
 
-			virtual void SetPlaybackData(TSharedPtr<TArray<uint8>, ESPMode::ThreadSafe>	PlaybackData) override
+			void SetPlaybackData(TSharedPtr<TArray<uint8>, ESPMode::ThreadSafe>	PlaybackData) override
 			{ Data = PlaybackData; }
 
-			virtual void SignalDataReady() override
+			void SignalDataReady() override
 			{ DoneSignal.Signal(); }
 
 			bool IsDone() const
 			{ return DoneSignal.IsSignaled(); }
-			
+
 			bool WaitDone(int32 WaitMicros)
 			{ return DoneSignal.WaitTimeout(WaitMicros); }
 
@@ -184,6 +184,7 @@ private:
 			AlreadyCached,
 			DownloadError,
 			ParseError,
+			InvalidFormat,
 			LicenseKeyError
 		};
 
@@ -215,6 +216,8 @@ private:
 		FMediaEvent												DownloadCompleteSignal;
 		TSharedPtrTS<IParserISO14496_12>						MP4Parser;
 		int32													NumMOOFBoxesFound = 0;
+		bool													bParsingInitSegment = false;
+		bool													bInvalidMP4 = false;
 
 		TMediaQueueDynamicNoLock<FAccessUnit *>					AccessUnitFIFO;
 		FTimeValue 												DurationSuccessfullyRead;
@@ -240,17 +243,16 @@ private:
 		void HTTPCompletionCallback(const IElectraHttpManager::FRequest* Request);
 		void HTTPUpdateStats(const FTimeValue& CurrentTime, const IElectraHttpManager::FRequest* Request);
 
-
 		bool HasErrored() const;
 
 		// Methods from IParserISO14496_12::IReader
-		virtual int64 ReadData(void* IntoBuffer, int64 NumBytesToRead) override;
-		virtual bool HasReachedEOF() const override;
-		virtual bool HasReadBeenAborted() const override;
-		virtual int64 GetCurrentOffset() const override;
+		int64 ReadData(void* IntoBuffer, int64 NumBytesToRead) override;
+		bool HasReachedEOF() const override;
+		bool HasReadBeenAborted() const override;
+		int64 GetCurrentOffset() const override;
 		// Methods from IParserISO14496_12::IBoxCallback
-		virtual IParserISO14496_12::IBoxCallback::EParseContinuation OnFoundBox(IParserISO14496_12::FBoxType Box, int64 BoxSizeInBytes, int64 FileDataOffset, int64 BoxDataOffset) override;
-		virtual IParserISO14496_12::IBoxCallback::EParseContinuation OnEndOfBox(IParserISO14496_12::FBoxType Box, int64 BoxSizeInBytes, int64 FileDataOffset, int64 BoxDataOffset) override;
+		IParserISO14496_12::IBoxCallback::EParseContinuation OnFoundBox(IParserISO14496_12::FBoxType Box, int64 BoxSizeInBytes, int64 FileDataOffset, int64 BoxDataOffset) override;
+		IParserISO14496_12::IBoxCallback::EParseContinuation OnEndOfBox(IParserISO14496_12::FBoxType Box, int64 BoxSizeInBytes, int64 FileDataOffset, int64 BoxDataOffset) override;
 	};
 
 	// Currently set to use 2 handlers, one for video and one for audio. This could become a pool of n if we need to stream
diff --git a/Engine/Plugins/Media/ElectraPlayer/Source/ElectraPlayerRuntime/Private/Runtime/Player/mp4/PlaylistReaderMP4.cpp b/Engine/Plugins/Media/ElectraPlayer/Source/ElectraPlayerRuntime/Private/Runtime/Player/mp4/PlaylistReaderMP4.cpp
--- a/Engine/Plugins/Media/ElectraPlayer/Source/ElectraPlayerRuntime/Private/Runtime/Player/mp4/PlaylistReaderMP4.cpp
+++ b/Engine/Plugins/Media/ElectraPlayer/Source/ElectraPlayerRuntime/Private/Runtime/Player/mp4/PlaylistReaderMP4.cpp
@@ -15,7 +15,8 @@
 #include "Player/mp4/ManifestMP4.h"
 #include "Player/mp4/OptionKeynamesMP4.h"
 
-#define ERRCODE_MP4_INVALID_FILE							1
+#define ERRCODE_MP4_INVALID_FILE	1
+#define ERRCODE_MP4_DOWNLOAD_ERROR	2
 
 
 DECLARE_CYCLE_STAT(TEXT("FPlaylistReaderMP4_WorkerThread"), STAT_ElectraPlayer_MP4_PlaylistWorker, STATGROUP_ElectraPlayer);
@@ -115,7 +116,7 @@ private:
 	TSharedPtrTS<IElectraHttpManager::FReceiveBuffer> 		ReceiveBuffer;
 	TSharedPtrTS<IElectraHttpManager::FProgressListener>	ProgressListener;
 	HTTP::FConnectionInfo									ConnectionInfo;
-	FPODRingbuffer 											Buffer;
+	FWaitableBuffer 										Buffer;
 	int64													ParsePos = 0;
 	int64													ChunkReadOffset = 0;
 	int64													FileSize = -1;
@@ -324,6 +325,7 @@ void FPlaylistReaderMP4::ReadNextChunk(int64 InFromOffset, int64 ChunkSize)
 		Buffer.SetEOD();
 		return;
 	}
+	Buffer.Reset();
 	ProgressListener = MakeSharedTS<IElectraHttpManager::FProgressListener>();
 	ProgressListener->CompletionDelegate = IElectraHttpManager::FProgressListener::FCompletionDelegate::CreateRaw(this, &FPlaylistReaderMP4::HTTPCompletionCallback);
 	ProgressListener->ProgressDelegate   = IElectraHttpManager::FProgressListener::FProgressDelegate::CreateRaw(this, &FPlaylistReaderMP4::HTTPProgressCallback);
@@ -409,7 +411,15 @@ void FPlaylistReaderMP4::WorkerThread()
 		}
 		else
 		{
-			PostError(FString::Printf(TEXT("Failed to parse mp4 \"%s\" with error %u"), *ConnectionInfo.EffectiveURL, parseError), ERRCODE_MP4_INVALID_FILE, UEMEDIA_ERROR_FORMAT_ERROR);
+			// See if there was a download error
+			if (ConnectionInfo.StatusInfo.ErrorDetail.IsError())
+			{
+				PostError(FString::Printf(TEXT("%s while downloading \"%s\""), *ConnectionInfo.StatusInfo.ErrorDetail.GetMessage(), *ConnectionInfo.EffectiveURL), ERRCODE_MP4_DOWNLOAD_ERROR, UEMEDIA_ERROR_READ_ERROR);
+			}
+			else
+			{
+				PostError(FString::Printf(TEXT("Failed to parse mp4 \"%s\" with error %u"), *ConnectionInfo.EffectiveURL, parseError), ERRCODE_MP4_INVALID_FILE, UEMEDIA_ERROR_FORMAT_ERROR);
+			}
 		}
 	}
 
@@ -435,7 +445,7 @@ void FPlaylistReaderMP4::WorkerThread()
 int64 FPlaylistReaderMP4::ReadData(void* IntoBuffer, int64 NumBytesToRead)
 {
 	uint8* OutputBuffer = (uint8*)IntoBuffer;
-	// Do we have enough data in the ringbuffer to satisfy the read?
+	// Do we have enough data in the buffer to satisfy the read?
 	if (Buffer.Num() >= NumBytesToRead)
 	{
 		// Yes. Get the data and return.
@@ -446,7 +456,7 @@ int64 FPlaylistReaderMP4::ReadData(void* IntoBuffer, int64 NumBytesToRead)
 	}
 	else
 	{
-		// Do not have enough data yet or we want to read more than the ringbuffer can hold.
+		// Do not have enough data yet or we want to read more than the buffer can hold?
 		int32 NumBytesToGo = NumBytesToRead;
 		int64 NextChunkReadOffset = ParsePos;
 		while(NumBytesToGo > 0)
@@ -472,6 +482,7 @@ int64 FPlaylistReaderMP4::ReadData(void* IntoBuffer, int64 NumBytesToRead)
 				// Trigger read of next chunk of data.
 				if (!bChunkReadInProgress)
 				{
+					check(Buffer.IsEmpty());
 					// Is the data to read actually used or is it skipped over?
 					if (OutputBuffer)
 					{
diff --git a/Engine/Plugins/Media/ElectraPlayer/Source/ElectraPlayerRuntime/Private/Runtime/Player/mp4/StreamReaderMP4.cpp b/Engine/Plugins/Media/ElectraPlayer/Source/ElectraPlayerRuntime/Private/Runtime/Player/mp4/StreamReaderMP4.cpp
--- a/Engine/Plugins/Media/ElectraPlayer/Source/ElectraPlayerRuntime/Private/Runtime/Player/mp4/StreamReaderMP4.cpp
+++ b/Engine/Plugins/Media/ElectraPlayer/Source/ElectraPlayerRuntime/Private/Runtime/Player/mp4/StreamReaderMP4.cpp
@@ -394,9 +394,6 @@ void FStreamReaderMP4::HandleRequest()
 
 	ReadBuffer.Reset();
 	ReadBuffer.ReceiveBuffer = MakeSharedTS<IElectraHttpManager::FReceiveBuffer>();
-	// Set the receive buffer to an okay-ish size. Too small and the file I/O may block too often and get too slow.
-	ReadBuffer.ReceiveBuffer->Buffer.Reserve(4 << 20);
-	ReadBuffer.ReceiveBuffer->bEnableRingbuffer = true;
 	ReadBuffer.SetCurrentPos(Request->FileStartOffset);
 
 	const FParamDict& Options = PlayerSessionServices->GetOptions();
@@ -595,7 +592,7 @@ void FStreamReaderMP4::HandleRequest()
 
 					SelectedTrack.bIsFirstInSequence = false;
 
-					int64 nr = ReadBuffer.ReadTo(AccessUnit->AUData, (int32)SampleSize);
+					int64 nr = ReadBuffer.ReadTo(AccessUnit->AUData, SampleSize);
 					if (nr == SampleSize)
 					{
 						SelectedTrack.DurationSuccessfullyRead += AccessUnit->Duration;
@@ -744,60 +741,44 @@ void FStreamReaderMP4::WorkerThread()
 }
 
 
-int32 FStreamReaderMP4::FReadBuffer::ReadTo(void* ToBuffer, int32 NumBytes)
+int32 FStreamReaderMP4::FReadBuffer::ReadTo(void* IntoBuffer, int64 NumBytesToRead)
 {
-	FPODRingbuffer& SourceBuffer = ReceiveBuffer->Buffer;
-
-	uint8* OutputBuffer = (uint8*)ToBuffer;
-	// Do we have enough data in the ringbuffer to satisfy the read?
-	if (SourceBuffer.Num() >= NumBytes)
+	FWaitableBuffer& SourceBuffer = ReceiveBuffer->Buffer;
+	// Make sure the buffer will have the amount of data we need.
+	while(1)
 	{
-		SCOPE_CYCLE_COUNTER(STAT_ElectraPlayer_MP4_StreamReader);
-		CSV_SCOPED_TIMING_STAT(ElectraPlayer, MP4_StreamReader);
-
-		// Yes. Get the data and return.
-		int32 NumGot = SourceBuffer.PopData(OutputBuffer, NumBytes);
-		check(NumGot == NumBytes);
-		CurrentPos += NumBytes;
-		return NumBytes;
-	}
-	else
-	{
-		// Do not have enough data yet or we want to read more than the ringbuffer can hold.
-		int32 NumBytesToGo = NumBytes;
-		while(NumBytesToGo > 0)
+		if (!SourceBuffer.WaitUntilSizeAvailable(ParsePos + NumBytesToRead, 1000 * 100))
 		{
 			if (bHasErrored || SourceBuffer.WasAborted() || bAbort)
 			{
 				return -1;
 			}
-			// EOD?
-			if (SourceBuffer.IsEndOfData())
-			{
-				return 0;
-			}
-
-			// Get whatever amount of data is currently available to free up the buffer for receiving more data.
-			int32 NumGot;
-			{
-				SCOPE_CYCLE_COUNTER(STAT_ElectraPlayer_MP4_StreamReader);
-				CSV_SCOPED_TIMING_STAT(ElectraPlayer, MP4_StreamReader);
-				NumGot = SourceBuffer.PopData(OutputBuffer, NumBytesToGo);
-			}
-			if ((NumBytesToGo -= NumGot) > 0)
+		}
+		else
+		{
+			SCOPE_CYCLE_COUNTER(STAT_ElectraPlayer_MP4_StreamReader);
+			CSV_SCOPED_TIMING_STAT(ElectraPlayer, MP4_StreamReader);
+			SourceBuffer.Lock();
+			if (SourceBuffer.Num() >= ParsePos + NumBytesToRead)
 			{
-				if (OutputBuffer)
+				if (IntoBuffer)
 				{
-					OutputBuffer += NumGot;
+					FMemory::Memcpy(IntoBuffer, SourceBuffer.GetLinearReadData() + ParsePos, NumBytesToRead);
 				}
-				// Wait for data to arrive in the ringbuffer.
-				int32 WaitForBytes = NumBytesToGo > SourceBuffer.Capacity() ? SourceBuffer.Capacity() : NumBytesToGo;
-				SourceBuffer.WaitUntilSizeAvailable(WaitForBytes, 1000 * 100);
+				SourceBuffer.Unlock();
+				ParsePos += NumBytesToRead;
+				CurrentPos += NumBytesToRead;
+				return NumBytesToRead;
+			}
+			else
+			{
+				// Return 0 at EOF and -1 on error.
+				SourceBuffer.Unlock();
+				return bHasErrored ? -1 : 0;
 			}
 		}
-		CurrentPos += NumBytes;
-		return NumBytes;
 	}
+	return -1;
 }
 
 
diff --git a/Engine/Plugins/Media/ElectraPlayer/Source/ElectraPlayerRuntime/Private/Runtime/Player/mp4/StreamReaderMP4.h b/Engine/Plugins/Media/ElectraPlayer/Source/ElectraPlayerRuntime/Private/Runtime/Player/mp4/StreamReaderMP4.h
--- a/Engine/Plugins/Media/ElectraPlayer/Source/ElectraPlayerRuntime/Private/Runtime/Player/mp4/StreamReaderMP4.h
+++ b/Engine/Plugins/Media/ElectraPlayer/Source/ElectraPlayerRuntime/Private/Runtime/Player/mp4/StreamReaderMP4.h
@@ -119,6 +119,7 @@ private:
 		{
 			ReceiveBuffer.Reset();
 			CurrentPos = 0;
+			ParsePos = 0;
 			bAbort = false;
 			bHasErrored = false;
 		}
@@ -143,9 +144,10 @@ private:
 		{
 			bHasErrored = true;
 		}
-		int32 ReadTo(void* ToBuffer, int32 NumBytes);
+		int32 ReadTo(void* ToBuffer, int64 NumBytes);
 		TSharedPtrTS<IElectraHttpManager::FReceiveBuffer>	ReceiveBuffer;
 		int64												CurrentPos;
+		int64												ParsePos;
 		bool												bAbort;
 		bool												bHasErrored;
 	};
diff --git a/Engine/Plugins/Media/ElectraPlayer/Source/ElectraPlayerRuntime/Private/Runtime/StreamDataBuffer.h b/Engine/Plugins/Media/ElectraPlayer/Source/ElectraPlayerRuntime/Private/Runtime/StreamDataBuffer.h
--- a/Engine/Plugins/Media/ElectraPlayer/Source/ElectraPlayerRuntime/Private/Runtime/StreamDataBuffer.h
+++ b/Engine/Plugins/Media/ElectraPlayer/Source/ElectraPlayerRuntime/Private/Runtime/StreamDataBuffer.h
@@ -7,37 +7,26 @@
 
 namespace Electra
 {
-
 	/**
-	 * Byte ringbuffer
+	 * A buffer with waiting capability.
 	**/
-	class FPODRingbuffer : private TMediaNoncopyable<FPODRingbuffer>
+	class FWaitableBuffer
 	{
 	public:
-		FPODRingbuffer(int32 InNumBytes = 0)
-			: Data(nullptr)
-			, DataEnd(nullptr)
-			, WritePos(nullptr)
-			, ReadPos(nullptr)
-			, DataSize(0)
-			, NumIn(0)
-			, WaitingForSize(0)
-			, bEOD(false)
-			, bWasAborted(false)
-		{
-			Allocate(InNumBytes);
-			SizeAvailableSignal.Reset();
-		}
-		~FPODRingbuffer()
+		FWaitableBuffer() = default;
+		FWaitableBuffer(const FWaitableBuffer&) = delete;
+		FWaitableBuffer& operator=(const FWaitableBuffer&) = delete;
+
+		~FWaitableBuffer()
 		{
-			FMediaCriticalSection::ScopedLock Lock(AccessLock);
+			FScopeLock Lock(&AccessLock);
 			Deallocate();
 		}
 
-		//! Allocates buffer of the specified capacity, destroying any previous buffer.
-		bool Reserve(int32 InNumBytes)
+		// Allocates buffer of the specified capacity, destroying any previous buffer.
+		bool Reserve(int64 InNumBytes)
 		{
-			FMediaCriticalSection::ScopedLock Lock(AccessLock);
+			FScopeLock Lock(&AccessLock);
 			Deallocate();
 			if (Allocate(InNumBytes))
 			{
@@ -47,32 +36,23 @@ namespace Electra
 			return false;
 		}
 
-		//! Enlarges the buffer to the new capacity, retaining the current content.
-		bool EnlargeTo(int32 InNewNumBytes)
+		// Enlarges the buffer to the new capacity, retaining the current content.
+		bool EnlargeTo(int64 InNewNumBytes)
 		{
-			FMediaCriticalSection::ScopedLock Lock(AccessLock);
+			FScopeLock Lock(&AccessLock);
 			// Do we need a bigger capacity?
 			if (InNewNumBytes > Capacity())
 			{
-				// Are we currently empty?
-				if (IsEmpty())
-				{
-					return Reserve(InNewNumBytes);
-				}
-				else
-				{
-					return InternalGrowTo(InNewNumBytes);
-				}
+				return IsEmpty() ? Reserve(InNewNumBytes) : InternalGrowTo(InNewNumBytes);
 			}
 			return true;
 		}
 
-		//! Clears the buffer
+		// Clears the buffer
 		void Reset()
 		{
-			FMediaCriticalSection::ScopedLock Lock(AccessLock);
-			NumIn = 0;
-			WritePos = ReadPos = Data;
+			FScopeLock Lock(&AccessLock);
+			WritePos = ReadPos = 0;
 			bEOD = false;
 			bWasAborted = false;
 			// Do not modify the "waiter" members. When waiting for data to arrive it needs to continue
@@ -81,52 +61,57 @@ namespace Electra
 				//SizeAvailableSignal.Signal();
 		}
 
-		//! Returns the ringbuffer capacity
-		int32 Capacity() const
+		// Returns the buffer capacity
+		int64 Capacity() const
 		{
+			FScopeLock Lock(&AccessLock);
 			return DataSize;
 		}
 
-		//! Returns the number of bytes in the buffer (amount that can be popped)
-		int32 Num() const
+		// Returns the number of bytes in the buffer (amount that can be popped)
+		int64 Num() const
 		{
-			return NumIn;
+			FScopeLock Lock(&AccessLock);
+			return WritePos - ReadPos;
 		}
 
-		//! Returns the number of free bytes in the buffer (amount that can be pushed)
-		int32 Avail() const
+		// Returns the number of free bytes in the buffer (amount that can be pushed)
+		int64 Avail() const
 		{
-			return DataSize - Num();
+			FScopeLock Lock(&AccessLock);
+			return DataSize - WritePos;
 		}
 
-		//! Checks if the buffer is empty.
+		// Checks if the buffer is empty.
 		bool IsEmpty() const
 		{
 			return Num() == 0;
 		}
 
-		//! Checks if the buffer is full.
+		// Checks if the buffer is full.
 		bool IsFull() const
 		{
 			return Avail() == 0;
 		}
 
-		//! Checks if the buffer has reached the end-of-data marker (marker is set and no more data is in the buffer).
+		// Checks if the buffer has reached the end-of-data marker (marker is set and no more data is in the buffer).
 		bool IsEndOfData() const
 		{
 			return IsEmpty() && bEOD;
 		}
 
-		//! Checks if the end-of-data flag has been set. There may still be data in the buffer though!
+		// Checks if the end-of-data flag has been set. There may still be data in the buffer though!
 		bool GetEOD() const
 		{
 			return bEOD;
 		}
 
-		//! Waits until the specified number of bytes has arrived in the buffer.
-		//! NOTE: This method is somewhat dangerous in that there is no guarantee the required amount will ever arrive.
-		//!       You must also never wait for more data than is the capcacity of the buffer!
-		bool WaitUntilSizeAvailable(int32 SizeNeeded, int32 TimeoutMicroseconds)
+		/*
+			Waits until the specified number of bytes has arrived in the buffer.
+			Note: This method is somewhat dangerous in that there is no guarantee the required amount will ever arrive.
+			      You must also never wait for more data than is the capcacity of the buffer!
+		*/
+		bool WaitUntilSizeAvailable(int64 SizeNeeded, int32 TimeoutMicroseconds)
 		{
 			// Only wait if not at EOD and more data than presently available is asked for.
 			// Otherwise return enough data to be present even if that is not actually the case.
@@ -161,10 +146,10 @@ namespace Electra
 		}
 
 
-		//! Inserts elements into the buffer. Returns true if successful, false if there is no room.
-		bool PushData(const uint8* InData, int32 NumElements, bool bRingbufferPush=true)
+		// Inserts elements into the buffer. Returns true if successful, false if there is no room.
+		bool PushData(const uint8* InData, int64 NumElements)
 		{
-			FMediaCriticalSection::ScopedLock Lock(AccessLock);
+			FScopeLock Lock(&AccessLock);
 			check(!bEOD);
 			// Zero elements can always be pushed...
 			if (NumElements == 0)
@@ -173,26 +158,9 @@ namespace Electra
 			}
 			if (Avail() >= NumElements)
 			{
-				int32 NumBeforeWrap = DataEnd - WritePos;
-				if (NumElements <= NumBeforeWrap)
-				{
-					CopyData(WritePos, InData, NumElements);
-					if ((WritePos += NumElements) == DataEnd)
-					{
-						if (bRingbufferPush)
-						{
-							WritePos = Data;
-						}
-					}
-				}
-				else
-				{
-					CopyData(WritePos, InData, NumBeforeWrap);
-					CopyData(Data, InData + NumBeforeWrap, NumElements - NumBeforeWrap);
-					WritePos = Data + (NumElements - NumBeforeWrap);
-				}
-				NumIn += NumElements;
-				if (NumIn >= WaitingForSize)
+				CopyData(Data + WritePos, InData, NumElements);
+				WritePos += NumElements;
+				if (WritePos >= WaitingForSize)
 				{
 					SizeAvailableSignal.Signal();
 				}
@@ -204,63 +172,20 @@ namespace Electra
 			}
 		}
 
-
-		//! "Opens" a push-block by returning 2 sets of pointers and block sizes where data to be pushed can be stored directly.
-		int32 PushBlockOpen(uint8*& Block1, int32& Block1Size, uint8*& Block2, int32& Block2Size)
-		{
-			FMediaCriticalSection::ScopedLock Lock(AccessLock);
-			int32 av = Avail();
-			int32 NumBeforeWrap = DataEnd - WritePos;
-			Block1 = WritePos;
-			if (av > NumBeforeWrap)
-			{
-				Block1Size = NumBeforeWrap;
-				Block2 = Data;
-				Block2Size = av - NumBeforeWrap;
-			}
-			else
-			{
-				Block1Size = av;
-				Block2 = nullptr;
-				Block2Size = 0;
-			}
-			return av;
-		}
-
-		//! "Closes" a previously opened push block and advances the write pointer by as many bytes as were actually stored.
-		void PushBlockClose(int32 NumElements)
-		{
-			if (NumElements)
-			{
-				FMediaCriticalSection::ScopedLock Lock(AccessLock);
-				if ((WritePos += NumElements) >= DataEnd)
-				{
-					WritePos -= DataSize;
-				}
-				NumIn += NumElements;
-				if (NumIn >= WaitingForSize)
-				{
-					SizeAvailableSignal.Signal();
-				}
-			}
-		}
-
-
-		//! "Pushes" an end-of-data marker signaling that no further data will be pushed. May be called more than once. Ringbuffer must be Reset() before next use.
+		// "Pushes" an end-of-data marker signaling that no further data will be pushed. May be called more than once. Buffer must be Reset() before next use.
 		void SetEOD()
 		{
-			FMediaCriticalSection::ScopedLock Lock(AccessLock);
+			FScopeLock Lock(&AccessLock);
 			bEOD = true;
 			// Signal that data is present to wake any waiters on WaitForData() even though there may be no data in the buffer anymore.
 			SizeAvailableSignal.Signal();
 		}
 
-		//! "Pops" data from the buffer to a destination. At most the specified number of bytes are popped, or fewer if not as many are available.
-		int32 PopData(uint8* OutData, int32 MaxElementsWanted)
+		// "Pops" data from the buffer to a destination. At most the specified number of bytes are popped, or fewer if not as many are available.
+		int64 PopData(uint8* OutData, int64 MaxElementsWanted)
 		{
-			FMediaCriticalSection::ScopedLock Lock(AccessLock);
-			int32 size = Num();
-			int32 NumBeforeWrap = DataEnd - ReadPos;
+			FScopeLock Lock(&AccessLock);
+			int64 size = Num();
 
 			if (MaxElementsWanted > size)
 			{
@@ -269,23 +194,9 @@ namespace Electra
 			// Copy out or skip over?
 			if (OutData)
 			{
-				// Can copy out without wrap?
-				if (MaxElementsWanted <= NumBeforeWrap)
-				{
-					CopyData(OutData, ReadPos, MaxElementsWanted);
-				}
-				else
-				{
-					CopyData(OutData, ReadPos, NumBeforeWrap);
-					CopyData(OutData + NumBeforeWrap, Data, MaxElementsWanted - NumBeforeWrap);
-				}
-			}
-
-			if ((ReadPos += MaxElementsWanted) >= DataEnd)
-			{
-				ReadPos -= DataSize;
+				CopyData(OutData, Data + ReadPos, MaxElementsWanted);
 			}
-			NumIn -= MaxElementsWanted;
+			ReadPos += MaxElementsWanted;
 			return MaxElementsWanted;
 		}
 
@@ -299,27 +210,50 @@ namespace Electra
 			AccessLock.Unlock();
 		}
 
-		int32 GetLinearReadSize() const
+		int64 GetLinearReadSize() const
 		{
-			FMediaCriticalSection::ScopedLock Lock(AccessLock);
-			int32 NumBeforeWrap = DataEnd - ReadPos;
-			return NumBeforeWrap;
+			FScopeLock Lock(&AccessLock);
+			return WritePos - ReadPos;
 		}
 
-		// NOTE: Must control Lock()/Unlock() externally!
+		// Must control Lock()/Unlock() externally!
 		const uint8* GetLinearReadData() const
 		{
-			return ReadPos;
+			return Data + ReadPos;
 		}
 		uint8* GetLinearReadData()
 		{
-			return ReadPos;
+			return Data + ReadPos;
 		}
 
+		uint8* GetLinearWriteData(int64 InNumBytesToAppend)
+		{
+			FScopeLock Lock(&AccessLock);
+			int64 Av = Avail();
+			if (InNumBytesToAppend > Av)
+			{
+				bool bOk = IsEmpty() ? Allocate(InNumBytesToAppend) : InternalGrowTo(DataSize + InNumBytesToAppend - Av);
+				check(bOk); (void)bOk;
+			}
+			return Data + WritePos;
+		}
+		void AppendedNewData(int64 InNumAppended)
+		{
+			FScopeLock Lock(&AccessLock);
+			check(!bEOD);
+			if (InNumAppended > 0)
+			{
+				WritePos += InNumAppended;
+				if (WritePos >= WaitingForSize)
+				{
+					SizeAvailableSignal.Signal();
+				}
+			}
+		}
 
 		void Abort()
 		{
-			FMediaCriticalSection::ScopedLock Lock(AccessLock);
+			FScopeLock Lock(&AccessLock);
 			bWasAborted = true;
 			SizeAvailableSignal.Signal();
 		}
@@ -330,15 +264,14 @@ namespace Electra
 		}
 
 	protected:
-		bool Allocate(int32 InSize)
+		bool Allocate(int64 InSize)
 		{
 			if (InSize)
 			{
 				DataSize = InSize;
 				Data = (uint8*)FMemory::Malloc(InSize);
-				DataEnd = Data + DataSize;
-				WritePos = Data;
-				ReadPos = Data;
+				WritePos = 0;
+				ReadPos = 0;
 				return Data != nullptr;
 			}
 			return true;
@@ -348,14 +281,12 @@ namespace Electra
 		{
 			FMemory::Free(Data);
 			Data = nullptr;
-			DataEnd = nullptr;
 			DataSize = 0;
-			NumIn = 0;
-			WritePos = nullptr;
-			ReadPos = nullptr;
+			WritePos = 0;
+			ReadPos = 0;
 		}
 
-		bool InternalGrowTo(int32 InNewNumBytes)
+		bool InternalGrowTo(int64 InNewNumBytes)
 		{
 			// Note: The access mutex must already be held here!
 			check(Data && InNewNumBytes);
@@ -363,21 +294,14 @@ namespace Electra
 			uint8* NewData = (uint8*)FMemory::Realloc(Data, InNewNumBytes);
 			if (NewData)
 			{
-				// Get current read & write offsets
-				int32 ReadOffset = ReadPos - Data;
-				int32 WriteOffset = WritePos - Data;
-				// Set the data pointers up in the new buffer.
 				DataSize = InNewNumBytes;
 				Data = NewData;
-				DataEnd = Data + DataSize;
-				WritePos = Data + WriteOffset;
-				ReadPos = Data + ReadOffset;
 				return true;
 			}
 			return false;
 		}
 
-		void CopyData(uint8* CopyTo, const uint8* CopyFrom, int32 NumElements)
+		void CopyData(uint8* CopyTo, const uint8* CopyFrom, int64 NumElements)
 		{
 			if (NumElements && CopyTo && CopyFrom)
 			{
@@ -385,21 +309,25 @@ namespace Electra
 			}
 		}
 
-		FMediaCriticalSection	AccessLock;
-		FMediaEvent				SizeAvailableSignal;	//!< signaled when WaitingForSize data is present
-		uint8* volatile			Data;					//!< Base address of buffer
-		uint8* volatile			DataEnd;				//!< End address of buffer
-		uint8* volatile			WritePos;				//!< Current write position
-		uint8* volatile			ReadPos;				//!< Current read position
-		volatile int32			DataSize;				//!< Maximum number of bytes in the buffer
-		volatile int32			NumIn;					//!< Current number of bytes in the buffer
-		volatile int32			WaitingForSize;			//!< If waiting for a certain number of bytes to become available
-		volatile bool			bEOD;					//!< true when the last packet of data was pushed into the buffer.
-		volatile bool			bWasAborted;
+		mutable FCriticalSection AccessLock;
+		// Signal which gets set when at least `WaitingForSize` amount of data is present
+		FMediaEvent					SizeAvailableSignal;
+		// Buffer address
+		uint8*						Data = nullptr;
+		// Allocated buffer size
+		uint64						DataSize = 0;
+		// Offset into buffer where to add new data
+		uint64						WritePos = 0;
+		// Offset into buffer from where to read the next data.
+		uint64						ReadPos = 0;
+		// Amount of data necessary to be present for `SizeAvailableSignal` to get set.
+		uint64						WaitingForSize = 0;
+		// Flag indicating that no additional data will be added to the buffer.
+		volatile bool				bEOD = false;	
+		// Flag indicating that reading into the buffer has been aborted.
+		volatile bool				bWasAborted = false;
 	};
 
-
-
 } // namespace Electra
 
 
