ElectraPlayer: Making sure dependent modules are loaded in the correct order so that a shutdown will not unload modules that may still be relied upon.

#rnx

diff --git a/Engine/Plugins/Media/ElectraPlayer/Source/ElectraPlayerRuntime/Private/Runtime/PlayerRuntimeGlobal.cpp b/Engine/Plugins/Media/ElectraPlayer/Source/ElectraPlayerRuntime/Private/Runtime/PlayerRuntimeGlobal.cpp
--- a/Engine/Plugins/Media/ElectraPlayer/Source/ElectraPlayerRuntime/Private/Runtime/PlayerRuntimeGlobal.cpp
+++ b/Engine/Plugins/Media/ElectraPlayer/Source/ElectraPlayerRuntime/Private/Runtime/PlayerRuntimeGlobal.cpp
@@ -3,6 +3,7 @@
 #include "PlayerCore.h"
 #include "PlayerRuntimeGlobal.h"
 #include "Misc/CoreDelegates.h"
+#include "Modules/ModuleManager.h"
 
 namespace Electra
 {
@@ -93,6 +94,14 @@ namespace Electra
 			ApplicationResumeDelegate = FCoreDelegates::ApplicationHasEnteredForegroundDelegate.AddStatic(&HandleApplicationHasEnteredForeground);
 		}
 
+		// Load the modules we depend on. They may have been loaded already, but we do it explicitly here to ensure that
+		// they will not be unloaded on shutdown before this module here, otherwise there could be crashes.
+		FModuleManager::Get().LoadModule(TEXT("ElectraBase"));
+		FModuleManager::Get().LoadModule(TEXT("ElectraSamples"));
+		FModuleManager::Get().LoadModule(TEXT("ElectraHTTPStream"));
+		FModuleManager::Get().LoadModule(TEXT("ElectraSubtitles"));
+		FModuleManager::Get().LoadModule(TEXT("ElectraCDM"));
+
 		EnabledAnalyticsEvents = InConfiguration.EnabledAnalyticsEvents;
 		return(true);
 	}
