ElectraPlayer: Locking the modular feature list when querying

#rnx

diff --git a/Engine/Plugins/Media/ElectraPlayer/Source/ElectraPlayerRuntime/Private/Runtime/Decoder/SubtitleDecoder.cpp b/Engine/Plugins/Media/ElectraPlayer/Source/ElectraPlayerRuntime/Private/Runtime/Decoder/SubtitleDecoder.cpp
--- a/Engine/Plugins/Media/ElectraPlayer/Source/ElectraPlayerRuntime/Private/Runtime/Decoder/SubtitleDecoder.cpp
+++ b/Engine/Plugins/Media/ElectraPlayer/Source/ElectraPlayerRuntime/Private/Runtime/Decoder/SubtitleDecoder.cpp
@@ -41,7 +41,9 @@ void FSubtitleDecoderPlugins::CollectSupportedDecoderPlugins()
 	if (!bIsValidList)
 	{
 		// Get the list of all the registered modular features implementing a subtitle decoder.
+		IModularFeatures::Get().LockModularFeatureList();
 		TArray<IElectraSubtitleModularFeature*> PluginImplementations = IModularFeatures::Get().GetModularFeatureImplementations<IElectraSubtitleModularFeature>(IElectraSubtitlesModule::GetModularFeatureName());
+		IModularFeatures::Get().UnlockModularFeatureList();
 		for(auto &Plugin : PluginImplementations)
 		{
 			if (!Plugin)
