ElectraPlayer: Safeguard in case the HTTP stream module failed to initialize

#rnx

diff --git a/Engine/Plugins/Media/ElectraPlayer/Source/ElectraPlayerRuntime/Private/Runtime/HTTP/HTTPManager.cpp b/Engine/Plugins/Media/ElectraPlayer/Source/ElectraPlayerRuntime/Private/Runtime/HTTP/HTTPManager.cpp
--- a/Engine/Plugins/Media/ElectraPlayer/Source/ElectraPlayerRuntime/Private/Runtime/HTTP/HTTPManager.cpp
+++ b/Engine/Plugins/Media/ElectraPlayer/Source/ElectraPlayerRuntime/Private/Runtime/HTTP/HTTPManager.cpp
@@ -341,7 +341,7 @@ namespace Electra
 		void HandleHTTPResponses(const FTimeValue& Now);
 		void RemoveAllRequests();
 
-		void StartHTTPManager();
+		bool StartHTTPManager();
 		void StopHTTPManager();
 		void ProcessHTTPManager();
 
@@ -426,13 +426,38 @@ namespace Electra
 	void FElectraHttpManager::Initialize()
 	{
 		ProgressInterval.SetFromMilliseconds(100);
-		StartHTTPManager();
-		bManagerStarted = true;
+		bManagerStarted = StartHTTPManager();
 	}
 
 	void FElectraHttpManager::AddRequest(TSharedPtrTS<FRequest> Request, bool bAutoRemoveWhenComplete)
 	{
 		FScopeLock lock(&Lock);
+		if (!HTTPStreamHandler.IsValid())
+		{
+			Request->ConnectionInfo.EffectiveURL = Request->Parameters.URL;
+			Request->ConnectionInfo.StatusInfo.ErrorCode = ERRCODE_HTTPMODULE_FAILURE;
+			Request->ConnectionInfo.StatusInfo.ErrorDetail.
+				SetError(UEMEDIA_ERROR_NOT_SUPPORTED).
+				SetFacility(Facility::EFacility::HTTPReader).
+				SetCode(ERRCODE_HTTPMODULE_FAILURE).
+				SetMessage(FString::Printf(TEXT("FElectraHttpManager did not initialize")));
+			Request->ConnectionInfo.RequestEndTime = MEDIAutcTime::Current();
+			Request->ConnectionInfo.bHasFinished = true;
+			TSharedPtrTS<FReceiveBuffer> ReceiveBuffer = Request->ReceiveBuffer.Pin();
+			if (ReceiveBuffer.IsValid())
+			{
+				ReceiveBuffer->Buffer.SetEOD();
+			}
+			TSharedPtrTS<FProgressListener> ProgressListener = Request->ProgressListener.Pin();
+			if (ProgressListener.IsValid())
+			{
+				if (ProgressListener->CompletionDelegate.IsBound())
+				{
+					ProgressListener->CallCompletionDelegate(Request.Get());
+				}
+			}
+			return;
+		}
 		if (!bTerminate)
 		{
 			// Not currently supported. Reserved for future use.
@@ -445,6 +470,10 @@ namespace Electra
 
 	void FElectraHttpManager::RemoveRequest(TSharedPtrTS<FRequest> Request, bool bDoNotWaitForRemoval)
 	{
+		if (!HTTPStreamHandler.IsValid())
+		{
+			return;
+		}
 		if (bDoNotWaitForRemoval)
 		{
 			Request->ReceiveBuffer.Reset();
@@ -519,8 +548,14 @@ namespace Electra
 		Handle->ActiveResponse.NumSubRangeRequest = 0;
 		Handle->ActiveResponse.OriginalRange = Request->Parameters.Range;
 
-		PrepareHTTPHandle(Now, Handle.Get(), Request, true);
-		return Handle.Release();
+		if (PrepareHTTPHandle(Now, Handle.Get(), Request, true))
+		{
+			return Handle.Release();
+		}
+		else
+		{
+			return nullptr;
+		}
 	}
 
 
@@ -624,9 +659,13 @@ namespace Electra
 			}
 
 			Handle->HttpRequest->NotificationDelegate().BindThreadSafeSP(Handle->HttpsRequestCallbackWrapper.ToSharedRef(), &FHTTPCallbackWrapper::ReportRequestNotification);
+			return true;
+		}
+		else
+		{
+			UE_LOG(LogElectraHTTPManager, Error, TEXT("FElectraHttpManager is not available, cannot create request."));
+			return false;
 		}
-
-		return true;
 	}
 
 
@@ -889,7 +928,7 @@ namespace Electra
 	}
 
 
-	void FElectraHttpManager::StartHTTPManager()
+	bool FElectraHttpManager::StartHTTPManager()
 	{
 		LLM_SCOPE(ELLMTag::ElectraPlayer);
 
@@ -898,7 +937,16 @@ namespace Electra
 			HttpOptions.Set(FString(TEXT("proxy")), FVariantValue(FString(ELECTRA_HTTPMANAGER_DEBUG_PROXY_ADDRESS)));
 		#endif
 		HTTPStreamHandler = IElectraHTTPStream::Create(HttpOptions);
-		HTTPStreamHandler->ThreadHandlerDelegate().BindThreadSafeSP(AsShared(), &FElectraHttpManager::ProcessHTTPManager);
+		if (HTTPStreamHandler.IsValid())
+		{
+			HTTPStreamHandler->ThreadHandlerDelegate().BindThreadSafeSP(AsShared(), &FElectraHttpManager::ProcessHTTPManager);
+			return true;
+		}
+		else
+		{
+			UE_LOG(LogElectraHTTPManager, Error, TEXT("Failed to create an FElectraHttpManager instance."));
+			return false;
+		}
 	}
 	
 	void FElectraHttpManager::StopHTTPManager()
diff --git a/Engine/Plugins/Media/ElectraUtil/Source/ElectraHTTPStream/Private/WinHttp/WinHttpElectraHTTPStream.cpp b/Engine/Plugins/Media/ElectraUtil/Source/ElectraHTTPStream/Private/WinHttp/WinHttpElectraHTTPStream.cpp
--- a/Engine/Plugins/Media/ElectraUtil/Source/ElectraHTTPStream/Private/WinHttp/WinHttpElectraHTTPStream.cpp
+++ b/Engine/Plugins/Media/ElectraUtil/Source/ElectraHTTPStream/Private/WinHttp/WinHttpElectraHTTPStream.cpp
@@ -1127,10 +1127,27 @@ bool FElectraHTTPStreamWinHttp::Initialize(const Electra::FParamDict& InOptions)
 #else
 	DWORD SecureProtocols = 0;
 	#ifdef WINHTTP_FLAG_SECURE_PROTOCOL_TLS1_1
-		SecureProtocols |= WINHTTP_FLAG_SECURE_PROTOCOL_TLS1_1;
+		#if PLATFORM_WINDOWS
+			const bool bIsWindows7OrGreater = FPlatformMisc::VerifyWindowsVersion(6, 1);
+			if (bIsWindows7OrGreater)
+			{
+				SecureProtocols |= WINHTTP_FLAG_SECURE_PROTOCOL_TLS1_1;
+			}
+		#else
+			SecureProtocols |= WINHTTP_FLAG_SECURE_PROTOCOL_TLS1_1;
+		#endif
 	#endif
+	
 	#ifdef WINHTTP_FLAG_SECURE_PROTOCOL_TLS1_2
-		SecureProtocols |= WINHTTP_FLAG_SECURE_PROTOCOL_TLS1_2;
+		#if PLATFORM_WINDOWS
+			const bool bIsWindows8Point1OrGreater = FPlatformMisc::VerifyWindowsVersion(6, 3);
+			if (bIsWindows8Point1OrGreater)
+			{
+				SecureProtocols |= WINHTTP_FLAG_SECURE_PROTOCOL_TLS1_2;
+			}
+		#else
+			SecureProtocols |= WINHTTP_FLAG_SECURE_PROTOCOL_TLS1_2;
+		#endif
 	#endif
 /*
 Removed for the time being since this is not supported on all versions of Windows 10 yet.
