ElectraPlayer: Removed an #include that was not required and could set off a cyclic include warning

#rnx

diff --git a/Engine/Plugins/Media/ElectraPlayer/Source/ElectraPlayerRuntime/Private/Runtime/Core/MediaThreads.h b/Engine/Plugins/Media/ElectraPlayer/Source/ElectraPlayerRuntime/Private/Runtime/Core/MediaThreads.h
--- a/Engine/Plugins/Media/ElectraPlayer/Source/ElectraPlayerRuntime/Private/Runtime/Core/MediaThreads.h
+++ b/Engine/Plugins/Media/ElectraPlayer/Source/ElectraPlayerRuntime/Private/Runtime/Core/MediaThreads.h
@@ -2,8 +2,6 @@
 
 #pragma once
 
-#include "PlayerCore.h"
-
 #include "Core/MediaTypes.h"
 #include "Core/MediaMacros.h"
 #include "Core/MediaNoncopyable.h"
