ElectraPlayer: Generating an error when an mp4 contains no playable tracks.

#rnx

diff --git a/Engine/Plugins/Media/ElectraPlayer/Source/ElectraPlayerRuntime/Private/Runtime/Player/mp4/ManifestMP4.cpp b/Engine/Plugins/Media/ElectraPlayer/Source/ElectraPlayerRuntime/Private/Runtime/Player/mp4/ManifestMP4.cpp
--- a/Engine/Plugins/Media/ElectraPlayer/Source/ElectraPlayerRuntime/Private/Runtime/Player/mp4/ManifestMP4.cpp
+++ b/Engine/Plugins/Media/ElectraPlayer/Source/ElectraPlayerRuntime/Private/Runtime/Player/mp4/ManifestMP4.cpp
@@ -16,7 +16,8 @@
 #include "Player/mp4/ManifestMP4.h"
 
 
-#define ERRCODE_MANIFEST_MP4_STARTSEGMENT_NOT_FOUND		1
+#define ERRCODE_MANIFEST_MP4_NO_PLAYABLE_STREAMS		1
+#define ERRCODE_MANIFEST_MP4_STARTSEGMENT_NOT_FOUND		2
 
 
 DECLARE_CYCLE_STAT(TEXT("FPlayPeriodMP4::FindSegment"), STAT_ElectraPlayer_MP4_FindSegment, STATGROUP_ElectraPlayer);
@@ -672,6 +673,16 @@ FErrorDetail FManifestMP4Internal::FTimelineAssetMP4::Build(IPlayerSessionServic
 		}
 	}
 
+	// No playable content?
+	if (VideoAdaptationSets.Num() == 0 && AudioAdaptationSets.Num() == 0)
+	{
+		FErrorDetail err;
+		err.SetFacility(Facility::EFacility::MP4Playlist);
+		err.SetMessage("No playable streams in this mp4");
+		err.SetCode(ERRCODE_MANIFEST_MP4_NO_PLAYABLE_STREAMS);
+		return err;
+	}
+
 // FIXME: fragmented mp4's with a sidx and moof boxes!
 
 	// Hold on to the parsed MOOV box for future reference.
diff --git a/Engine/Plugins/Media/ElectraPlayer/Source/ElectraPlayerRuntime/Private/Runtime/Player/mp4/PlaylistReaderMP4.cpp b/Engine/Plugins/Media/ElectraPlayer/Source/ElectraPlayerRuntime/Private/Runtime/Player/mp4/PlaylistReaderMP4.cpp
--- a/Engine/Plugins/Media/ElectraPlayer/Source/ElectraPlayerRuntime/Private/Runtime/Player/mp4/PlaylistReaderMP4.cpp
+++ b/Engine/Plugins/Media/ElectraPlayer/Source/ElectraPlayerRuntime/Private/Runtime/Player/mp4/PlaylistReaderMP4.cpp
@@ -393,7 +393,7 @@ void FPlaylistReaderMP4::WorkerThread()
 					FURL_RFC3986::GetQueryParams(URLFragmentComponents, URLFragment, false);	// The fragment is already URL escaped, so no need to do it again.
 					Manifest->SetURLFragmentComponents(MoveTemp(URLFragmentComponents));
 
-					FErrorDetail err = Manifest->Build(MP4Parser, MasterPlaylistURL, ConnectionInfo);
+					LastErrorDetail = Manifest->Build(MP4Parser, MasterPlaylistURL, ConnectionInfo);
 
 					// Notify that the "variant playlists" are ready. There are no variants in an mp4, but this is the trigger that the playlists are all set up and are good to go now.
 					PlayerSessionServices->SendMessageToPlayer(IPlaylistReader::PlaylistLoadedMessage::Create(LastErrorDetail, &ConnectionInfo, Playlist::EListType::Variant, Playlist::ELoadType::Initial));
