ElectraPlayer: Added support for Linux

#rnx

diff --git a/Engine/Plugins/Media/ElectraCDM/ElectraCDM.uplugin b/Engine/Plugins/Media/ElectraCDM/ElectraCDM.uplugin
--- a/Engine/Plugins/Media/ElectraCDM/ElectraCDM.uplugin
+++ b/Engine/Plugins/Media/ElectraCDM/ElectraCDM.uplugin
@@ -19,7 +19,7 @@
 			"Name": "ElectraCDM",
 			"Type": "RuntimeNoCommandlet",
 			"LoadingPhase": "PostEngineInit",
-			"PlatformAllowList": [ "Win64", "Mac", "IOS", "Android" ],
+			"PlatformAllowList": [ "Win64", "Mac", "IOS", "Android", "Linux" ],
 			"TargetDenyList": [ "Server" ]
 		},
 	]
diff --git a/Engine/Plugins/Media/ElectraPlayer/ElectraPlayer.uplugin b/Engine/Plugins/Media/ElectraPlayer/ElectraPlayer.uplugin
--- a/Engine/Plugins/Media/ElectraPlayer/ElectraPlayer.uplugin
+++ b/Engine/Plugins/Media/ElectraPlayer/ElectraPlayer.uplugin
@@ -19,28 +19,28 @@
 			"Name": "ElectraPlayerRuntime",
 			"Type": "RuntimeNoCommandlet",
 			"LoadingPhase": "PreLoadingScreen",
-			"PlatformAllowList": [ "Win64", "Mac", "IOS", "Android" ],
+			"PlatformAllowList": [ "Win64", "Mac", "IOS", "Android", "Linux" ],
 			"TargetDenyList": [ "Server" ]
 		},
 		{
 			"Name": "ElectraPlayerPlugin",
 			"Type": "RuntimeNoCommandlet",
 			"LoadingPhase": "PreLoadingScreen",
-			"PlatformAllowList": [ "Win64", "Mac", "IOS", "Android" ],
+			"PlatformAllowList": [ "Win64", "Mac", "IOS", "Android", "Linux" ],
 			"TargetDenyList": [ "Server" ]
 		},
 		{
 			"Name": "ElectraPlayerFactory",
 			"Type": "Editor",
 			"LoadingPhase": "PostEngineInit",
-			"PlatformAllowList": [ "Win64", "Mac" ],
+			"PlatformAllowList": [ "Win64", "Mac", "Linux" ],
 			"TargetDenyList": [ "Server" ]
 		},
 		{
 			"Name": "ElectraPlayerFactory",
 			"Type": "RuntimeNoCommandlet",
 			"LoadingPhase": "PostEngineInit",
-			"PlatformAllowList": [ "Win64", "Mac", "IOS", "Android" ],
+			"PlatformAllowList": [ "Win64", "Mac", "IOS", "Android", "Linux" ],
 			"TargetDenyList": [ "Server" ]
 		}
 	],
diff --git a/Engine/Plugins/Media/ElectraPlayer/Source/ElectraPlayerFactory/ElectraPlayerFactory.Build.cs b/Engine/Plugins/Media/ElectraPlayer/Source/ElectraPlayerFactory/ElectraPlayerFactory.Build.cs
--- a/Engine/Plugins/Media/ElectraPlayer/Source/ElectraPlayerFactory/ElectraPlayerFactory.Build.cs
+++ b/Engine/Plugins/Media/ElectraPlayer/Source/ElectraPlayerFactory/ElectraPlayerFactory.Build.cs
@@ -64,6 +64,7 @@ namespace UnrealBuildTool.Rules
 				|| Target.Platform == UnrealTargetPlatform.Mac
 				|| Target.Platform == UnrealTargetPlatform.IOS
 				|| Target.IsInPlatformGroup(UnrealPlatformGroup.Android)
+				|| Target.IsInPlatformGroup(UnrealPlatformGroup.Unix)
 				;
 		}
     }
diff --git a/Engine/Plugins/Media/ElectraPlayer/Source/ElectraPlayerFactory/Private/ElectraPlayerFactoryModule.cpp b/Engine/Plugins/Media/ElectraPlayer/Source/ElectraPlayerFactory/Private/ElectraPlayerFactoryModule.cpp
--- a/Engine/Plugins/Media/ElectraPlayer/Source/ElectraPlayerFactory/Private/ElectraPlayerFactoryModule.cpp
+++ b/Engine/Plugins/Media/ElectraPlayer/Source/ElectraPlayerFactory/Private/ElectraPlayerFactoryModule.cpp
@@ -130,6 +130,8 @@ public:
 		AddSupportedPlatform(FGuid(0xccf05903, 0x822b47e1, 0xb2236a28, 0xdfd78817));
 		AddSupportedPlatform(FGuid(0x5636fbc1, 0xd2b54f62, 0xac8e7d4f, 0xb184b45a));
 		AddSupportedPlatform(FGuid(0xb596ce6f, 0xd8324a9c, 0x84e9f880, 0x21322535));
+		AddSupportedPlatform(FGuid(0x115de4fe, 0x241b465b, 0x970a872f, 0x3167492a));
+		AddSupportedPlatform(FGuid(0xc0b45a33, 0x9de340c7, 0xbce24c47, 0x15c3babf));
 #endif
 
 		// supported schemes
diff --git a/Engine/Plugins/Media/ElectraPlayer/Source/ElectraPlayerPlugin/Private/ElectraPlayerPlugin.h b/Engine/Plugins/Media/ElectraPlayer/Source/ElectraPlayerPlugin/Private/ElectraPlayerPlugin.h
--- a/Engine/Plugins/Media/ElectraPlayer/Source/ElectraPlayerPlugin/Private/ElectraPlayerPlugin.h
+++ b/Engine/Plugins/Media/ElectraPlayer/Source/ElectraPlayerPlugin/Private/ElectraPlayerPlugin.h
@@ -17,6 +17,8 @@
 #include "MediaSampleQueue.h"
 #include "Templates/SharedPointer.h"
 #include "Logging/LogMacros.h"
+#include "ElectraTextureSample.h"
+#include "ElectraPlayerAudioSample.h"
 
 class IElectraPlayerRuntimeModule;
 class IElectraSafeMediaOptionInterface;
@@ -27,11 +29,6 @@ DECLARE_LOG_CATEGORY_EXTERN(LogElectraPlayerPlugin, Log, All);
 
 //-----------------------------------------------------------------------------
 
-#include "ElectraTextureSample.h"
-#include "ElectraPlayerAudioSample.h"
-
-//-----------------------------------------------------------------------------
-
 class IMediaOptions;
 class FMediaSamples;
 class FArchive;
@@ -203,4 +200,3 @@ private:
 
 	TSharedPtr<IElectraPlayerResourceDelegate,ESPMode::ThreadSafe> PlayerResourceDelegate;
 };
-
diff --git a/Engine/Plugins/Media/ElectraPlayer/Source/ElectraPlayerPlugin/Private/Linux/ElectraPlayerPlugin_Linux.cpp b/Engine/Plugins/Media/ElectraPlayer/Source/ElectraPlayerPlugin/Private/Linux/ElectraPlayerPlugin_Linux.cpp
new file mode 100644
--- /dev/null
+++ b/Engine/Plugins/Media/ElectraPlayer/Source/ElectraPlayerPlugin/Private/Linux/ElectraPlayerPlugin_Linux.cpp
@@ -0,0 +1,14 @@
+// Copyright Epic Games, Inc. All Rights Reserved.
+
+#include "CoreTypes.h"
+
+#include <ElectraPlayerPlugin.h>
+
+IElectraPlayerResourceDelegate* FElectraPlayerPlugin::PlatformCreatePlayerResourceDelegate()
+{
+	return nullptr;
+}
+
+void FElectraPlayerPlugin::PlatformSetupResourceParams(Electra::FParamDict& Params)
+{
+}
diff --git a/Engine/Plugins/Media/ElectraPlayer/Source/ElectraPlayerRuntime/ElectraPlayerRuntime.Build.cs b/Engine/Plugins/Media/ElectraPlayer/Source/ElectraPlayerRuntime/ElectraPlayerRuntime.Build.cs
--- a/Engine/Plugins/Media/ElectraPlayer/Source/ElectraPlayerRuntime/ElectraPlayerRuntime.Build.cs
+++ b/Engine/Plugins/Media/ElectraPlayer/Source/ElectraPlayerRuntime/ElectraPlayerRuntime.Build.cs
@@ -154,9 +154,12 @@ namespace UnrealBuildTool.Rules
 					PublicDefinitions.Add("CURL_ENABLE_NO_TIMEOUTS_OPTION=1");
 				}
 
-				PublicIncludePaths.Add("$(ModuleDir)/Public/Unix");
+				PrivateDefinitions.Add("ELECTRA_PLATFORM_HAS_H265_DECODER=1");
+
+				PublicIncludePaths.Add("$(ModuleDir)/Public/Linux");
+				PrivateIncludePaths.Add("ElectraPlayerRuntime/Private/Runtime/Decoder/Linux");
 
-				PrivateIncludePaths.Add("ElectraPlayerRuntime/Private/Runtime/Decoder/GStreamer");
+				AddEngineThirdPartyPrivateStaticDependencies(Target, "libav");
 			}
 		}
 
diff --git a/Engine/Plugins/Media/ElectraPlayer/Source/ElectraPlayerRuntime/Private/Runtime/Decoder/GStreamer/AudioDecoderAAC_GStreamer.cpp b/Engine/Plugins/Media/ElectraPlayer/Source/ElectraPlayerRuntime/Private/Runtime/Decoder/GStreamer/AudioDecoderAAC_GStreamer.cpp
deleted file mode 100644
diff --git a/Engine/Plugins/Media/ElectraPlayer/Source/ElectraPlayerRuntime/Private/Runtime/Decoder/GStreamer/VideoDecoderH264_GStreamer.cpp b/Engine/Plugins/Media/ElectraPlayer/Source/ElectraPlayerRuntime/Private/Runtime/Decoder/GStreamer/VideoDecoderH264_GStreamer.cpp
deleted file mode 100644
diff --git a/Engine/Plugins/Media/ElectraPlayer/Source/ElectraPlayerRuntime/Private/Runtime/Decoder/GStreamer/VideoDecoderResourceDelegate.h b/Engine/Plugins/Media/ElectraPlayer/Source/ElectraPlayerRuntime/Private/Runtime/Decoder/GStreamer/VideoDecoderResourceDelegate.h
deleted file mode 100644
diff --git a/Engine/Plugins/Media/ElectraPlayer/Source/ElectraPlayerRuntime/Private/Runtime/Decoder/Linux/AudioDecoderAAC_Linux.cpp b/Engine/Plugins/Media/ElectraPlayer/Source/ElectraPlayerRuntime/Private/Runtime/Decoder/Linux/AudioDecoderAAC_Linux.cpp
new file mode 100644
--- /dev/null
+++ b/Engine/Plugins/Media/ElectraPlayer/Source/ElectraPlayerRuntime/Private/Runtime/Decoder/Linux/AudioDecoderAAC_Linux.cpp
@@ -0,0 +1,858 @@
+// Copyright Epic Games, Inc. All Rights Reserved.
+
+#include "CoreMinimal.h"
+
+#include "PlayerCore.h"
+#include "PlayerRuntimeGlobal.h"
+
+#include "StreamAccessUnitBuffer.h"
+#include "Decoder/AudioDecoderAAC.h"
+#include "Renderer/RendererBase.h"
+#include "Player/PlayerSessionServices.h"
+#include "Utilities/Utilities.h"
+#include "Utilities/UtilsMPEGAudio.h"
+#include "Utilities/StringHelpers.h"
+#include "Utilities/AudioChannelMapper.h"
+#include "DecoderErrors_Linux.h"
+#include "HAL/LowLevelMemTracker.h"
+#include "ElectraPlayerPrivate.h"
+
+#include "libav_Decoder_Common.h"
+#include "libav_Decoder_AAC.h"
+
+/***************************************************************************************************************************************************/
+
+DECLARE_CYCLE_STAT(TEXT("FAudioDecoderAACLinuxLibavcodec::Decode()"), STAT_ElectraPlayer_AudioAACDecode, STATGROUP_ElectraPlayer);
+DECLARE_CYCLE_STAT(TEXT("FAudioDecoderAACLinuxLibavcodec::ConvertOutput()"), STAT_ElectraPlayer_AudioAACConvertOutput, STATGROUP_ElectraPlayer);
+
+
+namespace Electra
+{
+
+namespace AACDecoderChannelOutputMappingLibavcodec
+{
+	#define CP FAudioChannelMapper::EChannelPosition
+	static const CP _1[] = { CP::C };
+	static const CP _2[] = { CP::L, CP::R };
+	static const CP _3[] = { CP::L, CP::R, CP::C };
+	static const CP _4[] = { CP::L, CP::R, CP::C, CP::Cs };
+	static const CP _5[] = { CP::L, CP::R, CP::C, CP::Ls, CP::Rs };
+	static const CP _6[] = { CP::L, CP::R, CP::C, CP::LFE, CP::Ls, CP::Rs };
+	static const CP _7[] = { CP::L, CP::R, CP::C, CP::LFE, CP::Ls, CP::Rs, CP::Cs };
+	static const CP _8[] = { CP::L, CP::R, CP::C, CP::LFE, CP::Ls, CP::Rs, CP::Lsr, CP::Rsr };
+	static const CP * const Order[] = { _1, _2, _3, _4, _5, _6, _7, _8 };
+	#undef CP
+};
+
+
+class FAudioDecoderAACLinuxLibavcodec : public IAudioDecoderAAC, public FMediaThread
+{
+public:
+	FAudioDecoderAACLinuxLibavcodec();
+	~FAudioDecoderAACLinuxLibavcodec();
+
+	void SetPlayerSessionServices(IPlayerSessionServices* SessionServices) override;
+	void Open(const FInstanceConfiguration& Config) override;
+	void Close() override;
+	void SetRenderer(TSharedPtr<IMediaRenderer, ESPMode::ThreadSafe> InRenderer) override;
+	void AUdataPushAU(FAccessUnit* AccessUnit) override;
+	void AUdataPushEOD() override;
+	void AUdataClearEOD() override;
+	void AUdataFlushEverything() override;
+	void SetAUInputBufferListener(IAccessUnitBufferListener* Listener) override;
+	void SetReadyBufferListener(IDecoderOutputBufferListener* Listener) override;
+private:
+	struct FDecoderInput
+	{
+		~FDecoderInput()
+		{
+			ReleasePayload();
+		}
+		void ReleasePayload()
+		{
+			FAccessUnit::Release(AccessUnit);
+			AccessUnit = nullptr;
+		}
+
+		FAccessUnit*	AccessUnit = nullptr;
+		bool			bHasBeenPrepared = false;
+		int64			PTS = 0;
+		int64			EndPTS = 0;
+		FTimeValue		AdjustedPTS;
+		FTimeValue		AdjustedDuration;
+		FTimeValue		StartOverlapDuration;
+		FTimeValue		EndOverlapDuration;
+	};
+
+	void StartThread();
+	void StopThread();
+	void WorkerThread();
+
+	bool InternalDecoderCreate();
+	void InternalDecoderDestroy();
+
+	bool CreateDecodedSamplePool();
+	void DestroyDecodedSamplePool();
+
+	void ReturnUnusedOutputBuffer();
+
+	void NotifyReadyBufferListener(bool bHaveOutput);
+
+	void PrepareAU(TSharedPtrTS<FDecoderInput> AU);
+	bool Decode(TSharedPtrTS<FDecoderInput> AU);
+	bool IsDifferentFormat(TSharedPtrTS<FDecoderInput> AU);
+
+	void PostError(int32 ApiReturnValue, const FString& Message, uint16 Code, UEMediaError Error = UEMEDIA_ERROR_OK);
+	void LogMessage(IInfoLog::ELevel Level, const FString& Message);
+
+private:
+	FInstanceConfiguration													Config;
+
+	TAccessUnitQueue<TSharedPtrTS<FDecoderInput>>							NextAccessUnits;
+
+	FMediaEvent																TerminateThreadSignal;
+	FMediaEvent																FlushDecoderSignal;
+	FMediaEvent																DecoderFlushedSignal;
+	bool																	bThreadStarted = false;
+
+	FMediaEvent																ApplicationRunningSignal;
+	FMediaEvent																ApplicationSuspendConfirmedSignal;
+
+	TSharedPtr<IMediaRenderer, ESPMode::ThreadSafe>							Renderer;
+	int32																	MaxDecodeBufferSize = 0;
+
+	FCriticalSection														ListenerMutex;
+	IAccessUnitBufferListener*												InputBufferListener = nullptr;
+	IDecoderOutputBufferListener*											ReadyBufferListener = nullptr;
+
+	IPlayerSessionServices* 												SessionServices = nullptr;
+
+	TSharedPtr<MPEG::FAACDecoderConfigurationRecord, ESPMode::ThreadSafe>	ConfigRecord;
+	TSharedPtrTS<FAccessUnit::CodecData>									CurrentCodecData;
+	bool																	bInDummyDecodeMode = false;
+	bool																	bHaveDiscontinuity = false;
+
+	TSharedPtr<ILibavDecoderAAC, ESPMode::ThreadSafe>						DecoderInstance;
+
+	TArray<TSharedPtrTS<FDecoderInput>>										InDecoderInput;
+
+	IMediaRenderer::IBuffer*												CurrentOutputBuffer = nullptr;
+	FParamDict																BufferAcquireOptions;
+	FParamDict																OutputBufferSampleProperties;
+	int32																	PCMBufferSize = 0;
+	int16*																	PCMBuffer = nullptr;
+	FAudioChannelMapper														ChannelMapper;
+	static const uint8														NumChannelsForConfig[16];
+
+};
+const uint8 FAudioDecoderAACLinuxLibavcodec::NumChannelsForConfig[16] = { 0, 1, 2, 3, 4, 5, 6, 7, 0, 0, 0, 7, 8, 0, 8, 0 };
+
+/***************************************************************************************************************************************************/
+/***************************************************************************************************************************************************/
+/***************************************************************************************************************************************************/
+
+IAudioDecoderAAC::FSystemConfiguration::FSystemConfiguration()
+{ }
+
+IAudioDecoderAAC::FInstanceConfiguration::FInstanceConfiguration()
+{ }
+
+bool IAudioDecoderAAC::Startup(const IAudioDecoderAAC::FSystemConfiguration& InConfig)
+{
+	ILibavDecoder::Startup();
+	return ILibavDecoder::IsLibAvAvailable();
+}
+
+void IAudioDecoderAAC::Shutdown()
+{
+	ILibavDecoder::Shutdown();
+}
+
+bool IAudioDecoderAAC::CanDecodeStream(const FStreamCodecInformation& InCodecInfo)
+{
+	bool bSupported = ILibavDecoderAAC::IsAvailable();
+	if (bSupported)
+	{
+		bSupported = InCodecInfo.GetChannelConfiguration() == 1 ||		// Mono
+					 InCodecInfo.GetChannelConfiguration() == 2 ||		// Stereo
+					 InCodecInfo.GetChannelConfiguration() == 3 ||		// L/C/R
+					 InCodecInfo.GetChannelConfiguration() == 6;		// 5.1
+	}
+	return bSupported;
+}
+
+IAudioDecoderAAC* IAudioDecoderAAC::Create()
+{
+	return new FAudioDecoderAACLinuxLibavcodec;
+}
+
+/***************************************************************************************************************************************************/
+/***************************************************************************************************************************************************/
+/***************************************************************************************************************************************************/
+
+FAudioDecoderAACLinuxLibavcodec::FAudioDecoderAACLinuxLibavcodec()
+	: FMediaThread("ElectraPlayer::AAC decoder")
+{
+}
+
+FAudioDecoderAACLinuxLibavcodec::~FAudioDecoderAACLinuxLibavcodec()
+{
+	Close();
+}
+
+void FAudioDecoderAACLinuxLibavcodec::SetPlayerSessionServices(IPlayerSessionServices* InSessionServices)
+{
+	SessionServices = InSessionServices;
+}
+
+void FAudioDecoderAACLinuxLibavcodec::Open(const FInstanceConfiguration& InConfig)
+{
+	Config = InConfig;
+	StartThread();
+}
+
+void FAudioDecoderAACLinuxLibavcodec::Close()
+{
+	StopThread();
+}
+
+void FAudioDecoderAACLinuxLibavcodec::SetRenderer(TSharedPtr<IMediaRenderer, ESPMode::ThreadSafe> InRenderer)
+{
+	Renderer = InRenderer;
+}
+
+void FAudioDecoderAACLinuxLibavcodec::AUdataPushAU(FAccessUnit* InAccessUnit)
+{
+	InAccessUnit->AddRef();
+
+	TSharedPtrTS<FDecoderInput> NextAU = MakeSharedTS<FDecoderInput>();
+	NextAU->AccessUnit = InAccessUnit;
+	NextAccessUnits.Enqueue(MoveTemp(NextAU));
+}
+
+void FAudioDecoderAACLinuxLibavcodec::AUdataPushEOD()
+{
+	NextAccessUnits.SetEOD();
+}
+
+void FAudioDecoderAACLinuxLibavcodec::AUdataClearEOD()
+{
+	NextAccessUnits.ClearEOD();
+}
+
+void FAudioDecoderAACLinuxLibavcodec::AUdataFlushEverything()
+{
+	FlushDecoderSignal.Signal();
+	DecoderFlushedSignal.WaitAndReset();
+}
+
+void FAudioDecoderAACLinuxLibavcodec::SetAUInputBufferListener(IAccessUnitBufferListener* InListener)
+{
+	FScopeLock lock(&ListenerMutex);
+	InputBufferListener = InListener;
+}
+
+void FAudioDecoderAACLinuxLibavcodec::SetReadyBufferListener(IDecoderOutputBufferListener* InListener)
+{
+	FScopeLock lock(&ListenerMutex);
+	ReadyBufferListener = InListener;
+}
+
+void FAudioDecoderAACLinuxLibavcodec::StartThread()
+{
+	ThreadSetPriority(Config.ThreadConfig.Decoder.Priority);
+	ThreadSetStackSize(Config.ThreadConfig.Decoder.StackSize);
+	ThreadSetCoreAffinity(Config.ThreadConfig.Decoder.CoreAffinity);
+	ThreadStart(FMediaRunnable::FStartDelegate::CreateRaw(this, &FAudioDecoderAACLinuxLibavcodec::WorkerThread));
+	bThreadStarted = true;
+}
+
+void FAudioDecoderAACLinuxLibavcodec::StopThread()
+{
+	if (bThreadStarted)
+	{
+		TerminateThreadSignal.Signal();
+		ThreadWaitDone();
+		bThreadStarted = false;
+	}
+}
+
+bool FAudioDecoderAACLinuxLibavcodec::CreateDecodedSamplePool()
+{
+	check(Renderer);
+	FParamDict poolOpts;
+	uint32 frameSize = sizeof(int16) * 8 * 2048;
+	poolOpts.Set("max_buffer_size", FVariantValue((int64) frameSize));
+	poolOpts.Set("num_buffers", FVariantValue((int64) 8));
+	poolOpts.Set("samples_per_block", FVariantValue((int64) 2048));
+	poolOpts.Set("max_channels", FVariantValue((int64) 8));
+
+	UEMediaError Error = Renderer->CreateBufferPool(poolOpts);
+	check(Error == UEMEDIA_ERROR_OK);
+
+	if (Error != UEMEDIA_ERROR_OK)
+	{
+		PostError(0, "Failed to create sample pool", ERRCODE_INTERNAL_LINUX_COULD_NOT_CREATE_SAMPLE_POOL, Error);
+	}
+
+	MaxDecodeBufferSize = (int32) Renderer->GetBufferPoolProperties().GetValue("max_buffers").GetInt64();
+
+	return Error == UEMEDIA_ERROR_OK;
+}
+
+void FAudioDecoderAACLinuxLibavcodec::DestroyDecodedSamplePool()
+{
+	Renderer->ReleaseBufferPool();
+}
+
+void FAudioDecoderAACLinuxLibavcodec::NotifyReadyBufferListener(bool bHaveOutput)
+{
+	if (ReadyBufferListener)
+	{
+		IDecoderOutputBufferListener::FDecodeReadyStats stats;
+		stats.MaxDecodedElementsReady = MaxDecodeBufferSize;
+		stats.NumElementsInDecoder = CurrentOutputBuffer ? 1 : 0;
+		stats.bOutputStalled = !bHaveOutput;
+		stats.bEODreached = NextAccessUnits.ReachedEOD() && stats.NumDecodedElementsReady == 0 && stats.NumElementsInDecoder == 0;
+		ListenerMutex.Lock();
+		if (ReadyBufferListener)
+		{
+			ReadyBufferListener->DecoderOutputReady(stats);
+		}
+		ListenerMutex.Unlock();
+	}
+}
+
+void FAudioDecoderAACLinuxLibavcodec::PostError(int32 ApiReturnValue, const FString& Message, uint16 Code, UEMediaError Error)
+{
+	check(SessionServices);
+	if (SessionServices)
+	{
+		FErrorDetail err;
+		err.SetError(Error != UEMEDIA_ERROR_OK ? Error : UEMEDIA_ERROR_DETAIL);
+		err.SetFacility(Facility::EFacility::AACDecoder);
+		err.SetCode(Code);
+		err.SetMessage(Message);
+		err.SetPlatformMessage(FString::Printf(TEXT("%d (0x%08x)"), (int32) ApiReturnValue, (int32) ApiReturnValue));
+		SessionServices->PostError(err);
+	}
+}
+
+void FAudioDecoderAACLinuxLibavcodec::LogMessage(IInfoLog::ELevel Level, const FString& Message)
+{
+	if (SessionServices)
+	{
+		SessionServices->PostLog(Facility::EFacility::AACDecoder, Level, Message);
+	}
+}
+
+bool FAudioDecoderAACLinuxLibavcodec::InternalDecoderCreate()
+{
+	InternalDecoderDestroy();
+
+	if (!ConfigRecord.IsValid())
+	{
+		PostError(0, "No CSD to create audio decoder with", ERRCODE_INTERNAL_LINUX_COULD_NOT_CREATE_AUDIO_DECODER);
+		return false;
+	}
+
+	if (ILibavDecoderAAC::IsAvailable())
+	{
+		DecoderInstance = ILibavDecoderAAC::Create(ConfigRecord->GetCodecSpecificData());
+		if (!DecoderInstance.IsValid() || DecoderInstance->GetLastLibraryError())
+		{
+			InternalDecoderDestroy();
+			PostError(-2, "libavcodec failed to open audio decoder", ERRCODE_INTERNAL_LINUX_COULD_NOT_CREATE_AUDIO_DECODER);
+			return false;
+		}
+		return true;
+	}
+	PostError(-1, "libavcodec does not support this audio format", ERRCODE_INTERNAL_LINUX_AUDIO_DECODER_NOT_SUPPORTED);
+	return false;
+}
+
+void FAudioDecoderAACLinuxLibavcodec::InternalDecoderDestroy()
+{
+	DecoderInstance.Reset();
+}
+
+void FAudioDecoderAACLinuxLibavcodec::ReturnUnusedOutputBuffer()
+{
+	if (CurrentOutputBuffer)
+	{
+		OutputBufferSampleProperties.Clear();
+		Renderer->ReturnBuffer(CurrentOutputBuffer, false, OutputBufferSampleProperties);
+		CurrentOutputBuffer = nullptr;
+	}
+}
+
+bool FAudioDecoderAACLinuxLibavcodec::Decode(TSharedPtrTS<FDecoderInput> AU)
+{
+	FTimeValue CurrentPTS;
+	CurrentPTS = AU->AdjustedPTS;
+	check(CurrentPTS.IsValid());
+	if (!CurrentPTS.IsValid())
+	{
+		CurrentPTS = AU->AccessUnit->PTS;
+	}
+
+	bool bDone = false;
+	while(!bDone && !TerminateThreadSignal.IsSignaled() && !FlushDecoderSignal.IsSignaled())
+	{
+		// Need a new output buffer?
+		if (CurrentOutputBuffer == nullptr)
+		{
+			SCOPE_CYCLE_COUNTER(STAT_ElectraPlayer_AudioAACDecode);
+			CSV_SCOPED_TIMING_STAT(ElectraPlayer, AudioAACDecode);
+			UEMediaError bufResult = Renderer->AcquireBuffer(CurrentOutputBuffer, 0, BufferAcquireOptions);
+			check(bufResult == UEMEDIA_ERROR_OK || bufResult == UEMEDIA_ERROR_INSUFFICIENT_DATA);
+			if (bufResult != UEMEDIA_ERROR_OK && bufResult != UEMEDIA_ERROR_INSUFFICIENT_DATA)
+			{
+				PostError(0, "Failed to acquire sample buffer", ERRCODE_INTERNAL_LINUX_COULD_NOT_GET_SAMPLE_BUFFER, bufResult);
+				return false;
+			}
+		}
+
+		bool bHaveAvailSmpBlk = CurrentOutputBuffer != nullptr;
+
+		// Check if the renderer can accept the output we will want to send to it.
+		// If it can't right now we treat this as if we do not have an available output buffer.
+		if (Renderer.IsValid() && !Renderer->CanReceiveOutputFrames(1))
+		{
+			bHaveAvailSmpBlk = false;
+		}
+
+		NotifyReadyBufferListener(bHaveAvailSmpBlk);
+		if (bHaveAvailSmpBlk)
+		{
+			if (DecoderInstance && !AU->AccessUnit->bIsDummyData)
+			{
+				if (bInDummyDecodeMode)
+				{
+					bHaveDiscontinuity = true;
+					bInDummyDecodeMode = false;
+				}
+
+				ILibavDecoderAAC::FInputAU DecAU;
+				DecAU.Data = AU->AccessUnit->AUData;
+				DecAU.DataSize = (int32) AU->AccessUnit->AUSize;
+				DecAU.DTS = AU->PTS;
+				DecAU.PTS = AU->PTS;
+				DecAU.Duration = AU->EndPTS - AU->PTS;
+				DecAU.UserValue = AU->PTS;
+				ILibavDecoder::EDecoderError DecErr = DecoderInstance->DecodeAccessUnit(DecAU);
+				if (DecErr != ILibavDecoder::EDecoderError::None)
+				{
+					PostError((int32)DecErr, "Failed to decode access unit", ERRCODE_INTERNAL_LINUX_FAILED_TO_DECODE_AUDIO);
+					return false;
+				}
+				ILibavDecoderAAC::FOutputInfo Out;
+				while(DecoderInstance->HaveOutput(Out) == ILibavDecoder::EOutputStatus::Available)
+				{
+					SCOPE_CYCLE_COUNTER(STAT_ElectraPlayer_AudioAACConvertOutput);
+					CSV_SCOPED_TIMING_STAT(ElectraPlayer, AudioAACConvertOutput);
+
+					if (!DecoderInstance->GetOutputAsS16(PCMBuffer, PCMBufferSize))
+					{
+						PostError(0, "Could not get decoded output due to decoded format being unsupported", ERRCODE_INTERNAL_LINUX_UNSUPPORTED_DECODER_OUTPUT_FORMAT);
+						return false;
+					}
+					int32 SamplingRate = Out.SampleRate;
+					int32 NumSamplesProduced = Out.NumSamples;
+					int32 NumberOfChannels = Out.NumChannels;
+					int32 nBytesPerSample = NumberOfChannels * sizeof(int16);
+					int32 NumDecodedBytes = NumSamplesProduced * nBytesPerSample;
+					if (!ChannelMapper.IsInitialized())
+					{
+						if (NumberOfChannels)
+						{
+							TArray<FAudioChannelMapper::FSourceLayout> Layout;
+							const FAudioChannelMapper::EChannelPosition* ChannelPositions = AACDecoderChannelOutputMappingLibavcodec::Order[NumberOfChannels - 1];
+							for(uint32 i=0; i<NumberOfChannels; ++i)
+							{
+								FAudioChannelMapper::FSourceLayout lo;
+								lo.ChannelPosition = ChannelPositions[i];
+								Layout.Emplace(MoveTemp(lo));
+							}
+							if (ChannelMapper.Initialize(sizeof(int16), Layout))
+							{
+								// Ok.
+							}
+							else
+							{
+								PostError(0, "Failed to initialize audio channel mapper", ERRCODE_INTERNAL_LINUX_UNSUPPORTED_NUMBER_OF_AUDIO_CHANNELS);
+								return false;
+							}
+						}
+						else
+						{
+							PostError(0, "Bad number (0) of decoded audio channels", ERRCODE_INTERNAL_LINUX_UNSUPPORTED_NUMBER_OF_AUDIO_CHANNELS);
+							return false;
+						}
+					}
+
+					int32 ByteOffsetToFirstSample = 0;
+					if (AU->StartOverlapDuration.GetAsHNS() || AU->EndOverlapDuration.GetAsHNS())
+					{
+						int32 SkipStartSampleNum = (int32) (AU->StartOverlapDuration.GetAsHNS() * SamplingRate / 10000000);
+						int32 SkipEndSampleNum = (int32) (AU->EndOverlapDuration.GetAsHNS() * SamplingRate / 10000000);
+
+						if (SkipStartSampleNum + SkipEndSampleNum < NumSamplesProduced)
+						{
+							ByteOffsetToFirstSample = SkipStartSampleNum * nBytesPerSample;
+							NumSamplesProduced -= SkipStartSampleNum;
+							NumSamplesProduced -= SkipEndSampleNum;
+							NumDecodedBytes -= (SkipStartSampleNum + SkipEndSampleNum) * nBytesPerSample;
+						}
+						else
+						{
+							NumSamplesProduced = 0;
+							NumDecodedBytes = 0;
+						}
+					}
+
+					int32 NumDecodedSamplesPerChannel = NumDecodedBytes / (sizeof(int16) * NumberOfChannels);
+					if (NumDecodedSamplesPerChannel)
+					{
+						int32 CurrentRenderOutputBufferSize = (int32)CurrentOutputBuffer->GetBufferProperties().GetValue("size").GetInt64();
+						void* CurrentRenderOutputBufferAddress = CurrentOutputBuffer->GetBufferProperties().GetValue("address").GetPointer();
+						ChannelMapper.MapChannels(CurrentRenderOutputBufferAddress, CurrentRenderOutputBufferSize, AdvancePointer(PCMBuffer, ByteOffsetToFirstSample), NumDecodedBytes, NumDecodedSamplesPerChannel);
+
+						FTimeValue Duration;
+						Duration.SetFromND(NumSamplesProduced, SamplingRate);
+
+						OutputBufferSampleProperties.Clear();
+						OutputBufferSampleProperties.Set("num_channels", FVariantValue((int64)ChannelMapper.GetNumTargetChannels()));
+						OutputBufferSampleProperties.Set("byte_size", FVariantValue((int64)(ChannelMapper.GetNumTargetChannels() * NumDecodedSamplesPerChannel * sizeof(int16))));
+						OutputBufferSampleProperties.Set("sample_rate", FVariantValue((int64) SamplingRate));
+						OutputBufferSampleProperties.Set("duration", FVariantValue(Duration));
+						OutputBufferSampleProperties.Set("pts", FVariantValue(CurrentPTS));
+						OutputBufferSampleProperties.Set("discontinuity", FVariantValue((bool)bHaveDiscontinuity));
+
+						Renderer->ReturnBuffer(CurrentOutputBuffer, true, OutputBufferSampleProperties);
+						CurrentOutputBuffer = nullptr;
+					}
+					else
+					{
+						ReturnUnusedOutputBuffer();
+					}
+					bHaveDiscontinuity = false;
+				}
+				bDone = true;
+			}
+			else
+			{
+				SCOPE_CYCLE_COUNTER(STAT_ElectraPlayer_AudioAACConvertOutput);
+				CSV_SCOPED_TIMING_STAT(ElectraPlayer, AudioAACConvertOutput);
+				// Check if we are in dummy decode mode already.
+				if (!bInDummyDecodeMode)
+				{
+					bInDummyDecodeMode = true;
+					bHaveDiscontinuity = true;
+				}
+
+				// Clear to silence
+				FMemory::Memzero(CurrentOutputBuffer->GetBufferProperties().GetValue("address").GetPointer(), CurrentOutputBuffer->GetBufferProperties().GetValue("size").GetInt64());
+
+				// Assume sensible defaults.
+				int64 NumChannels = 2;
+				int64 SampleRate = 48000;
+				int64 SamplesPerBlock = 1024;
+
+				// With a valid configuration record we can use the actual values.
+				if (ConfigRecord.IsValid())
+				{
+					NumChannels = ConfigRecord->PSSignal > 0 ? 2 : NumChannelsForConfig[ConfigRecord->ChannelConfiguration];
+					SampleRate = ConfigRecord->ExtSamplingFrequency ? ConfigRecord->ExtSamplingFrequency : ConfigRecord->SamplingRate;
+					SamplesPerBlock = ConfigRecord->SBRSignal > 0 ? 2048 : 1024;
+				}
+
+				// A partial sample block has fewer samples.
+				if (AU->AdjustedDuration.IsValid())
+				{
+					SamplesPerBlock = AU->AdjustedDuration.GetAsHNS() * SampleRate / 10000000;
+				}
+
+				if (SamplesPerBlock)
+				{
+					FTimeValue Duration;
+					Duration.SetFromND(SamplesPerBlock, (uint32) SampleRate);
+
+					OutputBufferSampleProperties.Clear();
+					OutputBufferSampleProperties.Set("num_channels", FVariantValue(NumChannels));
+					OutputBufferSampleProperties.Set("sample_rate", FVariantValue(SampleRate));
+					OutputBufferSampleProperties.Set("byte_size", FVariantValue((int64)(NumChannels * sizeof(int16) * SamplesPerBlock)));
+					OutputBufferSampleProperties.Set("duration", FVariantValue(Duration));
+					OutputBufferSampleProperties.Set("pts", FVariantValue(CurrentPTS));
+					OutputBufferSampleProperties.Set("discontinuity", FVariantValue((bool)bHaveDiscontinuity));
+
+					Renderer->ReturnBuffer(CurrentOutputBuffer, true, OutputBufferSampleProperties);
+					CurrentOutputBuffer = nullptr;
+				}
+				else
+				{
+					ReturnUnusedOutputBuffer();
+				}
+				bHaveDiscontinuity = false;
+				bDone = true;
+			}
+		}
+		else
+		{
+			// No available buffer. Sleep for a bit. Can't sleep on a signal since we have to check two: abort and flush
+			// We sleep for 20ms since a sample block for LC AAC produces 1024 samples which amount to 21.3ms @48kHz (or more at lower sample rates).
+			FMediaRunnable::SleepMilliseconds(20);
+		}
+	}
+	return true;
+}
+
+bool FAudioDecoderAACLinuxLibavcodec::IsDifferentFormat(TSharedPtrTS<FDecoderInput> AU)
+{
+	if (AU.IsValid() && AU->AccessUnit->AUCodecData.IsValid() && (!CurrentCodecData.IsValid() || CurrentCodecData->CodecSpecificData != AU->AccessUnit->AUCodecData->CodecSpecificData))
+	{
+		return true;
+	}
+	return false;
+}
+
+void FAudioDecoderAACLinuxLibavcodec::PrepareAU(TSharedPtrTS<FDecoderInput> AU)
+{
+	if (!AU->bHasBeenPrepared)
+	{
+		AU->bHasBeenPrepared = true;
+
+		// Does this AU fall (partially) outside the range for rendering?
+		FTimeValue StartTime = AU->AccessUnit->PTS;
+		FTimeValue EndTime = AU->AccessUnit->PTS + AU->AccessUnit->Duration;
+		AU->PTS = StartTime.GetAsHNS();		// The PTS we give the decoder no matter any adjustment.
+		AU->EndPTS = EndTime.GetAsHNS();	// End PTS we need to check the PTS value returned by the decoder against.
+		AU->StartOverlapDuration.SetToZero();
+		AU->EndOverlapDuration.SetToZero();
+		if (AU->AccessUnit->EarliestPTS.IsValid())
+		{
+			// If the end time of the AU is before the earliest render PTS we do not need to decode it.
+			if (EndTime <= AU->AccessUnit->EarliestPTS)
+			{
+				StartTime.SetToInvalid();
+			}
+			else if (StartTime < AU->AccessUnit->EarliestPTS)
+			{
+				AU->StartOverlapDuration = AU->AccessUnit->EarliestPTS - StartTime;
+				StartTime = AU->AccessUnit->EarliestPTS;
+			}
+		}
+		if (StartTime.IsValid() && AU->AccessUnit->LatestPTS.IsValid())
+		{
+			// If the start time is behind the latest render PTS we donot need to decode it.
+			if (StartTime >= AU->AccessUnit->LatestPTS)
+			{
+				StartTime.SetToInvalid();
+			}
+			else if (EndTime >= AU->AccessUnit->LatestPTS)
+			{
+				AU->EndOverlapDuration = EndTime - AU->AccessUnit->LatestPTS;
+				EndTime = AU->AccessUnit->LatestPTS;
+			}
+		}
+		AU->AdjustedPTS = StartTime;
+		AU->AdjustedDuration = EndTime - StartTime;
+		if (AU->AdjustedDuration <= FTimeValue::GetZero())
+		{
+			AU->AdjustedPTS.SetToInvalid();
+		}
+	}
+}
+
+void FAudioDecoderAACLinuxLibavcodec::WorkerThread()
+{
+	LLM_SCOPE(ELLMTag::ElectraPlayer);
+
+	ApplicationRunningSignal.Signal();
+	ApplicationSuspendConfirmedSignal.Reset();
+
+	TOptional<int64> SequenceIndex;
+	bool bError = false;
+
+	PCMBufferSize = sizeof(int16) * 8 * 2048 * 2;				// max. of 8 channels decoding HE-AAC with twice the room to assemble packets
+	PCMBuffer = (int16*)FMemory::Malloc(PCMBufferSize, 32);
+	CurrentOutputBuffer = nullptr;
+	bInDummyDecodeMode = false;
+	bHaveDiscontinuity = false;
+
+	bError = !CreateDecodedSamplePool();
+	check(!bError);
+
+	while(!TerminateThreadSignal.IsSignaled())
+	{
+		// If in background, wait until we get activated again.
+		if (!ApplicationRunningSignal.IsSignaled())
+		{
+			UE_LOG(LogElectraPlayer, Log, TEXT("FAudioDecoderAACLinuxLibavcodec(%p): OnSuspending"), this);
+			ApplicationSuspendConfirmedSignal.Signal();
+			while(!ApplicationRunningSignal.WaitTimeout(100 * 1000) && !TerminateThreadSignal.IsSignaled())
+			{
+			}
+			UE_LOG(LogElectraPlayer, Log, TEXT("FAudioDecoderAACLinuxLibavcodec(%p): OnResuming"), this);
+		}
+
+		// Notify the buffer listener that we will now be needing an AU for our input buffer.
+		if (!bError && InputBufferListener && NextAccessUnits.IsEmpty())
+		{
+			SCOPE_CYCLE_COUNTER(STAT_ElectraPlayer_AudioAACDecode);
+			CSV_SCOPED_TIMING_STAT(ElectraPlayer, AudioAACDecode);
+			IAccessUnitBufferListener::FBufferStats	stats;
+			stats.bEODSignaled = NextAccessUnits.GetEOD();
+			stats.bEODReached = NextAccessUnits.ReachedEOD();
+			ListenerMutex.Lock();
+			if (InputBufferListener)
+			{
+				InputBufferListener->DecoderInputNeeded(stats);
+			}
+			ListenerMutex.Unlock();
+		}
+		// Get the AU to be decoded if one is there.
+		bool bHaveData = NextAccessUnits.Wait(1000 * 10);
+		if (bHaveData)
+		{
+			TSharedPtrTS<FDecoderInput> CurrentAccessUnit;
+			bool bOk = NextAccessUnits.Dequeue(CurrentAccessUnit);
+			MEDIA_UNUSED_VAR(bOk);
+			check(bOk);
+
+			PrepareAU(CurrentAccessUnit);
+			if (!SequenceIndex.IsSet())
+			{
+				SequenceIndex = CurrentAccessUnit->AccessUnit->PTS.GetSequenceIndex();
+			}
+
+			// Check if the format has changed such that we need to destroy and re-create the decoder.
+			if (CurrentAccessUnit->AccessUnit->bTrackChangeDiscontinuity || IsDifferentFormat(CurrentAccessUnit) || (CurrentAccessUnit->AccessUnit->bIsDummyData && !bInDummyDecodeMode) || SequenceIndex.GetValue() != CurrentAccessUnit->AccessUnit->PTS.GetSequenceIndex())
+			{
+				SCOPE_CYCLE_COUNTER(STAT_ElectraPlayer_AudioAACDecode);
+				CSV_SCOPED_TIMING_STAT(ElectraPlayer, AudioAACDecode);
+				if (DecoderInstance)
+				{
+					InternalDecoderDestroy();
+					bHaveDiscontinuity = true;
+				}
+				ReturnUnusedOutputBuffer();
+				ConfigRecord.Reset();
+				CurrentCodecData.Reset();
+				ChannelMapper.Reset();
+				InDecoderInput.Empty();
+			}
+
+			// Parse the CSD into a configuration record.
+			if (!ConfigRecord.IsValid() && !bError)
+			{
+				SCOPE_CYCLE_COUNTER(STAT_ElectraPlayer_AudioAACDecode);
+				CSV_SCOPED_TIMING_STAT(ElectraPlayer, AudioAACDecode);
+				if (CurrentAccessUnit->AccessUnit->AUCodecData.IsValid())
+				{
+					CurrentCodecData = CurrentAccessUnit->AccessUnit->AUCodecData;
+					ConfigRecord = MakeShared<MPEG::FAACDecoderConfigurationRecord, ESPMode::ThreadSafe>();
+					if (!ConfigRecord->ParseFrom(CurrentCodecData->CodecSpecificData.GetData(), CurrentCodecData->CodecSpecificData.Num()))
+					{
+						ConfigRecord.Reset();
+						CurrentCodecData.Reset();
+						PostError(0, "Failed to parse AAC configuration record", ERRCODE_INTERNAL_LINUX_FAILED_TO_PARSE_CSD);
+						bError = true;
+					}
+					else
+					{
+						// Check for unsupported configurations
+						if (ConfigRecord->ChannelConfiguration != 1 &&		// Mono
+							ConfigRecord->ChannelConfiguration != 2 &&		// Stereo
+							ConfigRecord->ChannelConfiguration != 3 &&		// L+C+R
+							ConfigRecord->ChannelConfiguration != 6)		// 5.1
+						{
+							ConfigRecord.Reset();
+							CurrentCodecData.Reset();
+							PostError(0, "Unsupported AAC channel configuration", ERRCODE_INTERNAL_LINUX_UNSUPPORTED_NUMBER_OF_AUDIO_CHANNELS);
+							bError = true;
+						}
+					}
+				}
+			}
+
+			SequenceIndex = CurrentAccessUnit->AccessUnit->PTS.GetSequenceIndex();
+			// Is this audio packet to be dropped?
+			if (!CurrentAccessUnit->AdjustedPTS.IsValid())
+			{
+				CurrentAccessUnit.Reset();
+				continue;
+			}
+
+			// Need to create a decoder instance?
+			if (DecoderInstance == nullptr && !bError)
+			{
+				if (!InternalDecoderCreate())
+				{
+					bError = true;
+				}
+			}
+
+			// Decode if not in error, otherwise just spend some idle time.
+			if (!bError)
+			{
+				if (!Decode(CurrentAccessUnit))
+				{
+					bError = true;
+				}
+			}
+			else
+			{
+				// Pace ourselves in consuming any more data until we are being stopped.
+				FMediaRunnable::SleepMicroseconds(CurrentAccessUnit->AccessUnit->Duration.GetAsMicroseconds());
+			}
+
+			// AU payload is no longer needed, just the information. Release the payload.
+			if (CurrentAccessUnit.IsValid())
+			{
+				CurrentAccessUnit->ReleasePayload();
+			}
+			CurrentAccessUnit.Reset();
+		}
+		else
+		{
+			// No data. Is the buffer at EOD?
+			if (NextAccessUnits.ReachedEOD())
+			{
+				InDecoderInput.Empty();
+				NotifyReadyBufferListener(true);
+				FMediaRunnable::SleepMilliseconds(20);
+			}
+		}
+
+		// Flush?
+		if (FlushDecoderSignal.IsSignaled())
+		{
+			SCOPE_CYCLE_COUNTER(STAT_ElectraPlayer_AudioAACDecode);
+			CSV_SCOPED_TIMING_STAT(ElectraPlayer, AudioAACDecode);
+			ReturnUnusedOutputBuffer();
+			NextAccessUnits.Empty();
+			InDecoderInput.Empty();
+			SequenceIndex.Reset();
+
+			FlushDecoderSignal.Reset();
+			DecoderFlushedSignal.Signal();
+		}
+	}
+
+	ReturnUnusedOutputBuffer();
+	// Close the decoder.
+	InternalDecoderDestroy();
+	DestroyDecodedSamplePool();
+
+	// Flush any remaining input data.
+	NextAccessUnits.Empty();
+	InDecoderInput.Empty();
+
+	CurrentCodecData.Reset();
+	ConfigRecord.Reset();
+
+	FMemory::Free(PCMBuffer);
+}
+
+} // namespace Electra
diff --git a/Engine/Plugins/Media/ElectraPlayer/Source/ElectraPlayerRuntime/Private/Runtime/Decoder/Linux/DecoderErrors_Linux.h b/Engine/Plugins/Media/ElectraPlayer/Source/ElectraPlayerRuntime/Private/Runtime/Decoder/Linux/DecoderErrors_Linux.h
new file mode 100644
--- /dev/null
+++ b/Engine/Plugins/Media/ElectraPlayer/Source/ElectraPlayerRuntime/Private/Runtime/Decoder/Linux/DecoderErrors_Linux.h
@@ -0,0 +1,25 @@
+// Copyright Epic Games, Inc. All Rights Reserved.
+
+#pragma once
+
+
+#define ERRCODE_INTERNAL_LINUX_AUDIO_DECODER_NOT_SUPPORTED		  	1
+#define ERRCODE_INTERNAL_LINUX_OUT_OF_MEMORY						2
+#define ERRCODE_INTERNAL_LINUX_COULD_NOT_CREATE_AUDIO_DECODER		3
+#define ERRCODE_INTERNAL_LINUX_COULD_NOT_CREATE_SAMPLE_POOL			4
+#define ERRCODE_INTERNAL_LINUX_COULD_NOT_GET_SAMPLE_BUFFER			5
+#define ERRCODE_INTERNAL_LINUX_FAILED_TO_PARSE_CSD					6
+#define ERRCODE_INTERNAL_LINUX_FAILED_TO_DECODE_AUDIO				7
+#define ERRCODE_INTERNAL_LINUX_AUDIO_OUTPUT_BUFFER_TOO_SMALL		8
+#define ERRCODE_INTERNAL_LINUX_UNSUPPORTED_NUMBER_OF_AUDIO_CHANNELS	9
+#define ERRCODE_INTERNAL_LINUX_UNSUPPORTED_DECODER_OUTPUT_FORMAT	10
+
+#define ERRCODE_INTERNAL_LINUX_VIDEO_DECODER_NOT_SUPPORTED			101
+#define ERRCODE_INTERNAL_LINUX_COULD_NOT_CREATE_VIDEO_DECODER		102
+#define ERRCODE_INTERNAL_LINUX_COULD_NOT_CREATE_IMAGE_POOL			103
+#define ERRCODE_INTERNAL_LINUX_COULD_NOT_GET_OUTPUT_BUFFER			104
+#define ERRCODE_INTERNAL_LINUX_FAILED_TO_DECODE_VIDEO				105
+#define ERRCODE_INTERNAL_LINUX_NO_VIDEO_CSD							106
+#define ERRCODE_INTERNAL_LINUX_BAD_VIDEO_CSD						107
+#define ERRCODE_INTERNAL_LINUX_FAILED_TO_CREATE_BLOCK_BUFFER		108
+#define ERRCODE_INTERNAL_LINUX_FAILED_TO_CREATE_SAMPLE_BUFFER		109
diff --git a/Engine/Plugins/Media/ElectraPlayer/Source/ElectraPlayerRuntime/Private/Runtime/Decoder/Linux/ElectraVideoDecoder_Linux.cpp b/Engine/Plugins/Media/ElectraPlayer/Source/ElectraPlayerRuntime/Private/Runtime/Decoder/Linux/ElectraVideoDecoder_Linux.cpp
new file mode 100644
--- /dev/null
+++ b/Engine/Plugins/Media/ElectraPlayer/Source/ElectraPlayerRuntime/Private/Runtime/Decoder/Linux/ElectraVideoDecoder_Linux.cpp
@@ -0,0 +1,94 @@
+// Copyright Epic Games, Inc. All Rights Reserved.
+
+#include "ElectraVideoDecoder_Linux.h"
+
+namespace Electra
+{
+
+FElectraPlayerVideoDecoderOutputLinux::FElectraPlayerVideoDecoderOutputLinux()
+{
+	Stride = 0;
+}
+
+FElectraPlayerVideoDecoderOutputLinux::~FElectraPlayerVideoDecoderOutputLinux()
+{
+}
+
+void FElectraPlayerVideoDecoderOutputLinux::SetDecodedImage(TSharedPtr<ILibavDecoderDecodedImage, ESPMode::ThreadSafe> InDecodedImage)
+{
+	DecodedImage = MoveTemp(InDecodedImage);
+}
+
+bool FElectraPlayerVideoDecoderOutputLinux::InitializeForBuffer(FIntPoint Dim, EPixelFormat PixFmt, FParamDict* InParamDict)
+{
+	FVideoDecoderOutputLinux::Initialize(InParamDict);
+	check(PixFmt == EPixelFormat::PF_NV12);
+	if (PixFmt != EPixelFormat::PF_NV12)
+	{
+		return false;
+	}
+
+	// Round up to multiple of 2.
+	Dim.X = ((Dim.X + 1) / 2) * 2;
+	Dim.Y = ((Dim.Y + 1) / 2) * 2;
+
+	int32 AllocSize = (Dim.X * Dim.Y) * 3 / 2;
+	if (Buffer.Num() < AllocSize)
+	{
+		Buffer.SetNum(AllocSize);
+	}
+	// Stride is the width of the buffer
+	Stride = Dim.X;
+	// The vertical sample dimension encompasses the height of the UV plane.
+	SampleDim.X = Dim.X;
+	SampleDim.Y = Dim.Y * 3 / 2;
+
+	return true;
+}
+
+FIntPoint FElectraPlayerVideoDecoderOutputLinux::GetBufferDimensions() const
+{
+	return FIntPoint(SampleDim.X, SampleDim.Y * 2 / 3);
+}
+
+TArray<uint8>& FElectraPlayerVideoDecoderOutputLinux::GetMutableBuffer()
+{
+	return Buffer;
+}
+
+const TArray<uint8>& FElectraPlayerVideoDecoderOutputLinux::GetBuffer() const
+{
+	return Buffer;
+}
+
+uint32 FElectraPlayerVideoDecoderOutputLinux::GetStride() const
+{
+	return Stride;
+}
+
+
+void FElectraPlayerVideoDecoderOutputLinux::SetOwner(const TSharedPtr<IDecoderOutputOwner, ESPMode::ThreadSafe>& InOwningRenderer)
+{
+	OwningRenderer = InOwningRenderer;
+}
+
+void FElectraPlayerVideoDecoderOutputLinux::ShutdownPoolable()
+{
+	// Release the decoded image.
+	DecodedImage.Reset();
+
+	TSharedPtr<IDecoderOutputOwner, ESPMode::ThreadSafe> lockedVideoRenderer = OwningRenderer.Pin();
+	if (lockedVideoRenderer.IsValid())
+	{
+		lockedVideoRenderer->SampleReleasedToPool(this);
+	}
+}
+
+}
+
+// -----------------------------------------------------------------------------------------------------------------------------
+
+FVideoDecoderOutput* FElectraPlayerPlatformVideoDecoderOutputFactory::Create()
+{
+    return new Electra::FElectraPlayerVideoDecoderOutputLinux();
+}
diff --git a/Engine/Plugins/Media/ElectraPlayer/Source/ElectraPlayerRuntime/Private/Runtime/Decoder/Linux/ElectraVideoDecoder_Linux.h b/Engine/Plugins/Media/ElectraPlayer/Source/ElectraPlayerRuntime/Private/Runtime/Decoder/Linux/ElectraVideoDecoder_Linux.h
new file mode 100644
--- /dev/null
+++ b/Engine/Plugins/Media/ElectraPlayer/Source/ElectraPlayerRuntime/Private/Runtime/Decoder/Linux/ElectraVideoDecoder_Linux.h
@@ -0,0 +1,51 @@
+// Copyright Epic Games, Inc. All Rights Reserved.
+
+#pragma once
+
+#include "HAL/Platform.h"
+
+#include "PlayerCore.h"
+#include "Linux/MediaVideoDecoderOutputLinux.h"
+#include "Renderer/RendererVideo.h"
+
+#include "libav_Decoder_Common_Video.h"
+
+namespace Electra
+{
+
+class FElectraPlayerVideoDecoderOutputLinux : public FVideoDecoderOutputLinux
+{
+public:
+	FElectraPlayerVideoDecoderOutputLinux();
+
+	~FElectraPlayerVideoDecoderOutputLinux();
+
+	void SetDecodedImage(TSharedPtr<ILibavDecoderDecodedImage, ESPMode::ThreadSafe> InDecodedImage);
+
+
+	bool InitializeForBuffer(FIntPoint Dim, EPixelFormat PixFmt, FParamDict* InParamDict);
+	TArray<uint8>& GetMutableBuffer();
+	FIntPoint GetBufferDimensions() const override;
+
+	const TArray<uint8>& GetBuffer() const override;
+	uint32 GetStride() const override;
+	FIntPoint GetDim() const override
+	{ return SampleDim; }
+
+	void SetOwner(const TSharedPtr<IDecoderOutputOwner, ESPMode::ThreadSafe>& InOwningRenderer) override;
+
+	void ShutdownPoolable() override;
+
+private:
+	FIntPoint SampleDim;
+	TArray<uint8> Buffer;
+	uint32 Stride;
+
+	TSharedPtr<ILibavDecoderDecodedImage, ESPMode::ThreadSafe> DecodedImage;
+
+	// We hold a weak reference to the video renderer. During destruction the video renderer could be destroyed while samples are still out there..
+	TWeakPtr<IDecoderOutputOwner, ESPMode::ThreadSafe> OwningRenderer;
+};
+
+
+} // namespace Electra
diff --git a/Engine/Plugins/Media/ElectraPlayer/Source/ElectraPlayerRuntime/Private/Runtime/Decoder/Linux/VideoDecoderH264_Linux.cpp b/Engine/Plugins/Media/ElectraPlayer/Source/ElectraPlayerRuntime/Private/Runtime/Decoder/Linux/VideoDecoderH264_Linux.cpp
new file mode 100644
--- /dev/null
+++ b/Engine/Plugins/Media/ElectraPlayer/Source/ElectraPlayerRuntime/Private/Runtime/Decoder/Linux/VideoDecoderH264_Linux.cpp
@@ -0,0 +1,1474 @@
+// Copyright Epic Games, Inc. All Rights Reserved.
+
+#include "CoreMinimal.h"
+#include "HAL/Platform.h"
+
+#include "PlayerCore.h"
+#include "StreamAccessUnitBuffer.h"
+#include "Decoder/VideoDecoderH264.h"
+#include "Renderer/RendererBase.h"
+#include "Player/PlayerSessionServices.h"
+#include "Utilities/StringHelpers.h"
+#include "HAL/LowLevelMemTracker.h"
+#include "ElectraPlayerPrivate.h"
+#include "DecoderErrors_Linux.h"
+
+#include "Linux/MediaVideoDecoderOutputLinux.h"
+#include "Renderer/RendererVideo.h"
+#include "ElectraVideoDecoder_Linux.h"
+
+#include "libav_Decoder_Common.h"
+#include "libav_Decoder_H264.h"
+
+/***************************************************************************************************************************************************/
+
+DECLARE_CYCLE_STAT(TEXT("FVideoDecoderH264::Decode()"), STAT_ElectraPlayer_VideoH264Decode, STATGROUP_ElectraPlayer);
+DECLARE_CYCLE_STAT(TEXT("FVideoDecoderH264::ConvertOutput()"), STAT_ElectraPlayer_VideoH264ConvertOutput, STATGROUP_ElectraPlayer);
+
+namespace Electra
+{
+
+/**
+ * H264 video decoder class implementation.
+**/
+class FVideoDecoderH264LinuxLibavcodec : public IVideoDecoderH264, public FMediaThread, public ILibavDecoderVideoResourceAllocator
+{
+public:
+	static bool Startup(const IVideoDecoderH264::FSystemConfiguration& InConfig);
+	static void Shutdown();
+
+	FVideoDecoderH264LinuxLibavcodec();
+	virtual ~FVideoDecoderH264LinuxLibavcodec();
+
+	void SetPlayerSessionServices(IPlayerSessionServices* SessionServices) override;
+
+	void Open(const FInstanceConfiguration& InConfig) override;
+	void Close() override;
+	void DrainForCodecChange() override;
+
+	void SetMaximumDecodeCapability(int32 MaxWidth, int32 MaxHeight, int32 MaxProfile, int32 MaxProfileLevel, const FParamDict& AdditionalOptions) override;
+
+	void SetAUInputBufferListener(IAccessUnitBufferListener* Listener) override;
+
+	void SetReadyBufferListener(IDecoderOutputBufferListener* Listener) override;
+
+	void SetRenderer(TSharedPtr<IMediaRenderer, ESPMode::ThreadSafe> InRenderer) override;
+
+    void SetResourceDelegate(const TSharedPtr<IVideoDecoderResourceDelegate, ESPMode::ThreadSafe>& ResourceDelegate) override;
+
+	void AUdataPushAU(FAccessUnit* AccessUnit) override;
+	void AUdataPushEOD() override;
+	void AUdataClearEOD() override;
+	void AUdataFlushEverything() override;
+
+private:
+	// Add methods from ILibavDecoderVideoResourceAllocator here
+	// ...
+
+	enum
+	{
+		MaxPendingOutputImages = 2
+	};
+
+	enum EDecodeResult
+	{
+		Ok,
+		Fail
+	};
+
+	struct FDecoderInput : public TSharedFromThis<FDecoderInput, ESPMode::ThreadSafe>
+	{
+		~FDecoderInput()
+		{
+			ReleasePayload();
+		}
+		void ReleasePayload()
+		{
+			FAccessUnit::Release(AccessUnit);
+			AccessUnit = nullptr;
+		}
+
+		FAccessUnit*	AccessUnit = nullptr;
+		bool			bHasBeenPrepared = false;
+		bool			bIsIDR = false;
+		bool			bIsDiscardable = false;
+		int64			PTS = 0;
+		int64			EndPTS = 0;
+		int64			SequenceIndex = 0;
+		FTimeValue		AdjustedPTS;
+		FTimeValue		AdjustedDuration;
+
+		int32			Width = 0;
+		int32			Height = 0;
+		int32			TotalWidth = 0;
+		int32			TotalHeight = 0;
+		int32			CropLeft = 0;
+		int32			CropRight = 0;
+		int32			CropTop = 0;
+		int32			CropBottom = 0;
+		int32			AspectX = 0;
+		int32			AspectY = 0;
+	};
+
+	struct FDecoderFormatInfo
+	{
+		void Reset()
+		{
+			CurrentCodecData.Reset();
+		}
+		bool IsDifferentFrom(TSharedPtr<FDecoderInput, ESPMode::ThreadSafe> AU);
+
+		TSharedPtr<const FAccessUnit::CodecData, ESPMode::ThreadSafe> CurrentCodecData;
+	};
+
+	struct FDecodedImage
+	{
+		FDecodedImage() = default;
+		FDecodedImage(const FDecodedImage& rhs)
+		{
+			InternalCopy(rhs);
+		}
+
+		FDecodedImage& operator=(const FDecodedImage& rhs)
+		{
+			if (this != &rhs)
+			{
+				InternalCopy(rhs);
+			}
+			return *this;
+		}
+
+		bool operator<(const FDecodedImage& rhs) const
+		{
+			if (!SourceInfo.IsValid() || !rhs.SourceInfo.IsValid())
+			{
+				return false;
+			}
+			return (SourceInfo->SequenceIndex < rhs.SourceInfo->SequenceIndex) || (SourceInfo->SequenceIndex == rhs.SourceInfo->SequenceIndex && SourceInfo->PTS < rhs.SourceInfo->PTS);
+		}
+
+		TSharedPtr<FDecoderInput, ESPMode::ThreadSafe>	SourceInfo;
+		TSharedPtr<ILibavDecoderDecodedImage, ESPMode::ThreadSafe> DecodedImage;
+		ILibavDecoderVideoCommon::FOutputInfo ImageInfo;
+	private:
+		void InternalCopy(const FDecodedImage& rhs)
+		{
+			SourceInfo = rhs.SourceInfo;
+			DecodedImage = rhs.DecodedImage;
+			ImageInfo = rhs.ImageInfo;
+		}
+	};
+
+
+	bool InternalDecoderCreate();
+	void InternalDecoderDestroy();
+	void StartThread();
+	void StopThread();
+	void WorkerThread();
+
+	bool CreateDecodedImagePool();
+	void DestroyDecodedImagePool();
+
+	void NotifyReadyBufferListener(bool bHaveOutput);
+
+	bool AcquireOutputBuffer(IMediaRenderer::IBuffer*& RenderOutputBuffer);
+
+	void PrepareAU(TSharedPtr<FDecoderInput, ESPMode::ThreadSafe> AU);
+	void GetAndPrepareInputAU();
+
+	void ProcessOutput(bool bFlush = false);
+
+	bool FlushDecoder();
+	void ClearInDecoderInfos();
+	void FlushPendingImages();
+
+	EDecodeResult Decode(TSharedPtr<FDecoderInput, ESPMode::ThreadSafe> AU);
+	bool DecodeDummy(TSharedPtr<FDecoderInput, ESPMode::ThreadSafe> AU);
+
+	void PostError(int32_t ApiReturnValue, const FString& Message, uint16 Code, UEMediaError Error = UEMEDIA_ERROR_OK);
+	void LogMessage(IInfoLog::ELevel Level, const FString& Message);
+
+	void ConvertDecodedImageToNV12(TArray<uint8>& OutNV12Buffer, FIntPoint OutBufDim, const FDecodedImage& InImage);
+
+	FInstanceConfiguration Config;
+
+	FMediaEvent TerminateThreadSignal;
+	FMediaEvent FlushDecoderSignal;
+	FMediaEvent DecoderFlushedSignal;
+	bool bThreadStarted = false;
+
+	FMediaEvent ApplicationRunningSignal;
+	FMediaEvent ApplicationSuspendConfirmedSignal;
+
+	IPlayerSessionServices* PlayerSessionServices = nullptr;
+
+	TSharedPtr<IMediaRenderer, ESPMode::ThreadSafe> Renderer;
+
+    TWeakPtr<IVideoDecoderResourceDelegate, ESPMode::ThreadSafe> ResourceDelegate;
+
+	TAccessUnitQueue<TSharedPtr<FDecoderInput, ESPMode::ThreadSafe>> NextAccessUnits;
+	TSharedPtr<FDecoderInput, ESPMode::ThreadSafe> CurrentAccessUnit;
+	bool bDrainForCodecChange = false;
+
+	FCriticalSection ListenerMutex;
+	IAccessUnitBufferListener* InputBufferListener = nullptr;
+	IDecoderOutputBufferListener* ReadyBufferListener = nullptr;
+
+	FDecoderFormatInfo CurrentStreamFormatInfo;
+	TSharedPtr<ILibavDecoderVideoCommon, ESPMode::ThreadSafe> DecoderInstance;
+	TArray<TSharedPtr<FDecoderInput, ESPMode::ThreadSafe>> InDecoderInput;
+
+	int32 MaxDecodeBufferSize = 0;
+	bool bError = false;
+
+	TArray<FDecodedImage> ReadyImages;
+
+public:
+	static FSystemConfiguration SystemConfig;
+};
+
+IVideoDecoderH264::FSystemConfiguration FVideoDecoderH264LinuxLibavcodec::SystemConfig;
+
+/***************************************************************************************************************************************************/
+/***************************************************************************************************************************************************/
+/***************************************************************************************************************************************************/
+
+bool IVideoDecoderH264::Startup(const IVideoDecoderH264::FSystemConfiguration& InConfig)
+{
+	return FVideoDecoderH264LinuxLibavcodec::Startup(InConfig);
+}
+
+void IVideoDecoderH264::Shutdown()
+{
+	FVideoDecoderH264LinuxLibavcodec::Shutdown();
+}
+
+bool IVideoDecoderH264::GetStreamDecodeCapability(FStreamDecodeCapability& OutResult, const FStreamDecodeCapability& InStreamParameter)
+{
+	OutResult = InStreamParameter;
+	OutResult.DecoderSupportType = ILibavDecoderH264::IsAvailable() ? FStreamDecodeCapability::ESupported::HardAndSoftware : FStreamDecodeCapability::ESupported::NotSupported;
+	return true;
+}
+
+IVideoDecoderH264::FSystemConfiguration::FSystemConfiguration()
+{
+	ThreadConfig.Decoder.Priority 	= TPri_Normal;
+	ThreadConfig.Decoder.StackSize	= 256 << 10;
+	ThreadConfig.Decoder.CoreAffinity = -1;
+}
+
+IVideoDecoderH264::FInstanceConfiguration::FInstanceConfiguration()
+	: MaxDecodedFrames(8)
+	, ThreadConfig(FVideoDecoderH264LinuxLibavcodec::SystemConfig.ThreadConfig)
+{
+}
+
+IVideoDecoderH264* IVideoDecoderH264::Create()
+{
+	return new FVideoDecoderH264LinuxLibavcodec;
+}
+
+
+/***************************************************************************************************************************************************/
+/***************************************************************************************************************************************************/
+/***************************************************************************************************************************************************/
+
+//-----------------------------------------------------------------------------
+/**
+ * Decoder system startup
+ *
+ * @param config
+ *
+ * @return
+ */
+bool FVideoDecoderH264LinuxLibavcodec::Startup(const IVideoDecoderH264::FSystemConfiguration& InConfig)
+{
+	SystemConfig = InConfig;
+	ILibavDecoder::Startup();
+	return ILibavDecoder::IsLibAvAvailable();
+}
+
+
+//-----------------------------------------------------------------------------
+/**
+ * Decoder system shutdown.
+ */
+void FVideoDecoderH264LinuxLibavcodec::Shutdown()
+{
+	ILibavDecoder::Shutdown();
+}
+
+
+//-----------------------------------------------------------------------------
+/**
+ * Constructor
+ */
+FVideoDecoderH264LinuxLibavcodec::FVideoDecoderH264LinuxLibavcodec()
+	: FMediaThread("ElectraPlayer::H264 decoder")
+{
+}
+
+
+//-----------------------------------------------------------------------------
+/**
+ * Destructor
+ */
+FVideoDecoderH264LinuxLibavcodec::~FVideoDecoderH264LinuxLibavcodec()
+{
+	Close();
+}
+
+
+//-----------------------------------------------------------------------------
+/**
+ * Sets an AU input buffer listener.
+ *
+ * @param Listener
+ */
+void FVideoDecoderH264LinuxLibavcodec::SetAUInputBufferListener(IAccessUnitBufferListener* Listener)
+{
+	FScopeLock lock(&ListenerMutex);
+	InputBufferListener = Listener;
+}
+
+
+//-----------------------------------------------------------------------------
+/**
+ * Sets a buffer-ready listener.
+ *
+ * @param Listener
+ */
+void FVideoDecoderH264LinuxLibavcodec::SetReadyBufferListener(IDecoderOutputBufferListener* Listener)
+{
+	FScopeLock lock(&ListenerMutex);
+	ReadyBufferListener = Listener;
+}
+
+
+//-----------------------------------------------------------------------------
+/**
+ * Sets the owning player's session service interface.
+ *
+ * @param InSessionServices
+ *
+ * @return
+ */
+void FVideoDecoderH264LinuxLibavcodec::SetPlayerSessionServices(IPlayerSessionServices* InSessionServices)
+{
+	PlayerSessionServices = InSessionServices;
+}
+
+
+//-----------------------------------------------------------------------------
+/**
+ * Opens a decoder instance
+ *
+ * @param InConfig
+ *
+ * @return
+ */
+void FVideoDecoderH264LinuxLibavcodec::Open(const IVideoDecoderH264::FInstanceConfiguration& InConfig)
+{
+	Config = InConfig;
+	StartThread();
+}
+
+
+//-----------------------------------------------------------------------------
+/**
+ * Closes the decoder instance.
+ */
+void FVideoDecoderH264LinuxLibavcodec::Close()
+{
+	StopThread();
+}
+
+
+//-----------------------------------------------------------------------------
+/**
+ * Drains the decoder of all enqueued input and ends it, after which the decoder must send an FDecoderMessage to the player
+ * to signal completion.
+ */
+void FVideoDecoderH264LinuxLibavcodec::DrainForCodecChange()
+{
+	bDrainForCodecChange = true;
+}
+
+
+//-----------------------------------------------------------------------------
+/**
+ * Sets a new decoder limit.
+ * As soon as a new sequence starts matching this limit the decoder will be
+ * destroyed and recreated to conserve memory.
+ * The assumption is that the data being streamed will not exceed this limit,
+ * but if it does, any access unit requiring a better decoder will take
+ * precedence and force a decoder capable of decoding it!
+ *
+ * @param MaxWidth
+ * @param MaxHeight
+ * @param MaxProfile
+ * @param MaxProfileLevel
+ * @param AdditionalOptions
+ */
+void FVideoDecoderH264LinuxLibavcodec::SetMaximumDecodeCapability(int32 MaxWidth, int32 MaxHeight, int32 MaxProfile, int32 MaxProfileLevel, const FParamDict& AdditionalOptions)
+{
+	// Not implemented
+}
+
+
+//-----------------------------------------------------------------------------
+/**
+ * Sets a new renderer.
+ *
+ * @param InRenderer
+ */
+void FVideoDecoderH264LinuxLibavcodec::SetRenderer(TSharedPtr<IMediaRenderer, ESPMode::ThreadSafe> InRenderer)
+{
+	Renderer = InRenderer;
+}
+
+void FVideoDecoderH264LinuxLibavcodec::SetResourceDelegate(const TSharedPtr<IVideoDecoderResourceDelegate, ESPMode::ThreadSafe>& InResourceDelegate)
+{
+    ResourceDelegate = InResourceDelegate;
+}
+
+//-----------------------------------------------------------------------------
+/**
+ * Creates and runs the decoder thread.
+ */
+void FVideoDecoderH264LinuxLibavcodec::StartThread()
+{
+	ThreadSetPriority(Config.ThreadConfig.Decoder.Priority);
+	ThreadSetCoreAffinity(Config.ThreadConfig.Decoder.CoreAffinity);
+	ThreadSetStackSize(Config.ThreadConfig.Decoder.StackSize);
+	ThreadStart(FMediaRunnable::FStartDelegate::CreateRaw(this, &FVideoDecoderH264LinuxLibavcodec::WorkerThread));
+	bThreadStarted = true;
+}
+
+
+//-----------------------------------------------------------------------------
+/**
+ * Stops the decoder thread.
+ */
+void FVideoDecoderH264LinuxLibavcodec::StopThread()
+{
+	if (bThreadStarted)
+	{
+		TerminateThreadSignal.Signal();
+		ThreadWaitDone();
+		bThreadStarted = false;
+	}
+}
+
+
+//-----------------------------------------------------------------------------
+/**
+ * Posts an error to the session service error listeners.
+ *
+ * @param ApiReturnValue
+ * @param Message
+ * @param Code
+ * @param Error
+ */
+void FVideoDecoderH264LinuxLibavcodec::PostError(int32_t ApiReturnValue, const FString& Message, uint16 Code, UEMediaError Error)
+{
+	check(PlayerSessionServices);
+	if (PlayerSessionServices)
+	{
+		FErrorDetail err;
+		err.SetError(Error != UEMEDIA_ERROR_OK ? Error : UEMEDIA_ERROR_DETAIL);
+		err.SetFacility(Facility::EFacility::H264Decoder);
+		err.SetCode(Code);
+		err.SetMessage(Message);
+		err.SetPlatformMessage(FString::Printf(TEXT("%d (0x%08x)"), ApiReturnValue, ApiReturnValue));
+		PlayerSessionServices->PostError(err);
+	}
+}
+
+
+//-----------------------------------------------------------------------------
+/**
+ * Sends a log message to the session service log.
+ *
+ * @param Level
+ * @param Message
+ */
+void FVideoDecoderH264LinuxLibavcodec::LogMessage(IInfoLog::ELevel Level, const FString& Message)
+{
+	if (PlayerSessionServices)
+	{
+		PlayerSessionServices->PostLog(Facility::EFacility::H264Decoder, Level, Message);
+	}
+}
+
+
+//-----------------------------------------------------------------------------
+/**
+ * Create a pool of decoded images for the decoder.
+ *
+ * @return
+ */
+bool FVideoDecoderH264LinuxLibavcodec::CreateDecodedImagePool()
+{
+	check(Renderer);
+	FParamDict poolOpts;
+
+	poolOpts.Set("num_buffers", FVariantValue((int64) Config.MaxDecodedFrames));
+
+	UEMediaError Error = Renderer->CreateBufferPool(poolOpts);
+	check(Error == UEMEDIA_ERROR_OK);
+
+	MaxDecodeBufferSize = (int32) Renderer->GetBufferPoolProperties().GetValue("max_buffers").GetInt64();
+
+	if (Error != UEMEDIA_ERROR_OK)
+	{
+		PostError(0, "Failed to create image pool", ERRCODE_INTERNAL_LINUX_COULD_NOT_CREATE_IMAGE_POOL, Error);
+	}
+
+	return Error == UEMEDIA_ERROR_OK;
+}
+
+
+//-----------------------------------------------------------------------------
+/**
+ * Destroys the pool of decoded images.
+ */
+void FVideoDecoderH264LinuxLibavcodec::DestroyDecodedImagePool()
+{
+	Renderer->ReleaseBufferPool();
+}
+
+
+//-----------------------------------------------------------------------------
+/**
+ * Called to receive a new input access unit for decoding.
+ *
+ * @param AccessUnit
+ */
+void FVideoDecoderH264LinuxLibavcodec::AUdataPushAU(FAccessUnit* InAccessUnit)
+{
+	InAccessUnit->AddRef();
+
+	TSharedPtr<FDecoderInput, ESPMode::ThreadSafe> NextAU = MakeShared<FDecoderInput, ESPMode::ThreadSafe>();
+	NextAU->AccessUnit = InAccessUnit;
+	NextAccessUnits.Enqueue(MoveTemp(NextAU));
+}
+
+
+//-----------------------------------------------------------------------------
+/**
+ * "Pushes" an End Of Data marker indicating no further access units will be added.
+ */
+void FVideoDecoderH264LinuxLibavcodec::AUdataPushEOD()
+{
+	NextAccessUnits.SetEOD();
+}
+
+
+//-----------------------------------------------------------------------------
+/**
+ * Notifies the decoder that there may be further access units.
+ */
+void FVideoDecoderH264LinuxLibavcodec::AUdataClearEOD()
+{
+	NextAccessUnits.ClearEOD();
+}
+
+
+//-----------------------------------------------------------------------------
+/**
+ * Flushes the decoder and clears the input access unit buffer.
+ */
+void FVideoDecoderH264LinuxLibavcodec::AUdataFlushEverything()
+{
+	FlushDecoderSignal.Signal();
+	DecoderFlushedSignal.WaitAndReset();
+}
+
+
+//-----------------------------------------------------------------------------
+/**
+ * Create a decoder instance.
+ *
+ * @return true if successful, false on error
+ */
+bool FVideoDecoderH264LinuxLibavcodec::InternalDecoderCreate()
+{
+	InternalDecoderDestroy();
+
+	if (ILibavDecoderH264::IsAvailable())
+	{
+		TMap<FString, FVariant> Options;
+		Options.Add(TEXT("hw_priority"), FVariant(FString("vdpau;cuda;vaapi")));
+		DecoderInstance = ILibavDecoderH264::Create(this, Options);
+		if (!DecoderInstance.IsValid() || DecoderInstance->GetLastLibraryError())
+		{
+			InternalDecoderDestroy();
+			PostError(-2, "libavcodec failed to open video decoder", ERRCODE_INTERNAL_LINUX_COULD_NOT_CREATE_VIDEO_DECODER);
+			return false;
+		}
+		return true;
+	}
+	PostError(-1, "libavcodec does not support this video format", ERRCODE_INTERNAL_LINUX_VIDEO_DECODER_NOT_SUPPORTED);
+	return false;
+}
+
+
+//-----------------------------------------------------------------------------
+/**
+ * Destroys the current decoder instance.
+ */
+void FVideoDecoderH264LinuxLibavcodec::InternalDecoderDestroy()
+{
+	DecoderInstance.Reset();
+}
+
+
+//-----------------------------------------------------------------------------
+/**
+ * Notifies the buffer-ready listener that we will now be producing output data
+ * and need an output buffer.
+ * This call is not intended to create or obtain an output buffer. It is merely
+ * indicating that new output will be produced.
+ *
+ * @param bHaveOutput
+ */
+void FVideoDecoderH264LinuxLibavcodec::NotifyReadyBufferListener(bool bHaveOutput)
+{
+	if (ReadyBufferListener)
+	{
+		IDecoderOutputBufferListener::FDecodeReadyStats stats;
+		stats.MaxDecodedElementsReady = MaxDecodeBufferSize;
+		stats.NumElementsInDecoder = InDecoderInput.Num();
+		stats.bOutputStalled = !bHaveOutput;
+		stats.bEODreached = NextAccessUnits.ReachedEOD() && stats.NumDecodedElementsReady == 0 && stats.NumElementsInDecoder == 0;
+		ListenerMutex.Lock();
+		if (ReadyBufferListener)
+		{
+			ReadyBufferListener->DecoderOutputReady(stats);
+		}
+		ListenerMutex.Unlock();
+	}
+}
+
+
+//-----------------------------------------------------------------------------
+/**
+ * Prepares the passed access unit to be sent to the decoder.
+ *
+ * @param AU
+ */
+void FVideoDecoderH264LinuxLibavcodec::PrepareAU(TSharedPtr<FDecoderInput, ESPMode::ThreadSafe> AU)
+{
+	if (!AU->bHasBeenPrepared)
+	{
+		AU->bHasBeenPrepared = true;
+
+		const FStreamCodecInformation::FResolution& res = AU->AccessUnit->AUCodecData->ParsedInfo.GetResolution();
+		const FStreamCodecInformation::FAspectRatio& ar = AU->AccessUnit->AUCodecData->ParsedInfo.GetAspectRatio();
+		const FStreamCodecInformation::FCrop& crop = AU->AccessUnit->AUCodecData->ParsedInfo.GetCrop();
+		AU->Width = res.Width;
+		AU->Height = res.Height;
+		AU->TotalWidth = res.Width + crop.Left + crop.Right;
+		AU->TotalHeight = res.Height + crop.Top + crop.Bottom;
+		AU->CropLeft = crop.Left;
+		AU->CropRight = crop.Right;
+		AU->CropTop = crop.Top;
+		AU->CropBottom = crop.Bottom;
+		AU->AspectX = ar.Width ? ar.Width : 1;
+		AU->AspectY = ar.Height ? ar.Height : 1;
+
+		if (!AU->AccessUnit->bIsDummyData)
+		{
+			// Process NALUs
+			AU->bIsDiscardable = true;
+			AU->bIsIDR = AU->AccessUnit->bIsSyncSample;
+			uint32* NALU = (uint32 *)AU->AccessUnit->AUData;
+			uint32* End  = (uint32 *)Electra::AdvancePointer(NALU, AU->AccessUnit->AUSize);
+			while(NALU < End)
+			{
+				// Check the nal_ref_idc in the NAL unit for dependencies.
+				uint8 nal = *(const uint8 *)(NALU + 1);
+				check((nal & 0x80) == 0);
+				if ((nal >> 5) != 0)
+				{
+					AU->bIsDiscardable = false;
+				}
+				// IDR frame?
+				if ((nal & 0x1f) == 5)
+				{
+					AU->bIsIDR = true;
+				}
+
+				// SEI message(s)?
+				if ((nal & 0x1f) == 6)
+				{
+					// ...
+				}
+
+				uint32 naluLen = MEDIA_FROM_BIG_ENDIAN(*NALU) + 4;
+				*NALU = MEDIA_TO_BIG_ENDIAN(0x00000001U);
+				NALU = Electra::AdvancePointer(NALU, naluLen);
+			}
+		}
+
+		// Does this AU fall (partially) outside the range for rendering?
+		FTimeValue StartTime = AU->AccessUnit->PTS;
+		FTimeValue EndTime = AU->AccessUnit->PTS + AU->AccessUnit->Duration;
+		AU->PTS = StartTime.GetAsHNS();				// The PTS we give the decoder no matter any adjustment.
+		AU->EndPTS = EndTime.GetAsHNS();			// End PTS we need to check the PTS value returned by the decoder against.
+		AU->SequenceIndex = StartTime.GetSequenceIndex();
+		if (AU->AccessUnit->EarliestPTS.IsValid())
+		{
+			// If the end time of the AU is before the earliest render PTS we need to decode it, but not display it.
+			if (EndTime <= AU->AccessUnit->EarliestPTS)
+			{
+				StartTime.SetToInvalid();
+			}
+			else if (StartTime < AU->AccessUnit->EarliestPTS)
+			{
+				StartTime = AU->AccessUnit->EarliestPTS;
+			}
+		}
+		if (StartTime.IsValid() && AU->AccessUnit->LatestPTS.IsValid())
+		{
+			// If the start time is behind the latest render PTS we do not need to decode at all.
+			if (StartTime >= AU->AccessUnit->LatestPTS)
+			{
+				StartTime.SetToInvalid();
+				AU->bIsDiscardable = true;
+			}
+			else if (EndTime >= AU->AccessUnit->LatestPTS)
+			{
+				EndTime = AU->AccessUnit->LatestPTS;
+			}
+		}
+		AU->AdjustedPTS = StartTime;
+		AU->AdjustedDuration = EndTime - StartTime;
+		if (AU->AdjustedDuration <= FTimeValue::GetZero())
+		{
+			AU->AdjustedPTS.SetToInvalid();
+		}
+
+		// Is there codec specific data?
+		if (!AU->AccessUnit->bIsDummyData && AU->AccessUnit->AUCodecData.IsValid())
+		{
+			// Yes.
+			if (AU->bIsIDR || AU->AccessUnit->bIsSyncSample || AU->AccessUnit->bIsFirstInSequence)
+			{
+				// Have to re-allocate the AU memory to preprend the codec data
+				if (AU->AccessUnit->AUCodecData->CodecSpecificData.Num())
+				{
+					uint64 nb = AU->AccessUnit->AUSize + AU->AccessUnit->AUCodecData->CodecSpecificData.Num();
+					void* pD = AU->AccessUnit->AllocatePayloadBuffer(nb);
+					void* pP = pD;
+					FMemory::Memcpy(pP, AU->AccessUnit->AUCodecData->CodecSpecificData.GetData(), AU->AccessUnit->AUCodecData->CodecSpecificData.Num());
+					pP = Electra::AdvancePointer(pP, AU->AccessUnit->AUCodecData->CodecSpecificData.Num());
+					FMemory::Memcpy(pP, AU->AccessUnit->AUData, AU->AccessUnit->AUSize);
+					AU->AccessUnit->AdoptNewPayloadBuffer(pD, nb);
+				}
+				AU->bIsDiscardable = false;
+			}
+		}
+	}
+}
+
+
+//-----------------------------------------------------------------------------
+/**
+ * Gets an input access unit and prepares it for use.
+ */
+void FVideoDecoderH264LinuxLibavcodec::GetAndPrepareInputAU()
+{
+	if (bDrainForCodecChange)
+	{
+		return;
+	}
+
+	// Need new input?
+	if (!CurrentAccessUnit.IsValid() && InputBufferListener && NextAccessUnits.IsEmpty())
+	{
+		{
+			SCOPE_CYCLE_COUNTER(STAT_ElectraPlayer_VideoH264Decode);
+			CSV_SCOPED_TIMING_STAT(ElectraPlayer, VideoH264Decode);
+
+			IAccessUnitBufferListener::FBufferStats	stats;
+			stats.bEODSignaled = NextAccessUnits.GetEOD();
+			stats.bEODReached = NextAccessUnits.ReachedEOD();
+			ListenerMutex.Lock();
+			if (InputBufferListener)
+			{
+				InputBufferListener->DecoderInputNeeded(stats);
+			}
+			ListenerMutex.Unlock();
+		}
+
+		// Get the AU to be decoded if one is there.
+		bool bHaveData = NextAccessUnits.Wait(1000);
+		if (bHaveData)
+		{
+			SCOPE_CYCLE_COUNTER(STAT_ElectraPlayer_VideoH264Decode);
+			CSV_SCOPED_TIMING_STAT(ElectraPlayer, VideoH264Decode);
+
+			bool bOk = NextAccessUnits.Dequeue(CurrentAccessUnit);
+			MEDIA_UNUSED_VAR(bOk);
+			check(bOk);
+
+			PrepareAU(CurrentAccessUnit);
+		}
+	}
+}
+
+
+//-----------------------------------------------------------------------------
+/**
+ * Gets an output buffer from the video renderer.
+ *
+ * @param RenderOutputBuffer
+ *
+ * @return true if successful, false on error
+ */
+bool FVideoDecoderH264LinuxLibavcodec::AcquireOutputBuffer(IMediaRenderer::IBuffer*& RenderOutputBuffer)
+{
+	RenderOutputBuffer = nullptr;
+	FParamDict BufferAcquireOptions;
+	while(!TerminateThreadSignal.IsSignaled() && !FlushDecoderSignal.IsSignaled())
+	{
+		UEMediaError bufResult = Renderer->AcquireBuffer(RenderOutputBuffer, 0, BufferAcquireOptions);
+		check(bufResult == UEMEDIA_ERROR_OK || bufResult == UEMEDIA_ERROR_INSUFFICIENT_DATA);
+		if (bufResult != UEMEDIA_ERROR_OK && bufResult != UEMEDIA_ERROR_INSUFFICIENT_DATA)
+		{
+			PostError(0, "Failed to acquire sample buffer", ERRCODE_INTERNAL_LINUX_COULD_NOT_GET_OUTPUT_BUFFER, bufResult);
+			return false;
+		}
+		bool bHaveAvailSmpBlk = RenderOutputBuffer != nullptr;
+		NotifyReadyBufferListener(bHaveAvailSmpBlk);
+		if (bHaveAvailSmpBlk)
+		{
+			break;
+		}
+		else
+		{
+			// No available buffer. Sleep for a bit. Can't sleep on a signal since we have to check two: abort and flush
+			FMediaRunnable::SleepMilliseconds(5);
+		}
+	}
+	return true;
+}
+
+
+//-----------------------------------------------------------------------------
+/**
+ * Flushes the decoder and passes all pending frames to the renderer.
+ *
+ * @return true if successful, false otherwise.
+ */
+bool FVideoDecoderH264LinuxLibavcodec::FlushDecoder()
+{
+	if (DecoderInstance.IsValid())
+	{
+		ILibavDecoder::EDecoderError DecErr = DecoderInstance->SendEndOfData();
+		check(DecErr == ILibavDecoder::EDecoderError::None); (void)DecErr;
+		ProcessOutput(true);
+	}
+
+	ClearInDecoderInfos();
+	return true;
+}
+
+
+//-----------------------------------------------------------------------------
+/**
+ * Clears out any leftover decoder source infos that were not picked by decoded frames.
+ * This must be called after having flushed the decoder to ensure no delayed decode
+ * callbacks will arrive.
+ */
+void FVideoDecoderH264LinuxLibavcodec::ClearInDecoderInfos()
+{
+	InDecoderInput.Empty();
+}
+
+
+//-----------------------------------------------------------------------------
+/**
+ * Creates a dummy output image that is not to be displayed and has no image data.
+ * Dummy access units are created when stream data is missing to ensure the data
+ * pipeline does not run dry and exhibits no gaps in the timeline.
+ *
+ * @param AU
+ *
+ * @return true if successful, false otherwise
+ */
+bool FVideoDecoderH264LinuxLibavcodec::DecodeDummy(TSharedPtr<FDecoderInput, ESPMode::ThreadSafe> AU)
+{
+	if (AU.IsValid() && AU->AdjustedPTS.IsValid())
+	{
+		FDecodedImage NextImage;
+		NextImage.SourceInfo = AU;
+		ReadyImages.Add(NextImage);
+	}
+    return true;
+}
+
+//-----------------------------------------------------------------------------
+/**
+ * Flushes all images not passed to the renderer from our list.
+ */
+void FVideoDecoderH264LinuxLibavcodec::FlushPendingImages()
+{
+	ReadyImages.Empty();
+}
+
+//-----------------------------------------------------------------------------
+/**
+ * Checks if the codec specific data has changed.
+ *
+ * @param AU
+ *
+ * @return false if the format is still the same, true if it has changed.
+ */
+bool FVideoDecoderH264LinuxLibavcodec::FDecoderFormatInfo::IsDifferentFrom(TSharedPtr<FDecoderInput, ESPMode::ThreadSafe> AU)
+{
+	if (AU->AccessUnit->AUCodecData.IsValid() && AU->AccessUnit->AUCodecData.Get() != CurrentCodecData.Get())
+	{
+		// Pointers are different. Is the content too?
+		bool bDifferent = !CurrentCodecData.IsValid() || (CurrentCodecData.IsValid() && AU->AccessUnit->AUCodecData->CodecSpecificData != CurrentCodecData->CodecSpecificData);
+		CurrentCodecData = AU->AccessUnit->AUCodecData;
+		return bDifferent;
+	}
+	return false;
+}
+
+
+//-----------------------------------------------------------------------------
+/**
+ * Sends an access unit to the decoder for decoding.
+ *
+ * @param AU
+ *
+ * @return
+ */
+FVideoDecoderH264LinuxLibavcodec::EDecodeResult FVideoDecoderH264LinuxLibavcodec::Decode(TSharedPtr<FDecoderInput, ESPMode::ThreadSafe> AU)
+{
+	if (!DecoderInstance.IsValid())
+	{
+		return EDecodeResult::Fail;
+	}
+
+	bool bAddInput = true;
+	// Decode
+	ILibavDecoderVideoCommon::FInputAU DecAU;
+	DecAU.Data = AU->AccessUnit->AUData;
+	DecAU.DataSize = (int32) AU->AccessUnit->AUSize;
+	DecAU.DTS = AU->AccessUnit->DTS.GetAsHNS();
+	DecAU.PTS = AU->AccessUnit->PTS.GetAsHNS();
+	DecAU.Duration = AU->AccessUnit->Duration.GetAsHNS();
+	DecAU.UserValue = AU->PTS;
+	DecAU.Flags = 0;
+	if (!AU->AdjustedPTS.IsValid())
+	{
+		DecAU.Flags |= ILibavDecoderVideoCommon::FInputAU::EFlags::EVidAUFlag_DoNotOutput;
+		bAddInput = false;
+	}
+	if (bAddInput)
+	{
+		InDecoderInput.Add(AU);
+	}
+
+	ILibavDecoder::EDecoderError DecErr = DecoderInstance->DecodeAccessUnit(DecAU);
+	if (DecErr == ILibavDecoder::EDecoderError::None)
+	{
+		// Ok.
+		return EDecodeResult::Ok;
+	}
+	else
+	{
+		InDecoderInput.RemoveSingle(AU);
+		PostError(DecoderInstance->GetLastLibraryError(), FString::Printf(TEXT("Failed to decode video frame (%d)"),(int)DecErr), ERRCODE_INTERNAL_LINUX_FAILED_TO_DECODE_VIDEO);
+		return EDecodeResult::Fail;
+	}
+}
+
+
+//-----------------------------------------------------------------------------
+/**
+ * H264 video decoder main threaded decode loop
+ */
+void FVideoDecoderH264LinuxLibavcodec::WorkerThread()
+{
+	LLM_SCOPE(ELLMTag::ElectraPlayer);
+
+	ApplicationRunningSignal.Signal();
+	ApplicationSuspendConfirmedSignal.Reset();
+
+	TOptional<int64> SequenceIndex;
+	bool bDone  = false;
+    bool bGotLastSequenceAU = false;
+	bool bInDummyDecodeMode = false;
+
+	bError = false;
+
+	// Create decoded image pool.
+	if (!CreateDecodedImagePool())
+	{
+		bError = true;
+	}
+
+	while(!TerminateThreadSignal.IsSignaled())
+	{
+		// If in background, wait until we get activated again.
+		if (!ApplicationRunningSignal.IsSignaled())
+		{
+			UE_LOG(LogElectraPlayer, Log, TEXT("FVideoDecoderH264LinuxLibavcodec(%p): OnSuspending"), this);
+			ApplicationSuspendConfirmedSignal.Signal();
+			while(!ApplicationRunningSignal.WaitTimeout(100 * 1000) && !TerminateThreadSignal.IsSignaled())
+			{
+			}
+			UE_LOG(LogElectraPlayer, Log, TEXT("FVideoDecoderH264LinuxLibavcodec(%p): OnResuming"), this);
+		}
+
+		if (!bDrainForCodecChange)
+		{
+			GetAndPrepareInputAU();
+
+			bool bHaveData = CurrentAccessUnit.IsValid();
+			bool bTooManyImagesWaiting = false;
+			if (bHaveData)
+			{
+				// When there is data, even and especially after a previous EOD, we are no longer done and idling.
+				bDone = false;
+
+				// If there is too much output still pending to be delivered to the renderer
+				bTooManyImagesWaiting = ReadyImages.Num() > MaxPendingOutputImages;
+				if (bTooManyImagesWaiting)
+				{
+					while(!TerminateThreadSignal.IsSignaled() && !FlushDecoderSignal.IsSignaled())
+					{
+						ProcessOutput();
+						if ((bTooManyImagesWaiting = ReadyImages.Num() > MaxPendingOutputImages) == false)
+						{
+							break;
+						}
+						FMediaRunnable::SleepMilliseconds(10);
+					}
+				}
+			}
+
+			if (bHaveData && !bTooManyImagesWaiting)
+			{
+				if (!CurrentAccessUnit->AccessUnit->bIsDummyData)
+				{
+					bInDummyDecodeMode = false;
+
+					bool bStreamFormatChanged = CurrentStreamFormatInfo.IsDifferentFrom(CurrentAccessUnit) || bGotLastSequenceAU;
+					bool bNeedNewDecoder = false;
+
+					if (!SequenceIndex.IsSet())
+					{
+						SequenceIndex = CurrentAccessUnit->AccessUnit->PTS.GetSequenceIndex();
+					}
+					bool bSequenceIndexChanged = SequenceIndex.GetValue() != CurrentAccessUnit->AccessUnit->PTS.GetSequenceIndex();
+					SequenceIndex = CurrentAccessUnit->AccessUnit->PTS.GetSequenceIndex();
+
+					bGotLastSequenceAU = CurrentAccessUnit->AccessUnit->bIsLastInPeriod;
+					if (bStreamFormatChanged || bSequenceIndexChanged || !DecoderInstance.IsValid())
+					{
+						SCOPE_CYCLE_COUNTER(STAT_ElectraPlayer_VideoH264Decode);
+						CSV_SCOPED_TIMING_STAT(ElectraPlayer, VideoH264Decode);
+						bNeedNewDecoder |= bSequenceIndexChanged || !DecoderInstance.IsValid();
+						if (bNeedNewDecoder)
+						{
+							FlushDecoder();
+							InternalDecoderDestroy();
+							if (!bError && !InternalDecoderCreate())
+							{
+								bError = true;
+							}
+						}
+					}
+
+					// If this AU falls outside the range where it is to be rendered and it is also discardable
+					// we do not need to concern ourselves with it at all.
+					if (CurrentAccessUnit->bIsDiscardable && !CurrentAccessUnit->AdjustedPTS.IsValid())
+					{
+						CurrentAccessUnit.Reset();
+						continue;
+					}
+
+					// Decode
+					if (!bError && DecoderInstance.IsValid())
+					{
+						EDecodeResult DecRes = Decode(CurrentAccessUnit);
+
+						// Process any output we might have pending
+						ProcessOutput();
+
+						// Decode ok?
+						if (DecRes == EDecodeResult::Ok)
+						{
+							CurrentAccessUnit.Reset();
+						}
+						// Decode failed
+						else
+						{
+							bError = true;
+						}
+					}
+				}
+				else
+				{
+					if (!bInDummyDecodeMode)
+					{
+						bInDummyDecodeMode = true;
+						FlushDecoder();
+					}
+
+					if (!DecodeDummy(CurrentAccessUnit))
+					{
+						bError = true;
+					}
+					// DecodeDummy() went through most of the regular path, but has returned the output buffer immediately
+					// and can thus always get a new one with no waiting. To avoid draining the player buffer by consuming
+					// the dummy AUs at rapid pace we put ourselves to sleep for the duration the AU was supposed to last.
+					FMediaRunnable::SleepMicroseconds(CurrentAccessUnit->AdjustedDuration.GetAsMicroseconds());
+
+					ProcessOutput();
+					CurrentAccessUnit.Reset();
+				}
+			}
+			else
+			{
+				ProcessOutput();
+
+				// No data. Is the buffer at EOD?
+				if (!CurrentAccessUnit.IsValid() && NextAccessUnits.ReachedEOD())
+				{
+					NotifyReadyBufferListener(true);
+					// Are we done yet?
+					if (!bDone && !bError)
+					{
+						bError = !FlushDecoder();
+					}
+					bDone = true;
+					FMediaRunnable::SleepMilliseconds(10);
+					ProcessOutput(true);
+				}
+			}
+		}
+		else
+		{
+			ProcessOutput();
+			NotifyReadyBufferListener(true);
+			// Are we done yet?
+			if (!bDone && !bError)
+			{
+				bError = !FlushDecoder();
+			}
+			bDone = true;
+			ProcessOutput(true);
+			break;
+		}
+
+		// Flush?
+		if (FlushDecoderSignal.IsSignaled())
+		{
+			SCOPE_CYCLE_COUNTER(STAT_ElectraPlayer_VideoH264Decode);
+			CSV_SCOPED_TIMING_STAT(ElectraPlayer, VideoH264Decode);
+
+			// Have to destroy the decoder!
+			InternalDecoderDestroy();
+			FlushPendingImages();
+			ClearInDecoderInfos();
+			NextAccessUnits.Empty();
+			SequenceIndex.Reset();
+			CurrentAccessUnit.Reset();
+
+			FlushDecoderSignal.Reset();
+			DecoderFlushedSignal.Signal();
+
+			// Reset done state.
+			bDone = false;
+		}
+	}
+
+	InternalDecoderDestroy();
+	FlushPendingImages();
+	DestroyDecodedImagePool();
+	ClearInDecoderInfos();
+	NextAccessUnits.Empty();
+	CurrentAccessUnit.Reset();
+
+	if (bDrainForCodecChange)
+	{
+		// Notify the player that we have finished draining.
+		PlayerSessionServices->SendMessageToPlayer(FDecoderMessage::Create(FDecoderMessage::EReason::DrainingFinished, this, EStreamType::Video, FStreamCodecInformation::ECodec::H264));
+		// We need to wait to get terminated. Also check if flushing is requested and acknowledge if it is.
+		while(!TerminateThreadSignal.IsSignaled())
+		{
+			if (FlushDecoderSignal.WaitTimeoutAndReset(1000 * 10))
+			{
+				DecoderFlushedSignal.Signal();
+			}
+		}
+	}
+}
+
+//-----------------------------------------------------------------------------
+/**
+ */
+void FVideoDecoderH264LinuxLibavcodec::ProcessOutput(bool bFlush)
+{
+	SCOPE_CYCLE_COUNTER(STAT_ElectraPlayer_VideoH264ConvertOutput);
+	CSV_SCOPED_TIMING_STAT(ElectraPlayer, VideoH264ConvertOutput);
+
+	// Pull output from decoder
+	while(DecoderInstance.IsValid())
+	{
+		FDecodedImage NextImage;
+
+		// Output available?
+		if (DecoderInstance->HaveOutput(NextImage.ImageInfo) != ILibavDecoderVideoCommon::EOutputStatus::Available)
+		{
+			break;
+		}
+		// Get the output.
+		NextImage.DecodedImage = DecoderInstance->GetOutput();
+		if (!NextImage.DecodedImage.IsValid())
+		{
+			break;
+		}
+		// Find the corresponding input source.
+		int32 NumCurrentDecodeInputs = InDecoderInput.Num();
+		for(int32 i=0; i<NumCurrentDecodeInputs; ++i)
+		{
+			if (InDecoderInput[i]->PTS == NextImage.ImageInfo.UserValue)
+			{
+				NextImage.SourceInfo = InDecoderInput[i];
+				InDecoderInput.RemoveSingle(NextImage.SourceInfo);
+				break;
+			}
+		}
+
+		if (!NextImage.SourceInfo.IsValid())
+		{
+			LogMessage(IInfoLog::ELevel::Warning, FString::Printf(TEXT("FVideoDecoderH264::ProcessOutput(): No source info found for decoded PTS %lld in %d pending infos"), (long long int)NextImage.ImageInfo.UserValue, NumCurrentDecodeInputs));
+			break;
+		}
+
+		// For safety we remove input that is older than what we got right now.
+		// This is to ensure that if libav does not return decoded output to us for some reason
+		// (like a corrupted frame or such) that we do not sit on stale input.
+		if (InDecoderInput.Num() > MaxPendingOutputImages)
+		{
+			for(int32 i=0; i<InDecoderInput.Num(); ++i)
+			{
+				if (InDecoderInput[i]->PTS < NextImage.ImageInfo.UserValue)
+				{
+					LogMessage(IInfoLog::ELevel::Warning, FString::Printf(TEXT("FVideoDecoderH264::ProcessOutput(): Removing stale input PTS %lld for output PTS %lld"), (long long int)InDecoderInput[i]->PTS, (long long int)NextImage.ImageInfo.UserValue));
+					InDecoderInput.RemoveAt(i);
+					--i;
+				}
+			}
+		}
+
+		ReadyImages.Add(NextImage);
+	}
+	ReadyImages.Sort();
+
+	// Send frames out to the renderer.
+	IMediaRenderer::IBuffer* RenderOutputBuffer = nullptr;
+	while(!TerminateThreadSignal.IsSignaled() && !FlushDecoderSignal.IsSignaled())
+	{
+		if (ReadyImages.Num() == 0)
+		{
+			break;
+		}
+
+		// When we are not flushing out the remaining frames and the renderer cannot take on
+		// anything new we leave.
+		if (!bFlush && !Renderer->CanReceiveOutputFrames(1))
+		{
+			NotifyReadyBufferListener(false);
+			break;
+		}
+
+		// Need a new output buffer?
+		if (!RenderOutputBuffer)
+		{
+			// If we failed to get a new buffer, leave the loop.
+			if (!AcquireOutputBuffer(RenderOutputBuffer))
+			{
+				break;
+			}
+		}
+
+		// Can we send the image to the renderer?
+		bool bHaveAvailableImage = RenderOutputBuffer != nullptr;
+		if (bFlush && !Renderer->CanReceiveOutputFrames(1))
+		{
+			bHaveAvailableImage = false;
+		}
+		NotifyReadyBufferListener(bHaveAvailableImage);
+		if (bHaveAvailableImage)
+		{
+			FDecodedImage NextImage(ReadyImages[0]);
+			ReadyImages.RemoveAt(0);
+
+			FParamDict* OutputBufferSampleProperties = new FParamDict();
+			FIntPoint Dim(NextImage.SourceInfo->Width, NextImage.SourceInfo->Height);
+
+			if (NextImage.DecodedImage.IsValid())
+			{
+				int32 ax = NextImage.ImageInfo.AspectNum;
+				int32 ay = NextImage.ImageInfo.AspectDenom;
+				Dim.X = NextImage.ImageInfo.Width;
+				Dim.Y = NextImage.ImageInfo.Height;
+
+				OutputBufferSampleProperties->Set("width", FVariantValue((int64) NextImage.ImageInfo.Width));
+				OutputBufferSampleProperties->Set("height", FVariantValue((int64) NextImage.ImageInfo.Height));
+				OutputBufferSampleProperties->Set("crop_left", FVariantValue((int64) 0));
+				OutputBufferSampleProperties->Set("crop_right", FVariantValue((int64) 0));
+				OutputBufferSampleProperties->Set("crop_top", FVariantValue((int64) 0));
+				OutputBufferSampleProperties->Set("crop_bottom", FVariantValue((int64) 0));
+				OutputBufferSampleProperties->Set("aspect_ratio", FVariantValue((double)ax / (double)ay));
+				OutputBufferSampleProperties->Set("aspect_w", FVariantValue((int64) ax));
+				OutputBufferSampleProperties->Set("aspect_h", FVariantValue((int64) ay));
+				OutputBufferSampleProperties->Set("fps_num", FVariantValue((int64) 0 ));
+				OutputBufferSampleProperties->Set("fps_denom", FVariantValue((int64) 0 ));
+				OutputBufferSampleProperties->Set("pixelfmt", FVariantValue((int64)EPixelFormat::PF_NV12));
+			}
+			else
+			{
+				OutputBufferSampleProperties->Set("is_dummy", FVariantValue(true));
+			}
+			OutputBufferSampleProperties->Set("pts", FVariantValue(NextImage.SourceInfo->AdjustedPTS));
+			OutputBufferSampleProperties->Set("duration", FVariantValue(NextImage.SourceInfo->AdjustedDuration));
+
+			bool bRender = NextImage.SourceInfo->AdjustedPTS.IsValid();
+			TSharedPtr<FElectraPlayerVideoDecoderOutputLinux, ESPMode::ThreadSafe> DecoderOutput = RenderOutputBuffer->GetBufferProperties().GetValue("texture").GetSharedPointer<FElectraPlayerVideoDecoderOutputLinux>();
+			FIntPoint BufferDim(NextImage.ImageInfo.Planes[0].Width, NextImage.ImageInfo.Planes[0].Height);
+			if (DecoderOutput->InitializeForBuffer(BufferDim, EPixelFormat::PF_NV12, OutputBufferSampleProperties))
+			{
+				TArray<uint8>& ImgBuf = DecoderOutput->GetMutableBuffer();
+				ConvertDecodedImageToNV12(ImgBuf, DecoderOutput->GetBufferDimensions(), NextImage);
+				// Have the decoder output keep a reference to the decoded image in case it will
+				// shared data with it at some point.
+				DecoderOutput->SetDecodedImage(NextImage.DecodedImage);
+			}
+
+			// Return the buffer to the renderer.
+			Renderer->ReturnBuffer(RenderOutputBuffer, bRender, *OutputBufferSampleProperties);
+			RenderOutputBuffer = nullptr;
+		}
+		else
+		{
+			// Since we need to check the abort and the flush signal we sleep for a bit and check again.
+			FMediaRunnable::SleepMilliseconds(10);
+		}
+	}
+	if (RenderOutputBuffer)
+	{
+		FParamDict None;
+		Renderer->ReturnBuffer(RenderOutputBuffer, false, None);
+		RenderOutputBuffer = nullptr;
+	}
+}
+
+
+void FVideoDecoderH264LinuxLibavcodec::ConvertDecodedImageToNV12(TArray<uint8>& OutNV12Buffer, FIntPoint OutBufDim, const FDecodedImage& InImage)
+{
+	if (OutNV12Buffer.GetData() && InImage.DecodedImage.IsValid())
+	{
+		if (InImage.ImageInfo.NumPlanes == 2 &&
+			InImage.ImageInfo.Planes[0].Content == ILibavDecoderVideoCommon::FPlaneInfo::EContent::Luma &&
+			InImage.ImageInfo.Planes[1].Content == ILibavDecoderVideoCommon::FPlaneInfo::EContent::ChromaUV)
+		{
+			const int32 w = InImage.ImageInfo.Planes[0].Width;
+			const int32 h = InImage.ImageInfo.Planes[0].Height;
+			const int32 aw = ((w + 1) / 2) * 2;
+			const int32 ah = ((h + 1) / 2) * 2;
+			uint8* DstY = OutNV12Buffer.GetData();
+			uint8* DstUV = DstY + aw * ah;
+			const uint8* SrcY = (const uint8*)InImage.ImageInfo.Planes[0].Address;
+			const uint8* SrcUV = (const uint8*)InImage.ImageInfo.Planes[1].Address;
+			// To simplify the conversion we require the output buffer to have the dimension of the planes.
+			check(OutBufDim.X == aw && OutBufDim.Y == ah);
+			if (!SrcY || !SrcUV || OutBufDim.X != aw || OutBufDim.Y != ah)
+			{
+				return;
+			}
+			if ((w & 1) == 0)
+			{
+				FMemory::Memcpy(DstY, SrcY, w*h);
+				FMemory::Memcpy(DstUV, SrcUV, w*h/2);
+			}
+			else
+			{
+				for(int32 y=0; y<h; ++y)
+				{
+					FMemory::Memcpy(DstY, SrcY, w);
+					DstY += aw;
+					SrcY += w;
+				}
+				for(int32 y=0; y<h/2; ++y)
+				{
+					FMemory::Memcpy(DstUV, SrcUV, w);
+					DstUV += aw;
+					SrcUV += w;
+				}
+			}
+		}
+		else if (InImage.ImageInfo.NumPlanes == 3 &&
+				 InImage.ImageInfo.Planes[0].Content == ILibavDecoderVideoCommon::FPlaneInfo::EContent::Luma &&
+				 InImage.ImageInfo.Planes[1].Content == ILibavDecoderVideoCommon::FPlaneInfo::EContent::ChromaU &&
+				 InImage.ImageInfo.Planes[2].Content == ILibavDecoderVideoCommon::FPlaneInfo::EContent::ChromaV)
+		{
+			const int32 w = InImage.ImageInfo.Planes[0].Width;
+			const int32 h = InImage.ImageInfo.Planes[0].Height;
+			const int32 aw = ((w + 1) / 2) * 2;
+			const int32 ah = ((h + 1) / 2) * 2;
+			uint8* DstY = OutNV12Buffer.GetData();
+			uint8* DstUV = DstY + aw * ah;
+			const uint8* SrcY = (const uint8*)InImage.ImageInfo.Planes[0].Address;
+			const uint8* SrcU = (const uint8*)InImage.ImageInfo.Planes[1].Address;
+			const uint8* SrcV = (const uint8*)InImage.ImageInfo.Planes[2].Address;
+			// To simplify the conversion we require the output buffer to have the dimension of the planes.
+			check(OutBufDim.X == aw && OutBufDim.Y == ah);
+			if (!SrcY || !SrcU || !SrcV || OutBufDim.X != aw || OutBufDim.Y != ah)
+			{
+				return;
+			}
+			if ((w & 1) == 0)
+			{
+				FMemory::Memcpy(DstY, SrcY, w*h);
+				for(int32 i=0, iMax=w*h/4; i<iMax; ++i)
+				{
+					*DstUV++ = *SrcU++;
+					*DstUV++ = *SrcV++;
+				}
+			}
+			else
+			{
+				for(int32 y=0; y<h; ++y)
+				{
+					FMemory::Memcpy(DstY, SrcY, w);
+					DstY += aw;
+					SrcY += w;
+				}
+				int32 padUV = (aw - w) * 2;
+				for(int32 v=0; v<h/2; ++v)
+				{
+					for(int32 u=0; u<w/2; ++u)
+					{
+						*DstUV++ = *SrcU++;
+						*DstUV++ = *SrcV++;
+					}
+					DstUV += padUV;
+				}
+			}
+		}
+	}
+}
+
+
+} // namespace Electra
diff --git a/Engine/Plugins/Media/ElectraPlayer/Source/ElectraPlayerRuntime/Private/Runtime/Decoder/Linux/VideoDecoderH265_Linux.cpp b/Engine/Plugins/Media/ElectraPlayer/Source/ElectraPlayerRuntime/Private/Runtime/Decoder/Linux/VideoDecoderH265_Linux.cpp
new file mode 100644
--- /dev/null
+++ b/Engine/Plugins/Media/ElectraPlayer/Source/ElectraPlayerRuntime/Private/Runtime/Decoder/Linux/VideoDecoderH265_Linux.cpp
@@ -0,0 +1,1460 @@
+// Copyright Epic Games, Inc. All Rights Reserved.
+
+#include "CoreMinimal.h"
+#include "HAL/Platform.h"
+
+#include "PlayerCore.h"
+#include "StreamAccessUnitBuffer.h"
+#include "Decoder/VideoDecoderH265.h"
+#include "Renderer/RendererBase.h"
+#include "Player/PlayerSessionServices.h"
+#include "Utilities/StringHelpers.h"
+#include "HAL/LowLevelMemTracker.h"
+#include "ElectraPlayerPrivate.h"
+#include "DecoderErrors_Linux.h"
+
+#include "Linux/MediaVideoDecoderOutputLinux.h"
+#include "Renderer/RendererVideo.h"
+#include "ElectraVideoDecoder_Linux.h"
+
+#include "libav_Decoder_Common.h"
+#include "libav_Decoder_H265.h"
+
+/***************************************************************************************************************************************************/
+
+DECLARE_CYCLE_STAT(TEXT("FVideoDecoderH265::Decode()"), STAT_ElectraPlayer_VideoH265Decode, STATGROUP_ElectraPlayer);
+DECLARE_CYCLE_STAT(TEXT("FVideoDecoderH265::ConvertOutput()"), STAT_ElectraPlayer_VideoH265ConvertOutput, STATGROUP_ElectraPlayer);
+
+namespace Electra
+{
+
+/**
+ * H265 video decoder class implementation.
+**/
+class FVideoDecoderH265LinuxLibavcodec : public IVideoDecoderH265, public FMediaThread, public ILibavDecoderVideoResourceAllocator
+{
+public:
+	static bool Startup(const FParamDict& Options);
+	static void Shutdown();
+
+	FVideoDecoderH265LinuxLibavcodec();
+	virtual ~FVideoDecoderH265LinuxLibavcodec();
+
+	void SetPlayerSessionServices(IPlayerSessionServices* SessionServices) override;
+
+	void Open(const FInstanceConfiguration& InConfig) override;
+	void Close() override;
+	void DrainForCodecChange() override;
+
+	void SetMaximumDecodeCapability(int32 MaxTier, int32 MaxWidth, int32 MaxHeight, int32 MaxProfile, int32 MaxProfileLevel, const FParamDict& AdditionalOptions) override;
+
+	void SetAUInputBufferListener(IAccessUnitBufferListener* Listener) override;
+
+	void SetReadyBufferListener(IDecoderOutputBufferListener* Listener) override;
+
+	void SetRenderer(TSharedPtr<IMediaRenderer, ESPMode::ThreadSafe> InRenderer) override;
+
+    void SetResourceDelegate(const TSharedPtr<IVideoDecoderResourceDelegate, ESPMode::ThreadSafe>& ResourceDelegate) override;
+
+	void AUdataPushAU(FAccessUnit* AccessUnit) override;
+	void AUdataPushEOD() override;
+	void AUdataClearEOD() override;
+	void AUdataFlushEverything() override;
+
+private:
+	// Add methods from ILibavDecoderVideoResourceAllocator here
+	// ...
+
+	enum
+	{
+		MaxPendingOutputImages = 2
+	};
+
+	enum EDecodeResult
+	{
+		Ok,
+		Fail
+	};
+
+	struct FDecoderInput : public TSharedFromThis<FDecoderInput, ESPMode::ThreadSafe>
+	{
+		~FDecoderInput()
+		{
+			ReleasePayload();
+		}
+		void ReleasePayload()
+		{
+			FAccessUnit::Release(AccessUnit);
+			AccessUnit = nullptr;
+		}
+
+		FAccessUnit*	AccessUnit = nullptr;
+		bool			bHasBeenPrepared = false;
+		bool			bIsIDR = false;
+		bool			bIsDiscardable = false;
+		int64			PTS = 0;
+		int64			EndPTS = 0;
+		int64			SequenceIndex = 0;
+		FTimeValue		AdjustedPTS;
+		FTimeValue		AdjustedDuration;
+
+		int32			Width = 0;
+		int32			Height = 0;
+		int32			TotalWidth = 0;
+		int32			TotalHeight = 0;
+		int32			CropLeft = 0;
+		int32			CropRight = 0;
+		int32			CropTop = 0;
+		int32			CropBottom = 0;
+		int32			AspectX = 0;
+		int32			AspectY = 0;
+	};
+
+	struct FDecoderFormatInfo
+	{
+		void Reset()
+		{
+			CurrentCodecData.Reset();
+		}
+		bool IsDifferentFrom(TSharedPtr<FDecoderInput, ESPMode::ThreadSafe> AU);
+
+		TSharedPtr<const FAccessUnit::CodecData, ESPMode::ThreadSafe> CurrentCodecData;
+	};
+
+	struct FDecodedImage
+	{
+		FDecodedImage() = default;
+		FDecodedImage(const FDecodedImage& rhs)
+		{
+			InternalCopy(rhs);
+		}
+
+		FDecodedImage& operator=(const FDecodedImage& rhs)
+		{
+			if (this != &rhs)
+			{
+				InternalCopy(rhs);
+			}
+			return *this;
+		}
+
+		bool operator<(const FDecodedImage& rhs) const
+		{
+			if (!SourceInfo.IsValid() || !rhs.SourceInfo.IsValid())
+			{
+				return false;
+			}
+			return (SourceInfo->SequenceIndex < rhs.SourceInfo->SequenceIndex) || (SourceInfo->SequenceIndex == rhs.SourceInfo->SequenceIndex && SourceInfo->PTS < rhs.SourceInfo->PTS);
+		}
+
+		TSharedPtr<FDecoderInput, ESPMode::ThreadSafe>	SourceInfo;
+		TSharedPtr<ILibavDecoderDecodedImage, ESPMode::ThreadSafe> DecodedImage;
+		ILibavDecoderVideoCommon::FOutputInfo ImageInfo;
+	private:
+		void InternalCopy(const FDecodedImage& rhs)
+		{
+			SourceInfo = rhs.SourceInfo;
+			DecodedImage = rhs.DecodedImage;
+			ImageInfo = rhs.ImageInfo;
+		}
+	};
+
+
+	bool InternalDecoderCreate();
+	void InternalDecoderDestroy();
+	void StartThread();
+	void StopThread();
+	void WorkerThread();
+
+	bool CreateDecodedImagePool();
+	void DestroyDecodedImagePool();
+
+	void NotifyReadyBufferListener(bool bHaveOutput);
+
+	bool AcquireOutputBuffer(IMediaRenderer::IBuffer*& RenderOutputBuffer);
+
+	void PrepareAU(TSharedPtr<FDecoderInput, ESPMode::ThreadSafe> AU);
+	void GetAndPrepareInputAU();
+
+	void ProcessOutput(bool bFlush = false);
+
+	bool FlushDecoder();
+	void ClearInDecoderInfos();
+	void FlushPendingImages();
+
+	EDecodeResult Decode(TSharedPtr<FDecoderInput, ESPMode::ThreadSafe> AU);
+	bool DecodeDummy(TSharedPtr<FDecoderInput, ESPMode::ThreadSafe> AU);
+
+	void PostError(int32_t ApiReturnValue, const FString& Message, uint16 Code, UEMediaError Error = UEMEDIA_ERROR_OK);
+	void LogMessage(IInfoLog::ELevel Level, const FString& Message);
+
+	void ConvertDecodedImageToNV12(TArray<uint8>& OutNV12Buffer, FIntPoint OutBufDim, const FDecodedImage& InImage);
+
+	FInstanceConfiguration Config;
+
+	FMediaEvent TerminateThreadSignal;
+	FMediaEvent FlushDecoderSignal;
+	FMediaEvent DecoderFlushedSignal;
+	bool bThreadStarted = false;
+
+	FMediaEvent ApplicationRunningSignal;
+	FMediaEvent ApplicationSuspendConfirmedSignal;
+
+	IPlayerSessionServices* PlayerSessionServices = nullptr;
+
+	TSharedPtr<IMediaRenderer, ESPMode::ThreadSafe> Renderer;
+
+    TWeakPtr<IVideoDecoderResourceDelegate, ESPMode::ThreadSafe> ResourceDelegate;
+
+	TAccessUnitQueue<TSharedPtr<FDecoderInput, ESPMode::ThreadSafe>> NextAccessUnits;
+	TSharedPtr<FDecoderInput, ESPMode::ThreadSafe> CurrentAccessUnit;
+	bool bDrainForCodecChange = false;
+
+	FCriticalSection ListenerMutex;
+	IAccessUnitBufferListener* InputBufferListener = nullptr;
+	IDecoderOutputBufferListener* ReadyBufferListener = nullptr;
+
+	FDecoderFormatInfo CurrentStreamFormatInfo;
+	TSharedPtr<ILibavDecoderVideoCommon, ESPMode::ThreadSafe> DecoderInstance;
+	TArray<TSharedPtr<FDecoderInput, ESPMode::ThreadSafe>> InDecoderInput;
+
+	int32 MaxDecodeBufferSize = 0;
+	bool bError = false;
+
+	TArray<FDecodedImage> ReadyImages;
+};
+
+/***************************************************************************************************************************************************/
+/***************************************************************************************************************************************************/
+/***************************************************************************************************************************************************/
+
+bool IVideoDecoderH265::Startup(const FParamDict& Options)
+{
+	return FVideoDecoderH265LinuxLibavcodec::Startup(Options);
+}
+
+void IVideoDecoderH265::Shutdown()
+{
+	FVideoDecoderH265LinuxLibavcodec::Shutdown();
+}
+
+bool IVideoDecoderH265::GetStreamDecodeCapability(FStreamDecodeCapability& OutResult, const FStreamDecodeCapability& InStreamParameter)
+{
+	OutResult = InStreamParameter;
+	OutResult.DecoderSupportType = ILibavDecoderH265::IsAvailable() ? FStreamDecodeCapability::ESupported::HardAndSoftware : FStreamDecodeCapability::ESupported::NotSupported;
+
+	// For now we handle only 8 bit encodings. If Main-10 profile, or only Main-10 profile compatibility is signaled we refuse this stream.
+	if (InStreamParameter.Profile > 1 || (InStreamParameter.CompatibilityFlags & 0x60000000) == 0x20000000)
+	{
+		OutResult.DecoderSupportType = FStreamDecodeCapability::ESupported::NotSupported;
+	}
+	return true;
+}
+
+IVideoDecoderH265::FInstanceConfiguration::FInstanceConfiguration()
+	: MaxDecodedFrames(8)
+{
+}
+
+IVideoDecoderH265* IVideoDecoderH265::Create()
+{
+	return new FVideoDecoderH265LinuxLibavcodec;
+}
+
+
+/***************************************************************************************************************************************************/
+/***************************************************************************************************************************************************/
+/***************************************************************************************************************************************************/
+
+//-----------------------------------------------------------------------------
+/**
+ * Decoder system startup
+ *
+ * @param Options
+ *
+ * @return
+ */
+bool FVideoDecoderH265LinuxLibavcodec::Startup(const FParamDict& Options)
+{
+	ILibavDecoder::Startup();
+	return ILibavDecoder::IsLibAvAvailable();
+}
+
+
+//-----------------------------------------------------------------------------
+/**
+ * Decoder system shutdown.
+ */
+void FVideoDecoderH265LinuxLibavcodec::Shutdown()
+{
+	ILibavDecoder::Shutdown();
+}
+
+
+//-----------------------------------------------------------------------------
+/**
+ * Constructor
+ */
+FVideoDecoderH265LinuxLibavcodec::FVideoDecoderH265LinuxLibavcodec()
+	: FMediaThread("ElectraPlayer::H265 decoder")
+{
+}
+
+
+//-----------------------------------------------------------------------------
+/**
+ * Destructor
+ */
+FVideoDecoderH265LinuxLibavcodec::~FVideoDecoderH265LinuxLibavcodec()
+{
+	Close();
+}
+
+
+//-----------------------------------------------------------------------------
+/**
+ * Sets an AU input buffer listener.
+ *
+ * @param Listener
+ */
+void FVideoDecoderH265LinuxLibavcodec::SetAUInputBufferListener(IAccessUnitBufferListener* Listener)
+{
+	FScopeLock lock(&ListenerMutex);
+	InputBufferListener = Listener;
+}
+
+
+//-----------------------------------------------------------------------------
+/**
+ * Sets a buffer-ready listener.
+ *
+ * @param Listener
+ */
+void FVideoDecoderH265LinuxLibavcodec::SetReadyBufferListener(IDecoderOutputBufferListener* Listener)
+{
+	FScopeLock lock(&ListenerMutex);
+	ReadyBufferListener = Listener;
+}
+
+
+//-----------------------------------------------------------------------------
+/**
+ * Sets the owning player's session service interface.
+ *
+ * @param InSessionServices
+ *
+ * @return
+ */
+void FVideoDecoderH265LinuxLibavcodec::SetPlayerSessionServices(IPlayerSessionServices* InSessionServices)
+{
+	PlayerSessionServices = InSessionServices;
+}
+
+
+//-----------------------------------------------------------------------------
+/**
+ * Opens a decoder instance
+ *
+ * @param InConfig
+ *
+ * @return
+ */
+void FVideoDecoderH265LinuxLibavcodec::Open(const IVideoDecoderH265::FInstanceConfiguration& InConfig)
+{
+	Config = InConfig;
+	StartThread();
+}
+
+
+//-----------------------------------------------------------------------------
+/**
+ * Closes the decoder instance.
+ */
+void FVideoDecoderH265LinuxLibavcodec::Close()
+{
+	StopThread();
+}
+
+
+//-----------------------------------------------------------------------------
+/**
+ * Drains the decoder of all enqueued input and ends it, after which the decoder must send an FDecoderMessage to the player
+ * to signal completion.
+ */
+void FVideoDecoderH265LinuxLibavcodec::DrainForCodecChange()
+{
+	bDrainForCodecChange = true;
+}
+
+
+//-----------------------------------------------------------------------------
+/**
+ * Sets a new decoder limit.
+ * As soon as a new sequence starts matching this limit the decoder will be
+ * destroyed and recreated to conserve memory.
+ * The assumption is that the data being streamed will not exceed this limit,
+ * but if it does, any access unit requiring a better decoder will take
+ * precedence and force a decoder capable of decoding it!
+ *
+ * @param MaxTier
+ * @param MaxWidth
+ * @param MaxHeight
+ * @param MaxProfile
+ * @param MaxProfileLevel
+ * @param AdditionalOptions
+ */
+void FVideoDecoderH265LinuxLibavcodec::SetMaximumDecodeCapability(int32 MaxTier, int32 MaxWidth, int32 MaxHeight, int32 MaxProfile, int32 MaxProfileLevel, const FParamDict& AdditionalOptions)
+{
+	// Not implemented
+}
+
+
+//-----------------------------------------------------------------------------
+/**
+ * Sets a new renderer.
+ *
+ * @param InRenderer
+ */
+void FVideoDecoderH265LinuxLibavcodec::SetRenderer(TSharedPtr<IMediaRenderer, ESPMode::ThreadSafe> InRenderer)
+{
+	Renderer = InRenderer;
+}
+
+void FVideoDecoderH265LinuxLibavcodec::SetResourceDelegate(const TSharedPtr<IVideoDecoderResourceDelegate, ESPMode::ThreadSafe>& InResourceDelegate)
+{
+    ResourceDelegate = InResourceDelegate;
+}
+
+//-----------------------------------------------------------------------------
+/**
+ * Creates and runs the decoder thread.
+ */
+void FVideoDecoderH265LinuxLibavcodec::StartThread()
+{
+	ThreadStart(FMediaRunnable::FStartDelegate::CreateRaw(this, &FVideoDecoderH265LinuxLibavcodec::WorkerThread));
+	bThreadStarted = true;
+}
+
+
+//-----------------------------------------------------------------------------
+/**
+ * Stops the decoder thread.
+ */
+void FVideoDecoderH265LinuxLibavcodec::StopThread()
+{
+	if (bThreadStarted)
+	{
+		TerminateThreadSignal.Signal();
+		ThreadWaitDone();
+		bThreadStarted = false;
+	}
+}
+
+
+//-----------------------------------------------------------------------------
+/**
+ * Posts an error to the session service error listeners.
+ *
+ * @param ApiReturnValue
+ * @param Message
+ * @param Code
+ * @param Error
+ */
+void FVideoDecoderH265LinuxLibavcodec::PostError(int32_t ApiReturnValue, const FString& Message, uint16 Code, UEMediaError Error)
+{
+	check(PlayerSessionServices);
+	if (PlayerSessionServices)
+	{
+		FErrorDetail err;
+		err.SetError(Error != UEMEDIA_ERROR_OK ? Error : UEMEDIA_ERROR_DETAIL);
+		err.SetFacility(Facility::EFacility::H265Decoder);
+		err.SetCode(Code);
+		err.SetMessage(Message);
+		err.SetPlatformMessage(FString::Printf(TEXT("%d (0x%08x)"), ApiReturnValue, ApiReturnValue));
+		PlayerSessionServices->PostError(err);
+	}
+}
+
+
+//-----------------------------------------------------------------------------
+/**
+ * Sends a log message to the session service log.
+ *
+ * @param Level
+ * @param Message
+ */
+void FVideoDecoderH265LinuxLibavcodec::LogMessage(IInfoLog::ELevel Level, const FString& Message)
+{
+	if (PlayerSessionServices)
+	{
+		PlayerSessionServices->PostLog(Facility::EFacility::H265Decoder, Level, Message);
+	}
+}
+
+
+//-----------------------------------------------------------------------------
+/**
+ * Create a pool of decoded images for the decoder.
+ *
+ * @return
+ */
+bool FVideoDecoderH265LinuxLibavcodec::CreateDecodedImagePool()
+{
+	check(Renderer);
+	FParamDict poolOpts;
+
+	poolOpts.Set("num_buffers", FVariantValue((int64) Config.MaxDecodedFrames));
+
+	UEMediaError Error = Renderer->CreateBufferPool(poolOpts);
+	check(Error == UEMEDIA_ERROR_OK);
+
+	MaxDecodeBufferSize = (int32) Renderer->GetBufferPoolProperties().GetValue("max_buffers").GetInt64();
+
+	if (Error != UEMEDIA_ERROR_OK)
+	{
+		PostError(0, "Failed to create image pool", ERRCODE_INTERNAL_LINUX_COULD_NOT_CREATE_IMAGE_POOL, Error);
+	}
+
+	return Error == UEMEDIA_ERROR_OK;
+}
+
+
+//-----------------------------------------------------------------------------
+/**
+ * Destroys the pool of decoded images.
+ */
+void FVideoDecoderH265LinuxLibavcodec::DestroyDecodedImagePool()
+{
+	Renderer->ReleaseBufferPool();
+}
+
+
+//-----------------------------------------------------------------------------
+/**
+ * Called to receive a new input access unit for decoding.
+ *
+ * @param AccessUnit
+ */
+void FVideoDecoderH265LinuxLibavcodec::AUdataPushAU(FAccessUnit* InAccessUnit)
+{
+	InAccessUnit->AddRef();
+
+	TSharedPtr<FDecoderInput, ESPMode::ThreadSafe> NextAU = MakeShared<FDecoderInput, ESPMode::ThreadSafe>();
+	NextAU->AccessUnit = InAccessUnit;
+	NextAccessUnits.Enqueue(MoveTemp(NextAU));
+}
+
+
+//-----------------------------------------------------------------------------
+/**
+ * "Pushes" an End Of Data marker indicating no further access units will be added.
+ */
+void FVideoDecoderH265LinuxLibavcodec::AUdataPushEOD()
+{
+	NextAccessUnits.SetEOD();
+}
+
+
+//-----------------------------------------------------------------------------
+/**
+ * Notifies the decoder that there may be further access units.
+ */
+void FVideoDecoderH265LinuxLibavcodec::AUdataClearEOD()
+{
+	NextAccessUnits.ClearEOD();
+}
+
+
+//-----------------------------------------------------------------------------
+/**
+ * Flushes the decoder and clears the input access unit buffer.
+ */
+void FVideoDecoderH265LinuxLibavcodec::AUdataFlushEverything()
+{
+	FlushDecoderSignal.Signal();
+	DecoderFlushedSignal.WaitAndReset();
+}
+
+
+//-----------------------------------------------------------------------------
+/**
+ * Create a decoder instance.
+ *
+ * @return true if successful, false on error
+ */
+bool FVideoDecoderH265LinuxLibavcodec::InternalDecoderCreate()
+{
+	InternalDecoderDestroy();
+
+	if (ILibavDecoderH265::IsAvailable())
+	{
+		TMap<FString, FVariant> Options;
+		Options.Add(TEXT("hw_priority"), FVariant(FString("vdpau;cuda;vaapi")));
+		DecoderInstance = ILibavDecoderH265::Create(this, Options);
+		if (!DecoderInstance.IsValid() || DecoderInstance->GetLastLibraryError())
+		{
+			InternalDecoderDestroy();
+			PostError(-2, "libavcodec failed to open video decoder", ERRCODE_INTERNAL_LINUX_COULD_NOT_CREATE_VIDEO_DECODER);
+			return false;
+		}
+		return true;
+	}
+	PostError(-1, "libavcodec does not support this video format", ERRCODE_INTERNAL_LINUX_VIDEO_DECODER_NOT_SUPPORTED);
+	return false;
+}
+
+
+//-----------------------------------------------------------------------------
+/**
+ * Destroys the current decoder instance.
+ */
+void FVideoDecoderH265LinuxLibavcodec::InternalDecoderDestroy()
+{
+	DecoderInstance.Reset();
+}
+
+
+//-----------------------------------------------------------------------------
+/**
+ * Notifies the buffer-ready listener that we will now be producing output data
+ * and need an output buffer.
+ * This call is not intended to create or obtain an output buffer. It is merely
+ * indicating that new output will be produced.
+ *
+ * @param bHaveOutput
+ */
+void FVideoDecoderH265LinuxLibavcodec::NotifyReadyBufferListener(bool bHaveOutput)
+{
+	if (ReadyBufferListener)
+	{
+		IDecoderOutputBufferListener::FDecodeReadyStats stats;
+		stats.MaxDecodedElementsReady = MaxDecodeBufferSize;
+		stats.NumElementsInDecoder = InDecoderInput.Num();
+		stats.bOutputStalled = !bHaveOutput;
+		stats.bEODreached = NextAccessUnits.ReachedEOD() && stats.NumDecodedElementsReady == 0 && stats.NumElementsInDecoder == 0;
+		ListenerMutex.Lock();
+		if (ReadyBufferListener)
+		{
+			ReadyBufferListener->DecoderOutputReady(stats);
+		}
+		ListenerMutex.Unlock();
+	}
+}
+
+
+//-----------------------------------------------------------------------------
+/**
+ * Prepares the passed access unit to be sent to the decoder.
+ *
+ * @param AU
+ */
+void FVideoDecoderH265LinuxLibavcodec::PrepareAU(TSharedPtr<FDecoderInput, ESPMode::ThreadSafe> AU)
+{
+	if (!AU->bHasBeenPrepared)
+	{
+		AU->bHasBeenPrepared = true;
+
+		const FStreamCodecInformation::FResolution& res = AU->AccessUnit->AUCodecData->ParsedInfo.GetResolution();
+		const FStreamCodecInformation::FAspectRatio& ar = AU->AccessUnit->AUCodecData->ParsedInfo.GetAspectRatio();
+		const FStreamCodecInformation::FCrop& crop = AU->AccessUnit->AUCodecData->ParsedInfo.GetCrop();
+		AU->Width = res.Width;
+		AU->Height = res.Height;
+		AU->TotalWidth = res.Width + crop.Left + crop.Right;
+		AU->TotalHeight = res.Height + crop.Top + crop.Bottom;
+		AU->CropLeft = crop.Left;
+		AU->CropRight = crop.Right;
+		AU->CropTop = crop.Top;
+		AU->CropBottom = crop.Bottom;
+		AU->AspectX = ar.Width ? ar.Width : 1;
+		AU->AspectY = ar.Height ? ar.Height : 1;
+
+		if (!AU->AccessUnit->bIsDummyData)
+		{
+			// Process NALUs
+			AU->bIsDiscardable = false;
+			AU->bIsIDR = AU->AccessUnit->bIsSyncSample;
+			uint32* NALU = (uint32 *)AU->AccessUnit->AUData;
+			uint32* End  = (uint32 *)Electra::AdvancePointer(NALU, AU->AccessUnit->AUSize);
+			while(NALU < End)
+			{
+				uint8 nut = *(const uint8 *)(NALU + 1);
+				nut >>= 1;
+				// IDR frame?
+				if (nut == 19 /*IDR_W_RADL*/ || nut == 20 /*IDR_N_LP*/ || nut == 21 /*CRA_NUT*/)
+				{
+					AU->bIsIDR = true;
+				}
+				// One of TRAIL_N, TSA_N, STSA_N, RADL_N, RASL_N, RSV_VCL_N10, RSV_VCL_N12 or RSV_VCL_N14 ?
+				else if (nut == 0 || nut == 2 || nut == 4 || nut == 6 || nut == 8 || nut == 10 || nut == 12 || nut == 14)
+				{
+					AU->bIsDiscardable = true;
+				}
+
+				uint32 naluLen = MEDIA_FROM_BIG_ENDIAN(*NALU) + 4;
+				*NALU = MEDIA_TO_BIG_ENDIAN(0x00000001U);
+				NALU = Electra::AdvancePointer(NALU, naluLen);
+			}
+		}
+
+		// Does this AU fall (partially) outside the range for rendering?
+		FTimeValue StartTime = AU->AccessUnit->PTS;
+		FTimeValue EndTime = AU->AccessUnit->PTS + AU->AccessUnit->Duration;
+		AU->PTS = StartTime.GetAsHNS();				// The PTS we give the decoder no matter any adjustment.
+		AU->EndPTS = EndTime.GetAsHNS();			// End PTS we need to check the PTS value returned by the decoder against.
+		AU->SequenceIndex = StartTime.GetSequenceIndex();
+		if (AU->AccessUnit->EarliestPTS.IsValid())
+		{
+			// If the end time of the AU is before the earliest render PTS we need to decode it, but not display it.
+			if (EndTime <= AU->AccessUnit->EarliestPTS)
+			{
+				StartTime.SetToInvalid();
+			}
+			else if (StartTime < AU->AccessUnit->EarliestPTS)
+			{
+				StartTime = AU->AccessUnit->EarliestPTS;
+			}
+		}
+		if (StartTime.IsValid() && AU->AccessUnit->LatestPTS.IsValid())
+		{
+			// If the start time is behind the latest render PTS we do not need to decode at all.
+			if (StartTime >= AU->AccessUnit->LatestPTS)
+			{
+				StartTime.SetToInvalid();
+				AU->bIsDiscardable = true;
+			}
+			else if (EndTime >= AU->AccessUnit->LatestPTS)
+			{
+				EndTime = AU->AccessUnit->LatestPTS;
+			}
+		}
+		AU->AdjustedPTS = StartTime;
+		AU->AdjustedDuration = EndTime - StartTime;
+		if (AU->AdjustedDuration <= FTimeValue::GetZero())
+		{
+			AU->AdjustedPTS.SetToInvalid();
+		}
+
+		// Is there codec specific data?
+		if (!AU->AccessUnit->bIsDummyData && AU->AccessUnit->AUCodecData.IsValid())
+		{
+			// Yes.
+			if (AU->bIsIDR || AU->AccessUnit->bIsSyncSample || AU->AccessUnit->bIsFirstInSequence)
+			{
+				// Have to re-allocate the AU memory to preprend the codec data
+				if (AU->AccessUnit->AUCodecData->CodecSpecificData.Num())
+				{
+					uint64 nb = AU->AccessUnit->AUSize + AU->AccessUnit->AUCodecData->CodecSpecificData.Num();
+					void* pD = AU->AccessUnit->AllocatePayloadBuffer(nb);
+					void* pP = pD;
+					FMemory::Memcpy(pP, AU->AccessUnit->AUCodecData->CodecSpecificData.GetData(), AU->AccessUnit->AUCodecData->CodecSpecificData.Num());
+					pP = Electra::AdvancePointer(pP, AU->AccessUnit->AUCodecData->CodecSpecificData.Num());
+					FMemory::Memcpy(pP, AU->AccessUnit->AUData, AU->AccessUnit->AUSize);
+					AU->AccessUnit->AdoptNewPayloadBuffer(pD, nb);
+				}
+				AU->bIsDiscardable = false;
+			}
+		}
+	}
+}
+
+
+//-----------------------------------------------------------------------------
+/**
+ * Gets an input access unit and prepares it for use.
+ */
+void FVideoDecoderH265LinuxLibavcodec::GetAndPrepareInputAU()
+{
+	if (bDrainForCodecChange)
+	{
+		return;
+	}
+
+	// Need new input?
+	if (!CurrentAccessUnit.IsValid() && InputBufferListener && NextAccessUnits.IsEmpty())
+	{
+		{
+			SCOPE_CYCLE_COUNTER(STAT_ElectraPlayer_VideoH265Decode);
+			CSV_SCOPED_TIMING_STAT(ElectraPlayer, VideoH265Decode);
+
+			IAccessUnitBufferListener::FBufferStats	stats;
+			stats.bEODSignaled = NextAccessUnits.GetEOD();
+			stats.bEODReached = NextAccessUnits.ReachedEOD();
+			ListenerMutex.Lock();
+			if (InputBufferListener)
+			{
+				InputBufferListener->DecoderInputNeeded(stats);
+			}
+			ListenerMutex.Unlock();
+		}
+
+		// Get the AU to be decoded if one is there.
+		bool bHaveData = NextAccessUnits.Wait(1000);
+		if (bHaveData)
+		{
+			SCOPE_CYCLE_COUNTER(STAT_ElectraPlayer_VideoH265Decode);
+			CSV_SCOPED_TIMING_STAT(ElectraPlayer, VideoH265Decode);
+
+			bool bOk = NextAccessUnits.Dequeue(CurrentAccessUnit);
+			MEDIA_UNUSED_VAR(bOk);
+			check(bOk);
+
+			PrepareAU(CurrentAccessUnit);
+		}
+	}
+}
+
+
+//-----------------------------------------------------------------------------
+/**
+ * Gets an output buffer from the video renderer.
+ *
+ * @param RenderOutputBuffer
+ *
+ * @return true if successful, false on error
+ */
+bool FVideoDecoderH265LinuxLibavcodec::AcquireOutputBuffer(IMediaRenderer::IBuffer*& RenderOutputBuffer)
+{
+	RenderOutputBuffer = nullptr;
+	FParamDict BufferAcquireOptions;
+	while(!TerminateThreadSignal.IsSignaled() && !FlushDecoderSignal.IsSignaled())
+	{
+		UEMediaError bufResult = Renderer->AcquireBuffer(RenderOutputBuffer, 0, BufferAcquireOptions);
+		check(bufResult == UEMEDIA_ERROR_OK || bufResult == UEMEDIA_ERROR_INSUFFICIENT_DATA);
+		if (bufResult != UEMEDIA_ERROR_OK && bufResult != UEMEDIA_ERROR_INSUFFICIENT_DATA)
+		{
+			PostError(0, "Failed to acquire sample buffer", ERRCODE_INTERNAL_LINUX_COULD_NOT_GET_OUTPUT_BUFFER, bufResult);
+			return false;
+		}
+		bool bHaveAvailSmpBlk = RenderOutputBuffer != nullptr;
+		NotifyReadyBufferListener(bHaveAvailSmpBlk);
+		if (bHaveAvailSmpBlk)
+		{
+			break;
+		}
+		else
+		{
+			// No available buffer. Sleep for a bit. Can't sleep on a signal since we have to check two: abort and flush
+			FMediaRunnable::SleepMilliseconds(5);
+		}
+	}
+	return true;
+}
+
+
+//-----------------------------------------------------------------------------
+/**
+ * Flushes the decoder and passes all pending frames to the renderer.
+ *
+ * @return true if successful, false otherwise.
+ */
+bool FVideoDecoderH265LinuxLibavcodec::FlushDecoder()
+{
+	if (DecoderInstance.IsValid())
+	{
+		ILibavDecoder::EDecoderError DecErr = DecoderInstance->SendEndOfData();
+		check(DecErr == ILibavDecoder::EDecoderError::None); (void)DecErr;
+		ProcessOutput(true);
+	}
+
+	ClearInDecoderInfos();
+	return true;
+}
+
+
+//-----------------------------------------------------------------------------
+/**
+ * Clears out any leftover decoder source infos that were not picked by decoded frames.
+ * This must be called after having flushed the decoder to ensure no delayed decode
+ * callbacks will arrive.
+ */
+void FVideoDecoderH265LinuxLibavcodec::ClearInDecoderInfos()
+{
+	InDecoderInput.Empty();
+}
+
+
+//-----------------------------------------------------------------------------
+/**
+ * Creates a dummy output image that is not to be displayed and has no image data.
+ * Dummy access units are created when stream data is missing to ensure the data
+ * pipeline does not run dry and exhibits no gaps in the timeline.
+ *
+ * @param AU
+ *
+ * @return true if successful, false otherwise
+ */
+bool FVideoDecoderH265LinuxLibavcodec::DecodeDummy(TSharedPtr<FDecoderInput, ESPMode::ThreadSafe> AU)
+{
+	if (AU.IsValid() && AU->AdjustedPTS.IsValid())
+	{
+		FDecodedImage NextImage;
+		NextImage.SourceInfo = AU;
+		ReadyImages.Add(NextImage);
+	}
+    return true;
+}
+
+//-----------------------------------------------------------------------------
+/**
+ * Flushes all images not passed to the renderer from our list.
+ */
+void FVideoDecoderH265LinuxLibavcodec::FlushPendingImages()
+{
+	ReadyImages.Empty();
+}
+
+//-----------------------------------------------------------------------------
+/**
+ * Checks if the codec specific data has changed.
+ *
+ * @param AU
+ *
+ * @return false if the format is still the same, true if it has changed.
+ */
+bool FVideoDecoderH265LinuxLibavcodec::FDecoderFormatInfo::IsDifferentFrom(TSharedPtr<FDecoderInput, ESPMode::ThreadSafe> AU)
+{
+	if (AU->AccessUnit->AUCodecData.IsValid() && AU->AccessUnit->AUCodecData.Get() != CurrentCodecData.Get())
+	{
+		// Pointers are different. Is the content too?
+		bool bDifferent = !CurrentCodecData.IsValid() || (CurrentCodecData.IsValid() && AU->AccessUnit->AUCodecData->CodecSpecificData != CurrentCodecData->CodecSpecificData);
+		CurrentCodecData = AU->AccessUnit->AUCodecData;
+		return bDifferent;
+	}
+	return false;
+}
+
+
+//-----------------------------------------------------------------------------
+/**
+ * Sends an access unit to the decoder for decoding.
+ *
+ * @param AU
+ *
+ * @return
+ */
+FVideoDecoderH265LinuxLibavcodec::EDecodeResult FVideoDecoderH265LinuxLibavcodec::Decode(TSharedPtr<FDecoderInput, ESPMode::ThreadSafe> AU)
+{
+	if (!DecoderInstance.IsValid())
+	{
+		return EDecodeResult::Fail;
+	}
+
+	bool bAddInput = true;
+	// Decode
+	ILibavDecoderVideoCommon::FInputAU DecAU;
+	DecAU.Data = AU->AccessUnit->AUData;
+	DecAU.DataSize = (int32) AU->AccessUnit->AUSize;
+	DecAU.DTS = AU->AccessUnit->DTS.GetAsHNS();
+	DecAU.PTS = AU->AccessUnit->PTS.GetAsHNS();
+	DecAU.Duration = AU->AccessUnit->Duration.GetAsHNS();
+	DecAU.UserValue = AU->PTS;
+	DecAU.Flags = 0;
+	if (!AU->AdjustedPTS.IsValid())
+	{
+		DecAU.Flags |= ILibavDecoderVideoCommon::FInputAU::EFlags::EVidAUFlag_DoNotOutput;
+		bAddInput = false;
+	}
+	if (bAddInput)
+	{
+		InDecoderInput.Add(AU);
+	}
+
+	ILibavDecoder::EDecoderError DecErr = DecoderInstance->DecodeAccessUnit(DecAU);
+	if (DecErr == ILibavDecoder::EDecoderError::None)
+	{
+		// Ok.
+		return EDecodeResult::Ok;
+	}
+	else
+	{
+		InDecoderInput.RemoveSingle(AU);
+		PostError(DecoderInstance->GetLastLibraryError(), FString::Printf(TEXT("Failed to decode video frame (%d)"),(int)DecErr), ERRCODE_INTERNAL_LINUX_FAILED_TO_DECODE_VIDEO);
+		return EDecodeResult::Fail;
+	}
+}
+
+
+//-----------------------------------------------------------------------------
+/**
+ * H265 video decoder main threaded decode loop
+ */
+void FVideoDecoderH265LinuxLibavcodec::WorkerThread()
+{
+	LLM_SCOPE(ELLMTag::ElectraPlayer);
+
+	ApplicationRunningSignal.Signal();
+	ApplicationSuspendConfirmedSignal.Reset();
+
+	TOptional<int64> SequenceIndex;
+	bool bDone  = false;
+    bool bGotLastSequenceAU = false;
+	bool bInDummyDecodeMode = false;
+
+	bError = false;
+
+	// Create decoded image pool.
+	if (!CreateDecodedImagePool())
+	{
+		bError = true;
+	}
+
+	while(!TerminateThreadSignal.IsSignaled())
+	{
+		// If in background, wait until we get activated again.
+		if (!ApplicationRunningSignal.IsSignaled())
+		{
+			UE_LOG(LogElectraPlayer, Log, TEXT("FVideoDecoderH265LinuxLibavcodec(%p): OnSuspending"), this);
+			ApplicationSuspendConfirmedSignal.Signal();
+			while(!ApplicationRunningSignal.WaitTimeout(100 * 1000) && !TerminateThreadSignal.IsSignaled())
+			{
+			}
+			UE_LOG(LogElectraPlayer, Log, TEXT("FVideoDecoderH265LinuxLibavcodec(%p): OnResuming"), this);
+		}
+
+		if (!bDrainForCodecChange)
+		{
+			GetAndPrepareInputAU();
+
+			bool bHaveData = CurrentAccessUnit.IsValid();
+			bool bTooManyImagesWaiting = false;
+			if (bHaveData)
+			{
+				// When there is data, even and especially after a previous EOD, we are no longer done and idling.
+				bDone = false;
+
+				// If there is too much output still pending to be delivered to the renderer
+				bTooManyImagesWaiting = ReadyImages.Num() > MaxPendingOutputImages;
+				if (bTooManyImagesWaiting)
+				{
+					while(!TerminateThreadSignal.IsSignaled() && !FlushDecoderSignal.IsSignaled())
+					{
+						ProcessOutput();
+						if ((bTooManyImagesWaiting = ReadyImages.Num() > MaxPendingOutputImages) == false)
+						{
+							break;
+						}
+						FMediaRunnable::SleepMilliseconds(10);
+					}
+				}
+			}
+
+			if (bHaveData && !bTooManyImagesWaiting)
+			{
+				// When there is data, even and especially after a previous EOD, we are no longer done and idling.
+				bDone = false;
+				if (!CurrentAccessUnit->AccessUnit->bIsDummyData)
+				{
+					bInDummyDecodeMode = false;
+
+					bool bStreamFormatChanged = CurrentStreamFormatInfo.IsDifferentFrom(CurrentAccessUnit) || bGotLastSequenceAU;
+					bool bNeedNewDecoder = false;
+
+					if (!SequenceIndex.IsSet())
+					{
+						SequenceIndex = CurrentAccessUnit->AccessUnit->PTS.GetSequenceIndex();
+					}
+					bool bSequenceIndexChanged = SequenceIndex.GetValue() != CurrentAccessUnit->AccessUnit->PTS.GetSequenceIndex();
+					SequenceIndex = CurrentAccessUnit->AccessUnit->PTS.GetSequenceIndex();
+
+					bGotLastSequenceAU = CurrentAccessUnit->AccessUnit->bIsLastInPeriod;
+					if (bStreamFormatChanged || bSequenceIndexChanged || !DecoderInstance.IsValid())
+					{
+						SCOPE_CYCLE_COUNTER(STAT_ElectraPlayer_VideoH265Decode);
+						CSV_SCOPED_TIMING_STAT(ElectraPlayer, VideoH265Decode);
+						bNeedNewDecoder |= bSequenceIndexChanged || !DecoderInstance.IsValid();
+						if (bNeedNewDecoder)
+						{
+							FlushDecoder();
+							InternalDecoderDestroy();
+							if (!bError && !InternalDecoderCreate())
+							{
+								bError = true;
+							}
+						}
+					}
+
+					// If this AU falls outside the range where it is to be rendered and it is also discardable
+					// we do not need to concern ourselves with it at all.
+					if (CurrentAccessUnit->bIsDiscardable && !CurrentAccessUnit->AdjustedPTS.IsValid())
+					{
+						CurrentAccessUnit.Reset();
+						continue;
+					}
+
+					// Decode
+					if (!bError && DecoderInstance.IsValid())
+					{
+						EDecodeResult DecRes = Decode(CurrentAccessUnit);
+
+						// Process any output we might have pending
+						ProcessOutput();
+
+						// Decode ok?
+						if (DecRes == EDecodeResult::Ok)
+						{
+							CurrentAccessUnit.Reset();
+						}
+						// Decode failed
+						else
+						{
+							bError = true;
+						}
+					}
+				}
+				else
+				{
+					if (!bInDummyDecodeMode)
+					{
+						bInDummyDecodeMode = true;
+						FlushDecoder();
+					}
+
+					if (!DecodeDummy(CurrentAccessUnit))
+					{
+						bError = true;
+					}
+					// DecodeDummy() went through most of the regular path, but has returned the output buffer immediately
+					// and can thus always get a new one with no waiting. To avoid draining the player buffer by consuming
+					// the dummy AUs at rapid pace we put ourselves to sleep for the duration the AU was supposed to last.
+					FMediaRunnable::SleepMicroseconds(CurrentAccessUnit->AdjustedDuration.GetAsMicroseconds());
+
+					ProcessOutput();
+					CurrentAccessUnit.Reset();
+				}
+			}
+			else
+			{
+				ProcessOutput();
+
+				// No data. Is the buffer at EOD?
+				if (!CurrentAccessUnit.IsValid() && NextAccessUnits.ReachedEOD())
+				{
+					NotifyReadyBufferListener(true);
+					// Are we done yet?
+					if (!bDone && !bError)
+					{
+						bError = !FlushDecoder();
+					}
+					bDone = true;
+					FMediaRunnable::SleepMilliseconds(10);
+					ProcessOutput(true);
+				}
+			}
+		}
+		else
+		{
+			ProcessOutput();
+			NotifyReadyBufferListener(true);
+			// Are we done yet?
+			if (!bDone && !bError)
+			{
+				bError = !FlushDecoder();
+			}
+			bDone = true;
+			ProcessOutput(true);
+			break;
+		}
+
+		// Flush?
+		if (FlushDecoderSignal.IsSignaled())
+		{
+			SCOPE_CYCLE_COUNTER(STAT_ElectraPlayer_VideoH265Decode);
+			CSV_SCOPED_TIMING_STAT(ElectraPlayer, VideoH265Decode);
+
+			// Have to destroy the decoder!
+			InternalDecoderDestroy();
+			FlushPendingImages();
+			ClearInDecoderInfos();
+			NextAccessUnits.Empty();
+			SequenceIndex.Reset();
+			CurrentAccessUnit.Reset();
+
+			FlushDecoderSignal.Reset();
+			DecoderFlushedSignal.Signal();
+
+			// Reset done state.
+			bDone = false;
+		}
+	}
+
+	InternalDecoderDestroy();
+	FlushPendingImages();
+	DestroyDecodedImagePool();
+	ClearInDecoderInfos();
+	NextAccessUnits.Empty();
+	CurrentAccessUnit.Reset();
+
+	if (bDrainForCodecChange)
+	{
+		// Notify the player that we have finished draining.
+		PlayerSessionServices->SendMessageToPlayer(FDecoderMessage::Create(FDecoderMessage::EReason::DrainingFinished, this, EStreamType::Video, FStreamCodecInformation::ECodec::H265));
+		// We need to wait to get terminated. Also check if flushing is requested and acknowledge if it is.
+		while(!TerminateThreadSignal.IsSignaled())
+		{
+			if (FlushDecoderSignal.WaitTimeoutAndReset(1000 * 10))
+			{
+				DecoderFlushedSignal.Signal();
+			}
+		}
+	}
+}
+
+//-----------------------------------------------------------------------------
+/**
+ */
+void FVideoDecoderH265LinuxLibavcodec::ProcessOutput(bool bFlush)
+{
+	SCOPE_CYCLE_COUNTER(STAT_ElectraPlayer_VideoH265ConvertOutput);
+	CSV_SCOPED_TIMING_STAT(ElectraPlayer, VideoH265ConvertOutput);
+
+	// Pull output from decoder
+	while(DecoderInstance.IsValid())
+	{
+		FDecodedImage NextImage;
+
+		// Output available?
+		if (DecoderInstance->HaveOutput(NextImage.ImageInfo) != ILibavDecoderVideoCommon::EOutputStatus::Available)
+		{
+			break;
+		}
+		// Get the output.
+		NextImage.DecodedImage = DecoderInstance->GetOutput();
+		if (!NextImage.DecodedImage.IsValid())
+		{
+			break;
+		}
+		// Find the corresponding input source.
+		int32 NumCurrentDecodeInputs = InDecoderInput.Num();
+		for(int32 i=0; i<NumCurrentDecodeInputs; ++i)
+		{
+			if (InDecoderInput[i]->PTS == NextImage.ImageInfo.UserValue)
+			{
+				NextImage.SourceInfo = InDecoderInput[i];
+				InDecoderInput.RemoveSingle(NextImage.SourceInfo);
+				break;
+			}
+		}
+
+		if (!NextImage.SourceInfo.IsValid())
+		{
+			LogMessage(IInfoLog::ELevel::Warning, FString::Printf(TEXT("FVideoDecoderH265::ProcessOutput(): No source info found for decoded PTS %lld in %d pending infos"), (long long int)NextImage.ImageInfo.UserValue, NumCurrentDecodeInputs));
+			break;
+		}
+
+		// For safety we remove input that is older than what we got right now.
+		// This is to ensure that if libav does not return decoded output to us for some reason
+		// (like a corrupted frame or such) that we do not sit on stale input.
+		if (InDecoderInput.Num() > MaxPendingOutputImages)
+		{
+			for(int32 i=0; i<InDecoderInput.Num(); ++i)
+			{
+				if (InDecoderInput[i]->PTS < NextImage.ImageInfo.UserValue)
+				{
+					LogMessage(IInfoLog::ELevel::Warning, FString::Printf(TEXT("FVideoDecoderH265::ProcessOutput(): Removing stale input PTS %lld for output PTS %lld"), (long long int)InDecoderInput[i]->PTS, (long long int)NextImage.ImageInfo.UserValue));
+					InDecoderInput.RemoveAt(i);
+					--i;
+				}
+			}
+		}
+
+		ReadyImages.Add(NextImage);
+	}
+	ReadyImages.Sort();
+
+	// Send frames out to the renderer.
+	IMediaRenderer::IBuffer* RenderOutputBuffer = nullptr;
+	while(!TerminateThreadSignal.IsSignaled() && !FlushDecoderSignal.IsSignaled())
+	{
+		if (ReadyImages.Num() == 0)
+		{
+			break;
+		}
+
+		// When we are not flushing out the remaining frames and the renderer cannot take on
+		// anything new we leave.
+		if (!bFlush && !Renderer->CanReceiveOutputFrames(1))
+		{
+			NotifyReadyBufferListener(false);
+			break;
+		}
+
+		// Need a new output buffer?
+		if (!RenderOutputBuffer)
+		{
+			// If we failed to get a new buffer, leave the loop.
+			if (!AcquireOutputBuffer(RenderOutputBuffer))
+			{
+				break;
+			}
+		}
+
+		// Can we send the image to the renderer?
+		bool bHaveAvailableImage = RenderOutputBuffer != nullptr;
+		if (bFlush && !Renderer->CanReceiveOutputFrames(1))
+		{
+			bHaveAvailableImage = false;
+		}
+		NotifyReadyBufferListener(bHaveAvailableImage);
+		if (bHaveAvailableImage)
+		{
+			FDecodedImage NextImage(ReadyImages[0]);
+			ReadyImages.RemoveAt(0);
+
+			FParamDict* OutputBufferSampleProperties = new FParamDict();
+			FIntPoint Dim(NextImage.SourceInfo->Width, NextImage.SourceInfo->Height);
+
+			if (NextImage.DecodedImage.IsValid())
+			{
+				int32 ax = NextImage.ImageInfo.AspectNum;
+				int32 ay = NextImage.ImageInfo.AspectDenom;
+				Dim.X = NextImage.ImageInfo.Width;
+				Dim.Y = NextImage.ImageInfo.Height;
+
+				OutputBufferSampleProperties->Set("width", FVariantValue((int64) NextImage.ImageInfo.Width));
+				OutputBufferSampleProperties->Set("height", FVariantValue((int64) NextImage.ImageInfo.Height));
+				OutputBufferSampleProperties->Set("crop_left", FVariantValue((int64) 0));
+				OutputBufferSampleProperties->Set("crop_right", FVariantValue((int64) 0));
+				OutputBufferSampleProperties->Set("crop_top", FVariantValue((int64) 0));
+				OutputBufferSampleProperties->Set("crop_bottom", FVariantValue((int64) 0));
+				OutputBufferSampleProperties->Set("aspect_ratio", FVariantValue((double)ax / (double)ay));
+				OutputBufferSampleProperties->Set("aspect_w", FVariantValue((int64) ax));
+				OutputBufferSampleProperties->Set("aspect_h", FVariantValue((int64) ay));
+				OutputBufferSampleProperties->Set("fps_num", FVariantValue((int64) 0 ));
+				OutputBufferSampleProperties->Set("fps_denom", FVariantValue((int64) 0 ));
+				OutputBufferSampleProperties->Set("pixelfmt", FVariantValue((int64)EPixelFormat::PF_NV12));
+			}
+			else
+			{
+				OutputBufferSampleProperties->Set("is_dummy", FVariantValue(true));
+			}
+			OutputBufferSampleProperties->Set("pts", FVariantValue(NextImage.SourceInfo->AdjustedPTS));
+			OutputBufferSampleProperties->Set("duration", FVariantValue(NextImage.SourceInfo->AdjustedDuration));
+
+			bool bRender = NextImage.SourceInfo->AdjustedPTS.IsValid();
+			TSharedPtr<FElectraPlayerVideoDecoderOutputLinux, ESPMode::ThreadSafe> DecoderOutput = RenderOutputBuffer->GetBufferProperties().GetValue("texture").GetSharedPointer<FElectraPlayerVideoDecoderOutputLinux>();
+			FIntPoint BufferDim(NextImage.ImageInfo.Planes[0].Width, NextImage.ImageInfo.Planes[0].Height);
+			if (DecoderOutput->InitializeForBuffer(BufferDim, EPixelFormat::PF_NV12, OutputBufferSampleProperties))
+			{
+				TArray<uint8>& ImgBuf = DecoderOutput->GetMutableBuffer();
+				ConvertDecodedImageToNV12(ImgBuf, DecoderOutput->GetBufferDimensions(), NextImage);
+				// Have the decoder output keep a reference to the decoded image in case it will
+				// shared data with it at some point.
+				DecoderOutput->SetDecodedImage(NextImage.DecodedImage);
+			}
+
+			// Return the buffer to the renderer.
+			Renderer->ReturnBuffer(RenderOutputBuffer, bRender, *OutputBufferSampleProperties);
+			RenderOutputBuffer = nullptr;
+		}
+		else
+		{
+			// Since we need to check the abort and the flush signal we sleep for a bit and check again.
+			FMediaRunnable::SleepMilliseconds(10);
+		}
+	}
+	if (RenderOutputBuffer)
+	{
+		FParamDict None;
+		Renderer->ReturnBuffer(RenderOutputBuffer, false, None);
+		RenderOutputBuffer = nullptr;
+	}
+}
+
+
+void FVideoDecoderH265LinuxLibavcodec::ConvertDecodedImageToNV12(TArray<uint8>& OutNV12Buffer, FIntPoint OutBufDim, const FDecodedImage& InImage)
+{
+	if (OutNV12Buffer.GetData() && InImage.DecodedImage.IsValid())
+	{
+		if (InImage.ImageInfo.NumPlanes == 2 &&
+			InImage.ImageInfo.Planes[0].Content == ILibavDecoderVideoCommon::FPlaneInfo::EContent::Luma &&
+			InImage.ImageInfo.Planes[1].Content == ILibavDecoderVideoCommon::FPlaneInfo::EContent::ChromaUV)
+		{
+			const int32 w = InImage.ImageInfo.Planes[0].Width;
+			const int32 h = InImage.ImageInfo.Planes[0].Height;
+			const int32 aw = ((w + 1) / 2) * 2;
+			const int32 ah = ((h + 1) / 2) * 2;
+			uint8* DstY = OutNV12Buffer.GetData();
+			uint8* DstUV = DstY + aw * ah;
+			const uint8* SrcY = (const uint8*)InImage.ImageInfo.Planes[0].Address;
+			const uint8* SrcUV = (const uint8*)InImage.ImageInfo.Planes[1].Address;
+			// To simplify the conversion we require the output buffer to have the dimension of the planes.
+			check(OutBufDim.X == aw && OutBufDim.Y == ah);
+			if (!SrcY || !SrcUV || OutBufDim.X != aw || OutBufDim.Y != ah)
+			{
+				return;
+			}
+			if ((w & 1) == 0)
+			{
+				FMemory::Memcpy(DstY, SrcY, w*h);
+				FMemory::Memcpy(DstUV, SrcUV, w*h/2);
+			}
+			else
+			{
+				for(int32 y=0; y<h; ++y)
+				{
+					FMemory::Memcpy(DstY, SrcY, w);
+					DstY += aw;
+					SrcY += w;
+				}
+				for(int32 y=0; y<h/2; ++y)
+				{
+					FMemory::Memcpy(DstUV, SrcUV, w);
+					DstUV += aw;
+					SrcUV += w;
+				}
+			}
+		}
+		else if (InImage.ImageInfo.NumPlanes == 3 &&
+				 InImage.ImageInfo.Planes[0].Content == ILibavDecoderVideoCommon::FPlaneInfo::EContent::Luma &&
+				 InImage.ImageInfo.Planes[1].Content == ILibavDecoderVideoCommon::FPlaneInfo::EContent::ChromaU &&
+				 InImage.ImageInfo.Planes[2].Content == ILibavDecoderVideoCommon::FPlaneInfo::EContent::ChromaV)
+		{
+			const int32 w = InImage.ImageInfo.Planes[0].Width;
+			const int32 h = InImage.ImageInfo.Planes[0].Height;
+			const int32 aw = ((w + 1) / 2) * 2;
+			const int32 ah = ((h + 1) / 2) * 2;
+			uint8* DstY = OutNV12Buffer.GetData();
+			uint8* DstUV = DstY + aw * ah;
+			const uint8* SrcY = (const uint8*)InImage.ImageInfo.Planes[0].Address;
+			const uint8* SrcU = (const uint8*)InImage.ImageInfo.Planes[1].Address;
+			const uint8* SrcV = (const uint8*)InImage.ImageInfo.Planes[2].Address;
+			// To simplify the conversion we require the output buffer to have the dimension of the planes.
+			check(OutBufDim.X == aw && OutBufDim.Y == ah);
+			if (!SrcY || !SrcU || !SrcV || OutBufDim.X != aw || OutBufDim.Y != ah)
+			{
+				return;
+			}
+			if ((w & 1) == 0)
+			{
+				FMemory::Memcpy(DstY, SrcY, w*h);
+				for(int32 i=0, iMax=w*h/4; i<iMax; ++i)
+				{
+					*DstUV++ = *SrcU++;
+					*DstUV++ = *SrcV++;
+				}
+			}
+			else
+			{
+				for(int32 y=0; y<h; ++y)
+				{
+					FMemory::Memcpy(DstY, SrcY, w);
+					DstY += aw;
+					SrcY += w;
+				}
+				int32 padUV = (aw - w) * 2;
+				for(int32 v=0; v<h/2; ++v)
+				{
+					for(int32 u=0; u<w/2; ++u)
+					{
+						*DstUV++ = *SrcU++;
+						*DstUV++ = *SrcV++;
+					}
+					DstUV += padUV;
+				}
+			}
+		}
+	}
+}
+
+
+} // namespace Electra
diff --git a/Engine/Plugins/Media/ElectraPlayer/Source/ElectraPlayerRuntime/Private/Runtime/Decoder/Linux/VideoDecoderResourceDelegate.h b/Engine/Plugins/Media/ElectraPlayer/Source/ElectraPlayerRuntime/Private/Runtime/Decoder/Linux/VideoDecoderResourceDelegate.h
new file mode 100644
--- /dev/null
+++ b/Engine/Plugins/Media/ElectraPlayer/Source/ElectraPlayerRuntime/Private/Runtime/Decoder/Linux/VideoDecoderResourceDelegate.h
@@ -0,0 +1,13 @@
+// Copyright Epic Games, Inc. All Rights Reserved.
+
+#pragma once
+
+namespace Electra {
+
+class IVideoDecoderResourceDelegate
+{
+public:
+	virtual ~IVideoDecoderResourceDelegate() {}
+};
+
+}
diff --git a/Engine/Plugins/Media/ElectraPlayer/Source/ElectraPlayerRuntime/Private/Runtime/Platforms/Linux/ElectraPlayer_Linux.cpp b/Engine/Plugins/Media/ElectraPlayer/Source/ElectraPlayerRuntime/Private/Runtime/Platforms/Linux/ElectraPlayer_Linux.cpp
new file mode 100644
--- /dev/null
+++ b/Engine/Plugins/Media/ElectraPlayer/Source/ElectraPlayerRuntime/Private/Runtime/Platforms/Linux/ElectraPlayer_Linux.cpp
@@ -0,0 +1,65 @@
+// Copyright Epic Games, Inc. All Rights Reserved.
+
+#include "CoreTypes.h"
+
+#include "IElectraPlayerInterface.h"
+#include "ElectraPlayerPrivate.h"
+#include "ElectraPlayer.h"
+#include "ParameterDictionary.h"
+
+// ----------------------------------------------------------------------------------------------------------------------
+// ----------------------------------------------------------------------------------------------------------------------
+
+namespace Electra
+{
+
+/**
+ * Early startup during module load process
+ */
+bool PlatformEarlyStartup()
+{
+	return true;
+}
+
+/**
+ * Any platform specific memory hooks
+ */
+bool PlatformMemorySetup()
+{
+	return true;
+}
+
+/**
+ * Shutdown of module
+ */
+bool PlatformShutdown()
+{
+	return true;
+}
+
+TSharedPtr<IVideoDecoderResourceDelegate, ESPMode::ThreadSafe> PlatformCreateVideoDecoderResourceDelegate(const TWeakPtr<IElectraPlayerAdapterDelegate, ESPMode::ThreadSafe>& AdapterDelegate)
+{
+	return nullptr;
+}
+
+} //namespace Electra
+
+// ----------------------------------------------------------------------------------------------------------------------
+// ----------------------------------------------------------------------------------------------------------------------
+
+/**
+ * Public call for pre-use initialization after module is loaded, allowing for user passed parameters
+ */
+bool FElectraPlayerPlatform::StartupPlatformResources(const Electra::FParamDict& Params)
+{
+	return true;
+}
+
+
+void FElectraPlayer::PlatformNotifyOfOptionChange()
+{
+}
+
+void FElectraPlayer::PlatformSuspendOrResumeDecoders(bool /*bSuspend*/)
+{
+}
diff --git a/Engine/Plugins/Media/ElectraPlayer/Source/ElectraPlayerRuntime/Private/Runtime/Platforms/Linux/PlayerRuntime_Linux.cpp b/Engine/Plugins/Media/ElectraPlayer/Source/ElectraPlayerRuntime/Private/Runtime/Platforms/Linux/PlayerRuntime_Linux.cpp
new file mode 100644
--- /dev/null
+++ b/Engine/Plugins/Media/ElectraPlayer/Source/ElectraPlayerRuntime/Private/Runtime/Platforms/Linux/PlayerRuntime_Linux.cpp
@@ -0,0 +1,14 @@
+// Copyright Epic Games, Inc. All Rights Reserved.
+
+#include "SynchronizedClock.h"
+
+namespace Electra
+{
+
+int64 MEDIAutcTime::CurrentMSec()
+{
+	FTimespan localTime(FDateTime::UtcNow().GetTicks());
+	return (int64)localTime.GetTotalMilliseconds();
+};
+
+}
diff --git a/Engine/Plugins/Media/ElectraPlayer/Source/ElectraPlayerRuntime/Public/Linux/ElectraPlayerResourceDelegate.h b/Engine/Plugins/Media/ElectraPlayer/Source/ElectraPlayerRuntime/Public/Linux/ElectraPlayerResourceDelegate.h
new file mode 100644
--- /dev/null
+++ b/Engine/Plugins/Media/ElectraPlayer/Source/ElectraPlayerRuntime/Public/Linux/ElectraPlayerResourceDelegate.h
@@ -0,0 +1,8 @@
+// Copyright Epic Games, Inc. All Rights Reserved.
+#pragma once
+
+class IElectraPlayerResourceDelegate
+{
+public:
+	virtual ~IElectraPlayerResourceDelegate() {}
+};
diff --git a/Engine/Plugins/Media/ElectraPlayer/Source/ElectraPlayerRuntime/Public/Unix/ElectraPlayerResourceDelegate.h b/Engine/Plugins/Media/ElectraPlayer/Source/ElectraPlayerRuntime/Public/Unix/ElectraPlayerResourceDelegate.h
--- a/Engine/Plugins/Media/ElectraPlayer/Source/ElectraPlayerRuntime/Public/Unix/ElectraPlayerResourceDelegate.h
+++ b/Engine/Plugins/Media/ElectraPlayer/Source/ElectraPlayerRuntime/Public/Unix/ElectraPlayerResourceDelegate.h
@@ -1,8 +1,4 @@
 // Copyright Epic Games, Inc. All Rights Reserved.
 #pragma once
 
-class IElectraPlayerResourceDelegate
-{
-public:
-	virtual ~IElectraPlayerResourceDelegate() {}
-};
+#include "Linux/ElectraPlayerResourceDelegate.h"
diff --git a/Engine/Plugins/Media/ElectraSubtitles/ElectraSubtitles.uplugin b/Engine/Plugins/Media/ElectraSubtitles/ElectraSubtitles.uplugin
--- a/Engine/Plugins/Media/ElectraSubtitles/ElectraSubtitles.uplugin
+++ b/Engine/Plugins/Media/ElectraSubtitles/ElectraSubtitles.uplugin
@@ -19,7 +19,7 @@
 			"Name": "ElectraSubtitles",
 			"Type": "RuntimeNoCommandlet",
 			"LoadingPhase": "PostEngineInit",
-			"PlatformAllowList": [ "Win64", "Mac", "IOS", "Android" ],
+			"PlatformAllowList": [ "Win64", "Mac", "IOS", "Android", "Linux" ],
 			"TargetDenyList": [ "Server" ]
 		}
 	],
diff --git a/Engine/Plugins/Media/ElectraUtil/ElectraUtil.uplugin b/Engine/Plugins/Media/ElectraUtil/ElectraUtil.uplugin
--- a/Engine/Plugins/Media/ElectraUtil/ElectraUtil.uplugin
+++ b/Engine/Plugins/Media/ElectraUtil/ElectraUtil.uplugin
@@ -19,21 +19,21 @@
 			"Name": "ElectraBase",
 			"Type": "RuntimeNoCommandlet",
 			"LoadingPhase": "PostEngineInit",
-			"PlatformAllowList": [ "Win64", "Mac", "IOS", "Android" ],
+			"PlatformAllowList": [ "Win64", "Mac", "IOS", "Android", "Linux" ],
 			"TargetDenyList": [ "Server" ]
 		},
 		{
 			"Name": "ElectraSamples",
 			"Type": "RuntimeNoCommandlet",
 			"LoadingPhase": "PostEngineInit",
-			"PlatformAllowList": [ "Win64", "Mac", "IOS", "Android" ],
+			"PlatformAllowList": [ "Win64", "Mac", "IOS", "Android", "Linux" ],
 			"TargetDenyList": [ "Server" ]
 		},
 		{
 			"Name": "ElectraHTTPStream",
 			"Type": "RuntimeNoCommandlet",
 			"LoadingPhase": "PostEngineInit",
-			"PlatformAllowList": [ "Win64", "Mac", "IOS", "Android" ],
+			"PlatformAllowList": [ "Win64", "Mac", "IOS", "Android", "Linux" ],
 			"TargetDenyList": [ "Server" ]
 		}
 	]
diff --git a/Engine/Plugins/Media/ElectraUtil/Source/ElectraBase/Public/Linux/MediaVideoDecoderOutputLinux.h b/Engine/Plugins/Media/ElectraUtil/Source/ElectraBase/Public/Linux/MediaVideoDecoderOutputLinux.h
new file mode 100644
--- /dev/null
+++ b/Engine/Plugins/Media/ElectraUtil/Source/ElectraBase/Public/Linux/MediaVideoDecoderOutputLinux.h
@@ -0,0 +1,15 @@
+// Copyright Epic Games, Inc. All Rights Reserved.
+#pragma once
+
+#include "MediaVideoDecoderOutput.h"
+#include "Containers/Array.h"
+
+class FVideoDecoderOutputLinux : public FVideoDecoderOutput
+{
+public:
+	~FVideoDecoderOutputLinux() = default;
+
+	virtual const TArray<uint8>& GetBuffer() const = 0;
+	virtual FIntPoint GetBufferDimensions() const = 0;
+	virtual uint32 GetStride() const = 0;
+};
diff --git a/Engine/Plugins/Media/ElectraUtil/Source/ElectraSamples/ElectraSamples.Build.cs b/Engine/Plugins/Media/ElectraUtil/Source/ElectraSamples/ElectraSamples.Build.cs
--- a/Engine/Plugins/Media/ElectraUtil/Source/ElectraSamples/ElectraSamples.Build.cs
+++ b/Engine/Plugins/Media/ElectraUtil/Source/ElectraSamples/ElectraSamples.Build.cs
@@ -80,9 +80,9 @@ namespace UnrealBuildTool.Rules
 				string PluginPath = Utils.MakePathRelativeTo(ModuleDirectory, Target.RelativeEnginePath);
 				AdditionalPropertiesForReceipt.Add("AndroidPlugin", Path.Combine(PluginPath, "ElectraSamples_UPL.xml"));
 			}
-			else if (Target.Platform == UnrealTargetPlatform.Linux)
+			else if (Target.IsInPlatformGroup(UnrealPlatformGroup.Unix))
 			{
-				PublicIncludePaths.Add("$(ModuleDir)/Public/GStreamer");
+				PublicIncludePaths.Add("$(ModuleDir)/Public/Linux");
 			}
 		}
 	}
diff --git a/Engine/Plugins/Media/ElectraUtil/Source/ElectraSamples/Private/ElectraTextureSampleUnix.cpp b/Engine/Plugins/Media/ElectraUtil/Source/ElectraSamples/Private/ElectraTextureSampleUnix.cpp
deleted file mode 100644
diff --git a/Engine/Plugins/Media/ElectraUtil/Source/ElectraSamples/Private/Linux/ElectraTextureSampleLinux.cpp b/Engine/Plugins/Media/ElectraUtil/Source/ElectraSamples/Private/Linux/ElectraTextureSampleLinux.cpp
new file mode 100644
--- /dev/null
+++ b/Engine/Plugins/Media/ElectraUtil/Source/ElectraSamples/Private/Linux/ElectraTextureSampleLinux.cpp
@@ -0,0 +1,91 @@
+// Copyright Epic Games, Inc. All Rights Reserved.
+
+#include "CoreMinimal.h"
+
+#if !UE_SERVER
+
+#include "ElectraTextureSample.h"
+
+FElectraTextureSampleLinux::~FElectraTextureSampleLinux()
+{
+}
+
+
+FIntPoint FElectraTextureSampleLinux::GetDim() const
+{
+	if (VideoDecoderOutput.IsValid())
+	{
+		return VideoDecoderOutput->GetDim();
+	}
+	return FIntPoint::ZeroValue;
+}
+
+FMediaTimeStamp FElectraTextureSampleLinux::GetTime() const
+{
+	if (VideoDecoderOutput.IsValid())
+	{
+		const FDecoderTimeStamp TimeStamp = VideoDecoderOutput->GetTime();
+		return FMediaTimeStamp(TimeStamp.Time, TimeStamp.SequenceIndex);
+	}
+	return FMediaTimeStamp();
+}
+
+FTimespan FElectraTextureSampleLinux::GetDuration() const
+{
+	if (VideoDecoderOutput.IsValid())
+	{
+		return VideoDecoderOutput->GetDuration();
+	}
+	return FTimespan::Zero();
+}
+
+FIntPoint FElectraTextureSampleLinux::GetOutputDim() const
+{
+	if (VideoDecoderOutput.IsValid())
+	{
+		return VideoDecoderOutput->GetOutputDim();
+	}
+	return FIntPoint::ZeroValue;
+}
+
+uint32 FElectraTextureSampleLinux::GetStride() const
+{
+	if (VideoDecoderOutput.IsValid())
+	{
+		return VideoDecoderOutput->GetStride();
+	}
+	return 0;
+}
+
+const void* FElectraTextureSampleLinux::GetBuffer()
+{
+	if (VideoDecoderOutput)
+	{
+		return VideoDecoderOutput->GetBuffer().GetData();
+	}
+	return nullptr;
+}
+
+EMediaTextureSampleFormat FElectraTextureSampleLinux::GetFormat() const
+{
+	return EMediaTextureSampleFormat::CharNV12;
+}
+
+
+/**
+ * Init code for realloacting an image from the the pool
+ */
+void FElectraTextureSampleLinux::InitializePoolable()
+{
+}
+
+/**
+ *  Return the object to the pool and inform the renderer about this...
+ */
+void FElectraTextureSampleLinux::ShutdownPoolable()
+{
+	VideoDecoderOutput.Reset();
+	Texture = nullptr;
+}
+
+#endif
diff --git a/Engine/Plugins/Media/ElectraUtil/Source/ElectraSamples/Public/GStreamer/ElectraTextureSample.h b/Engine/Plugins/Media/ElectraUtil/Source/ElectraSamples/Public/GStreamer/ElectraTextureSample.h
deleted file mode 100644
diff --git a/Engine/Plugins/Media/ElectraUtil/Source/ElectraSamples/Public/Linux/ElectraTextureSample.h b/Engine/Plugins/Media/ElectraUtil/Source/ElectraSamples/Public/Linux/ElectraTextureSample.h
new file mode 100644
--- /dev/null
+++ b/Engine/Plugins/Media/ElectraUtil/Source/ElectraSamples/Public/Linux/ElectraTextureSample.h
@@ -0,0 +1,113 @@
+// Copyright Epic Games, Inc. All Rights Reserved.
+
+#pragma once
+
+#include "CoreMinimal.h"
+#include "IMediaTextureSample.h"
+#include "MediaObjectPool.h"
+#include "RHI.h"
+#include "RHIUtilities.h"
+
+#include "Linux/MediaVideoDecoderOutputLinux.h"
+
+class FRHITexture;
+
+class FElectraTextureSampleLinux
+	: public IMediaTextureSample
+	, public IMediaPoolable
+{
+public:
+	FElectraTextureSampleLinux() = default;
+
+	virtual	~FElectraTextureSampleLinux();
+
+public:
+	void Initialize(FVideoDecoderOutput* InVideoDecoderOutput)
+	{
+		VideoDecoderOutput = StaticCastSharedPtr<FVideoDecoderOutputLinux, IDecoderOutputPoolable, ESPMode::ThreadSafe>(InVideoDecoderOutput->AsShared());
+	}
+/*
+	void CreateTexture()
+	{
+		check(IsInRenderingThread());
+
+		const FRHITextureCreateDesc Desc =
+			FRHITextureCreateDesc::Create2D(TEXT("DummyTexture2D"))
+			.SetExtent(TotalSize)
+			.SetFormat(PF_B8G8R8A8)
+			.SetFlags(ETextureCreateFlags::Dynamic | ETextureCreateFlags::SRGB | ETextureCreateFlags::RenderTargetable | ETextureCreateFlags::ShaderResource)
+			.SetInitialState(ERHIAccess::SRVMask);
+
+		Texture = RHICreateTexture(Desc);
+	}
+*/
+
+	const void* GetBuffer() override;
+
+	FIntPoint GetDim() const override;
+
+	FTimespan GetDuration() const override;
+
+	double GetAspectRatio() const override
+	{
+		return VideoDecoderOutput->GetAspectRatio();
+	}
+
+	EMediaOrientation GetOrientation() const override
+	{
+		return (EMediaOrientation)VideoDecoderOutput->GetOrientation();
+	}
+
+	EMediaTextureSampleFormat GetFormat() const override;
+
+	FIntPoint GetOutputDim() const override;
+	uint32 GetStride() const override;
+
+	FRHITexture* GetTexture() const override
+	{
+		return Texture;
+	}
+
+	FMediaTimeStamp GetTime() const override;
+
+	bool IsCacheable() const override
+	{
+		return true;
+	}
+
+	bool IsOutputSrgb() const override
+	{
+		return true;
+	}
+
+    IMediaTextureSampleConverter* GetMediaTextureSampleConverter() override
+    {
+        return nullptr;
+    }
+
+#if !UE_SERVER
+	void InitializePoolable() override;
+	void ShutdownPoolable() override;
+#endif
+/*
+	TRefCountPtr<FRHITexture2D> GetTextureRef() const
+	{
+		return Texture;
+	}
+*/
+private:
+	/** The sample's texture resource. */
+	TRefCountPtr<FRHITexture2D> Texture;
+
+	/** Output data from video decoder. */
+	TSharedPtr<FVideoDecoderOutputLinux, ESPMode::ThreadSafe> VideoDecoderOutput;
+};
+
+using FElectraTextureSamplePtr = TSharedPtr<FElectraTextureSampleLinux, ESPMode::ThreadSafe>;
+using FElectraTextureSampleRef = TSharedRef<FElectraTextureSampleLinux, ESPMode::ThreadSafe>;
+
+class FElectraTextureSamplePool : public TMediaObjectPool<FElectraTextureSampleLinux>
+{
+public:
+	void PrepareForDecoderShutdown() {}
+};
