Remove use of unacceptable words in ElectraPlayerRuntime

diff --git a/Engine/Plugins/Media/ElectraPlayer/Source/ElectraPlayerRuntime/Private/Runtime/Player/ABRRules/ABRLiveStream.cpp b/Engine/Plugins/Media/ElectraPlayer/Source/ElectraPlayerRuntime/Private/Runtime/Player/ABRRules/ABRLiveStream.cpp
--- a/Engine/Plugins/Media/ElectraPlayer/Source/ElectraPlayerRuntime/Private/Runtime/Player/ABRRules/ABRLiveStream.cpp
+++ b/Engine/Plugins/Media/ElectraPlayer/Source/ElectraPlayerRuntime/Private/Runtime/Player/ABRRules/ABRLiveStream.cpp
@@ -1356,7 +1356,7 @@ IAdaptiveStreamSelector::ESegmentAction FABRLiveStream::PerformSelection(const T
 	}
 	else
 	{
-		// Note: Are there representations we could lift the temporary blacklist from?
+		// Note: Are there representations we could lift the temporary denylist from?
 		if (bSkipWithFiller)
 		{
 			return IAdaptiveStreamSelector::ESegmentAction::Fill;
diff --git a/Engine/Plugins/Media/ElectraPlayer/Source/ElectraPlayerRuntime/Private/Runtime/Player/ABRRules/ABRVoDPlus.cpp b/Engine/Plugins/Media/ElectraPlayer/Source/ElectraPlayerRuntime/Private/Runtime/Player/ABRRules/ABRVoDPlus.cpp
--- a/Engine/Plugins/Media/ElectraPlayer/Source/ElectraPlayerRuntime/Private/Runtime/Player/ABRRules/ABRVoDPlus.cpp
+++ b/Engine/Plugins/Media/ElectraPlayer/Source/ElectraPlayerRuntime/Private/Runtime/Player/ABRRules/ABRVoDPlus.cpp
@@ -827,7 +827,7 @@ IAdaptiveStreamSelector::ESegmentAction FABROnDemandPlus::PerformSelection(const
 	}
 	else
 	{
-		// Note: Are there representations we could lift the temporary blacklist from?
+		// Note: Are there representations we could lift the temporary denylist from?
 		if (bSkipWithFiller)
 		{
 			return IAdaptiveStreamSelector::ESegmentAction::Fill;
diff --git a/Engine/Plugins/Media/ElectraPlayer/Source/ElectraPlayerRuntime/Private/Runtime/Player/AdaptiveStreamingPlayerABR.cpp b/Engine/Plugins/Media/ElectraPlayer/Source/ElectraPlayerRuntime/Private/Runtime/Player/AdaptiveStreamingPlayerABR.cpp
--- a/Engine/Plugins/Media/ElectraPlayer/Source/ElectraPlayerRuntime/Private/Runtime/Player/AdaptiveStreamingPlayerABR.cpp
+++ b/Engine/Plugins/Media/ElectraPlayer/Source/ElectraPlayerRuntime/Private/Runtime/Player/AdaptiveStreamingPlayerABR.cpp
@@ -33,8 +33,8 @@ namespace Electra
 		FTimeValue GetMinBufferTimeForPlayback(EMinBufferType InBufferingType, FTimeValue InDefaultMBT) override;
 		FRebufferAction GetRebufferAction(const FParamDict& CurrentPlayerOptions) override;
 		EHandlingAction PeriodicHandle() override;
-		void MarkStreamAsUnavailable(const FBlacklistedStream& BlacklistedStream) override;
-		void MarkStreamAsAvailable(const FBlacklistedStream& NoLongerBlacklistedStream) override;
+		void MarkStreamAsUnavailable(const FDenylistedStream& DenylistedStream) override;
+		void MarkStreamAsAvailable(const FDenylistedStream& NoLongerDenylistedStream) override;
 		int64 GetLastBandwidth() override;
 		int64 GetAverageBandwidth() override;
 		int64 GetAverageThroughput() override;
@@ -128,7 +128,7 @@ namespace Electra
 		FString												CurrentVideoAdaptationSetID;
 		FString												CurrentAudioAdaptationSetID;
 
-		TArray<FBlacklistedStream>							BlacklistedExternally;
+		TArray<FDenylistedStream>							DenylistedExternally;
 
 		FStreamCodecInformation::FResolution				MaxStreamResolution;
 		int32												BandwidthCeiling = 0x7fffffff;
@@ -380,12 +380,12 @@ namespace Electra
 	 * This is used in formats like HLS when a dedicated stream playlist is in error
 	 * and thus segment information for this stream is not available.
 	 *
-	 * @param InBlacklistedStream
+	 * @param InDenylistedStream
 	 */
-	void FAdaptiveStreamSelector::MarkStreamAsUnavailable(const FBlacklistedStream& InBlacklistedStream)
+	void FAdaptiveStreamSelector::MarkStreamAsUnavailable(const FDenylistedStream& InDenylistedStream)
 	{
 		FScopeLock lock(&AccessMutex);
-		BlacklistedExternally.Push(InBlacklistedStream);
+		DenylistedExternally.Push(InDenylistedStream);
 	}
 
 	//-----------------------------------------------------------------------------
@@ -393,12 +393,12 @@ namespace Electra
 	 * Marks a previously set to unavailable stream as being available again.
 	 * Usually when the stream specific playlist (in HLS) has become available again.
 	 *
-	 * @param InNoLongerBlacklistedStream
+	 * @param InNoLongerDenylistedStream
 	 */
-	void FAdaptiveStreamSelector::MarkStreamAsAvailable(const FBlacklistedStream& InNoLongerBlacklistedStream)
+	void FAdaptiveStreamSelector::MarkStreamAsAvailable(const FDenylistedStream& InNoLongerDenylistedStream)
 	{
 		FScopeLock lock(&AccessMutex);
-		BlacklistedExternally.Remove(InNoLongerBlacklistedStream);
+		DenylistedExternally.Remove(InNoLongerDenylistedStream);
 	}
 
 	//-----------------------------------------------------------------------------
@@ -624,16 +624,16 @@ namespace Electra
 			}
 		}
 
-		// Remove those that have been externally blacklisted. This may include all of them.
+		// Remove those that have been externally denylisted. This may include all of them.
 		AccessMutex.Lock();
-		if (BlacklistedExternally.Num())
+		if (DenylistedExternally.Num())
 		{
 			for(int32 j=0; j<CandidateRepresentations.Num(); ++j)
 			{
-				for(int32 i=0, iMax=BlacklistedExternally.Num(); i<iMax; ++i)
+				for(int32 i=0, iMax=DenylistedExternally.Num(); i<iMax; ++i)
 				{
-					if (CandidateRepresentations[j]->AdaptationSetUniqueID == BlacklistedExternally[i].AdaptationSetUniqueID &&
-						CandidateRepresentations[j]->RepresentationUniqueID == BlacklistedExternally[i].RepresentationUniqueID)
+					if (CandidateRepresentations[j]->AdaptationSetUniqueID == DenylistedExternally[i].AdaptationSetUniqueID &&
+						CandidateRepresentations[j]->RepresentationUniqueID == DenylistedExternally[i].RepresentationUniqueID)
 					{
 						CandidateRepresentations.RemoveAt(j);
 						--j;
diff --git a/Engine/Plugins/Media/ElectraPlayer/Source/ElectraPlayerRuntime/Private/Runtime/Player/AdaptiveStreamingPlayerABR.h b/Engine/Plugins/Media/ElectraPlayer/Source/ElectraPlayerRuntime/Private/Runtime/Player/AdaptiveStreamingPlayerABR.h
--- a/Engine/Plugins/Media/ElectraPlayer/Source/ElectraPlayerRuntime/Private/Runtime/Player/AdaptiveStreamingPlayerABR.h
+++ b/Engine/Plugins/Media/ElectraPlayer/Source/ElectraPlayerRuntime/Private/Runtime/Player/AdaptiveStreamingPlayerABR.h
@@ -119,13 +119,13 @@ public:
 	// The ABR may take internal action to switch quality when this is called!
 	virtual FRebufferAction GetRebufferAction(const FParamDict& CurrentPlayerOptions) = 0;
 
-	struct FBlacklistedStream
+	struct FDenylistedStream
 	{
 		FString		AssetUniqueID;
 		FString		AdaptationSetUniqueID;
 		FString		RepresentationUniqueID;
 		FString		CDN;
-		bool operator == (const FBlacklistedStream& rhs) const
+		bool operator == (const FDenylistedStream& rhs) const
 		{
 			return AssetUniqueID == rhs.AssetUniqueID &&
 				   AdaptationSetUniqueID == rhs.AdaptationSetUniqueID &&
@@ -133,8 +133,8 @@ public:
 				   CDN == rhs.CDN;
 		}
 	};
-	virtual void MarkStreamAsUnavailable(const FBlacklistedStream& BlacklistedStream) = 0;
-	virtual void MarkStreamAsAvailable(const FBlacklistedStream& NoLongerBlacklistedStream) = 0;
+	virtual void MarkStreamAsUnavailable(const FDenylistedStream& DenylistedStream) = 0;
+	virtual void MarkStreamAsAvailable(const FDenylistedStream& NoLongerDenylistedStream) = 0;
 
 
 	virtual FABRDownloadProgressDecision ReportDownloadProgress(const Metrics::FSegmentDownloadStats& SegmentDownloadStats) = 0;
diff --git a/Engine/Plugins/Media/ElectraPlayer/Source/ElectraPlayerRuntime/Private/Runtime/Player/HLS/ManifestBuilderHLS.cpp b/Engine/Plugins/Media/ElectraPlayer/Source/ElectraPlayerRuntime/Private/Runtime/Player/HLS/ManifestBuilderHLS.cpp
--- a/Engine/Plugins/Media/ElectraPlayer/Source/ElectraPlayerRuntime/Private/Runtime/Player/HLS/ManifestBuilderHLS.cpp
+++ b/Engine/Plugins/Media/ElectraPlayer/Source/ElectraPlayerRuntime/Private/Runtime/Player/HLS/ManifestBuilderHLS.cpp
@@ -142,10 +142,10 @@
 	virtual FErrorDetail BuildFromMasterPlaylist(TSharedPtrTS<FManifestHLSInternal>& OutHLSPlaylist, const HLSPlaylistParser::FPlaylist& Playlist, const FPlaylistLoadRequestHLS& SourceRequest, const HTTP::FConnectionInfo* ConnectionInfo) override;
 
 	virtual FErrorDetail GetInitialPlaylistLoadRequests(TArray<FPlaylistLoadRequestHLS>& OutRequests, TSharedPtrTS<FManifestHLSInternal> Manifest) override;
-	virtual UEMediaError UpdateFailedInitialPlaylistLoadRequest(FPlaylistLoadRequestHLS& InOutFailedRequest, const HTTP::FConnectionInfo* ConnectionInfo, TSharedPtrTS<HTTP::FRetryInfo> PreviousAttempts, const FTimeValue& BlacklistUntilUTC, TSharedPtrTS<FManifestHLSInternal> Manifest) override;
+	virtual UEMediaError UpdateFailedInitialPlaylistLoadRequest(FPlaylistLoadRequestHLS& InOutFailedRequest, const HTTP::FConnectionInfo* ConnectionInfo, TSharedPtrTS<HTTP::FRetryInfo> PreviousAttempts, const FTimeValue& DenylistUntilUTC, TSharedPtrTS<FManifestHLSInternal> Manifest) override;
 
 	virtual FErrorDetail UpdateFromVariantPlaylist(TSharedPtrTS<FManifestHLSInternal> InOutHLSPlaylist, const HLSPlaylistParser::FPlaylist& VariantPlaylist, const FPlaylistLoadRequestHLS& SourceRequest, const HTTP::FConnectionInfo* ConnectionInfo, uint32 ResponseCRC) override;
-	virtual void SetVariantPlaylistFailure(TSharedPtrTS<FManifestHLSInternal> InHLSPlaylist, const FPlaylistLoadRequestHLS& SourceRequest, const HTTP::FConnectionInfo* ConnectionInfo, TSharedPtrTS<HTTP::FRetryInfo> PreviousAttempts, const FTimeValue& BlacklistUntilUTC) override;
+	virtual void SetVariantPlaylistFailure(TSharedPtrTS<FManifestHLSInternal> InHLSPlaylist, const FPlaylistLoadRequestHLS& SourceRequest, const HTTP::FConnectionInfo* ConnectionInfo, TSharedPtrTS<HTTP::FRetryInfo> PreviousAttempts, const FTimeValue& DenylistUntilUTC) override;
 
 private:
 	FErrorDetail SetupRenditions(FManifestHLSInternal* Manifest, const HLSPlaylistParser::FPlaylist& Playlist);
@@ -418,8 +418,8 @@
 
 	return FErrorDetail();
 }
+
 
-
 FErrorDetail FManifestBuilderHLS::SetupVariants(FManifestHLSInternal* Manifest, const HLSPlaylistParser::FPlaylist& Playlist)
 {
 	if (Playlist.GetPlaylists().Num() == 0)
@@ -1128,7 +1128,7 @@
 	return FErrorDetail();
 }
 
-UEMediaError FManifestBuilderHLS::UpdateFailedInitialPlaylistLoadRequest(FPlaylistLoadRequestHLS& InOutFailedRequest, const HTTP::FConnectionInfo* ConnectionInfo, TSharedPtrTS<HTTP::FRetryInfo> PreviousAttempts, const FTimeValue& BlacklistUntilUTC, TSharedPtrTS<FManifestHLSInternal> Manifest)
+UEMediaError FManifestBuilderHLS::UpdateFailedInitialPlaylistLoadRequest(FPlaylistLoadRequestHLS& InOutFailedRequest, const HTTP::FConnectionInfo* ConnectionInfo, TSharedPtrTS<HTTP::FRetryInfo> PreviousAttempts, const FTimeValue& DenylistUntilUTC, TSharedPtrTS<FManifestHLSInternal> Manifest)
 {
 	// Is the failed request a video variant?
 	for(int32 i=0, iMax=Manifest->VariantStreams.Num(); i<iMax; ++i)
@@ -1159,7 +1159,7 @@
 				{
 					if (!AlternateVariantStream.IsValid() || AlternateVariantStream->Bandwidth < Manifest->VariantStreams[nVariant]->Bandwidth)
 					{
-						if (!Manifest->VariantStreams[nVariant]->Internal.Blacklisted.Get())
+						if (!Manifest->VariantStreams[nVariant]->Internal.Denylisted.Get())
 						{
 							AlternateVariantStream = Manifest->VariantStreams[nVariant];
 						}
@@ -1177,30 +1177,30 @@
 					}
 					if (!AlternateVariantStream.IsValid() || AlternateVariantStream->Bandwidth > Manifest->VariantStreams[nVariant]->Bandwidth)
 					{
-						if (!Manifest->VariantStreams[nVariant]->Internal.Blacklisted.Get())
+						if (!Manifest->VariantStreams[nVariant]->Internal.Denylisted.Get())
 						{
 							AlternateVariantStream = Manifest->VariantStreams[nVariant];
 						}
 					}
 				}
 			}
-			// FIXME: We could also check if there is any variant that isn't blacklisted yet (or whose blacklist can be lifted again).
+			// FIXME: We could also check if there is any variant that isn't denylisted yet (or whose denylist can be lifted again).
 			if (!AlternateVariantStream.IsValid())
 			{
 				return UEMEDIA_ERROR_END_OF_STREAM;
 			}
 
-			// Blacklist the failed variant only when we have found an alternative. Otherwise we may want to retry the failed stream.
-			VideoVariant->Internal.Blacklisted = MakeSharedTS<FManifestHLSInternal::FBlacklist>();
-			VideoVariant->Internal.Blacklisted->PreviousAttempts				= PreviousAttempts;
-			VideoVariant->Internal.Blacklisted->BecomesAvailableAgainAtUTC  	= BlacklistUntilUTC;
-			VideoVariant->Internal.Blacklisted->AssetIDs.AssetUniqueID  		= DefaultAssetNameHLS;
-			VideoVariant->Internal.Blacklisted->AssetIDs.AdaptationSetUniqueID  = InOutFailedRequest.AdaptationSetUniqueID;
-			VideoVariant->Internal.Blacklisted->AssetIDs.RepresentationUniqueID = InOutFailedRequest.RepresentationUniqueID;
-			VideoVariant->Internal.Blacklisted->AssetIDs.CDN					= InOutFailedRequest.CDN;
+			// Denylist the failed variant only when we have found an alternative. Otherwise we may want to retry the failed stream.
+			VideoVariant->Internal.Denylisted = MakeSharedTS<FManifestHLSInternal::FDenylist>();
+			VideoVariant->Internal.Denylisted->PreviousAttempts				= PreviousAttempts;
+			VideoVariant->Internal.Denylisted->BecomesAvailableAgainAtUTC  	= DenylistUntilUTC;
+			VideoVariant->Internal.Denylisted->AssetIDs.AssetUniqueID  		= DefaultAssetNameHLS;
+			VideoVariant->Internal.Denylisted->AssetIDs.AdaptationSetUniqueID  = InOutFailedRequest.AdaptationSetUniqueID;
+			VideoVariant->Internal.Denylisted->AssetIDs.RepresentationUniqueID = InOutFailedRequest.RepresentationUniqueID;
+			VideoVariant->Internal.Denylisted->AssetIDs.CDN					= InOutFailedRequest.CDN;
 			// Tell the stream selector that this stream is temporarily unavailable.
 			TSharedPtrTS<IAdaptiveStreamSelector> StreamSelector(PlayerSessionServices->GetStreamSelector());
-			StreamSelector->MarkStreamAsUnavailable(VideoVariant->Internal.Blacklisted->AssetIDs);
+			StreamSelector->MarkStreamAsUnavailable(VideoVariant->Internal.Denylisted->AssetIDs);
 
 			FURL_RFC3986 UrlBuilder;
 			UrlBuilder.Parse(Manifest->MasterPlaylistVars.PlaylistLoadRequest.URL);
@@ -1237,17 +1237,17 @@
 			return UEMEDIA_ERROR_END_OF_STREAM;
 		// Implement this later when there actually are different auiod variants to choose from.
 			#if 0
-				// Blacklist the failed variant only when we have found an alternative. Otherwise we may want to retry the failed stream.
-				AudioRendition->Internal.Blacklisted = MakeSharedTS<FManifestHLSInternal::FBlacklist>();
-				AudioRendition->Internal.Blacklisted->PreviousAttempts  			  = PreviousAttempts;
-				AudioRendition->Internal.Blacklisted->BecomesAvailableAgainAtUTC	  = BlacklistUntilUTC;
-				AudioRendition->Internal.Blacklisted->AssetIDs.AssetUniqueID		  = DefaultAssetNameHLS;
-				AudioRendition->Internal.Blacklisted->AssetIDs.AdaptationSetUniqueID  = InOutFailedRequest.AdaptationSetUniqueID;
-				AudioRendition->Internal.Blacklisted->AssetIDs.RepresentationUniqueID = InOutFailedRequest.RepresentationUniqueID;
-				AudioRendition->Internal.Blacklisted->AssetIDs.CDN  				  = InOutFailedRequest.CDN;
+				// Denylist the failed variant only when we have found an alternative. Otherwise we may want to retry the failed stream.
+				AudioRendition->Internal.Denylisted = MakeSharedTS<FManifestHLSInternal::FDenylist>();
+				AudioRendition->Internal.Denylisted->PreviousAttempts  			  = PreviousAttempts;
+				AudioRendition->Internal.Denylisted->BecomesAvailableAgainAtUTC	  = DenylistUntilUTC;
+				AudioRendition->Internal.Denylisted->AssetIDs.AssetUniqueID		  = DefaultAssetNameHLS;
+				AudioRendition->Internal.Denylisted->AssetIDs.AdaptationSetUniqueID  = InOutFailedRequest.AdaptationSetUniqueID;
+				AudioRendition->Internal.Denylisted->AssetIDs.RepresentationUniqueID = InOutFailedRequest.RepresentationUniqueID;
+				AudioRendition->Internal.Denylisted->AssetIDs.CDN  				  = InOutFailedRequest.CDN;
 				// Tell the stream selector that this stream is temporarily unavailable.
 				TSharedPtrTS<IAdaptiveStreamSelector> StreamSelector(PlayerSessionServices->GetStreamSelector());
-				StreamSelector->MarkStreamAsUnavailable(AudioRendition->Internal.Blacklisted->AssetIDs);
+				StreamSelector->MarkStreamAsUnavailable(AudioRendition->Internal.Denylisted->AssetIDs);
 				return UEMEDIA_ERROR_OK;
 			#endif
 		}
@@ -1257,7 +1257,7 @@
 }
 
 
-void FManifestBuilderHLS::SetVariantPlaylistFailure(TSharedPtrTS<FManifestHLSInternal> InHLSPlaylist, const FPlaylistLoadRequestHLS& SourceRequest, const HTTP::FConnectionInfo* ConnectionInfo, TSharedPtrTS<HTTP::FRetryInfo> PreviousAttempts, const FTimeValue& BlacklistUntilUTC)
+void FManifestBuilderHLS::SetVariantPlaylistFailure(TSharedPtrTS<FManifestHLSInternal> InHLSPlaylist, const FPlaylistLoadRequestHLS& SourceRequest, const HTTP::FConnectionInfo* ConnectionInfo, TSharedPtrTS<HTTP::FRetryInfo> PreviousAttempts, const FTimeValue& DenylistUntilUTC)
 {
 	// Find the variant or rendition this playlist updates.
 	InHLSPlaylist->LockPlaylists();
@@ -1278,18 +1278,18 @@
 			}
 			Playlist->Internal.ExpiresAtTime.SetToPositiveInfinity();
 
-			Playlist->Internal.Blacklisted = MakeSharedTS<FManifestHLSInternal::FBlacklist>();
-			Playlist->Internal.Blacklisted->PreviousAttempts		   = PreviousAttempts;
-			Playlist->Internal.Blacklisted->BecomesAvailableAgainAtUTC = BlacklistUntilUTC;
+			Playlist->Internal.Denylisted = MakeSharedTS<FManifestHLSInternal::FDenylist>();
+			Playlist->Internal.Denylisted->PreviousAttempts		   = PreviousAttempts;
+			Playlist->Internal.Denylisted->BecomesAvailableAgainAtUTC = DenylistUntilUTC;
 
-			Playlist->Internal.Blacklisted->AssetIDs.AssetUniqueID  		= DefaultAssetNameHLS;
-			Playlist->Internal.Blacklisted->AssetIDs.AdaptationSetUniqueID  = SourceRequest.AdaptationSetUniqueID;
-			Playlist->Internal.Blacklisted->AssetIDs.RepresentationUniqueID = SourceRequest.RepresentationUniqueID;
-			Playlist->Internal.Blacklisted->AssetIDs.CDN					= SourceRequest.CDN;
+			Playlist->Internal.Denylisted->AssetIDs.AssetUniqueID  		= DefaultAssetNameHLS;
+			Playlist->Internal.Denylisted->AssetIDs.AdaptationSetUniqueID  = SourceRequest.AdaptationSetUniqueID;
+			Playlist->Internal.Denylisted->AssetIDs.RepresentationUniqueID = SourceRequest.RepresentationUniqueID;
+			Playlist->Internal.Denylisted->AssetIDs.CDN					= SourceRequest.CDN;
 
 			// Tell the stream selector that this stream is temporarily unavailable.
 			TSharedPtrTS<IAdaptiveStreamSelector> StreamSelector(PlayerSessionServices->GetStreamSelector());
-			StreamSelector->MarkStreamAsUnavailable(Playlist->Internal.Blacklisted->AssetIDs);
+			StreamSelector->MarkStreamAsUnavailable(Playlist->Internal.Denylisted->AssetIDs);
 		}
 	}
 	InHLSPlaylist->UnlockPlaylists();
@@ -1336,7 +1336,7 @@
 			}
 			Playlist->Internal.ExpiresAtTime	   = SourceRequest.RequestedAtTime + (Playlist->Internal.MediaStream->TargetDuration / 2);
 			Playlist->Internal.bReloadTriggered    = false;
-			Playlist->Internal.Blacklisted.Reset();
+			Playlist->Internal.Denylisted.Reset();
 		}
 		else
 		{
@@ -1705,7 +1705,7 @@
 		Playlist->Internal.PlaylistLoadRequest.LastUpdateCRC32 = ResponseCRC;
 		Playlist->Internal.MediaStream  					   = UpdatedMediaStream;
 		Playlist->Internal.LoadState						   = FManifestHLSInternal::FPlaylistBase::FInternal::ELoadState::Loaded;
-		Playlist->Internal.Blacklisted.Reset();
+		Playlist->Internal.Denylisted.Reset();
 
 		// Set the expiration time of this playlist.
 		// A VOD or list with an end tag does not expire.
diff --git a/Engine/Plugins/Media/ElectraPlayer/Source/ElectraPlayerRuntime/Private/Runtime/Player/HLS/ManifestBuilderHLS.h b/Engine/Plugins/Media/ElectraPlayer/Source/ElectraPlayerRuntime/Private/Runtime/Player/HLS/ManifestBuilderHLS.h
--- a/Engine/Plugins/Media/ElectraPlayer/Source/ElectraPlayerRuntime/Private/Runtime/Player/HLS/ManifestBuilderHLS.h
+++ b/Engine/Plugins/Media/ElectraPlayer/Source/ElectraPlayerRuntime/Private/Runtime/Player/HLS/ManifestBuilderHLS.h
@@ -151,11 +151,11 @@ struct FManifestHLSInternal
 	};
 
 
-	struct FBlacklist
+	struct FDenylist
 	{
 		TSharedPtrTS<HTTP::FRetryInfo>					PreviousAttempts;
 		FTimeValue										BecomesAvailableAgainAtUTC;
-		IAdaptiveStreamSelector::FBlacklistedStream		AssetIDs;
+		IAdaptiveStreamSelector::FDenylistedStream		AssetIDs;
 	};
 
 
@@ -182,7 +182,7 @@ struct FManifestHLSInternal
 			FInternal() : LoadState(ELoadState::NotLoaded), UniqueID(0), bReloadTriggered(false), bNewlySelected(true), bHasVideo(false), bHasAudio(false) {}
 			FPlaylistLoadRequestHLS			PlaylistLoadRequest;
 			TSharedPtrTS<FMediaStream>		MediaStream;
-			TSharedPtrTS<FBlacklist>		Blacklisted;
+			TSharedPtrTS<FDenylist>		Denylisted;
 			FTimeValue						ExpiresAtTime;							//!< Synchronized UTC time (from session service GetSynchronizedUTCTime()) at which this list expires. Set to infinite if it does not.
 			ELoadState						LoadState;
 			uint32							UniqueID;
@@ -433,12 +433,12 @@ public:
 	 *                 Failed request for which an alterative shall be returned.
 	 * @param ConnectionInfo
 	 * @param PreviousAttempts
-	 * @param BlacklistUntilUTC
+	 * @param DenylistUntilUTC
 	 * @param Manifest
 	 *
 	 * @return UEMEDIA_ERROR_OK if an alternative was found or UEMEDIA_ERROR_END_OF_STREAM if no further alternatives exist.
 	 */
-	virtual UEMediaError UpdateFailedInitialPlaylistLoadRequest(FPlaylistLoadRequestHLS& InOutFailedRequest, const HTTP::FConnectionInfo* ConnectionInfo, TSharedPtrTS<HTTP::FRetryInfo> PreviousAttempts, const FTimeValue& BlacklistUntilUTC, TSharedPtrTS<FManifestHLSInternal> Manifest) = 0;
+	virtual UEMediaError UpdateFailedInitialPlaylistLoadRequest(FPlaylistLoadRequestHLS& InOutFailedRequest, const HTTP::FConnectionInfo* ConnectionInfo, TSharedPtrTS<HTTP::FRetryInfo> PreviousAttempts, const FTimeValue& DenylistUntilUTC, TSharedPtrTS<FManifestHLSInternal> Manifest) = 0;
 
 
 	/**
@@ -462,9 +462,9 @@ public:
 	 * @param SourceRequest
 	 * @param ConnectionInfo
 	 * @param PreviousAttempts
-	 * @param BlacklistUntilUTC
+	 * @param DenylistUntilUTC
 	 */
-	virtual void SetVariantPlaylistFailure(TSharedPtrTS<FManifestHLSInternal> InHLSPlaylist, const FPlaylistLoadRequestHLS& SourceRequest, const HTTP::FConnectionInfo* ConnectionInfo, TSharedPtrTS<HTTP::FRetryInfo> PreviousAttempts, const FTimeValue& BlacklistUntilUTC) = 0;
+	virtual void SetVariantPlaylistFailure(TSharedPtrTS<FManifestHLSInternal> InHLSPlaylist, const FPlaylistLoadRequestHLS& SourceRequest, const HTTP::FConnectionInfo* ConnectionInfo, TSharedPtrTS<HTTP::FRetryInfo> PreviousAttempts, const FTimeValue& DenylistUntilUTC) = 0;
 };
 
 
diff --git a/Engine/Plugins/Media/ElectraPlayer/Source/ElectraPlayerRuntime/Private/Runtime/Player/HLS/ManifestHLS.cpp b/Engine/Plugins/Media/ElectraPlayer/Source/ElectraPlayerRuntime/Private/Runtime/Player/HLS/ManifestHLS.cpp
--- a/Engine/Plugins/Media/ElectraPlayer/Source/ElectraPlayerRuntime/Private/Runtime/Player/HLS/ManifestHLS.cpp
+++ b/Engine/Plugins/Media/ElectraPlayer/Source/ElectraPlayerRuntime/Private/Runtime/Player/HLS/ManifestHLS.cpp
@@ -89,7 +89,7 @@ private:
 
 	IManifest::FResult FindSegment(TSharedPtrTS<FStreamSegmentRequestHLSfmp4>& OutRequest, TSharedPtrTS<FManifestHLSInternal::FPlaylistBase> InPlaylist, TSharedPtrTS<FManifestHLSInternal::FMediaStream> InStream, uint32 StreamUniqueID, EStreamType StreamType, const FSegSearchParam& SearchParam, IManifest::ESearchType SearchType);
 
-	void RefreshBlacklistState();
+	void RefreshDenylistState();
 
 	TSharedPtrTS<FManifestHLSInternal>			InternalManifest;
 	IPlayerSessionServices* 					SessionServices;
@@ -490,10 +490,10 @@ IManifest::FResult FPlayPeriodHLS::GetMediaStreamForID(TSharedPtrTS<FManifestHLS
 				{
 					OutPlaylist = Playlist;
 
-					// Playlist currently blacklisted?
-					if (Playlist->Internal.Blacklisted.IsValid())
+					// Playlist currently denylisted?
+					if (Playlist->Internal.Denylisted.IsValid())
 					{
-						// Return and assume a non-blacklisted stream will be selected.
+						// Return and assume a allowlist stream will be selected.
 						return IManifest::FResult().RetryAfterMilliseconds(50);
 					}
 					// Check the load state
@@ -834,7 +834,7 @@ IManifest::FResult FPlayPeriodHLS::GetStartingSegment(TSharedPtrTS<IStreamSegmen
 {
 	FManifestHLSInternal::ScopedLockPlaylists lock(InternalManifest);
 
-	RefreshBlacklistState();
+	RefreshDenylistState();
 
 	TSharedPtrTS<FManifestHLSInternal::FPlaylistBase>	VideoPlaylist;
 	TSharedPtrTS<FManifestHLSInternal::FPlaylistBase>	AudioPlaylist;
@@ -1095,7 +1095,7 @@ IManifest::FResult FPlayPeriodHLS::GetNextOrRetrySegment(TSharedPtrTS<IStreamSeg
 		return IManifest::FResult(IManifest::FResult::EType::NotFound).SetErrorDetail(FErrorDetail().SetMessage(FString::Printf(TEXT("Cannot get next segment stream type \"%s\" since no stream is actively selected!"), GetStreamTypeName(CurrentRequest->GetType()))));
 	}
 
-	RefreshBlacklistState();
+	RefreshDenylistState();
 
 	TSharedPtrTS<FManifestHLSInternal::FMediaStream>	Stream;
 	TSharedPtrTS<FManifestHLSInternal::FPlaylistBase> Playlist;
@@ -1207,9 +1207,9 @@ IManifest::FResult FPlayPeriodHLS::GetRetrySegment(TSharedPtrTS<IStreamSegment>&
 
 
 /**
- * Checks if any potentially blacklisted stream can be used again.
+ * Checks if any potentially denylisted stream can be used again.
  */
-void FPlayPeriodHLS::RefreshBlacklistState()
+void FPlayPeriodHLS::RefreshDenylistState()
 {
 	FTimeValue Now = SessionServices->GetSynchronizedUTCTime()->GetTime();
 
@@ -1219,9 +1219,9 @@ void FPlayPeriodHLS::RefreshBlacklistState()
 		TSharedPtrTS<FManifestHLSInternal::FRendition::FPlaylistBase> Stream = It.Value().Pin();
 		if (Stream.IsValid())
 		{
-			if (Stream->Internal.Blacklisted.IsValid())
+			if (Stream->Internal.Denylisted.IsValid())
 			{
-				if (Now >= Stream->Internal.Blacklisted->BecomesAvailableAgainAtUTC)
+				if (Now >= Stream->Internal.Denylisted->BecomesAvailableAgainAtUTC)
 				{
 					Stream->Internal.LoadState  	  = FManifestHLSInternal::FPlaylistBase::FInternal::ELoadState::NotLoaded;
 					Stream->Internal.bReloadTriggered = false;
@@ -1230,10 +1230,10 @@ void FPlayPeriodHLS::RefreshBlacklistState()
 
 					// Tell the stream selector that this stream is available again.
 					TSharedPtrTS<IAdaptiveStreamSelector> StreamSelector(SessionServices->GetStreamSelector());
-					StreamSelector->MarkStreamAsAvailable(Stream->Internal.Blacklisted->AssetIDs);
+					StreamSelector->MarkStreamAsAvailable(Stream->Internal.Denylisted->AssetIDs);
 
-					Stream->Internal.Blacklisted.Reset();
-					LogMessage(IInfoLog::ELevel::Info, FString::Printf(TEXT("Lifting blacklist of playlist \"%s\""), *Stream->Internal.PlaylistLoadRequest.URL));
+					Stream->Internal.Denylisted.Reset();
+					LogMessage(IInfoLog::ELevel::Info, FString::Printf(TEXT("Lifting denylist of playlist \"%s\""), *Stream->Internal.PlaylistLoadRequest.URL));
 				}
 			}
 		}
@@ -1399,7 +1399,7 @@ TSharedPtrTS<ITimelineMediaAsset> FPlayPeriodHLS::GetMediaAsset() const
 
 void FPlayPeriodHLS::SelectStream(const FString& AdaptationSetID, const FString& RepresentationID)
 {
-	RefreshBlacklistState();
+	RefreshDenylistState();
 
 	// The representation ID is the unique ID of the stream as a string. Convert it back
 	uint32 UniqueID;
diff --git a/Engine/Plugins/Media/ElectraPlayer/Source/ElectraPlayerRuntime/Private/Runtime/Player/HLS/PlaylistReaderHLS.cpp b/Engine/Plugins/Media/ElectraPlayer/Source/ElectraPlayerRuntime/Private/Runtime/Player/HLS/PlaylistReaderHLS.cpp
--- a/Engine/Plugins/Media/ElectraPlayer/Source/ElectraPlayerRuntime/Private/Runtime/Player/HLS/PlaylistReaderHLS.cpp
+++ b/Engine/Plugins/Media/ElectraPlayer/Source/ElectraPlayerRuntime/Private/Runtime/Player/HLS/PlaylistReaderHLS.cpp
@@ -997,9 +997,9 @@ void FPlaylistReaderHLS::HandleCompletedPlaylistDownloads(FTimeValue& TimeNow)
 
 						if (bFail)
 						{
-							int32 BlacklistForMilliseconds = 1000 * 15;
-							Builder->SetVariantPlaylistFailure(Manifest, Request->GetPlaylistLoadRequest(), ConnInfo, RetryInfo, TimeNow + FTimeValue().SetFromMilliseconds(BlacklistForMilliseconds));
-							LogMessage(IInfoLog::ELevel::Warning, FString::Printf(TEXT("Failed %s playlist \"%s\" (%s), blacklisting for %d milliseconds"), ParseError.IsSet() ? TEXT("parsing") : TEXT("downloading"), *ConnInfo->EffectiveURL, ParseError.IsSet() ? *ParseError.GetPrintable() : *ConnInfo->StatusInfo.ErrorDetail.GetMessage(), BlacklistForMilliseconds));
+							int32 DenylistForMilliseconds = 1000 * 15;
+							Builder->SetVariantPlaylistFailure(Manifest, Request->GetPlaylistLoadRequest(), ConnInfo, RetryInfo, TimeNow + FTimeValue().SetFromMilliseconds(DenylistForMilliseconds));
+							LogMessage(IInfoLog::ELevel::Warning, FString::Printf(TEXT("Failed %s playlist \"%s\" (%s), denylisting for %d milliseconds"), ParseError.IsSet() ? TEXT("parsing") : TEXT("downloading"), *ConnInfo->EffectiveURL, ParseError.IsSet() ? *ParseError.GetPrintable() : *ConnInfo->StatusInfo.ErrorDetail.GetMessage(), DenylistForMilliseconds));
 						}
 					}
 					else
