ElectraPlayer: Added DASH MPD reloading of a Live stream on too many segment failures to see if the presentation may have ended.

#rnx

diff --git a/Engine/Plugins/Media/ElectraPlayer/Source/ElectraPlayerRuntime/Private/Runtime/Player/ABRRules/ABRLiveStream.cpp b/Engine/Plugins/Media/ElectraPlayer/Source/ElectraPlayerRuntime/Private/Runtime/Player/ABRRules/ABRLiveStream.cpp
--- a/Engine/Plugins/Media/ElectraPlayer/Source/ElectraPlayerRuntime/Private/Runtime/Player/ABRRules/ABRLiveStream.cpp
+++ b/Engine/Plugins/Media/ElectraPlayer/Source/ElectraPlayerRuntime/Private/Runtime/Player/ABRRules/ABRLiveStream.cpp
@@ -243,6 +243,7 @@ private:
 		bool bBufferingWillCompleteOnTime = false;
 
 		FABRDownloadProgressDecision StreamReaderDecision;
+		double LastTimePlaylistRefreshRequested = -1.0;
 	};
 
 
@@ -338,6 +339,10 @@ private:
 
 	const int32 NumQualityLevelDropsWhenAborting = 2;
 
+	const double TimeUntilNextPlaylistRefresh = 10.0;
+	const int32 TriggerPlaylistRefreshAfterNumLateSegments = 3;
+	const int32 MaxTimesToRetryLateSegment = 10;
+
 	// When not playing on the Live edge we delay loading of a segment in the lower
 	// quality range if the buffer has more data (and the distance to the Live edge is greater).
 	// This avoids loading in many poor quality segments and allows for faster upswitching if bandwidth improves.
@@ -870,6 +875,13 @@ void FABRLiveStream::PrepareStreamCandidateList(TArray<TSharedPtrTS<FABRStreamIn
 IAdaptiveStreamSelector::ESegmentAction FABRLiveStream::EvaluateForError(TArray<TSharedPtrTS<FABRStreamInformation>>& InOutCandidates, EStreamType StreamType, FTimeValue& OutDelay, const TSharedPtrTS<IManifest::IPlayPeriod>& CurrentPlayPeriod, const TSharedPtrTS<const IStreamSegment>& CurrentSegment, const FTimeValue& TimeNow)
 {
 	FStreamWorkVars* WorkVars = GetWorkVars(StreamType);
+	double Now = Info->ABRGetWallclockTime().GetAsSeconds();
+
+	// Allow another playlist updates after enough time has elapsed.
+	if (WorkVars->LastTimePlaylistRefreshRequested > 0.0 && Now - WorkVars->LastTimePlaylistRefreshRequested > TimeUntilNextPlaylistRefresh)
+	{
+		WorkVars->LastTimePlaylistRefreshRequested = -1.0;
+	}
 
 	if (CurrentSegment.IsValid())
 	{
@@ -904,7 +916,13 @@ IAdaptiveStreamSelector::ESegmentAction FABRLiveStream::EvaluateForError(TArray<
 						CurrentPlayPeriod->IncreaseSegmentFetchDelay(FTimeValue(FTimeValue::MillisecondsToHNS(100)));
 					}
 					OutDelay.SetFromMilliseconds(150);
-					return IAdaptiveStreamSelector::ESegmentAction::Retry;
+
+					if (Stats.RetryNumber > TriggerPlaylistRefreshAfterNumLateSegments && WorkVars->LastTimePlaylistRefreshRequested < 0.0)
+					{
+						WorkVars->LastTimePlaylistRefreshRequested = Info->ABRGetWallclockTime().GetAsSeconds();
+						Info->ABRTriggerPlaylistRefresh();
+					}
+					return Stats.RetryNumber < MaxTimesToRetryLateSegment ? IAdaptiveStreamSelector::ESegmentAction::Retry : IAdaptiveStreamSelector::ESegmentAction::Fail;
 				}
 
 				// Too many failures already? It's unlikely that there is a way that would magically fix itself now.
diff --git a/Engine/Plugins/Media/ElectraPlayer/Source/ElectraPlayerRuntime/Private/Runtime/Player/AdaptiveStreamingPlayer.cpp b/Engine/Plugins/Media/ElectraPlayer/Source/ElectraPlayerRuntime/Private/Runtime/Player/AdaptiveStreamingPlayer.cpp
--- a/Engine/Plugins/Media/ElectraPlayer/Source/ElectraPlayerRuntime/Private/Runtime/Player/AdaptiveStreamingPlayer.cpp
+++ b/Engine/Plugins/Media/ElectraPlayer/Source/ElectraPlayerRuntime/Private/Runtime/Player/AdaptiveStreamingPlayer.cpp
@@ -1157,6 +1157,11 @@ TSharedPtrTS<FDRMManager> FAdaptiveStreamingPlayer::GetDRMManager()
 	return DrmManager;
 }
 
+void FAdaptiveStreamingPlayer::SetPlaybackEnd(const FTimeValue& InEndAtTime, IPlayerSessionServices::EPlayEndReason InEndingReason, TSharedPtrTS<IPlayerSessionServices::IPlayEndReason> InCustomManifestObject)
+{
+	WorkerThread.SendPlaybackEndMessage(InEndAtTime, InEndingReason, InCustomManifestObject);
+}
+
 
 
 //-----------------------------------------------------------------------------
@@ -1557,6 +1562,14 @@ bool FAdaptiveStreamingPlayer::InternalHandleThreadMessages()
 				break;
 			}
 
+			case FWorkerThreadMessages::FMessage::EType::EndPlaybackAt:
+			{
+				// With a defined end time we just play out until then.
+				PlaybackState.SetShouldPlayOnLiveEdge(false);
+				PlaybackState.SetPlaybackEndAtTime(msg.Data.EndPlaybackAt.EndAtTime);
+				break;
+			}
+
 			default:
 			{
 				checkNoEntry();
@@ -3919,6 +3932,22 @@ void FAdaptiveStreamingPlayer::CheckForStreamEnd()
 	{
 		if (StreamState == EStreamState::eStream_Running)
 		{
+			// First check if there is an end time set at which we need to stop.
+			FTimeValue EndAtTime = PlaybackState.GetPlaybackEndAtTime();
+			if (EndAtTime.IsValid())
+			{
+				if (PlaybackState.GetPlayPosition() >= EndAtTime)
+				{
+					// A forced end time is intended to stop playback. We do NOT look at the loop state here.
+					InternalSetPlaybackEnded();
+					DataBuffersCriticalSection.Lock();
+					NextDataBuffers.Empty();
+					DataBuffersCriticalSection.Unlock();
+					return;
+				}
+			}
+
+			// Check if all data has been used.
 			const bool bHaveVid = bHaveVideoReader.GetWithDefault(false);
 			const bool bHaveAud = bHaveAudioReader.GetWithDefault(false);
 			const bool bHaveTxt = bHaveTextReader.GetWithDefault(false);
@@ -4279,6 +4308,7 @@ void FAdaptiveStreamingPlayer::InternalStop(bool bHoldCurrentFrame)
 	PlaybackRate = 0.0;
 	CurrentState = EPlayerState::eState_Paused;
 	PlaybackState.SetPausedAndPlaying(false, false);
+	PlaybackState.SetPlaybackEndAtTime(FTimeValue::GetInvalid());
 
 	// Flush all access unit buffers.
 	DataBuffersCriticalSection.Lock();
@@ -4620,6 +4650,14 @@ void FAdaptiveStreamingPlayer::ABRTriggerClockSync(IAdaptiveStreamSelector::IPla
 	}
 }
 
+void FAdaptiveStreamingPlayer::ABRTriggerPlaylistRefresh()
+{
+	if (Manifest.IsValid())
+	{
+		Manifest->TriggerPlaylistRefresh();
+	}
+}
+
 
 
 
diff --git a/Engine/Plugins/Media/ElectraPlayer/Source/ElectraPlayerRuntime/Private/Runtime/Player/AdaptiveStreamingPlayerABR.cpp b/Engine/Plugins/Media/ElectraPlayer/Source/ElectraPlayerRuntime/Private/Runtime/Player/AdaptiveStreamingPlayerABR.cpp
--- a/Engine/Plugins/Media/ElectraPlayer/Source/ElectraPlayerRuntime/Private/Runtime/Player/AdaptiveStreamingPlayerABR.cpp
+++ b/Engine/Plugins/Media/ElectraPlayer/Source/ElectraPlayerRuntime/Private/Runtime/Player/AdaptiveStreamingPlayerABR.cpp
@@ -114,6 +114,8 @@ namespace Electra
 		{ return PlayerLiveControl->ABRGetRenderRateScale(); }
 		void ABRTriggerClockSync(IAdaptiveStreamSelector::IPlayerLiveControl::EClockSyncType InClockSyncType) override
 		{ PlayerLiveControl->ABRTriggerClockSync(InClockSyncType); }
+		void ABRTriggerPlaylistRefresh() override
+		{ PlayerLiveControl->ABRTriggerPlaylistRefresh(); }
 
 
 		IPlayerSessionServices*								PlayerSessionServices;
diff --git a/Engine/Plugins/Media/ElectraPlayer/Source/ElectraPlayerRuntime/Private/Runtime/Player/AdaptiveStreamingPlayerABR.h b/Engine/Plugins/Media/ElectraPlayer/Source/ElectraPlayerRuntime/Private/Runtime/Player/AdaptiveStreamingPlayerABR.h
--- a/Engine/Plugins/Media/ElectraPlayer/Source/ElectraPlayerRuntime/Private/Runtime/Player/AdaptiveStreamingPlayerABR.h
+++ b/Engine/Plugins/Media/ElectraPlayer/Source/ElectraPlayerRuntime/Private/Runtime/Player/AdaptiveStreamingPlayerABR.h
@@ -68,6 +68,8 @@ public:
 			Required
 		};
 		virtual void ABRTriggerClockSync(EClockSyncType InClockSyncType) = 0;
+
+		virtual void ABRTriggerPlaylistRefresh() = 0;
 	};
 
 
diff --git a/Engine/Plugins/Media/ElectraPlayer/Source/ElectraPlayerRuntime/Private/Runtime/Player/AdaptiveStreamingPlayerEvents.cpp b/Engine/Plugins/Media/ElectraPlayer/Source/ElectraPlayerRuntime/Private/Runtime/Player/AdaptiveStreamingPlayerEvents.cpp
--- a/Engine/Plugins/Media/ElectraPlayer/Source/ElectraPlayerRuntime/Private/Runtime/Player/AdaptiveStreamingPlayerEvents.cpp
+++ b/Engine/Plugins/Media/ElectraPlayer/Source/ElectraPlayerRuntime/Private/Runtime/Player/AdaptiveStreamingPlayerEvents.cpp
@@ -178,7 +178,7 @@ void FAdaptiveStreamingPlayerAEMSHandler::ExecutePendingEvents(FTimeValue Curren
 	if (DeltaTime >= HandlingInterval)
 	{
 		FTimeRange TriggerRange;
-		TriggerRange.Start = CurrentPlaybackTime - HandlingInterval;
+		TriggerRange.Start = LastPlayPos.IsValid() ? LastPlayPos : CurrentPlaybackTime - HandlingInterval;
 		TriggerRange.End = CurrentPlaybackTime;
 		LastPlayPos = CurrentPlaybackTime;
 		FireEventsInRange(TriggerRange);
diff --git a/Engine/Plugins/Media/ElectraPlayer/Source/ElectraPlayerRuntime/Private/Runtime/Player/AdaptiveStreamingPlayerInternal.h b/Engine/Plugins/Media/ElectraPlayer/Source/ElectraPlayerRuntime/Private/Runtime/Player/AdaptiveStreamingPlayerInternal.h
--- a/Engine/Plugins/Media/ElectraPlayer/Source/ElectraPlayerRuntime/Private/Runtime/Player/AdaptiveStreamingPlayerInternal.h
+++ b/Engine/Plugins/Media/ElectraPlayer/Source/ElectraPlayerRuntime/Private/Runtime/Player/AdaptiveStreamingPlayerInternal.h
@@ -189,6 +189,7 @@ struct FPlaybackState
 		CurrentPlayPosition.SetToInvalid();
 		EncoderLatency.SetToInvalid();
 		CurrentLiveLatency.SetToInvalid();
+		EndPlaybackAtTime.SetToInvalid();
 		LoopState = {};
 		bHaveMetadata = false;
 		bHasEnded = false;
@@ -208,6 +209,7 @@ struct FPlaybackState
 	FTimeValue								CurrentPlayPosition;
 	FTimeValue								EncoderLatency;
 	FTimeValue								CurrentLiveLatency;
+	FTimeValue								EndPlaybackAtTime;
 	IAdaptiveStreamingPlayer::FLoopState	LoopState;
 	bool									bHaveMetadata;
 	bool									bHasEnded;
@@ -455,6 +457,18 @@ struct FPlaybackState
 		FScopeLock lock(&Lock);
 		return bLoopStateHasChanged;
 	}
+
+	void SetPlaybackEndAtTime(const FTimeValue& InAtTime)
+	{
+		FScopeLock lock(&Lock);
+		EndPlaybackAtTime = InAtTime;
+	}
+	FTimeValue GetPlaybackEndAtTime() const
+	{
+		FScopeLock lock(&Lock);
+		return EndPlaybackAtTime;
+	}
+
 };
 
 
@@ -808,11 +822,11 @@ public:
 	FAdaptiveStreamingPlayer(const IAdaptiveStreamingPlayer::FCreateParam& InCreateParameters);
 	virtual ~FAdaptiveStreamingPlayer();
 
-	virtual void SetStaticResourceProviderCallback(const TSharedPtr<IAdaptiveStreamingPlayerResourceProvider, ESPMode::ThreadSafe>& InStaticResourceProvider) override;
-	virtual void SetVideoDecoderResourceDelegate(const TSharedPtr<IVideoDecoderResourceDelegate, ESPMode::ThreadSafe>& ResourceDelegate) override;
+	void SetStaticResourceProviderCallback(const TSharedPtr<IAdaptiveStreamingPlayerResourceProvider, ESPMode::ThreadSafe>& InStaticResourceProvider) override;
+	void SetVideoDecoderResourceDelegate(const TSharedPtr<IVideoDecoderResourceDelegate, ESPMode::ThreadSafe>& ResourceDelegate) override;
 
-	virtual void AddMetricsReceiver(IAdaptiveStreamingPlayerMetrics* InMetricsReceiver) override;
-	virtual void RemoveMetricsReceiver(IAdaptiveStreamingPlayerMetrics* InMetricsReceiver) override;
+	void AddMetricsReceiver(IAdaptiveStreamingPlayerMetrics* InMetricsReceiver) override;
+	void RemoveMetricsReceiver(IAdaptiveStreamingPlayerMetrics* InMetricsReceiver) override;
 
 	void HandleOnce();
 
@@ -831,57 +845,57 @@ public:
 	}
 
 
-	virtual void AddAEMSReceiver(TWeakPtrTS<IAdaptiveStreamingPlayerAEMSReceiver> InReceiver, FString InForSchemeIdUri, FString InForValue, IAdaptiveStreamingPlayerAEMSReceiver::EDispatchMode InDispatchMode) override;
-	virtual void RemoveAEMSReceiver(TWeakPtrTS<IAdaptiveStreamingPlayerAEMSReceiver> InReceiver, FString InForSchemeIdUri, FString InForValue, IAdaptiveStreamingPlayerAEMSReceiver::EDispatchMode InDispatchMode) override;
+	void AddAEMSReceiver(TWeakPtrTS<IAdaptiveStreamingPlayerAEMSReceiver> InReceiver, FString InForSchemeIdUri, FString InForValue, IAdaptiveStreamingPlayerAEMSReceiver::EDispatchMode InDispatchMode) override;
+	void RemoveAEMSReceiver(TWeakPtrTS<IAdaptiveStreamingPlayerAEMSReceiver> InReceiver, FString InForSchemeIdUri, FString InForValue, IAdaptiveStreamingPlayerAEMSReceiver::EDispatchMode InDispatchMode) override;
 
-	virtual void AddSubtitleReceiver(TWeakPtrTS<IAdaptiveStreamingPlayerSubtitleReceiver> InReceiver) override;
-	virtual void RemoveSubtitleReceiver(TWeakPtrTS<IAdaptiveStreamingPlayerSubtitleReceiver> InReceiver) override;
+	void AddSubtitleReceiver(TWeakPtrTS<IAdaptiveStreamingPlayerSubtitleReceiver> InReceiver) override;
+	void RemoveSubtitleReceiver(TWeakPtrTS<IAdaptiveStreamingPlayerSubtitleReceiver> InReceiver) override;
 
-	virtual void Initialize(const FParamDict& Options) override;
-	virtual void ModifyOptions(const FParamDict& InOptionsToSetOrChange, const FParamDict& InOptionsToClear) override;
+	void Initialize(const FParamDict& Options) override;
+	void ModifyOptions(const FParamDict& InOptionsToSetOrChange, const FParamDict& InOptionsToClear) override;
 
-	virtual void SetInitialStreamAttributes(EStreamType StreamType, const FStreamSelectionAttributes& InitialSelection) override;
-	virtual void EnableFrameAccurateSeeking(bool bEnabled) override;
-	virtual void LoadManifest(const FString& manifestURL) override;
+	void SetInitialStreamAttributes(EStreamType StreamType, const FStreamSelectionAttributes& InitialSelection) override;
+	void EnableFrameAccurateSeeking(bool bEnabled) override;
+	void LoadManifest(const FString& manifestURL) override;
 
-	virtual void SeekTo(const FSeekParam& NewPosition) override;
-	virtual void Pause() override;
-	virtual void Resume() override;
-	virtual void Stop() override;
-	virtual void SetPlaybackRange(const FPlaybackRange& InPlaybackRange) override;
-	virtual void GetPlaybackRange(FPlaybackRange& OutPlaybackRange) override;
-	virtual void SetLooping(const FLoopParam& InLoopParams) override;
+	void SeekTo(const FSeekParam& NewPosition) override;
+	void Pause() override;
+	void Resume() override;
+	void Stop() override;
+	void SetPlaybackRange(const FPlaybackRange& InPlaybackRange) override;
+	void GetPlaybackRange(FPlaybackRange& OutPlaybackRange) override;
+	void SetLooping(const FLoopParam& InLoopParams) override;
 
-	virtual FErrorDetail GetError() const override;
+	FErrorDetail GetError() const override;
 
-	virtual bool HaveMetadata() const override;
-	virtual FTimeValue GetDuration() const override;
-	virtual FTimeValue GetPlayPosition() const override;
-	virtual void GetTimelineRange(FTimeRange& OutRange) const override;
-	virtual void GetSeekableRange(FTimeRange& OutRange) const override;
-	virtual void GetSeekablePositions(TArray<FTimespan>& OutPositions) const override;
-	virtual bool HasEnded() const override;
-	virtual bool IsSeeking() const override;
-	virtual bool IsBuffering() const override;
-	virtual bool IsPlaying() const override;
-	virtual bool IsPaused() const override;
+	bool HaveMetadata() const override;
+	FTimeValue GetDuration() const override;
+	FTimeValue GetPlayPosition() const override;
+	void GetTimelineRange(FTimeRange& OutRange) const override;
+	void GetSeekableRange(FTimeRange& OutRange) const override;
+	void GetSeekablePositions(TArray<FTimespan>& OutPositions) const override;
+	bool HasEnded() const override;
+	bool IsSeeking() const override;
+	bool IsBuffering() const override;
+	bool IsPlaying() const override;
+	bool IsPaused() const override;
 
-	virtual void GetLoopState(FLoopState& OutLoopState) const override;
-	virtual void GetTrackMetadata(TArray<FTrackMetadata>& OutTrackMetadata, EStreamType StreamType) const override;
-	//virtual void GetSelectedTrackMetadata(TOptional<FTrackMetadata>& OutSelectedTrackMetadata, EStreamType StreamType) const override;
-	virtual void GetSelectedTrackAttributes(FStreamSelectionAttributes& OutAttributes, EStreamType StreamType) const override;
+	void GetLoopState(FLoopState& OutLoopState) const override;
+	void GetTrackMetadata(TArray<FTrackMetadata>& OutTrackMetadata, EStreamType StreamType) const override;
+	//void GetSelectedTrackMetadata(TOptional<FTrackMetadata>& OutSelectedTrackMetadata, EStreamType StreamType) const override;
+	void GetSelectedTrackAttributes(FStreamSelectionAttributes& OutAttributes, EStreamType StreamType) const override;
 
-	virtual void SetBitrateCeiling(int32 highestSelectableBitrate) override;
-	virtual void SetMaxResolution(int32 MaxWidth, int32 MaxHeight) override;
+	void SetBitrateCeiling(int32 highestSelectableBitrate) override;
+	void SetMaxResolution(int32 MaxWidth, int32 MaxHeight) override;
 
-	//virtual void SelectTrackByMetadata(EStreamType StreamType, const FTrackMetadata& StreamMetadata) override;
-	virtual void SelectTrackByAttributes(EStreamType StreamType, const FStreamSelectionAttributes& Attributes) override;
-	virtual void DeselectTrack(EStreamType StreamType) override;
-	virtual bool IsTrackDeselected(EStreamType StreamType) override;
+	//void SelectTrackByMetadata(EStreamType StreamType, const FTrackMetadata& StreamMetadata) override;
+	void SelectTrackByAttributes(EStreamType StreamType, const FStreamSelectionAttributes& Attributes) override;
+	void DeselectTrack(EStreamType StreamType) override;
+	bool IsTrackDeselected(EStreamType StreamType) override;
 
 #if PLATFORM_ANDROID
-	virtual void Android_UpdateSurface(const TSharedPtr<IOptionPointerValueContainer>& Surface) override;
-	virtual void Android_SuspendOrResumeDecoder(bool bSuspend) override;
+	void Android_UpdateSurface(const TSharedPtr<IOptionPointerValueContainer>& Surface) override;
+	void Android_SuspendOrResumeDecoder(bool bSuspend) override;
 	static FParamDict& Android_Workarounds(FStreamCodecInformation::ECodec InForCodec);
 #endif
 
@@ -890,27 +904,27 @@ public:
 
 private:
 	// Methods from IPlayerSessionServices
-	virtual void PostError(const FErrorDetail& Error) override;
-	virtual void PostLog(Facility::EFacility FromFacility, IInfoLog::ELevel LogLevel, const FString& Message) override;
-	virtual void SendMessageToPlayer(TSharedPtrTS<IPlayerMessage> PlayerMessage) override;
-	virtual void GetExternalGuid(FGuid& OutExternalGuid) override;
-	virtual ISynchronizedUTCTime* GetSynchronizedUTCTime() override;
-	virtual TSharedPtr<IAdaptiveStreamingPlayerResourceProvider, ESPMode::ThreadSafe> GetStaticResourceProvider() override;
-	virtual TSharedPtrTS<IElectraHttpManager> GetHTTPManager() override;
-	virtual TSharedPtrTS<IAdaptiveStreamSelector> GetStreamSelector() override;
-	virtual IPlayerStreamFilter* GetStreamFilter() override;
-	virtual const FCodecSelectionPriorities& GetCodecSelectionPriorities(EStreamType ForStream) override;
+	void PostError(const FErrorDetail& Error) override;
+	void PostLog(Facility::EFacility FromFacility, IInfoLog::ELevel LogLevel, const FString& Message) override;
+	void SendMessageToPlayer(TSharedPtrTS<IPlayerMessage> PlayerMessage) override;
+	void GetExternalGuid(FGuid& OutExternalGuid) override;
+	ISynchronizedUTCTime* GetSynchronizedUTCTime() override;
+	TSharedPtr<IAdaptiveStreamingPlayerResourceProvider, ESPMode::ThreadSafe> GetStaticResourceProvider() override;
+	TSharedPtrTS<IElectraHttpManager> GetHTTPManager() override;
+	TSharedPtrTS<IAdaptiveStreamSelector> GetStreamSelector() override;
+	IPlayerStreamFilter* GetStreamFilter() override;
+	const FCodecSelectionPriorities& GetCodecSelectionPriorities(EStreamType ForStream) override;
 	virtual	TSharedPtrTS<IPlaylistReader> GetManifestReader() override;
-	virtual TSharedPtrTS<IPlayerEntityCache> GetEntityCache() override;
-	virtual TSharedPtrTS<IHTTPResponseCache> GetHTTPResponseCache() override;
-
-	virtual IAdaptiveStreamingPlayerAEMSHandler* GetAEMSEventHandler() override;
-	virtual FParamDict& GetOptions() override;
-	virtual TSharedPtrTS<FDRMManager> GetDRMManager() override;
+	TSharedPtrTS<IPlayerEntityCache> GetEntityCache() override;
+	TSharedPtrTS<IHTTPResponseCache> GetHTTPResponseCache() override;
+	IAdaptiveStreamingPlayerAEMSHandler* GetAEMSEventHandler() override;
+	FParamDict& GetOptions() override;
+	TSharedPtrTS<FDRMManager> GetDRMManager() override;
+	void SetPlaybackEnd(const FTimeValue& InEndAtTime, IPlayerSessionServices::EPlayEndReason InEndingReason, TSharedPtrTS<IPlayerSessionServices::IPlayEndReason> InCustomManifestObject) override;
 
 	// Methods from IPlayerStreamFilter
-	virtual bool CanDecodeStream(const FStreamCodecInformation& InStreamCodecInfo) const override;
-	virtual bool CanDecodeSubtitle(const FString& MimeType, const FString& Codec) const override;
+	bool CanDecodeStream(const FStreamCodecInformation& InStreamCodecInfo) const override;
+	bool CanDecodeSubtitle(const FString& MimeType, const FString& Codec) const override;
 
 
 	// Methods from IAdaptiveStreamSelector::IPlayerLiveControl
@@ -928,6 +942,7 @@ private:
 	void ABRSetRenderRateScale(double InRenderRateScale) override;
 	double ABRGetRenderRateScale() const override;
 	void ABRTriggerClockSync(IAdaptiveStreamSelector::IPlayerLiveControl::EClockSyncType InClockSyncType) override;
+	void ABRTriggerPlaylistRefresh() override;
 
 
 	enum class EPlayerState
@@ -1281,6 +1296,7 @@ private:
 				SelectTrackByMetadata,
 				SelectTrackByAttributes,
 				DeselectTrack,
+				EndPlaybackAt,
 				// Player session message
 				PlayerSession,
 				// Fragment reader messages
@@ -1349,6 +1365,12 @@ private:
 					FTrackMetadata									TrackMetadata;
 					FStreamSelectionAttributes						TrackAttributes;
 				};
+				struct FEndPlaybackAt
+				{
+					FTimeValue											EndAtTime;
+					IPlayerSessionServices::EPlayEndReason				EndingReason;
+					TSharedPtrTS<IPlayerSessionServices::IPlayEndReason> CustomManifestObject;
+				};
 
 				FLoadManifest				ManifestToLoad;
 				FStreamReader				StreamReader;
@@ -1359,6 +1381,7 @@ private:
 				FResolution					Resolution;
 				FInitialStreamSelect		InitialStreamAttribute;
 				FMetadataTrackSelection		TrackSelection;
+				FEndPlaybackAt				EndPlaybackAt;
 			};
 			EType					Type;
 			FData   				Data;
@@ -1472,6 +1495,16 @@ private:
 			TriggerSharedWorkerThread(MoveTemp(Msg));
 		}
 
+		void SendPlaybackEndMessage(const FTimeValue& InEndAtTime, IPlayerSessionServices::EPlayEndReason InEndingReason, TSharedPtrTS<IPlayerSessionServices::IPlayEndReason> InCustomManifestObject)
+		{
+			FMessage Msg;
+			Msg.Type = FMessage::EType::EndPlaybackAt;
+			Msg.Data.EndPlaybackAt.EndAtTime = InEndAtTime;
+			Msg.Data.EndPlaybackAt.EndingReason = InEndingReason;
+			Msg.Data.EndPlaybackAt.CustomManifestObject = InCustomManifestObject;
+			TriggerSharedWorkerThread(MoveTemp(Msg));
+		}
+
 		void TriggerSharedWorkerThread(FMessage Mesg)
 		{
 			WorkMessages.Enqueue(MoveTemp(Mesg));
diff --git a/Engine/Plugins/Media/ElectraPlayer/Source/ElectraPlayerRuntime/Private/Runtime/Player/DASH/ManifestBuilderDASH.cpp b/Engine/Plugins/Media/ElectraPlayer/Source/ElectraPlayerRuntime/Private/Runtime/Player/DASH/ManifestBuilderDASH.cpp
--- a/Engine/Plugins/Media/ElectraPlayer/Source/ElectraPlayerRuntime/Private/Runtime/Player/DASH/ManifestBuilderDASH.cpp
+++ b/Engine/Plugins/Media/ElectraPlayer/Source/ElectraPlayerRuntime/Private/Runtime/Player/DASH/ManifestBuilderDASH.cpp
@@ -2370,7 +2370,7 @@ FTimeValue FManifestDASHInternal::CalculateDistanceToLiveEdge() const
 	if (!Distance.IsValid())
 	{
 		Distance = MPDRoot->GetSuggestedPresentationDelay();
-
+#if 0
 		/*
 			If the @suggestedPresentationDelay is set equal to or shorter than the @minBufferTime there is a
 			conflict. Since we have to buffer @minBufferTime worth of content, which could be arriving in real-time only
@@ -2381,6 +2381,7 @@ FTimeValue FManifestDASHInternal::CalculateDistanceToLiveEdge() const
 			bWarnedAboutTooSmallSuggestedPresentationDelay = true;
 			LogMessage(PlayerSessionServices, IInfoLog::ELevel::Info, FString::Printf(TEXT("MPD@suggestedPresentationDelay of %.3fs is less than MPD@minBufferTime of %.3fs"), Distance.GetAsSeconds(), MPDRoot->GetMinBufferTime().GetAsSeconds()));
 		}
+#endif
 	}
 	// If not set see if there is an MPD@maxSegmentDuration and use that.
 	if (!Distance.IsValid())
diff --git a/Engine/Plugins/Media/ElectraPlayer/Source/ElectraPlayerRuntime/Private/Runtime/Player/DASH/ManifestDASH.cpp b/Engine/Plugins/Media/ElectraPlayer/Source/ElectraPlayerRuntime/Private/Runtime/Player/DASH/ManifestDASH.cpp
--- a/Engine/Plugins/Media/ElectraPlayer/Source/ElectraPlayerRuntime/Private/Runtime/Player/DASH/ManifestDASH.cpp
+++ b/Engine/Plugins/Media/ElectraPlayer/Source/ElectraPlayerRuntime/Private/Runtime/Player/DASH/ManifestDASH.cpp
@@ -367,6 +367,17 @@ void FManifestDASH::TriggerClockSync(IManifest::EClockSyncType InClockSyncType)
 	Reader->RequestClockResync();
 }
 
+void FManifestDASH::TriggerPlaylistRefresh()
+{
+	TSharedPtrTS<FManifestDASHInternal> Manifest = CurrentManifest;
+	// Trigger only when updates are not expected regularly.
+	if (Manifest.IsValid() && Manifest->AreUpdatesExpected() && (Manifest->GetMinimumUpdatePeriod() == FTimeValue::GetZero() || Manifest->GetMinimumUpdatePeriod() > FTimeValue(10.0)))
+	{
+		IPlaylistReaderDASH* Reader = static_cast<IPlaylistReaderDASH*>(PlayerSessionServices->GetManifestReader().Get());
+		Reader->RequestMPDUpdate(IPlaylistReaderDASH::EMPDRequestType::GetLatestSegment);
+	}
+}
+
 
 IStreamReader* FManifestDASH::CreateStreamReaderHandler()
 {
diff --git a/Engine/Plugins/Media/ElectraPlayer/Source/ElectraPlayerRuntime/Private/Runtime/Player/DASH/ManifestDASH.h b/Engine/Plugins/Media/ElectraPlayer/Source/ElectraPlayerRuntime/Private/Runtime/Player/DASH/ManifestDASH.h
--- a/Engine/Plugins/Media/ElectraPlayer/Source/ElectraPlayerRuntime/Private/Runtime/Player/DASH/ManifestDASH.h
+++ b/Engine/Plugins/Media/ElectraPlayer/Source/ElectraPlayerRuntime/Private/Runtime/Player/DASH/ManifestDASH.h
@@ -34,6 +34,7 @@ public:
 	void GetTrackMetadata(TArray<FTrackMetadata>& OutMetadata, EStreamType StreamType) const override;
 	void UpdateDynamicRefetchCounter() override;
 	void TriggerClockSync(IManifest::EClockSyncType InClockSyncType) override;
+	void TriggerPlaylistRefresh() override;
 	IStreamReader* CreateStreamReaderHandler() override;
 
 	FResult FindPlayPeriod(TSharedPtrTS<IPlayPeriod>& OutPlayPeriod, const FPlayStartPosition& StartPosition, ESearchType SearchType) override;
diff --git a/Engine/Plugins/Media/ElectraPlayer/Source/ElectraPlayerRuntime/Private/Runtime/Player/DASH/PlaylistReaderDASH.cpp b/Engine/Plugins/Media/ElectraPlayer/Source/ElectraPlayerRuntime/Private/Runtime/Player/DASH/PlaylistReaderDASH.cpp
--- a/Engine/Plugins/Media/ElectraPlayer/Source/ElectraPlayerRuntime/Private/Runtime/Player/DASH/PlaylistReaderDASH.cpp
+++ b/Engine/Plugins/Media/ElectraPlayer/Source/ElectraPlayerRuntime/Private/Runtime/Player/DASH/PlaylistReaderDASH.cpp
@@ -454,6 +454,7 @@ void FPlaylistReaderDASH::CheckForMPDUpdate()
 						}
 						if (bEventApplies)
 						{
+							PlayerSessionServices->SetPlaybackEnd(Request->ValidUntil + Request->NewDuration, IPlayerSessionServices::EPlayEndReason::EndAll, nullptr);
 							Manifest->EndPresentationAt(Request->ValidUntil + Request->NewDuration, Request->PeriodId);
 						}
 					}
@@ -1331,7 +1332,7 @@ void FPlaylistReaderDASH::ManifestUpdateDownloadCompleted(FResourceLoadRequestPt
 {
 	bIsMPDUpdateInProgress = false;
 
-	if (bSuccess)
+	if (bSuccess || (Request->GetConnectionInfo() && Request->GetConnectionInfo()->StatusInfo.HTTPStatus == 410))
 	{
 		if (!Request->Request.IsValid() || !Request->Request->GetResponseBuffer().IsValid())
 		{
@@ -1374,6 +1375,13 @@ void FPlaylistReaderDASH::ManifestUpdateDownloadCompleted(FResourceLoadRequestPt
 				RequestsLock.Unlock();
 			}
 		}
+		// If we got a "410 - Gone" we set the end of presentation to "now"
+		else if (ConnInfo->StatusInfo.HTTPStatus == 410)
+		{
+			PlayerSessionServices->SetPlaybackEnd(FetchTime, IPlayerSessionServices::EPlayEndReason::EndAll, nullptr);
+			Manifest->GetMPDRoot()->SetFetchTime(FetchTime);
+			Manifest->EndPresentationAt(FetchTime, FString());
+		}
 		else
 		{
 			PlayerSessionServices->SendMessageToPlayer(IPlaylistReader::PlaylistDownloadMessage::Create(ConnInfo, Playlist::EListType::Master, Playlist::ELoadType::Update));
@@ -1621,7 +1629,12 @@ void FPlaylistReaderDASH::OnMediaPlayerEventReceived(TSharedPtrTS<IAdaptiveStrea
 	}
 	else if (InEvent->GetSchemeIdUri().Equals(DASH::Schemes::ManifestEvents::Scheme_urn_mpeg_dash_event_ttfn_2016))
 	{
-		// TBD
+		if (Manifest.IsValid())
+		{
+			// For now we do not consider the event value and just end playback.
+			PlayerSessionServices->SetPlaybackEnd(InEvent->GetPresentationTime(), IPlayerSessionServices::EPlayEndReason::EndAll, nullptr);
+			Manifest->EndPresentationAt(InEvent->GetPresentationTime(), FString());
+		}
 	}
 }
 
diff --git a/Engine/Plugins/Media/ElectraPlayer/Source/ElectraPlayerRuntime/Private/Runtime/Player/HLS/ManifestHLS.cpp b/Engine/Plugins/Media/ElectraPlayer/Source/ElectraPlayerRuntime/Private/Runtime/Player/HLS/ManifestHLS.cpp
--- a/Engine/Plugins/Media/ElectraPlayer/Source/ElectraPlayerRuntime/Private/Runtime/Player/HLS/ManifestHLS.cpp
+++ b/Engine/Plugins/Media/ElectraPlayer/Source/ElectraPlayerRuntime/Private/Runtime/Player/HLS/ManifestHLS.cpp
@@ -275,12 +275,15 @@ IManifest::FResult FManifestHLS::FindNextPlayPeriod(TSharedPtrTS<IPlayPeriod>& O
 	return IManifest::FResult(IManifest::FResult::EType::PastEOS);
 }
 
-
 void FManifestHLS::TriggerClockSync(IManifest::EClockSyncType InClockSyncType)
 {
 	// No-op.
 }
 
+void FManifestHLS::TriggerPlaylistRefresh()
+{
+	// No-op.
+}
 
 /**
  * Creates a stream reader for the media segments.
diff --git a/Engine/Plugins/Media/ElectraPlayer/Source/ElectraPlayerRuntime/Private/Runtime/Player/HLS/ManifestHLS.h b/Engine/Plugins/Media/ElectraPlayer/Source/ElectraPlayerRuntime/Private/Runtime/Player/HLS/ManifestHLS.h
--- a/Engine/Plugins/Media/ElectraPlayer/Source/ElectraPlayerRuntime/Private/Runtime/Player/HLS/ManifestHLS.h
+++ b/Engine/Plugins/Media/ElectraPlayer/Source/ElectraPlayerRuntime/Private/Runtime/Player/HLS/ManifestHLS.h
@@ -32,6 +32,7 @@ public:
 	void GetTrackMetadata(TArray<FTrackMetadata>& OutMetadata, EStreamType StreamType) const override;
 	void UpdateDynamicRefetchCounter() override;
 	void TriggerClockSync(EClockSyncType InClockSyncType) override;
+	void TriggerPlaylistRefresh() override;
 	IStreamReader *CreateStreamReaderHandler() override;
 
 	FResult FindPlayPeriod(TSharedPtrTS<IPlayPeriod>& OutPlayPeriod, const FPlayStartPosition& StartPosition, ESearchType SearchType) override;
diff --git a/Engine/Plugins/Media/ElectraPlayer/Source/ElectraPlayerRuntime/Private/Runtime/Player/Manifest.h b/Engine/Plugins/Media/ElectraPlayer/Source/ElectraPlayerRuntime/Private/Runtime/Player/Manifest.h
--- a/Engine/Plugins/Media/ElectraPlayer/Source/ElectraPlayerRuntime/Private/Runtime/Player/Manifest.h
+++ b/Engine/Plugins/Media/ElectraPlayer/Source/ElectraPlayerRuntime/Private/Runtime/Player/Manifest.h
@@ -331,6 +331,7 @@ namespace Electra
 		};
 		virtual void TriggerClockSync(EClockSyncType InClockSyncType) = 0;
 
+		virtual void TriggerPlaylistRefresh() = 0;
 
 		//-------------------------------------------------------------------------
 		// Stream fragment reader
diff --git a/Engine/Plugins/Media/ElectraPlayer/Source/ElectraPlayerRuntime/Private/Runtime/Player/PlayerSessionServices.h b/Engine/Plugins/Media/ElectraPlayer/Source/ElectraPlayerRuntime/Private/Runtime/Player/PlayerSessionServices.h
--- a/Engine/Plugins/Media/ElectraPlayer/Source/ElectraPlayerRuntime/Private/Runtime/Player/PlayerSessionServices.h
+++ b/Engine/Plugins/Media/ElectraPlayer/Source/ElectraPlayerRuntime/Private/Runtime/Player/PlayerSessionServices.h
@@ -136,6 +136,19 @@ namespace Electra
 		  * Returns the DRM manager, if any.
 		  */
 		virtual TSharedPtrTS<FDRMManager> GetDRMManager() = 0;
+
+		enum class EPlayEndReason
+		{
+			EndAll,
+			NextItem,
+			ErrorCondition
+		};
+		class IPlayEndReason
+		{
+		public:
+			virtual ~IPlayEndReason() = default;
+		};
+		virtual void SetPlaybackEnd(const FTimeValue& InEndAtTime, EPlayEndReason InEndingReason, TSharedPtrTS<IPlayEndReason> InCustomManifestObject) = 0;
 	};
 
 
diff --git a/Engine/Plugins/Media/ElectraPlayer/Source/ElectraPlayerRuntime/Private/Runtime/Player/mp4/ManifestMP4.cpp b/Engine/Plugins/Media/ElectraPlayer/Source/ElectraPlayerRuntime/Private/Runtime/Player/mp4/ManifestMP4.cpp
--- a/Engine/Plugins/Media/ElectraPlayer/Source/ElectraPlayerRuntime/Private/Runtime/Player/mp4/ManifestMP4.cpp
+++ b/Engine/Plugins/Media/ElectraPlayer/Source/ElectraPlayerRuntime/Private/Runtime/Player/mp4/ManifestMP4.cpp
@@ -190,6 +190,11 @@ void FManifestMP4Internal::TriggerClockSync(IManifest::EClockSyncType InClockSyn
 	// No-op.
 }
 
+void FManifestMP4Internal::TriggerPlaylistRefresh()
+{
+	// No-op.
+}
+
 
 //-----------------------------------------------------------------------------
 /**
diff --git a/Engine/Plugins/Media/ElectraPlayer/Source/ElectraPlayerRuntime/Private/Runtime/Player/mp4/ManifestMP4.h b/Engine/Plugins/Media/ElectraPlayer/Source/ElectraPlayerRuntime/Private/Runtime/Player/mp4/ManifestMP4.h
--- a/Engine/Plugins/Media/ElectraPlayer/Source/ElectraPlayerRuntime/Private/Runtime/Player/mp4/ManifestMP4.h
+++ b/Engine/Plugins/Media/ElectraPlayer/Source/ElectraPlayerRuntime/Private/Runtime/Player/mp4/ManifestMP4.h
@@ -60,6 +60,7 @@ public:
 	TSharedPtrTS<IProducerReferenceTimeInfo> GetProducerReferenceTimeInfo(int64 ID) const override;
 	void UpdateDynamicRefetchCounter() override;
 	void TriggerClockSync(EClockSyncType InClockSyncType) override;
+	void TriggerPlaylistRefresh() override;
 
 	IStreamReader* CreateStreamReaderHandler() override;
 	FResult FindPlayPeriod(TSharedPtrTS<IPlayPeriod>& OutPlayPeriod, const FPlayStartPosition& StartPosition, ESearchType SearchType) override;
