ElectraPlayer: Allowing for the facade to use block on fetch

#rnx

diff --git a/Engine/Plugins/Media/ElectraPlayer/Source/ElectraPlayerPlugin/Private/ElectraPlayerPlugin.cpp b/Engine/Plugins/Media/ElectraPlayer/Source/ElectraPlayerPlugin/Private/ElectraPlayerPlugin.cpp
--- a/Engine/Plugins/Media/ElectraPlayer/Source/ElectraPlayerPlugin/Private/ElectraPlayerPlugin.cpp
+++ b/Engine/Plugins/Media/ElectraPlayer/Source/ElectraPlayerPlugin/Private/ElectraPlayerPlugin.cpp
@@ -709,21 +709,22 @@ uint32 FElectraPlayerPlugin::GetNewResourcesOnOpen() const
 bool FElectraPlayerPlugin::CanControl(EMediaControl Control) const
 {
 	EMediaState CurrentState = GetState();
-	if (Control == EMediaControl::Pause)
+	if (Control == EMediaControl::BlockOnFetch)
 	{
 		return CurrentState == EMediaState::Playing;
 	}
-
-	if (Control == EMediaControl::Resume)
+	else if (Control == EMediaControl::Pause)
+	{
+		return CurrentState == EMediaState::Playing;
+	}
+	else if (Control == EMediaControl::Resume)
 	{
 		return CurrentState == EMediaState::Paused || CurrentState == EMediaState::Stopped;
 	}
-
-	if (Control == EMediaControl::Seek)
+	else if (Control == EMediaControl::Seek)
 	{
-		return CurrentState == EMediaState::Playing || CurrentState == EMediaState::Paused || CurrentState == EMediaState::Stopped ;//|| CurrentState == Preparing;
+		return CurrentState == EMediaState::Playing || CurrentState == EMediaState::Paused || CurrentState == EMediaState::Stopped;
 	}
-
 	return false;
 }
 
