ElectraPlayer: Added handling of changing track availability in multiperiod DASH

#rnx

diff --git a/Engine/Plugins/Media/ElectraPlayer/Source/ElectraPlayerRuntime/Private/ElectraPlayer.cpp b/Engine/Plugins/Media/ElectraPlayer/Source/ElectraPlayerRuntime/Private/ElectraPlayer.cpp
--- a/Engine/Plugins/Media/ElectraPlayer/Source/ElectraPlayerRuntime/Private/ElectraPlayer.cpp
+++ b/Engine/Plugins/Media/ElectraPlayer/Source/ElectraPlayerRuntime/Private/ElectraPlayer.cpp
@@ -202,6 +202,7 @@ void FElectraPlayer::ClearToDefaultState()
 	SelectedVideoTrackIndex = -1;
 	SelectedAudioTrackIndex = -1;
 	SelectedSubtitleTrackIndex = -1;
+	bVideoTrackIndexDirty = true;
 	bAudioTrackIndexDirty = true;
 	bSubtitleTrackIndexDirty = true;
 	bInitialSeekPerformed = false;
@@ -880,6 +881,7 @@ bool FElectraPlayer::SetLooping(bool bLooping)
 		IAdaptiveStreamingPlayer::FLoopParam loop;
 		loop.bEnableLooping = bLooping;
 		CurrentPlayer->AdaptivePlayer->SetLooping(loop);
+		UE_LOG(LogElectraPlayer, Verbose, TEXT("[%p][%p] IMediaPlayer::SetLooping(%s)"), this, CurrentPlayer.Get(), bLooping?TEXT("true"):TEXT("false"));
 		return true;
 	}
 	return false;
@@ -1312,11 +1314,6 @@ int32 FElectraPlayer::GetNumTrackFormats(EPlayerTrackType TrackType, int32 Track
 
 int32 FElectraPlayer::GetSelectedTrack(EPlayerTrackType TrackType) const
 {
-	// For now Electra never has any caption or metadata tracks
-	if (TrackType == EPlayerTrackType::Video)
-	{
-		return SelectedVideoTrackIndex;
-	}
 	/*
 		To reduce the overhead of this function we check for the track the underlying player has
 		actually selected only when we were told the tracks changed.
@@ -1326,71 +1323,53 @@ int32 FElectraPlayer::GetSelectedTrack(EPlayerTrackType TrackType) const
 		the audio stream when transitioning from one period into the next, which may change the index of
 		the selected track.
 	*/
-	else if (TrackType == EPlayerTrackType::Audio)
+
+	auto CheckAndReselectTrack = [this](Electra::EStreamType InStreamType, bool& InOutDirtyFlag, int32& InOutSelectedIndex, int32 InNumTracks) -> int32
 	{
-		if (bAudioTrackIndexDirty)
+		if (InOutDirtyFlag)
 		{
-			if (NumTracksAudio == 0)
+			if (InNumTracks == 0)
 			{
-				SelectedAudioTrackIndex = -1;
+				InOutSelectedIndex = -1;
 			}
 			else
 			{
 				TSharedPtr<FInternalPlayerImpl, ESPMode::ThreadSafe> LockedPlayer = CurrentPlayer;
 				if (LockedPlayer.IsValid() && LockedPlayer->AdaptivePlayer.IsValid())
 				{
-					if (LockedPlayer->AdaptivePlayer->IsTrackDeselected(Electra::EStreamType::Audio))
+					if (LockedPlayer->AdaptivePlayer->IsTrackDeselected(InStreamType))
 					{
-						SelectedAudioTrackIndex = -1;
-						bAudioTrackIndexDirty = false;
+						InOutSelectedIndex = -1;
+						InOutDirtyFlag = false;
 					}
 					else
 					{
 						Electra::FStreamSelectionAttributes Attributes;
-						LockedPlayer->AdaptivePlayer->GetSelectedTrackAttributes(Attributes, Electra::EStreamType::Audio);
+						LockedPlayer->AdaptivePlayer->GetSelectedTrackAttributes(Attributes, InStreamType);
 						if (Attributes.OverrideIndex.IsSet())
 						{
-							SelectedAudioTrackIndex = Attributes.OverrideIndex.GetValue();
-							bAudioTrackIndexDirty = false;
+							InOutSelectedIndex = Attributes.OverrideIndex.GetValue();
+							InOutDirtyFlag = false;
 						}
 					}
 				}
 			}
 		}
-		return SelectedAudioTrackIndex;
+		return InOutSelectedIndex;
+	};
+
+	// Electra does not have caption or metadata tracks, handle only video, audio and subtitles.
+	if (TrackType == EPlayerTrackType::Video)
+	{
+		return CheckAndReselectTrack(Electra::EStreamType::Video, bVideoTrackIndexDirty, SelectedVideoTrackIndex, NumTracksVideo);
+	}
+	else if (TrackType == EPlayerTrackType::Audio)
+	{
+		return CheckAndReselectTrack(Electra::EStreamType::Audio, bAudioTrackIndexDirty, SelectedAudioTrackIndex, NumTracksAudio);
 	}
 	else if (TrackType == EPlayerTrackType::Subtitle)
 	{
-		if (bSubtitleTrackIndexDirty)
-		{
-			if (NumTracksSubtitle == 0)
-			{
-				SelectedSubtitleTrackIndex = -1;
-			}
-			else
-			{
-				TSharedPtr<FInternalPlayerImpl, ESPMode::ThreadSafe> LockedPlayer = CurrentPlayer;
-				if (LockedPlayer.IsValid() && LockedPlayer->AdaptivePlayer.IsValid())
-				{
-					if (LockedPlayer->AdaptivePlayer->IsTrackDeselected(Electra::EStreamType::Subtitle))
-					{
-						SelectedSubtitleTrackIndex = -1;
-						bSubtitleTrackIndexDirty = false;
-					}
-					else
-					{
-						Electra::FStreamSelectionAttributes Attributes;
-						LockedPlayer->AdaptivePlayer->GetSelectedTrackAttributes(Attributes, Electra::EStreamType::Subtitle);
-						if (Attributes.OverrideIndex.IsSet())
-						{
-							SelectedSubtitleTrackIndex = Attributes.OverrideIndex.GetValue();
-							bSubtitleTrackIndexDirty = false;
-						}
-					}
-				}
-			}
-		}
-		return SelectedSubtitleTrackIndex;
+		return CheckAndReselectTrack(Electra::EStreamType::Subtitle, bSubtitleTrackIndexDirty, SelectedSubtitleTrackIndex, NumTracksSubtitle);
 	}
 	return -1;
 }
@@ -1979,6 +1958,7 @@ void FElectraPlayer::HandlePlayerEventTracksChanged()
 	CurrentPlayer->AdaptivePlayer->GetTrackMetadata(SubtitleStreamMetaData, Electra::EStreamType::Subtitle);
 	NumTracksSubtitle = SubtitleStreamMetaData.Num();
 
+	bVideoTrackIndexDirty = true;
 	bAudioTrackIndexDirty = true;
 	bSubtitleTrackIndexDirty = true;
 
diff --git a/Engine/Plugins/Media/ElectraPlayer/Source/ElectraPlayerRuntime/Private/ElectraPlayer.h b/Engine/Plugins/Media/ElectraPlayer/Source/ElectraPlayerRuntime/Private/ElectraPlayer.h
--- a/Engine/Plugins/Media/ElectraPlayer/Source/ElectraPlayerRuntime/Private/ElectraPlayer.h
+++ b/Engine/Plugins/Media/ElectraPlayer/Source/ElectraPlayerRuntime/Private/ElectraPlayer.h
@@ -386,6 +386,7 @@ private:
 	mutable int32									SelectedVideoTrackIndex;
 	mutable int32									SelectedAudioTrackIndex;
 	mutable int32									SelectedSubtitleTrackIndex;
+	mutable bool									bVideoTrackIndexDirty;
 	mutable bool									bAudioTrackIndexDirty;
 	mutable bool									bSubtitleTrackIndexDirty;
 
diff --git a/Engine/Plugins/Media/ElectraPlayer/Source/ElectraPlayerRuntime/Private/Runtime/Player/ABRRules/ABRLiveStream.cpp b/Engine/Plugins/Media/ElectraPlayer/Source/ElectraPlayerRuntime/Private/Runtime/Player/ABRRules/ABRLiveStream.cpp
--- a/Engine/Plugins/Media/ElectraPlayer/Source/ElectraPlayerRuntime/Private/Runtime/Player/ABRRules/ABRLiveStream.cpp
+++ b/Engine/Plugins/Media/ElectraPlayer/Source/ElectraPlayerRuntime/Private/Runtime/Player/ABRRules/ABRLiveStream.cpp
@@ -276,7 +276,7 @@ private:
 	{
 		IAdaptiveStreamSelector::IPlayerLiveControl::FABRBufferStats bs;
 		Info->ABRGetStreamBufferStats(bs, InStreamType);
-		bEOS = bs.bReachedEnd;
+		bEOS = bs.bReachedEnd || bs.bEndOfTrack;
 		return bs.PlayableContentDuration.GetAsSeconds();
 	}
 
@@ -1018,19 +1018,22 @@ IAdaptiveStreamSelector::ESegmentAction FABRLiveStream::EvaluateForQuality(TArra
 		return IAdaptiveStreamSelector::ESegmentAction::FetchNext;
 	}
 
-	TArray<TSharedPtrTS<FABRStreamInformation>> PossibleRepresentations;
-	PossibleRepresentations.Add(InOutCandidates[0]);
-	const int32 MaxAllowedBandwidth = Info->GetBandwidthCeiling();
-	const FStreamCodecInformation::FResolution MaxAllowedResolution = Info->GetMaxStreamResolution();
-	for(int32 nStr=1; nStr<InOutCandidates.Num(); ++nStr)
+	if (InOutCandidates.Num())
 	{
-		// Check if bitrate and resolution are acceptable
-		if (InOutCandidates[nStr]->Bitrate <= MaxAllowedBandwidth && !InOutCandidates[nStr]->Resolution.ExceedsLimit(MaxAllowedResolution))
+		TArray<TSharedPtrTS<FABRStreamInformation>> PossibleRepresentations;
+		PossibleRepresentations.Add(InOutCandidates[0]);
+		const int32 MaxAllowedBandwidth = Info->GetBandwidthCeiling();
+		const FStreamCodecInformation::FResolution MaxAllowedResolution = Info->GetMaxStreamResolution();
+		for(int32 nStr=1; nStr<InOutCandidates.Num(); ++nStr)
 		{
-			PossibleRepresentations.Add(InOutCandidates[nStr]);
+			// Check if bitrate and resolution are acceptable
+			if (InOutCandidates[nStr]->Bitrate <= MaxAllowedBandwidth && !InOutCandidates[nStr]->Resolution.ExceedsLimit(MaxAllowedResolution))
+			{
+				PossibleRepresentations.Add(InOutCandidates[nStr]);
+			}
 		}
+		Swap(InOutCandidates, PossibleRepresentations);
 	}
-	Swap(InOutCandidates, PossibleRepresentations);
 	return IAdaptiveStreamSelector::ESegmentAction::FetchNext;
 }
 
diff --git a/Engine/Plugins/Media/ElectraPlayer/Source/ElectraPlayerRuntime/Private/Runtime/Player/ABRRules/ABRVoDPlus.cpp b/Engine/Plugins/Media/ElectraPlayer/Source/ElectraPlayerRuntime/Private/Runtime/Player/ABRRules/ABRVoDPlus.cpp
--- a/Engine/Plugins/Media/ElectraPlayer/Source/ElectraPlayerRuntime/Private/Runtime/Player/ABRRules/ABRVoDPlus.cpp
+++ b/Engine/Plugins/Media/ElectraPlayer/Source/ElectraPlayerRuntime/Private/Runtime/Player/ABRRules/ABRVoDPlus.cpp
@@ -141,7 +141,7 @@ private:
 	{
 		IAdaptiveStreamSelector::IPlayerLiveControl::FABRBufferStats bs;
 		Info->ABRGetStreamBufferStats(bs, InStreamType);
-		bEOS = bs.bReachedEnd;
+		bEOS = bs.bReachedEnd || bs.bEndOfTrack;
 		return bs.PlayableContentDuration.GetAsSeconds();
 	}
 
@@ -704,19 +704,22 @@ IAdaptiveStreamSelector::ESegmentAction FABROnDemandPlus::EvaluateForQuality(TAr
 		return IAdaptiveStreamSelector::ESegmentAction::FetchNext;
 	}
 
-	TArray<TSharedPtrTS<FABRStreamInformation>> PossibleRepresentations;
-	PossibleRepresentations.Add(InOutCandidates[0]);
-	const int32 MaxAllowedBandwidth = Info->GetBandwidthCeiling();
-	const FStreamCodecInformation::FResolution MaxAllowedResolution = Info->GetMaxStreamResolution();
-	for(int32 nStr=1; nStr<InOutCandidates.Num(); ++nStr)
+	if (InOutCandidates.Num())
 	{
-		// Check if bitrate and resolution are acceptable
-		if (InOutCandidates[nStr]->Bitrate <= MaxAllowedBandwidth && !InOutCandidates[nStr]->Resolution.ExceedsLimit(MaxAllowedResolution))
+		TArray<TSharedPtrTS<FABRStreamInformation>> PossibleRepresentations;
+		PossibleRepresentations.Add(InOutCandidates[0]);
+		const int32 MaxAllowedBandwidth = Info->GetBandwidthCeiling();
+		const FStreamCodecInformation::FResolution MaxAllowedResolution = Info->GetMaxStreamResolution();
+		for(int32 nStr=1; nStr<InOutCandidates.Num(); ++nStr)
 		{
-			PossibleRepresentations.Add(InOutCandidates[nStr]);
+			// Check if bitrate and resolution are acceptable
+			if (InOutCandidates[nStr]->Bitrate <= MaxAllowedBandwidth && !InOutCandidates[nStr]->Resolution.ExceedsLimit(MaxAllowedResolution))
+			{
+				PossibleRepresentations.Add(InOutCandidates[nStr]);
+			}
 		}
+		Swap(InOutCandidates, PossibleRepresentations);
 	}
-	Swap(InOutCandidates, PossibleRepresentations);
 	return IAdaptiveStreamSelector::ESegmentAction::FetchNext;
 }
 
diff --git a/Engine/Plugins/Media/ElectraPlayer/Source/ElectraPlayerRuntime/Private/Runtime/Player/AdaptiveStreamingPlayer.cpp b/Engine/Plugins/Media/ElectraPlayer/Source/ElectraPlayerRuntime/Private/Runtime/Player/AdaptiveStreamingPlayer.cpp
--- a/Engine/Plugins/Media/ElectraPlayer/Source/ElectraPlayerRuntime/Private/Runtime/Player/AdaptiveStreamingPlayer.cpp
+++ b/Engine/Plugins/Media/ElectraPlayer/Source/ElectraPlayerRuntime/Private/Runtime/Player/AdaptiveStreamingPlayer.cpp
@@ -88,14 +88,12 @@ FAdaptiveStreamingPlayer::FAdaptiveStreamingPlayer(const IAdaptiveStreamingPlaye
 	PlaybackRate   		    = 0.0;
 	RenderRateScale			= 1.0;
 	StreamReaderHandler     = nullptr;
+	bStreamingHasStarted	= false;
 	RebufferCause   	    = ERebufferCause::None;
 	Manifest  				= nullptr;
 	LastBufferingState 		= EPlayerState::eState_Buffering;
 	bIsClosing 				= false;
 
-	bHaveVideoReader.Reset();
-	bHaveAudioReader.Reset();
-	bHaveTextReader.Reset();
 	RebufferDetectedAtPlayPos.SetToInvalid();
 
 	CurrentPlaybackSequenceID[0] = 0;
@@ -1273,10 +1271,10 @@ void FAdaptiveStreamingPlayer::InternalHandleOnce()
 		{
 			if (RenderClock->IsRunning())
 			{
-				IMediaRenderClock::ERendererType ClockSource = IMediaRenderClock::ERendererType::Video;
-				if (DataAvailabilityStateAud.Availability == Metrics::FDataAvailabilityChange::EAvailability::DataAvailable)
+				IMediaRenderClock::ERendererType ClockSource = IMediaRenderClock::ERendererType::Audio;
+				if (DataAvailabilityStateAud.Availability == Metrics::FDataAvailabilityChange::EAvailability::DataNotAvailable)
 				{
-					ClockSource = IMediaRenderClock::ERendererType::Audio;
+					ClockSource = IMediaRenderClock::ERendererType::Video;
 				}
 				FTimeValue playPos = RenderClock->GetInterpolatedRenderTime(ClockSource);
 				PlaybackState.SetPlayPosition(ClampTimeToCurrentRange(playPos, false, true));
@@ -1732,8 +1730,6 @@ void FAdaptiveStreamingPlayer::HandleMetadataChanges()
 	}
 
 	FTimeValue CurrentTime = PlaybackState.GetPlayPosition();
-	FLoopState CurrentLoop;
-	PlaybackState.GetLoopState(CurrentLoop);
 
 	// See if we may be reaching the next period.
 	if (UpcomingPeriods.Num() && CurrentTime.IsValid() && CurrentTime != MetadataHandlingState.LastHandlingTime)
@@ -1761,6 +1757,11 @@ void FAdaptiveStreamingPlayer::HandleMetadataChanges()
 						DispatchEvent(FMetricEvent::ReportTracksChanged());
 					}
 
+					// Update the data availability for this period.
+					UpdateDataAvailabilityState(DataAvailabilityStateVid, MetadataVideo.Num() ? Metrics::FDataAvailabilityChange::EAvailability::DataAvailable : Metrics::FDataAvailabilityChange::EAvailability::DataNotAvailable);
+					UpdateDataAvailabilityState(DataAvailabilityStateAud, MetadataAudio.Num() ? Metrics::FDataAvailabilityChange::EAvailability::DataAvailable : Metrics::FDataAvailabilityChange::EAvailability::DataNotAvailable);
+					UpdateDataAvailabilityState(DataAvailabilityStateTxt, MetadataSubtitle.Num() ? Metrics::FDataAvailabilityChange::EAvailability::DataAvailable : Metrics::FDataAvailabilityChange::EAvailability::DataNotAvailable);
+
 					// Inform AEMS handler of the period transition.
 					AEMSEventHandler->PlaybackPeriodTransition(UpcomingPeriods[i].ID, UpcomingPeriods[i].TimeRange);
 				}
@@ -1773,7 +1774,8 @@ void FAdaptiveStreamingPlayer::HandleMetadataChanges()
 		// if they are used in a future loop cycle.
 		for(int32 i=0; i<UpcomingPeriods.Num(); ++i)
 		{
-			if (UpcomingPeriods[i].TimeRange.End.IsValid() && CurrentTime >= UpcomingPeriods[i].TimeRange.End && UpcomingPeriods[i].LoopCount <= CurrentLoop.Count)
+			if ((UpcomingPeriods[i].TimeRange.End.IsValid() && CurrentTime >= UpcomingPeriods[i].TimeRange.End && UpcomingPeriods[i].LoopCount <= CurrentTime.GetSequenceIndex()) ||
+				(UpcomingPeriods[i].TimeRange.Start.IsValid() && CurrentTime < UpcomingPeriods[i].TimeRange.Start && UpcomingPeriods[i].LoopCount < CurrentTime.GetSequenceIndex()))
 			{
 				UpcomingPeriods.RemoveAt(i);
 				--i;
@@ -1793,7 +1795,7 @@ void FAdaptiveStreamingPlayer::HandleMetadataChanges()
 		{
 			if (((ActivePeriods[i].TimeRange.End.IsValid() && PastTime > ActivePeriods[i].TimeRange.End) ||
 				(!ActivePeriods[i].TimeRange.End.IsValid() && (i+1) < ActivePeriods.Num() && PastTime >= ActivePeriods[i+1].TimeRange.Start)) &&
-				ActivePeriods[i].LoopCount <= CurrentLoop.Count)
+				ActivePeriods[i].LoopCount <= CurrentTime.GetSequenceIndex())
 			{
 				// Do not remove the only active period.
 				if (ActivePeriods.Num() > 1)
@@ -1817,7 +1819,7 @@ void FAdaptiveStreamingPlayer::HandleAEMSEvents()
 {
 	if (CurrentState == EPlayerState::eState_Playing)
 	{
-		FTimeValue Current = GetCurrentPlayTime();
+		FTimeValue Current = PlaybackState.GetPlayPosition();
 		AEMSEventHandler->ExecutePendingEvents(Current);
 	}
 }
@@ -1957,78 +1959,88 @@ double FAdaptiveStreamingPlayer::GetMinBufferTimeBeforePlayback()
 	return kMinBufferBeforePlayback;
 }
 
+bool FAdaptiveStreamingPlayer::IsExpectedToStreamNow(EStreamType InType)
+{
+	if (!StreamReaderHandler)
+	{
+		return false;
+	}
+	if (InType == EStreamType::Video)
+	{
+		return CurrentPlayPeriodVideo.IsValid() ? CurrentPlayPeriodVideo->GetSelectedStreamBufferSourceInfo(InType).IsValid() : false;
+	}
+	else if (InType == EStreamType::Audio)
+	{
+		return CurrentPlayPeriodAudio.IsValid() ? CurrentPlayPeriodAudio->GetSelectedStreamBufferSourceInfo(InType).IsValid() : false;
+	}
+	else if (InType == EStreamType::Subtitle)
+	{
+		return CurrentPlayPeriodText.IsValid() ? CurrentPlayPeriodText->GetSelectedStreamBufferSourceInfo(InType).IsValid() : false;
+	}
+	return false;
+}
+
+
 bool FAdaptiveStreamingPlayer::HaveEnoughBufferedDataToStartPlayback()
 {
 	double kMinBufferBeforePlayback = GetMinBufferTimeBeforePlayback();
 
 	DiagnosticsCriticalSection.Lock();
 
-	bool bHaveEnoughVideo = false;
-	if (bHaveVideoReader.IsSet())
+	bool bHaveEnoughVideo = true;
+	if (IsExpectedToStreamNow(EStreamType::Video))
 	{
-		if (StreamReaderHandler && bHaveVideoReader.Value())
+		bHaveEnoughVideo = false;
+		// If the video stream has reached the end then there won't be any new data and whatever we have is all there is.
+		if (VideoBufferStats.StreamBuffer.bEndOfData)
+		{
+			bHaveEnoughVideo = true;
+		}
+		else
 		{
-			// If the video stream has reached the end then there won't be any new data and whatever we have is all there is.
-			if (VideoBufferStats.StreamBuffer.bEndOfData)
+			// Does the buffer have the requested amount of data?
+			if (VideoBufferStats.StreamBuffer.PlayableDuration.GetAsSeconds() >= kMinBufferBeforePlayback)
 			{
 				bHaveEnoughVideo = true;
 			}
 			else
 			{
-				// Does the buffer have the requested amount of data?
-				if (VideoBufferStats.StreamBuffer.PlayableDuration.GetAsSeconds() >= kMinBufferBeforePlayback)
+				// There is a chance that the input stream buffer is too small to hold the requested amount of material.
+				if (VideoBufferStats.StreamBuffer.bLastPushWasBlocked)
 				{
+					// We won't be able to receive any more, so we have enough as it is.
 					bHaveEnoughVideo = true;
 				}
-				else
-				{
-					// There is a chance that the input stream buffer is too small to hold the requested amount of material.
-					if (VideoBufferStats.StreamBuffer.bLastPushWasBlocked)
-					{
-						// We won't be able to receive any more, so we have enough as it is.
-						bHaveEnoughVideo = true;
-					}
-				}
 			}
 		}
-		else
-		{
-			bHaveEnoughVideo = true;
-		}
 	}
 
-	bool bHaveEnoughAudio = false;
-	if (bHaveAudioReader.IsSet())
+	bool bHaveEnoughAudio = true;
+	if (IsExpectedToStreamNow(EStreamType::Audio))
 	{
-		if (StreamReaderHandler && bHaveAudioReader.Value())
+		bHaveEnoughAudio = false;
+		// If the audio stream has reached the end then there won't be any new data and whatever we have is all there is.
+		if (AudioBufferStats.StreamBuffer.bEndOfData)
 		{
-			// If the audio stream has reached the end then there won't be any new data and whatever we have is all there is.
-			if (AudioBufferStats.StreamBuffer.bEndOfData)
+			bHaveEnoughAudio = true;
+		}
+		else
+		{
+			// Does the buffer have the requested amount of data?
+			if (AudioBufferStats.StreamBuffer.PlayableDuration.GetAsSeconds() >= kMinBufferBeforePlayback)
 			{
 				bHaveEnoughAudio = true;
 			}
 			else
 			{
-				// Does the buffer have the requested amount of data?
-				if (AudioBufferStats.StreamBuffer.PlayableDuration.GetAsSeconds() >= kMinBufferBeforePlayback)
+				// There is a chance that the input stream buffer is too small to hold the requested amount of material.
+				if (AudioBufferStats.StreamBuffer.bLastPushWasBlocked)
 				{
+					// We won't be able to receive any more, so we have enough as it is.
 					bHaveEnoughAudio = true;
 				}
-				else
-				{
-					// There is a chance that the input stream buffer is too small to hold the requested amount of material.
-					if (AudioBufferStats.StreamBuffer.bLastPushWasBlocked)
-					{
-						// We won't be able to receive any more, so we have enough as it is.
-						bHaveEnoughAudio = true;
-					}
-				}
 			}
 		}
-		else
-		{
-			bHaveEnoughAudio = true;
-		}
 	}
 
 	// When we are dealing with a single multiplexed stream and one buffer is blocked then essentially all buffers
@@ -2060,7 +2072,7 @@ void FAdaptiveStreamingPlayer::PrepareForPrerolling()
 	DiagnosticsCriticalSection.Lock();
 
 	// Check video buffer
-	if (StreamReaderHandler && bHaveVideoReader.GetWithDefault(false))
+	if (IsExpectedToStreamNow(EStreamType::Video))
 	{
 		if (VideoBufferStats.StreamBuffer.bEndOfData && (VideoBufferStats.StreamBuffer.PlayableDuration == FTimeValue::GetZero() || VideoBufferStats.StreamBuffer.NumCurrentAccessUnits == 0))
 		{
@@ -2072,7 +2084,7 @@ void FAdaptiveStreamingPlayer::PrepareForPrerolling()
 	}
 
 	// Check audio buffer
-	if (StreamReaderHandler && bHaveAudioReader.GetWithDefault(false))
+	if (IsExpectedToStreamNow(EStreamType::Audio))
 	{
 		if (AudioBufferStats.StreamBuffer.bEndOfData && (AudioBufferStats.StreamBuffer.PlayableDuration == FTimeValue::GetZero() || AudioBufferStats.StreamBuffer.NumCurrentAccessUnits == 0))
 		{
@@ -2094,7 +2106,7 @@ void FAdaptiveStreamingPlayer::PrepareForPrerolling()
 void FAdaptiveStreamingPlayer::HandleNewBufferedData()
 {
 	// Are we in buffering state? (this includes rebuffering)
-	if (CurrentState == EPlayerState::eState_Buffering && DecoderState == EDecoderState::eDecoder_Paused)
+	if (bStreamingHasStarted && CurrentState == EPlayerState::eState_Buffering && DecoderState == EDecoderState::eDecoder_Paused)
 	{
 		if (!SeekVars.bForScrubbing)
 		{
@@ -2164,43 +2176,6 @@ void FAdaptiveStreamingPlayer::HandleNewBufferedData()
 	}
 }
 
-//-----------------------------------------------------------------------------
-/**
- * Returns the current internal playback time.
- * This is not the same as GetPlayPosition() in that this time here must not
- * continue to run if the render clock is not stopped.
- * This time should be fairly accurate but is not currently required to be.
- * Events and period transitions are triggered off of this.
- */
-FTimeValue FAdaptiveStreamingPlayer::GetCurrentPlayTime()
-{
-	FTimeValue Current;
-	FTimeValue v;
-	TSharedPtrTS<FMultiTrackAccessUnitBuffer> Buffer = GetCurrentOutputStreamBuffer(EStreamType::Video);
-	if (Buffer.IsValid())
-	{
-		v = Buffer->GetLastPoppedDTS();
-	}
-	if (v.IsValid() && !Current.IsValid())
-	{
-		Current = v;
-	}
-	Buffer = GetCurrentOutputStreamBuffer(EStreamType::Audio);
-	if (Buffer.IsValid())
-	{
-		v = Buffer->GetLastPoppedDTS();
-	}
-	if (v.IsValid() && !Current.IsValid())
-	{
-		Current = v;
-	}
-	if (!Current.IsValid())
-	{
-		Current = PlaybackState.GetPlayPosition();
-	}
-	return Current;
-}
-
 
 //-----------------------------------------------------------------------------
 /**
@@ -2277,45 +2252,39 @@ void FAdaptiveStreamingPlayer::HandleNewOutputData()
 
 		// Check video
 		bool bHaveEnoughVideo = PrerollVars.bHaveEnoughVideo;
-		if (bHaveVideoReader.IsSet())
+		if (IsExpectedToStreamNow(EStreamType::Video))
 		{
-			if (StreamReaderHandler && bHaveVideoReader.Value())
-			{
-				// If the video decoder output buffer is stalled then we have enough video output. There won't be any more coming in.
-				// Also when the stream has reached the end and that has propagated through the decoder's buffer.
-				if (VideoBufferStats.DecoderOutputBuffer.bOutputStalled ||
-					(VideoBufferStats.StreamBuffer.bEndOfData && VideoBufferStats.DecoderInputBuffer.bEODSignaled && VideoBufferStats.DecoderInputBuffer.bEODReached) ||
-					bIsVideoDeselected)
-				{
-					bHaveEnoughVideo = true;
-				}
-			}
-			else
+			// If the video decoder output buffer is stalled then we have enough video output. There won't be any more coming in.
+			// Also when the stream has reached the end and that has propagated through the decoder's buffer.
+			if (VideoBufferStats.DecoderOutputBuffer.bOutputStalled ||
+				(VideoBufferStats.StreamBuffer.bEndOfData && VideoBufferStats.DecoderInputBuffer.bEODSignaled && VideoBufferStats.DecoderInputBuffer.bEODReached) ||
+				bIsVideoDeselected)
 			{
 				bHaveEnoughVideo = true;
 			}
 		}
+		else
+		{
+			bHaveEnoughVideo = true;
+		}
 
 		// Check audio
 		bool bHaveEnoughAudio = PrerollVars.bHaveEnoughAudio;
-		if (bHaveAudioReader.IsSet())
+		if (IsExpectedToStreamNow(EStreamType::Audio))
 		{
-			if (StreamReaderHandler && bHaveAudioReader.Value())
-			{
-				// If the audio decoder output buffer is stalled then we have enough video output. There won't be any more coming in.
-				// Also when the stream has reached the end and that has propagated through the decoder's buffer.
-				if (AudioBufferStats.DecoderOutputBuffer.bOutputStalled ||
-					(AudioBufferStats.StreamBuffer.bEndOfData && AudioBufferStats.DecoderInputBuffer.bEODSignaled && AudioBufferStats.DecoderInputBuffer.bEODReached) ||
-					bIsAudioDeselected)
-				{
-					bHaveEnoughAudio = true;
-				}
-			}
-			else
+			// If the audio decoder output buffer is stalled then we have enough video output. There won't be any more coming in.
+			// Also when the stream has reached the end and that has propagated through the decoder's buffer.
+			if (AudioBufferStats.DecoderOutputBuffer.bOutputStalled ||
+				(AudioBufferStats.StreamBuffer.bEndOfData && AudioBufferStats.DecoderInputBuffer.bEODSignaled && AudioBufferStats.DecoderInputBuffer.bEODReached) ||
+				bIsAudioDeselected)
 			{
 				bHaveEnoughAudio = true;
 			}
 		}
+		else
+		{
+			bHaveEnoughAudio = true;
+		}
 
 		DiagnosticsCriticalSection.Unlock();
 
@@ -2610,8 +2579,6 @@ void FAdaptiveStreamingPlayer::InternalHandlePendingStartRequest(const FTimeValu
 						// default start time so we can seek backwards/forward from it to a different time.
 						Manifest->ClearDefaultStartTime();
 
-					// TODO: this should update the stream metadata.
-
 						// At playback start all streams begin in the same period.
 						CurrentPlayPeriodVideo = InitialPlayPeriod;
 						CurrentPlayPeriodAudio = InitialPlayPeriod;
@@ -2680,7 +2647,7 @@ void FAdaptiveStreamingPlayer::InternalHandlePendingStartRequest(const FTimeValu
 							{
 								FPlayerSequenceState NextPlaybackState = CurrentPlaybackSequenceState;
 								++NextPlaybackState.SequenceIndex;
-								Result = InitialPlayPeriod->GetLoopingSegment(FirstSegmentRequest, NextPlaybackState, PendingStartRequest->FinishedRequests, StartAt, PendingStartRequest->SearchType);
+								Result = InitialPlayPeriod->GetLoopingSegment(FirstSegmentRequest, NextPlaybackState, StartAt, PendingStartRequest->SearchType);
 								if (Result.GetType() == IManifest::FResult::EType::Found)
 								{
 									CurrentPlaybackSequenceState = NextPlaybackState;
@@ -2821,6 +2788,19 @@ void FAdaptiveStreamingPlayer::InternalHandlePendingFirstSegmentRequest(const FT
 			NewReceiveBuffers->VidBuffer->SelectTrackWhenAvailable(CurrentPlaybackSequenceID[StreamTypeToArrayIndex(EStreamType::Video)], BufferSourceInfoVid);
 			NewReceiveBuffers->AudBuffer->SelectTrackWhenAvailable(CurrentPlaybackSequenceID[StreamTypeToArrayIndex(EStreamType::Audio)], BufferSourceInfoAud);
 			NewReceiveBuffers->TxtBuffer->SelectTrackWhenAvailable(CurrentPlaybackSequenceID[StreamTypeToArrayIndex(EStreamType::Subtitle)], BufferSourceInfoTxt);
+			// Flag tracks that will not deliver any data accordingly.
+			if (!BufferSourceInfoVid.IsValid())
+			{
+				NewReceiveBuffers->VidBuffer->SetEndOfTrackAll();
+			}
+			if (!BufferSourceInfoAud.IsValid())
+			{
+				NewReceiveBuffers->AudBuffer->SetEndOfTrackAll();
+			}
+			if (!BufferSourceInfoTxt.IsValid())
+			{
+				NewReceiveBuffers->TxtBuffer->SetEndOfTrackAll();
+			}
 
 			// Add the new receive buffer to the list.
 			DataBuffersCriticalSection.Lock();
@@ -2856,31 +2836,6 @@ void FAdaptiveStreamingPlayer::InternalHandlePendingFirstSegmentRequest(const FT
 					DispatchEvent(FMetricEvent::ReportJumpInPlayPosition(NewPlayPos, CurrentPlayPos, Metrics::ETimeJumpReason::UserSeek));
 				}
 
-				// Get the expected stream types from the starting segment request.
-			// FIXME: get rid of the "bHave...XYZ...Reader" states to become more flexible with dynamic stream selection (enabling/disabling) during playback.
-				bHaveVideoReader.Set(false);
-				bHaveAudioReader.Set(false);
-				bHaveTextReader.Set(false);
-				switch(PendingFirstSegmentRequest->GetType())
-				{
-					case EStreamType::Video:	bHaveVideoReader.Set(true);	break;
-					case EStreamType::Audio:	bHaveAudioReader.Set(true);	break;
-					case EStreamType::Subtitle:	bHaveTextReader.Set(true);	break;
-					default:												break;
-				}
-				TArray<TSharedPtrTS<IStreamSegment>> DependentStreams;
-				PendingFirstSegmentRequest->GetDependentStreams(DependentStreams);
-				for(int32 i=0; i<DependentStreams.Num(); ++i)
-				{
-					switch(DependentStreams[i]->GetType())
-					{
-						case EStreamType::Video:	bHaveVideoReader.Set(true);	break;
-						case EStreamType::Audio:	bHaveAudioReader.Set(true);	break;
-						case EStreamType::Subtitle:	bHaveTextReader.Set(true);	break;
-						default:												break;
-					}
-				}
-
 				if (AudioRender.Renderer)
 				{
 					RenderClock->SetCurrentTime(IMediaRenderClock::ERendererType::Audio, NewPlayPos);
@@ -2931,6 +2886,7 @@ void FAdaptiveStreamingPlayer::InternalHandlePendingFirstSegmentRequest(const FT
 			}
 		}
 		PendingFirstSegmentRequest.Reset();
+		bStreamingHasStarted = true;
 	}
 }
 
@@ -2941,42 +2897,33 @@ void FAdaptiveStreamingPlayer::InternalHandleCompletedSegmentRequests(const FTim
 	if (CompletedSegmentRequests.Num())
 	{
 		bool bAllAtEOS = true;
-		if (bHaveVideoReader.GetWithDefault(false))
-		{
-			TSharedPtrTS<FMultiTrackAccessUnitBuffer> VidRcvBuffer = GetCurrentReceiveStreamBuffer(EStreamType::Video);
-			if (!CompletedSegmentRequests.Contains(EStreamType::Video))
-			{
-				bAllAtEOS = false;
-			}
-			else if (VidRcvBuffer.IsValid())
-			{
-				VidRcvBuffer->PushEndOfDataAll();
-			}
-		}
-		if (bHaveAudioReader.GetWithDefault(false))
-		{
-			TSharedPtrTS<FMultiTrackAccessUnitBuffer> AudRcvBuffer = GetCurrentReceiveStreamBuffer(EStreamType::Audio);
-			if (!CompletedSegmentRequests.Contains(EStreamType::Audio))
-			{
-				bAllAtEOS = false;
-			}
-			else if (AudRcvBuffer.IsValid())
-			{
-				AudRcvBuffer->PushEndOfDataAll();
-			}
-		}
-		if (bHaveTextReader.GetWithDefault(false))
+
+		auto CheckStreamTypeEOS = [this](bool& bInOutAllEOS, EStreamType InStreamType) -> void
 		{
-			TSharedPtrTS<FMultiTrackAccessUnitBuffer> TxtRcvBuffer = GetCurrentReceiveStreamBuffer(EStreamType::Subtitle);
-			if (!CompletedSegmentRequests.Contains(EStreamType::Subtitle))
+			TSharedPtrTS<FMultiTrackAccessUnitBuffer> RcvBuffer = GetCurrentReceiveStreamBuffer(InStreamType);
+			if (IsExpectedToStreamNow(InStreamType))
 			{
-				bAllAtEOS = false;
+				if (!CompletedSegmentRequests.Contains(InStreamType))
+				{
+					bInOutAllEOS = false;
+				}
+				else if (RcvBuffer.IsValid())
+				{
+					RcvBuffer->SetEndOfTrackAll();
+					RcvBuffer->PushEndOfDataAll();
+				}
 			}
-			else if (TxtRcvBuffer.IsValid())
+			else if (RcvBuffer.IsValid())
 			{
-				TxtRcvBuffer->PushEndOfDataAll();
+				RcvBuffer->SetEndOfTrackAll();
+				RcvBuffer->PushEndOfDataAll();
 			}
-		}
+		};
+
+		CheckStreamTypeEOS(bAllAtEOS, EStreamType::Video);
+		CheckStreamTypeEOS(bAllAtEOS, EStreamType::Audio);
+		CheckStreamTypeEOS(bAllAtEOS, EStreamType::Subtitle);
+
 		if (bAllAtEOS)
 		{
 			// Move the finished requests into a local work map, emptying the original.
@@ -2996,7 +2943,6 @@ void FAdaptiveStreamingPlayer::InternalHandleCompletedSegmentRequests(const FTim
 					PendingStartRequest->StartAt.Time = Seekable.Start;
 					SetPlaystartOptions(PendingStartRequest->StartAt.Options);
 					PendingStartRequest->StartType = FPendingStartRequest::EStartType::LoopPoint;
-					PendingStartRequest->FinishedRequests = MoveTemp(LocalFinishedRequests);
 					// Increase the internal loop count
 					++CurrentLoopState.Count;
 					CurrentLoopState.To = ClampTimeToCurrentRange(Seekable.Start, true, false);
@@ -3008,6 +2954,48 @@ void FAdaptiveStreamingPlayer::InternalHandleCompletedSegmentRequests(const FTim
 }
 
 
+void FAdaptiveStreamingPlayer::RequestNewPeriodStreams(EStreamType InType, FPendingSegmentRequest& InOutCurrentRequest)
+{
+	// Check if the next period has streams the current one does not have.
+	// We need to trigger segment requests for those.
+	if (!InOutCurrentRequest.bDidRequestNewPeriodStreams)
+	{
+		InOutCurrentRequest.bDidRequestNewPeriodStreams = true;
+
+		auto RequestIfNeeded = [this, &InOutCurrentRequest](EStreamType InStreamType)
+		{
+			if (!IsExpectedToStreamNow(InStreamType))
+			{
+				FPendingSegmentRequest NextReq;
+				NextReq.Period = InOutCurrentRequest.Period;
+				NextReq.StreamType = InStreamType;
+				NextReq.StartoverPosition.Time = InOutCurrentRequest.Period->GetMediaAsset()->GetTimeRange().Start;
+				SetPlaystartOptions(NextReq.StartoverPosition.Options);
+				NextReq.bStartOver = true;
+				NextReq.bDidRequestNewPeriodStreams = true;
+				NextPendingSegmentRequests.Enqueue(MoveTemp(NextReq));
+			}
+		};
+
+		if (InType == EStreamType::Video)
+		{
+			RequestIfNeeded(EStreamType::Audio);
+			RequestIfNeeded(EStreamType::Subtitle);
+		}
+		else if (InType == EStreamType::Audio)
+		{
+			RequestIfNeeded(EStreamType::Video);
+			RequestIfNeeded(EStreamType::Subtitle);
+		}
+		else if (InType == EStreamType::Subtitle)
+		{
+			RequestIfNeeded(EStreamType::Video);
+			RequestIfNeeded(EStreamType::Audio);
+		}
+	}
+}
+
+
 //-----------------------------------------------------------------------------
 /**
  * Handles pending media segment requests.
@@ -3027,17 +3015,6 @@ void FAdaptiveStreamingPlayer::HandlePendingMediaSegmentRequests()
 
 	FTimeValue CurrentPlaybackPos = GetPlayPosition();
 
-	// Add the requests that are play position triggered to the pending requests to process next.
-	for(int32 i=0; i<PlayPosPendingSegmentRequests.Num(); ++i)
-	{
-		if (CurrentPlaybackPos >= PlayPosPendingSegmentRequests[i].StartoverPosition.Time)
-		{
-			NextPendingSegmentRequests.Enqueue(PlayPosPendingSegmentRequests[i]);
-			PlayPosPendingSegmentRequests.RemoveAtSwap(i);
-			--i;
-		}
-	}
-
 	// Are there completed downloads for which we need to find the next segment?
 	TQueue<FPendingSegmentRequest> pendingRequests;
 	Swap(pendingRequests, NextPendingSegmentRequests);
@@ -3139,6 +3116,8 @@ void FAdaptiveStreamingPlayer::HandlePendingMediaSegmentRequests()
 				}
 				else if (PeriodState == IManifest::IPlayPeriod::EReadyState::IsReady)
 				{
+					RequestNewPeriodStreams(StreamType, FinishedReq);
+
 					TSharedPtrTS<FMultiTrackAccessUnitBuffer> StreamRcvBuffer = GetCurrentReceiveStreamBuffer(StreamType);
 					if (StreamType == EStreamType::Video)
 					{
@@ -3293,7 +3272,7 @@ void FAdaptiveStreamingPlayer::HandlePendingMediaSegmentRequests()
 							// Get the dependent stream types, if any.
 							// This is mainly for multiplexed streams like an .mp4 at this point.
 							// All stream types have reached EOS now.
-							// NOTE: If is possible that there is no actual request here. This happens when a stream
+							// NOTE: It is possible that there is no actual request here. This happens when a stream
 							// gets enabled (like a subtitle stream) but is immediately at EOS and thus has no actual segment request!
 							if (FinishedReq.Request.IsValid())
 							{
@@ -3314,31 +3293,16 @@ void FAdaptiveStreamingPlayer::HandlePendingMediaSegmentRequests()
 				case IManifest::FResult::EType::NotFound:
 				{
 					// Not found here should indicate that the type of stream is not available (or not selected) in the new period.
-					if (StreamType == EStreamType::Subtitle)
+					TSharedPtrTS<ITimelineMediaAsset> ThisPeriod = SegmentPlayPeriod->GetMediaAsset();
+					if (ThisPeriod.IsValid())
 					{
-						TSharedPtrTS<ITimelineMediaAsset> ThisPeriod = SegmentPlayPeriod->GetMediaAsset();
-						if (ThisPeriod.IsValid())
+						UpdatePeriodStreamBufferSourceInfo(SegmentPlayPeriod, StreamType, TSharedPtrTS<FBufferSourceInfo>());
+						TSharedPtrTS<FMultiTrackAccessUnitBuffer> RcvBuffer = GetCurrentReceiveStreamBuffer(StreamType);
+						if (RcvBuffer.IsValid())
 						{
-							FTimeRange ThisPeriodTimeRange = ThisPeriod->GetTimeRange();
-							// Attempt to reselect this stream with the beginning of the following period.
-							FPendingSegmentRequest ReselectReq;
-							ReselectReq.bPlayPosAutoReselect = true;
-							ReselectReq.StartoverPosition.Time = ThisPeriodTimeRange.End;
-							ReselectReq.StreamType = StreamType;
-							PlayPosPendingSegmentRequests.Emplace(MoveTemp(ReselectReq));
-							// Remember that this period has no selected track.
-							UpdatePeriodStreamBufferSourceInfo(SegmentPlayPeriod, StreamType, TSharedPtrTS<FBufferSourceInfo>());
+							RcvBuffer->SetEndOfTrackAll();
 						}
 					}
-					else
-					{
-						// Throw a playback error for now.
-						FErrorDetail err;
-						err.SetFacility(Facility::EFacility::Player);
-						err.SetMessage("Could not locate next segment");
-						err.SetCode(INTERR_FRAGMENT_NOT_AVAILABLE);
-						PostError(err);
-					}
 					break;
 				}
 				case IManifest::FResult::EType::BeforeStart:
@@ -3440,14 +3404,6 @@ void FAdaptiveStreamingPlayer::CancelPendingMediaSegmentRequests(EStreamType Str
 			ReadyWaitingSegmentRequests.Enqueue(ReadyReq);
 		}
 	}
-	for(int32 i=0; i<PlayPosPendingSegmentRequests.Num(); ++i)
-	{
-		if (PlayPosPendingSegmentRequests[i].StreamType == StreamType)
-		{
-			PlayPosPendingSegmentRequests.RemoveAtSwap(i);
-			--i;
-		}
-	}
 }
 
 
@@ -3947,40 +3903,63 @@ void FAdaptiveStreamingPlayer::CheckForStreamEnd()
 				}
 			}
 
-			// Check if all data has been used.
-			const bool bHaveVid = bHaveVideoReader.GetWithDefault(false);
-			const bool bHaveAud = bHaveAudioReader.GetWithDefault(false);
-			const bool bHaveTxt = bHaveTextReader.GetWithDefault(false);
-			// For simplicity we assume all streams are done if we know whether or not they exist.
-			bool bEndVid = bHaveVideoReader.IsSet();
-			bool bEndAud = bHaveAudioReader.IsSet();
-			bool bEndTxt = bHaveTextReader.IsSet();
-			int64 VidStalled = 0;
-			int64 AudStalled = 0;
-
 			DiagnosticsCriticalSection.Lock();
 			FBufferStats vidStats = VideoBufferStats;
 			FBufferStats audStats = AudioBufferStats;
 			DiagnosticsCriticalSection.Unlock();
 
+			// As long as no buffer is at EOD we don't need to check further
+			if (!vidStats.StreamBuffer.bEndOfData && !audStats.StreamBuffer.bEndOfData)
+			{
+				return;
+			}
+
+			// Look at the last period's buffer source infos to see what streams we are expecting to be currently active.
+			bool bHaveVid = false;
+			bool bHaveAud = false;
+			bool bHavetxt = false;
+			ActivePeriodCriticalSection.Lock();
+			if (ActivePeriods.Num())
+			{
+				bHaveVid = ActivePeriods.Last().BufferSourceInfoVid.IsValid();
+				bHaveAud = ActivePeriods.Last().BufferSourceInfoAud.IsValid();
+				bHavetxt = ActivePeriods.Last().BufferSourceInfoTxt.IsValid();
+			}
+			ActivePeriodCriticalSection.Unlock();
+
+			// For simplicity we assume all streams are done and check for those that are active at the moment.
+			bool bEndVid = true;
+			bool bEndAud = true;
+			bool bEndTxt = true;
+			int64 VidStalled = 0;
+			int64 AudStalled = 0;
+
 			// Check for end of video stream
-			if (StreamReaderHandler && bHaveVid)
+			if (bHaveVid)
 			{
 				// Source buffer and decoder input buffer at end of data?
 				// We do not check the decoder output buffer since it is possible that no decodable output was produced
 				// which tends to happen at the end of the stream or when there is no video any more when audio is longer.
 				bEndVid = (vidStats.StreamBuffer.bEndOfData && vidStats.DecoderInputBuffer.bEODSignaled && vidStats.DecoderInputBuffer.bEODReached);
 				VidStalled = vidStats.StreamBuffer.bEndOfData ? vidStats.GetStalledDurationMillisec() : 0;
+				if (bEndVid)
+				{
+					UpdateDataAvailabilityState(DataAvailabilityStateVid, Metrics::FDataAvailabilityChange::EAvailability::DataNotAvailable);
+				}
 			}
 
 			// Check for end of audio stream
-			if (StreamReaderHandler && bHaveAud)
+			if (bHaveAud)
 			{
 				// Source buffer and decoder input buffer at end of data?
 				// We do not check the decoder output buffer since it is possible that no decodable output was produced
 				// which tends to happen at the end of the stream or when there is no audio any more when video is longer.
 				bEndAud = (audStats.StreamBuffer.bEndOfData && audStats.DecoderInputBuffer.bEODSignaled && audStats.DecoderInputBuffer.bEODReached);
 				AudStalled = audStats.StreamBuffer.bEndOfData ? audStats.GetStalledDurationMillisec() : 0;
+				if (bEndAud)
+				{
+					UpdateDataAvailabilityState(DataAvailabilityStateAud, Metrics::FDataAvailabilityChange::EAvailability::DataNotAvailable);
+				}
 			}
 
 			// Text stream
@@ -4047,6 +4026,8 @@ void FAdaptiveStreamingPlayer::CheckForStreamEnd()
 							PlaybackState.SetLoopState(newLoopState);
 							DispatchEvent(FMetricEvent::ReportJumpInPlayPosition(newLoopState.To, newLoopState.From, Metrics::ETimeJumpReason::Looping));
 						}
+						// Update the metadata state again.
+						MetadataHandlingState.Reset();
 					}
 					else
 					{
@@ -4123,7 +4104,6 @@ void FAdaptiveStreamingPlayer::InternalStartAt(const FSeekParam& NewPosition)
 	check(ReadyWaitingSegmentRequests.IsEmpty());
 	check(CompletedSegmentRequests.Num() == 0);
 	NextPendingSegmentRequests.Empty();
-	PlayPosPendingSegmentRequests.Empty();
 	ReadyWaitingSegmentRequests.Empty();
 	PendingFirstSegmentRequest.Reset();
 	CompletedSegmentRequests.Empty();
@@ -4150,6 +4130,7 @@ void FAdaptiveStreamingPlayer::InternalStartAt(const FSeekParam& NewPosition)
 		PendingStartRequest->SearchType = IManifest::ESearchType::Closest;
 	}
 	bFirstSegmentRequestIsForLooping = false;
+	bStreamingHasStarted = false;
 
 	// Let the AEMS handler know that we are starting up.
 	AEMSEventHandler->PlaybackStartingUp();
@@ -4188,7 +4169,7 @@ void FAdaptiveStreamingPlayer::InternalSetPlaybackEnded()
 	// In case a seek back into the stream will happen we reset the first start state.
 	PrerollVars.bIsVeryFirstStart = true;
 	// Let the AEMS handler know as well.
-	AEMSEventHandler->PlaybackReachedEnd(GetCurrentPlayTime());
+	AEMSEventHandler->PlaybackReachedEnd(TimelineRange.End);
 }
 
 
@@ -4281,16 +4262,12 @@ void FAdaptiveStreamingPlayer::InternalStop(bool bHoldCurrentFrame)
 		CurrentStreamReader->Close();
 	}
 	delete CurrentStreamReader;
-
-	bHaveVideoReader.Reset();
-	bHaveAudioReader.Reset();
-	bHaveTextReader.Reset();
+	bStreamingHasStarted = false;
 
 	// Release any pending segment requests.
 	PendingStartRequest.Reset();
 	PendingFirstSegmentRequest.Reset();
 	NextPendingSegmentRequests.Empty();
-	PlayPosPendingSegmentRequests.Empty();
 	ReadyWaitingSegmentRequests.Empty();
 	CompletedSegmentRequests.Empty();
 	InitialPlayPeriod.Reset();
@@ -4424,11 +4401,11 @@ void FAdaptiveStreamingPlayer::InternalSetLoop(const FLoopParam& LoopParam)
 				TSharedPtrTS<FMultiTrackAccessUnitBuffer> RcvBuffer;
 				bool bAnyReceiveBufferAtEOD = false;
 				RcvBuffer = GetCurrentReceiveStreamBuffer(EStreamType::Video);
-				bAnyReceiveBufferAtEOD |= RcvBuffer.IsValid() ? RcvBuffer->IsEODFlagSet() : false;
+				bAnyReceiveBufferAtEOD |= RcvBuffer.IsValid() ? RcvBuffer->IsEODFlagSet() || RcvBuffer->IsEndOfTrack() : false;
 				RcvBuffer = GetCurrentReceiveStreamBuffer(EStreamType::Audio);
-				bAnyReceiveBufferAtEOD |= RcvBuffer.IsValid() ? RcvBuffer->IsEODFlagSet() : false;
+				bAnyReceiveBufferAtEOD |= RcvBuffer.IsValid() ? RcvBuffer->IsEODFlagSet() || RcvBuffer->IsEndOfTrack() : false;
 				RcvBuffer = GetCurrentReceiveStreamBuffer(EStreamType::Subtitle);
-				bAnyReceiveBufferAtEOD |= RcvBuffer.IsValid() ? RcvBuffer->IsEODFlagSet() : false;
+				bAnyReceiveBufferAtEOD |= RcvBuffer.IsValid() ? RcvBuffer->IsEODFlagSet() || RcvBuffer->IsEndOfTrack() : false;
 
 				bStartOver = bAnyReceiveBufferAtEOD;
 			}
@@ -4588,14 +4565,17 @@ void FAdaptiveStreamingPlayer::ABRGetStreamBufferStats(IAdaptiveStreamSelector::
 		case EStreamType::Video:
 			OutBufferStats.PlayableContentDuration = VideoBufferStats.StreamBuffer.PlayableDuration;
 			OutBufferStats.bReachedEnd = VideoBufferStats.StreamBuffer.bEndOfData;
+			OutBufferStats.bEndOfTrack = VideoBufferStats.StreamBuffer.bEndOfTrack;
 			break;
 		case EStreamType::Audio:
 			OutBufferStats.PlayableContentDuration = AudioBufferStats.StreamBuffer.PlayableDuration;
 			OutBufferStats.bReachedEnd = AudioBufferStats.StreamBuffer.bEndOfData;
+			OutBufferStats.bEndOfTrack = AudioBufferStats.StreamBuffer.bEndOfTrack;
 			break;
 		case EStreamType::Subtitle:
 			OutBufferStats.PlayableContentDuration = TextBufferStats.StreamBuffer.PlayableDuration;
 			OutBufferStats.bReachedEnd = TextBufferStats.StreamBuffer.bEndOfData;
+			OutBufferStats.bEndOfTrack = TextBufferStats.StreamBuffer.bEndOfTrack;
 			break;
 		default:
 			break;
@@ -4718,24 +4698,24 @@ void FAdaptiveStreamingPlayer::DebugPrint(void* pPlayer, void (*pPrintFN)(void*
 #if 1
 	DiagnosticsCriticalSection.Lock();
 	const FAccessUnitBufferInfo *pD = nullptr;
-	if (StreamReaderHandler && bHaveVideoReader.GetWithDefault(false))
+	if(1)
 	{
 		pD = &VideoBufferStats.StreamBuffer;
-		pPrintFN(pPlayer, "Video buffer  : %3u AUs; %8u/%8u bytes in; eod %d; %#7.4fs", (uint32)pD->NumCurrentAccessUnits, (uint32)pD->CurrentMemInUse, (uint32)PlayerConfig.StreamBufferConfigVideo.MaxDataSize, pD->bEndOfData, pD->PlayableDuration.GetAsSeconds());
+		pPrintFN(pPlayer, "Video buffer  : EOT %d; EOS %d; %3u AUs; %8u/%8u bytes in; %#7.4fs", pD->bEndOfTrack, pD->bEndOfData, (uint32)pD->NumCurrentAccessUnits, (uint32)pD->CurrentMemInUse, (uint32)PlayerConfig.StreamBufferConfigVideo.MaxDataSize, pD->PlayableDuration.GetAsSeconds());
 		pPrintFN(pPlayer, "Video decoder : %2u in decoder, %zu total, %s; EOD in %d; EOD out %d", (uint32)VideoBufferStats.DecoderOutputBuffer.NumElementsInDecoder, (uint32)VideoBufferStats.DecoderOutputBuffer.MaxDecodedElementsReady, VideoBufferStats.DecoderOutputBuffer.bOutputStalled?"    stalled":"not stalled", VideoBufferStats.DecoderInputBuffer.bEODReached, VideoBufferStats.DecoderOutputBuffer.bEODreached);
 		pPrintFN(pPlayer, "Video renderer: %#7.4fs enqueued", VideoRender.Renderer.IsValid() ? VideoRender.Renderer->GetEnqueuedSampleDuration().GetAsSeconds() : 0.0);
 	}
-	if (StreamReaderHandler && bHaveAudioReader.GetWithDefault(false))
+	if(1)
 	{
 		pD = &AudioBufferStats.StreamBuffer;
-		pPrintFN(pPlayer, "Audio buffer  : %3u AUs; %8u/%8u bytes in; eod %d; %#7.4fs", (uint32)pD->NumCurrentAccessUnits, (uint32)pD->CurrentMemInUse, (uint32)PlayerConfig.StreamBufferConfigAudio.MaxDataSize, pD->bEndOfData, pD->PlayableDuration.GetAsSeconds());
+		pPrintFN(pPlayer, "Audio buffer  : EOT %d; EOS %d; %3u AUs; %8u/%8u bytes in; %#7.4fs", pD->bEndOfTrack, pD->bEndOfData, (uint32)pD->NumCurrentAccessUnits, (uint32)pD->CurrentMemInUse, (uint32)PlayerConfig.StreamBufferConfigAudio.MaxDataSize, pD->PlayableDuration.GetAsSeconds());
 		pPrintFN(pPlayer, "Audio decoder : %2u in decoder, %2u total, %s; EOD in %d; EOD out %d", (uint32)AudioBufferStats.DecoderOutputBuffer.NumElementsInDecoder, (uint32)AudioBufferStats.DecoderOutputBuffer.MaxDecodedElementsReady, AudioBufferStats.DecoderOutputBuffer.bOutputStalled ? "    stalled" : "not stalled", AudioBufferStats.DecoderInputBuffer.bEODReached, AudioBufferStats.DecoderOutputBuffer.bEODreached);
 		pPrintFN(pPlayer, "Audio renderer: %#7.4fs enqueued", AudioRender.Renderer.IsValid() ? AudioRender.Renderer->GetEnqueuedSampleDuration().GetAsSeconds() : 0.0);
 	}
-	if (StreamReaderHandler && bHaveTextReader.GetWithDefault(false))
+	if(1)
 	{
 		pD = &TextBufferStats.StreamBuffer;
-		pPrintFN(pPlayer, "Text buffer   : %3u AUs; %8u/%8u bytes in; eod %d; %#7.4fs", (uint32)pD->NumCurrentAccessUnits, (uint32)pD->CurrentMemInUse, (uint32)PlayerConfig.StreamBufferConfigText.MaxDataSize, pD->bEndOfData, pD->PlayableDuration.GetAsSeconds());
+		pPrintFN(pPlayer, "Text buffer   : EOT %d; EOS %d; %3u AUs; %8u/%8u bytes in; eod %d; %#7.4fs", pD->bEndOfTrack, pD->bEndOfData, (uint32)pD->NumCurrentAccessUnits, (uint32)pD->CurrentMemInUse, (uint32)PlayerConfig.StreamBufferConfigText.MaxDataSize, pD->PlayableDuration.GetAsSeconds());
 	}
 	DiagnosticsCriticalSection.Unlock();
 
@@ -4760,8 +4740,6 @@ void FAdaptiveStreamingPlayer::DebugPrint(void* pPlayer, void (*pPrintFN)(void*
 	pPrintFN(pPlayer, "     is paused: %s", IsPaused() ? "Yes" : "No");
 	pPrintFN(pPlayer, "    is seeking: %s", IsSeeking() ? "Yes" : "No");
 	pPrintFN(pPlayer, "     has ended: %s", HasEnded() ? "Yes" : "No");
-	pPrintFN(pPlayer, "clock (%s): %.3f", RenderClock->IsRunning()?"running":"paused", RenderClock->GetInterpolatedRenderTime(IMediaRenderClock::ERendererType::Video).GetAsSeconds(-1.0));
-	pPrintFN(pPlayer, " internal time: [%d] %.3f", (int32) GetCurrentPlayTime().GetSequenceIndex(), GetCurrentPlayTime().GetAsSeconds());
 #endif
 	if (StreamSelector.IsValid())
 	{
diff --git a/Engine/Plugins/Media/ElectraPlayer/Source/ElectraPlayerRuntime/Private/Runtime/Player/AdaptiveStreamingPlayerABR.cpp b/Engine/Plugins/Media/ElectraPlayer/Source/ElectraPlayerRuntime/Private/Runtime/Player/AdaptiveStreamingPlayerABR.cpp
--- a/Engine/Plugins/Media/ElectraPlayer/Source/ElectraPlayerRuntime/Private/Runtime/Player/AdaptiveStreamingPlayerABR.cpp
+++ b/Engine/Plugins/Media/ElectraPlayer/Source/ElectraPlayerRuntime/Private/Runtime/Player/AdaptiveStreamingPlayerABR.cpp
@@ -599,6 +599,11 @@ namespace Electra
 		// Get the initial list of representations that are currently possible to choose from.
 		TArray<TSharedPtrTS<FABRStreamInformation>> CandidateRepresentations;
 		ABRMethod->PrepareStreamCandidateList(CandidateRepresentations, StreamType, CurrentPlayPeriod, TimeNow);
+		// No streams, nothing to do. Assume there is something coming up, so continue.
+		if (CandidateRepresentations.Num() == 0)
+		{
+			return ESegmentAction::FetchNext;
+		}
 
 		// First have the ABR method evaluate this download for possible errors and decide on an action.
 		IAdaptiveStreamSelector::ESegmentAction NextSegmentAction = ABRMethod->EvaluateForError(CandidateRepresentations, StreamType, OutDelay, CurrentPlayPeriod, CurrentSegment, TimeNow);
diff --git a/Engine/Plugins/Media/ElectraPlayer/Source/ElectraPlayerRuntime/Private/Runtime/Player/AdaptiveStreamingPlayerABR.h b/Engine/Plugins/Media/ElectraPlayer/Source/ElectraPlayerRuntime/Private/Runtime/Player/AdaptiveStreamingPlayerABR.h
--- a/Engine/Plugins/Media/ElectraPlayer/Source/ElectraPlayerRuntime/Private/Runtime/Player/AdaptiveStreamingPlayerABR.h
+++ b/Engine/Plugins/Media/ElectraPlayer/Source/ElectraPlayerRuntime/Private/Runtime/Player/AdaptiveStreamingPlayerABR.h
@@ -55,6 +55,7 @@ public:
 		{
 			FTimeValue PlayableContentDuration;
 			bool bReachedEnd = false;
+			bool bEndOfTrack = false;
 		};
 		virtual void ABRGetStreamBufferStats(FABRBufferStats& OutBufferStats, EStreamType ForStream) = 0;
 
diff --git a/Engine/Plugins/Media/ElectraPlayer/Source/ElectraPlayerRuntime/Private/Runtime/Player/AdaptiveStreamingPlayerDecoders.cpp b/Engine/Plugins/Media/ElectraPlayer/Source/ElectraPlayerRuntime/Private/Runtime/Player/AdaptiveStreamingPlayerDecoders.cpp
--- a/Engine/Plugins/Media/ElectraPlayer/Source/ElectraPlayerRuntime/Private/Runtime/Player/AdaptiveStreamingPlayerDecoders.cpp
+++ b/Engine/Plugins/Media/ElectraPlayer/Source/ElectraPlayerRuntime/Private/Runtime/Player/AdaptiveStreamingPlayerDecoders.cpp
@@ -167,12 +167,12 @@ bool FAdaptiveStreamingPlayer::CanDecodeSubtitle(const FString& MimeType, const
 
 
 
-bool FAdaptiveStreamingPlayer::FindMatchingStreamInfo(FStreamCodecInformation& OutStreamInfo, const FTimeValue& AtTime, int32 MaxWidth, int32 MaxHeight)
+bool FAdaptiveStreamingPlayer::FindMatchingStreamInfo(FStreamCodecInformation& OutStreamInfo, const FString& InPeriodID, const FTimeValue& AtTime, int32 MaxWidth, int32 MaxHeight)
 {
 	TArray<FStreamCodecInformation> VideoCodecInfos;
 	TSharedPtrTS<ITimelineMediaAsset> Asset;
 
-	// Locate the period for the specified time.
+	// Locate the period by ID or the specified time.
 	ActivePeriodCriticalSection.Lock();
 	if (ActivePeriods.Num() == 0)
 	{
@@ -180,12 +180,27 @@ bool FAdaptiveStreamingPlayer::FindMatchingStreamInfo(FStreamCodecInformation& O
 		return false;
 	}
 	bool bFound = false;
-	for(int32 i=0; i<ActivePeriods.Num(); ++i)
+	if (InPeriodID.Len())
 	{
-		if (AtTime >= ActivePeriods[i].TimeRange.Start && AtTime < (ActivePeriods[i].TimeRange.End.IsValid() ? ActivePeriods[i].TimeRange.End : FTimeValue::GetPositiveInfinity()))
+		for(int32 i=0; i<ActivePeriods.Num(); ++i)
 		{
-			Asset = ActivePeriods[i].Period;
-			bFound = true;
+			if (ActivePeriods[i].ID.Equals(InPeriodID))
+			{
+				Asset = ActivePeriods[i].Period;
+				bFound = true;
+				break;
+			}
+		}
+	}
+	if (!bFound && AtTime.IsValid())
+	{
+		for(int32 i=0; i<ActivePeriods.Num(); ++i)
+		{
+			if (AtTime >= ActivePeriods[i].TimeRange.Start && AtTime < (ActivePeriods[i].TimeRange.End.IsValid() ? ActivePeriods[i].TimeRange.End : FTimeValue::GetPositiveInfinity()))
+			{
+				Asset = ActivePeriods[i].Period;
+				bFound = true;
+			}
 		}
 	}
 	if (!bFound)
@@ -327,7 +342,9 @@ int32 FAdaptiveStreamingPlayer::CreateDecoder(EStreamType type)
 			bool bOk = GetCurrentOutputStreamBuffer(type).IsValid() && GetCurrentOutputStreamBuffer(type)->PeekAndAddRef(AccessUnit);
 			if (AccessUnit)
 			{
-				FTimeValue DecodeTime = AccessUnit->PTS;
+				FTimeValue DecodeTime = AccessUnit->EarliestPTS.IsValid() ? AccessUnit->EarliestPTS : AccessUnit->PTS;
+				check(AccessUnit->BufferSourceInfo.IsValid());
+				FString PeriodID = AccessUnit->BufferSourceInfo.IsValid() ? AccessUnit->BufferSourceInfo->PeriodID : FString();
 				VideoDecoder.CurrentCodecInfo.Clear();
 				if (AccessUnit->AUCodecData.IsValid())
 				{
@@ -341,7 +358,7 @@ int32 FAdaptiveStreamingPlayer::CreateDecoder(EStreamType type)
 				// This is only an initial selection as there could be other adaptation sets in upcoming periods
 				// that have a larger resolution that is still within the allowed limits.
 				FStreamCodecInformation HighestStream;
-				if (!FindMatchingStreamInfo(HighestStream, DecodeTime, 0, 0))
+				if (!FindMatchingStreamInfo(HighestStream, PeriodID, DecodeTime, 0, 0))
 				{
 					FErrorDetail err;
 					err.SetFacility(Facility::EFacility::Player);
@@ -601,7 +618,7 @@ void FAdaptiveStreamingPlayer::HandleSubtitleDecoder()
 		}
 		FTimeValue NextPTS = PeekedAU ? PeekedAU->PTS - SubtitleDecoder.Decoder->GetStreamedDeliveryTimeOffset() : FTimeValue::GetInvalid();
 		FAccessUnit::Release(PeekedAU);
-		FTimeValue CurrentPos = GetCurrentPlayTime();
+		FTimeValue CurrentPos = PlaybackState.GetPlayPosition();
 		if (NextPTS.IsValid() && CurrentPos >= NextPTS)
 		{
 			FeedDecoder(EStreamType::Subtitle, SubtitleDecoder.Decoder, false);
@@ -704,7 +721,6 @@ void FAdaptiveStreamingPlayer::FeedDecoder(EStreamType Type, IAccessUnitBufferIn
 	}
 
 	FBufferStats* pStats = nullptr;
-	Metrics::FDataAvailabilityChange* pAvailability = nullptr;
 	FStreamCodecInformation* CurrentCodecInfo = nullptr;
 	TSharedPtrTS<FAccessUnit::CodecData>* LastSentAUCodecData = nullptr;
 	bool bCodecChangeDetected = false;
@@ -714,21 +730,18 @@ void FAdaptiveStreamingPlayer::FeedDecoder(EStreamType Type, IAccessUnitBufferIn
 	{
 		case EStreamType::Video:
 			pStats = &VideoBufferStats;
-			pAvailability = &DataAvailabilityStateVid;
 			CurrentCodecInfo = &VideoDecoder.CurrentCodecInfo;
 			LastSentAUCodecData = &VideoDecoder.LastSentAUCodecData;
 			bIsDeselected = bIsVideoDeselected;
 			break;
 		case EStreamType::Audio:
 			pStats = &AudioBufferStats;
-			pAvailability = &DataAvailabilityStateAud;
 			CurrentCodecInfo = &AudioDecoder.CurrentCodecInfo;
 			LastSentAUCodecData = &AudioDecoder.LastSentAUCodecData;
 			bIsDeselected = bIsAudioDeselected;
 			break;
 		case EStreamType::Subtitle:
 			pStats = &TextBufferStats;
-			pAvailability = &DataAvailabilityStateTxt;
 			CurrentCodecInfo = &SubtitleDecoder.CurrentCodecInfo;
 			LastSentAUCodecData = &SubtitleDecoder.LastSentAUCodecData;
 			bIsDeselected = bIsTextDeselected;
@@ -749,8 +762,11 @@ void FAdaptiveStreamingPlayer::FeedDecoder(EStreamType Type, IAccessUnitBufferIn
 		// Check for buffer underrun.
 		if (bHandleUnderrun && RebufferCause == ERebufferCause::None && CurrentState == EPlayerState::eState_Playing && StreamState == EStreamState::eStream_Running && PipelineState == EPipelineState::ePipeline_Running)
 		{
+			const FTimeValue MinPushedDuration(FTimeValue::MillisecondsToHNS(500));
 			bool bEODSet = FromMultistreamBuffer->IsEODFlagSet();
-			if (!bEODSet && FromMultistreamBuffer->Num() == 0)
+			bool bEOTSet = FromMultistreamBuffer->IsEndOfTrack();
+			FTimeValue PushedDuration = FromMultistreamBuffer->GetPlayableDurationPushedSinceEOT();
+			if (!bEODSet && !bEOTSet && FromMultistreamBuffer->Num() == 0 && PushedDuration >= MinPushedDuration)
 			{
 				FTimeValue EnqueuedDuration(FTimeValue::GetZero());
 				int32 NumEnqueuedSamples = 0;
@@ -839,7 +855,10 @@ void FAdaptiveStreamingPlayer::FeedDecoder(EStreamType Type, IAccessUnitBufferIn
 				if (Type == EStreamType::Video && VideoDecoder.bApplyNewLimits)
 				{
 					FStreamCodecInformation StreamInfo;
-					if (FindMatchingStreamInfo(StreamInfo, AccessUnit->PTS, VideoResolutionLimitWidth, VideoResolutionLimitHeight))
+					FTimeValue DecodeTime = AccessUnit->EarliestPTS.IsValid() ? AccessUnit->EarliestPTS : AccessUnit->PTS;
+					check(AccessUnit->BufferSourceInfo.IsValid());
+					FString PeriodID = AccessUnit->BufferSourceInfo.IsValid() ? AccessUnit->BufferSourceInfo->PeriodID : FString();
+					if (FindMatchingStreamInfo(StreamInfo, PeriodID, DecodeTime, VideoResolutionLimitWidth, VideoResolutionLimitHeight))
 					{
 						if (VideoDecoder.Decoder)
 						{
@@ -854,6 +873,14 @@ void FAdaptiveStreamingPlayer::FeedDecoder(EStreamType Type, IAccessUnitBufferIn
 
 				if (Decoder)
 				{
+					// Since we are providing a new access unit now the decoder won't be at EOD any more.
+					// If it was before, clear it out.
+					if (pStats && pStats->DecoderInputBuffer.bEODSignaled)
+					{
+						pStats->DecoderInputBuffer.bEODSignaled = false;
+						pStats->DecoderInputBuffer.bEODReached = false;
+						Decoder->AUdataClearEOD();
+					}
 					Decoder->AUdataPushAU(AccessUnit);
 				}
 
@@ -868,12 +895,6 @@ void FAdaptiveStreamingPlayer::FeedDecoder(EStreamType Type, IAccessUnitBufferIn
 					*LastSentAUCodecData = AccessUnit->AUCodecData;
 				}
 				
-				// Notify of any change in data availability.
-				if (pAvailability)
-				{
-					UpdateDataAvailabilityState(*pAvailability, Metrics::FDataAvailabilityChange::EAvailability::DataAvailable);
-				}
-
 				// Check for presence of encoder latency and store it.
 				// This may be updated by different streams/decoders which can't be helped. There should only be one actively
 				// observed element but we have seen manifests that have the same <ProducerReferenceTime@id> in multiple AdaptationSets.
@@ -888,7 +909,7 @@ void FAdaptiveStreamingPlayer::FeedDecoder(EStreamType Type, IAccessUnitBufferIn
 	}
 	// An AU is not tagged as being "the last" one. Instead the EOD is handled separately and must be dealt with
 	// by the decoders accordingly.
-	if (!bCodecChangeDetected && FromMultistreamBuffer->IsEODFlagSet() && FromMultistreamBuffer->Num() == 0)
+	if (!bCodecChangeDetected && (FromMultistreamBuffer->IsEODFlagSet() || FromMultistreamBuffer->IsEndOfTrack()) && FromMultistreamBuffer->Num() == 0)
 	{
 		if (pStats && !pStats->DecoderInputBuffer.bEODSignaled)
 		{
@@ -897,10 +918,6 @@ void FAdaptiveStreamingPlayer::FeedDecoder(EStreamType Type, IAccessUnitBufferIn
 				Decoder->AUdataPushEOD();
 			}
 		}
-		if (pAvailability)
-		{
-			UpdateDataAvailabilityState(*pAvailability, Metrics::FDataAvailabilityChange::EAvailability::DataNotAvailable);
-		}
 	}
 }
 
diff --git a/Engine/Plugins/Media/ElectraPlayer/Source/ElectraPlayerRuntime/Private/Runtime/Player/AdaptiveStreamingPlayerInternal.h b/Engine/Plugins/Media/ElectraPlayer/Source/ElectraPlayerRuntime/Private/Runtime/Player/AdaptiveStreamingPlayerInternal.h
--- a/Engine/Plugins/Media/ElectraPlayer/Source/ElectraPlayerRuntime/Private/Runtime/Player/AdaptiveStreamingPlayerInternal.h
+++ b/Engine/Plugins/Media/ElectraPlayer/Source/ElectraPlayerRuntime/Private/Runtime/Player/AdaptiveStreamingPlayerInternal.h
@@ -420,6 +420,13 @@ struct FPlaybackState
 		OutAudioTracks = AudioTracks;
 		OutSubtitleTracks = SubtitleTracks;
 	}
+	void HaveTrackMetadata(bool& bOutHaveVideo, bool& bOutHaveAudio, bool& bOutHaveSubtitle) const
+	{
+		FScopeLock lock(&Lock);
+		bOutHaveVideo = VideoTracks.Num() != 0;
+		bOutHaveAudio = AudioTracks.Num() != 0;
+		bOutHaveSubtitle = SubtitleTracks.Num() != 0;
+	}
 
 	void SetLoopState(const IAdaptiveStreamingPlayer::FLoopState& InLoopState)
 	{
@@ -1539,7 +1546,6 @@ private:
 		IManifest::ESearchType									SearchType = IManifest::ESearchType::Closest;
 		FTimeValue												RetryAtTime;
 		EStartType												StartType = EStartType::PlayStart;
-		TMultiMap<EStreamType, TSharedPtrTS<IStreamSegment>>	FinishedRequests;
 	};
 
 	struct FBufferStats
@@ -1724,6 +1730,7 @@ private:
 		TSharedPtrTS<IManifest::IPlayPeriod>		Period;							//!< Set if transitioning between periods. This is the new period that needs to be readied.
 		bool										bStartOver = false;				//!< True when switching tracks within an ongoing period
 		bool										bPlayPosAutoReselect = false;	//!< True when trying to re-select a stream at the start of a new period when it was not available before.
+		bool										bDidRequestNewPeriodStreams = false;
 		EStreamType									StreamType = EStreamType::Unsupported;
 		FPlayStartPosition							StartoverPosition;
 	};
@@ -1820,9 +1827,10 @@ private:
 
 	int32 CreateDecoder(EStreamType type);
 	void DestroyDecoders();
-	bool FindMatchingStreamInfo(FStreamCodecInformation& OutStreamInfo, const FTimeValue& AtTime, int32 MaxWidth, int32 MaxHeight);
+	bool FindMatchingStreamInfo(FStreamCodecInformation& OutStreamInfo, const FString& InPeriodID, const FTimeValue& AtTime, int32 MaxWidth, int32 MaxHeight);
 	void UpdateStreamResolutionLimit();
 	void AddUpcomingPeriod(TSharedPtrTS<IManifest::IPlayPeriod> InUpcomingPeriod);
+	void RequestNewPeriodStreams(EStreamType InType, FPendingSegmentRequest& InOutCurrentRequest);
 	void UpdatePeriodStreamBufferSourceInfo(TSharedPtrTS<IManifest::IPlayPeriod> InForPeriod, EStreamType InStreamType, TSharedPtrTS<FBufferSourceInfo> InBufferSourceInfo);
 	TSharedPtrTS<FBufferSourceInfo> GetStreamBufferInfoAtTime(bool& bOutHavePeriods, bool& bOutFoundTime, EStreamType InStreamType, const FTimeValue& InAtTime) const;
 
@@ -1904,6 +1912,8 @@ private:
 		return GetStreamBuffer(InStreamType, CurrentDataReceiveBuffers);
 	}
 
+	bool IsExpectedToStreamNow(EStreamType InType);
+
 	bool IsSeamlessBufferSwitchPossible(EStreamType InStreamType, const TSharedPtrTS<FStreamDataBuffers>& InFromStreamBuffers);
 
 	void GetStreamBufferUtilization(FAccessUnitBufferInfo& OutInfo, EStreamType BufferType);
@@ -1932,9 +1942,6 @@ private:
 	void StartRendering();
 	void StopRendering();
 
-	FTimeValue GetCurrentPlayTime();
-
-
 	//
 	// Member variables
 	//
@@ -1970,9 +1977,7 @@ private:
 	TSharedPtrTS<FHTTPResourceRequest>									ManifestMimeTypeRequest;
 
 	IStreamReader*														StreamReaderHandler;
-	TMediaOptionalValue<bool> 											bHaveVideoReader;
-	TMediaOptionalValue<bool> 											bHaveAudioReader;
-	TMediaOptionalValue<bool> 											bHaveTextReader;
+	bool																bStreamingHasStarted;
 
 	FCriticalSection													DataBuffersCriticalSection;
 	TArray<TSharedPtrTS<FStreamDataBuffers>>							NextDataBuffers;
@@ -2045,7 +2050,6 @@ private:
 	TSharedPtrTS<FPendingStartRequest>									PendingStartRequest;
 	TSharedPtrTS<IStreamSegment>										PendingFirstSegmentRequest;
 	TQueue<FPendingSegmentRequest>										NextPendingSegmentRequests;
-	TArray<FPendingSegmentRequest>										PlayPosPendingSegmentRequests;
 	TQueue<TSharedPtrTS<IStreamSegment>>								ReadyWaitingSegmentRequests;
 	TMultiMap<EStreamType, TSharedPtrTS<IStreamSegment>>				CompletedSegmentRequests;
 	bool																bFirstSegmentRequestIsForLooping;
diff --git a/Engine/Plugins/Media/ElectraPlayer/Source/ElectraPlayerRuntime/Private/Runtime/Player/DASH/ManifestDASH.cpp b/Engine/Plugins/Media/ElectraPlayer/Source/ElectraPlayerRuntime/Private/Runtime/Player/DASH/ManifestDASH.cpp
--- a/Engine/Plugins/Media/ElectraPlayer/Source/ElectraPlayerRuntime/Private/Runtime/Player/DASH/ManifestDASH.cpp
+++ b/Engine/Plugins/Media/ElectraPlayer/Source/ElectraPlayerRuntime/Private/Runtime/Player/DASH/ManifestDASH.cpp
@@ -123,7 +123,7 @@ public:
 	void TriggerInitSegmentPreload(const TArray<FInitSegmentPreload>& InitSegmentsToPreload) override;
 	IManifest::FResult GetStartingSegment(TSharedPtrTS<IStreamSegment>& OutSegment, const FPlayerSequenceState& InSequenceState, const FPlayStartPosition& StartPosition, IManifest::ESearchType SearchType) override;
 	IManifest::FResult GetContinuationSegment(TSharedPtrTS<IStreamSegment>& OutSegment, EStreamType StreamType, const FPlayerSequenceState& SequenceState, const FPlayStartPosition& StartPosition, IManifest::ESearchType SearchType) override;
-	IManifest::FResult GetLoopingSegment(TSharedPtrTS<IStreamSegment>& OutSegment, const FPlayerSequenceState& SequenceState, const TMultiMap<EStreamType, TSharedPtrTS<IStreamSegment>>& InFinishedSegments, const FPlayStartPosition& StartPosition, IManifest::ESearchType SearchType) override;
+	IManifest::FResult GetLoopingSegment(TSharedPtrTS<IStreamSegment>& OutSegment, const FPlayerSequenceState& SequenceState, const FPlayStartPosition& StartPosition, IManifest::ESearchType SearchType) override;
 	IManifest::FResult GetNextSegment(TSharedPtrTS<IStreamSegment>& OutSegment, TSharedPtrTS<const IStreamSegment> CurrentSegment, const FPlayStartOptions& Options) override;
 	IManifest::FResult GetRetrySegment(TSharedPtrTS<IStreamSegment>& OutSegment, TSharedPtrTS<const IStreamSegment> CurrentSegment, const FPlayStartOptions& Options, bool bReplaceWithFillerData) override;
 	void IncreaseSegmentFetchDelay(const FTimeValue& IncreaseAmount) override;
@@ -735,6 +735,7 @@ TSharedPtrTS<FManifestDASHInternal::FAdaptationSet> FDASHPlayPeriod::SelectAdapt
 			OutBufferSourceInfo->Language = tm.Language;
 			OutBufferSourceInfo->Codec = tm.HighestBandwidthCodec.GetCodecName();
 			OutBufferSourceInfo->HardIndex = SelectedTypeIndex;
+			OutBufferSourceInfo->PeriodID = Period->GetUniqueIdentifier();
 			OutBufferSourceInfo->PeriodAdaptationSetID = Period->GetUniqueIdentifier() + TEXT("/") + AS->GetUniqueIdentifier();
 		}
 	}
@@ -1578,18 +1579,9 @@ IManifest::FResult FDASHPlayPeriod::GetRetrySegment(TSharedPtrTS<IStreamSegment>
 }
 
 
-IManifest::FResult FDASHPlayPeriod::GetLoopingSegment(TSharedPtrTS<IStreamSegment>& OutSegment, const FPlayerSequenceState& InSequenceState, const TMultiMap<EStreamType, TSharedPtrTS<IStreamSegment>>& InFinishedSegments, const FPlayStartPosition& StartPosition, IManifest::ESearchType SearchType)
+IManifest::FResult FDASHPlayPeriod::GetLoopingSegment(TSharedPtrTS<IStreamSegment>& OutSegment, const FPlayerSequenceState& InSequenceState, const FPlayStartPosition& StartPosition, IManifest::ESearchType SearchType)
 {
-	if (InFinishedSegments.Num())
-	{
-		IManifest::FResult Result = GetStartingSegment(OutSegment, InSequenceState, StartPosition, SearchType);
-		if (Result.GetType() == IManifest::FResult::EType::Found)
-		{
-			return Result;
-		}
-	}
-	// Return past EOS when we can't loop to indicate we're really done now.
-	return IManifest::FResult(IManifest::FResult::EType::PastEOS);
+	return GetStartingSegment(OutSegment, InSequenceState, StartPosition, SearchType);
 }
 
 
diff --git a/Engine/Plugins/Media/ElectraPlayer/Source/ElectraPlayerRuntime/Private/Runtime/Player/HLS/ManifestHLS.cpp b/Engine/Plugins/Media/ElectraPlayer/Source/ElectraPlayerRuntime/Private/Runtime/Player/HLS/ManifestHLS.cpp
--- a/Engine/Plugins/Media/ElectraPlayer/Source/ElectraPlayerRuntime/Private/Runtime/Player/HLS/ManifestHLS.cpp
+++ b/Engine/Plugins/Media/ElectraPlayer/Source/ElectraPlayerRuntime/Private/Runtime/Player/HLS/ManifestHLS.cpp
@@ -47,7 +47,7 @@ public:
 	IManifest::FResult GetContinuationSegment(TSharedPtrTS<IStreamSegment>& OutSegment, EStreamType StreamType, const FPlayerSequenceState& SequenceState, const FPlayStartPosition& StartPosition, IManifest::ESearchType SearchType) override;
 	IManifest::FResult GetNextSegment(TSharedPtrTS<IStreamSegment>& OutSegment, TSharedPtrTS<const IStreamSegment> CurrentSegment, const FPlayStartOptions& Options) override;
 	IManifest::FResult GetRetrySegment(TSharedPtrTS<IStreamSegment>& OutSegment, TSharedPtrTS<const IStreamSegment> CurrentSegment, const FPlayStartOptions& Options, bool bReplaceWithFillerData) override;
-	IManifest::FResult GetLoopingSegment(TSharedPtrTS<IStreamSegment>& OutSegment, const FPlayerSequenceState& SequenceState, const TMultiMap<EStreamType, TSharedPtrTS<IStreamSegment>>& InFinishedSegments, const FPlayStartPosition& StartPosition, IManifest::ESearchType SearchType) override;
+	IManifest::FResult GetLoopingSegment(TSharedPtrTS<IStreamSegment>& OutSegment, const FPlayerSequenceState& SequenceState, const FPlayStartPosition& StartPosition, IManifest::ESearchType SearchType) override;
 	void IncreaseSegmentFetchDelay(const FTimeValue& IncreaseAmount) override;
 
 	// Obtains information on the stream segmentation of a particular stream starting at a given current reference segment (optional, if not given returns suitable default values).
@@ -454,12 +454,20 @@ void FPlayPeriodHLS::PrepareForPlay()
 
 	// Set up source buffer information for video and audio.
 	// These are currently dummies since we do not support track switching yet.
-	CurrentSourceBufferInfoVideo = MakeSharedTS<FBufferSourceInfo>();
-	CurrentSourceBufferInfoAudio = MakeSharedTS<FBufferSourceInfo>();
-	CurrentSourceBufferInfoVideo->PeriodAdaptationSetID = TEXT("video.0");
-	CurrentSourceBufferInfoVideo->HardIndex = 0;
-	CurrentSourceBufferInfoAudio->PeriodAdaptationSetID = TEXT("audio.0");
-	CurrentSourceBufferInfoAudio->HardIndex = 0;
+	if (ActiveVideoUniqueID)
+	{
+		CurrentSourceBufferInfoVideo = MakeSharedTS<FBufferSourceInfo>();
+		CurrentSourceBufferInfoVideo->PeriodID = InternalManifest->CurrentMediaAsset->GetAssetIdentifier();
+		CurrentSourceBufferInfoVideo->PeriodAdaptationSetID = TEXT("video.0");
+		CurrentSourceBufferInfoVideo->HardIndex = 0;
+	}
+	if (ActiveAudioUniqueID)
+	{
+		CurrentSourceBufferInfoAudio = MakeSharedTS<FBufferSourceInfo>();
+		CurrentSourceBufferInfoAudio->PeriodID = InternalManifest->CurrentMediaAsset->GetAssetIdentifier();
+		CurrentSourceBufferInfoAudio->PeriodAdaptationSetID = TEXT("audio.0");
+		CurrentSourceBufferInfoAudio->HardIndex = 0;
+	}
 
 	CurrentReadyState = IManifest::IPlayPeriod::EReadyState::IsReady;
 }
@@ -1040,18 +1048,9 @@ void FPlayPeriodHLS::IncreaseSegmentFetchDelay(const FTimeValue& IncreaseAmount)
 }
 
 
-IManifest::FResult FPlayPeriodHLS::GetLoopingSegment(TSharedPtrTS<IStreamSegment>& OutSegment, const FPlayerSequenceState& SequenceState, const TMultiMap<EStreamType, TSharedPtrTS<IStreamSegment>>& InFinishedSegments, const FPlayStartPosition& StartPosition, IManifest::ESearchType SearchType)
+IManifest::FResult FPlayPeriodHLS::GetLoopingSegment(TSharedPtrTS<IStreamSegment>& OutSegment, const FPlayerSequenceState& SequenceState, const FPlayStartPosition& StartPosition, IManifest::ESearchType SearchType)
 {
-	if (InFinishedSegments.Num())
-	{
-		IManifest::FResult res = GetStartingSegment(OutSegment, SequenceState, StartPosition, SearchType);
-		if (res.GetType() == IManifest::FResult::EType::Found)
-		{
-			return res;
-		}
-	}
-	// Return past EOS when we can't loop to indicate we're really done now.
-	return IManifest::FResult(IManifest::FResult::EType::PastEOS);
+	return GetStartingSegment(OutSegment, SequenceState, StartPosition, SearchType);
 }
 
 
diff --git a/Engine/Plugins/Media/ElectraPlayer/Source/ElectraPlayerRuntime/Private/Runtime/Player/Manifest.h b/Engine/Plugins/Media/ElectraPlayer/Source/ElectraPlayerRuntime/Private/Runtime/Player/Manifest.h
--- a/Engine/Plugins/Media/ElectraPlayer/Source/ElectraPlayerRuntime/Private/Runtime/Player/Manifest.h
+++ b/Engine/Plugins/Media/ElectraPlayer/Source/ElectraPlayerRuntime/Private/Runtime/Player/Manifest.h
@@ -446,13 +446,12 @@ namespace Electra
 			 *
 			 * @param OutSegment
 			 * @param InSequenceState
-			 * @param InFinishedSegments
 			 * @param StartPosition
 			 * @param SearchType
 			 *
 			 * @return
 			 */
-			virtual FResult GetLoopingSegment(TSharedPtrTS<IStreamSegment>& OutSegment, const FPlayerSequenceState& InSequenceState, const TMultiMap<EStreamType, TSharedPtrTS<IStreamSegment>>& InFinishedSegments, const FPlayStartPosition& StartPosition, ESearchType SearchType) = 0;
+			virtual FResult GetLoopingSegment(TSharedPtrTS<IStreamSegment>& OutSegment, const FPlayerSequenceState& InSequenceState, const FPlayStartPosition& StartPosition, ESearchType SearchType) = 0;
 
 			/**
 			 * Gets the segment request for the segment following the specified earlier request.
diff --git a/Engine/Plugins/Media/ElectraPlayer/Source/ElectraPlayerRuntime/Private/Runtime/Player/mp4/ManifestMP4.cpp b/Engine/Plugins/Media/ElectraPlayer/Source/ElectraPlayerRuntime/Private/Runtime/Player/mp4/ManifestMP4.cpp
--- a/Engine/Plugins/Media/ElectraPlayer/Source/ElectraPlayerRuntime/Private/Runtime/Player/mp4/ManifestMP4.cpp
+++ b/Engine/Plugins/Media/ElectraPlayer/Source/ElectraPlayerRuntime/Private/Runtime/Player/mp4/ManifestMP4.cpp
@@ -454,6 +454,7 @@ void FManifestMP4Internal::FPlayPeriodMP4::MakeBufferSourceInfoFromMetadata(EStr
 		OutBufferSourceInfo->Language = InMetadata->Language;
 		OutBufferSourceInfo->Codec = InMetadata->HighestBandwidthCodec.GetCodecName();
 		TSharedPtrTS<FTimelineAssetMP4> Asset = MediaAsset.Pin();
+		OutBufferSourceInfo->PeriodID = Asset->GetUniqueIdentifier();
 		OutBufferSourceInfo->PeriodAdaptationSetID = Asset->GetUniqueIdentifier() + TEXT(".") + InMetadata->ID;
 		TArray<FTrackMetadata> Metadata;
 		Asset->GetMetaData(Metadata, StreamType);
@@ -542,16 +543,15 @@ IManifest::FResult FManifestMP4Internal::FPlayPeriodMP4::GetContinuationSegment(
  *
  * @param OutSegment
  * @param SequenceState
- * @param InFinishedSegments
  * @param StartPosition
  * @param SearchType
  *
  * @return
  */
-IManifest::FResult FManifestMP4Internal::FPlayPeriodMP4::GetLoopingSegment(TSharedPtrTS<IStreamSegment>& OutSegment, const FPlayerSequenceState& SequenceState, const TMultiMap<EStreamType, TSharedPtrTS<IStreamSegment>>& InFinishedSegments, const FPlayStartPosition& StartPosition, ESearchType SearchType)
+IManifest::FResult FManifestMP4Internal::FPlayPeriodMP4::GetLoopingSegment(TSharedPtrTS<IStreamSegment>& OutSegment, const FPlayerSequenceState& SequenceState, const FPlayStartPosition& StartPosition, ESearchType SearchType)
 {
 	TSharedPtrTS<FTimelineAssetMP4> ma = MediaAsset.Pin();
-	return ma.IsValid() ? ma->GetLoopingSegment(OutSegment, SequenceState, InFinishedSegments, StartPosition, SearchType) : IManifest::FResult(IManifest::FResult::EType::NotFound);
+	return ma.IsValid() ? ma->GetLoopingSegment(OutSegment, SequenceState, StartPosition, SearchType) : IManifest::FResult(IManifest::FResult::EType::NotFound);
 }
 
 
@@ -1181,23 +1181,9 @@ IManifest::FResult FManifestMP4Internal::FTimelineAssetMP4::GetRetrySegment(TSha
 }
 
 
-IManifest::FResult FManifestMP4Internal::FTimelineAssetMP4::GetLoopingSegment(TSharedPtrTS<IStreamSegment>& OutSegment, const FPlayerSequenceState& SequenceState, const TMultiMap<EStreamType, TSharedPtrTS<IStreamSegment>>& InFinishedSegments, const FPlayStartPosition& StartPosition, ESearchType SearchType)
+IManifest::FResult FManifestMP4Internal::FTimelineAssetMP4::GetLoopingSegment(TSharedPtrTS<IStreamSegment>& OutSegment, const FPlayerSequenceState& SequenceState, const FPlayStartPosition& StartPosition, ESearchType SearchType)
 {
-	if (InFinishedSegments.Num())
-	{
-		auto It = InFinishedSegments.CreateConstIterator();
-		const FStreamSegmentRequestMP4* FinishedRequest = static_cast<const FStreamSegmentRequestMP4*>(It->Value.Get());
-		if (FinishedRequest)
-		{
-			IManifest::FResult res = GetStartingSegment(OutSegment, SequenceState, StartPosition, SearchType, -1);
-			if (res.GetType() == IManifest::FResult::EType::Found)
-			{
-				return res;
-			}
-		}
-	}
-	// Return past EOS when we can't loop to indicate we're really done now.
-	return IManifest::FResult(IManifest::FResult::EType::PastEOS);
+	return GetStartingSegment(OutSegment, SequenceState, StartPosition, SearchType, -1);
 }
 
 
diff --git a/Engine/Plugins/Media/ElectraPlayer/Source/ElectraPlayerRuntime/Private/Runtime/Player/mp4/ManifestMP4.h b/Engine/Plugins/Media/ElectraPlayer/Source/ElectraPlayerRuntime/Private/Runtime/Player/mp4/ManifestMP4.h
--- a/Engine/Plugins/Media/ElectraPlayer/Source/ElectraPlayerRuntime/Private/Runtime/Player/mp4/ManifestMP4.h
+++ b/Engine/Plugins/Media/ElectraPlayer/Source/ElectraPlayerRuntime/Private/Runtime/Player/mp4/ManifestMP4.h
@@ -235,7 +235,7 @@ public:
 		FResult GetStartingSegment(TSharedPtrTS<IStreamSegment>& OutSegment, const FPlayerSequenceState& InSequenceState, const FPlayStartPosition& StartPosition, ESearchType SearchType, int64 AtAbsoluteFilePos);
 		FResult GetNextSegment(TSharedPtrTS<IStreamSegment>& OutSegment, TSharedPtrTS<const IStreamSegment> CurrentSegment, const FPlayStartOptions& Options);
 		FResult GetRetrySegment(TSharedPtrTS<IStreamSegment>& OutSegment, TSharedPtrTS<const IStreamSegment> CurrentSegment, const FPlayStartOptions& Options, bool bReplaceWithFillerData);
-		FResult GetLoopingSegment(TSharedPtrTS<IStreamSegment>& OutSegment, const FPlayerSequenceState& InSequenceState, const TMultiMap<EStreamType, TSharedPtrTS<IStreamSegment>>& InFinishedSegments, const FPlayStartPosition& StartPosition, ESearchType SearchType);
+		FResult GetLoopingSegment(TSharedPtrTS<IStreamSegment>& OutSegment, const FPlayerSequenceState& InSequenceState, const FPlayStartPosition& StartPosition, ESearchType SearchType);
 
 		void GetSegmentInformation(TArray<IManifest::IPlayPeriod::FSegmentInformation>& OutSegmentInformation, FTimeValue& OutAverageSegmentDuration, TSharedPtrTS<const IStreamSegment> CurrentSegment, const FTimeValue& LookAheadTime, const FString& AdaptationSetID, const FString& RepresentationID);
 		TSharedPtrTS<IParserISO14496_12>	GetMoovBoxParser();
@@ -291,7 +291,7 @@ public:
 		FResult GetContinuationSegment(TSharedPtrTS<IStreamSegment>& OutSegment, EStreamType StreamType, const FPlayerSequenceState& LoopState, const FPlayStartPosition& StartPosition, ESearchType SearchType) override;
 		FResult GetNextSegment(TSharedPtrTS<IStreamSegment>& OutSegment, TSharedPtrTS<const IStreamSegment> CurrentSegment, const FPlayStartOptions& Options) override;
 		FResult GetRetrySegment(TSharedPtrTS<IStreamSegment>& OutSegment, TSharedPtrTS<const IStreamSegment> CurrentSegment, const FPlayStartOptions& Options, bool bReplaceWithFillerData) override;
-		FResult GetLoopingSegment(TSharedPtrTS<IStreamSegment>& OutSegment, const FPlayerSequenceState& InSequenceState, const TMultiMap<EStreamType, TSharedPtrTS<IStreamSegment>>& InFinishedSegments, const FPlayStartPosition& StartPosition, ESearchType SearchType) override;
+		FResult GetLoopingSegment(TSharedPtrTS<IStreamSegment>& OutSegment, const FPlayerSequenceState& InSequenceState, const FPlayStartPosition& StartPosition, ESearchType SearchType) override;
 		void IncreaseSegmentFetchDelay(const FTimeValue& IncreaseAmount) override;
 		void GetSegmentInformation(TArray<FSegmentInformation>& OutSegmentInformation, FTimeValue& OutAverageSegmentDuration, TSharedPtrTS<const IStreamSegment> CurrentSegment, const FTimeValue& LookAheadTime, const FString& AdaptationSetID, const FString& RepresentationID) override;
 
diff --git a/Engine/Plugins/Media/ElectraPlayer/Source/ElectraPlayerRuntime/Private/Runtime/Player/mp4/StreamReaderMP4.cpp b/Engine/Plugins/Media/ElectraPlayer/Source/ElectraPlayerRuntime/Private/Runtime/Player/mp4/StreamReaderMP4.cpp
--- a/Engine/Plugins/Media/ElectraPlayer/Source/ElectraPlayerRuntime/Private/Runtime/Player/mp4/StreamReaderMP4.cpp
+++ b/Engine/Plugins/Media/ElectraPlayer/Source/ElectraPlayerRuntime/Private/Runtime/Player/mp4/StreamReaderMP4.cpp
@@ -469,6 +469,7 @@ void FStreamReaderMP4::HandleRequest()
 					meta->Kind = SelectedTrackMetadata->Kind;
 					meta->Language = SelectedTrackMetadata->Language;
 					meta->Codec = st.CSD.IsValid() ? st.CSD->ParsedInfo.GetCodecName() : FString();
+					meta->PeriodID = SelectedTrackMetadata->PeriodID;
 					meta->PeriodAdaptationSetID = SelectedTrackMetadata->PeriodID + TEXT(".") + SelectedTrackMetadata->AdaptationSetID;
 					meta->HardIndex = SelectedTrackMetadata->Index;
 					meta->PlaybackSequenceID = PlaybackSequenceID;
diff --git a/Engine/Plugins/Media/ElectraPlayer/Source/ElectraPlayerRuntime/Private/Runtime/StreamAccessUnitBuffer.cpp b/Engine/Plugins/Media/ElectraPlayer/Source/ElectraPlayerRuntime/Private/Runtime/StreamAccessUnitBuffer.cpp
--- a/Engine/Plugins/Media/ElectraPlayer/Source/ElectraPlayerRuntime/Private/Runtime/StreamAccessUnitBuffer.cpp
+++ b/Engine/Plugins/Media/ElectraPlayer/Source/ElectraPlayerRuntime/Private/Runtime/StreamAccessUnitBuffer.cpp
@@ -10,7 +10,9 @@ namespace Electra
 		EmptyBuffer = CreateNewBuffer();
 		LastPoppedDTS.SetToInvalid();
 		LastPoppedPTS.SetToInvalid();
+		PlayableDurationPushedSinceEOT.SetToZero();
 		bEndOfData = false;
+		bEndOfTrack = false;
 		bLastPushWasBlocked = false;
 		bIsParallelTrackMode = false;
 		bPopAsDummyUntilSyncFrame = true;
@@ -80,6 +82,19 @@ namespace Electra
 		// Pushing data to either buffer means that there is data and we are not at EOD here any more.
 		// This does not necessarily mean that the selected track is not at EOD. Just that some track is not.
 		bEndOfData = false;
+		// Keep track of how much data since the last end-of-track was signaled has been added.
+		// This is useful when looking at available content to detect underruns which may not be done if
+		// a newly enabled track did not have enough time to collect new data. Only when sufficient data
+		// was available might it make sense to check for underruns.
+		if (bEndOfTrack)
+		{
+			bEndOfTrack = false;
+			PlayableDurationPushedSinceEOT.SetToZero();
+		}
+		if (AU->DropState == FAccessUnit::EDropState::None)
+		{
+			PlayableDurationPushedSinceEOT += AU->Duration;
+		}
 
 		AccessLock.Unlock();
 
@@ -116,6 +131,30 @@ namespace Electra
 		}
 	}
 
+	void FMultiTrackAccessUnitBuffer::SetEndOfTrackFor(TSharedPtrTS<const FBufferSourceInfo> InStreamSourceInfo)
+	{
+		AccessLock.Lock();
+		TSharedPtrTS<const FBufferSourceInfo> BufferSourceInfo;
+		TSharedPtrTS<FAccessUnitBuffer> Buffer = FindOrCreateBufferFor(BufferSourceInfo, InStreamSourceInfo, true);
+		ActivateInitialBuffer();
+		AccessLock.Unlock();
+		if (Buffer.IsValid())
+		{
+			Buffer->SetEndOfTrack();
+		}
+	}
+
+	void FMultiTrackAccessUnitBuffer::SetEndOfTrackAll()
+	{
+		FScopeLock lock(&AccessLock);
+		bEndOfTrack = true;
+		// Push an end-of-data into all tracks.
+		for(auto& It : BufferList)
+		{
+			It.Buffer->SetEndOfTrack();
+		}
+	}
+
 	void FMultiTrackAccessUnitBuffer::Clear()
 	{
 		BufferList.Empty();
@@ -125,7 +164,9 @@ namespace Electra
 		LastPoppedBufferInfo.Reset();
 		LastPoppedDTS.SetToInvalid();
 		LastPoppedPTS.SetToInvalid();
+		PlayableDurationPushedSinceEOT.SetToZero();
 		bEndOfData = false;
+		bEndOfTrack = false;
 		bLastPushWasBlocked = false;
 		bPopAsDummyUntilSyncFrame = true;
 	}
@@ -231,6 +272,12 @@ namespace Electra
 		return LastPoppedDTS;
 	}
 
+	FTimeValue FMultiTrackAccessUnitBuffer::GetPlayableDurationPushedSinceEOT()
+	{
+		FScopeLock lock(&AccessLock);
+		return PlayableDurationPushedSinceEOT;
+	}
+
 	TSharedPtrTS<FAccessUnitBuffer> FMultiTrackAccessUnitBuffer::GetSelectedTrackBuffer()
 	{
 		FScopeLock lock(&AccessLock);
@@ -246,6 +293,7 @@ namespace Electra
 		if (Buf == EmptyBuffer)
 		{
 			OutStats.bEndOfData = bEndOfData;
+			OutStats.bEndOfTrack = bEndOfTrack;
 		}
 	}
 
@@ -346,6 +394,13 @@ namespace Electra
 		return Buf == EmptyBuffer ? bEndOfData : Buf->IsEODFlagSet();
 	}
 
+	bool FMultiTrackAccessUnitBuffer::IsEndOfTrack()
+	{
+		// Note: We assume the access lock is held by the caller!
+		TSharedPtrTS<FAccessUnitBuffer> Buf = GetSelectedTrackBuffer();
+		return Buf == EmptyBuffer ? bEndOfTrack : Buf->IsEndOfTrack();
+	}
+
 	int32 FMultiTrackAccessUnitBuffer::Num()
 	{
 		// Note: We assume the access lock is held by the caller!
diff --git a/Engine/Plugins/Media/ElectraPlayer/Source/ElectraPlayerRuntime/Private/Runtime/StreamAccessUnitBuffer.h b/Engine/Plugins/Media/ElectraPlayer/Source/ElectraPlayerRuntime/Private/Runtime/StreamAccessUnitBuffer.h
--- a/Engine/Plugins/Media/ElectraPlayer/Source/ElectraPlayerRuntime/Private/Runtime/StreamAccessUnitBuffer.h
+++ b/Engine/Plugins/Media/ElectraPlayer/Source/ElectraPlayerRuntime/Private/Runtime/StreamAccessUnitBuffer.h
@@ -35,6 +35,8 @@ namespace Electra
 	 */
 	struct FBufferSourceInfo
 	{
+		// The period the data comes from. Necessary to track period transitions.
+		FString		PeriodID;
 		// Identifies the period and track (adaptation set) this data is originating from.
 		FString		PeriodAdaptationSetID;
 		// Partial track metadata. See FTrackMetadata.
@@ -206,6 +208,7 @@ namespace Electra
 			CurrentMemInUse = 0;
 			NumCurrentAccessUnits = 0;
 			bEndOfData = false;
+			bEndOfTrack = false;
 			bLastPushWasBlocked = false;
 		}
 
@@ -215,6 +218,7 @@ namespace Electra
 		int64				CurrentMemInUse;
 		int64				NumCurrentAccessUnits;
 		bool				bEndOfData;
+		bool				bEndOfTrack;
 		bool				bLastPushWasBlocked;
 	};
 
@@ -248,6 +252,7 @@ namespace Electra
 			, PlayableDuration(FTimeValue::GetZero())
 			, CurrentMemInUse(0)
 			, bEndOfData(false)
+			, bEndOfTrack(false)
 			, bLastPushWasBlocked(false)
 		{
 		}
@@ -290,6 +295,7 @@ namespace Electra
 			OutStats.CurrentMemInUse = CurrentMemInUse;
 			OutStats.NumCurrentAccessUnits = AccessUnits.Num();
 			OutStats.bEndOfData = bEndOfData;
+			OutStats.bEndOfTrack = bEndOfTrack;
 			OutStats.bLastPushWasBlocked = bLastPushWasBlocked;
 		}
 
@@ -300,6 +306,7 @@ namespace Electra
 			// Pushing new data unconditionally clears the EOD flag even if the buffer is currently full.
 			// The attempt to push implies there will be more data.
 			bEndOfData = false;
+			bEndOfTrack = false;
 			ExternalInfo = ExternalInfo ? ExternalInfo : &ZeroExternalInfo;
 			if (CanPush(AU, Limit, ExternalInfo))
 			{
@@ -334,6 +341,19 @@ namespace Electra
 			bEndOfData = true;
 		}
 
+		void SetEndOfTrack()
+		{
+			bEndOfTrack = true;
+		}
+		void ClearEndOfTrack()
+		{
+			bEndOfTrack = false;
+		}
+		bool IsEndOfTrack() const
+		{
+			return bEndOfTrack;
+		}
+
 		//! Removes and returns the oldest access unit from the FIFO. Returns false if the FIFO is empty.
 		bool Pop(FAccessUnit*& OutAU)
 		{
@@ -537,6 +557,7 @@ namespace Electra
 			}
 			CurrentMemInUse = 0;
 			bEndOfData = false;
+			bEndOfTrack = false;
 			bLastPushWasBlocked = false;
 			FrontDTS.SetToInvalid();
 			PushedDuration.SetToZero();
@@ -629,9 +650,10 @@ namespace Electra
 		FTimeValue									FrontDTS;					//!< DTS of first AU in buffer
 		FTimeValue									PushedDuration;
 		FTimeValue									PlayableDuration;
-		int64										CurrentMemInUse;
-		bool										bEndOfData;
-		bool										bLastPushWasBlocked;
+		int64										CurrentMemInUse = 0;
+		bool										bEndOfData = false;
+		bool										bEndOfTrack = false;
+		bool										bLastPushWasBlocked = false;
 	};
 
 
@@ -651,10 +673,13 @@ namespace Electra
 		bool Push(FAccessUnit*& AU, const FAccessUnitBuffer::FConfiguration* BufferConfiguration, const FAccessUnitBuffer::FExternalBufferInfo* InCurrentTotalBufferUtilization);
 		void PushEndOfDataFor(TSharedPtrTS<const FBufferSourceInfo> InStreamSourceInfo);
 		void PushEndOfDataAll();
+		void SetEndOfTrackFor(TSharedPtrTS<const FBufferSourceInfo> InStreamSourceInfo);
+		void SetEndOfTrackAll();
 		void Flush();
 		void GetStats(FAccessUnitBufferInfo& OutStats);
 		FTimeValue GetLastPoppedDTS();
 		FTimeValue GetLastPoppedPTS();
+		FTimeValue GetPlayableDurationPushedSinceEOT();
 
 		// Helper class to lock the AU buffer
 		class FScopedLock
@@ -681,6 +706,7 @@ namespace Electra
 		bool Pop(FAccessUnit*& OutAU);
 		void PopDiscardUntil(FTimeValue UntilTime);
 		bool IsEODFlagSet();
+		bool IsEndOfTrack();
 		int32 Num();
 		bool WasLastPushBlocked();
 		bool HasPendingTrackSwitch();
@@ -726,7 +752,9 @@ namespace Electra
 		TSharedPtrTS<const FBufferSourceInfo>			LastPoppedBufferInfo;
 		FTimeValue										LastPoppedDTS;
 		FTimeValue										LastPoppedPTS;
+		FTimeValue										PlayableDurationPushedSinceEOT;
 		bool											bEndOfData;
+		bool											bEndOfTrack;
 		bool											bLastPushWasBlocked;
 		bool											bPopAsDummyUntilSyncFrame;
 		bool											bIsParallelTrackMode;
