ElectraPlayer: Added player option to have a CDN's HTTP response code of choice disable a stream permanently

#rnx

diff --git a/Engine/Plugins/Media/ElectraPlayer/Source/ElectraPlayerPlugin/Private/ElectraPlayerPlugin.cpp b/Engine/Plugins/Media/ElectraPlayer/Source/ElectraPlayerPlugin/Private/ElectraPlayerPlugin.cpp
--- a/Engine/Plugins/Media/ElectraPlayer/Source/ElectraPlayerPlugin/Private/ElectraPlayerPlugin.cpp
+++ b/Engine/Plugins/Media/ElectraPlayer/Source/ElectraPlayerPlugin/Private/ElectraPlayerPlugin.cpp
@@ -595,6 +595,18 @@ bool FElectraPlayerPlugin::Open(const FString& Url, const IMediaOptions* Options
 		PlayerOptions.Set("throw_error_when_rebuffering", Electra::FVariantValue(bThrowErrorWhenRebuffering));
 		UE_LOG(LogElectraPlayerPlugin, Log, TEXT("[%p] IMediaPlayer::Open: Throw playback error when rebuffering"), this);
 	}
+	FString CDNHTTPStatusDenyStream = Options->GetMediaOption(TEXT("ElectraGetDenyStreamCode"), FString());
+	if (CDNHTTPStatusDenyStream.Len())
+	{
+		int32 HTTPStatus = -1;
+		LexFromString(HTTPStatus, *CDNHTTPStatusDenyStream);
+		if (HTTPStatus > 0 && HTTPStatus < 1000)
+		{
+			PlayerOptions.Set("abr:cdn_deny_httpstatus", Electra::FVariantValue((int64)HTTPStatus));
+			UE_LOG(LogElectraPlayerPlugin, Log, TEXT("[%p] IMediaPlayer::Open: CDN HTTP status %d will deny a stream permanently"), this, HTTPStatus);
+		}
+	}
+
 
 	// Check for options that can be changed during playback and apply them at startup already.
 	// If a media source supports the MaxResolutionForMediaStreaming option then we can override the max resolution.
diff --git a/Engine/Plugins/Media/ElectraPlayer/Source/ElectraPlayerRuntime/Private/Runtime/Player/ABRRules/ABRLiveStream.cpp b/Engine/Plugins/Media/ElectraPlayer/Source/ElectraPlayerRuntime/Private/Runtime/Player/ABRRules/ABRLiveStream.cpp
--- a/Engine/Plugins/Media/ElectraPlayer/Source/ElectraPlayerRuntime/Private/Runtime/Player/ABRRules/ABRLiveStream.cpp
+++ b/Engine/Plugins/Media/ElectraPlayer/Source/ElectraPlayerRuntime/Private/Runtime/Player/ABRRules/ABRLiveStream.cpp
@@ -2,6 +2,7 @@
 
 #include "Player/ABRRules/ABRLiveStream.h"
 #include "Player/ABRRules/ABRStatisticTypes.h"
+#include "Player/ABRRules/ABROptionKeynames.h"
 #include "Player/AdaptivePlayerOptionKeynames.h"
 
 #include "Player/Manifest.h"
@@ -186,6 +187,7 @@ private:
 			HighestQualityStreamBitrate = 0;
 			LowestQualityStreamBitrate = 0;
 			NumQualities = 0;
+			NumConsecutiveSegmentErrors = 0;
 			SegmentDownloadHistory.Clear();
 			OverDownloadTimeTotal = 0.0;
 			QualityMetrics.Empty();
@@ -223,6 +225,7 @@ private:
 		int64 HighestQualityStreamBitrate = 0;
 		int64 LowestQualityStreamBitrate = 0;
 		int32 NumQualities = 0;
+		int32 NumConsecutiveSegmentErrors = 0;
 		bool bIsLowLatencyEnabled = false;
 
 		double BufferContentDurationAtSegmentStart = -1.0;
@@ -328,6 +331,7 @@ private:
 	bool bSkipWithFiller = false;
 
 	// Configuration
+	TMediaOptionalValue<int32> HTTPStatusCodeToDenyStream;
 	bool bRebufferingJustResumesLoading = false;
 	bool bEmergencyBufferSlowdownPossible = false;
 
@@ -393,6 +397,10 @@ FABRLiveStream::FABRLiveStream(IABRInfoInterface* InInfo, EMediaFormatType InFor
 	, PresentationType(InPresentationType)
 {
 	bRebufferingJustResumesLoading = Info->GetPlayerOptions().GetValue(OptionRebufferingContinuesLoading).SafeGetBool(false);
+	if (Info->GetPlayerOptions().HaveKey(ABR::OptionKeyABR_CDNSegmentDenyHTTPStatus))
+	{
+		HTTPStatusCodeToDenyStream.Set((int32) Info->GetPlayerOptions().GetValue(ABR::OptionKeyABR_CDNSegmentDenyHTTPStatus).SafeGetInt64(-1));
+	}
 	PrepareLatencyConfiguration();
 
 	FTimeRange RenderRange = Info->ABRGetSupportedRenderRateScale();
@@ -893,11 +901,18 @@ IAdaptiveStreamSelector::ESegmentAction FABRLiveStream::EvaluateForError(TArray<
 		TSharedPtrTS<FABRStreamInformation> CurrentStreamInfo = Info->GetStreamInformation(Stats);
 		if (CurrentStreamInfo.IsValid())
 		{
-			// Remember statis
+			// Remember stats
 			CurrentStreamInfo->Health.LastDownloadStats = Stats;
 
 			if (!Stats.bWasSuccessful)
 			{
+				// If this is the first failure for this segment (not being retried yet) increase
+				// the number of consecutively failed segments.
+				if (Stats.RetryNumber == 0)
+				{
+					++WorkVars->NumConsecutiveSegmentErrors;
+				}
+
 				// Was not successful. Figure out what to do now.
 				// In the case where the segment had an availability window set the stream reader was waiting for
 				// to be entered and we got a 404 back, the server did not manage to publish the new segment in time.
@@ -931,21 +946,32 @@ IAdaptiveStreamSelector::ESegmentAction FABRLiveStream::EvaluateForError(TArray<
 					Info->LogMessage(IInfoLog::ELevel::Error, FString::Printf(TEXT("Exceeded permissable number of retries (%d). Failing now."), Stats.RetryNumber));
 					return IAdaptiveStreamSelector::ESegmentAction::Fail;
 				}
+				if (WorkVars->NumConsecutiveSegmentErrors >= 5)
+				{
+					Info->LogMessage(IInfoLog::ELevel::Error, FString::Printf(TEXT("Exceeded permissable number of consecutive segment failures (%d). Failing now."), WorkVars->NumConsecutiveSegmentErrors));
+					return IAdaptiveStreamSelector::ESegmentAction::Fail;
+				}
 
 				// Is this an init segment?
 				if (Stats.SegmentType == Metrics::ESegmentType::Init)
 				{
 					bRetryIfPossible = true;
-					// Is the init segment absent or failed to be parsed? If so this stream is dead for good.
-					if (Stats.bParseFailure || Stats.HTTPStatusCode == 404)
-					{
-						CurrentStreamInfo->Health.BecomesAvailableAgainAtUTC.SetToInvalid();
-					}
-					// Any other case of failure other than having aborted in ReportDownloadProgress() ?
-					else if (!Stats.bWasAborted)
+					// A failure other than having aborted in ReportDownloadProgress() ?
+					if (!Stats.bWasAborted)
 					{
+						// Did the init segment fail to be parsed or does the HTTP status match the one used to disable the stream?
+						// If so this stream is dead for good.
+						if (Stats.bParseFailure || (HTTPStatusCodeToDenyStream.IsSet() && HTTPStatusCodeToDenyStream.Value() == Stats.HTTPStatusCode))
+						{
+							CurrentStreamInfo->Health.BecomesAvailableAgainAtUTC.SetToInvalid();
+						}
+						else if (Stats.HTTPStatusCode && Stats.RetryNumber == 0)
+						{
+							// Do one more try.
+							OutDelay.SetFromMilliseconds(150);
+						}
 						// Take this stream offline for a brief moment.
-						if (CurrentStreamInfo->Health.BecomesAvailableAgainAtUTC == FTimeValue::GetZero())
+						else if (CurrentStreamInfo->Health.BecomesAvailableAgainAtUTC == FTimeValue::GetZero())
 						{
 							CurrentStreamInfo->Health.BecomesAvailableAgainAtUTC = TimeNow + FTimeValue().SetFromMilliseconds(1000);
 						}
@@ -954,49 +980,75 @@ IAdaptiveStreamSelector::ESegmentAction FABRLiveStream::EvaluateForError(TArray<
 				// A media segment failure
 				else
 				{
-					bRetryIfPossible = true;
-					bSkipWithFiller = true;
-
-					// Did we abort that stream in ReportDownloadProgress() ?
-					if (!Stats.bWasAborted)
+					// Check if there is a HTTP status code set which, if returned, is used to disable this stream permanently.
+					if (HTTPStatusCodeToDenyStream.IsSet() && HTTPStatusCodeToDenyStream.Value() == Stats.HTTPStatusCode)
 					{
-						// Take a video stream offline for a brief moment unless it is the only one. Audio is usually the only one with no alternative to switch to, so don't do that!
-						if (Stats.StreamType == EStreamType::Video)
+						CurrentStreamInfo->Health.BecomesAvailableAgainAtUTC.SetToInvalid();
+						bRetryIfPossible = true;
+						bSkipWithFiller = false;
+					}
+					else
+					{
+						bRetryIfPossible = true;
+						bSkipWithFiller = true;
+
+						// Did we abort that stream in ReportDownloadProgress() ?
+						if (!Stats.bWasAborted)
 						{
-							if (Info->GetStreamInformations(EStreamType::Video).Num() > 1)
+							// Take a video stream offline for a brief moment unless it is the worst one. Audio is usually the only one with no alternative to switch to, so don't do that!
+							if (Stats.StreamType == EStreamType::Video)
 							{
-								if (CurrentStreamInfo->Health.BecomesAvailableAgainAtUTC == FTimeValue::GetZero())
+								if (CurrentStreamInfo->QualityIndex != 0)
 								{
-									CurrentStreamInfo->Health.BecomesAvailableAgainAtUTC = TimeNow + FTimeValue().SetFromMilliseconds(1000);
+									if (CurrentStreamInfo->Health.BecomesAvailableAgainAtUTC == FTimeValue::GetZero())
+									{
+										CurrentStreamInfo->Health.BecomesAvailableAgainAtUTC = TimeNow + FTimeValue().SetFromMilliseconds(1000);
+									}
+								}
+								else
+								{
+									// Need to insert filler data instead.
+									bRetryIfPossible = false;
 								}
 							}
 							else
 							{
-								// Need to insert filler data instead.
-								bRetryIfPossible = false;
+								// For VoD we retry the audio segment once with a small delay before skipping over it, which we do for Live right away
+								// since retrying is a luxury we do not have in this case.
+								// Note: This here ABR is optimized for Live. For VoD a dedicated Live ABR should be employed!
+								if (PresentationType == EABRPresentationType::OnDemand)
+								{
+									if (Stats.RetryNumber == 0)
+									{
+										OutDelay.SetFromMilliseconds(150);
+										bRetryIfPossible = true;
+									}
+									else
+									{
+										bRetryIfPossible = false;
+									}
+								}
+								else
+								{
+									bRetryIfPossible = false;
+								}
 							}
 						}
-						else
+
+						// If any content has already been put out into the buffers we cannot retry the segment on another quality level.
+						if (Stats.DurationDelivered > 0.0)
 						{
-							// For VoD we retry the audio segment up to a certain number of times while for Live we skip over it
-							// since retrying is a luxury we do not have in this case.
-							if (PresentationType == EABRPresentationType::OnDemand)
-							{
-								bRetryIfPossible = true;
-							}
-							else
-							{
-								bRetryIfPossible = false;
-							}
+							bRetryIfPossible = false;
+							bSkipWithFiller = false;
 						}
 					}
-
-					// If any content has already been put out into the buffers we cannot retry the segment on another quality level.
-					if (Stats.DurationDelivered > 0.0)
-					{
-						bRetryIfPossible = false;
-						bSkipWithFiller = false;
-					}
+				}
+			}
+			else
+			{
+				if (!Stats.bInsertedFillerData)
+				{
+					WorkVars->NumConsecutiveSegmentErrors = 0;
 				}
 			}
 		}
diff --git a/Engine/Plugins/Media/ElectraPlayer/Source/ElectraPlayerRuntime/Private/Runtime/Player/ABRRules/ABROptionKeynames.h b/Engine/Plugins/Media/ElectraPlayer/Source/ElectraPlayerRuntime/Private/Runtime/Player/ABRRules/ABROptionKeynames.h
new file mode 100644
--- /dev/null
+++ b/Engine/Plugins/Media/ElectraPlayer/Source/ElectraPlayerRuntime/Private/Runtime/Player/ABRRules/ABROptionKeynames.h
@@ -0,0 +1,19 @@
+// Copyright Epic Games, Inc. All Rights Reserved.
+
+#pragma once
+
+#include "PlayerCore.h"
+
+namespace Electra
+{
+namespace ABR
+{
+	//! (bool) true to just finish the currently loading segment when rebuffering. false to start over with.
+	const TCHAR* const OptionKeyABR_RebufferingContinuesLoading = TEXT("abr:rebuffering_continues_loading");
+
+	//! (int) HTTP status code if returned by CDN for any segment request will take that stream permanently offline.
+	const TCHAR* const OptionKeyABR_CDNSegmentDenyHTTPStatus = TEXT("abr:cdn_deny_httpstatus");
+}
+} // namespace Electra
+
+
diff --git a/Engine/Plugins/Media/ElectraPlayer/Source/ElectraPlayerRuntime/Private/Runtime/Player/ABRRules/ABRVoDPlus.cpp b/Engine/Plugins/Media/ElectraPlayer/Source/ElectraPlayerRuntime/Private/Runtime/Player/ABRRules/ABRVoDPlus.cpp
--- a/Engine/Plugins/Media/ElectraPlayer/Source/ElectraPlayerRuntime/Private/Runtime/Player/ABRRules/ABRVoDPlus.cpp
+++ b/Engine/Plugins/Media/ElectraPlayer/Source/ElectraPlayerRuntime/Private/Runtime/Player/ABRRules/ABRVoDPlus.cpp
@@ -2,6 +2,7 @@
 
 #include "Player/ABRRules/ABRVoDPlus.h"
 #include "Player/ABRRules/ABRStatisticTypes.h"
+#include "Player/ABRRules/ABROptionKeynames.h"
 #include "Player/AdaptivePlayerOptionKeynames.h"
 
 #include "Player/Manifest.h"
@@ -82,6 +83,7 @@ private:
 			LowestQualityStreamBitrate = 0;
 			HighestQualityStreamBitrate = 0;
 			NumQualities = 0;
+			NumConsecutiveSegmentErrors = 0;
 			// Note: Do not reset AverageBandwidth and AverageLatency!
 		}
 
@@ -109,6 +111,7 @@ private:
 		int64 LowestQualityStreamBitrate = 0;
 		int64 HighestQualityStreamBitrate = 0;
 		int32 NumQualities = 0;
+		int32 NumConsecutiveSegmentErrors = 0;
 
 		int32 QualityIndexDownloading = 0;
 		int32 BitrateDownloading = 0;
@@ -186,6 +189,7 @@ private:
 	bool bSkipWithFiller = false;
 
 	// Configuration
+	TMediaOptionalValue<int32> HTTPStatusCodeToDenyStream;
 	bool bRebufferingJustResumesLoading = false;
 	bool bEmergencyBufferSlowdownPossible = false;
 
@@ -238,6 +242,10 @@ FABROnDemandPlus::FABROnDemandPlus(IABRInfoInterface* InInfo, EMediaFormatType I
 	, PresentationType(InPresentationType)
 {
 	bRebufferingJustResumesLoading = Info->GetPlayerOptions().GetValue(OptionRebufferingContinuesLoading).SafeGetBool(false);
+	if (Info->GetPlayerOptions().HaveKey(ABR::OptionKeyABR_CDNSegmentDenyHTTPStatus))
+	{
+		HTTPStatusCodeToDenyStream.Set((int32) Info->GetPlayerOptions().GetValue(ABR::OptionKeyABR_CDNSegmentDenyHTTPStatus).SafeGetInt64(-1));
+	}
 
 	FTimeRange RenderRange = Info->ABRGetSupportedRenderRateScale();
 	bEmergencyBufferSlowdownPossible = RenderRange.IsValid();
@@ -596,6 +604,12 @@ IAdaptiveStreamSelector::ESegmentAction FABROnDemandPlus::EvaluateForError(TArra
 
 			if (!Stats.bWasSuccessful)
 			{
+				// If this is the first failure for this segment (not being retried yet) increase
+				// the number of consecutively failed segments.
+				if (Stats.RetryNumber == 0)
+				{
+					++WorkVars->NumConsecutiveSegmentErrors;
+				}
 				// Was not successful. Figure out what to do now.
 				// In the case where the segment had an availability window set the stream reader was waiting for
 				// to be entered and we got a 404 back, the server did not manage to publish the new segment in time.
@@ -607,7 +621,7 @@ IAdaptiveStreamSelector::ESegmentAction FABROnDemandPlus::EvaluateForError(TArra
 					//Info->LogMessage(IInfoLog::ELevel::Warning, FString::Printf(TEXT("Segment \"%s\" not available at announced time. Trying again in 0.15s"), *GetFilenameFromURL(Stats.URL)));
 					CurrentPlayPeriod->IncreaseSegmentFetchDelay(FTimeValue(FTimeValue::MillisecondsToHNS(100)));
 					OutDelay.SetFromMilliseconds(150);
-					return IAdaptiveStreamSelector::ESegmentAction::Retry;
+					return Stats.RetryNumber < 3 ? IAdaptiveStreamSelector::ESegmentAction::Retry : IAdaptiveStreamSelector::ESegmentAction::Fail;
 				}
 
 				// Too many failures already? It's unlikely that there is a way that would magically fix itself now.
@@ -616,21 +630,32 @@ IAdaptiveStreamSelector::ESegmentAction FABROnDemandPlus::EvaluateForError(TArra
 					Info->LogMessage(IInfoLog::ELevel::Error, FString::Printf(TEXT("Exceeded permissable number of retries (%d). Failing now."), Stats.RetryNumber));
 					return IAdaptiveStreamSelector::ESegmentAction::Fail;
 				}
+				if (WorkVars->NumConsecutiveSegmentErrors >= 3)
+				{
+					Info->LogMessage(IInfoLog::ELevel::Error, FString::Printf(TEXT("Exceeded permissable number of consecutive segment failures (%d). Failing now."), WorkVars->NumConsecutiveSegmentErrors));
+					return IAdaptiveStreamSelector::ESegmentAction::Fail;
+				}
 
 				// Is this an init segment?
 				if (Stats.SegmentType == Metrics::ESegmentType::Init)
 				{
 					bRetryIfPossible = true;
-					// Is the init segment absent or failed to be parsed? If so this stream is dead for good.
-					if (Stats.bParseFailure || Stats.HTTPStatusCode == 404)
-					{
-						CurrentStreamInfo->Health.BecomesAvailableAgainAtUTC.SetToInvalid();
-					}
-					// Any other case of failure other than having aborted in ReportDownloadProgress() ?
-					else if (!Stats.bWasAborted)
+					// A failure other than having aborted in ReportDownloadProgress() ?
+					if (!Stats.bWasAborted)
 					{
+						// Did the init segment fail to be parsed or does the HTTP status match the one used to disable the stream?
+						// If so this stream is dead for good.
+						if (Stats.bParseFailure || (HTTPStatusCodeToDenyStream.IsSet() && HTTPStatusCodeToDenyStream.Value() == Stats.HTTPStatusCode))
+						{
+							CurrentStreamInfo->Health.BecomesAvailableAgainAtUTC.SetToInvalid();
+						}
+						else if (Stats.HTTPStatusCode && Stats.RetryNumber == 0)
+						{
+							// Do one more try.
+							OutDelay.SetFromMilliseconds(150);
+						}
 						// Take this stream offline for a brief moment.
-						if (CurrentStreamInfo->Health.BecomesAvailableAgainAtUTC == FTimeValue::GetZero())
+						else if (CurrentStreamInfo->Health.BecomesAvailableAgainAtUTC == FTimeValue::GetZero())
 						{
 							CurrentStreamInfo->Health.BecomesAvailableAgainAtUTC = TimeNow + FTimeValue().SetFromMilliseconds(1000);
 						}
@@ -639,40 +664,58 @@ IAdaptiveStreamSelector::ESegmentAction FABROnDemandPlus::EvaluateForError(TArra
 				// A media segment failure
 				else
 				{
-					bRetryIfPossible = true;
-					bSkipWithFiller = true;
-
-					// Did we abort that stream in ReportDownloadProgress() ?
-					if (!Stats.bWasAborted)
+					// Check if there is a HTTP status code set which, if returned, is used to disable this stream permanently.
+					if (HTTPStatusCodeToDenyStream.IsSet() && HTTPStatusCodeToDenyStream.Value() == Stats.HTTPStatusCode)
+					{
+						CurrentStreamInfo->Health.BecomesAvailableAgainAtUTC.SetToInvalid();
+						bRetryIfPossible = true;
+						bSkipWithFiller = false;
+					}
+					else
 					{
-						// Take a video stream offline for a brief moment unless it is the only one. Audio is usually the only one with no alternative to switch to, so don't do that!
-						if (Stats.StreamType == EStreamType::Video)
+						bRetryIfPossible = true;
+						bSkipWithFiller = true;
+
+						// Did we abort that stream in ReportDownloadProgress() ?
+						if (!Stats.bWasAborted)
 						{
-							if (Info->GetStreamInformations(EStreamType::Video).Num() > 1)
+							// Take a video stream offline for a brief moment unless it is the worst one. Audio is usually the only one with no alternative to switch to, so don't do that!
+							if (Stats.StreamType == EStreamType::Video)
 							{
-								if (CurrentStreamInfo->Health.BecomesAvailableAgainAtUTC == FTimeValue::GetZero())
+								if (CurrentStreamInfo->QualityIndex != 0)
 								{
-									CurrentStreamInfo->Health.BecomesAvailableAgainAtUTC = TimeNow + FTimeValue().SetFromMilliseconds(1000);
+									if (CurrentStreamInfo->Health.BecomesAvailableAgainAtUTC == FTimeValue::GetZero())
+									{
+										CurrentStreamInfo->Health.BecomesAvailableAgainAtUTC = TimeNow + FTimeValue().SetFromMilliseconds(1000);
+									}
+								}
+								else
+								{
+									// Need to insert filler data instead.
+									bRetryIfPossible = false;
 								}
 							}
 							else
 							{
-								// Need to insert filler data instead.
-								bRetryIfPossible = false;
-							}
-						}
-						else
-						{
-							// For VoD we retry the audio segment up to a certain number of times while for Live we skip over it
-							// since retrying is a luxury we do not have in this case.
-							// Note: This here ABR is optimized for VoD. For Live a dedicated Live ABR should be employed!
-							if (PresentationType == EABRPresentationType::OnDemand)
-							{
-								bRetryIfPossible = true;
-							}
-							else
-							{
-								bRetryIfPossible = false;
+								// For VoD we retry the audio segment once with a small delay before skipping over it, which we do for Live right away
+								// since retrying is a luxury we do not have in this case.
+								// Note: This here ABR is optimized for VoD. For Live a dedicated Live ABR should be employed!
+								if (PresentationType == EABRPresentationType::OnDemand)
+								{
+									if (Stats.RetryNumber == 0)
+									{
+										OutDelay.SetFromMilliseconds(150);
+										bRetryIfPossible = true;
+									}
+									else
+									{
+										bRetryIfPossible = false;
+									}
+								}
+								else
+								{
+									bRetryIfPossible = false;
+								}
 							}
 						}
 					}
@@ -685,6 +728,13 @@ IAdaptiveStreamSelector::ESegmentAction FABROnDemandPlus::EvaluateForError(TArra
 					}
 				}
 			}
+			else
+			{
+				if (!Stats.bInsertedFillerData)
+				{
+					WorkVars->NumConsecutiveSegmentErrors = 0;
+				}
+			}
 		}
 	}
 	return IAdaptiveStreamSelector::ESegmentAction::FetchNext;
