ElectraPlayer: Retaining the last 20 concise log messages to add to the analytics blob to provide better context for errors.

#rnx

diff --git a/Engine/Plugins/Media/ElectraPlayer/Source/ElectraPlayerRuntime/Private/ElectraPlayer.cpp b/Engine/Plugins/Media/ElectraPlayer/Source/ElectraPlayerRuntime/Private/ElectraPlayer.cpp
--- a/Engine/Plugins/Media/ElectraPlayer/Source/ElectraPlayerRuntime/Private/ElectraPlayer.cpp
+++ b/Engine/Plugins/Media/ElectraPlayer/Source/ElectraPlayerRuntime/Private/ElectraPlayer.cpp
@@ -1818,6 +1818,7 @@ void FElectraPlayer::HandlePlayerEventOpenSource(const FString& URL)
 
 	// Update statistics
 	FScopeLock Lock(&StatisticsLock);
+	Statistics.AddMessageToHistory(TEXT("Opening stream"));
 	Statistics.InitialURL = URL;
 	Statistics.TimeAtOpen = FPlatformTime::Seconds();
 	Statistics.LastState  = "Opening";
@@ -1834,13 +1835,15 @@ void FElectraPlayer::HandlePlayerEventOpenSource(const FString& URL)
 
 void FElectraPlayer::HandlePlayerEventReceivedMasterPlaylist(const FString& EffectiveURL)
 {
+	UE_LOG(LogElectraPlayer, Log, TEXT("[%p][%p] Received master playlist from \"%s\""), this, CurrentPlayer.Get(), *SanitizeMessage(EffectiveURL));
+
 	// Update statistics
 	FScopeLock Lock(&StatisticsLock);
+	Statistics.AddMessageToHistory(TEXT("Got master playlist"));
 	// Note the time it took to get the master playlist
 	Statistics.TimeToLoadMasterPlaylist = FPlatformTime::Seconds() - Statistics.TimeAtOpen;
 	Statistics.LastState = "Preparing";
 
-	UE_LOG(LogElectraPlayer, Log, TEXT("[%p][%p] Received master playlist from \"%s\""), this, CurrentPlayer.Get(), *SanitizeMessage(EffectiveURL));
 
 	// Enqueue a "MasterPlaylist" event.
 	static const FString kEventNameElectraMasterPlaylist(TEXT("Electra.MasterPlaylist"));
@@ -1872,6 +1875,7 @@ void FElectraPlayer::HandlePlayerEventReceivedPlaylists()
 
 	// Update statistics
 	StatisticsLock.Lock();
+	Statistics.AddMessageToHistory(TEXT("Got initial playlists"));
 	// Note the time it took to get the stream playlist
 	Statistics.TimeToLoadStreamPlaylists = FPlatformTime::Seconds() - Statistics.TimeAtOpen;
 	Statistics.LastState = "Idle";
@@ -1999,10 +2003,14 @@ void FElectraPlayer::HandlePlayerEventLicenseKey(const Electra::Metrics::FLicens
 	if (LicenseKeyStats.bWasSuccessful)
 	{
 		UE_LOG(LogElectraPlayer, Log, TEXT("[%p][%p] License key obtained"), this, CurrentPlayer.Get());
+		FScopeLock Lock(&StatisticsLock);
+		Statistics.AddMessageToHistory(TEXT("Obtained license key"));
 	}
 	else
 	{
 		UE_LOG(LogElectraPlayer, Log, TEXT("[%p][%p] License key error \"%s\""), this, CurrentPlayer.Get(), *LicenseKeyStats.FailureReason);
+		FScopeLock Lock(&StatisticsLock);
+		Statistics.AddMessageToHistory(TEXT("License key error"));
 	}
 }
 
@@ -2065,7 +2073,11 @@ void FElectraPlayer::HandlePlayerEventBufferingStart(Electra::Metrics::EBufferin
 		AnalyticEvent->ParamArray.Add(FAnalyticsEventAttribute(TEXT("Type"), Electra::Metrics::GetBufferingReasonString(BufferingReason)));
 		EnqueueAnalyticsEvent(AnalyticEvent);
 	}
-	UE_LOG(LogElectraPlayer, Log, TEXT("[%p][%p] %s buffering starts"), this, CurrentPlayer.Get(), Electra::Metrics::GetBufferingReasonString(BufferingReason));
+
+	FString Msg = FString::Printf(TEXT("%s buffering starts"), Electra::Metrics::GetBufferingReasonString(BufferingReason));
+	Statistics.AddMessageToHistory(Msg);
+
+	UE_LOG(LogElectraPlayer, Log, TEXT("[%p][%p] %s"), this, CurrentPlayer.Get(), *Msg);
 	CSV_EVENT(ElectraPlayer, TEXT("Buffering starts"));
 }
 
@@ -2105,6 +2117,7 @@ void FElectraPlayer::HandlePlayerEventBufferingEnd(Electra::Metrics::EBufferingR
 		EnqueueAnalyticsEvent(AnalyticEvent);
 	}
 	UE_LOG(LogElectraPlayer, Log, TEXT("[%p][%p] %s buffering ended after %.3fs"), this, CurrentPlayer.Get(), Electra::Metrics::GetBufferingReasonString(BufferingReason), BufferingDuration);
+	Statistics.AddMessageToHistory(TEXT("Buffering ended"));
 // Should we set the state (back?) to something or wait for the following play/pause event to set a new one?
 	Statistics.LastState = "Ready";
 
@@ -2165,6 +2178,11 @@ void FElectraPlayer::HandlePlayerEventSegmentDownload(const Electra::Metrics::FS
 	{
 		UE_LOG(LogElectraPlayer, Log, TEXT("[%p][%p] %s segment download issue (%s): retry:%d, success:%d, aborted:%d, filler:%d"), this, CurrentPlayer.Get(), Electra::Metrics::GetSegmentTypeString(SegmentDownloadStats.SegmentType), *SegmentDownloadStats.FailureReason, SegmentDownloadStats.RetryNumber, SegmentDownloadStats.bWasSuccessful, SegmentDownloadStats.bWasAborted, SegmentDownloadStats.bInsertedFillerData);
 
+		if (SegmentDownloadStats.FailureReason.Len())
+		{
+			Statistics.AddMessageToHistory(FString::Printf(TEXT("%s segment download issue (%s)"), Electra::Metrics::GetSegmentTypeString(SegmentDownloadStats.SegmentType), *SegmentDownloadStats.FailureReason));
+		}
+
 		static const FString kEventNameElectraSegmentIssue(TEXT("Electra.SegmentIssue"));
 		if (Electra::IsAnalyticsEventEnabled(kEventNameElectraSegmentIssue))
 		{
@@ -2329,6 +2347,7 @@ void FElectraPlayer::HandlePlayerEventPlaybackStart()
 	}
 	Statistics.LastState = "Playing";
 	UE_LOG(LogElectraPlayer, Log, TEXT("[%p][%p] Playback started at play position %.3f"), this, CurrentPlayer.Get(), PlayPos);
+	Statistics.AddMessageToHistory(TEXT("Playback started"));
 
 	// Enqueue a "Start" event.
 	static const FString kEventNameElectraStart(TEXT("Electra.Start"));
@@ -2347,6 +2366,7 @@ void FElectraPlayer::HandlePlayerEventPlaybackPaused()
 	Statistics.LastState = "Paused";
 	double PlayPos = CurrentPlayer->AdaptivePlayer->GetPlayPosition().GetAsSeconds();
 	UE_LOG(LogElectraPlayer, Log, TEXT("[%p][%p] Playback paused at play position %.3f"), this, CurrentPlayer.Get(), PlayPos);
+	Statistics.AddMessageToHistory(TEXT("Playback paused"));
 
 	// Enqueue a "Pause" event.
 	static const FString kEventNameElectraPause(TEXT("Electra.Pause"));
@@ -2365,6 +2385,7 @@ void FElectraPlayer::HandlePlayerEventPlaybackResumed()
 	Statistics.LastState = "Playing";
 	double PlayPos = CurrentPlayer->AdaptivePlayer->GetPlayPosition().GetAsSeconds();
 	UE_LOG(LogElectraPlayer, Log, TEXT("[%p][%p] Playback resumed at play position %.3f"), this, CurrentPlayer.Get(), PlayPos);
+	Statistics.AddMessageToHistory(TEXT("Playback resumed"));
 
 	// Enqueue a "Resume" event.
 	static const FString kEventNameElectraResume(TEXT("Electra.Resume"));
@@ -2387,6 +2408,8 @@ void FElectraPlayer::HandlePlayerEventPlaybackEnded()
 	Statistics.LastState = "Ended";
 	Statistics.bDidPlaybackEnd = true;
 	UE_LOG(LogElectraPlayer, Log, TEXT("[%p][%p] Playback reached end at play position %.3f"), this, CurrentPlayer.Get(), PlayPos);
+	Statistics.AddMessageToHistory(TEXT("Playback ended"));
+
 	// Enqueue an "End" event.
 	static const FString kEventNameElectraEnd(TEXT("Electra.End"));
 	if (Electra::IsAnalyticsEventEnabled(kEventNameElectraEnd))
@@ -2425,6 +2448,7 @@ void FElectraPlayer::HandlePlayerEventJumpInPlayPosition(const Electra::FTimeVal
 		Electra::IAdaptiveStreamingPlayer::FLoopState loopState;
 		CurrentPlayer->AdaptivePlayer->GetLoopState(loopState);
 		UE_LOG(LogElectraPlayer, Log, TEXT("[%p][%p] Looping (%d) from %.3f to %.3f"), this, CurrentPlayer.Get(), loopState.Count, FromTime.GetAsSeconds(), ToNewTime.GetAsSeconds());
+		Statistics.AddMessageToHistory(TEXT("Looped"));
 	}
 
 	// Enqueue a "PositionJump" event.
@@ -2451,6 +2475,7 @@ void FElectraPlayer::HandlePlayerEventPlaybackStopped()
 	Statistics.bDidPlaybackEnd = true;
 	// Note: we do not change Statistics.LastState since we want to keep the state the player was in when it got closed.
 	UE_LOG(LogElectraPlayer, Log, TEXT("[%p][%p] Playback stopped. Last play position %.3f"), this, CurrentPlayer.Get(), PlayPos);
+	Statistics.AddMessageToHistory(TEXT("Stopped"));
 
 	// Enqueue a "Stop" event.
 	static const FString kEventNameElectraStop(TEXT("Electra.Stop"));
@@ -2481,6 +2506,13 @@ void FElectraPlayer::HandlePlayerEventError(const FString& ErrorReason)
 	}
 	// Note: we do not change Statistics.LastState to something like 'error' because we want to know the state the player was in when it errored.
 	UE_LOG(LogElectraPlayer, Error, TEXT("[%p][%p] ReportError: \"%s\""), this, CurrentPlayer.Get(), *SanitizeMessage(ErrorReason));
+	Statistics.AddMessageToHistory(FString::Printf(TEXT("Error: %s"), *SanitizeMessage(ErrorReason)));
+	FString MessageHistory;
+	for(auto &msg : Statistics.MessageHistoryBuffer)
+	{
+		MessageHistory.Append(msg);
+		MessageHistory.Append(TEXT("\n"));
+	}
 
 	// Enqueue an "Error" event.
 	static const FString kEventNameElectraError(TEXT("Electra.Error"));
@@ -2489,6 +2521,7 @@ void FElectraPlayer::HandlePlayerEventError(const FString& ErrorReason)
 		TSharedPtr<FAnalyticsEvent> AnalyticEvent = CreateAnalyticsEvent(kEventNameElectraError);
 		AnalyticEvent->ParamArray.Add(FAnalyticsEventAttribute(TEXT("Reason"), *ErrorReason));
 		AnalyticEvent->ParamArray.Add(FAnalyticsEventAttribute(TEXT("LastState"), *Statistics.LastState));
+		AnalyticEvent->ParamArray.Add(FAnalyticsEventAttribute(TEXT("MessageHistory"), MessageHistory));
 		EnqueueAnalyticsEvent(AnalyticEvent);
 	}
 }
@@ -2499,17 +2532,29 @@ void FElectraPlayer::HandlePlayerEventLogMessage(Electra::IInfoLog::ELevel InLog
 	switch(InLogLevel)
 	{
 		case Electra::IInfoLog::ELevel::Error:
+		{
 			UE_LOG(LogElectraPlayer, Error, TEXT("[%p][%p] %s"), this, CurrentPlayer.Get(), *m);
+			FScopeLock Lock(&StatisticsLock);
+			Statistics.AddMessageToHistory(m);
 			break;
+		}
 		case Electra::IInfoLog::ELevel::Warning:
+		{
 			UE_LOG(LogElectraPlayer, Warning, TEXT("[%p][%p] %s"), this, CurrentPlayer.Get(), *m);
+			FScopeLock Lock(&StatisticsLock);
+			Statistics.AddMessageToHistory(m);
 			break;
+		}
 		case Electra::IInfoLog::ELevel::Info:
+		{
 			UE_LOG(LogElectraPlayer, Log, TEXT("[%p][%p] %s"), this, CurrentPlayer.Get(), *m);
 			break;
+		}
 		case Electra::IInfoLog::ELevel::Verbose:
+		{
 			UE_LOG(LogElectraPlayer, Verbose, TEXT("[%p][%p] %s"), this, CurrentPlayer.Get(), *m);
 			break;
+		}
 	}
 }
 
@@ -2670,6 +2715,17 @@ void FElectraPlayer::LogStatistics()
 			Statistics.SubtitlesResponseTime,
 			*Statistics.SubtitlesLastError
 		);
+		
+		if (Statistics.LastError.Len())
+		{
+			FString MessageHistory;
+			for(auto &msg : Statistics.MessageHistoryBuffer)
+			{
+				MessageHistory.Append(msg);
+				MessageHistory.Append(TEXT("\n"));
+			}
+			UE_LOG(LogElectraPlayer, Log, TEXT("Most recent log messages:\n%s"), *MessageHistory);
+		}
 	}
 }
 
@@ -2702,8 +2758,15 @@ void FElectraPlayer::SendAnalyticMetrics(const TSharedPtr<IAnalyticsProviderET>&
 	TArray<FAnalyticsEventAttribute> ParamArray;
 	AddCommonAnalyticsAttributes(ParamArray);
 	StatisticsLock.Lock();
+	FString MessageHistory;
+	for(auto &msg : Statistics.MessageHistoryBuffer)
+	{
+		MessageHistory.Append(msg);
+		MessageHistory.Append(TEXT("\n"));
+	}
 	ParamArray.Add(FAnalyticsEventAttribute(TEXT("URL"), Statistics.InitialURL));
 	ParamArray.Add(FAnalyticsEventAttribute(TEXT("LastState"), Statistics.LastState));
+	ParamArray.Add(FAnalyticsEventAttribute(TEXT("MessageHistory"), MessageHistory));
 	ParamArray.Add(FAnalyticsEventAttribute(TEXT("LastError"), Statistics.LastError));
 	ParamArray.Add(FAnalyticsEventAttribute(TEXT("FinalVideoResolution"), FString::Printf(TEXT("%d*%d"), Statistics.CurrentlyActiveResolutionWidth, Statistics.CurrentlyActiveResolutionHeight)));
 	ParamArray.Add(FAnalyticsEventAttribute(TEXT("TimeElapsedToMasterPlaylist"), Statistics.TimeToLoadMasterPlaylist));
diff --git a/Engine/Plugins/Media/ElectraPlayer/Source/ElectraPlayerRuntime/Private/ElectraPlayer.h b/Engine/Plugins/Media/ElectraPlayer/Source/ElectraPlayerRuntime/Private/ElectraPlayer.h
--- a/Engine/Plugins/Media/ElectraPlayer/Source/ElectraPlayerRuntime/Private/ElectraPlayer.h
+++ b/Engine/Plugins/Media/ElectraPlayer/Source/ElectraPlayerRuntime/Private/ElectraPlayer.h
@@ -669,6 +669,15 @@ private:
 			MediaTimelineAtStart.Reset();
 			MediaTimelineAtEnd.Reset();
 			MediaDuration = 0.0;
+			MessageHistoryBuffer.Empty();
+		}
+		void AddMessageToHistory(FString InMessage)
+		{
+			if (MessageHistoryBuffer.Num() >= 20)
+			{
+				MessageHistoryBuffer.RemoveAt(0);
+			}
+			MessageHistoryBuffer.Emplace(MoveTemp(InMessage));
 		}
 
 		FString					InitialURL;
@@ -711,6 +720,7 @@ private:
 		FTimeRange				MediaTimelineAtStart;
 		FTimeRange				MediaTimelineAtEnd;
 		double					MediaDuration;
+		TArray<FString>			MessageHistoryBuffer;
 	};
 
 	struct FAnalyticsEvent
diff --git a/Engine/Plugins/Media/ElectraPlayer/Source/ElectraPlayerRuntime/Private/Runtime/Player/AdaptiveStreamingPlayerResourceRequest.h b/Engine/Plugins/Media/ElectraPlayer/Source/ElectraPlayerRuntime/Private/Runtime/Player/AdaptiveStreamingPlayerResourceRequest.h
--- a/Engine/Plugins/Media/ElectraPlayer/Source/ElectraPlayerRuntime/Private/Runtime/Player/AdaptiveStreamingPlayerResourceRequest.h
+++ b/Engine/Plugins/Media/ElectraPlayer/Source/ElectraPlayerRuntime/Private/Runtime/Player/AdaptiveStreamingPlayerResourceRequest.h
@@ -114,6 +114,7 @@ public:
 		else if (Error == 1)					return FString(TEXT("Connection timeout"));
 		else if (Error == 2)					return FString(TEXT("Data timeout"));
 		else if (Error == 3)					return FString(TEXT("Connection closed"));
+		else if (Error == 4)					return FString(TEXT("No connection"));
 		else if (Error >= 100 && Error < 600)	return FString::Printf(TEXT("HTTP status code %d"), Error);
 		else									return FString::Printf(TEXT("Unknown <code %d>"), Error);
 	}
