ElectraHTTPStream: Added guard around callback delegate use to prevent it being cleared while executing; moved destruction to a task on the thread pool to avoid a possible internal deadlock of the thread trying to delete itself.

#rnx

diff --git a/Engine/Plugins/Media/ElectraPlayer/Source/ElectraPlayerRuntime/Private/Runtime/HTTP/HTTPManager.cpp b/Engine/Plugins/Media/ElectraPlayer/Source/ElectraPlayerRuntime/Private/Runtime/HTTP/HTTPManager.cpp
--- a/Engine/Plugins/Media/ElectraPlayer/Source/ElectraPlayerRuntime/Private/Runtime/HTTP/HTTPManager.cpp
+++ b/Engine/Plugins/Media/ElectraPlayer/Source/ElectraPlayerRuntime/Private/Runtime/HTTP/HTTPManager.cpp
@@ -11,6 +11,9 @@
 // For base64 encoding/decoding
 #include "Misc/Base64.h"
 
+// Asynchronous termination
+#include "Async/Async.h"
+
 // For file:// scheme archive deserialization
 #include "HAL/FileManager.h"
 #include "Misc/FileHelper.h"
@@ -54,6 +57,19 @@ namespace Electra
 		virtual void RemoveRequest(TSharedPtrTS<FRequest> Request, bool bDoNotWaitForRemoval) override;
 
 	private:
+		class TDeleter
+		{
+		public:
+			void operator()(FElectraHttpManager* InInstanceToDelete)
+			{
+				TFunction<void()> DeleteTask = [InInstanceToDelete]()
+				{
+					delete InInstanceToDelete;
+				};
+				Async(EAsyncExecution::ThreadPool, MoveTemp(DeleteTask));
+			}
+		};
+
 		struct FLocalByteStream
 		{
 			virtual ~FLocalByteStream() = default;
@@ -405,7 +421,11 @@ namespace Electra
 			FElectraHttpManager* Manager = new FElectraHttpManager;
 			if (Manager)
 			{
-				Self = MakeShareable(Manager);
+				// A custom deleter is required to ensure that the destruction will not happen within the enclosing worker thread
+				// due to the shared pointer to ourselves passed as a callback delegate in StartHTTPManager() being released.
+				// If that happens we encounter a deadlock in the worker thread trying to delete itself.
+				// To avoid this the custom deleter moves the final destruction out as a task to the thread pool.
+				Self = MakeShareable(Manager, TDeleter());
 				Manager->Initialize();
 			}
 			SingletonSelf = Self;
@@ -939,7 +959,7 @@ namespace Electra
 		HTTPStreamHandler = IElectraHTTPStream::Create(HttpOptions);
 		if (HTTPStreamHandler.IsValid())
 		{
-			HTTPStreamHandler->ThreadHandlerDelegate().BindThreadSafeSP(AsShared(), &FElectraHttpManager::ProcessHTTPManager);
+			HTTPStreamHandler->AddThreadHandlerDelegate(IElectraHTTPStream::FElectraHTTPStreamThreadHandlerDelegate::CreateThreadSafeSP(AsShared(), &FElectraHttpManager::ProcessHTTPManager));
 			return true;
 		}
 		else
@@ -955,7 +975,7 @@ namespace Electra
 
 		if (HTTPStreamHandler.IsValid())
 		{
-			HTTPStreamHandler->ThreadHandlerDelegate().Unbind();
+			HTTPStreamHandler->RemoveThreadHandlerDelegate();
 		}
 
 		RemoveAllRequests();
diff --git a/Engine/Plugins/Media/ElectraUtil/Source/ElectraHTTPStream/Private/Apple/AppleElectraHTTPStream.cpp b/Engine/Plugins/Media/ElectraUtil/Source/ElectraHTTPStream/Private/Apple/AppleElectraHTTPStream.cpp
--- a/Engine/Plugins/Media/ElectraUtil/Source/ElectraHTTPStream/Private/Apple/AppleElectraHTTPStream.cpp
+++ b/Engine/Plugins/Media/ElectraUtil/Source/ElectraHTTPStream/Private/Apple/AppleElectraHTTPStream.cpp
@@ -69,8 +69,16 @@ public:
 	FElectraHTTPStreamApple();
 	bool Initialize(const Electra::FParamDict& InOptions);
 
-	FElectraHTTPStreamThreadHandlerDelegate& ThreadHandlerDelegate() override
-	{ return ThreadHandlerCallback; }
+	void AddThreadHandlerDelegate(FElectraHTTPStreamThreadHandlerDelegate InDelegate) override
+	{
+		FScopeLock lock(&CallbackLock);
+		ThreadHandlerCallback = MoveTemp(InDelegate);
+	}
+	void RemoveThreadHandlerDelegate() override
+	{
+		FScopeLock lock(&CallbackLock);
+		ThreadHandlerCallback.Unbind();
+	}
 
 	void Close() override;
 
@@ -141,6 +149,7 @@ private:
 	FRunnableThread* Thread = nullptr;
 	FTimeWaitableSignal HaveWorkSignal;
 
+	FCriticalSection CallbackLock;
 	FElectraHTTPStreamThreadHandlerDelegate ThreadHandlerCallback;
 
 	FCriticalSection RequestLock;
@@ -1123,6 +1132,7 @@ uint32 FElectraHTTPStreamApple::Run()
 		// User callback
 		{
 		SCOPE_CYCLE_COUNTER(STAT_ElectraHTTPThread_CustomHandler);
+		FScopeLock lock(&CallbackLock);
 		ThreadHandlerCallback.ExecuteIfBound();
 		}
 	}
diff --git a/Engine/Plugins/Media/ElectraUtil/Source/ElectraHTTPStream/Private/Curl/CurlElectraHTTPStream.cpp b/Engine/Plugins/Media/ElectraUtil/Source/ElectraHTTPStream/Private/Curl/CurlElectraHTTPStream.cpp
--- a/Engine/Plugins/Media/ElectraUtil/Source/ElectraHTTPStream/Private/Curl/CurlElectraHTTPStream.cpp
+++ b/Engine/Plugins/Media/ElectraUtil/Source/ElectraHTTPStream/Private/Curl/CurlElectraHTTPStream.cpp
@@ -961,8 +961,6 @@ void FElectraHTTPStreamLibCurl::Close()
 	if (Thread)
 	{
 		Thread->Kill(true);
-		ThreadFinished.Lock();
-		ThreadFinished.Unlock();
 		delete Thread;
 		Thread = nullptr;
 	}
@@ -1065,6 +1063,7 @@ void FElectraHTTPStreamLibCurl::WorkInnerLoop()
 			// User callback
 			{
 			SCOPE_CYCLE_COUNTER(STAT_ElectraHTTPThread_CustomHandler);
+			FScopeLock lock(&CallbackLock);
 			ThreadHandlerCallback.ExecuteIfBound();
 			}
 		}
@@ -1099,6 +1098,7 @@ void FElectraHTTPStreamLibCurl::WorkMultiPoll()
 			// User callback
 			{
 			SCOPE_CYCLE_COUNTER(STAT_ElectraHTTPThread_CustomHandler);
+			FScopeLock lock(&CallbackLock);
 			ThreadHandlerCallback.ExecuteIfBound();
 			}
 		}
@@ -1111,8 +1111,6 @@ uint32 FElectraHTTPStreamLibCurl::Run()
 {
 	LLM_SCOPE(ELLMTag::MediaStreaming);
 
-	ThreadFinished.Lock();
-
 #if ELECTRA_HTTPSTREAM_CURL_USE_MULTIPOLL
 	WorkMultiPoll();
 #else
@@ -1135,7 +1133,6 @@ uint32 FElectraHTTPStreamLibCurl::Run()
 	}
 	RequestLock.Unlock();
 	HandleCompletedRequests();
-	ThreadFinished.Unlock();
 
 	return 0;
 }
diff --git a/Engine/Plugins/Media/ElectraUtil/Source/ElectraHTTPStream/Private/Curl/CurlElectraHTTPStreamImpl.h b/Engine/Plugins/Media/ElectraUtil/Source/ElectraHTTPStream/Private/Curl/CurlElectraHTTPStreamImpl.h
--- a/Engine/Plugins/Media/ElectraUtil/Source/ElectraHTTPStream/Private/Curl/CurlElectraHTTPStreamImpl.h
+++ b/Engine/Plugins/Media/ElectraUtil/Source/ElectraHTTPStream/Private/Curl/CurlElectraHTTPStreamImpl.h
@@ -28,14 +28,22 @@ public:
 	FElectraHTTPStreamLibCurl();
 	virtual bool Initialize(const Electra::FParamDict& InOptions);
 
-	virtual FElectraHTTPStreamThreadHandlerDelegate& ThreadHandlerDelegate() override
-	{ return ThreadHandlerCallback; }
+	void AddThreadHandlerDelegate(FElectraHTTPStreamThreadHandlerDelegate InDelegate) override
+	{
+		FScopeLock lock(&CallbackLock);
+		ThreadHandlerCallback = MoveTemp(InDelegate);
+	}
+	void RemoveThreadHandlerDelegate() override
+	{
+		FScopeLock lock(&CallbackLock);
+		ThreadHandlerCallback.Unbind();
+	}
 
-	virtual void Close() override;
+	void Close() override;
 
-	virtual IElectraHTTPStreamRequestPtr CreateRequest() override;
+	IElectraHTTPStreamRequestPtr CreateRequest() override;
 
-	virtual void AddRequest(IElectraHTTPStreamRequestPtr Request) override;
+	void AddRequest(IElectraHTTPStreamRequestPtr Request) override;
 
 #if ELECTRA_HTTPSTREAM_CURL_USE_SHARE
 	CURLSH* GetCurlShareHandle()
@@ -49,8 +57,8 @@ public:
 	{ return ProxyAddressAndPort; }
 
 	void TriggerWorkSignal()
-	{ 
-		HaveWorkSignal.Signal(); 
+	{
+		HaveWorkSignal.Signal();
 #if ELECTRA_HTTPSTREAM_CURL_USE_MULTIPOLL
 		curl_multi_wakeup(CurlMultiHandle);
 #endif
@@ -58,8 +66,8 @@ public:
 
 private:
 	// Methods from FRunnable
-	virtual uint32 Run() override final;
-	virtual void Stop() override final;
+	uint32 Run() override final;
+	void Stop() override final;
 
 	void SetupNewRequests();
 	void UpdateActiveRequests();
@@ -84,8 +92,8 @@ private:
 	FThreadSafeCounter ExitRequest;
 	FRunnableThread* Thread = nullptr;
 	FTimeWaitableSignal HaveWorkSignal;
-	FCriticalSection ThreadFinished;
 
+	FCriticalSection CallbackLock;
 	FElectraHTTPStreamThreadHandlerDelegate ThreadHandlerCallback;
 
 	FCriticalSection RequestLock;
diff --git a/Engine/Plugins/Media/ElectraUtil/Source/ElectraHTTPStream/Private/WinHttp/WinHttpElectraHTTPStream.cpp b/Engine/Plugins/Media/ElectraUtil/Source/ElectraHTTPStream/Private/WinHttp/WinHttpElectraHTTPStream.cpp
--- a/Engine/Plugins/Media/ElectraUtil/Source/ElectraHTTPStream/Private/WinHttp/WinHttpElectraHTTPStream.cpp
+++ b/Engine/Plugins/Media/ElectraUtil/Source/ElectraHTTPStream/Private/WinHttp/WinHttpElectraHTTPStream.cpp
@@ -90,8 +90,17 @@ public:
 	static uint64 GetNewRequestIndex()
 	{ return ++NextRequestIndex; }
 
-	FElectraHTTPStreamThreadHandlerDelegate& ThreadHandlerDelegate() override
-	{ return ThreadHandlerCallback; }
+	void AddThreadHandlerDelegate(FElectraHTTPStreamThreadHandlerDelegate InDelegate) override
+	{
+		FScopeLock lock(&CallbackLock);
+		ThreadHandlerCallback = MoveTemp(InDelegate);
+	}
+	void RemoveThreadHandlerDelegate() override
+	{
+		FScopeLock lock(&CallbackLock);
+		ThreadHandlerCallback.Unbind();
+	}
+
 
 	void Close() override;
 
@@ -127,8 +136,8 @@ private:
 	};
 
 	// Methods from FRunnable
-	virtual uint32 Run() override final;
-	virtual void Stop() override final;
+	uint32 Run() override final;
+	void Stop() override final;
 
 
 	static FRequestPointers GetRequestByIndex(uint64 InRequestIndex)
@@ -183,6 +192,7 @@ private:
 	FRunnableThread* Thread = nullptr;
 	FTimeWaitableSignal HaveWorkSignal;
 
+	FCriticalSection CallbackLock;
 	FElectraHTTPStreamThreadHandlerDelegate ThreadHandlerCallback;
 
 	FCriticalSection RequestLock;
@@ -1361,6 +1371,7 @@ uint32 FElectraHTTPStreamWinHttp::Run()
 		// User callback
 		{
 		SCOPE_CYCLE_COUNTER(STAT_ElectraHTTPThread_CustomHandler);
+		FScopeLock lock(&CallbackLock);
 		ThreadHandlerCallback.ExecuteIfBound();
 		}
 	}
diff --git a/Engine/Plugins/Media/ElectraUtil/Source/ElectraHTTPStream/Public/ElectraHTTPStream.h b/Engine/Plugins/Media/ElectraUtil/Source/ElectraHTTPStream/Public/ElectraHTTPStream.h
--- a/Engine/Plugins/Media/ElectraUtil/Source/ElectraHTTPStream/Public/ElectraHTTPStream.h
+++ b/Engine/Plugins/Media/ElectraUtil/Source/ElectraHTTPStream/Public/ElectraHTTPStream.h
@@ -233,7 +233,8 @@ public:
 	 * The NotificationDelegate() of the IElectraHTTPStreamRequestPtr request object should be all you need.
 	 */
 	DECLARE_DELEGATE(FElectraHTTPStreamThreadHandlerDelegate);
-	virtual FElectraHTTPStreamThreadHandlerDelegate& ThreadHandlerDelegate() = 0;
+	virtual void AddThreadHandlerDelegate(FElectraHTTPStreamThreadHandlerDelegate InDelegate) = 0;
+	virtual void RemoveThreadHandlerDelegate() = 0;
 
 	/**
 	 * Closes this handler and all currently active or pending requests.
