ElectraPlayer: Added audio decoder render throttling

#rnx

diff --git a/Engine/Plugins/Media/ElectraPlayer/Source/ElectraPlayerRuntime/Private/Runtime/Decoder/DX/AudioDecoderAAC_DX.cpp b/Engine/Plugins/Media/ElectraPlayer/Source/ElectraPlayerRuntime/Private/Runtime/Decoder/DX/AudioDecoderAAC_DX.cpp
--- a/Engine/Plugins/Media/ElectraPlayer/Source/ElectraPlayerRuntime/Private/Runtime/Decoder/DX/AudioDecoderAAC_DX.cpp
+++ b/Engine/Plugins/Media/ElectraPlayer/Source/ElectraPlayerRuntime/Private/Runtime/Decoder/DX/AudioDecoderAAC_DX.cpp
@@ -925,7 +925,15 @@ namespace Electra
 							return false;
 						}
 					}
+
 					bool bHaveAvailSmpBlk = CurrentRenderOutputBuffer != nullptr;
+					// Check if the renderer can accept the output we will want to send to it.
+					// If it can't right now we treat this as if we do not have an available output buffer.
+					if (Renderer.IsValid() && !Renderer->CanReceiveOutputFrames(1))
+					{
+						bHaveAvailSmpBlk = false;
+					}
+
 					NotifyReadyBufferListener(bHaveAvailSmpBlk);
 					if (bHaveAvailSmpBlk)
 					{
