ElectraPlayer: Sequence index resets to zero on Seek(); player flush in facade is now ignored since it happens async and may empty the sink while a BlockOnRange is waiting for that data; sink flush is now performed on user triggered seek with Electra making sure to not deliver any new samples until the player sample queues have been cleared to prevent spilling of stale data from the queues to the sinks.

#rnx

diff --git a/Engine/Plugins/Media/ElectraPlayer/Source/ElectraPlayerPlugin/Private/ElectraPlayerPlugin.h b/Engine/Plugins/Media/ElectraPlayer/Source/ElectraPlayerPlugin/Private/ElectraPlayerPlugin.h
--- a/Engine/Plugins/Media/ElectraPlayer/Source/ElectraPlayerPlugin/Private/ElectraPlayerPlugin.h
+++ b/Engine/Plugins/Media/ElectraPlayer/Source/ElectraPlayerPlugin/Private/ElectraPlayerPlugin.h
@@ -99,7 +99,7 @@ public:
 	void SetLastAudioRenderedSampleTime(FTimespan SampleTime) override;
 	bool FlushOnSeekStarted() const override
 	{
-		return false;
+		return true;
 	}
 	bool FlushOnSeekCompleted() const override
 	{
diff --git a/Engine/Plugins/Media/ElectraPlayer/Source/ElectraPlayerRuntime/Private/ElectraPlayer.cpp b/Engine/Plugins/Media/ElectraPlayer/Source/ElectraPlayerRuntime/Private/ElectraPlayer.cpp
--- a/Engine/Plugins/Media/ElectraPlayer/Source/ElectraPlayerRuntime/Private/ElectraPlayer.cpp
+++ b/Engine/Plugins/Media/ElectraPlayer/Source/ElectraPlayerRuntime/Private/ElectraPlayer.cpp
@@ -206,6 +206,7 @@ void FElectraPlayer::ClearToDefaultState()
 	bAudioTrackIndexDirty = true;
 	bSubtitleTrackIndexDirty = true;
 	bInitialSeekPerformed = false;
+	bDiscardOutputUntilCleanStart = false;
 	LastPresentedFrameDimension = FIntPoint::ZeroValue;
 	CurrentlyActiveVideoStreamFormat.Reset();
 	DeferredPlayerEvents.Empty();
@@ -765,8 +766,15 @@ bool FElectraPlayer::PresentVideoFrame(const FVideoDecoderOutputPtr& InVideoFram
 	TSharedPtr<IElectraPlayerAdapterDelegate, ESPMode::ThreadSafe> PinnedAdapterDelegate = AdapterDelegate.Pin();
 	if (PinnedAdapterDelegate.IsValid())
 	{
-		PinnedAdapterDelegate->PresentVideoFrame(InVideoFrame);
-		LastPresentedFrameDimension = InVideoFrame->GetOutputDim();
+		if (!bDiscardOutputUntilCleanStart)
+		{
+			PinnedAdapterDelegate->PresentVideoFrame(InVideoFrame);
+			LastPresentedFrameDimension = InVideoFrame->GetOutputDim();
+		}
+		else
+		{
+			UE_LOG(LogElectraPlayer, Log, TEXT("[%p][%p] Dropping video frame @ %.6f;%d due to seek block"), this, CurrentPlayer.Get(), InVideoFrame->GetTime().Time.GetTotalSeconds(), (int)InVideoFrame->GetTime().SequenceIndex);
+		}
 	}
 	return true;
 }
@@ -782,7 +790,14 @@ bool FElectraPlayer::PresentAudioFrame(const IAudioDecoderOutputPtr& DecoderOutp
 	TSharedPtr<IElectraPlayerAdapterDelegate, ESPMode::ThreadSafe> PinnedAdapterDelegate = AdapterDelegate.Pin();
 	if (PinnedAdapterDelegate.IsValid())
 	{
-		PinnedAdapterDelegate->PresentAudioFrame(DecoderOutput);
+		if (!bDiscardOutputUntilCleanStart)
+		{
+			PinnedAdapterDelegate->PresentAudioFrame(DecoderOutput);
+		}
+		else
+		{
+			UE_LOG(LogElectraPlayer, Log, TEXT("[%p][%p] Dropping audio frame @ %.6f;%d due to seek block"), this, CurrentPlayer.Get(), DecoderOutput->GetTime().Time.GetTotalSeconds(), (int)DecoderOutput->GetTime().SequenceIndex);
+		}
 	}
 	return true;
 }
@@ -793,7 +808,10 @@ bool FElectraPlayer::PresentSubtitle(const ISubtitleDecoderOutputPtr& DecoderOut
 	TSharedPtr<IElectraPlayerAdapterDelegate, ESPMode::ThreadSafe> PinnedAdapterDelegate = AdapterDelegate.Pin();
 	if (PinnedAdapterDelegate.IsValid())
 	{
-		PinnedAdapterDelegate->PresentSubtitleSample(DecoderOutput);
+		if (!bDiscardOutputUntilCleanStart)
+		{
+			PinnedAdapterDelegate->PresentSubtitleSample(DecoderOutput);
+		}
 	}
 	return true;
 }
@@ -1108,6 +1126,7 @@ bool FElectraPlayer::Seek(const FTimespan& Time, const FSeekParam& Param)
 		seek.bOptimizeForScrubbing = Param.bOptimizeForScrubbing;
 		seek.DistanceThreshold = Param.DistanceThreshold;
 		bInitialSeekPerformed = true;
+		bDiscardOutputUntilCleanStart = true;
 		CurrentPlayer->AdaptivePlayer->SeekTo(seek);
 		return true;
 	}
@@ -1677,6 +1696,11 @@ void FElectraPlayer::HandleDeferredPlayerEvents()
 				HandlePlayerEventPlaylistDownload(Ev->PlaylistDownloadStats);
 				break;
 			}
+			case FPlayerMetricEventBase::EType::CleanStart:
+			{
+				bDiscardOutputUntilCleanStart = false;
+				break;
+			}
 			case FPlayerMetricEventBase::EType::BufferingStart:
 			{
 				FPlayerMetricEvent_BufferingStart* Ev = static_cast<FPlayerMetricEvent_BufferingStart*>(Event.Get());
@@ -2127,7 +2151,7 @@ void FElectraPlayer::HandlePlayerEventBufferingEnd(Electra::Metrics::EBufferingR
 void FElectraPlayer::HandlePlayerEventBandwidth(int64 EffectiveBps, int64 ThroughputBps, double LatencyInSeconds)
 {
 //	FScopeLock Lock(&StatisticsLock);
-	UE_LOG(LogElectraPlayer, VeryVerbose, TEXT("[%p][%p] Observed bandwidth of %d bps; throughput = %d; latency = %.3fs"), this, CurrentPlayer.Get(), (int32)EffectiveBps, (int32)ThroughputBps, LatencyInSeconds);
+	UE_LOG(LogElectraPlayer, VeryVerbose, TEXT("[%p][%p] Observed bandwidth of %lld Kbps; throughput = %lld Kbps; latency = %.3fs"), this, CurrentPlayer.Get(), EffectiveBps/1000, ThroughputBps/1000, LatencyInSeconds);
 }
 
 void FElectraPlayer::HandlePlayerEventBufferUtilization(const Electra::Metrics::FBufferStats& BufferStats)
@@ -2293,6 +2317,7 @@ void FElectraPlayer::HandlePlayerEventCodecFormatChange(const Electra::FStreamCo
 
 void FElectraPlayer::HandlePlayerEventPrerollStart()
 {
+	bDiscardOutputUntilCleanStart = false;
 	// Update statistics
 	FScopeLock Lock(&StatisticsLock);
 	Statistics.TimeAtPrerollBegin = FPlatformTime::Seconds();
@@ -2490,6 +2515,7 @@ void FElectraPlayer::HandlePlayerEventPlaybackStopped()
 void FElectraPlayer::HandlePlayerEventSeekCompleted()
 {
 	UE_LOG(LogElectraPlayer, Log, TEXT("[%p][%p] Seek completed"), this, CurrentPlayer.Get());
+	bDiscardOutputUntilCleanStart = false;
 	MediaStateOnSeekFinished();
 }
 
diff --git a/Engine/Plugins/Media/ElectraPlayer/Source/ElectraPlayerRuntime/Private/ElectraPlayer.h b/Engine/Plugins/Media/ElectraPlayer/Source/ElectraPlayerRuntime/Private/ElectraPlayer.h
--- a/Engine/Plugins/Media/ElectraPlayer/Source/ElectraPlayerRuntime/Private/ElectraPlayer.h
+++ b/Engine/Plugins/Media/ElectraPlayer/Source/ElectraPlayerRuntime/Private/ElectraPlayer.h
@@ -183,6 +183,7 @@ private:
 			ReceivedPlaylists,
 			TracksChanged,
 			PlaylistDownload,
+			CleanStart,
 			BufferingStart,
 			BufferingEnd,
 			Bandwidth,
@@ -317,6 +318,10 @@ private:
 	{ DeferredPlayerEvents.Enqueue(MakeSharedTS<FPlayerMetricEventBase>(FPlayerMetricEventBase::EType::TracksChanged)); }
 	virtual void ReportPlaylistDownload(const Metrics::FPlaylistDownloadStats& PlaylistDownloadStats) override
 	{ DeferredPlayerEvents.Enqueue(MakeSharedTS<FPlayerMetricEvent_PlaylistDownload>(PlaylistDownloadStats)); }
+	virtual void ReportCleanStart() override
+	{ /*DeferredPlayerEvents.Enqueue(MakeSharedTS<FPlayerMetricEventBase>(FPlayerMetricEventBase::EType::CleanStart));*/ 
+		bDiscardOutputUntilCleanStart = false;
+	}
 	virtual void ReportBufferingStart(Metrics::EBufferingReason BufferingReason) override
 	{ DeferredPlayerEvents.Enqueue(MakeSharedTS<FPlayerMetricEvent_BufferingStart>(BufferingReason)); }
 	virtual void ReportBufferingEnd(Metrics::EBufferingReason BufferingReason) override
@@ -391,6 +396,7 @@ private:
 	mutable bool									bSubtitleTrackIndexDirty;
 
 	bool											bInitialSeekPerformed;
+	bool											bDiscardOutputUntilCleanStart;
 
 	FPlaybackRange									CurrentPlaybackRange;
 	TOptional<bool>									bFrameAccurateSeeking;
diff --git a/Engine/Plugins/Media/ElectraPlayer/Source/ElectraPlayerRuntime/Private/Runtime/Player/AdaptiveStreamingPlayer.cpp b/Engine/Plugins/Media/ElectraPlayer/Source/ElectraPlayerRuntime/Private/Runtime/Player/AdaptiveStreamingPlayer.cpp
--- a/Engine/Plugins/Media/ElectraPlayer/Source/ElectraPlayerRuntime/Private/Runtime/Player/AdaptiveStreamingPlayer.cpp
+++ b/Engine/Plugins/Media/ElectraPlayer/Source/ElectraPlayerRuntime/Private/Runtime/Player/AdaptiveStreamingPlayer.cpp
@@ -523,6 +523,101 @@ void FAdaptiveStreamingPlayer::DispatchEvent(TSharedPtrTS<FMetricEvent> Event)
 	}
 }
 
+void FAdaptiveStreamingPlayer::FireSyncEvent(TSharedPtrTS<FMetricEvent> Event)
+{
+	LockMetricsReceivers();
+	const TArray<IAdaptiveStreamingPlayerMetrics*, TInlineAllocator<4>>& Listeners = GetMetricsReceivers();
+	for(int32 i=0; i<Listeners.Num(); ++i)
+	{
+		switch(Event->Type)
+		{
+			case FMetricEvent::EType::OpenSource:
+				Listeners[i]->ReportOpenSource(Event->Param.URL);
+				break;
+			case FMetricEvent::EType::ReceivedMasterPlaylist:
+				Listeners[i]->ReportReceivedMasterPlaylist(Event->Param.URL);
+				break;
+			case FMetricEvent::EType::ReceivedPlaylists:
+				Listeners[i]->ReportReceivedPlaylists();
+				break;
+			case FMetricEvent::EType::TracksChanged:
+				Listeners[i]->ReportTracksChanged();
+				break;
+			case FMetricEvent::EType::CleanStart:
+				Listeners[i]->ReportCleanStart();
+				break;
+			case FMetricEvent::EType::BufferingStart:
+				Listeners[i]->ReportBufferingStart(Event->Param.BufferingReason);
+				break;
+			case FMetricEvent::EType::BufferingEnd:
+				Listeners[i]->ReportBufferingEnd(Event->Param.BufferingReason);
+				break;
+			case FMetricEvent::EType::Bandwidth:
+				Listeners[i]->ReportBandwidth(Event->Param.Bandwidth.EffectiveBps, Event->Param.Bandwidth.ThroughputBps, Event->Param.Bandwidth.Latency);
+				break;
+			case FMetricEvent::EType::BufferUtilization:
+				Listeners[i]->ReportBufferUtilization(Event->Param.BufferStats);
+				break;
+			case FMetricEvent::EType::PlaylistDownload:
+				Listeners[i]->ReportPlaylistDownload(Event->Param.PlaylistStats);
+				break;
+			case FMetricEvent::EType::SegmentDownload:
+				Listeners[i]->ReportSegmentDownload(Event->Param.SegmentStats);
+				break;
+			case FMetricEvent::EType::LicenseKey:
+				Listeners[i]->ReportLicenseKey(Event->Param.LicenseKeyStats);
+				break;
+			case FMetricEvent::EType::DataAvailabilityChange:
+				Listeners[i]->ReportDataAvailabilityChange(Event->Param.DataAvailability);
+				break;
+			case FMetricEvent::EType::VideoQualityChange:
+				Listeners[i]->ReportVideoQualityChange(Event->Param.QualityChange.NewBitrate, Event->Param.QualityChange.PrevBitrate, Event->Param.QualityChange.bIsDrastic);
+				break;
+			case FMetricEvent::EType::CodecFormatChange:
+				Listeners[i]->ReportDecodingFormatChange(Event->Param.CodecFormatChange);
+				break;
+			case FMetricEvent::EType::PrerollStart:
+				Listeners[i]->ReportPrerollStart();
+				break;
+			case FMetricEvent::EType::PrerollEnd:
+				Listeners[i]->ReportPrerollEnd();
+				break;
+			case FMetricEvent::EType::PlaybackStart:
+				Listeners[i]->ReportPlaybackStart();
+				break;
+			case FMetricEvent::EType::PlaybackPaused:
+				Listeners[i]->ReportPlaybackPaused();
+				break;
+			case FMetricEvent::EType::PlaybackResumed:
+				Listeners[i]->ReportPlaybackResumed();
+				break;
+			case FMetricEvent::EType::PlaybackEnded:
+				Listeners[i]->ReportPlaybackEnded();
+				break;
+			case FMetricEvent::EType::PlaybackJumped:
+				Listeners[i]->ReportJumpInPlayPosition(Event->Param.TimeJump.ToNewTime, Event->Param.TimeJump.FromTime, Event->Param.TimeJump.Reason);
+				break;
+			case FMetricEvent::EType::PlaybackStopped:
+				Listeners[i]->ReportPlaybackStopped();
+				break;
+			case FMetricEvent::EType::SeekCompleted:
+				Listeners[i]->ReportSeekCompleted();
+				break;
+			case FMetricEvent::EType::Errored:
+				Listeners[i]->ReportError(Event->Param.ErrorDetail.GetPrintable());
+				break;
+			case FMetricEvent::EType::LogMessage:
+				Listeners[i]->ReportLogMessage(Event->Param.LogMessage.Level, Event->Param.LogMessage.Message, Event->Param.LogMessage.AtMillis);
+				break;
+		}
+	}
+	UnlockMetricsReceivers();
+	if (Event->EventSignal)
+	{
+		Event->EventSignal->Signal();
+	}
+}
+
 
 //-----------------------------------------------------------------------------
 /**
@@ -1673,7 +1768,7 @@ void FAdaptiveStreamingPlayer::HandleSeeking()
 		}
 
 		// Check the distance to the last seek performed, if there is one.
-		if (SeekVars.LastFinishedRequest.IsSet())
+		if (SeekVars.LastFinishedRequest.IsSet() && SeekVars.PendingRequest.GetValue().DistanceThreshold.IsSet())
 		{
 			double Distance = Utils::AbsoluteValue(SeekVars.LastFinishedRequest.GetValue().Time.GetAsSeconds() - SeekVars.PendingRequest.GetValue().Time.GetAsSeconds());
 			if (Distance <= SeekVars.PendingRequest.GetValue().DistanceThreshold.Get(0.0))
@@ -1692,6 +1787,15 @@ void FAdaptiveStreamingPlayer::HandleSeeking()
 		SeekVars.PendingRequest.Reset();
 		SeekVars.ClearWorkVars();
 
+		// On a user-induced seek the sequence index is supposed to reset to 0.
+		CurrentPlaybackSequenceState.SequenceIndex = 0;
+		// And since it is a seek on purpose the loop counter is reset as well.
+		FInternalLoopState LoopStateNow;
+		PlaybackState.GetLoopState(LoopStateNow);
+		LoopStateNow.Count = 0;
+		PlaybackState.SetLoopState(LoopStateNow);
+		CurrentLoopState.Count = 0;
+
 		if (SeekVars.bIsPlayStart)
 		{
 			CurrentState = EPlayerState::eState_Buffering;
@@ -2640,6 +2744,12 @@ void FAdaptiveStreamingPlayer::InternalHandlePendingStartRequest(const FTimeValu
 							IManifest::FResult Result;
 							if (PendingStartRequest->StartType != FPendingStartRequest::EStartType::LoopPoint)
 							{
+								// Get the current loop state to fetch the last published loop count.
+								// The internal counter in CurrentLoopState may be some iterations ahead since it could have fetched segments from the loop point already.
+								FInternalLoopState LoopStateNow;
+								PlaybackState.GetLoopState(LoopStateNow);
+								CurrentLoopState.Count = LoopStateNow.Count;
+
 								Result = InitialPlayPeriod->GetStartingSegment(FirstSegmentRequest, CurrentPlaybackSequenceState, StartAt, PendingStartRequest->SearchType);
 								bFirstSegmentRequestIsForLooping = false;
 							}
@@ -4137,6 +4247,14 @@ void FAdaptiveStreamingPlayer::InternalStartAt(const FSeekParam& NewPosition)
 
 	// From this point on any further start will not be the very first one any more.
 	SeekVars.bIsPlayStart = false;
+
+	// Starting with a clean slate. Fire this event synchronously directly to the receivers.
+	// This is intended to release any upper layer blocks of passing decoded output to sinks
+	// that may be in place to prevent receiving decoder output that may be "stale" from the
+	// perspective of the upper layer.
+	// It is IMPERATIVE for the receiver of this event to do NO BLOCKING tasks OR call in
+	// any of the player API!
+	FireSyncEvent(FMetricEvent::ReportCleanStart());
 }
 
 
@@ -4916,95 +5034,7 @@ void FAdaptiveStreamingPlayerEventHandler::WorkerThreadFN()
 			TSharedPtrTS<FAdaptiveStreamingPlayer>	Player = pEvt->Player.Pin();
 			if (Player.IsValid())
 			{
-				// Call listeners under lock. They are not allowed to add or remove themselves or other listeners.
-				Player->LockMetricsReceivers();
-				const TArray<IAdaptiveStreamingPlayerMetrics*, TInlineAllocator<4>>& MetricListeners = Player->GetMetricsReceivers();
-				for(int32 i=0; i<MetricListeners.Num(); ++i)
-				{
-					switch(pEvt->Type)
-					{
-						case FMetricEvent::EType::OpenSource:
-							MetricListeners[i]->ReportOpenSource(pEvt->Param.URL);
-							break;
-						case FMetricEvent::EType::ReceivedMasterPlaylist:
-							MetricListeners[i]->ReportReceivedMasterPlaylist(pEvt->Param.URL);
-							break;
-						case FMetricEvent::EType::ReceivedPlaylists:
-							MetricListeners[i]->ReportReceivedPlaylists();
-							break;
-						case FMetricEvent::EType::TracksChanged:
-							MetricListeners[i]->ReportTracksChanged();
-							break;
-						case FMetricEvent::EType::BufferingStart:
-							MetricListeners[i]->ReportBufferingStart(pEvt->Param.BufferingReason);
-							break;
-						case FMetricEvent::EType::BufferingEnd:
-							MetricListeners[i]->ReportBufferingEnd(pEvt->Param.BufferingReason);
-							break;
-						case FMetricEvent::EType::Bandwidth:
-							MetricListeners[i]->ReportBandwidth(pEvt->Param.Bandwidth.EffectiveBps, pEvt->Param.Bandwidth.ThroughputBps, pEvt->Param.Bandwidth.Latency);
-							break;
-						case FMetricEvent::EType::BufferUtilization:
-							MetricListeners[i]->ReportBufferUtilization(pEvt->Param.BufferStats);
-							break;
-						case FMetricEvent::EType::PlaylistDownload:
-							MetricListeners[i]->ReportPlaylistDownload(pEvt->Param.PlaylistStats);
-							break;
-						case FMetricEvent::EType::SegmentDownload:
-							MetricListeners[i]->ReportSegmentDownload(pEvt->Param.SegmentStats);
-							break;
-						case FMetricEvent::EType::LicenseKey:
-							MetricListeners[i]->ReportLicenseKey(pEvt->Param.LicenseKeyStats);
-							break;
-						case FMetricEvent::EType::DataAvailabilityChange:
-							MetricListeners[i]->ReportDataAvailabilityChange(pEvt->Param.DataAvailability);
-							break;
-						case FMetricEvent::EType::VideoQualityChange:
-							MetricListeners[i]->ReportVideoQualityChange(pEvt->Param.QualityChange.NewBitrate, pEvt->Param.QualityChange.PrevBitrate, pEvt->Param.QualityChange.bIsDrastic);
-							break;
-						case FMetricEvent::EType::CodecFormatChange:
-							MetricListeners[i]->ReportDecodingFormatChange(pEvt->Param.CodecFormatChange);
-							break;
-						case FMetricEvent::EType::PrerollStart:
-							MetricListeners[i]->ReportPrerollStart();
-							break;
-						case FMetricEvent::EType::PrerollEnd:
-							MetricListeners[i]->ReportPrerollEnd();
-							break;
-						case FMetricEvent::EType::PlaybackStart:
-							MetricListeners[i]->ReportPlaybackStart();
-							break;
-						case FMetricEvent::EType::PlaybackPaused:
-							MetricListeners[i]->ReportPlaybackPaused();
-							break;
-						case FMetricEvent::EType::PlaybackResumed:
-							MetricListeners[i]->ReportPlaybackResumed();
-							break;
-						case FMetricEvent::EType::PlaybackEnded:
-							MetricListeners[i]->ReportPlaybackEnded();
-							break;
-						case FMetricEvent::EType::PlaybackJumped:
-							MetricListeners[i]->ReportJumpInPlayPosition(pEvt->Param.TimeJump.ToNewTime, pEvt->Param.TimeJump.FromTime, pEvt->Param.TimeJump.Reason);
-							break;
-						case FMetricEvent::EType::PlaybackStopped:
-							MetricListeners[i]->ReportPlaybackStopped();
-							break;
-						case FMetricEvent::EType::SeekCompleted:
-							MetricListeners[i]->ReportSeekCompleted();
-							break;
-						case FMetricEvent::EType::Errored:
-							MetricListeners[i]->ReportError(pEvt->Param.ErrorDetail.GetPrintable());
-							break;
-						case FMetricEvent::EType::LogMessage:
-							MetricListeners[i]->ReportLogMessage(pEvt->Param.LogMessage.Level, pEvt->Param.LogMessage.Message, pEvt->Param.LogMessage.AtMillis);
-							break;
-					}
-				}
-				Player->UnlockMetricsReceivers();
-			}
-			if (pEvt->EventSignal)
-			{
-				pEvt->EventSignal->Signal();
+				Player->FireSyncEvent(pEvt);
 			}
 		}
 	}
diff --git a/Engine/Plugins/Media/ElectraPlayer/Source/ElectraPlayerRuntime/Private/Runtime/Player/AdaptiveStreamingPlayerABR.cpp b/Engine/Plugins/Media/ElectraPlayer/Source/ElectraPlayerRuntime/Private/Runtime/Player/AdaptiveStreamingPlayerABR.cpp
--- a/Engine/Plugins/Media/ElectraPlayer/Source/ElectraPlayerRuntime/Private/Runtime/Player/AdaptiveStreamingPlayerABR.cpp
+++ b/Engine/Plugins/Media/ElectraPlayer/Source/ElectraPlayerRuntime/Private/Runtime/Player/AdaptiveStreamingPlayerABR.cpp
@@ -52,6 +52,7 @@ namespace Electra
 		void ReportReceivedMasterPlaylist(const FString& EffectiveURL) override {}
 		void ReportReceivedPlaylists() override {}
 		void ReportTracksChanged() override {}
+		void ReportCleanStart() override {};
 		void ReportPlaylistDownload(const Metrics::FPlaylistDownloadStats& PlaylistDownloadStats) override {}
 		void ReportBandwidth(int64 EffectiveBps, int64 ThroughputBps, double LatencyInSeconds) override {}
 		void ReportBufferUtilization(const Metrics::FBufferStats& BufferStats) override {}
diff --git a/Engine/Plugins/Media/ElectraPlayer/Source/ElectraPlayerRuntime/Private/Runtime/Player/AdaptiveStreamingPlayerInternal.h b/Engine/Plugins/Media/ElectraPlayer/Source/ElectraPlayerRuntime/Private/Runtime/Player/AdaptiveStreamingPlayerInternal.h
--- a/Engine/Plugins/Media/ElectraPlayer/Source/ElectraPlayerRuntime/Private/Runtime/Player/AdaptiveStreamingPlayerInternal.h
+++ b/Engine/Plugins/Media/ElectraPlayer/Source/ElectraPlayerRuntime/Private/Runtime/Player/AdaptiveStreamingPlayerInternal.h
@@ -491,6 +491,7 @@ struct FMetricEvent
 		ReceivedMasterPlaylist,
 		ReceivedPlaylists,
 		TracksChanged,
+		CleanStart,
 		BufferingStart,
 		BufferingEnd,
 		Bandwidth,
@@ -566,6 +567,12 @@ struct FMetricEvent
 	TWeakPtr<FAdaptiveStreamingPlayer, ESPMode::ThreadSafe>	Player;
 	FMediaEvent	*EventSignal = nullptr;
 
+	static TSharedPtrTS<FMetricEvent> ReportCleanStart()
+	{
+		TSharedPtrTS<FMetricEvent> Evt = MakeSharedTS<FMetricEvent>();
+		Evt->Type = EType::CleanStart;
+		return Evt;
+	}
 	static TSharedPtrTS<FMetricEvent> ReportOpenSource(const FString& URL)
 	{
 		TSharedPtrTS<FMetricEvent> Evt = MakeSharedTS<FMetricEvent>();
@@ -850,6 +857,7 @@ public:
 	{
 		return MetricListeners;
 	}
+	void FireSyncEvent(TSharedPtrTS<FMetricEvent> Event);
 
 
 	void AddAEMSReceiver(TWeakPtrTS<IAdaptiveStreamingPlayerAEMSReceiver> InReceiver, FString InForSchemeIdUri, FString InForValue, IAdaptiveStreamingPlayerAEMSReceiver::EDispatchMode InDispatchMode) override;
diff --git a/Engine/Plugins/Media/ElectraPlayer/Source/ElectraPlayerRuntime/Private/Runtime/Player/AdaptiveStreamingPlayerMetrics.h b/Engine/Plugins/Media/ElectraPlayer/Source/ElectraPlayerRuntime/Private/Runtime/Player/AdaptiveStreamingPlayerMetrics.h
--- a/Engine/Plugins/Media/ElectraPlayer/Source/ElectraPlayerRuntime/Private/Runtime/Player/AdaptiveStreamingPlayerMetrics.h
+++ b/Engine/Plugins/Media/ElectraPlayer/Source/ElectraPlayerRuntime/Private/Runtime/Player/AdaptiveStreamingPlayerMetrics.h
@@ -227,6 +227,11 @@ public:
 	 */
 	virtual void ReportPlaylistDownload(const Metrics::FPlaylistDownloadStats& PlaylistDownloadStats) = 0;
 
+	/**
+	 * Called when the player starts over from a clean state.
+	 */
+	virtual void ReportCleanStart() = 0;
+
 	/**
 	 * Called when buffering of data begins.
 	 */
diff --git a/Engine/Source/Runtime/MediaUtils/Private/MediaPlayerFacade.cpp b/Engine/Source/Runtime/MediaUtils/Private/MediaPlayerFacade.cpp
--- a/Engine/Source/Runtime/MediaUtils/Private/MediaPlayerFacade.cpp
+++ b/Engine/Source/Runtime/MediaUtils/Private/MediaPlayerFacade.cpp
@@ -1770,7 +1770,7 @@ void FMediaPlayerFacade::ProcessEvent(EMediaEvent Event, bool bIsBroadcastAllowe
 	{
 		if (!Player.IsValid() || Player->FlushOnSeekCompleted())
 		{
-			Flush(Player->GetPlayerFeatureFlag(IMediaPlayer::EFeatureFlag::PlayerUsesInternalFlushOnSeek));
+			Flush(!Player.IsValid() || Player->GetPlayerFeatureFlag(IMediaPlayer::EFeatureFlag::PlayerUsesInternalFlushOnSeek));
 		}
 	}
 	else if (Event == EMediaEvent::MediaClosed)
@@ -2730,8 +2730,6 @@ void FMediaPlayerFacade::ProcessSubtitleSamples(IMediaSamples& Samples, TRange<F
 
 void FMediaPlayerFacade::ReceiveMediaEvent(EMediaEvent Event)
 {
-	UE_LOG(LogMediaUtils, VeryVerbose, TEXT("PlayerFacade %p: Received media event %s"), this, *MediaUtils::EventToString(Event));
-
 	if (Event >= EMediaEvent::Internal_Start)
 	{
 		switch (Event)
@@ -2783,8 +2781,11 @@ void FMediaPlayerFacade::ReceiveMediaEvent(EMediaEvent Event)
 
 		case	EMediaEvent::Internal_ResetForDiscontinuity:
 		{
-			UE_LOG(LogMediaUtils, VeryVerbose, TEXT("PlayerFacade %p: Reset for discontinuity"), this);
-			bIsSinkFlushPending = true;
+			// Disabled for now to prevent a flush on the next handling iteration that might discard the samples a blocking range is waiting for.
+			#if 0
+				UE_LOG(LogMediaUtils, VeryVerbose, TEXT("PlayerFacade %p: Reset for discontinuity"), this);
+				bIsSinkFlushPending = true;
+			#endif
 			break;
 		}
 		case	EMediaEvent::Internal_RenderClockStart:
@@ -2824,11 +2825,15 @@ void FMediaPlayerFacade::ReceiveMediaEvent(EMediaEvent Event)
 		}
 
 		default:
+		{
+			UE_LOG(LogMediaUtils, VeryVerbose, TEXT("PlayerFacade %p: Received media event %s"), this, *MediaUtils::EventToString(Event));
 			break;
 		}
+		}
 	}
 	else
 	{
+		UE_LOG(LogMediaUtils, VeryVerbose, TEXT("PlayerFacade %p: Received media event %s"), this, *MediaUtils::EventToString(Event));
 		QueuedEvents.Enqueue(Event);
 	}
 }
