ElectraPlayer: Fix for adjusting the playback range not working while looping and not showing the target frame while scrubbing

#rnx

diff --git a/Engine/Plugins/Media/ElectraPlayer/Source/ElectraPlayerRuntime/Private/Runtime/Player/AdaptivePlayerOptionKeynames.h b/Engine/Plugins/Media/ElectraPlayer/Source/ElectraPlayerRuntime/Private/Runtime/Player/AdaptivePlayerOptionKeynames.h
--- a/Engine/Plugins/Media/ElectraPlayer/Source/ElectraPlayerRuntime/Private/Runtime/Player/AdaptivePlayerOptionKeynames.h
+++ b/Engine/Plugins/Media/ElectraPlayer/Source/ElectraPlayerRuntime/Private/Runtime/Player/AdaptivePlayerOptionKeynames.h
@@ -39,6 +39,9 @@ namespace Electra
 	//! (FTimeValue) absolute end time (excluding) of the range to limit playback to
 	const TCHAR* const OptionPlayRangeEnd = TEXT("play_range_end");
 
+	//! (bool) true to emit the first decoded video frame while prerolling so it can be displayed while scrubbing.
+	const TCHAR* const OptionKeyDoNotHoldBackFirstVideoFrame = TEXT("do_not_hold_back_first_frame");
+	
 
 
 	const TCHAR* const OptionKeyCurrentAvgStartingVideoBitrate = TEXT("current:avg_video_bitrate");
diff --git a/Engine/Plugins/Media/ElectraPlayer/Source/ElectraPlayerRuntime/Private/Runtime/Player/AdaptiveStreamingPlayer.cpp b/Engine/Plugins/Media/ElectraPlayer/Source/ElectraPlayerRuntime/Private/Runtime/Player/AdaptiveStreamingPlayer.cpp
--- a/Engine/Plugins/Media/ElectraPlayer/Source/ElectraPlayerRuntime/Private/Runtime/Player/AdaptiveStreamingPlayer.cpp
+++ b/Engine/Plugins/Media/ElectraPlayer/Source/ElectraPlayerRuntime/Private/Runtime/Player/AdaptiveStreamingPlayer.cpp
@@ -1152,6 +1152,13 @@ int32 FAdaptiveStreamingPlayer::CreateRenderers()
 	// Set the render clock with the renderes.
 	VideoRender.Renderer->SetRenderClock(RenderClock);
 	AudioRender.Renderer->SetRenderClock(RenderClock);
+
+	// Hold back all frames during preroll or emit the first frame for scrubbing?
+	if (PlayerOptions.HaveKey(OptionKeyDoNotHoldBackFirstVideoFrame))
+	{
+		VideoRender.Renderer->DisableHoldbackOfFirstRenderableVideoFrame(PlayerOptions.GetValue(OptionKeyDoNotHoldBackFirstVideoFrame).SafeGetBool(false));
+	}
+
 	return 0;
 }
 
@@ -2549,6 +2556,8 @@ void FAdaptiveStreamingPlayer::InternalHandlePendingStartRequest(const FTimeValu
 							FTimeRange Seekable = Manifest->GetSeekableTimeRange();
 							PendingStartRequest->SearchType = IManifest::ESearchType::Closest;
 							StartAt.Time = Seekable.Start;
+							ClampStartRequestTime(StartAt.Time);
+							PendingStartRequest->StartAt.Time = StartAt.Time;
 						}
 						else
 						{
diff --git a/Engine/Plugins/Media/ElectraPlayer/Source/ElectraPlayerRuntime/Private/Runtime/Player/AdaptiveStreamingPlayerInternal.h b/Engine/Plugins/Media/ElectraPlayer/Source/ElectraPlayerRuntime/Private/Runtime/Player/AdaptiveStreamingPlayerInternal.h
--- a/Engine/Plugins/Media/ElectraPlayer/Source/ElectraPlayerRuntime/Private/Runtime/Player/AdaptiveStreamingPlayerInternal.h
+++ b/Engine/Plugins/Media/ElectraPlayer/Source/ElectraPlayerRuntime/Private/Runtime/Player/AdaptiveStreamingPlayerInternal.h
@@ -50,6 +50,8 @@ public:
 	virtual FTimeValue GetEnqueuedSampleDuration() = 0;
 	virtual int32 GetNumEnqueuedSamples(FTimeValue* OutOptionalDuration) = 0;
 
+	virtual void DisableHoldbackOfFirstRenderableVideoFrame(bool bDisableHoldback) = 0;
+
 	virtual FTimeRange GetSupportedRenderRateScale() = 0;
 	virtual void SetPlayRateScale(double InNewScale) = 0;
 	virtual double GetPlayRateScale() = 0;
diff --git a/Engine/Plugins/Media/ElectraPlayer/Source/ElectraPlayerRuntime/Private/Runtime/Player/AdaptiveStreamingPlayerRenderers.cpp b/Engine/Plugins/Media/ElectraPlayer/Source/ElectraPlayerRuntime/Private/Runtime/Player/AdaptiveStreamingPlayerRenderers.cpp
--- a/Engine/Plugins/Media/ElectraPlayer/Source/ElectraPlayerRuntime/Private/Runtime/Player/AdaptiveStreamingPlayerRenderers.cpp
+++ b/Engine/Plugins/Media/ElectraPlayer/Source/ElectraPlayerRuntime/Private/Runtime/Player/AdaptiveStreamingPlayerRenderers.cpp
@@ -81,6 +81,8 @@ namespace Electra
 		FTimeValue GetEnqueuedSampleDuration() override;
 		int32 GetNumEnqueuedSamples(FTimeValue* OutOptionalDuration) override;
 
+		void DisableHoldbackOfFirstRenderableVideoFrame(bool bDisableHoldback) override;
+
 		FTimeRange GetSupportedRenderRateScale() override;
 		void SetPlayRateScale(double InNewScale) override;
 		double GetPlayRateScale() override;
@@ -114,6 +116,8 @@ namespace Electra
 		int64 CurrentValidityValue = 0;
 		EStreamType Type = EStreamType::Unsupported;
 		bool bIsRunning = false;
+		bool bDoNotHoldBackFirstVideoFrame = false;
+		uint32 NumBuffersNotHeldBack = 0;
 
 		TQueue<FPendingReturnBuffer> PendingReturnBuffers;
 		std::atomic<int32> NumPendingReturnBuffers;
@@ -425,7 +429,23 @@ UEMediaError FAdaptiveStreamingWrappedRenderer::ReturnBufferCommon(IBuffer* Buff
 	FScopeLock lock(&Lock);
 	EnqueuedDuration += Duration;
 	++NumEnqueuedSamples;
-	if (!bIsRunning)
+
+	bool bHoldback = !bIsRunning;
+	// If the video renderer shall not hold back the first frame (used for scrubbing video)
+	// then we pass it out. The count is reset in Flush().
+	if (Type == EStreamType::Video && bDoNotHoldBackFirstVideoFrame)
+	{
+		if (NumBuffersNotHeldBack == 0)
+		{
+			bHoldback = false;
+		}
+		if (!bHoldback)
+		{
+			++NumBuffersNotHeldBack;
+		}
+	}
+
+	if (bHoldback)
 	{
 		FPendingReturnBuffer pb;
 		pb.Buffer = Buffer;
@@ -484,6 +504,7 @@ UEMediaError FAdaptiveStreamingWrappedRenderer::Flush(const FParamDict& InOption
 	ReturnAllPendingBuffers(true);
 	++CurrentValidityValue;
 	NumEnqueuedSamples = 0;
+	NumBuffersNotHeldBack = 0;
 	EnqueuedDuration.SetToZero();
 	Lock.Unlock();
 
@@ -546,6 +567,12 @@ int32 FAdaptiveStreamingWrappedRenderer::GetNumEnqueuedSamples(FTimeValue* OutOp
 	return NumEnqueuedSamples;
 }
 
+void FAdaptiveStreamingWrappedRenderer::DisableHoldbackOfFirstRenderableVideoFrame(bool bInDisableHoldback)
+{
+	FScopeLock lock(&Lock);
+	bDoNotHoldBackFirstVideoFrame = bInDisableHoldback;
+}
+
 FTimeRange FAdaptiveStreamingWrappedRenderer::GetSupportedRenderRateScale()
 {
 	FScopeLock lock(&Lock);
