ElectraPlayer: Correcting backslashes to forward slashes and spaces in file: scheme in case they are accidentally used.

#rnx

diff --git a/Engine/Plugins/Media/ElectraPlayer/Source/ElectraPlayerRuntime/Private/Runtime/Player/AdaptiveStreamingPlayerManifest.cpp b/Engine/Plugins/Media/ElectraPlayer/Source/ElectraPlayerRuntime/Private/Runtime/Player/AdaptiveStreamingPlayerManifest.cpp
--- a/Engine/Plugins/Media/ElectraPlayer/Source/ElectraPlayerRuntime/Private/Runtime/Player/AdaptiveStreamingPlayerManifest.cpp
+++ b/Engine/Plugins/Media/ElectraPlayer/Source/ElectraPlayerRuntime/Private/Runtime/Player/AdaptiveStreamingPlayerManifest.cpp
@@ -24,6 +24,21 @@ static const FString MIMETypeHLS(TEXT("application/vnd.apple.mpegURL"));
 static const FString MIMETypeDASH(TEXT("application/dash+xml"));
 
 
+/**
+ * Fix common mistakes in a file:// URL like backslashes instead of forward slashes
+ * and spaces that are not escaped.
+ */
+FString FixLocalFileSchemeURL(const FString& InURL)
+{
+	FString URL(InURL);
+	if (URL.StartsWith("file:", ESearchCase::IgnoreCase))
+	{
+		URL.ReplaceCharInline(TCHAR('\\'), TCHAR('/'));
+		URL.ReplaceInline(TEXT(" "), TEXT("%20"));
+	}
+	return URL;
+}
+
 /**
  * Returns the MIME type of the media playlist pointed to by the given URL.
  * This only parses the URL for known extensions or scheme.
@@ -33,7 +48,7 @@ static const FString MIMETypeDASH(TEXT("application/dash+xml"));
  *
  * @return MIME type or empty string if it cannot be determined
  */
-FString GetMIMETypeForURL(const FString &URL)
+FString GetMIMETypeForURL(const FString& URL)
 {
 	FString MimeType;
 	FURL_RFC3986 UrlParser;
@@ -139,8 +154,10 @@ void FAdaptiveStreamingPlayer::InternalCloseManifestReader()
  * @param URL
  * @param MimeType
  */
-void FAdaptiveStreamingPlayer::InternalLoadManifest(const FString& URL, const FString& MimeType)
+void FAdaptiveStreamingPlayer::InternalLoadManifest(const FString& InURL, const FString& MimeType)
 {
+	// Fix potential errors in a local file URL.
+	FString URL = Playlist::FixLocalFileSchemeURL(InURL);
 	// Remember the original request URL since we may lose the fragment part in requests.
 	ManifestURL = URL;
 	ManifestMimeTypeRequest.Reset();
