Updated ElectraPlayerRuntime to use a shared PCH.

Before: Total CPU Time: 92 s
After: Total CPU Time: 89.625 s

diff --git a/Engine/Plugins/Media/ElectraPlayer/Source/ElectraPlayerRuntime/ElectraPlayerRuntime.Build.cs b/Engine/Plugins/Media/ElectraPlayer/Source/ElectraPlayerRuntime/ElectraPlayerRuntime.Build.cs
--- a/Engine/Plugins/Media/ElectraPlayer/Source/ElectraPlayerRuntime/ElectraPlayerRuntime.Build.cs
+++ b/Engine/Plugins/Media/ElectraPlayer/Source/ElectraPlayerRuntime/ElectraPlayerRuntime.Build.cs
@@ -13,7 +13,8 @@ namespace UnrealBuildTool.Rules
 			//
 
 			bLegalToDistributeObjectCode = true;
-			PCHUsage = PCHUsageMode.NoPCHs;
+			PCHUsage = PCHUsageMode.UseSharedPCHs;
+			PrivatePCHHeaderFile = "Private/ElectraPlayerPrivate.h";
 
 			PrivateDependencyModuleNames.AddRange(
 				new string[] {
