Fixed compile issue when ELECTRA_ENABLE_SWDECODE is not defined.

diff --git a/Engine/Plugins/Media/ElectraPlayer/Source/ElectraPlayerRuntime/Private/Runtime/Decoder/DX/VideoDecoderH264_DX.cpp b/Engine/Plugins/Media/ElectraPlayer/Source/ElectraPlayerRuntime/Private/Runtime/Decoder/DX/VideoDecoderH264_DX.cpp
--- a/Engine/Plugins/Media/ElectraPlayer/Source/ElectraPlayerRuntime/Private/Runtime/Decoder/DX/VideoDecoderH264_DX.cpp
+++ b/Engine/Plugins/Media/ElectraPlayer/Source/ElectraPlayerRuntime/Private/Runtime/Decoder/DX/VideoDecoderH264_DX.cpp
@@ -45,7 +45,7 @@ namespace Electra {
 
 IVideoDecoderH264::FSystemConfiguration			FVideoDecoderH264::SystemConfig;
 
-#if ELECTRA_ENABLE_SWDECODE
+#ifdef ELECTRA_ENABLE_SWDECODE
 bool FVideoDecoderH264::bDidCheckHWSupport = false;
 bool FVideoDecoderH264::bIsHWSupported = false;
 #endif
