Updated a few modules not use the shared PCHs with a private header. This hides issues that customers will see if PCHs are disabled.

diff --git a/Engine/Plugins/Media/ElectraCDM/Source/ElectraCDM/ElectraCDM.Build.cs b/Engine/Plugins/Media/ElectraCDM/Source/ElectraCDM/ElectraCDM.Build.cs
--- a/Engine/Plugins/Media/ElectraCDM/Source/ElectraCDM/ElectraCDM.Build.cs
+++ b/Engine/Plugins/Media/ElectraCDM/Source/ElectraCDM/ElectraCDM.Build.cs
@@ -6,8 +6,7 @@ namespace UnrealBuildTool.Rules
 	{
 		public ElectraCDM(ReadOnlyTargetRules Target) : base(Target)
 		{
-			PCHUsage = PCHUsageMode.UseSharedPCHs;
-			PrivatePCHHeaderFile = "Public/ElectraCDM.h";
+			bEnforceIWYU = false; // Disabled because of third party code
 			PublicDependencyModuleNames.AddRange(
 				new string[]
 				{
diff --git a/Engine/Plugins/Media/ElectraPlayer/Source/ElectraPlayerRuntime/ElectraPlayerRuntime.Build.cs b/Engine/Plugins/Media/ElectraPlayer/Source/ElectraPlayerRuntime/ElectraPlayerRuntime.Build.cs
--- a/Engine/Plugins/Media/ElectraPlayer/Source/ElectraPlayerRuntime/ElectraPlayerRuntime.Build.cs
+++ b/Engine/Plugins/Media/ElectraPlayer/Source/ElectraPlayerRuntime/ElectraPlayerRuntime.Build.cs
@@ -13,8 +13,7 @@ namespace UnrealBuildTool.Rules
 			//
 
 			bLegalToDistributeObjectCode = true;
-			PCHUsage = PCHUsageMode.UseSharedPCHs;
-			PrivatePCHHeaderFile = "Private/ElectraPlayerPrivate.h";
+			bEnforceIWYU = false;
 
 			PrivateDependencyModuleNames.AddRange(
 				new string[] {
diff --git a/Engine/Plugins/Media/ElectraSubtitles/Source/ElectraSubtitles/ElectraSubtitles.Build.cs b/Engine/Plugins/Media/ElectraSubtitles/Source/ElectraSubtitles/ElectraSubtitles.Build.cs
--- a/Engine/Plugins/Media/ElectraSubtitles/Source/ElectraSubtitles/ElectraSubtitles.Build.cs
+++ b/Engine/Plugins/Media/ElectraSubtitles/Source/ElectraSubtitles/ElectraSubtitles.Build.cs
@@ -6,9 +6,6 @@ namespace UnrealBuildTool.Rules
 	{
 		public ElectraSubtitles(ReadOnlyTargetRules Target) : base(Target)
 		{
-			PCHUsage = PCHUsageMode.UseSharedPCHs;
-			PrivatePCHHeaderFile = "Public/ElectraSubtitleModule.h";
-
 			PrivateDependencyModuleNames.AddRange(
 				new string[]
 				{
