ElectraPlayer: Fix for some Android decoders returning unexpected EOS buffers

#rnx

diff --git a/Engine/Plugins/Media/ElectraPlayer/Source/ElectraPlayerRuntime/Private/Runtime/Decoder/Android/AudioDecoderAAC_Android.cpp b/Engine/Plugins/Media/ElectraPlayer/Source/ElectraPlayerRuntime/Private/Runtime/Decoder/Android/AudioDecoderAAC_Android.cpp
--- a/Engine/Plugins/Media/ElectraPlayer/Source/ElectraPlayerRuntime/Private/Runtime/Decoder/Android/AudioDecoderAAC_Android.cpp
+++ b/Engine/Plugins/Media/ElectraPlayer/Source/ElectraPlayerRuntime/Private/Runtime/Decoder/Android/AudioDecoderAAC_Android.cpp
@@ -178,6 +178,7 @@ private:
 
 	TSharedPtr<IAndroidJavaAACAudioDecoder, ESPMode::ThreadSafe>			DecoderInstance;
 	IAndroidJavaAACAudioDecoder::FOutputFormatInfo							CurrentOutputFormatInfo;
+	bool																	bIsOutputFormatInfoValid = false;
 	EDecodingState															DecodingState = EDecodingState::Regular;
 	int64																	LastPushedPresentationTimeUs = 0;
 	bool																	bGotEOS = false;
@@ -578,6 +579,7 @@ bool FAudioDecoderAAC::InternalDecoderCreate()
 	InternalDecoderDestroy();
 
 	DecoderInstance = IAndroidJavaAACAudioDecoder::Create();
+	bIsOutputFormatInfoValid = false;
 	int32 result = DecoderInstance->InitializeDecoder(*ConfigRecord);
 	if (result)
 	{
@@ -767,13 +769,41 @@ FAudioDecoderAAC::EOutputResult FAudioDecoderAAC::GetOutput()
 
 	if (OutputBufferInfo.BufferIndex >= 0)
 	{
-		int32 SamplingRate = CurrentOutputFormatInfo.SampleRate;
-		int32 NumberOfChannels = CurrentOutputFormatInfo.NumChannels;
-		int32 OutputByteCount = OutputBufferInfo.Size;
-		int32 nBytesPerSample = NumberOfChannels * sizeof(int16);
-		int32 nSamplesProduced = OutputByteCount / nBytesPerSample;
-		bool bEOS = OutputBufferInfo.bIsEOS;
+		// Is this an empty EOS buffer?
+		if (OutputBufferInfo.bIsEOS && OutputBufferInfo.Size == 0)
+		{
+			void* CopyBuffer = (void*)PCMBuffer;
+			Result = DecoderInstance->GetOutputBufferAndRelease(CopyBuffer, PCMBufferSize, OutputBufferInfo);
+			LogMessage(IInfoLog::ELevel::Info, FString::Printf(TEXT("Got an EOS buffer, release buffer returned %d"), Result));
+			return EOutputResult::EOS;
+		}
+
+		// Check if we got a MediaCodec_INFO_OUTPUT_FORMAT_CHANGED with valid format information.
+		if (!bIsOutputFormatInfoValid)
+		{
+			IAndroidJavaAACAudioDecoder::FOutputFormatInfo OutputFormatInfo;
+			LogMessage(IInfoLog::ELevel::Info, TEXT("Got a buffer without a preceeding format info message, querying current buffer for properties."));
+			Result = DecoderInstance->GetOutputFormatInfo(OutputFormatInfo, OutputBufferInfo.BufferIndex);
+			if (Result == 0)
+			{
+				CurrentOutputFormatInfo = OutputFormatInfo;
+				bIsOutputFormatInfoValid = CurrentOutputFormatInfo.NumChannels > 0 && CurrentOutputFormatInfo.SampleRate > 0;
+				LogMessage(IInfoLog::ELevel::Info, FString::Printf(TEXT("Info from current buffer: %d channels @ %d Hz; size=%d, eos=%d, cfg=%d"), CurrentOutputFormatInfo.NumChannels, CurrentOutputFormatInfo.SampleRate, OutputBufferInfo.Size, OutputBufferInfo.bIsEOS, OutputBufferInfo.bIsConfig));
+			}
+			else
+			{
+				LogMessage(IInfoLog::ELevel::Info, TEXT("Failed to get properties from current buffer."));
+				if (ConfigRecord.IsValid())
+				{
+					CurrentOutputFormatInfo.NumChannels = ConfigRecord->PSSignal > 0 ? 2 : NumChannelsForConfig[ConfigRecord->ChannelConfiguration];
+					CurrentOutputFormatInfo.SampleRate = ConfigRecord->ExtSamplingFrequency ? ConfigRecord->ExtSamplingFrequency : ConfigRecord->SamplingRate;
+					bIsOutputFormatInfoValid = CurrentOutputFormatInfo.NumChannels > 0 && CurrentOutputFormatInfo.SampleRate > 0;
+					LogMessage(IInfoLog::ELevel::Info, FString::Printf(TEXT("Info from config record: %d channels @ %d Hz; sbr=%d, ps=%d; size=%d, eos=%d, cfg=%d"), CurrentOutputFormatInfo.NumChannels, CurrentOutputFormatInfo.SampleRate, ConfigRecord->SBRSignal, ConfigRecord->PSSignal, OutputBufferInfo.Size, OutputBufferInfo.bIsEOS, OutputBufferInfo.bIsConfig));
+				}
+			}
+		}
 
+		bool bEOS = OutputBufferInfo.bIsEOS;
 		// Get the frontmost AU info that should correspond to this output.
 		TSharedPtrTS<FDecoderInput> MatchingInput;
 		if (InDecoderInput.Num())
@@ -789,6 +819,12 @@ FAudioDecoderAAC::EOutputResult FAudioDecoderAAC::GetOutput()
 			bEOS = true;
 		}
 
+		int32 SamplingRate = CurrentOutputFormatInfo.SampleRate;
+		int32 NumberOfChannels = CurrentOutputFormatInfo.NumChannels;
+		int32 OutputByteCount = OutputBufferInfo.Size;
+		int32 nBytesPerSample = NumberOfChannels * sizeof(int16);
+		int32 nSamplesProduced = OutputByteCount / nBytesPerSample;
+
 		check(PCMBufferSize >= OutputByteCount);
 		if (PCMBufferSize >= OutputByteCount)
 		{
@@ -891,6 +927,8 @@ FAudioDecoderAAC::EOutputResult FAudioDecoderAAC::GetOutput()
 		if (Result == 0)
 		{
 			CurrentOutputFormatInfo = OutputFormatInfo;
+			bIsOutputFormatInfoValid = CurrentOutputFormatInfo.NumChannels > 0 && CurrentOutputFormatInfo.SampleRate > 0;
+			LogMessage(IInfoLog::ELevel::Info, FString::Printf(TEXT("Output format: %d channels @ %d Hz"), CurrentOutputFormatInfo.NumChannels, CurrentOutputFormatInfo.SampleRate));
 		}
 		else
 		{
diff --git a/Engine/Plugins/Media/ElectraPlayer/Source/ElectraPlayerRuntime/Private/Runtime/Decoder/Android/Java/src/com/epicgames/unreal/ElectraAudioDecoderAAC.java b/Engine/Plugins/Media/ElectraPlayer/Source/ElectraPlayerRuntime/Private/Runtime/Decoder/Android/Java/src/com/epicgames/unreal/ElectraAudioDecoderAAC.java
--- a/Engine/Plugins/Media/ElectraPlayer/Source/ElectraPlayerRuntime/Private/Runtime/Decoder/Android/Java/src/com/epicgames/unreal/ElectraAudioDecoderAAC.java
+++ b/Engine/Plugins/Media/ElectraPlayer/Source/ElectraPlayerRuntime/Private/Runtime/Decoder/Android/Java/src/com/epicgames/unreal/ElectraAudioDecoderAAC.java
@@ -281,9 +281,34 @@ public class ElectraAudioDecoderAAC
 
 	public FOutputFormatInfo GetOutputFormatInfo(int InOutputBufferIndex)
 	{
-		//FOutputFormatInfo OutputFormatInfo = new FOutputFormatInfo();
-		//return OutputFormatInfo;
-		return CurrentOutputFormatInfo;
+		if (bIsInitialized)
+		{
+			if (InOutputBufferIndex < 0)
+			{
+				return CurrentOutputFormatInfo;
+			}
+			else
+			{
+				MediaFormat Format = DecoderHandle.getOutputFormat(InOutputBufferIndex);
+				if (Format != null)
+				{
+					FOutputFormatInfo NewFormatInfo = new FOutputFormatInfo();
+					NewFormatInfo.SampleRate	 = Format.getInteger(MediaFormat.KEY_SAMPLE_RATE);
+					NewFormatInfo.NumChannels    = Format.getInteger(MediaFormat.KEY_CHANNEL_COUNT);
+					NewFormatInfo.BytesPerSample = 2;
+					if (Format.containsKey(MediaFormat.KEY_PCM_ENCODING))
+					{
+						boolean bIsFloat = Format.getInteger(MediaFormat.KEY_PCM_ENCODING) == AudioFormat.ENCODING_PCM_FLOAT;
+						if (bIsFloat)
+						{
+							NewFormatInfo.BytesPerSample = 4;
+						}
+					}
+					return NewFormatInfo;
+				}
+			}
+		}
+		return null;
 	}
 
 
@@ -312,7 +337,7 @@ public class ElectraAudioDecoderAAC
 				else if (OutputBufferInfo.BufferIndex == MediaCodec.INFO_OUTPUT_FORMAT_CHANGED)			// -2
 				{
 					MediaFormat Format = DecoderHandle.getOutputFormat();
-					//GameActivity.Log.debug("Android Audio Decoder: New format " + Format);
+					GameActivity.Log.debug("Android Audio Decoder: New format " + Format);
 
 					// Mwahahaaa... this only works from API level 29 onward...
 					//		boolean bIsFloat = Format.getInteger(MediaFormat.KEY_PCM_ENCODING, AudioFormat.ENCODING_PCM_16BIT) == AudioFormat.ENCODING_PCM_FLOAT;
