ElectraPlayer: Not accepting decoded output instead of discarding it to prevent a tight decoder loop.

#rnx

diff --git a/Engine/Plugins/Media/ElectraPlayer/Source/ElectraPlayerRuntime/Private/ElectraPlayer.cpp b/Engine/Plugins/Media/ElectraPlayer/Source/ElectraPlayerRuntime/Private/ElectraPlayer.cpp
--- a/Engine/Plugins/Media/ElectraPlayer/Source/ElectraPlayerRuntime/Private/ElectraPlayer.cpp
+++ b/Engine/Plugins/Media/ElectraPlayer/Source/ElectraPlayerRuntime/Private/ElectraPlayer.cpp
@@ -766,15 +766,8 @@ bool FElectraPlayer::PresentVideoFrame(const FVideoDecoderOutputPtr& InVideoFram
 	TSharedPtr<IElectraPlayerAdapterDelegate, ESPMode::ThreadSafe> PinnedAdapterDelegate = AdapterDelegate.Pin();
 	if (PinnedAdapterDelegate.IsValid())
 	{
-		if (!bDiscardOutputUntilCleanStart)
-		{
-			PinnedAdapterDelegate->PresentVideoFrame(InVideoFrame);
-			LastPresentedFrameDimension = InVideoFrame->GetOutputDim();
-		}
-		else
-		{
-			UE_LOG(LogElectraPlayer, Log, TEXT("[%p][%p] Dropping video frame @ %.6f;%d due to seek block"), this, CurrentPlayer.Get(), InVideoFrame->GetTime().Time.GetTotalSeconds(), (int)InVideoFrame->GetTime().SequenceIndex);
-		}
+		PinnedAdapterDelegate->PresentVideoFrame(InVideoFrame);
+		LastPresentedFrameDimension = InVideoFrame->GetOutputDim();
 	}
 	return true;
 }
@@ -790,14 +783,7 @@ bool FElectraPlayer::PresentAudioFrame(const IAudioDecoderOutputPtr& DecoderOutp
 	TSharedPtr<IElectraPlayerAdapterDelegate, ESPMode::ThreadSafe> PinnedAdapterDelegate = AdapterDelegate.Pin();
 	if (PinnedAdapterDelegate.IsValid())
 	{
-		if (!bDiscardOutputUntilCleanStart)
-		{
-			PinnedAdapterDelegate->PresentAudioFrame(DecoderOutput);
-		}
-		else
-		{
-			UE_LOG(LogElectraPlayer, Log, TEXT("[%p][%p] Dropping audio frame @ %.6f;%d due to seek block"), this, CurrentPlayer.Get(), DecoderOutput->GetTime().Time.GetTotalSeconds(), (int)DecoderOutput->GetTime().SequenceIndex);
-		}
+		PinnedAdapterDelegate->PresentAudioFrame(DecoderOutput);
 	}
 	return true;
 }
@@ -808,10 +794,7 @@ bool FElectraPlayer::PresentSubtitle(const ISubtitleDecoderOutputPtr& DecoderOut
 	TSharedPtr<IElectraPlayerAdapterDelegate, ESPMode::ThreadSafe> PinnedAdapterDelegate = AdapterDelegate.Pin();
 	if (PinnedAdapterDelegate.IsValid())
 	{
-		if (!bDiscardOutputUntilCleanStart)
-		{
-			PinnedAdapterDelegate->PresentSubtitleSample(DecoderOutput);
-		}
+		PinnedAdapterDelegate->PresentSubtitleSample(DecoderOutput);
 	}
 	return true;
 }
@@ -840,7 +823,10 @@ bool FElectraPlayer::CanPresentVideoFrames(uint64 NumFrames)
 	TSharedPtr<IElectraPlayerAdapterDelegate, ESPMode::ThreadSafe> PinnedAdapterDelegate = AdapterDelegate.Pin();
 	if (PinnedAdapterDelegate.IsValid())
 	{
-		return PinnedAdapterDelegate->CanReceiveVideoSamples(NumFrames);
+		if (!bDiscardOutputUntilCleanStart)
+		{
+			return PinnedAdapterDelegate->CanReceiveVideoSamples(NumFrames);
+		}
 	}
 	return false;
 }
@@ -850,7 +836,10 @@ bool FElectraPlayer::CanPresentAudioFrames(uint64 NumFrames)
 	TSharedPtr<IElectraPlayerAdapterDelegate, ESPMode::ThreadSafe> PinnedAdapterDelegate = AdapterDelegate.Pin();
 	if (PinnedAdapterDelegate.IsValid())
 	{
-		return PinnedAdapterDelegate->CanReceiveAudioSamples(NumFrames);
+		if (!bDiscardOutputUntilCleanStart)
+		{
+			return PinnedAdapterDelegate->CanReceiveAudioSamples(NumFrames);
+		}
 	}
 	return false;
 }
