ElectraPlayer: Always allow emitting of the first decoded and rendarable frame to allow for scrubbing in editor.

#rnx

diff --git a/Engine/Plugins/Media/ElectraPlayer/Source/ElectraPlayerRuntime/Private/Runtime/Player/AdaptiveStreamingPlayerRenderers.cpp b/Engine/Plugins/Media/ElectraPlayer/Source/ElectraPlayerRuntime/Private/Runtime/Player/AdaptiveStreamingPlayerRenderers.cpp
--- a/Engine/Plugins/Media/ElectraPlayer/Source/ElectraPlayerRuntime/Private/Runtime/Player/AdaptiveStreamingPlayerRenderers.cpp
+++ b/Engine/Plugins/Media/ElectraPlayer/Source/ElectraPlayerRuntime/Private/Runtime/Player/AdaptiveStreamingPlayerRenderers.cpp
@@ -145,7 +145,7 @@ namespace Electra
 		int64 CurrentValidityValue = 0;
 		EStreamType Type = EStreamType::Unsupported;
 		bool bIsRunning = false;
-		bool bDoNotHoldBackFirstVideoFrame = false;
+		bool bDoNotHoldBackFirstVideoFrame = true;
 		uint32 NumBuffersNotHeldBack = 0;
 
 		TQueue<FPendingReturnBuffer> PendingReturnBuffers;
@@ -467,31 +467,33 @@ UEMediaError FAdaptiveStreamingWrappedRenderer::ReturnBufferCommon(IBuffer* Buff
 	EnqueuedDuration += Duration;
 	++NumEnqueuedSamples;
 
-	bool bHoldback = !bIsRunning;
-	// If the video renderer shall not hold back the first frame (used for scrubbing video)
-	// then we pass it out. The count is reset in Flush().
-	if (Type == EStreamType::Video && bDoNotHoldBackFirstVideoFrame)
+	if (bRender)
 	{
-		if (NumBuffersNotHeldBack == 0)
+		bool bHoldback = !bIsRunning;
+		// If the video renderer shall not hold back the first frame (used for scrubbing video)
+		// then we pass it out. The count is reset in Flush().
+		if (Type == EStreamType::Video && bDoNotHoldBackFirstVideoFrame)
 		{
-			bHoldback = false;
+			if (NumBuffersNotHeldBack == 0)
+			{
+				bHoldback = false;
+			}
+			if (!bHoldback)
+			{
+				++NumBuffersNotHeldBack;
+			}
 		}
-		if (!bHoldback)
+		if (bHoldback)
 		{
-			++NumBuffersNotHeldBack;
+			FPendingReturnBuffer pb;
+			pb.Buffer = Buffer;
+			pb.bRender = bRender;
+			pb.Properties = InSampleProperties;
+			PendingReturnBuffers.Enqueue(MoveTemp(pb));
+			++NumPendingReturnBuffers;
+			return UEMEDIA_ERROR_OK;
 		}
 	}
-
-	if (bHoldback)
-	{
-		FPendingReturnBuffer pb;
-		pb.Buffer = Buffer;
-		pb.bRender = bRender;
-		pb.Properties = InSampleProperties;
-		PendingReturnBuffers.Enqueue(MoveTemp(pb));
-		++NumPendingReturnBuffers;
-		return UEMEDIA_ERROR_OK;
-	}
 	lock.Unlock();
 	return WrappedRenderer->ReturnBuffer(Buffer, bRender, InSampleProperties);
 }
