iOS15 deprecation fixes

#rnx

diff --git a/Engine/Plugins/Media/ElectraPlayer/Source/ElectraPlayerRuntime/Private/Runtime/Decoder/Apple/VideoDecoderH264_Apple.cpp b/Engine/Plugins/Media/ElectraPlayer/Source/ElectraPlayerRuntime/Private/Runtime/Decoder/Apple/VideoDecoderH264_Apple.cpp
--- a/Engine/Plugins/Media/ElectraPlayer/Source/ElectraPlayerRuntime/Private/Runtime/Decoder/Apple/VideoDecoderH264_Apple.cpp
+++ b/Engine/Plugins/Media/ElectraPlayer/Source/ElectraPlayerRuntime/Private/Runtime/Decoder/Apple/VideoDecoderH264_Apple.cpp
@@ -1705,7 +1705,7 @@ void FVideoDecoderH264::ProcessOutput(bool bFlush)
 					long ax = NextImage.SourceInfo->AspectX;
 					long ay = NextImage.SourceInfo->AspectY;
 					// If there is aspect ratio information on the image itself and it's valid, use that instead.
-					NSDictionary* Dict = (NSDictionary*)CVBufferGetAttachments(ImageBufferRef, kCVAttachmentMode_ShouldPropagate);
+					NSDictionary* Dict = (NSDictionary*)CVBufferCopyAttachments(ImageBufferRef, kCVAttachmentMode_ShouldPropagate);
 					if (Dict)
 					{
 						NSDictionary* AspectDict = (NSDictionary*)Dict[(__bridge NSString*)kCVImageBufferPixelAspectRatioKey];
diff --git a/Engine/Plugins/Media/ElectraPlayer/Source/ElectraPlayerRuntime/Private/Runtime/Decoder/Apple/VideoDecoderH265_Apple.cpp b/Engine/Plugins/Media/ElectraPlayer/Source/ElectraPlayerRuntime/Private/Runtime/Decoder/Apple/VideoDecoderH265_Apple.cpp
--- a/Engine/Plugins/Media/ElectraPlayer/Source/ElectraPlayerRuntime/Private/Runtime/Decoder/Apple/VideoDecoderH265_Apple.cpp
+++ b/Engine/Plugins/Media/ElectraPlayer/Source/ElectraPlayerRuntime/Private/Runtime/Decoder/Apple/VideoDecoderH265_Apple.cpp
@@ -1687,7 +1687,7 @@ void FVideoDecoderH265::ProcessOutput(bool bFlush)
 					long ax = NextImage.SourceInfo->AspectX;
 					long ay = NextImage.SourceInfo->AspectY;
 					// If there is aspect ratio information on the image itself and it's valid, use that instead.
-					NSDictionary* Dict = (NSDictionary*)CVBufferGetAttachments(ImageBufferRef, kCVAttachmentMode_ShouldPropagate);
+					NSDictionary* Dict = (NSDictionary*)CVBufferCopyAttachments(ImageBufferRef, kCVAttachmentMode_ShouldPropagate);
 					if (Dict)
 					{
 						NSDictionary* AspectDict = (NSDictionary*)Dict[(__bridge NSString*)kCVImageBufferPixelAspectRatioKey];
diff --git a/Engine/Plugins/Media/ElectraUtil/Source/ElectraHTTPStream/Private/Apple/AppleElectraHTTPStream.cpp b/Engine/Plugins/Media/ElectraUtil/Source/ElectraHTTPStream/Private/Apple/AppleElectraHTTPStream.cpp
--- a/Engine/Plugins/Media/ElectraUtil/Source/ElectraHTTPStream/Private/Apple/AppleElectraHTTPStream.cpp
+++ b/Engine/Plugins/Media/ElectraUtil/Source/ElectraHTTPStream/Private/Apple/AppleElectraHTTPStream.cpp
@@ -1315,9 +1315,9 @@ FElectraHTTPStreamApple::FChallengeResponse FElectraHTTPStreamApple::ReceivedCha
 			}
 			
 			CFIndex CertificateCount = CFArrayGetCount(Certificates);
-			for (int i = 0; CertificateCount; ++i)
+			for (int i = 0; i < CertificateCount; ++i)
 			{
-				SecCertificateRef Cert = SecTrustGetCertificateAtIndex(RemoteTrust, i);
+				SecCertificateRef Cert = (SecCertificateRef)CFArrayGetValueAtIndex(Certificates, i);
 
 				// this is not great, but the only way to extract a public key from a SecCertificateRef
 				// is to create an individual SecTrustRef for each cert that only contains itself and then
