ElectraPlayer: Use currently set playback range on a Seek() to avoid issues with the seek running asynchronously and a new range being set in the interim.

#rnx

diff --git a/Engine/Plugins/Media/ElectraPlayer/Source/ElectraPlayerRuntime/Private/Runtime/Player/AdaptivePlayerOptionKeynames.h b/Engine/Plugins/Media/ElectraPlayer/Source/ElectraPlayerRuntime/Private/Runtime/Player/AdaptivePlayerOptionKeynames.h
--- a/Engine/Plugins/Media/ElectraPlayer/Source/ElectraPlayerRuntime/Private/Runtime/Player/AdaptivePlayerOptionKeynames.h
+++ b/Engine/Plugins/Media/ElectraPlayer/Source/ElectraPlayerRuntime/Private/Runtime/Player/AdaptivePlayerOptionKeynames.h
@@ -33,12 +33,6 @@ namespace Electra
 	//! (bool) true to optimize seeking for faster frame scrubbing, false to optimize for playback.
 	const TCHAR* const OptionKeyFrameOptimizeSeekForScrubbing = TEXT("optimize_seek_for_scrubbing");
 	
-	//! (FTimeValue) absolute start time (including) of the range to limit playback to
-	const TCHAR* const OptionPlayRangeStart = TEXT("play_range_start");
-
-	//! (FTimeValue) absolute end time (excluding) of the range to limit playback to
-	const TCHAR* const OptionPlayRangeEnd = TEXT("play_range_end");
-
 	//! (bool) true to emit the first decoded video frame while prerolling so it can be displayed while scrubbing.
 	const TCHAR* const OptionKeyDoNotHoldBackFirstVideoFrame = TEXT("do_not_hold_back_first_frame");
 	
diff --git a/Engine/Plugins/Media/ElectraPlayer/Source/ElectraPlayerRuntime/Private/Runtime/Player/AdaptiveStreamingPlayer.cpp b/Engine/Plugins/Media/ElectraPlayer/Source/ElectraPlayerRuntime/Private/Runtime/Player/AdaptiveStreamingPlayer.cpp
--- a/Engine/Plugins/Media/ElectraPlayer/Source/ElectraPlayerRuntime/Private/Runtime/Player/AdaptiveStreamingPlayer.cpp
+++ b/Engine/Plugins/Media/ElectraPlayer/Source/ElectraPlayerRuntime/Private/Runtime/Player/AdaptiveStreamingPlayer.cpp
@@ -671,31 +671,7 @@ void FAdaptiveStreamingPlayer::EnableFrameAccurateSeeking(bool bEnabled)
  */
 void FAdaptiveStreamingPlayer::SetPlaybackRange(const FPlaybackRange& InPlaybackRange)
 {
-	FTimeValue RangeStartNow = GetOptions().GetValue(OptionPlayRangeStart).SafeGetTimeValue(FTimeValue());
-	FTimeValue RangeEndNow = GetOptions().GetValue(OptionPlayRangeEnd).SafeGetTimeValue(FTimeValue());
-	FTimeValue NewRangeStart = InPlaybackRange.Start.IsSet() ? InPlaybackRange.Start.GetValue() : FTimeValue();
-	FTimeValue NewRangeEnd = InPlaybackRange.End.IsSet() ? InPlaybackRange.End.GetValue() : FTimeValue();
-	if (NewRangeStart != RangeStartNow || NewRangeEnd != RangeEndNow)
-	{
-		if (NewRangeStart.IsValid())
-		{
-			GetOptions().SetOrUpdate(OptionPlayRangeStart, FVariantValue(NewRangeStart));
-		}
-		else
-		{
-			GetOptions().Remove(OptionPlayRangeStart);
-		}
-		if (NewRangeEnd.IsValid())
-		{
-			GetOptions().SetOrUpdate(OptionPlayRangeEnd, FVariantValue(NewRangeEnd));
-		}
-		else
-		{
-			GetOptions().Remove(OptionPlayRangeEnd);
-		}
-		// Mark the range as dirty so the player can take appropriate action.
-		PlaybackState.SetPlayRangeHasChanged(true);
-	}
+	PlaybackState.SetPlayRange(InPlaybackRange);
 }
 
 
@@ -705,24 +681,7 @@ void FAdaptiveStreamingPlayer::SetPlaybackRange(const FPlaybackRange& InPlayback
  */
 void FAdaptiveStreamingPlayer::GetPlaybackRange(FPlaybackRange& OutPlaybackRange)
 {
-	OutPlaybackRange.Start.Reset();
-	OutPlaybackRange.End.Reset();
-	if (GetOptions().HaveKey(OptionPlayRangeStart))
-	{
-		FTimeValue RangeStart = GetOptions().GetValue(OptionPlayRangeStart).SafeGetTimeValue(FTimeValue());
-		if (RangeStart.IsValid())
-		{
-			OutPlaybackRange.Start = RangeStart;
-		}
-	}
-	if (GetOptions().HaveKey(OptionPlayRangeEnd))
-	{
-		FTimeValue RangeEnd = GetOptions().GetValue(OptionPlayRangeEnd).SafeGetTimeValue(FTimeValue());
-		if (RangeEnd.IsValid())
-		{
-			OutPlaybackRange.End = RangeEnd;
-		}
-	}
+	PlaybackState.GetPlayRange(OutPlaybackRange);
 }
 
 
@@ -766,8 +725,10 @@ void FAdaptiveStreamingPlayer::Resume()
  */
 void FAdaptiveStreamingPlayer::SeekTo(const FSeekParam& NewPosition)
 {
+	FTimeRange NewPlayRange = PlaybackState.GetPlayRange();
 	FScopeLock lock(&SeekVars.Lock);
 	SeekVars.PendingRequest = NewPosition;
+	SeekVars.PlayrangeOnRequest = NewPlayRange;
 	WorkerThread.TriggerSharedWorkerThread();
 }
 
@@ -930,6 +891,14 @@ bool FAdaptiveStreamingPlayer::IsPaused() const
 void FAdaptiveStreamingPlayer::GetLoopState(FLoopState& OutLoopState) const
 {
 	PlaybackState.GetLoopState(OutLoopState);
+
+	// Check if there is a pending seek. A seek resets the loop count, but the seek itself may not be executing
+	// and the count not being reset yet.
+	FScopeLock lock(&SeekVars.Lock);
+	if (SeekVars.PendingRequest.IsSet())
+	{
+		OutLoopState.Count = 0;
+	}
 }
 
 //-----------------------------------------------------------------------------
@@ -1774,6 +1743,18 @@ void FAdaptiveStreamingPlayer::HandleSeeking()
 			return;
 		}
 
+		// Not playing on the Live edge. This may get set to true later on handling the start segment.
+		PlaybackState.SetShouldPlayOnLiveEdge(false);
+		// On a user-induced seek the sequence index is supposed to reset to 0.
+		// We need to reset this even when we would not actually issue the seek.
+		CurrentPlaybackSequenceState.SequenceIndex = 0;
+		// And since it is a seek on purpose the loop counter is reset as well.
+		FInternalLoopState LoopStateNow;
+		PlaybackState.GetLoopState(LoopStateNow);
+		LoopStateNow.Count = 0;
+		PlaybackState.SetLoopState(LoopStateNow);
+		CurrentLoopState.Count = 0;
+
 		// Check the distance to the last seek performed, if there is one.
 		if (SeekVars.LastFinishedRequest.IsSet() && SeekVars.PendingRequest.GetValue().DistanceThreshold.IsSet())
 		{
@@ -1787,22 +1768,12 @@ void FAdaptiveStreamingPlayer::HandleSeeking()
 			}
 		}
 
-		// Not playing on the Live edge. This may get set to true later on handling the start segment.
-		PlaybackState.SetShouldPlayOnLiveEdge(false);
 		// Trigger the seek.
 		FSeekParam SeekParam = SeekVars.PendingRequest.GetValue();
+		FTimeRange SeekPlayrangeOnRequest = SeekVars.PlayrangeOnRequest;
 		SeekVars.PendingRequest.Reset();
 		SeekVars.ClearWorkVars();
 
-		// On a user-induced seek the sequence index is supposed to reset to 0.
-		CurrentPlaybackSequenceState.SequenceIndex = 0;
-		// And since it is a seek on purpose the loop counter is reset as well.
-		FInternalLoopState LoopStateNow;
-		PlaybackState.GetLoopState(LoopStateNow);
-		LoopStateNow.Count = 0;
-		PlaybackState.SetLoopState(LoopStateNow);
-		CurrentLoopState.Count = 0;
-
 		if (SeekVars.bIsPlayStart)
 		{
 			CurrentState = EPlayerState::eState_Buffering;
@@ -1820,7 +1791,7 @@ void FAdaptiveStreamingPlayer::HandleSeeking()
 			SeekVars.bForScrubbing = bIsForScrubbing;
 		}
 		SeekVars.ActiveRequest = SeekParam;
-		InternalStartAt(SeekVars.ActiveRequest.GetValue());
+		InternalStartAt(SeekVars.ActiveRequest.GetValue(), &SeekPlayrangeOnRequest);
 	}
 }
 
@@ -2854,6 +2825,17 @@ void FAdaptiveStreamingPlayer::InternalHandlePendingStartRequest(const FTimeValu
 								case IManifest::FResult::EType::BeforeStart:
 								case IManifest::FResult::EType::PastEOS:
 								{
+									PostLog(Facility::EFacility::Player, IInfoLog::ELevel::Info, FString::Printf(TEXT("miss(%d,a=%d,l=%d) for %lld/%d in range %lld/%d - %lld/%d"), (int32)Result.GetType(), StartAt.Options.bFrameAccuracy, PendingStartRequest->StartType == FPendingStartRequest::EStartType::LoopPoint,
+											(long long int) StartAt.Time.GetAsHNS(), StartAt.Time.IsValid(), 
+											(long long int) StartAt.Options.PlaybackRange.Start.GetAsHNS(),
+											StartAt.Options.PlaybackRange.Start.IsValid(),
+											(long long int) StartAt.Options.PlaybackRange.End.GetAsHNS(),
+											StartAt.Options.PlaybackRange.End.IsValid()));
+									FTimeRange Seekable = Manifest->GetSeekableTimeRange();
+									PostLog(Facility::EFacility::Player, IInfoLog::ELevel::Info, FString::Printf(TEXT("timeline is %lld/%d - %lld/%d"),
+											(long long int) Seekable.Start.GetAsHNS(), Seekable.Start.IsValid(),
+											(long long int) Seekable.End.GetAsHNS(), Seekable.End.IsValid()));
+
 									FErrorDetail err;
 									err.SetFacility(Facility::EFacility::Player);
 									err.SetMessage("Could not locate the stream segment to begin playback at.");
@@ -3572,7 +3554,7 @@ void FAdaptiveStreamingPlayer::InternalStartoverAtCurrentPosition()
 	CurrentState = EPlayerState::eState_Seeking;
 	// Get the current loop state back from the playback state to keep the loop counter as it is now.
 	PlaybackState.GetLoopState(CurrentLoopState);
-	InternalStartAt(NewPosition);
+	InternalStartAt(NewPosition, nullptr);
 }
 
 void FAdaptiveStreamingPlayer::InternalStartoverAtLiveEdge()
@@ -3581,7 +3563,7 @@ void FAdaptiveStreamingPlayer::InternalStartoverAtLiveEdge()
 	InternalStop(PlayerConfig.bHoldLastFrameDuringSeek);
 	CurrentState = EPlayerState::eState_Seeking;
 	PlaybackState.GetLoopState(CurrentLoopState);
-	InternalStartAt(NewPosition);
+	InternalStartAt(NewPosition, nullptr);
 }
 
 
@@ -3868,11 +3850,10 @@ TSharedPtrTS<FBufferSourceInfo> FAdaptiveStreamingPlayer::GetStreamBufferInfoAtT
 
 void FAdaptiveStreamingPlayer::SetPlaystartOptions(FPlayStartOptions& OutOptions)
 {
-	FTimeValue RangeStart = GetOptions().GetValue(OptionPlayRangeStart).SafeGetTimeValue(FTimeValue());
-	FTimeValue RangeEnd = GetOptions().GetValue(OptionPlayRangeEnd).SafeGetTimeValue(FTimeValue());
+	FTimeRange ActiveRange = PlaybackState.GetActivePlayRange();
 
-	OutOptions.PlaybackRange.Start = RangeStart.IsValid() ? RangeStart : FTimeValue::GetZero();
-	OutOptions.PlaybackRange.End = RangeEnd.IsValid() ? RangeEnd : FTimeValue::GetPositiveInfinity();
+	OutOptions.PlaybackRange.Start = ActiveRange.Start.IsValid() ? ActiveRange.Start : FTimeValue::GetZero();
+	OutOptions.PlaybackRange.End = ActiveRange.End.IsValid() ? ActiveRange.End : FTimeValue::GetPositiveInfinity();
 
 	OutOptions.bFrameAccuracy = GetOptions().GetValue(OptionKeyFrameAccurateSeek).SafeGetBool(false);
 }
@@ -3905,14 +3886,13 @@ FTimeValue FAdaptiveStreamingPlayer::ClampTimeToCurrentRange(const FTimeValue& I
 {
 	FTimeValue Out(InTime);
 
-	FTimeRange TimelineRange;
+	FTimeRange TimelineRange, ActiveRange = PlaybackState.GetActivePlayRange();
 	PlaybackState.GetTimelineRange(TimelineRange);
 	if (bClampToEnd)
 	{
-		FTimeValue RangeEnd = GetOptions().GetValue(OptionPlayRangeEnd).SafeGetTimeValue(FTimeValue());
-		if (RangeEnd.IsValid() && RangeEnd < TimelineRange.End)
+		if (ActiveRange.End.IsValid() && ActiveRange.End < TimelineRange.End)
 		{
-			TimelineRange.End = RangeEnd;
+			TimelineRange.End = ActiveRange.End;
 		}
 		if (Out > TimelineRange.End)
 		{
@@ -3921,10 +3901,9 @@ FTimeValue FAdaptiveStreamingPlayer::ClampTimeToCurrentRange(const FTimeValue& I
 	}
 	if (bClampToStart)
 	{
-		FTimeValue RangeStart = GetOptions().GetValue(OptionPlayRangeStart).SafeGetTimeValue(FTimeValue());
-		if (RangeStart.IsValid() && RangeStart > TimelineRange.Start)
+		if (ActiveRange.Start.IsValid() && ActiveRange.Start > TimelineRange.Start)
 		{
-			TimelineRange.Start = RangeStart;
+			TimelineRange.Start = ActiveRange.Start;
 		}
 		if (Out < TimelineRange.Start)
 		{
@@ -4171,8 +4150,9 @@ void FAdaptiveStreamingPlayer::CheckForStreamEnd()
  * and issues the first fragment request.
  *
  * @param NewPosition
+ * @param InTimeRange
  */
-void FAdaptiveStreamingPlayer::InternalStartAt(const FSeekParam& NewPosition)
+void FAdaptiveStreamingPlayer::InternalStartAt(const FSeekParam& NewPosition, const FTimeRange* InTimeRange)
 {
 	uint32 NextPlaybackID = Utils::Max(Utils::Max(CurrentPlaybackSequenceID[0], CurrentPlaybackSequenceID[1]), CurrentPlaybackSequenceID[2]) + 1;
 	CurrentPlaybackSequenceID[0] = NextPlaybackID;
@@ -4182,7 +4162,6 @@ void FAdaptiveStreamingPlayer::InternalStartAt(const FSeekParam& NewPosition)
 	StreamState = EStreamState::eStream_Running;
 
 	PlaybackState.SetHasEnded(false);
-	PlaybackState.SetPlayRangeHasChanged(false);
 	PlaybackState.SetLoopStateHasChanged(false);
 	PlaybackState.SetEncoderLatency(FTimeValue::GetInvalid());
 
@@ -4235,6 +4214,7 @@ void FAdaptiveStreamingPlayer::InternalStartAt(const FSeekParam& NewPosition)
 	PendingStartRequest->StartAt.Time = NewPosition.Time;
 	PendingStartRequest->StartingBitrate = NewPosition.StartingBitrate;
 	PendingStartRequest->StartType = SeekVars.bIsPlayStart ? FPendingStartRequest::EStartType::PlayStart : FPendingStartRequest::EStartType::Seeking;
+	PlaybackState.ActivateNewPlayRange(InTimeRange);
 	SetPlaystartOptions(PendingStartRequest->StartAt.Options);
 
 	// The fragment should be the closest to the time stamp unless we are rebuffering in which case it should be the one we failed on.
@@ -4273,7 +4253,7 @@ void FAdaptiveStreamingPlayer::InternalSetPlaybackEnded()
 
 	FTimeRange TimelineRange;
 	PlaybackState.GetTimelineRange(TimelineRange);
-	FTimeValue RangeEnd = GetOptions().GetValue(OptionPlayRangeEnd).SafeGetTimeValue(FTimeValue());
+	FTimeValue RangeEnd = PlaybackState.GetActivePlayRange().End;
 	if (RangeEnd.IsValid() && RangeEnd < TimelineRange.End)
 	{
 		TimelineRange.End = RangeEnd;
@@ -4630,7 +4610,7 @@ void FAdaptiveStreamingPlayer::InternalRebuffer()
 				StartAtTime.StartingBitrate = Action.SuggestedRestartBitrate;
 			}
 
-			InternalStartAt(StartAtTime);
+			InternalStartAt(StartAtTime, nullptr);
 		}
 	}
 }
@@ -4851,8 +4831,7 @@ void FAdaptiveStreamingPlayer::DebugPrint(void* pPlayer, void (*pPrintFN)(void*
 	pPrintFN(pPlayer, "Pipeline state: %s", GetPipelineStateName(PipelineState));
 	pPrintFN(pPlayer, "Stream state  : %s", GetStreamStateName(StreamState));
 	pPrintFN(pPlayer, "-----------------------------------");
-	FTimeValue RangeStart = GetOptions().GetValue(OptionPlayRangeStart).SafeGetTimeValue(FTimeValue());
-	FTimeValue RangeEnd = GetOptions().GetValue(OptionPlayRangeEnd).SafeGetTimeValue(FTimeValue());
+	FTimeRange ActiveRange = PlaybackState.GetActivePlayRange();
 	FTimeRange seekable, timeline;
 	GetSeekableRange(seekable);
 	PlaybackState.GetTimelineRange(timeline);
@@ -4860,7 +4839,7 @@ void FAdaptiveStreamingPlayer::DebugPrint(void* pPlayer, void (*pPrintFN)(void*
 	pPrintFN(pPlayer, "      duration: %.3f", GetDuration().GetAsSeconds());
 	pPrintFN(pPlayer, "timeline range: %.3f - %.3f", timeline.Start.GetAsSeconds(), timeline.End.GetAsSeconds());
 	pPrintFN(pPlayer, "seekable range: %.3f - %.3f", seekable.Start.GetAsSeconds(), seekable.End.GetAsSeconds());
-	pPrintFN(pPlayer, "playback range: %.3f - %.3f", RangeStart.GetAsSeconds(), RangeEnd.GetAsSeconds());
+	pPrintFN(pPlayer, "playback range: %.3f - %.3f", ActiveRange.Start.GetAsSeconds(), ActiveRange.End.GetAsSeconds());
 	pPrintFN(pPlayer, " play position: [%d] %.3f", (int32) GetPlayPosition().GetSequenceIndex(), GetPlayPosition().GetAsSeconds());
 	pPrintFN(pPlayer, "  is buffering: %s", IsBuffering() ? "Yes" : "No");
 	pPrintFN(pPlayer, "    is playing: %s", IsPlaying() ? "Yes" : "No");
diff --git a/Engine/Plugins/Media/ElectraPlayer/Source/ElectraPlayerRuntime/Private/Runtime/Player/AdaptiveStreamingPlayerInternal.h b/Engine/Plugins/Media/ElectraPlayer/Source/ElectraPlayerRuntime/Private/Runtime/Player/AdaptiveStreamingPlayerInternal.h
--- a/Engine/Plugins/Media/ElectraPlayer/Source/ElectraPlayerRuntime/Private/Runtime/Player/AdaptiveStreamingPlayerInternal.h
+++ b/Engine/Plugins/Media/ElectraPlayer/Source/ElectraPlayerRuntime/Private/Runtime/Player/AdaptiveStreamingPlayerInternal.h
@@ -193,6 +193,8 @@ struct FPlaybackState
 		CurrentLiveLatency.SetToInvalid();
 		EndPlaybackAtTime.SetToInvalid();
 		LoopState = {};
+		ActivePlaybackRange.Reset();
+		NewPlaybackRange.Reset();
 		bHaveMetadata = false;
 		bHasEnded = false;
 		bIsSeeking = false;
@@ -213,6 +215,8 @@ struct FPlaybackState
 	FTimeValue								CurrentLiveLatency;
 	FTimeValue								EndPlaybackAtTime;
 	IAdaptiveStreamingPlayer::FLoopState	LoopState;
+	FTimeRange								ActivePlaybackRange;
+	FTimeRange								NewPlaybackRange;
 	bool									bHaveMetadata;
 	bool									bHasEnded;
 	bool									bIsSeeking;
@@ -446,16 +450,64 @@ struct FPlaybackState
 		LoopState.bIsEnabled = bEnable;
 	}
 
-	void SetPlayRangeHasChanged(bool bWasChanged)
+	void SetPlayRange(const IAdaptiveStreamingPlayer::FPlaybackRange& InNewRange)
 	{
 		FScopeLock lock(&Lock);
-		bPlayrangeHasChanged = bWasChanged;
+		FTimeValue NewRangeStart = InNewRange.Start.IsSet() ? InNewRange.Start.GetValue() : FTimeValue();
+		FTimeValue NewRangeEnd = InNewRange.End.IsSet() ? InNewRange.End.GetValue() : FTimeValue();
+		if (NewRangeStart != NewPlaybackRange.Start || NewRangeEnd != NewPlaybackRange.End)
+		{
+			NewPlaybackRange.Start = NewRangeStart;
+			NewPlaybackRange.End = NewRangeEnd;
+			bPlayrangeHasChanged = true;
+		}
+	}
+	void SetPlayRange(const FTimeRange& InNewRange)
+	{
+		FScopeLock lock(&Lock);
+		if (InNewRange.Start != NewPlaybackRange.Start || InNewRange.End != NewPlaybackRange.End)
+		{
+			NewPlaybackRange = InNewRange;
+			bPlayrangeHasChanged = true;
+		}
+	}
+	FTimeRange GetPlayRange()
+	{
+		FScopeLock lock(&Lock);
+		return NewPlaybackRange;
+	}
+	void GetPlayRange(IAdaptiveStreamingPlayer::FPlaybackRange& OutRange)
+	{
+		OutRange.Start.Reset();
+		OutRange.End.Reset();
+		FScopeLock lock(&Lock);
+		if (NewPlaybackRange.Start.IsValid())
+		{
+			OutRange.Start = NewPlaybackRange.Start;
+		}
+		if (NewPlaybackRange.End.IsValid())
+		{
+			OutRange.End = NewPlaybackRange.End;
+		}
+	}
+	void ActivateNewPlayRange(const FTimeRange* InTimeRange)
+	{
+		FScopeLock lock(&Lock);
+		ActivePlaybackRange = InTimeRange ? *InTimeRange : NewPlaybackRange;
+		bPlayrangeHasChanged = InTimeRange ? bPlayrangeHasChanged : false;
+	}
+	FTimeRange GetActivePlayRange()
+	{
+		FScopeLock lock(&Lock);
+		return ActivePlaybackRange;
 	}
 	bool GetPlayRangeHasChanged()
 	{
 		FScopeLock lock(&Lock);
 		return bPlayrangeHasChanged;
 	}
+
+
 	void SetLoopStateHasChanged(bool bWasChanged)
 	{
 		FScopeLock lock(&Lock);
@@ -1691,6 +1743,7 @@ private:
 			PendingRequest.Reset();
 			ActiveRequest.Reset();
 			LastFinishedRequest.Reset();
+			PlayrangeOnRequest.Reset();
 			bForScrubbing = false;
 			bScrubPrerollDone = false;
 		}
@@ -1713,9 +1766,10 @@ private:
 		TOptional<FSeekParam> ActiveRequest;
 		TOptional<FSeekParam> LastFinishedRequest;
 
-		FCriticalSection Lock;
+		mutable FCriticalSection Lock;
 		// The pending request must be accessed under lock since it is written to by the main thread and read from the worker thread!
 		TOptional<FSeekParam> PendingRequest;
+		FTimeRange PlayrangeOnRequest;
 	};
 
 	struct FStreamBitrateInfo
@@ -1931,7 +1985,7 @@ private:
 	void FeedDecoder(EStreamType Type, IAccessUnitBufferInterface* Decoder, bool bHandleUnderrun);
 	void ClearAllDecoderEODs();
 
-	void InternalStartAt(const FSeekParam& NewPosition);
+	void InternalStartAt(const FSeekParam& NewPosition, const FTimeRange* InTimeRange);
 	void InternalPause();
 	void InternalResume();
 	void InternalRebuffer();
diff --git a/Engine/Plugins/Media/ElectraPlayer/Source/ElectraPlayerRuntime/Private/Runtime/Player/AdaptiveStreamingPlayerManifest.cpp b/Engine/Plugins/Media/ElectraPlayer/Source/ElectraPlayerRuntime/Private/Runtime/Player/AdaptiveStreamingPlayerManifest.cpp
--- a/Engine/Plugins/Media/ElectraPlayer/Source/ElectraPlayerRuntime/Private/Runtime/Player/AdaptiveStreamingPlayerManifest.cpp
+++ b/Engine/Plugins/Media/ElectraPlayer/Source/ElectraPlayerRuntime/Private/Runtime/Player/AdaptiveStreamingPlayerManifest.cpp
@@ -259,7 +259,7 @@ bool FAdaptiveStreamingPlayer::SelectManifest()
 		if (ManifestType != EMediaFormatType::Unknown)
 		{
 			TArray<FTimespan> SeekablePositions;
-			FTimeRange RestrictedPlaybackRange;
+			FTimeRange RestrictedPlaybackRange, ActivePlaybackRange;
 			TSharedPtrTS<IManifest> NewPresentation = ManifestReader->GetManifest();
 			check(NewPresentation.IsValid());
 
@@ -273,14 +273,18 @@ bool FAdaptiveStreamingPlayer::SelectManifest()
 			// like example.mp4#t=10.8,18.4
 			// These do not override any user defined range values!
 			RestrictedPlaybackRange = NewPresentation->GetPlaybackRange();
-			if (RestrictedPlaybackRange.Start.IsValid() && !GetOptions().HaveKey(OptionPlayRangeStart))
+			ActivePlaybackRange = PlaybackState.GetActivePlayRange();
+			if (RestrictedPlaybackRange.Start.IsValid() && !ActivePlaybackRange.Start.IsValid())
 			{
-				GetOptions().Set(OptionPlayRangeStart, FVariantValue(RestrictedPlaybackRange.Start));
+				ActivePlaybackRange.Start = RestrictedPlaybackRange.Start;
 			}
-			if (RestrictedPlaybackRange.End.IsValid() && !GetOptions().HaveKey(OptionPlayRangeEnd))
+			if (RestrictedPlaybackRange.End.IsValid() && !ActivePlaybackRange.End.IsValid())
 			{
-				GetOptions().Set(OptionPlayRangeEnd, FVariantValue(RestrictedPlaybackRange.End));
+				ActivePlaybackRange.End = RestrictedPlaybackRange.End;
 			}
+			// Set the new range, but say it did not change. This is the first start and we do not need to issue a seek.
+			PlaybackState.SetPlayRange(ActivePlaybackRange);
+			PlaybackState.ActivateNewPlayRange(nullptr);
 
 			TArray<FTrackMetadata> VideoTrackMetadata;
 			TArray<FTrackMetadata> AudioTrackMetadata;
