ElectraPlayer: Performing the bulk of the initialization from within the worker thread now instead of the game thread.

#rnx

diff --git a/Engine/Plugins/Media/ElectraPlayer/Source/ElectraPlayerRuntime/Private/Runtime/Player/AdaptiveStreamingPlayer.cpp b/Engine/Plugins/Media/ElectraPlayer/Source/ElectraPlayerRuntime/Private/Runtime/Player/AdaptiveStreamingPlayer.cpp
--- a/Engine/Plugins/Media/ElectraPlayer/Source/ElectraPlayerRuntime/Private/Runtime/Player/AdaptiveStreamingPlayer.cpp
+++ b/Engine/Plugins/Media/ElectraPlayer/Source/ElectraPlayerRuntime/Private/Runtime/Player/AdaptiveStreamingPlayer.cpp
@@ -283,120 +283,7 @@
 void FAdaptiveStreamingPlayer::Initialize(const FParamDict& Options)
 {
 	PlayerOptions = Options;
-
-	// Get the HTTP manager. This is a shared instance for all players.
-	HttpManager = IElectraHttpManager::Create();
-
-	// Create the DRM manager.
-	DrmManager = FDRMManager::Create(this);
-
-	// Create an entity cache.
-	EntityCache = IPlayerEntityCache::Create(this, PlayerOptions);
-
-	// Create an HTTP response cache.
-	HttpResponseCache = IHTTPResponseCache::Create(this, PlayerOptions);
-
-	// Create the ABR stream selector.
-	StreamSelector = IAdaptiveStreamSelector::Create(this, this);
-	AddMetricsReceiver(StreamSelector.Get());
-
-	// Create renderers.
-	CreateRenderers();
-
-	// Set up video decoder resolution limits. As the media playlists are parsed the video streams will be
-	// compared against these limits and those that exceed the limit will not be considered for playback.
-
-	// Maximum allowed vertical resolution specified?
-	if (PlayerOptions.HaveKey(TEXT("max_resoY")))
-	{
-		PlayerConfig.H264LimitUpto30fps.MaxResolution.Height = (int32) PlayerOptions.GetValue(TEXT("max_resoY")).GetInt64();
-		PlayerConfig.H264LimitAbove30fps.MaxResolution.Height = (int32) PlayerOptions.GetValue(TEXT("max_resoY")).GetInt64();
-	}
-	// A limit in vertical resolution for streams with more than 30fps?
-	if (PlayerOptions.HaveKey(TEXT("max_resoY_above_30fps")))
-	{
-		PlayerConfig.H264LimitAbove30fps.MaxResolution.Height = (int32) PlayerOptions.GetValue(TEXT("max_resoY_above_30fps")).GetInt64();
-	}
-	// Note: We could add additional limits here if need be.
-	//       Eventually these need to be differentiated based on codec as well.
-
-	// Get global video decoder capabilities and if supported, set the stream resolution limit accordingly.
-	IVideoDecoderH264::FStreamDecodeCapability Capability, StreamParam;
-	// Do a one-time global capability check with a default-empty stream param structure.
-	// This then gets used in the individual stream capability checks.
-	if (IVideoDecoderH264::GetStreamDecodeCapability(Capability, StreamParam))
-	{
-		if (Capability.Profile && Capability.Level)
-		{
-			PlayerConfig.H264LimitUpto30fps.MaxTierProfileLevel.Profile = Capability.Profile;
-			PlayerConfig.H264LimitUpto30fps.MaxTierProfileLevel.Level = Capability.Level;
-			PlayerConfig.H264LimitAbove30fps.MaxTierProfileLevel.Profile = Capability.Profile;
-			PlayerConfig.H264LimitAbove30fps.MaxTierProfileLevel.Level = Capability.Level;
-		}
-		if (Capability.Height)
-		{
-			PlayerConfig.H264LimitUpto30fps.MaxResolution.Height = Utils::Min(PlayerConfig.H264LimitUpto30fps.MaxResolution.Height, Capability.Height);
-			PlayerConfig.H264LimitAbove30fps.MaxResolution.Height = Utils::Min(PlayerConfig.H264LimitAbove30fps.MaxResolution.Height, Capability.Height);
-		}
-
-		// If this is software decoding only and there is limit for this in place on Windows, apply it.
-		if (PlayerOptions.HaveKey(TEXT("max_resoY_windows_software")) && Capability.DecoderSupportType == IVideoDecoderH264::FStreamDecodeCapability::ESupported::SoftwareOnly)
-		{
-			int32 MaxWinSWHeight = (int32) PlayerOptions.GetValue(TEXT("max_resoY_windows_software")).GetInt64();
-			PlayerConfig.H264LimitUpto30fps.MaxResolution.Height = Utils::Min(PlayerConfig.H264LimitUpto30fps.MaxResolution.Height, MaxWinSWHeight);
-			PlayerConfig.H264LimitAbove30fps.MaxResolution.Height = Utils::Min(PlayerConfig.H264LimitAbove30fps.MaxResolution.Height, MaxWinSWHeight);
-		}
-
-		// If the maximum fps is only up to 30 fps set the resolution for streams above 30fps so small
-		// that they will get rejected.
-		if (Capability.FPS > 0.0 && Capability.FPS <= 30.0)
-		{
-			PlayerConfig.H264LimitAbove30fps.MaxResolution.Height = 16;
-			PlayerConfig.H264LimitAbove30fps.MaxTierProfileLevel.Profile = 66;
-			PlayerConfig.H264LimitAbove30fps.MaxTierProfileLevel.Level = 10;
-		}
-	}
-
-	// Check for codecs that are not to be used as per the user's choice, even if the device supports them.
-	auto GetExcludedCodecPrefixes = [](TArray<FString>& OutList, const FParamDict& InOptions, const FString& InKey) -> void
-	{
-		if (InOptions.HaveKey(InKey))
-		{
-			const TCHAR* const CommaDelimiter = TEXT(",");
-			TArray<FString> Prefixes;
-			FString Value = InOptions.GetValue(InKey).GetFString();
-			Value.TrimQuotesInline();
-			Value.ParseIntoArray(Prefixes, CommaDelimiter, true);
-			for(auto &Prefix : Prefixes)
-			{
-				OutList.Emplace(Prefix.TrimStartAndEnd().TrimQuotes().TrimStartAndEnd());
-			}
-		}
-	};
-	GetExcludedCodecPrefixes(ExcludedVideoDecoderPrefixes,    PlayerOptions, TEXT("excluded_codecs_video"));
-	GetExcludedCodecPrefixes(ExcludedAudioDecoderPrefixes,    PlayerOptions, TEXT("excluded_codecs_audio"));
-	GetExcludedCodecPrefixes(ExcludedSubtitleDecoderPrefixes, PlayerOptions, TEXT("excluded_codecs_subtitles"));
-
-	auto GetCodecSelectionPriorities = [this](FCodecSelectionPriorities& OutPriorities, const FParamDict& InOptions, const FString& InKey, const TCHAR* const InType) -> void
-	{
-		if (InOptions.HaveKey(InKey))
-		{
-			if (!OutPriorities.Initialize(InOptions.GetValue(InKey).GetFString()))
-			{
-				PostLog(Facility::EFacility::Player, IInfoLog::ELevel::Warning, FString::Printf(TEXT("Failed to parse %s codec selection priority configuration"), InType));
-			}
-		}
-	};
-	GetCodecSelectionPriorities(CodecPrioritiesVideo,     PlayerOptions, TEXT("preferred_codecs_video"),     TEXT("video"));
-	GetCodecSelectionPriorities(CodecPrioritiesAudio,     PlayerOptions, TEXT("preferred_codecs_audio"),     TEXT("audio"));
-	GetCodecSelectionPriorities(CodecPrioritiesSubtitles, PlayerOptions, TEXT("preferred_codecs_subtitles"), TEXT("subtitle"));
-
-
-	// Unless already specified in the options to either value, enable frame accurate seeking now.
-	if (!PlayerOptions.HaveKey(OptionKeyFrameAccurateSeek))
-	{
-		PlayerOptions.Set(OptionKeyFrameAccurateSeek, FVariantValue(true));
-	}
+	WorkerThread.SendInitializeMessage();
 }
 
 void FAdaptiveStreamingPlayer::ModifyOptions(const FParamDict& InOptionsToSetOrChange, const FParamDict& InOptionsToClear)
@@ -1434,6 +1321,11 @@
 		// What is it?
 		switch(msg.Type)
 		{
+			case FWorkerThreadMessages::FMessage::EType::Initialize:
+			{
+				FAdaptiveStreamingPlayer::InternalInitialize();
+				break;
+			}
 			case FWorkerThreadMessages::FMessage::EType::LoadManifest:
 			{
 				InternalLoadManifest(msg.Data.ManifestToLoad.URL, msg.Data.ManifestToLoad.MimeType);
@@ -4423,6 +4315,124 @@
 }
 
 
+void FAdaptiveStreamingPlayer::InternalInitialize()
+{
+	// Get the HTTP manager. This is a shared instance for all players.
+	HttpManager = IElectraHttpManager::Create();
+
+	// Create the DRM manager.
+	DrmManager = FDRMManager::Create(this);
+
+	// Create an entity cache.
+	EntityCache = IPlayerEntityCache::Create(this, PlayerOptions);
+
+	// Create an HTTP response cache.
+	HttpResponseCache = IHTTPResponseCache::Create(this, PlayerOptions);
+
+	// Create the ABR stream selector.
+	StreamSelector = IAdaptiveStreamSelector::Create(this, this);
+	AddMetricsReceiver(StreamSelector.Get());
+
+	// Create renderers.
+	CreateRenderers();
+
+	// Set up video decoder resolution limits. As the media playlists are parsed the video streams will be
+	// compared against these limits and those that exceed the limit will not be considered for playback.
+
+	// Maximum allowed vertical resolution specified?
+	if (PlayerOptions.HaveKey(TEXT("max_resoY")))
+	{
+		PlayerConfig.H264LimitUpto30fps.MaxResolution.Height = (int32)PlayerOptions.GetValue(TEXT("max_resoY")).GetInt64();
+		PlayerConfig.H264LimitAbove30fps.MaxResolution.Height = (int32)PlayerOptions.GetValue(TEXT("max_resoY")).GetInt64();
+	}
+	// A limit in vertical resolution for streams with more than 30fps?
+	if (PlayerOptions.HaveKey(TEXT("max_resoY_above_30fps")))
+	{
+		PlayerConfig.H264LimitAbove30fps.MaxResolution.Height = (int32)PlayerOptions.GetValue(TEXT("max_resoY_above_30fps")).GetInt64();
+	}
+	// Note: We could add additional limits here if need be.
+	//       Eventually these need to be differentiated based on codec as well.
+
+	// Get global video decoder capabilities and if supported, set the stream resolution limit accordingly.
+	IVideoDecoderH264::FStreamDecodeCapability Capability, StreamParam;
+	// Do a one-time global capability check with a default-empty stream param structure.
+	// This then gets used in the individual stream capability checks.
+	if (IVideoDecoderH264::GetStreamDecodeCapability(Capability, StreamParam))
+	{
+		if (Capability.Profile && Capability.Level)
+		{
+			PlayerConfig.H264LimitUpto30fps.MaxTierProfileLevel.Profile = Capability.Profile;
+			PlayerConfig.H264LimitUpto30fps.MaxTierProfileLevel.Level = Capability.Level;
+			PlayerConfig.H264LimitAbove30fps.MaxTierProfileLevel.Profile = Capability.Profile;
+			PlayerConfig.H264LimitAbove30fps.MaxTierProfileLevel.Level = Capability.Level;
+		}
+		if (Capability.Height)
+		{
+			PlayerConfig.H264LimitUpto30fps.MaxResolution.Height = Utils::Min(PlayerConfig.H264LimitUpto30fps.MaxResolution.Height, Capability.Height);
+			PlayerConfig.H264LimitAbove30fps.MaxResolution.Height = Utils::Min(PlayerConfig.H264LimitAbove30fps.MaxResolution.Height, Capability.Height);
+		}
+
+		// If this is software decoding only and there is limit for this in place on Windows, apply it.
+		if (PlayerOptions.HaveKey(TEXT("max_resoY_windows_software")) && Capability.DecoderSupportType == IVideoDecoderH264::FStreamDecodeCapability::ESupported::SoftwareOnly)
+		{
+			int32 MaxWinSWHeight = (int32)PlayerOptions.GetValue(TEXT("max_resoY_windows_software")).GetInt64();
+			PlayerConfig.H264LimitUpto30fps.MaxResolution.Height = Utils::Min(PlayerConfig.H264LimitUpto30fps.MaxResolution.Height, MaxWinSWHeight);
+			PlayerConfig.H264LimitAbove30fps.MaxResolution.Height = Utils::Min(PlayerConfig.H264LimitAbove30fps.MaxResolution.Height, MaxWinSWHeight);
+		}
+
+		// If the maximum fps is only up to 30 fps set the resolution for streams above 30fps so small
+		// that they will get rejected.
+		if (Capability.FPS > 0.0 && Capability.FPS <= 30.0)
+		{
+			PlayerConfig.H264LimitAbove30fps.MaxResolution.Height = 16;
+			PlayerConfig.H264LimitAbove30fps.MaxTierProfileLevel.Profile = 66;
+			PlayerConfig.H264LimitAbove30fps.MaxTierProfileLevel.Level = 10;
+		}
+	}
+
+	// Check for codecs that are not to be used as per the user's choice, even if the device supports them.
+	auto GetExcludedCodecPrefixes = [](TArray<FString>& OutList, const FParamDict& InOptions, const FString& InKey) -> void
+	{
+		if (InOptions.HaveKey(InKey))
+		{
+			const TCHAR* const CommaDelimiter = TEXT(",");
+			TArray<FString> Prefixes;
+			FString Value = InOptions.GetValue(InKey).GetFString();
+			Value.TrimQuotesInline();
+			Value.ParseIntoArray(Prefixes, CommaDelimiter, true);
+			for (auto& Prefix : Prefixes)
+			{
+				OutList.Emplace(Prefix.TrimStartAndEnd().TrimQuotes().TrimStartAndEnd());
+			}
+		}
+	};
+	GetExcludedCodecPrefixes(ExcludedVideoDecoderPrefixes, PlayerOptions, TEXT("excluded_codecs_video"));
+	GetExcludedCodecPrefixes(ExcludedAudioDecoderPrefixes, PlayerOptions, TEXT("excluded_codecs_audio"));
+	GetExcludedCodecPrefixes(ExcludedSubtitleDecoderPrefixes, PlayerOptions, TEXT("excluded_codecs_subtitles"));
+
+	auto GetCodecSelectionPriorities = [this](FCodecSelectionPriorities& OutPriorities, const FParamDict& InOptions, const FString& InKey, const TCHAR* const InType) -> void
+	{
+		if (InOptions.HaveKey(InKey))
+		{
+			if (!OutPriorities.Initialize(InOptions.GetValue(InKey).GetFString()))
+			{
+				PostLog(Facility::EFacility::Player, IInfoLog::ELevel::Warning, FString::Printf(TEXT("Failed to parse %s codec selection priority configuration"), InType));
+			}
+		}
+	};
+	GetCodecSelectionPriorities(CodecPrioritiesVideo, PlayerOptions, TEXT("preferred_codecs_video"), TEXT("video"));
+	GetCodecSelectionPriorities(CodecPrioritiesAudio, PlayerOptions, TEXT("preferred_codecs_audio"), TEXT("audio"));
+	GetCodecSelectionPriorities(CodecPrioritiesSubtitles, PlayerOptions, TEXT("preferred_codecs_subtitles"), TEXT("subtitle"));
+
+
+	// Unless already specified in the options to either value, enable frame accurate seeking now.
+	if (!PlayerOptions.HaveKey(OptionKeyFrameAccurateSeek))
+	{
+		PlayerOptions.Set(OptionKeyFrameAccurateSeek, FVariantValue(true));
+	}
+}
+
+
 void FAdaptiveStreamingPlayer::InternalClose()
 {
 	// No longer need the manifest reader/updater
diff --git a/Engine/Plugins/Media/ElectraPlayer/Source/ElectraPlayerRuntime/Private/Runtime/Player/AdaptiveStreamingPlayerInternal.h b/Engine/Plugins/Media/ElectraPlayer/Source/ElectraPlayerRuntime/Private/Runtime/Player/AdaptiveStreamingPlayerInternal.h
--- a/Engine/Plugins/Media/ElectraPlayer/Source/ElectraPlayerRuntime/Private/Runtime/Player/AdaptiveStreamingPlayerInternal.h
+++ b/Engine/Plugins/Media/ElectraPlayer/Source/ElectraPlayerRuntime/Private/Runtime/Player/AdaptiveStreamingPlayerInternal.h
@@ -1354,6 +1354,7 @@ private:
 			enum class EType
 			{
 				// Commands
+				Initialize,
 				LoadManifest,
 				Pause,
 				Resume,
@@ -1476,6 +1477,12 @@ private:
 			Msg.Data.MediaEvent.Event = pEventSignal;
 			TriggerSharedWorkerThread(MoveTemp(Msg));
 		}
+		void SendInitializeMessage()
+		{
+			FMessage Msg;
+			Msg.Type = FMessage::EType::Initialize;
+			TriggerSharedWorkerThread(MoveTemp(Msg));
+		}
 		void SendLoadManifestMessage(const FString& URL, const FString& MimeType)
 		{
 			FMessage Msg;
@@ -1909,6 +1916,7 @@ private:
 	void OnFragmentReachedEOS(EStreamType InStreamType, TSharedPtr<const FBufferSourceInfo, ESPMode::ThreadSafe> InStreamSourceInfo) override;
 	void OnFragmentClose(TSharedPtrTS<IStreamSegment> pRequest) override;
 
+	void InternalInitialize();
 	void InternalHandleOnce();
 	bool InternalHandleThreadMessages();
 	void InternalHandleManifestReader();
