ElectraPlayer: Fixing initial playback range given via URL parameters overwriting range set by application

#rnx

diff --git a/Engine/Plugins/Media/ElectraPlayer/Source/ElectraPlayerRuntime/Private/Runtime/Player/AdaptiveStreamingPlayerManifest.cpp b/Engine/Plugins/Media/ElectraPlayer/Source/ElectraPlayerRuntime/Private/Runtime/Player/AdaptiveStreamingPlayerManifest.cpp
--- a/Engine/Plugins/Media/ElectraPlayer/Source/ElectraPlayerRuntime/Private/Runtime/Player/AdaptiveStreamingPlayerManifest.cpp
+++ b/Engine/Plugins/Media/ElectraPlayer/Source/ElectraPlayerRuntime/Private/Runtime/Player/AdaptiveStreamingPlayerManifest.cpp
@@ -259,7 +259,8 @@ bool FAdaptiveStreamingPlayer::SelectManifest()
 		if (ManifestType != EMediaFormatType::Unknown)
 		{
 			TArray<FTimespan> SeekablePositions;
-			FTimeRange RestrictedPlaybackRange, ActivePlaybackRange;
+			FPlaybackRange PlaybackRange;
+			FTimeRange RestrictedPlaybackRange;
 			TSharedPtrTS<IManifest> NewPresentation = ManifestReader->GetManifest();
 			check(NewPresentation.IsValid());
 
@@ -273,17 +274,17 @@ bool FAdaptiveStreamingPlayer::SelectManifest()
 			// like example.mp4#t=10.8,18.4
 			// These do not override any user defined range values!
 			RestrictedPlaybackRange = NewPresentation->GetPlaybackRange();
-			ActivePlaybackRange = PlaybackState.GetActivePlayRange();
-			if (RestrictedPlaybackRange.Start.IsValid() && !ActivePlaybackRange.Start.IsValid())
+			PlaybackState.GetPlayRange(PlaybackRange);
+			if (RestrictedPlaybackRange.Start.IsValid() && !PlaybackRange.Start.Get(FTimeValue()).IsValid())
 			{
-				ActivePlaybackRange.Start = RestrictedPlaybackRange.Start;
+				PlaybackRange.Start = RestrictedPlaybackRange.Start;
 			}
-			if (RestrictedPlaybackRange.End.IsValid() && !ActivePlaybackRange.End.IsValid())
+			if (RestrictedPlaybackRange.End.IsValid() && !PlaybackRange.End.Get(FTimeValue()).IsValid())
 			{
-				ActivePlaybackRange.End = RestrictedPlaybackRange.End;
+				PlaybackRange.End = RestrictedPlaybackRange.End;
 			}
 			// Set the new range, but say it did not change. This is the first start and we do not need to issue a seek.
-			PlaybackState.SetPlayRange(ActivePlaybackRange);
+			PlaybackState.SetPlayRange(PlaybackRange);
 			PlaybackState.ActivateNewPlayRange(nullptr);
 
 			TArray<FTrackMetadata> VideoTrackMetadata;
