ElectraPlayer: Adjusted max buffer sizes that were increased while working on supporting HDR down to what they were before.

#rnx

diff --git a/Engine/Plugins/Media/ElectraPlayer/Source/ElectraPlayerRuntime/Private/Runtime/Player/AdaptiveStreamingPlayerInternalConfig.h b/Engine/Plugins/Media/ElectraPlayer/Source/ElectraPlayerRuntime/Private/Runtime/Player/AdaptiveStreamingPlayerInternalConfig.h
--- a/Engine/Plugins/Media/ElectraPlayer/Source/ElectraPlayerRuntime/Private/Runtime/Player/AdaptiveStreamingPlayerInternalConfig.h
+++ b/Engine/Plugins/Media/ElectraPlayer/Source/ElectraPlayerRuntime/Private/Runtime/Player/AdaptiveStreamingPlayerInternalConfig.h
@@ -23,10 +23,10 @@ namespace Electra
 				WorkerThread.CoreAffinity = -1;
 
 				// Set default values to maximum permitted values.
-				StreamBufferConfigVideo.MaxDataSize = 256 << 20;
+				StreamBufferConfigVideo.MaxDataSize = 16 << 20;
 				StreamBufferConfigVideo.MaxDuration.SetFromSeconds(20.0);
 
-				StreamBufferConfigAudio.MaxDataSize = 8 << 20;
+				StreamBufferConfigAudio.MaxDataSize = 2 << 20;
 				StreamBufferConfigAudio.MaxDuration.SetFromSeconds(20.0);
 
 				// Subtitle streams tend to be sparse and each AU could potentially have
