ElectraPlayer: Fixed referenced frame not being sent to video decoder when it falls out of the playback range

#rnx

diff --git a/Engine/Plugins/Media/ElectraPlayer/Source/ElectraPlayerRuntime/Private/Runtime/Decoder/Android/VideoDecoderH264_Android.cpp b/Engine/Plugins/Media/ElectraPlayer/Source/ElectraPlayerRuntime/Private/Runtime/Decoder/Android/VideoDecoderH264_Android.cpp
--- a/Engine/Plugins/Media/ElectraPlayer/Source/ElectraPlayerRuntime/Private/Runtime/Decoder/Android/VideoDecoderH264_Android.cpp
+++ b/Engine/Plugins/Media/ElectraPlayer/Source/ElectraPlayerRuntime/Private/Runtime/Decoder/Android/VideoDecoderH264_Android.cpp
@@ -1063,7 +1063,10 @@ void FVideoDecoderH264::PrepareAU(TSharedPtrTS<FDecoderInput> AU)
 			if (StartTime >= AU->AccessUnit->LatestPTS)
 			{
 				StartTime.SetToInvalid();
-				AU->bIsDiscardable = true;
+				if (AU->AccessUnit->DTS.IsValid() && AU->AccessUnit->DTS >= AU->AccessUnit->LatestPTS)
+				{
+					AU->bIsDiscardable = true;
+				}
 			}
 			else if (EndTime >= AU->AccessUnit->LatestPTS)
 			{
diff --git a/Engine/Plugins/Media/ElectraPlayer/Source/ElectraPlayerRuntime/Private/Runtime/Decoder/Android/VideoDecoderH265_Android.cpp b/Engine/Plugins/Media/ElectraPlayer/Source/ElectraPlayerRuntime/Private/Runtime/Decoder/Android/VideoDecoderH265_Android.cpp
--- a/Engine/Plugins/Media/ElectraPlayer/Source/ElectraPlayerRuntime/Private/Runtime/Decoder/Android/VideoDecoderH265_Android.cpp
+++ b/Engine/Plugins/Media/ElectraPlayer/Source/ElectraPlayerRuntime/Private/Runtime/Decoder/Android/VideoDecoderH265_Android.cpp
@@ -1088,7 +1088,10 @@ void FVideoDecoderH265::PrepareAU(TSharedPtrTS<FDecoderInput> AU)
 			if (StartTime >= AU->AccessUnit->LatestPTS)
 			{
 				StartTime.SetToInvalid();
-				AU->bIsDiscardable = true;
+				if (AU->AccessUnit->DTS.IsValid() && AU->AccessUnit->DTS >= AU->AccessUnit->LatestPTS)
+				{
+					AU->bIsDiscardable = true;
+				}
 			}
 			else if (EndTime >= AU->AccessUnit->LatestPTS)
 			{
diff --git a/Engine/Plugins/Media/ElectraPlayer/Source/ElectraPlayerRuntime/Private/Runtime/Decoder/Apple/VideoDecoderH264_Apple.cpp b/Engine/Plugins/Media/ElectraPlayer/Source/ElectraPlayerRuntime/Private/Runtime/Decoder/Apple/VideoDecoderH264_Apple.cpp
--- a/Engine/Plugins/Media/ElectraPlayer/Source/ElectraPlayerRuntime/Private/Runtime/Decoder/Apple/VideoDecoderH264_Apple.cpp
+++ b/Engine/Plugins/Media/ElectraPlayer/Source/ElectraPlayerRuntime/Private/Runtime/Decoder/Apple/VideoDecoderH264_Apple.cpp
@@ -982,7 +982,10 @@ void FVideoDecoderH264::PrepareAU(TSharedPtr<FDecoderInput, ESPMode::ThreadSafe>
 			if (StartTime >= AU->AccessUnit->LatestPTS)
 			{
 				StartTime.SetToInvalid();
-				AU->bIsDiscardable = true;
+				if (AU->AccessUnit->DTS.IsValid() && AU->AccessUnit->DTS >= AU->AccessUnit->LatestPTS)
+				{
+					AU->bIsDiscardable = true;
+				}
 			}
 			else if (EndTime >= AU->AccessUnit->LatestPTS)
 			{
@@ -1769,6 +1772,7 @@ void FVideoDecoderH264::ProcessOutput(bool bFlush)
 								}
 							}
 						}
+						CFRelease(Dict);
 					}
 					double PixelAspectRatio = (double)ax / (double)ay;
 
diff --git a/Engine/Plugins/Media/ElectraPlayer/Source/ElectraPlayerRuntime/Private/Runtime/Decoder/Apple/VideoDecoderH265_Apple.cpp b/Engine/Plugins/Media/ElectraPlayer/Source/ElectraPlayerRuntime/Private/Runtime/Decoder/Apple/VideoDecoderH265_Apple.cpp
--- a/Engine/Plugins/Media/ElectraPlayer/Source/ElectraPlayerRuntime/Private/Runtime/Decoder/Apple/VideoDecoderH265_Apple.cpp
+++ b/Engine/Plugins/Media/ElectraPlayer/Source/ElectraPlayerRuntime/Private/Runtime/Decoder/Apple/VideoDecoderH265_Apple.cpp
@@ -1013,7 +1013,10 @@ void FVideoDecoderH265::PrepareAU(TSharedPtr<FDecoderInput, ESPMode::ThreadSafe>
 			if (StartTime >= AU->AccessUnit->LatestPTS)
 			{
 				StartTime.SetToInvalid();
-				AU->bIsDiscardable = true;
+				if (AU->AccessUnit->DTS.IsValid() && AU->AccessUnit->DTS >= AU->AccessUnit->LatestPTS)
+				{
+					AU->bIsDiscardable = true;
+				}
 			}
 			else if (EndTime >= AU->AccessUnit->LatestPTS)
 			{
@@ -1791,6 +1794,7 @@ void FVideoDecoderH265::ProcessOutput(bool bFlush)
 								}
 							}
 						}
+						CFRelease(Dict);
 					}
 					double PixelAspectRatio = (double)ax / (double)ay;
 
diff --git a/Engine/Plugins/Media/ElectraPlayer/Source/ElectraPlayerRuntime/Private/Runtime/Decoder/DX/VideoDecoderH264_DX.cpp b/Engine/Plugins/Media/ElectraPlayer/Source/ElectraPlayerRuntime/Private/Runtime/Decoder/DX/VideoDecoderH264_DX.cpp
--- a/Engine/Plugins/Media/ElectraPlayer/Source/ElectraPlayerRuntime/Private/Runtime/Decoder/DX/VideoDecoderH264_DX.cpp
+++ b/Engine/Plugins/Media/ElectraPlayer/Source/ElectraPlayerRuntime/Private/Runtime/Decoder/DX/VideoDecoderH264_DX.cpp
@@ -829,7 +829,10 @@ void FVideoDecoderH264::PrepareAU(TSharedPtrTS<FDecoderInput> AU)
 			if (StartTime >= AU->AccessUnit->LatestPTS)
 			{
 				StartTime.SetToInvalid();
-				AU->bIsDiscardable = true;
+				if (AU->AccessUnit->DTS.IsValid() && AU->AccessUnit->DTS >= AU->AccessUnit->LatestPTS)
+				{
+					AU->bIsDiscardable = true;
+				}
 			}
 			else if (EndTime >= AU->AccessUnit->LatestPTS)
 			{
diff --git a/Engine/Plugins/Media/ElectraPlayer/Source/ElectraPlayerRuntime/Private/Runtime/Decoder/DX/VideoDecoderH265_DX.cpp b/Engine/Plugins/Media/ElectraPlayer/Source/ElectraPlayerRuntime/Private/Runtime/Decoder/DX/VideoDecoderH265_DX.cpp
--- a/Engine/Plugins/Media/ElectraPlayer/Source/ElectraPlayerRuntime/Private/Runtime/Decoder/DX/VideoDecoderH265_DX.cpp
+++ b/Engine/Plugins/Media/ElectraPlayer/Source/ElectraPlayerRuntime/Private/Runtime/Decoder/DX/VideoDecoderH265_DX.cpp
@@ -719,7 +719,10 @@ void FVideoDecoderH265::PrepareAU(TSharedPtrTS<FDecoderInput> AU)
 			if (StartTime >= AU->AccessUnit->LatestPTS)
 			{
 				StartTime.SetToInvalid();
-				AU->bIsDiscardable = true;
+				if (AU->AccessUnit->DTS.IsValid() && AU->AccessUnit->DTS >= AU->AccessUnit->LatestPTS)
+				{
+					AU->bIsDiscardable = true;
+				}
 			}
 			else if (EndTime >= AU->AccessUnit->LatestPTS)
 			{
diff --git a/Engine/Plugins/Media/ElectraPlayer/Source/ElectraPlayerRuntime/Private/Runtime/Decoder/Linux/VideoDecoderH264_Linux.cpp b/Engine/Plugins/Media/ElectraPlayer/Source/ElectraPlayerRuntime/Private/Runtime/Decoder/Linux/VideoDecoderH264_Linux.cpp
--- a/Engine/Plugins/Media/ElectraPlayer/Source/ElectraPlayerRuntime/Private/Runtime/Decoder/Linux/VideoDecoderH264_Linux.cpp
+++ b/Engine/Plugins/Media/ElectraPlayer/Source/ElectraPlayerRuntime/Private/Runtime/Decoder/Linux/VideoDecoderH264_Linux.cpp
@@ -736,7 +736,10 @@ void FVideoDecoderH264LinuxLibavcodec::PrepareAU(TSharedPtr<FDecoderInput, ESPMo
 			if (StartTime >= AU->AccessUnit->LatestPTS)
 			{
 				StartTime.SetToInvalid();
-				AU->bIsDiscardable = true;
+				if (AU->AccessUnit->DTS.IsValid() && AU->AccessUnit->DTS >= AU->AccessUnit->LatestPTS)
+				{
+					AU->bIsDiscardable = true;
+				}
 			}
 			else if (EndTime >= AU->AccessUnit->LatestPTS)
 			{
diff --git a/Engine/Plugins/Media/ElectraPlayer/Source/ElectraPlayerRuntime/Private/Runtime/Decoder/Linux/VideoDecoderH265_Linux.cpp b/Engine/Plugins/Media/ElectraPlayer/Source/ElectraPlayerRuntime/Private/Runtime/Decoder/Linux/VideoDecoderH265_Linux.cpp
--- a/Engine/Plugins/Media/ElectraPlayer/Source/ElectraPlayerRuntime/Private/Runtime/Decoder/Linux/VideoDecoderH265_Linux.cpp
+++ b/Engine/Plugins/Media/ElectraPlayer/Source/ElectraPlayerRuntime/Private/Runtime/Decoder/Linux/VideoDecoderH265_Linux.cpp
@@ -805,7 +805,10 @@ void FVideoDecoderH265LinuxLibavcodec::PrepareAU(TSharedPtr<FDecoderInput, ESPMo
 			if (StartTime >= AU->AccessUnit->LatestPTS)
 			{
 				StartTime.SetToInvalid();
-				AU->bIsDiscardable = true;
+				if (AU->AccessUnit->DTS.IsValid() && AU->AccessUnit->DTS >= AU->AccessUnit->LatestPTS)
+				{
+					AU->bIsDiscardable = true;
+				}
 			}
 			else if (EndTime >= AU->AccessUnit->LatestPTS)
 			{
