ElectraPlayer: Fixed software decoding on Windows

#rnx

diff --git a/Engine/Plugins/Media/ElectraPlayer/Source/ElectraPlayerRuntime/Private/Runtime/Decoder/DX/VideoDecoderH264_DX.cpp b/Engine/Plugins/Media/ElectraPlayer/Source/ElectraPlayerRuntime/Private/Runtime/Decoder/DX/VideoDecoderH264_DX.cpp
--- a/Engine/Plugins/Media/ElectraPlayer/Source/ElectraPlayerRuntime/Private/Runtime/Decoder/DX/VideoDecoderH264_DX.cpp
+++ b/Engine/Plugins/Media/ElectraPlayer/Source/ElectraPlayerRuntime/Private/Runtime/Decoder/DX/VideoDecoderH264_DX.cpp
@@ -1149,8 +1149,12 @@ bool FVideoDecoderH264::Decode(TSharedPtrTS<FDecoderInput> AU, bool bResolutionC
 			// Check if we got an output buffer. When flushing or terminating AcquireOutputBuffer() may return with no buffer!
 			if (CurrentRenderOutputBuffer)
 			{
-				// Setup texture sample to receive data during decode
-				PreInitDecodeOutputForSW(FIntPoint(CurrentSampleInfo.GetResolution().Width, CurrentSampleInfo.GetResolution().Height));
+				// Setup texture sample to receive data during decode. If this created a new IMFSample then we need to
+				// reset and re-create the output buffer to update the sample in there.
+				if (PreInitDecodeOutputForSW(FIntPoint(CurrentSampleInfo.GetResolution().Width, CurrentSampleInfo.GetResolution().Height)))
+				{
+					CurrentDecoderOutputBuffer.Reset();
+				}
 			}
 			else
 			{
diff --git a/Engine/Plugins/Media/ElectraPlayer/Source/ElectraPlayerRuntime/Private/Runtime/Decoder/DX/VideoDecoderH264_DX.h b/Engine/Plugins/Media/ElectraPlayer/Source/ElectraPlayerRuntime/Private/Runtime/Decoder/DX/VideoDecoderH264_DX.h
--- a/Engine/Plugins/Media/ElectraPlayer/Source/ElectraPlayerRuntime/Private/Runtime/Decoder/DX/VideoDecoderH264_DX.h
+++ b/Engine/Plugins/Media/ElectraPlayer/Source/ElectraPlayerRuntime/Private/Runtime/Decoder/DX/VideoDecoderH264_DX.h
@@ -194,7 +194,7 @@ namespace Electra
 		// Per platform specialization
 		virtual bool InternalDecoderCreate() = 0;
 		virtual bool CreateDecoderOutputBuffer() = 0;
-		virtual void PreInitDecodeOutputForSW(const FIntPoint& Dim) = 0;
+		virtual bool PreInitDecodeOutputForSW(const FIntPoint& Dim) = 0;
 		virtual bool SetupDecodeOutputData(const FIntPoint& ImageDim, const TRefCountPtr<IMFSample>& DecodedOutputSample, FParamDict* OutputBufferSampleProperties) = 0;
 		virtual void PlatformTick() {}
 
diff --git a/Engine/Plugins/Media/ElectraPlayer/Source/ElectraPlayerRuntime/Private/Runtime/Decoder/Windows/ElectraVideoDecoder_PC.cpp b/Engine/Plugins/Media/ElectraPlayer/Source/ElectraPlayerRuntime/Private/Runtime/Decoder/Windows/ElectraVideoDecoder_PC.cpp
--- a/Engine/Plugins/Media/ElectraPlayer/Source/ElectraPlayerRuntime/Private/Runtime/Decoder/Windows/ElectraVideoDecoder_PC.cpp
+++ b/Engine/Plugins/Media/ElectraPlayer/Source/ElectraPlayerRuntime/Private/Runtime/Decoder/Windows/ElectraVideoDecoder_PC.cpp
@@ -7,6 +7,7 @@
 FElectraPlayerVideoDecoderOutputPC::FElectraPlayerVideoDecoderOutputPC()
 	: OutputType(EOutputType::Unknown)
 	, Stride(0)
+	, AllocatedBufferSize(0)
 	, SampleDim(0,0)
 {
 }
@@ -59,13 +60,22 @@ FIntPoint FElectraPlayerVideoDecoderOutputPC::GetDim() const
 	return SampleDim;
 }
 
-void FElectraPlayerVideoDecoderOutputPC::PreInitForDecode(FIntPoint OutputDim, const TFunction<void(int32 /*ApiReturnValue*/, const FString& /*Message*/, uint16 /*Code*/, UEMediaError /*Error*/ )>& PostError)
+bool FElectraPlayerVideoDecoderOutputPC::PreInitForDecode(FIntPoint OutputDim, const TFunction<void(int32 /*ApiReturnValue*/, const FString& /*Message*/, uint16 /*Code*/, UEMediaError /*Error*/ )>& PostError)
 {
+	DWORD SwBufAllocSize = 0;
 	FIntPoint Dim;
 	Dim.X = (OutputDim.X + 15) & ~15;
 	Dim.Y = ((OutputDim.Y + 15) & ~15) * 3/2;	// adjust height to encompass Y and UV planes
 
+	// Windows 8+ and a valid DX11 rendering device?
+	// (Rendering device will be null if the rendering device is not DX11 (DX12, Vulkan...))
+	bool bIsHW = Electra::IsWindows8Plus() && Electra::FDXDeviceInfo::s_DXDeviceInfo->RenderingDx11Device;
 	bool bNeedNew = !MFSample.IsValid() || SampleDim != Dim;
+	if (!bIsHW)
+	{
+		SwBufAllocSize = TargetBufferAllocSize.Get(Dim.X * Dim.Y * sizeof(uint8));
+		bNeedNew |= SwBufAllocSize != AllocatedBufferSize;
+	}
 
 	if (bNeedNew)
 	{
@@ -76,9 +86,7 @@ void FElectraPlayerVideoDecoderOutputPC::PreInitForDecode(FIntPoint OutputDim, c
 
 		SampleDim = Dim;
 
-		// Windows 8+ and a valid DX11 rendering device?
-		// (Rendering device will be null if the rendering device is not DX11 (DX12, Vulkan...))
-		if (Electra::IsWindows8Plus() && Electra::FDXDeviceInfo::s_DXDeviceInfo->RenderingDx11Device)
+		if (bIsHW)
 		{
 			// Software decode into DX11 texture
 
@@ -109,37 +117,29 @@ void FElectraPlayerVideoDecoderOutputPC::PreInitForDecode(FIntPoint OutputDim, c
 			{
 				UE_LOG(LogElectraPlayer, Error, TEXT("CreateTexture2D() failed with 0x%X %s"), Result, *GetComErrorDescription(Result));
 				PostError(Result, "Failed to create software decode output texture", ERRCODE_INTERNAL_COULD_NOT_CREATE_OUTPUTBUFFER, UEMEDIA_ERROR_OK);
-				return;
+				return false;
 			}
 			if (FAILED(Result = MFCreateDXGISurfaceBuffer(__uuidof(ID3D11Texture2D), Texture, 0, false, MediaBuffer.GetInitReference())))
 			{
 				UE_LOG(LogElectraPlayer, Error, TEXT("MFCreateDXGISurfaceBuffer() failed with 0x%X %s"), Result, *GetComErrorDescription(Result));
 				PostError(Result, "Failed to create software decode output texture", ERRCODE_INTERNAL_COULD_NOT_CREATE_OUTPUTBUFFER, UEMEDIA_ERROR_OK);
-				return;
+				return false;
 			}
 
 			OutputType = EOutputType::SoftwareWin8Plus;
 		}
 		else // Software decode into CPU buffer
 		{
-			DWORD AllocSize;
-
-			if (TargetBufferAllocSize.IsSet())
-			{
-				AllocSize = TargetBufferAllocSize.GetValue();
-			}
-			else
-			{
-				AllocSize = Dim.X * Dim.Y * sizeof(uint8);
-			}
+			check(SwBufAllocSize);
 
 			// SW decode results are just delivered in a simple CPU-side buffer. Create the decoder side version of this...
-			if (FAILED(Result = MFCreateMemoryBuffer(AllocSize, MediaBuffer.GetInitReference())))
+			if (FAILED(Result = MFCreateMemoryBuffer(SwBufAllocSize, MediaBuffer.GetInitReference())))
 			{
 				UE_LOG(LogElectraPlayer, Error, TEXT("MFCreateMemoryBuffer() failed with 0x%X %s"), Result, *GetComErrorDescription(Result));
 				PostError(Result, "Failed to create software decode output buffer", ERRCODE_INTERNAL_COULD_NOT_CREATE_OUTPUTBUFFER, UEMEDIA_ERROR_OK);
-				return;
+				return false;
 			}
+			AllocatedBufferSize = SwBufAllocSize;
 			OutputType = EOutputType::SoftwareWin7;
 		}
 
@@ -149,16 +149,29 @@ void FElectraPlayerVideoDecoderOutputPC::PreInitForDecode(FIntPoint OutputDim, c
 			{
 				UE_LOG(LogElectraPlayer, Error, TEXT("IMFSample::AddBuffer() failed with 0x%X %s"), Result, *GetComErrorDescription(Result));
 				PostError(Result, "Failed to add output buffer to software decoder sample", ERRCODE_INTERNAL_COULD_NOT_ADD_OUTPUT_BUFFER_TO_SAMPLE, UEMEDIA_ERROR_OK);
-				return;
+				return false;
 			}
 		}
 		else
 		{
 			UE_LOG(LogElectraPlayer, Error, TEXT("MFCreateSample() failed with 0x%X %s"), Result, *GetComErrorDescription(Result));
 			PostError(Result, "Failed to create software decode output sample", ERRCODE_INTERNAL_COULD_NOT_CREATE_OUTPUT_SAMPLE, UEMEDIA_ERROR_OK);
-			return;
+			return false;
+		}
+	}
+	else if (!bIsHW)
+	{
+		// In software mode we have to set the size of the buffer back to zero. A *new* decoder transform checks for this
+		// on the first decode and will fail with a "0x80004005 : 'CopyDecodedFrame failed'" if this is not the case.
+		// Once it is running this does not seem to matter. If a new decoder is required after a seek then this must be an empty buffer.
+		check(MFSample.IsValid());
+		TRefCountPtr<IMFMediaBuffer> MediaBuffer;
+		if (MFSample->GetBufferByIndex(0, MediaBuffer.GetInitReference()) == S_OK)
+		{
+			MediaBuffer->SetCurrentLength(0);
 		}
 	}
+	return bNeedNew;
 }
 
 
diff --git a/Engine/Plugins/Media/ElectraPlayer/Source/ElectraPlayerRuntime/Private/Runtime/Decoder/Windows/ElectraVideoDecoder_PC.h b/Engine/Plugins/Media/ElectraPlayer/Source/ElectraPlayerRuntime/Private/Runtime/Decoder/Windows/ElectraVideoDecoder_PC.h
--- a/Engine/Plugins/Media/ElectraPlayer/Source/ElectraPlayerRuntime/Private/Runtime/Decoder/Windows/ElectraVideoDecoder_PC.h
+++ b/Engine/Plugins/Media/ElectraPlayer/Source/ElectraPlayerRuntime/Private/Runtime/Decoder/Windows/ElectraVideoDecoder_PC.h
@@ -42,7 +42,7 @@ public:
 
 	// Software decode (into texture if DX11 device specified - available only Win8+)
 	void SetSWDecodeTargetBufferSize(uint32 InTargetBufferSize);
-	void PreInitForDecode(FIntPoint OutputDim, const TFunction<void(int32 /*ApiReturnValue*/, const FString& /*Message*/, uint16 /*Code*/, UEMediaError /*Error*/)>& PostError);
+	bool PreInitForDecode(FIntPoint OutputDim, const TFunction<void(int32 /*ApiReturnValue*/, const FString& /*Message*/, uint16 /*Code*/, UEMediaError /*Error*/)>& PostError);
 	void ProcessDecodeOutput(FIntPoint OutputDim, Electra::FParamDict* InParamDict);
 
 	void SetOwner(const TSharedPtr<IDecoderOutputOwner, ESPMode::ThreadSafe>& InOwningRenderer) override;
@@ -79,6 +79,7 @@ private:
 
 	// WMF sample (owned by this class if SW decoder is used)
 	TRefCountPtr<IMFSample> MFSample;
+	uint32 AllocatedBufferSize;
 
 	// Dimension of any internally allocated buffer - stored explicitly to cover various special cases for DX
 	FIntPoint SampleDim;
diff --git a/Engine/Plugins/Media/ElectraPlayer/Source/ElectraPlayerRuntime/Private/Runtime/Decoder/Windows/VideoDecoderH264_PC.cpp b/Engine/Plugins/Media/ElectraPlayer/Source/ElectraPlayerRuntime/Private/Runtime/Decoder/Windows/VideoDecoderH264_PC.cpp
--- a/Engine/Plugins/Media/ElectraPlayer/Source/ElectraPlayerRuntime/Private/Runtime/Decoder/Windows/VideoDecoderH264_PC.cpp
+++ b/Engine/Plugins/Media/ElectraPlayer/Source/ElectraPlayerRuntime/Private/Runtime/Decoder/Windows/VideoDecoderH264_PC.cpp
@@ -62,7 +62,7 @@ public:
 private:
 	virtual bool InternalDecoderCreate() override;
 	virtual bool CreateDecoderOutputBuffer() override;
-	virtual void PreInitDecodeOutputForSW(const FIntPoint& Dim);
+	virtual bool PreInitDecodeOutputForSW(const FIntPoint& Dim);
 	virtual bool SetupDecodeOutputData(const FIntPoint& ImageDim, const TRefCountPtr<IMFSample>& DecodedOutputSample, FParamDict* OutputBufferSampleProperties) override;
 
 	bool CopyTexture(const TRefCountPtr<IMFSample>& DecodedSample, Electra::FParamDict* ParamDict, FIntPoint OutputDim);
@@ -288,7 +288,7 @@ bool FVideoDecoderH264_PC::CreateDecoderOutputBuffer()
 }
 
 
-void FVideoDecoderH264_PC::PreInitDecodeOutputForSW(const FIntPoint& Dim)
+bool FVideoDecoderH264_PC::PreInitDecodeOutputForSW(const FIntPoint& Dim)
 {
 	TSharedPtr<FElectraPlayerVideoDecoderOutputPC, ESPMode::ThreadSafe> DecoderOutput = CurrentRenderOutputBuffer->GetBufferProperties().GetValue("texture").GetSharedPointer<FElectraPlayerVideoDecoderOutputPC>();
 	check(DecoderOutput);
@@ -304,7 +304,7 @@ void FVideoDecoderH264_PC::PreInitDecodeOutputForSW(const FIntPoint& Dim)
 		}
 	}
 
-	DecoderOutput->PreInitForDecode(Dim, [this](int32 ApiReturnValue, const FString& Message, uint16 Code, UEMediaError Error) { PostError(ApiReturnValue, Message, Code, Error); });
+	return DecoderOutput->PreInitForDecode(Dim, [this](int32 ApiReturnValue, const FString& Message, uint16 Code, UEMediaError Error) { PostError(ApiReturnValue, Message, Code, Error); });
 }
 
 
