ElectraPlayer: Use mp4 media duration based on longest duration of the individual tracks (their 'tkhd' boxes) over the duration given in the 'mdhd' box, which may be larger than the longest track.

#rnx

diff --git a/Engine/Plugins/Media/ElectraPlayer/Source/ElectraPlayerRuntime/Private/Runtime/Demuxer/ParserISO14496-12.cpp b/Engine/Plugins/Media/ElectraPlayer/Source/ElectraPlayerRuntime/Private/Runtime/Demuxer/ParserISO14496-12.cpp
--- a/Engine/Plugins/Media/ElectraPlayer/Source/ElectraPlayerRuntime/Private/Runtime/Demuxer/ParserISO14496-12.cpp
+++ b/Engine/Plugins/Media/ElectraPlayer/Source/ElectraPlayerRuntime/Private/Runtime/Demuxer/ParserISO14496-12.cpp
@@ -5771,13 +5771,13 @@ private:
 		if (TFHDBox == nullptr)
 		{
 			// Not fragmented.
-			if (MDHDBox)
+			if (TKHDBox && MVHDBox)
 			{
-				return FTimeFraction((int64)MDHDBox->GetDuration(), MDHDBox->GetTimescale());
+				return FTimeFraction(TKHDBox->GetDuration(), MVHDBox->GetTimescale());
 			}
-			else if (TKHDBox && MVHDBox)
+			else if (MDHDBox)
 			{
-				return FTimeFraction(TKHDBox->GetDuration(), MVHDBox->GetTimescale());
+				return FTimeFraction((int64)MDHDBox->GetDuration(), MDHDBox->GetTimescale());
 			}
 		}
 		return FTimeFraction();
diff --git a/Engine/Plugins/Media/ElectraPlayer/Source/ElectraPlayerRuntime/Private/Runtime/Player/mp4/ManifestMP4.cpp b/Engine/Plugins/Media/ElectraPlayer/Source/ElectraPlayerRuntime/Private/Runtime/Player/mp4/ManifestMP4.cpp
--- a/Engine/Plugins/Media/ElectraPlayer/Source/ElectraPlayerRuntime/Private/Runtime/Player/mp4/ManifestMP4.cpp
+++ b/Engine/Plugins/Media/ElectraPlayer/Source/ElectraPlayerRuntime/Private/Runtime/Player/mp4/ManifestMP4.cpp
@@ -642,6 +642,7 @@ FErrorDetail FManifestMP4Internal::FTimelineAssetMP4::Build(IPlayerSessionServic
 	PlayerSessionServices = InPlayerSessionServices;
 	MediaURL = URL;
 	// Go over the supported tracks and create an internal manifest-like structure for the player to work with.
+	LongestTrackDuration = FTimeValue::GetZero();
 	for(int32 nTrack=0,nMaxTrack=MP4Parser->GetNumberOfTracks(); nTrack < nMaxTrack; ++nTrack)
 	{
 		const IParserISO14496_12::ITrack* Track = MP4Parser->GetTrackByIndex(nTrack);
@@ -670,6 +671,12 @@ FErrorDetail FManifestMP4Internal::FTimelineAssetMP4::Build(IPlayerSessionServic
 					default:
 						break;
 				}
+				FTimeValue TrkDur;
+				TrkDur.SetFromTimeFraction(Track->GetDuration());
+				if (TrkDur > LongestTrackDuration)
+				{
+					LongestTrackDuration = TrkDur;
+				}
 			}
 			else
 			{
@@ -798,7 +805,9 @@ void FManifestMP4Internal::FTimelineAssetMP4::UpdatePlayRangeEndInfo(const FTime
 							bFirst = false;
 							TrackLocalTime = PlayRangeEndInfo.Time.GetAsTimebase(TrkIt->GetTimescale());
 						}
-						if (TrkIt->GetPTS() >= TrackLocalTime)
+						// Look at the DTS, not the PTS since a frame with a larger PTS may still be referenced
+						// by earlier frames (B frame reordering).
+						if (TrkIt->GetDTS() >= TrackLocalTime)
 						{
 							TrackEndBytePos = TrkIt->GetSampleFileOffset();
 							if (TrackEndBytePos > PlayRangeEndInfo.TotalFileEndOffset)
@@ -827,6 +836,10 @@ IManifest::FResult FManifestMP4Internal::FTimelineAssetMP4::GetStartingSegment(T
 	bool bFrameAccurateSearch = AtAbsoluteFilePos < 0 ? StartPosition.Options.bFrameAccuracy : false;
 	FTimeValue PlayRangeEnd = StartPosition.Options.PlaybackRange.End;
 	check(PlayRangeEnd.IsValid());
+	if (PlayRangeEnd > LongestTrackDuration)
+	{
+		PlayRangeEnd = LongestTrackDuration;
+	}
 	UpdatePlayRangeEndInfo(PlayRangeEnd);
 
 	// Look at the actual tracks. If there is video search there first for a keyframe/IDR frame.
diff --git a/Engine/Plugins/Media/ElectraPlayer/Source/ElectraPlayerRuntime/Private/Runtime/Player/mp4/ManifestMP4.h b/Engine/Plugins/Media/ElectraPlayer/Source/ElectraPlayerRuntime/Private/Runtime/Player/mp4/ManifestMP4.h
--- a/Engine/Plugins/Media/ElectraPlayer/Source/ElectraPlayerRuntime/Private/Runtime/Player/mp4/ManifestMP4.h
+++ b/Engine/Plugins/Media/ElectraPlayer/Source/ElectraPlayerRuntime/Private/Runtime/Player/mp4/ManifestMP4.h
@@ -149,17 +149,7 @@ public:
 
 		virtual FTimeValue GetDuration() const override
 		{
-			if (MoovBoxParser.IsValid())
-			{
-				TMediaOptionalValue<FTimeFraction> movieDuration = MoovBoxParser->GetMovieDuration();
-				if (movieDuration.IsSet())
-				{
-					FTimeValue dur;
-					dur.SetFromTimeFraction(movieDuration.Value());
-					return dur;
-				}
-			}
-			return FTimeValue();
+			return LongestTrackDuration;
 		}
 
 		virtual FString GetAssetIdentifier() const override
@@ -263,6 +253,7 @@ public:
 		void LogMessage(IInfoLog::ELevel Level, const FString& Message);
 		IPlayerSessionServices* 						PlayerSessionServices;
 		FString											MediaURL;
+		FTimeValue										LongestTrackDuration;
 		TSharedPtrTS<IParserISO14496_12>				MoovBoxParser;
 		TArray<TSharedPtrTS<FAdaptationSetMP4>>			VideoAdaptationSets;
 		TArray<TSharedPtrTS<FAdaptationSetMP4>>			AudioAdaptationSets;
