ElectraPlayer: Removed a wrong check that could cause an infinite loop when asked to flush while already flushing

#rnx

diff --git a/Engine/Plugins/Media/ElectraPlayer/Source/ElectraPlayerRuntime/Private/Runtime/Decoder/DX/AudioDecoderAAC_DX.cpp b/Engine/Plugins/Media/ElectraPlayer/Source/ElectraPlayerRuntime/Private/Runtime/Decoder/DX/AudioDecoderAAC_DX.cpp
--- a/Engine/Plugins/Media/ElectraPlayer/Source/ElectraPlayerRuntime/Private/Runtime/Decoder/DX/AudioDecoderAAC_DX.cpp
+++ b/Engine/Plugins/Media/ElectraPlayer/Source/ElectraPlayerRuntime/Private/Runtime/Decoder/DX/AudioDecoderAAC_DX.cpp
@@ -908,7 +908,7 @@ namespace Electra
 				// Get an sample block from the pool.
 				while(!TerminateThreadEvent.IsSignaled())
 				{
-					if (FlushDecoderEvent.IsSignaled() && !bFlush)
+					if (FlushDecoderEvent.IsSignaled())
 					{
 						break;
 					}
