Fixed non-unity/pch errors

diff --git a/Engine/Plugins/Experimental/MeshModelingToolsetExp/Source/MeshModelingToolsExp/Private/MeshGroupPaintTool.cpp b/Engine/Plugins/Experimental/MeshModelingToolsetExp/Source/MeshModelingToolsExp/Private/MeshGroupPaintTool.cpp
--- a/Engine/Plugins/Experimental/MeshModelingToolsetExp/Source/MeshModelingToolsExp/Private/MeshGroupPaintTool.cpp
+++ b/Engine/Plugins/Experimental/MeshModelingToolsetExp/Source/MeshModelingToolsExp/Private/MeshGroupPaintTool.cpp
@@ -1,6 +1,7 @@
 // Copyright Epic Games, Inc. All Rights Reserved.
 
 #include "MeshGroupPaintTool.h"
+#include "Engine/Engine.h"
 #include "Engine/World.h"
 #include "InteractiveToolManager.h"
 #include "InteractiveGizmoManager.h"
diff --git a/Engine/Plugins/Experimental/MeshModelingToolsetExp/Source/MeshModelingToolsExp/Private/MeshInspectorTool.cpp b/Engine/Plugins/Experimental/MeshModelingToolsetExp/Source/MeshModelingToolsExp/Private/MeshInspectorTool.cpp
--- a/Engine/Plugins/Experimental/MeshModelingToolsetExp/Source/MeshModelingToolsExp/Private/MeshInspectorTool.cpp
+++ b/Engine/Plugins/Experimental/MeshModelingToolsetExp/Source/MeshModelingToolsExp/Private/MeshInspectorTool.cpp
@@ -1,6 +1,7 @@
 // Copyright Epic Games, Inc. All Rights Reserved.
 
 #include "MeshInspectorTool.h"
+#include "Engine/Engine.h"
 #include "InteractiveToolManager.h"
 #include "ToolBuilderUtil.h"
 
diff --git a/Engine/Plugins/Media/ElectraPlayer/Source/ElectraPlayerRuntime/Private/Runtime/Decoder/Apple/VideoDecoderH265_Apple.cpp b/Engine/Plugins/Media/ElectraPlayer/Source/ElectraPlayerRuntime/Private/Runtime/Decoder/Apple/VideoDecoderH265_Apple.cpp
--- a/Engine/Plugins/Media/ElectraPlayer/Source/ElectraPlayerRuntime/Private/Runtime/Decoder/Apple/VideoDecoderH265_Apple.cpp
+++ b/Engine/Plugins/Media/ElectraPlayer/Source/ElectraPlayerRuntime/Private/Runtime/Decoder/Apple/VideoDecoderH265_Apple.cpp
@@ -7,6 +7,7 @@
 #include "PlayerCore.h"
 #include "PlayerRuntimeGlobal.h"
 
+#include "Stats/Stats.h"
 #include "StreamAccessUnitBuffer.h"
 #include "Decoder/VideoDecoderH265.h"
 #include "Renderer/RendererBase.h"
diff --git a/Engine/Plugins/Runtime/IOSReplayKit/Source/IOSReplayKit/Private/ReplayKitRecorder.cpp b/Engine/Plugins/Runtime/IOSReplayKit/Source/IOSReplayKit/Private/ReplayKitRecorder.cpp
--- a/Engine/Plugins/Runtime/IOSReplayKit/Source/IOSReplayKit/Private/ReplayKitRecorder.cpp
+++ b/Engine/Plugins/Runtime/IOSReplayKit/Source/IOSReplayKit/Private/ReplayKitRecorder.cpp
@@ -10,6 +10,7 @@
 #include "Misc/Paths.h"
 #include "Engine/Engine.h"
 #include "Engine/GameViewportClient.h"
+#include "UnrealClient.h"
 
 @implementation ReplayKitRecorder
 
diff --git a/Engine/Source/Runtime/IOS/IOSPlatformFeatures/Private/IOSGamepadUtils.cpp b/Engine/Source/Runtime/IOS/IOSPlatformFeatures/Private/IOSGamepadUtils.cpp
--- a/Engine/Source/Runtime/IOS/IOSPlatformFeatures/Private/IOSGamepadUtils.cpp
+++ b/Engine/Source/Runtime/IOS/IOSPlatformFeatures/Private/IOSGamepadUtils.cpp
@@ -4,6 +4,7 @@
 #include "GameDelegates.h"
 #include "IOS/IOSApplication.h"
 #include "IOS/IOSPlatformApplicationMisc.h"
+#include "TextureResource.h"
 
 DEFINE_LOG_CATEGORY_STATIC(LogIOSGamepadUtils, Log, All);
 
