Fix for test client compile errors

diff --git a/Engine/Plugins/Media/ElectraPlayer/Source/ElectraPlayerRuntime/Public/ElectraPlayerMisc.h b/Engine/Plugins/Media/ElectraPlayer/Source/ElectraPlayerRuntime/Public/ElectraPlayerMisc.h
--- a/Engine/Plugins/Media/ElectraPlayer/Source/ElectraPlayerRuntime/Public/ElectraPlayerMisc.h
+++ b/Engine/Plugins/Media/ElectraPlayer/Source/ElectraPlayerRuntime/Public/ElectraPlayerMisc.h
@@ -4,7 +4,7 @@
 
 #include "Logging/LogMacros.h"
 #include "ProfilingDebugging/CsvProfiler.h"
-#include "Stats/Stats2.h"
+#include "Stats/Stats.h"
 
 /** Log category for the ElectraPlayer module. */
 DECLARE_LOG_CATEGORY_EXTERN(LogElectraPlayer, Log, All);
