ElectraPlayer: Added custom HTTP user agent

#rnx

diff --git a/Engine/Plugins/Media/ElectraPlayer/Source/ElectraPlayerRuntime/Private/Runtime/HTTP/HTTPManager.cpp b/Engine/Plugins/Media/ElectraPlayer/Source/ElectraPlayerRuntime/Private/Runtime/HTTP/HTTPManager.cpp
--- a/Engine/Plugins/Media/ElectraPlayer/Source/ElectraPlayerRuntime/Private/Runtime/HTTP/HTTPManager.cpp
+++ b/Engine/Plugins/Media/ElectraPlayer/Source/ElectraPlayerRuntime/Private/Runtime/HTTP/HTTPManager.cpp
@@ -663,7 +663,14 @@ namespace Electra
 				Handle->HttpRequest->EnableTimingTraces();
 			}
 			Handle->HttpRequest->SetURL(Request->Parameters.URL);
-			Handle->HttpRequest->SetUserAgent(ELECTRA_HTTPMANAGER_USER_AGENT);
+			if (Request->Parameters.UserAgent.IsSet())
+			{
+				Handle->HttpRequest->SetUserAgent(Request->Parameters.UserAgent.Value());
+			}
+			else
+			{
+				Handle->HttpRequest->SetUserAgent(ELECTRA_HTTPMANAGER_USER_AGENT);
+			}
 			Handle->HttpRequest->AllowCompression(!Request->Parameters.AcceptEncoding.GetWithDefault(TEXT("")).Equals(TEXT("identity")));
 			#if defined(ELECTRA_HTTPMANAGER_ALLOW_UNSAFE_CONNECTIONS) && ELECTRA_HTTPMANAGER_ALLOW_UNSAFE_CONNECTIONS
 				Handle->HttpRequest->AllowUnsafeRequestsForDebugging();
diff --git a/Engine/Plugins/Media/ElectraPlayer/Source/ElectraPlayerRuntime/Private/Runtime/HTTP/HTTPManager.h b/Engine/Plugins/Media/ElectraPlayer/Source/ElectraPlayerRuntime/Private/Runtime/HTTP/HTTPManager.h
--- a/Engine/Plugins/Media/ElectraPlayer/Source/ElectraPlayerRuntime/Private/Runtime/HTTP/HTTPManager.h
+++ b/Engine/Plugins/Media/ElectraPlayer/Source/ElectraPlayerRuntime/Private/Runtime/HTTP/HTTPManager.h
@@ -237,6 +237,7 @@ namespace Electra
 
 			FString								URL;							//!< URL
 			FString								Verb;							//!< GET (default if not set), HEAD, OPTIONS,....
+			TMediaOptionalValue<FString>		UserAgent;						//!< Custom user agent string
 			ElectraHTTPStream::FHttpRange		Range;							//!< Optional request range
 			TArray<HTTP::FHTTPHeader>			RequestHeaders;					//!< Request headers
 			TMediaOptionalValue<FString>		AcceptEncoding;					//!< Optional accepted encoding
