ElectraPlayer: Checking with CDM system if scheme is supported

#rnx

diff --git a/Engine/Plugins/Media/ElectraPlayer/Source/ElectraPlayerRuntime/Private/Runtime/Player/DASH/ManifestBuilderDASH.cpp b/Engine/Plugins/Media/ElectraPlayer/Source/ElectraPlayerRuntime/Private/Runtime/Player/DASH/ManifestBuilderDASH.cpp
--- a/Engine/Plugins/Media/ElectraPlayer/Source/ElectraPlayerRuntime/Private/Runtime/Player/DASH/ManifestBuilderDASH.cpp
+++ b/Engine/Plugins/Media/ElectraPlayer/Source/ElectraPlayerRuntime/Private/Runtime/Player/DASH/ManifestBuilderDASH.cpp
@@ -2075,6 +2075,11 @@ bool FManifestDASHInternal::CanUseEncryptedAdaptation(const TSharedPtrTS<FAdapta
 			if (DRMCapabilities.IsValid())
 			{
 				ElectraCDM::IMediaCDMCapabilities::ESupportResult Result;
+				Result = DRMCapabilities->SupportsCipher(InAdaptationSet->CommonEncryptionScheme);
+				if (Result != ElectraCDM::IMediaCDMCapabilities::ESupportResult::Supported)
+				{
+					continue;
+				}
 				Result = DRMCapabilities->SupportsType(Mime);
 				if (Result != ElectraCDM::IMediaCDMCapabilities::ESupportResult::Supported)
 				{
