ElectraPlayer: Added method to fetch a custom playlist prior to opening the media

#rnx

diff --git a/Engine/Plugins/Media/ElectraPlayer/Source/ElectraPlayerPlugin/Private/ElectraPlayerPlugin.cpp b/Engine/Plugins/Media/ElectraPlayer/Source/ElectraPlayerPlugin/Private/ElectraPlayerPlugin.cpp
--- a/Engine/Plugins/Media/ElectraPlayer/Source/ElectraPlayerPlugin/Private/ElectraPlayerPlugin.cpp
+++ b/Engine/Plugins/Media/ElectraPlayer/Source/ElectraPlayerPlugin/Private/ElectraPlayerPlugin.cpp
@@ -224,8 +224,11 @@ public:
 
 //-----------------------------------------------------------------------------
 
+void FElectraPlayerPlugin::FPlayerAdapterDelegate::BlobReceived(const TSharedPtr<TArray<uint8>, ESPMode::ThreadSafe>& InBlobData, IElectraPlayerAdapterDelegate::EBlobResultType InResultType, int32 InResultCode, const Electra::FParamDict* InExtraInfo)
+{
+}
 
-Electra::FVariantValue FElectraPlayerPlugin::FPlayerAdapterDelegate::QueryOptions(EOptionType Type, const Electra::FVariantValue & Param)
+Electra::FVariantValue FElectraPlayerPlugin::FPlayerAdapterDelegate::QueryOptions(EOptionType Type, const Electra::FVariantValue& Param)
 {
 	TSharedPtr<FElectraPlayerPlugin, ESPMode::ThreadSafe> PinnedHost = Host.Pin();
 	if (PinnedHost.IsValid())
@@ -626,7 +629,7 @@ bool FElectraPlayerPlugin::Open(const FString& Url, const IMediaOptions* Options
 		LocalPlaystartOptions.MaxBandwidthForStreaming = (int32)MaxBandwidthForStreaming;
 	}
 
-	return Player->OpenInternal(Url, PlayerOptions, LocalPlaystartOptions);
+	return Player->OpenInternal(Url, PlayerOptions, LocalPlaystartOptions, IElectraPlayerInterface::EOpenType::Media);
 }
 
 //-----------------------------------------------------------------------------
diff --git a/Engine/Plugins/Media/ElectraPlayer/Source/ElectraPlayerPlugin/Private/ElectraPlayerPlugin.h b/Engine/Plugins/Media/ElectraPlayer/Source/ElectraPlayerPlugin/Private/ElectraPlayerPlugin.h
--- a/Engine/Plugins/Media/ElectraPlayer/Source/ElectraPlayerPlugin/Private/ElectraPlayerPlugin.h
+++ b/Engine/Plugins/Media/ElectraPlayer/Source/ElectraPlayerPlugin/Private/ElectraPlayerPlugin.h
@@ -157,6 +157,7 @@ private:
 		FPlayerAdapterDelegate(const TSharedPtr<FElectraPlayerPlugin, ESPMode::ThreadSafe> & InHost) : Host(InHost) {}
 
 		Electra::FVariantValue QueryOptions(EOptionType Type, const Electra::FVariantValue & Param) override;
+		void BlobReceived(const TSharedPtr<TArray<uint8>, ESPMode::ThreadSafe>& InBlobData, IElectraPlayerAdapterDelegate::EBlobResultType InResultType, int32 InResultCode, const Electra::FParamDict* InExtraInfo) override;
 		void SendMediaEvent(EPlayerEvent Event) override;
 		void OnVideoFlush() override;
 		void OnAudioFlush() override;
diff --git a/Engine/Plugins/Media/ElectraPlayer/Source/ElectraPlayerRuntime/Private/ElectraPlayer.cpp b/Engine/Plugins/Media/ElectraPlayer/Source/ElectraPlayerRuntime/Private/ElectraPlayer.cpp
--- a/Engine/Plugins/Media/ElectraPlayer/Source/ElectraPlayerRuntime/Private/ElectraPlayer.cpp
+++ b/Engine/Plugins/Media/ElectraPlayer/Source/ElectraPlayerRuntime/Private/ElectraPlayer.cpp
@@ -218,12 +218,23 @@ void FElectraPlayer::ClearToDefaultState()
 /**
  *	Open player
  */
-bool FElectraPlayer::OpenInternal(const FString& Url, const FParamDict& PlayerOptions, const FPlaystartOptions& InPlaystartOptions)
+bool FElectraPlayer::OpenInternal(const FString& Url, const FParamDict& InPlayerOptions, const FPlaystartOptions& InPlaystartOptions, EOpenType InOpenType)
 {
 	LLM_SCOPE(ELLMTag::ElectraPlayer);
 	CSV_EVENT(ElectraPlayer, TEXT("Open"));
 
-	CloseInternal(false);
+	// Open the provided URL as a media or a blob?
+	FString BlobParams;
+	bool bPreviousOpenLoadedBlob = PendingBlobRequest.IsValid();
+	PendingBlobRequest.Reset();
+	bool bCreateNewPlayer = (InOpenType == IElectraPlayerInterface::EOpenType::Media && !bPreviousOpenLoadedBlob) ||
+							(InOpenType == IElectraPlayerInterface::EOpenType::Blob);
+	if (bCreateNewPlayer)
+	{
+		CloseInternal(false);
+	}
+	TSharedPtr<FInternalPlayerImpl, ESPMode::ThreadSafe> NewPlayer = MoveTemp(CurrentPlayer);
+
 	// Clear out our work variables
 	ClearToDefaultState();
 	bAllowKillAfterCloseEvent = false;
@@ -240,74 +251,111 @@ bool FElectraPlayer::OpenInternal(const FString& Url, const FParamDict& PlayerOp
 
 	PlaystartOptions = InPlaystartOptions;
 
-	// Create a new empty player structure. This contains the actual player instance, its associated renderers and sample queues.
-	TSharedPtr<FInternalPlayerImpl, ESPMode::ThreadSafe> NewPlayer = MakeShared<FInternalPlayerImpl, ESPMode::ThreadSafe>();
-
 	// Get a writable copy of the URL so we can sanitize it if necessary.
 	MediaUrl = Url;
 	MediaUrl.TrimStartAndEndInline();
-	UE_LOG(LogElectraPlayer, Log, TEXT("[%p][%p] IMediaPlayer::Open(%s)"), this, NewPlayer.Get(), *SanitizeMessage(MediaUrl));
 
-	// Create the renderers so we can pass them to the internal player.
-	// They get a pointer to ourselves which they will call On[Video|Audio]Decoded() and On[Video|Audio]Flush() on.
-	NewPlayer->RendererVideo = MakeShared<FElectraRendererVideo, ESPMode::ThreadSafe>(SharedThis(this));
-	NewPlayer->RendererAudio = MakeShared<FElectraRendererAudio, ESPMode::ThreadSafe>(SharedThis(this));
+	if (!NewPlayer.IsValid())
+	{
+		FParamDict PlayerOptions(InPlayerOptions);
 
-	// Create the internal player and register ourselves as metrics receiver and static resource provider.
-	IAdaptiveStreamingPlayer::FCreateParam CreateParams;
-	CreateParams.VideoRenderer = NewPlayer->RendererVideo;
-	CreateParams.AudioRenderer = NewPlayer->RendererAudio;
-	CreateParams.ExternalPlayerGUID = PlayerGuid;
-	NewPlayer->AdaptivePlayer = IAdaptiveStreamingPlayer::Create(CreateParams);
-	NewPlayer->AdaptivePlayer->AddMetricsReceiver(this);
-	NewPlayer->AdaptivePlayer->SetStaticResourceProviderCallback(StaticResourceProvider);
-	NewPlayer->AdaptivePlayer->SetVideoDecoderResourceDelegate(VideoDecoderResourceDelegate);
+		// Create a new empty player structure. This contains the actual player instance, its associated renderers and sample queues.
+		NewPlayer = MakeShared<FInternalPlayerImpl, ESPMode::ThreadSafe>();
 
-	// Create the subtitle receiver and register it with the player.
-	MediaPlayerSubtitleReceiver = MakeSharedTS<FSubtitleEventReceiver>();
-	MediaPlayerSubtitleReceiver->GetSubtitleReceivedDelegate().BindRaw(this, &FElectraPlayer::OnSubtitleDecoded);
-	MediaPlayerSubtitleReceiver->GetSubtitleFlushDelegate().BindRaw(this, &FElectraPlayer::OnSubtitleFlush);
-	NewPlayer->AdaptivePlayer->AddSubtitleReceiver(MediaPlayerSubtitleReceiver);
+		// Create the renderers so we can pass them to the internal player.
+		// They get a pointer to ourselves which they will call On[Video|Audio]Decoded() and On[Video|Audio]Flush() on.
+		NewPlayer->RendererVideo = MakeShared<FElectraRendererVideo, ESPMode::ThreadSafe>(SharedThis(this));
+		NewPlayer->RendererAudio = MakeShared<FElectraRendererAudio, ESPMode::ThreadSafe>(SharedThis(this));
 
-	// Create a new media player event receiver and register it to receive all non player internal events as soon as they are received.
-	MediaPlayerEventReceiver = MakeSharedTS<FAEMSEventReceiver>();
-	MediaPlayerEventReceiver->GetEventReceivedDelegate().BindRaw(this, &FElectraPlayer::OnMediaPlayerEventReceived);
-	NewPlayer->AdaptivePlayer->AddAEMSReceiver(MediaPlayerEventReceiver, TEXT("*"), TEXT(""), IAdaptiveStreamingPlayerAEMSReceiver::EDispatchMode::OnReceive);
+		// Create the internal player and register ourselves as metrics receiver and static resource provider.
+		IAdaptiveStreamingPlayer::FCreateParam CreateParams;
+		CreateParams.VideoRenderer = NewPlayer->RendererVideo;
+		CreateParams.AudioRenderer = NewPlayer->RendererAudio;
+		CreateParams.ExternalPlayerGUID = PlayerGuid;
+		NewPlayer->AdaptivePlayer = IAdaptiveStreamingPlayer::Create(CreateParams);
+		NewPlayer->AdaptivePlayer->AddMetricsReceiver(this);
+		NewPlayer->AdaptivePlayer->SetStaticResourceProviderCallback(StaticResourceProvider);
+		NewPlayer->AdaptivePlayer->SetVideoDecoderResourceDelegate(VideoDecoderResourceDelegate);
 
-	NewPlayer->AdaptivePlayer->Initialize(PlayerOptions);
+		// Create the subtitle receiver and register it with the player.
+		MediaPlayerSubtitleReceiver = MakeSharedTS<FSubtitleEventReceiver>();
+		MediaPlayerSubtitleReceiver->GetSubtitleReceivedDelegate().BindRaw(this, &FElectraPlayer::OnSubtitleDecoded);
+		MediaPlayerSubtitleReceiver->GetSubtitleFlushDelegate().BindRaw(this, &FElectraPlayer::OnSubtitleFlush);
+		NewPlayer->AdaptivePlayer->AddSubtitleReceiver(MediaPlayerSubtitleReceiver);
 
-	// Check for options that can be changed during playback and apply them at startup already.
-	// If a media source supports the MaxResolutionForMediaStreaming option then we can override the max resolution.
-	if (PlaystartOptions.MaxVerticalStreamResolution.IsSet())
-	{
-		NewPlayer->AdaptivePlayer->SetMaxResolution(0, PlaystartOptions.MaxVerticalStreamResolution.GetValue());
-	}
+		// Create a new media player event receiver and register it to receive all non player internal events as soon as they are received.
+		MediaPlayerEventReceiver = MakeSharedTS<FAEMSEventReceiver>();
+		MediaPlayerEventReceiver->GetEventReceivedDelegate().BindRaw(this, &FElectraPlayer::OnMediaPlayerEventReceived);
+		NewPlayer->AdaptivePlayer->AddAEMSReceiver(MediaPlayerEventReceiver, TEXT("*"), TEXT(""), IAdaptiveStreamingPlayerAEMSReceiver::EDispatchMode::OnReceive);
 
-	if (PlaystartOptions.MaxBandwidthForStreaming.IsSet())
-	{
-		NewPlayer->AdaptivePlayer->SetBitrateCeiling(PlaystartOptions.MaxBandwidthForStreaming.GetValue());
+		if (InOpenType == IElectraPlayerInterface::EOpenType::Blob)
+		{
+			const TCHAR * const KeyBlob = TEXT("blobparams");
+			if (PlayerOptions.HaveKey(KeyBlob))
+			{
+				BlobParams = PlayerOptions.GetValue(KeyBlob).SafeGetFString();
+				PlayerOptions.Remove(KeyBlob);
+			}
+		}
+		NewPlayer->AdaptivePlayer->Initialize(PlayerOptions);
 	}
 
-	// Set the player member variable to the new player now that it has been fully initialized.
-	CurrentPlayer = MoveTemp(NewPlayer);
-	// Apply options that may have been set prior to calling Open().
-	// Set these only if they have defined values as to not override what might have been set in the PlayerOptions.
-	if (bFrameAccurateSeeking.IsSet())
+	if (InOpenType == IElectraPlayerInterface::EOpenType::Media)
 	{
-		SetFrameAccurateSeekMode(bFrameAccurateSeeking.GetValue());
-	}
-	if (bEnableLooping.IsSet())
-	{
-		SetLooping(bEnableLooping.GetValue());
+		// Check for options that can be changed during playback and apply them at startup already.
+		// If a media source supports the MaxResolutionForMediaStreaming option then we can override the max resolution.
+		if (PlaystartOptions.MaxVerticalStreamResolution.IsSet())
+		{
+			NewPlayer->AdaptivePlayer->SetMaxResolution(0, PlaystartOptions.MaxVerticalStreamResolution.GetValue());
+		}
+
+		if (PlaystartOptions.MaxBandwidthForStreaming.IsSet())
+		{
+			NewPlayer->AdaptivePlayer->SetBitrateCeiling(PlaystartOptions.MaxBandwidthForStreaming.GetValue());
+		}
+
+		// Apply options that may have been set prior to calling Open().
+		// Set these only if they have defined values as to not override what might have been set in the PlayerOptions.
+		if (bFrameAccurateSeeking.IsSet())
+		{
+			SetFrameAccurateSeekMode(bFrameAccurateSeeking.GetValue());
+		}
+		if (bEnableLooping.IsSet())
+		{
+			SetLooping(bEnableLooping.GetValue());
+		}
+		if (CurrentPlaybackRange.Start.IsSet() || CurrentPlaybackRange.End.IsSet())
+		{
+			SetPlaybackRange(CurrentPlaybackRange);
+		}
+
+		// Set the player member variable to the new player now that it has been fully initialized.
+		CurrentPlayer = MoveTemp(NewPlayer);
+
+		if (bPreviousOpenLoadedBlob)
+		{
+			CurrentPlayer->AdaptivePlayer->ModifyOptions(InPlayerOptions, Electra::FParamDict());
+		}
+
+		// Issue load of the playlist.
+		UE_LOG(LogElectraPlayer, Log, TEXT("[%p][%p] IMediaPlayer::Open(%s)"), this, CurrentPlayer.Get(), *SanitizeMessage(MediaUrl));
+		CurrentPlayer->AdaptivePlayer->LoadManifest(MediaUrl);
 	}
-	if (CurrentPlaybackRange.Start.IsSet() || CurrentPlaybackRange.End.IsSet())
+	else
 	{
-		SetPlaybackRange(CurrentPlaybackRange);
+		PendingBlobRequest = MakeShared<FBlobRequest, ESPMode::ThreadSafe>();
+		if (!PendingBlobRequest->Request->SetFromJSON(BlobParams))
+		{
+			UE_LOG(LogElectraPlayer, Error, TEXT("[%p] IMediaPlayer::OpenBlob(%s) has bad JSON parameters"), this, *SanitizeMessage(MediaUrl));
+			PendingBlobRequest.Reset();
+			return false;
+		}
+		CurrentPlayer = MoveTemp(NewPlayer);
+		UE_LOG(LogElectraPlayer, Log, TEXT("[%p][%p] IMediaPlayer::OpenBlob(%s)"), this, CurrentPlayer.Get(), *SanitizeMessage(MediaUrl));
+		PendingBlobRequest->Request->URL(MediaUrl).Callback().BindThreadSafeSP(PendingBlobRequest.ToSharedRef(), &FBlobRequest::OnBlobRequestComplete);
+		CurrentPlayer->AdaptivePlayer->LoadBlob(PendingBlobRequest->Request);
 	}
-
-	// Issue load of the playlist.
-	CurrentPlayer->AdaptivePlayer->LoadManifest(MediaUrl);
- 	return true;
+	return true;
 }
 
 //-----------------------------------------------------------------------------
@@ -558,6 +606,9 @@ void FElectraPlayer::Tick(FTimespan DeltaTime, FTimespan Timecode)
 		// Handle static resource fetch requests.
 		StaticResourceProvider->ProcessPendingStaticResourceRequests();
 
+		// Check for blob loading completed
+		HandleBlobDownload();
+
 		// Check for option changes
 		TSharedPtr<IElectraPlayerAdapterDelegate, ESPMode::ThreadSafe> PinnedAdapterDelegate = AdapterDelegate.Pin();
 		if (PinnedAdapterDelegate.IsValid())
@@ -622,6 +673,33 @@ void FElectraPlayer::Tick(FTimespan DeltaTime, FTimespan Timecode)
 	}
 }
 
+
+void FElectraPlayer::HandleBlobDownload()
+{
+	if (PendingBlobRequest.IsValid() && PendingBlobRequest->bIsComplete && !PendingBlobRequest->bDispatched)
+	{
+		PendingBlobRequest->bDispatched = true;
+		if (!PendingBlobRequest->Request->GetWasCanceled())
+		{
+			int32 ErrCode = PendingBlobRequest->Request->GetError();
+			IElectraPlayerAdapterDelegate::EBlobResultType Result = ErrCode == 0 ? IElectraPlayerAdapterDelegate::EBlobResultType::Success :
+																	ErrCode > 0 && ErrCode < 100 ? IElectraPlayerAdapterDelegate::EBlobResultType::TimedOut :
+																	IElectraPlayerAdapterDelegate::EBlobResultType::HttpFailure;
+			TSharedPtr<TArray<uint8>, ESPMode::ThreadSafe> BlobData = MakeShared<TArray<uint8>, ESPMode::ThreadSafe>();
+			TSharedPtrTS<IElectraHttpManager::FReceiveBuffer> ResponseBuffer = PendingBlobRequest->Request->GetResponseBuffer();
+			if (ResponseBuffer.IsValid())
+			{
+				BlobData->Append((const uint8*)ResponseBuffer->Buffer.GetLinearReadData(), ResponseBuffer->Buffer.Num());
+			}
+			TSharedPtr<IElectraPlayerAdapterDelegate, ESPMode::ThreadSafe> PinnedAdapterDelegate = AdapterDelegate.Pin();
+			if (PinnedAdapterDelegate.IsValid())
+			{
+				PinnedAdapterDelegate->BlobReceived(BlobData, Result, ErrCode, nullptr);
+			}
+		}
+	}
+}
+
 //-----------------------------------------------------------------------------
 /**
  * The video renderer is adding a buffer to the queue
@@ -2999,7 +3077,7 @@ void FElectraPlayer::FAdaptiveStreamingPlayerResourceProvider::ProcessPendingSta
 						UE_LOG(LogElectraPlayer, Log, TEXT("[%p][%p] IMediaPlayer::Open: Source provided playlist data for '%s'"), this, CurrentPlayer.Get(), *SanitizeMessage(InOutRequest->GetResourceURL()));
 						*/
 
-						// FString is Unicode but the HTTP response for an HLS playlist is a UTF-8 string.
+						// FString is Unicode but the HTTP response for a playlist is expected to be a UTF-8 string.
 						// Create a plain array from this.
 						TSharedPtr<TArray<uint8>, ESPMode::ThreadSafe> ResponseDataPtr = MakeShared<TArray<uint8>, ESPMode::ThreadSafe>((const uint8*)TCHAR_TO_UTF8(*PlaylistData), PlaylistData.Len());
 						// And put it into the request
diff --git a/Engine/Plugins/Media/ElectraPlayer/Source/ElectraPlayerRuntime/Private/ElectraPlayer.h b/Engine/Plugins/Media/ElectraPlayer/Source/ElectraPlayerRuntime/Private/ElectraPlayer.h
--- a/Engine/Plugins/Media/ElectraPlayer/Source/ElectraPlayerRuntime/Private/ElectraPlayer.h
+++ b/Engine/Plugins/Media/ElectraPlayer/Source/ElectraPlayerRuntime/Private/ElectraPlayer.h
@@ -87,7 +87,7 @@ public:
 
 	// -------- PlayerAdapter (Plugin/Native) API
 
-	bool OpenInternal(const FString& Url, const FParamDict& PlayerOptions, const FPlaystartOptions& InPlaystartOptions) override;
+	bool OpenInternal(const FString& Url, const FParamDict& PlayerOptions, const FPlaystartOptions& InPlaystartOptions, EOpenType InOpenType) override;
 	void CloseInternal(bool bKillAfterClose) override;
 
 	void Tick(FTimespan DeltaTime, FTimespan Timecode) override;
@@ -505,6 +505,21 @@ private:
 	TSharedPtr<FAdaptiveStreamingPlayerResourceProvider, ESPMode::ThreadSafe> StaticResourceProvider;
 	TSharedPtr<IVideoDecoderResourceDelegate, ESPMode::ThreadSafe> VideoDecoderResourceDelegate;
 
+	class FBlobRequest
+	{
+	public:
+		FBlobRequest() : Request(MakeShared<Electra::FHTTPResourceRequest, ESPMode::ThreadSafe>())
+		{ }
+		void OnBlobRequestComplete(TSharedPtrTS<Electra::FHTTPResourceRequest> InRequest)
+		{
+			bIsComplete = true;
+		}
+		TSharedPtr<Electra::FHTTPResourceRequest, ESPMode::ThreadSafe> Request;
+		bool bIsComplete = false;
+		bool bDispatched = false;
+	};
+	TSharedPtr<FBlobRequest, ESPMode::ThreadSafe> PendingBlobRequest;
+
 
 	class FAverageValue
 	{
@@ -754,6 +769,7 @@ private:
 	/** Sequential analytics event number. Helps sorting events. **/
 	uint32									AnalyticsInstanceEventCount;
 
+	void HandleBlobDownload();
 
 	void HandleDeferredPlayerEvents();
 	void HandlePlayerEventOpenSource(const FString& URL);
diff --git a/Engine/Plugins/Media/ElectraPlayer/Source/ElectraPlayerRuntime/Private/Runtime/Player/AdaptiveStreamingPlayer.cpp b/Engine/Plugins/Media/ElectraPlayer/Source/ElectraPlayerRuntime/Private/Runtime/Player/AdaptiveStreamingPlayer.cpp
--- a/Engine/Plugins/Media/ElectraPlayer/Source/ElectraPlayerRuntime/Private/Runtime/Player/AdaptiveStreamingPlayer.cpp
+++ b/Engine/Plugins/Media/ElectraPlayer/Source/ElectraPlayerRuntime/Private/Runtime/Player/AdaptiveStreamingPlayer.cpp
@@ -406,7 +406,10 @@ void FAdaptiveStreamingPlayer::DispatchEvent(TSharedPtrTS<FMetricEvent> Event)
 {
 	if (EventDispatcher.IsValid())
 	{
-		Event->Player = SharedThis(this);
+		if (DoesSharedInstanceExist())
+		{
+			Event->Player = SharedThis(this);
+		}
 		EventDispatcher->DispatchEvent(Event);
 	}
 }
@@ -575,13 +578,18 @@ void FAdaptiveStreamingPlayer::GetPlaybackRange(FPlaybackRange& OutPlaybackRange
 
 //-----------------------------------------------------------------------------
 /**
- * Starts loading the manifest / master playlist.
+ * Starts loading the manifest / main playlist.
  */
 void FAdaptiveStreamingPlayer::LoadManifest(const FString& InManifestURL)
 {
 	WorkerThread.SendLoadManifestMessage(InManifestURL, FString());
 }
 
+void FAdaptiveStreamingPlayer::LoadBlob(TSharedPtr<FHTTPResourceRequest, ESPMode::ThreadSafe> InBlobLoadRequest)
+{
+	WorkerThread.SendLoadBlobMessage(InBlobLoadRequest);
+}
+
 
 //-----------------------------------------------------------------------------
 /**
@@ -1332,6 +1340,14 @@ bool FAdaptiveStreamingPlayer::InternalHandleThreadMessages()
 				InternalLoadManifest(msg.Data.ManifestToLoad.URL, msg.Data.ManifestToLoad.MimeType);
 				break;
 			}
+			case FWorkerThreadMessages::FMessage::EType::LoadBlob:
+			{
+				if (msg.Data.BlobToLoad.BlobLoadRequest.IsValid())
+				{
+					msg.Data.BlobToLoad.BlobLoadRequest->StartGet(this);
+				}
+				break;
+			}
 			case FWorkerThreadMessages::FMessage::EType::Pause:
 			{
 				bShouldBePaused = true;
diff --git a/Engine/Plugins/Media/ElectraPlayer/Source/ElectraPlayerRuntime/Private/Runtime/Player/AdaptiveStreamingPlayer.h b/Engine/Plugins/Media/ElectraPlayer/Source/ElectraPlayerRuntime/Private/Runtime/Player/AdaptiveStreamingPlayer.h
--- a/Engine/Plugins/Media/ElectraPlayer/Source/ElectraPlayerRuntime/Private/Runtime/Player/AdaptiveStreamingPlayer.h
+++ b/Engine/Plugins/Media/ElectraPlayer/Source/ElectraPlayerRuntime/Private/Runtime/Player/AdaptiveStreamingPlayer.h
@@ -39,7 +39,7 @@ public:
 	//-------------------------------------------------------------------------
 	// Various resource providers/delegates
 	//
-	virtual void SetStaticResourceProviderCallback(const TSharedPtr<IAdaptiveStreamingPlayerResourceProvider, ESPMode::ThreadSafe> & InStaticResourceProvider) = 0;
+	virtual void SetStaticResourceProviderCallback(const TSharedPtr<IAdaptiveStreamingPlayerResourceProvider, ESPMode::ThreadSafe>& InStaticResourceProvider) = 0;
 	virtual void SetVideoDecoderResourceDelegate(const TSharedPtr<IVideoDecoderResourceDelegate, ESPMode::ThreadSafe>& ResourceDelegate) = 0;
 
 
@@ -100,7 +100,10 @@ public:
 	//-------------------------------------------------------------------------
 	// Manifest loading functions
 	//
-	//! Issue a load and parse of the manifest/master playlist file. Make initial stream selection choice by calling SetInitialStreamAttributes() beforehand.
+	//! Issues loading of a binary blob of arbitrary data.
+	virtual void LoadBlob(TSharedPtr<FHTTPResourceRequest, ESPMode::ThreadSafe> InBlobLoadRequest) = 0;
+
+	//! Issue a load and parse of the manifest/main playlist file. Make initial stream selection choice by calling SetInitialStreamAttributes() beforehand.
 	virtual void LoadManifest(const FString& ManifestURL) = 0;
 
 
diff --git a/Engine/Plugins/Media/ElectraPlayer/Source/ElectraPlayerRuntime/Private/Runtime/Player/AdaptiveStreamingPlayerInternal.h b/Engine/Plugins/Media/ElectraPlayer/Source/ElectraPlayerRuntime/Private/Runtime/Player/AdaptiveStreamingPlayerInternal.h
--- a/Engine/Plugins/Media/ElectraPlayer/Source/ElectraPlayerRuntime/Private/Runtime/Player/AdaptiveStreamingPlayerInternal.h
+++ b/Engine/Plugins/Media/ElectraPlayer/Source/ElectraPlayerRuntime/Private/Runtime/Player/AdaptiveStreamingPlayerInternal.h
@@ -925,6 +925,7 @@ public:
 
 	void SetInitialStreamAttributes(EStreamType StreamType, const FStreamSelectionAttributes& InitialSelection) override;
 	void EnableFrameAccurateSeeking(bool bEnabled) override;
+	void LoadBlob(TSharedPtr<FHTTPResourceRequest, ESPMode::ThreadSafe> InBlobLoadRequest) override;
 	void LoadManifest(const FString& manifestURL) override;
 
 	void SeekTo(const FSeekParam& NewPosition) override;
@@ -1356,6 +1357,7 @@ private:
 				// Commands
 				Initialize,
 				LoadManifest,
+				LoadBlob,
 				Pause,
 				Resume,
 				Loop,
@@ -1383,6 +1385,10 @@ private:
 					FString											URL;
 					FString											MimeType;
 				};
+				struct FLoadBlob
+				{
+					TSharedPtrTS<FHTTPResourceRequest>				BlobLoadRequest;
+				};
 				struct FStreamReader
 				{
 					FStreamReader()
@@ -1443,6 +1449,7 @@ private:
 				};
 
 				FLoadManifest				ManifestToLoad;
+				FLoadBlob					BlobToLoad;
 				FStreamReader				StreamReader;
 				FEvent						MediaEvent;
 				FLoop						Looping;
@@ -1491,6 +1498,13 @@ private:
 			Msg.Data.ManifestToLoad.MimeType = MimeType;
 			TriggerSharedWorkerThread(MoveTemp(Msg));
 		}
+		void SendLoadBlobMessage(TSharedPtr<FHTTPResourceRequest, ESPMode::ThreadSafe> InBlobLoadRequest)
+		{
+			FMessage Msg;
+			Msg.Type = FMessage::EType::LoadBlob;
+			Msg.Data.BlobToLoad.BlobLoadRequest = MoveTemp(InBlobLoadRequest);
+			TriggerSharedWorkerThread(MoveTemp(Msg));
+		}
 		void SendPauseMessage()
 		{
 			FMessage Msg;
diff --git a/Engine/Plugins/Media/ElectraPlayer/Source/ElectraPlayerRuntime/Private/Runtime/Player/AdaptiveStreamingPlayerResourceRequest.cpp b/Engine/Plugins/Media/ElectraPlayer/Source/ElectraPlayerRuntime/Private/Runtime/Player/AdaptiveStreamingPlayerResourceRequest.cpp
--- a/Engine/Plugins/Media/ElectraPlayer/Source/ElectraPlayerRuntime/Private/Runtime/Player/AdaptiveStreamingPlayerResourceRequest.cpp
+++ b/Engine/Plugins/Media/ElectraPlayer/Source/ElectraPlayerRuntime/Private/Runtime/Player/AdaptiveStreamingPlayerResourceRequest.cpp
@@ -3,11 +3,14 @@
 #include "AdaptiveStreamingPlayerResourceRequest.h"
 #include "Player/PlayerSessionServices.h"
 
-namespace Electra
-{
-
+#include <Dom/JsonObject.h>
+#include <Serialization/JsonReader.h>
+#include <Serialization/JsonSerializer.h>
+#include <Misc/Base64.h>
 
 
+namespace Electra
+{
 
 FHTTPResourceRequest::FHTTPResourceRequest()
 {
@@ -26,21 +29,91 @@ FHTTPResourceRequest::~FHTTPResourceRequest()
 	Cancel();
 	if (bWasAdded && Request.IsValid())
 	{
-		TSharedPtrTS<IElectraHttpManager> HTTPManager = PlayerSessionServices->GetHTTPManager();
-		if (HTTPManager.IsValid())
+		TSharedPtrTS<IElectraHttpManager> PinnedHTTPManager = HTTPManager.Pin();
+		if (PinnedHTTPManager.IsValid())
 		{
-			HTTPManager->RemoveRequest(Request, true);
+			PinnedHTTPManager->RemoveRequest(Request, true);
 		}
 	}
 	Request.Reset();
 }
 
-void FHTTPResourceRequest::StartGet(IPlayerSessionServices* InPlayerSessionServices)
+bool FHTTPResourceRequest::SetFromJSON(const FString& InJSONParams)
 {
-	PlayerSessionServices = InPlayerSessionServices;
+	if (InJSONParams.Len())
+	{
+		TSharedRef<TJsonReader<>> Reader = TJsonReaderFactory<>::Create(InJSONParams);
+		TSharedPtr<FJsonObject> JSONParams;
+		if (FJsonSerializer::Deserialize(Reader, JSONParams))
+		{
+			TArray<FString> StringArray;
+			FString String;
+			// Verb
+			if (JSONParams->TryGetStringField(TEXT("verb"), String))
+			{
+				Verb(String);
+				// If POST see if there is base64 encoded post data to send.
+				if (String.Equals(TEXT("POST")) && JSONParams->TryGetStringField(TEXT("data"), String))
+				{
+					TArray<uint8> Data;
+					if (FBase64::Decode(String, Data))
+					{
+						PostData(Data);
+					}
+				}
+			}
+			else
+			{
+				Verb(TEXT("GET"));
+			}
+
+			// Custom user agent
+			if (JSONParams->TryGetStringField(TEXT("agent"), String))
+			{
+				UserAgent(String);
+			}
+
+			// Accept-encoding
+			if (JSONParams->TryGetStringField(TEXT("encoding"), String))
+			{
+				AcceptEncoding(String);
+			}
+
+			// Headers. Must always be an array of strings
+			if (JSONParams->TryGetStringArrayField(TEXT("hdrs"), StringArray))
+			{
+				Headers(StringArray);
+				StringArray.Empty();
+			}
+
+			// Connection timeout in milliseconds
+			int32 TimeOutMS = 0;
+			if (JSONParams->TryGetNumberField(TEXT("ctoms"), TimeOutMS))
+			{
+				ConnectionTimeout(FTimeValue(FTimeValue::MillisecondsToHNS(TimeOutMS)));
+			}
+
+			// No-data timeout in milliseconds
+			TimeOutMS = 0;
+			if (JSONParams->TryGetNumberField(TEXT("ndtoms"), TimeOutMS))
+			{
+				NoDataTimeout(FTimeValue(FTimeValue::MillisecondsToHNS(TimeOutMS)));
+			}
+		}
+		else
+		{
+			return false;
+		}
+	}
+	return true;
+}
+
 
+void FHTTPResourceRequest::StartGet(IPlayerSessionServices* InPlayerSessionServices)
+{
+	HTTPManager = InPlayerSessionServices->GetHTTPManager();
 	// Is there a static resource provider that we can try?
-	TSharedPtr<IAdaptiveStreamingPlayerResourceProvider, ESPMode::ThreadSafe> StaticResourceProvider = PlayerSessionServices->GetStaticResourceProvider();
+	TSharedPtr<IAdaptiveStreamingPlayerResourceProvider, ESPMode::ThreadSafe> StaticResourceProvider = InPlayerSessionServices->GetStaticResourceProvider();
 	if (StaticQueryType.IsSet() && StaticResourceProvider.IsValid())
 	{
 		TSharedPtr<FStaticResourceRequest, ESPMode::ThreadSafe>	StaticRequest = MakeShared<FStaticResourceRequest, ESPMode::ThreadSafe>(AsShared());
@@ -49,10 +122,10 @@ void FHTTPResourceRequest::StartGet(IPlayerSessionServices* InPlayerSessionServi
 	else
 	{
 		bWasAdded = true;
-		TSharedPtrTS<IElectraHttpManager> HTTPManager = PlayerSessionServices->GetHTTPManager();
-		if (HTTPManager.IsValid())
+		TSharedPtrTS<IElectraHttpManager> PinnedHTTPManager = HTTPManager.Pin();
+		if (PinnedHTTPManager.IsValid())
 		{
-			HTTPManager->AddRequest(Request, false);
+			PinnedHTTPManager->AddRequest(Request, false);
 		}
 	}
 }
@@ -64,10 +137,10 @@ void FHTTPResourceRequest::Cancel()
 	ReceiveBuffer.Reset();
 	if (bWasAdded && Request.IsValid())
 	{
-		TSharedPtrTS<IElectraHttpManager> HTTPManager = PlayerSessionServices->GetHTTPManager();
-		if (HTTPManager.IsValid())
+		TSharedPtrTS<IElectraHttpManager> PinnedHTTPManager = HTTPManager.Pin();
+		if (PinnedHTTPManager.IsValid())
 		{
-			HTTPManager->RemoveRequest(Request, true);
+			PinnedHTTPManager->RemoveRequest(Request, true);
 		}
 		bWasAdded = false;
 	}
@@ -83,10 +156,10 @@ void FHTTPResourceRequest::StaticDataReady()
 		{
 			// Do the actual HTTP request now.
 			bWasAdded = true;
-			TSharedPtrTS<IElectraHttpManager> HTTPManager = PlayerSessionServices->GetHTTPManager();
-			if (HTTPManager.IsValid())
+			TSharedPtrTS<IElectraHttpManager> PinnedHTTPManager = HTTPManager.Pin();
+			if (PinnedHTTPManager.IsValid())
 			{
-				HTTPManager->AddRequest(Request, false);
+				PinnedHTTPManager->AddRequest(Request, false);
 			}
 		}
 		else
diff --git a/Engine/Plugins/Media/ElectraPlayer/Source/ElectraPlayerRuntime/Private/Runtime/Player/AdaptiveStreamingPlayerResourceRequest.h b/Engine/Plugins/Media/ElectraPlayer/Source/ElectraPlayerRuntime/Private/Runtime/Player/AdaptiveStreamingPlayerResourceRequest.h
--- a/Engine/Plugins/Media/ElectraPlayer/Source/ElectraPlayerRuntime/Private/Runtime/Player/AdaptiveStreamingPlayerResourceRequest.h
+++ b/Engine/Plugins/Media/ElectraPlayer/Source/ElectraPlayerRuntime/Private/Runtime/Player/AdaptiveStreamingPlayerResourceRequest.h
@@ -65,12 +65,17 @@ public:
 
 	DECLARE_DELEGATE_OneParam(FOnRequestCompletedCallback, TSharedPtrTS<FHTTPResourceRequest> /*Request*/);
 
+	virtual bool SetFromJSON(const FString& InJSONParams);
+
 	virtual FHTTPResourceRequest& URL(const FString& InURL)
 	{ Request->Parameters.URL = InURL; return *this; }
 
 	virtual FHTTPResourceRequest& Verb(const FString& InVerb)
 	{ Request->Parameters.Verb = InVerb; return *this; }
 
+	virtual FHTTPResourceRequest& UserAgent(const FString& InUserAgent)
+	{ Request->Parameters.UserAgent.Set(InUserAgent); return *this; }
+
 	virtual FHTTPResourceRequest& PostData(TArray<uint8> InPostData)
 	{ Request->Parameters.PostData = MoveTemp(InPostData); return *this; }
 	
@@ -176,7 +181,7 @@ private:
 			}
 		}
 
-		//! Signal request completion. Must be called with ot without data being set.
+		//! Signal request completion. Must be called with or without data being set.
 		virtual void SignalDataReady()
 		{
 			TSharedPtrTS<FHTTPResourceRequest> p(Owner.Pin());
@@ -200,7 +205,7 @@ private:
 	TSharedPtrTS<IElectraHttpManager::FProgressListener> ProgressListener;
 	FOnRequestCompletedCallback CompletedCallback;
 	TWeakPtrTS<IHTTPResourceRequestObject> UserObject;
-	IPlayerSessionServices* PlayerSessionServices = nullptr;
+	TWeakPtrTS<IElectraHttpManager> HTTPManager;
 	TMediaOptionalValue<IAdaptiveStreamingPlayerResourceRequest::EPlaybackResourceType> StaticQueryType;
 	bool bWasAdded = false;
 	bool bWasCanceled = false;
diff --git a/Engine/Plugins/Media/ElectraPlayer/Source/ElectraPlayerRuntime/Public/IElectraPlayerInterface.h b/Engine/Plugins/Media/ElectraPlayer/Source/ElectraPlayerRuntime/Public/IElectraPlayerInterface.h
--- a/Engine/Plugins/Media/ElectraPlayer/Source/ElectraPlayerRuntime/Public/IElectraPlayerInterface.h
+++ b/Engine/Plugins/Media/ElectraPlayer/Source/ElectraPlayerRuntime/Public/IElectraPlayerInterface.h
@@ -42,6 +42,14 @@ public:
 	};
 	virtual Electra::FVariantValue QueryOptions(EOptionType Type, const Electra::FVariantValue& Param = Electra::FVariantValue()) = 0;
 
+	enum class EBlobResultType
+	{
+		Success,
+		TimedOut,
+		HttpFailure
+	};
+	virtual void BlobReceived(const TSharedPtr<TArray<uint8>, ESPMode::ThreadSafe>& InBlobData, EBlobResultType InResultType, int32 InResultCode, const Electra::FParamDict* InExtraInfo) = 0;
+
 	enum class EPlayerEvent
 	{
 		MediaBuffering = 0,
@@ -126,7 +134,13 @@ public:
 		bool						bDoNotPreload = false;
 	};
 
-	virtual bool OpenInternal(const FString& Url, const Electra::FParamDict& PlayerOptions, const FPlaystartOptions& InPlaystartOptions) = 0;
+	enum class EOpenType
+	{
+		Media,
+		Blob
+	};
+
+	virtual bool OpenInternal(const FString& Url, const Electra::FParamDict& PlayerOptions, const FPlaystartOptions& InPlaystartOptions, EOpenType InOpenType) = 0;
 	virtual void CloseInternal(bool bKillAfterClose) = 0;
 
 	virtual void Tick(FTimespan DeltaTime, FTimespan Timecode) = 0;
