ElectraPlayer: Fixed DASH URL resolving of relative paths in BaseURL elements; fixed segment index calculation when PTO is specified; added fallback of time parsing in ISO scheme if response is not formatted as such and is just a number instead.

#rnx

diff --git a/Engine/Plugins/Media/ElectraPlayer/Source/ElectraPlayerRuntime/Private/Runtime/Player/DASH/ManifestBuilderDASH.cpp b/Engine/Plugins/Media/ElectraPlayer/Source/ElectraPlayerRuntime/Private/Runtime/Player/DASH/ManifestBuilderDASH.cpp
--- a/Engine/Plugins/Media/ElectraPlayer/Source/ElectraPlayerRuntime/Private/Runtime/Player/DASH/ManifestBuilderDASH.cpp
+++ b/Engine/Plugins/Media/ElectraPlayer/Source/ElectraPlayerRuntime/Private/Runtime/Player/DASH/ManifestBuilderDASH.cpp
@@ -153,7 +153,7 @@ namespace DASHUrlHelpers
 	 * Collects BaseURL elements beginning at the specified element up the MPD hierarchy.
 	 * At most one BaseURL per level is added to the output. If a preferred service location is specified the BaseURL element matching this will
 	 * be added. If no BaseURL on that level matching the preferred service location the first BaseURL element is used.
-	 * Results are added in reverse hierarchy order, eg. Representation, Adaptation, Period, MPD.
+	 * Results are added in hierarchy order, eg. MPD, Period, Adaptation, Representation.
 	 */
 	void GetAllHierarchyBaseURLs(IPlayerSessionServices* InPlayerSessionServices, TArray<TSharedPtrTS<const FDashMPD_BaseURLType>>& OutBaseURLs, TSharedPtrTS<const IDashMPDElement> StartingAt, const TCHAR* PreferredServiceLocation)
 	{
@@ -164,7 +164,7 @@ namespace DASHUrlHelpers
 			{
 				if (ElementBaseURLs.Num())
 				{
-					OutBaseURLs.Emplace(ElementBaseURLs[0]);
+					OutBaseURLs.Insert(ElementBaseURLs[0], 0);
 				}
 			}
 			else
@@ -173,7 +173,7 @@ namespace DASHUrlHelpers
 				{
 					if (i == 0 || ElementBaseURLs[i]->GetServiceLocation().Equals(PreferredServiceLocation))
 					{
-						OutBaseURLs.Emplace(ElementBaseURLs[i]);
+						OutBaseURLs.Insert(ElementBaseURLs[i], 0);
 						break;
 					}
 				}
@@ -563,8 +563,7 @@ namespace DASHUrlHelpers
 
 
 	/**
-	 * Resolves relative URLs against their parent BaseURL elements in the hierarchy.
-	 * If this does not produce an absolute URL it will finally be resolved against the URL the MPD was loaded from.
+	 * Resolves an element URLs against its parent BaseURL elements in the hierarchy.
 	 * Returns true if an absolute URL could be generated, false if not.
 	 * Since the MPD URL had to be an absolute URL this cannot actually fail.
 	 */
@@ -575,60 +574,39 @@ namespace DASHUrlHelpers
 		bATOComplete.Reset();
 		FURL_RFC3986 UrlParser;
 		FString ElementURL(InElementURL);
-		// If the element URL is empty it is specified as the first entry in the BaseURL array.
-		// We MUST NOT resolve against an empty URL because as per RFC 3986 section 5.2.2 the query string and fragment come from the relative
-		// URL and if we were to start with an empty one we would lose those!
-		int32 nBase = 0;
-		if (ElementURL.IsEmpty())
+
+		// If there is no element URL and BaseURLs this is probably for an MPD update in which case we return the original document URL.
+		if (ElementURL.IsEmpty() && BaseURLs.Num() == 0)
 		{
-			// If there is no element URL and BaseURLs this is probably for an MPD update in which case we return the original document URL.
-			if (BaseURLs.Num() == 0)
-			{
-				OutURL = DocumentURL;
-				return true;
-			}
-			ElementURL = BaseURLs[0]->GetURL();
-			++nBase;
+			OutURL = DocumentURL;
+			return true;
 		}
-		if (UrlParser.Parse(ElementURL))
+
+		// Starting with the document URL apply the <BaseURL> elements down the path to the element.
+		// If at any level an absolute URL is specified it replaces the URL as a whole, otherwise
+		// relative URLs are resolved against the parent.
+		if (!UrlParser.Parse(DocumentURL))
 		{
-			while(1)
+			return false;
+		}
+		for(int32 nBase=0; nBase<BaseURLs.Num(); ++nBase)
+		{
+			UrlParser.ResolveWith(BaseURLs[nBase]->GetURL());
+			SumOfUrlATOs += BaseURLs[nBase]->GetAvailabilityTimeOffset();
+			if (!bATOComplete.IsSet())
 			{
-				if (UrlParser.IsAbsolute())
-				{
-					break;
-				}
-				else
-				{
-					if (nBase < BaseURLs.Num())
-					{
-						UrlParser.ResolveAgainst(BaseURLs[nBase]->GetURL());
-						SumOfUrlATOs += BaseURLs[nBase]->GetAvailabilityTimeOffset();
-						if (!bATOComplete.IsSet())
-						{
-							bATOComplete = BaseURLs[nBase]->GetAvailabilityTimeComplete();
-						}
-						else
-						{
-							if (BaseURLs[nBase]->GetAvailabilityTimeComplete().IsSet() && BaseURLs[nBase]->GetAvailabilityTimeComplete().Value() != bATOComplete.Value())
-							{
-								// Inconsistent availabilityTimeComplete across the hierarchy.
-							}
-						}
-						++nBase;
-					}
-					else
-					{
-						UrlParser.ResolveAgainst(DocumentURL);
-						break;
-					}
-				}
+				bATOComplete = BaseURLs[nBase]->GetAvailabilityTimeComplete();
+			}
+			else if (BaseURLs[nBase]->GetAvailabilityTimeComplete().IsSet() && BaseURLs[nBase]->GetAvailabilityTimeComplete().Value() != bATOComplete.Value())
+			{
+				// Inconsistent availabilityTimeComplete across the hierarchy.
+				// TBD
 			}
-			OutURL = UrlParser.Get();
-			ATO = SumOfUrlATOs;
-			return UrlParser.IsAbsolute();
 		}
-		return false;
+		UrlParser.ResolveWith(InElementURL);
+		OutURL = UrlParser.Get();
+		ATO = SumOfUrlATOs;
+		return UrlParser.IsAbsolute();
 	}
 
 	FString ApplyAnnexEByteRange(FString InURL, FString InRange, const TArray<TSharedPtrTS<const FDashMPD_BaseURLType>>& BaseURLs)
diff --git a/Engine/Plugins/Media/ElectraPlayer/Source/ElectraPlayerRuntime/Private/Runtime/Player/DASH/ManifestDASH.cpp b/Engine/Plugins/Media/ElectraPlayer/Source/ElectraPlayerRuntime/Private/Runtime/Player/DASH/ManifestDASH.cpp
--- a/Engine/Plugins/Media/ElectraPlayer/Source/ElectraPlayerRuntime/Private/Runtime/Player/DASH/ManifestDASH.cpp
+++ b/Engine/Plugins/Media/ElectraPlayer/Source/ElectraPlayerRuntime/Private/Runtime/Player/DASH/ManifestDASH.cpp
@@ -2747,8 +2747,9 @@ FManifestDASHInternal::FRepresentation::ESearchResult FManifestDASHInternal::FRe
 	}
 	uint32 SegmentDuration = Duration.Value();
 
-	// Get the period local time into media local timescale and add the PTO.
-	int64 MediaLocalSearchTime = InSearchOptions.PeriodLocalTime.GetAsTimebase(MPDTimescale) + PTO;
+	// Get the period local time into media local timescale.
+	// Note that the PTO is not relevant here since we are not calculating on internal media time.
+	int64 MediaLocalSearchTime = InSearchOptions.PeriodLocalTime.GetAsTimebase(MPDTimescale);
 	// If the first media segment does not fall onto the period start there will be an EPT delta that is usually negative.
 	// To simplify calculation of the segment index we shift the search time such that 0 would correspond to the EPT.
 	int32 EPTdelta = EptDelta.GetWithDefault(0);
@@ -2758,8 +2759,8 @@ FManifestDASHInternal::FRepresentation::ESearchResult FManifestDASHInternal::FRe
 		MediaLocalSearchTime = 0;
 	}
 
-	int64 MediaLocalPeriodEnd = InSearchOptions.PeriodDuration.IsValid() && !InSearchOptions.PeriodDuration.IsInfinity() ? InSearchOptions.PeriodDuration.GetAsTimebase(MPDTimescale) - EPTdelta + PTO : TNumericLimits<int64>::Max();
-	int64 MediaLocalPresentationEnd = InSearchOptions.PeriodPresentationEnd.IsValid() && !InSearchOptions.PeriodPresentationEnd.IsInfinity() ? InSearchOptions.PeriodPresentationEnd.GetAsTimebase(MPDTimescale) + PTO : TNumericLimits<int64>::Max();
+	int64 MediaLocalPeriodEnd = InSearchOptions.PeriodDuration.IsValid() && !InSearchOptions.PeriodDuration.IsInfinity() ? InSearchOptions.PeriodDuration.GetAsTimebase(MPDTimescale) - EPTdelta : TNumericLimits<int64>::Max();
+	int64 MediaLocalPresentationEnd = InSearchOptions.PeriodPresentationEnd.IsValid() && !InSearchOptions.PeriodPresentationEnd.IsInfinity() ? InSearchOptions.PeriodPresentationEnd.GetAsTimebase(MPDTimescale) : TNumericLimits<int64>::Max();
 	int64 MediaLocalEndTime = Utils::Min(MediaLocalPeriodEnd, MediaLocalPresentationEnd);
 	uint32 MaxSegmentsInPeriod = MediaLocalEndTime == TNumericLimits<int64>::Max() ? TNumericLimits<uint32>::Max() : (uint32)((MediaLocalEndTime + SegmentDuration - 1) / SegmentDuration);
 
@@ -2831,14 +2832,14 @@ FManifestDASHInternal::FRepresentation::ESearchResult FManifestDASHInternal::FRe
 		return FManifestDASHInternal::FRepresentation::ESearchResult::PastEOS;
 	}
 
-	OutSegmentInfo.PeriodLocalSegmentStartTime.SetFromND(SegmentNum * (int64)SegmentDuration - PTO, MPDTimescale);
-	OutSegmentInfo.Time = EPTdelta + SegmentNum * (int64)SegmentDuration;
+	OutSegmentInfo.PeriodLocalSegmentStartTime.SetFromND(SegmentNum * (int64)SegmentDuration, MPDTimescale);
+	OutSegmentInfo.Time = PTO + EPTdelta + SegmentNum * (int64)SegmentDuration;
 	OutSegmentInfo.PTO = PTO;
 	OutSegmentInfo.EPTdelta = EPTdelta;
 	OutSegmentInfo.Duration = SegmentDuration;
 	OutSegmentInfo.Number = StartNumber + SegmentNum;
-	OutSegmentInfo.MediaLocalFirstAUTime = OutSegmentInfo.MediaLocalFirstPTS = MediaLocalSearchTime;
-	OutSegmentInfo.MediaLocalLastAUTime = MediaLocalEndTime;
+	OutSegmentInfo.MediaLocalFirstAUTime = OutSegmentInfo.MediaLocalFirstPTS = MediaLocalSearchTime + PTO;
+	OutSegmentInfo.MediaLocalLastAUTime = MediaLocalEndTime + PTO;
 	OutSegmentInfo.Timescale = MPDTimescale;
 	OutSegmentInfo.bMayBeMissing = SegmentNum + 1 >= MaxSegmentsInPeriod;
 	OutSegmentInfo.bIsLastInPeriod = OutSegmentInfo.bMayBeMissing && InSearchOptions.bHasFollowingPeriod;
diff --git a/Engine/Plugins/Media/ElectraPlayer/Source/ElectraPlayerRuntime/Private/Runtime/Player/DASH/PlaylistReaderDASH.cpp b/Engine/Plugins/Media/ElectraPlayer/Source/ElectraPlayerRuntime/Private/Runtime/Player/DASH/PlaylistReaderDASH.cpp
--- a/Engine/Plugins/Media/ElectraPlayer/Source/ElectraPlayerRuntime/Private/Runtime/Player/DASH/PlaylistReaderDASH.cpp
+++ b/Engine/Plugins/Media/ElectraPlayer/Source/ElectraPlayerRuntime/Private/Runtime/Player/DASH/PlaylistReaderDASH.cpp
@@ -1559,6 +1559,12 @@ void FPlaylistReaderDASH::Timesync_httpiso_Completed(FResourceLoadRequestPtr Req
 			{
 				PlayerSessionServices->GetSynchronizedUTCTime()->SetTime(NewTime);
 			}
+			// If parsing failed then maybe the response is just a number (possibly with frational digits) giving the
+			// current Unix epoch time.
+			else if (UnixEpoch::ParseFloatString(NewTime, Response))
+			{
+				PlayerSessionServices->GetSynchronizedUTCTime()->SetTime(NewTime);
+			}
 		}
 	}
 	// Presently we do not try another time sync method if this one has failed. Clear out what we have attempted so far and be done with it.
