ElectraPlayer: Increased buffer limit size to accomodate UHD HDR content at high bitrates. Limiting to 6s worth of content instead of memory size.

#rnx

diff --git a/Engine/Plugins/Media/ElectraPlayer/Source/ElectraPlayerRuntime/Private/Runtime/Player/AdaptiveStreamingPlayerInternalConfig.h b/Engine/Plugins/Media/ElectraPlayer/Source/ElectraPlayerRuntime/Private/Runtime/Player/AdaptiveStreamingPlayerInternalConfig.h
--- a/Engine/Plugins/Media/ElectraPlayer/Source/ElectraPlayerRuntime/Private/Runtime/Player/AdaptiveStreamingPlayerInternalConfig.h
+++ b/Engine/Plugins/Media/ElectraPlayer/Source/ElectraPlayerRuntime/Private/Runtime/Player/AdaptiveStreamingPlayerInternalConfig.h
@@ -23,11 +23,11 @@ namespace Electra
 				WorkerThread.CoreAffinity = -1;
 
 				// Set default values to maximum permitted values.
-				StreamBufferConfigVideo.MaxDataSize = 16 << 20;
-				StreamBufferConfigVideo.MaxDuration.SetFromSeconds(20.0);
+				StreamBufferConfigVideo.MaxDataSize = 128 << 20;
+				StreamBufferConfigVideo.MaxDuration.SetFromSeconds(6.0);
 
-				StreamBufferConfigAudio.MaxDataSize = 2 << 20;
-				StreamBufferConfigAudio.MaxDuration.SetFromSeconds(20.0);
+				StreamBufferConfigAudio.MaxDataSize = 8 << 20;
+				StreamBufferConfigAudio.MaxDuration.SetFromSeconds(6.0);
 
 				// Subtitle streams tend to be sparse and each AU could potentially have
 				// a huge duration, so we need to ignore durations altogether.
