Fixing up H.264 playback with Electra under Vulkan on PC

diff --git a/Engine/Plugins/Media/ElectraPlayer/Source/ElectraPlayerRuntime/Private/Runtime/Decoder/Windows/VideoDecoderH264_PC.cpp b/Engine/Plugins/Media/ElectraPlayer/Source/ElectraPlayerRuntime/Private/Runtime/Decoder/Windows/VideoDecoderH264_PC.cpp
--- a/Engine/Plugins/Media/ElectraPlayer/Source/ElectraPlayerRuntime/Private/Runtime/Decoder/Windows/VideoDecoderH264_PC.cpp
+++ b/Engine/Plugins/Media/ElectraPlayer/Source/ElectraPlayerRuntime/Private/Runtime/Decoder/Windows/VideoDecoderH264_PC.cpp
@@ -363,7 +363,7 @@ bool FVideoDecoderH264_PC::CopyTexture(const TRefCountPtr<IMFSample>& Sample, El
 			return false;
 		}
 	}
-	else if (FDXDeviceInfo::s_DXDeviceInfo->DxVersion == FDXDeviceInfo::ED3DVersion::Version12Win10)
+	else if (FDXDeviceInfo::s_DXDeviceInfo->DxVersion == FDXDeviceInfo::ED3DVersion::Version12Win10 || FDXDeviceInfo::s_DXDeviceInfo->DxVersion == FDXDeviceInfo::ED3DVersion::VersionNoneDxWin10)
 	{
 		TRefCountPtr<IMFMediaBuffer> Buffer;
 		CHECK_HR(Sample->GetBufferByIndex(0, Buffer.GetInitReference()));
