ElectraPlayer: Removed a check that triggered with iOS16 due to some internal API changes.

#rnx

diff --git a/Engine/Plugins/Media/ElectraPlayer/Source/ElectraPlayerRuntime/Private/Runtime/Decoder/Apple/AudioDecoderAAC_Apple.cpp b/Engine/Plugins/Media/ElectraPlayer/Source/ElectraPlayerRuntime/Private/Runtime/Decoder/Apple/AudioDecoderAAC_Apple.cpp
--- a/Engine/Plugins/Media/ElectraPlayer/Source/ElectraPlayerRuntime/Private/Runtime/Decoder/Apple/AudioDecoderAAC_Apple.cpp
+++ b/Engine/Plugins/Media/ElectraPlayer/Source/ElectraPlayerRuntime/Private/Runtime/Decoder/Apple/AudioDecoderAAC_Apple.cpp
@@ -687,7 +687,6 @@ OSStatus FAudioDecoderAAC::FAudioConverterInstance::AudioToolbox_ComplexInputCal
 
 			if (OutDataPacketDescription)
 			{
-				check(*OutDataPacketDescription == nullptr);
 				*OutDataPacketDescription = &PacketDescr;
 				PacketDescr.mStartOffset 		    = 0;
 				PacketDescr.mVariableFramesInPacket = 0;
@@ -1115,7 +1114,7 @@ void FAudioDecoderAAC::WorkerThread()
 			}
 			UE_LOG(LogElectraPlayer, Log, TEXT("FAudioDecoderAAC(%p): OnResuming"), this);
 		}
-	
+
 		// Notify the buffer listener that we will now be needing an AU for our input buffer.
 		if (!bError && InputBufferListener && NextAccessUnits.IsEmpty())
 		{
@@ -1274,7 +1273,7 @@ void FAudioDecoderAAC::WorkerThread()
 	ConfigRecord.Reset();
 
 	FMemory::Free(PCMBuffer);
-	
+
 	RemoveBGFGNotificationHandler(FGBGHandlers);
 }
 
