ElectraPlayer: Added stream quality cycling cvar for VoD assets for debugging purposes

#rnx

diff --git a/Engine/Plugins/Media/ElectraPlayer/Source/ElectraPlayerRuntime/Private/Runtime/Player/ABRRules/ABRVoDPlus.cpp b/Engine/Plugins/Media/ElectraPlayer/Source/ElectraPlayerRuntime/Private/Runtime/Player/ABRRules/ABRVoDPlus.cpp
--- a/Engine/Plugins/Media/ElectraPlayer/Source/ElectraPlayerRuntime/Private/Runtime/Player/ABRRules/ABRVoDPlus.cpp
+++ b/Engine/Plugins/Media/ElectraPlayer/Source/ElectraPlayerRuntime/Private/Runtime/Player/ABRRules/ABRVoDPlus.cpp
@@ -10,6 +10,19 @@
 #include "Utilities/Utilities.h"
 #include "Utilities/URLParser.h"
 
+#define ENABLE_CYCLE_CVAR 0
+
+#if ENABLE_CYCLE_CVAR
+#include "HAL/IConsoleManager.h"
+static TAutoConsoleVariable<int32> CVarElectraABRVODCycle(
+	TEXT("Electra.ABR.VoD.CycleQualities"),
+	0,
+	TEXT("Cycles through all quality levels even if not sustainable.\n")
+	TEXT(" 0: do not cycle; 1-n: cycle up to quality level n."),
+	ECVF_Default);
+#endif
+
+
 
 namespace Electra
 {
@@ -222,6 +235,14 @@ private:
 	
 	// Scale of highest stream bitrate to clamp bandwidth to so it does not get ridiculously large.
 	const double ClampBandwidthToMaxStreamBitrateScaleFactor = 2.0;
+	// Clamp to this rate unless the max stream bitrate after applying the scaling factor is already higher.
+	const double ClampBandwidthToMaxNonStreamBitrate = 20000000.0;
+
+#if ENABLE_CYCLE_CVAR
+	bool bForceQualityCycling = false;
+	int32 QualityCyclingIndex = 0;
+	int32 QualityCyclingLimit = 0;
+#endif
 };
 
 
@@ -417,7 +438,7 @@ void FABROnDemandPlus::ReportDownloadEnd(const Metrics::FSegmentDownloadStats& S
 				const double DlBandwidth = SegmentDownloadStats.NumBytesDownloaded * 8 / (SegmentDownloadStats.TimeToDownload - SegmentDownloadStats.TimeToFirstByte);
 				if (DlBandwidth > 0.0)
 				{
-					Bandwidth = Utils::Min((int64)DlBandwidth, (int64)(WorkVars->HighestQualityStreamBitrate * ClampBandwidthToMaxStreamBitrateScaleFactor));
+					Bandwidth = (int64) Utils::Min(DlBandwidth, Utils::Max(WorkVars->HighestQualityStreamBitrate * ClampBandwidthToMaxStreamBitrateScaleFactor, ClampBandwidthToMaxNonStreamBitrate));
 				}
 			}
 
@@ -846,10 +867,26 @@ IAdaptiveStreamSelector::ESegmentAction FABROnDemandPlus::PerformSelection(const
 					//Info->LogMessage(IInfoLog::ELevel::Verbose, FString::Printf(TEXT("Delaying segment by %.3fs"), OutDelay.GetAsSeconds()));
 				}
 			}
+
+#if ENABLE_CYCLE_CVAR
+			if (bForceQualityCycling)
+			{
+				if (QualityCyclingIndex >= Utils::Min(InCandidates.Num(), QualityCyclingLimit))
+				{
+					QualityCyclingIndex = 0;
+				}
+				NewQualityIndex = InCandidates[QualityCyclingIndex++]->QualityIndex;
+				TSharedPtrTS<FABRStreamInformation> Cnd = GetStreamInfoForQualityIndex(StreamType, NewQualityIndex);
+				if (Cnd.IsValid())
+				{
+					Info->LogMessage(IInfoLog::ELevel::Info, FString::Printf(TEXT("Cycling to quality index %d (%d*%d @ %d bps)"), NewQualityIndex, Cnd->Resolution.Width, Cnd->Resolution.Height, Cnd->Bitrate));
+				}
+			}
+#endif
 		}
 		else
 		{
-			NewQualityIndex = InCandidates.Num() - 1;
+			NewQualityIndex = InCandidates[InCandidates.Num() - 1]->QualityIndex;
 		}
 
 		check(NewQualityIndex >= 0);
@@ -935,6 +972,24 @@ IAdaptiveStreamSelector::EHandlingAction FABROnDemandPlus::PeriodicHandle()
 			}
 		}
 	}
+
+#if ENABLE_CYCLE_CVAR
+	int32 CycleQuality = CVarElectraABRVODCycle.GetValueOnAnyThread();
+	if (CycleQuality && CycleQuality!=QualityCyclingLimit)
+	{
+		QualityCyclingIndex = bForceQualityCycling ? QualityCyclingIndex : 0;
+		QualityCyclingLimit = CycleQuality;
+		bForceQualityCycling = true;
+		Info->LogMessage(IInfoLog::ELevel::Info, FString::Printf(TEXT("Enabling quality cycling up to quality %d"), CycleQuality));
+	}
+	else if (CycleQuality == 0 && bForceQualityCycling)
+	{
+		bForceQualityCycling = false;
+		QualityCyclingLimit = 0;
+		Info->LogMessage(IInfoLog::ELevel::Info, FString::Printf(TEXT("Diable quality cycling")));
+	}
+#endif
+
 	return NextAction;
 }
 
