ElectraPlayer: exposing metadata through media options

#rnx

diff --git a/Engine/Plugins/Media/ElectraPlayer/Source/ElectraPlayerPlugin/Private/ElectraPlayerPlugin.cpp b/Engine/Plugins/Media/ElectraPlayer/Source/ElectraPlayerPlugin/Private/ElectraPlayerPlugin.cpp
--- a/Engine/Plugins/Media/ElectraPlayer/Source/ElectraPlayerPlugin/Private/ElectraPlayerPlugin.cpp
+++ b/Engine/Plugins/Media/ElectraPlayer/Source/ElectraPlayerPlugin/Private/ElectraPlayerPlugin.cpp
@@ -278,6 +278,18 @@ Electra::FVariantValue FElectraPlayerPlugin::FPlayerAdapterDelegate::QueryOption
 						break;
 					}
 
+					case EOptionType::MediaMetadataUpdate:
+					{
+						static const FName MetadataUpdateOptionKey = TEXT("ElectraMetaDataUpdate");
+						if (SafeOptions->HasMediaOption(MetadataUpdateOptionKey))
+						{
+							check(Param.IsType(FVariantValue::EDataType::TypeFString));
+							// This only provides metadata, the return value of the Get is of no consequence.
+							SafeOptions->GetMediaOption(MetadataUpdateOptionKey, Param.GetFString());
+						}
+						break;
+					}
+
 					case EOptionType::PlaystartPosFromSeekPositions:
 					{
 						static const FName PlaystartOptionKey = TEXT("ElectraGetPlaystartPosFromSeekPositions");
