ElectraPlayer: Possible fix for a rare race condition at the loop point

#rnx

diff --git a/Engine/Plugins/Media/ElectraPlayer/Source/ElectraPlayerRuntime/Private/Runtime/Player/AdaptiveStreamingPlayerInternal.h b/Engine/Plugins/Media/ElectraPlayer/Source/ElectraPlayerRuntime/Private/Runtime/Player/AdaptiveStreamingPlayerInternal.h
--- a/Engine/Plugins/Media/ElectraPlayer/Source/ElectraPlayerRuntime/Private/Runtime/Player/AdaptiveStreamingPlayerInternal.h
+++ b/Engine/Plugins/Media/ElectraPlayer/Source/ElectraPlayerRuntime/Private/Runtime/Player/AdaptiveStreamingPlayerInternal.h
@@ -2041,11 +2041,13 @@ private:
 
 	TSharedPtrTS<FMultiTrackAccessUnitBuffer> GetCurrentOutputStreamBuffer(EStreamType InStreamType)
 	{
+		FScopeLock lock(&DataBuffersCriticalSection);
 		return GetStreamBuffer(InStreamType, ActiveDataOutputBuffers);
 	}
 
 	TSharedPtrTS<FMultiTrackAccessUnitBuffer> GetCurrentReceiveStreamBuffer(EStreamType InStreamType)
 	{
+		FScopeLock lock(&DataBuffersCriticalSection);
 		return GetStreamBuffer(InStreamType, CurrentDataReceiveBuffers);
 	}
 
