ElectraPlayer: Handling of in-stream metadata and dynamic changes for DASH playback

#rnx

diff --git a/Engine/Plugins/Media/ElectraPlayer/Source/ElectraPlayerRuntime/Private/Runtime/Player/AdaptiveStreamingPlayerInternal.h b/Engine/Plugins/Media/ElectraPlayer/Source/ElectraPlayerRuntime/Private/Runtime/Player/AdaptiveStreamingPlayerInternal.h
--- a/Engine/Plugins/Media/ElectraPlayer/Source/ElectraPlayerRuntime/Private/Runtime/Player/AdaptiveStreamingPlayerInternal.h
+++ b/Engine/Plugins/Media/ElectraPlayer/Source/ElectraPlayerRuntime/Private/Runtime/Player/AdaptiveStreamingPlayerInternal.h
@@ -1897,7 +1897,8 @@ private:
 		void Reset()
 		{
 			NextEntries.Empty();
-			// Do not reset the active metadata.
+			// Do NOT reset the active metadata, but the time it became valid!
+			ActiveSince.SetToInvalid();
 		}
 		void AddEntry(const FTimeValue& InValidFrom, const TSharedPtrTS<UtilsMP4::FMetadataParser>& InMetadata)
 		{
@@ -1906,7 +1907,11 @@ private:
 			e.Metadata = InMetadata;
 			NextEntries.StableSort([](const FEntry& a, const FEntry& b)
 			{
-				return a.ValidFrom < b.ValidFrom;
+				const FTimeValue& t1 = a.ValidFrom;
+				const FTimeValue& t2 = b.ValidFrom;
+				const int64 s1 = t1.GetSequenceIndex();
+				const int64 s2 = t2.GetSequenceIndex();
+				return (s1 == s2 && t1 < t2) || (s1 < s2);
 			});
 		}
 		bool Handle(const FTimeValue& InAtTime);
@@ -1922,6 +1927,7 @@ private:
 		};
 		TArray<FEntry> NextEntries;
 		TSharedPtrTS<UtilsMP4::FMetadataParser> ActiveMetadata;
+		FTimeValue ActiveSince;
 	};
 
 	struct FMetadataHandlingState
diff --git a/Engine/Plugins/Media/ElectraPlayer/Source/ElectraPlayerRuntime/Private/Runtime/Player/AdaptiveStreamingPlayerManifest.cpp b/Engine/Plugins/Media/ElectraPlayer/Source/ElectraPlayerRuntime/Private/Runtime/Player/AdaptiveStreamingPlayerManifest.cpp
--- a/Engine/Plugins/Media/ElectraPlayer/Source/ElectraPlayerRuntime/Private/Runtime/Player/AdaptiveStreamingPlayerManifest.cpp
+++ b/Engine/Plugins/Media/ElectraPlayer/Source/ElectraPlayerRuntime/Private/Runtime/Player/AdaptiveStreamingPlayerManifest.cpp
@@ -347,27 +347,52 @@ void FAdaptiveStreamingPlayer::UpdateManifest()
 bool FAdaptiveStreamingPlayer::FMediaMetadataUpdate::Handle(const FTimeValue& InAtTime)
 {
 	TSharedPtrTS<UtilsMP4::FMetadataParser> NextMetadata;
+	FTimeValue NextActiveTime;
 	while(NextEntries.Num())
 	{
+		// Make the first metadata available right away if there is none yet and the time is not valid either.
 		if (!InAtTime.IsValid() && !ActiveMetadata.IsValid())
 		{
 			NextMetadata = NextEntries[0].Metadata;
+			NextActiveTime = NextEntries[0].ValidFrom;
 			break;
 		}
-		else if (NextEntries[0].ValidFrom <= InAtTime)
-		{
-			NextMetadata = NextEntries[0].Metadata;
-			NextEntries.RemoveAt(0);
-		}
 		else
 		{
-			break;
+			const int64 s1 = NextEntries[0].ValidFrom.GetSequenceIndex();
+			const int64 s2 = InAtTime.GetSequenceIndex();
+			// Sequence count of metadata less or equal to the one of the current time?
+			if (s1 <= s2)
+			{
+				if (NextEntries[0].ValidFrom <= InAtTime)
+				{
+					// Is this metadata older than the one we used before? This happens when there is metadata in the
+					// init segment, which is sent for every media segment so that no metadata is missed when seeking.
+					// Also, when stream switches occur the metadata of the init segment from the new stream is sent
+					// as well.
+					if (!ActiveSince.IsValid() || NextEntries[0].ValidFrom.GetSequenceIndex() > ActiveSince.GetSequenceIndex() || NextEntries[0].ValidFrom > ActiveSince)
+					{
+						NextMetadata = NextEntries[0].Metadata;
+						NextActiveTime = NextEntries[0].ValidFrom;
+					}
+					NextEntries.RemoveAt(0);
+				}
+				else
+				{
+					break;
+				}
+			}
+			else
+			{
+				break;
+			}
 		}
 	}
 	if (NextMetadata.IsValid())
 	{
 		bool bChanged = !ActiveMetadata.IsValid() || ActiveMetadata->IsDifferentFrom(*NextMetadata);
 		ActiveMetadata = MoveTemp(NextMetadata);
+		ActiveSince = NextActiveTime;
 		return bChanged;
 	}
 	return false;
diff --git a/Engine/Plugins/Media/ElectraPlayer/Source/ElectraPlayerRuntime/Private/Runtime/Player/DASH/StreamReaderFMP4DASH.cpp b/Engine/Plugins/Media/ElectraPlayer/Source/ElectraPlayerRuntime/Private/Runtime/Player/DASH/StreamReaderFMP4DASH.cpp
--- a/Engine/Plugins/Media/ElectraPlayer/Source/ElectraPlayerRuntime/Private/Runtime/Player/DASH/StreamReaderFMP4DASH.cpp
+++ b/Engine/Plugins/Media/ElectraPlayer/Source/ElectraPlayerRuntime/Private/Runtime/Player/DASH/StreamReaderFMP4DASH.cpp
@@ -17,6 +17,7 @@
 #include "Player/DRM/DRMManager.h"
 #include "Utilities/Utilities.h"
 #include "Utilities/TimeUtilities.h"
+#include "Utilities/UtilsMP4.h"
 #include "Async/Async.h"
 
 #define INTERNAL_ERROR_INIT_SEGMENT_DOWNLOAD_ERROR					1
@@ -1023,6 +1024,28 @@ void FStreamReaderFMP4DASH::FStreamHandler::HandleRequest()
 					}
 					if (parseError == UEMEDIA_ERROR_OK)
 					{
+						// Get the metadata from the movie fragment or the init segment.
+						const IParserISO14496_12::IMetadata* MoofMetadata = MP4Parser->GetMetadata(IParserISO14496_12::EBaseBoxType::Moof);
+						const IParserISO14496_12::IMetadata* MoovMetadata = !MoofMetadata ? MP4InitSegment.IsValid() ? MP4InitSegment->GetMetadata(IParserISO14496_12::EBaseBoxType::Moov) : nullptr : nullptr;
+						if (MoofMetadata || MoovMetadata)
+						{
+							const IParserISO14496_12::IMetadata* md = MoofMetadata ? MoofMetadata : MoovMetadata;
+							uint32 hdlr = md->GetHandler();
+							uint32 res0 = md->GetReserved(0);
+							TArray<UtilsMP4::FMetadataParser::FBoxInfo> Boxes;
+							for(int32 i=0, iMax=md->GetNumChildBoxes(); i<iMax; ++i)
+							{
+								Boxes.Emplace(UtilsMP4::FMetadataParser::FBoxInfo(md->GetChildBoxType(i), md->GetChildBoxData(i), md->GetChildBoxDataSize(i)));
+							}
+							TSharedPtrTS<UtilsMP4::FMetadataParser> MediaMetadata = MakeSharedTS<UtilsMP4::FMetadataParser>();
+							if (MediaMetadata->Parse(hdlr, res0, Boxes) == UtilsMP4::FMetadataParser::EResult::Success)
+							{
+								FTimeValue StartTime = md == MoofMetadata ? Request->GetFirstPTS() : TimeOffset;
+								StartTime.SetSequenceIndex(Request->TimestampSequenceIndex);
+  								PlayerSessionService->SendMessageToPlayer(FPlaylistMetadataUpdateMessage::Create(StartTime, MediaMetadata));
+							}
+						}
+
 						// For the time being we only want to have a single track in the movie segments.
 						if (MP4Parser->GetNumberOfTracks() == 1)
 						{
