ElectraPlayer: moved setting of new instance up in order to use internal configuration methods on that instance

#rnx

diff --git a/Engine/Plugins/Media/ElectraPlayer/Source/ElectraPlayerRuntime/Private/ElectraPlayer.cpp b/Engine/Plugins/Media/ElectraPlayer/Source/ElectraPlayerRuntime/Private/ElectraPlayer.cpp
--- a/Engine/Plugins/Media/ElectraPlayer/Source/ElectraPlayerRuntime/Private/ElectraPlayer.cpp
+++ b/Engine/Plugins/Media/ElectraPlayer/Source/ElectraPlayerRuntime/Private/ElectraPlayer.cpp
@@ -314,6 +314,9 @@ bool FElectraPlayer::OpenInternal(const FString& Url, const FParamDict& InPlayer
 			NewPlayer->AdaptivePlayer->SetBitrateCeiling(PlaystartOptions.MaxBandwidthForStreaming.GetValue());
 		}
 
+        // Set the player member variable to the new player so we can use our internal configuration methods on the new player. 
+        CurrentPlayer = MoveTemp(NewPlayer);
+		
 		// Apply options that may have been set prior to calling Open().
 		// Set these only if they have defined values as to not override what might have been set in the PlayerOptions.
 		if (bFrameAccurateSeeking.IsSet())
@@ -329,9 +332,6 @@ bool FElectraPlayer::OpenInternal(const FString& Url, const FParamDict& InPlayer
 			SetPlaybackRange(CurrentPlaybackRange);
 		}
 
-		// Set the player member variable to the new player now that it has been fully initialized.
-		CurrentPlayer = MoveTemp(NewPlayer);
-
 		if (bPreviousOpenLoadedBlob)
 		{
 			CurrentPlayer->AdaptivePlayer->ModifyOptions(InPlayerOptions, Electra::FParamDict());
