Adding ElectraPlayerPluginHandler to ElectraPlayer plugin

diff --git a/Engine/Plugins/Media/ElectraPlayer/ElectraPlayer.uplugin b/Engine/Plugins/Media/ElectraPlayer/ElectraPlayer.uplugin
--- a/Engine/Plugins/Media/ElectraPlayer/ElectraPlayer.uplugin
+++ b/Engine/Plugins/Media/ElectraPlayer/ElectraPlayer.uplugin
@@ -15,18 +15,22 @@
 	"IsBetaVersion": false,
 	"Installed": false,
 	"Modules": [
+		{
+			"Name": "ElectraPlayerPluginHandler",
+			"Type": "RuntimeNoCommandlet",
+			"LoadingPhase": "PreLoadingScreen",
+			"TargetDenyList": [ "Server" ]
+		},
 		{
 			"Name": "ElectraPlayerRuntime",
 			"Type": "RuntimeNoCommandlet",
 			"LoadingPhase": "PreLoadingScreen",
-			"PlatformAllowList": [ "Win64", "Mac", "IOS", "Android", "Linux" ],
 			"TargetDenyList": [ "Server" ]
 		},
 		{
 			"Name": "ElectraPlayerPlugin",
 			"Type": "RuntimeNoCommandlet",
 			"LoadingPhase": "PreLoadingScreen",
-			"PlatformAllowList": [ "Win64", "Mac", "IOS", "Android", "Linux" ],
 			"TargetDenyList": [ "Server" ]
 		},
 		{
@@ -43,6 +47,7 @@
 			"PlatformAllowList": [ "Win64", "Mac", "IOS", "Android", "Linux" ],
 			"TargetDenyList": [ "Server" ]
 		}
+		
 	],
 	"Plugins": [
 		{
diff --git a/Engine/Plugins/Media/ElectraPlayer/Source/ElectraPlayerPluginHandler/ElectraPlayerPluginHandler.Build.cs b/Engine/Plugins/Media/ElectraPlayer/Source/ElectraPlayerPluginHandler/ElectraPlayerPluginHandler.Build.cs
new file mode 100644
--- /dev/null
+++ b/Engine/Plugins/Media/ElectraPlayer/Source/ElectraPlayerPluginHandler/ElectraPlayerPluginHandler.Build.cs
@@ -0,0 +1,40 @@
+// Copyright Epic Games, Inc. All Rights Reserved.
+
+using System;
+using System.IO;
+
+namespace UnrealBuildTool.Rules
+{
+	public class ElectraPlayerPluginHandler : ModuleRules
+	{ 
+		public ElectraPlayerPluginHandler(ReadOnlyTargetRules Target) : base(Target)
+		{
+			if (DoesPlatformSupportElectra(Target))
+			{
+				PublicDefinitions.Add("HAVE_ELECTRA=1");
+				PublicDependencyModuleNames.Add("ElectraPlayerRuntime");
+				PublicDependencyModuleNames.Add("ElectraPlayerPlugin");
+			}
+			else
+			{
+				PublicDefinitions.Add("HAVE_ELECTRA=0");
+			}
+
+			PublicDependencyModuleNames.AddRange(
+				new string[]
+				{
+					"Core"
+				}
+			);
+		}
+		
+		virtual protected bool DoesPlatformSupportElectra(ReadOnlyTargetRules InTarget) 
+		{
+			return InTarget.Platform == UnrealTargetPlatform.Win64 || 
+			       InTarget.Platform == UnrealTargetPlatform.IOS || 
+			       InTarget.Platform == UnrealTargetPlatform.Mac ||
+			       InTarget.Platform == UnrealTargetPlatform.Linux ||
+			       InTarget.Platform == UnrealTargetPlatform.Android;
+		}
+	}
+}
diff --git a/Engine/Plugins/Media/ElectraPlayer/Source/ElectraPlayerPluginHandler/Private/ElectraPlayerPluginHandlerModule.cpp b/Engine/Plugins/Media/ElectraPlayer/Source/ElectraPlayerPluginHandler/Private/ElectraPlayerPluginHandlerModule.cpp
new file mode 100644
--- /dev/null
+++ b/Engine/Plugins/Media/ElectraPlayer/Source/ElectraPlayerPluginHandler/Private/ElectraPlayerPluginHandlerModule.cpp
@@ -0,0 +1,20 @@
+// Copyright Epic Games, Inc. All Rights Reserved.
+
+#include "ElectraPlayerPluginHandlerModule.h"
+
+#define LOCTEXT_NAMESPACE "FElectraPlayerPluginHandlerModule"
+
+void FElectraPlayerPluginHandlerModule::StartupModule()
+{
+	// This code will execute after your module is loaded into memory; the exact timing is specified in the .uplugin file per-module
+}
+
+void FElectraPlayerPluginHandlerModule::ShutdownModule()
+{
+	// This function may be called during shutdown to clean up your module.  For modules that support dynamic reloading,
+	// we call this function before unloading the module.
+}
+
+#undef LOCTEXT_NAMESPACE
+	
+IMPLEMENT_MODULE(FElectraPlayerPluginHandlerModule, ElectraPlayerPluginHandler)
\ No newline at end of file
diff --git a/Engine/Plugins/Media/ElectraPlayer/Source/ElectraPlayerPluginHandler/Public/ElectraPlayerPluginHandlerModule.h b/Engine/Plugins/Media/ElectraPlayer/Source/ElectraPlayerPluginHandler/Public/ElectraPlayerPluginHandlerModule.h
new file mode 100644
--- /dev/null
+++ b/Engine/Plugins/Media/ElectraPlayer/Source/ElectraPlayerPluginHandler/Public/ElectraPlayerPluginHandlerModule.h
@@ -0,0 +1,15 @@
+// Copyright Epic Games, Inc. All Rights Reserved.
+
+#pragma once
+
+#include "CoreMinimal.h"
+#include "Modules/ModuleManager.h"
+
+class ELECTRAPLAYERPLUGINHANDLER_API FElectraPlayerPluginHandlerModule : public IModuleInterface
+{
+public:
+
+	/** IModuleInterface implementation */
+	virtual void StartupModule() override;
+	virtual void ShutdownModule() override;
+};
