Excluding server target from ElectraPlayer

diff --git a/Engine/Plugins/Media/ElectraPlayer/Source/ElectraPlayerPluginHandler/ElectraPlayerPluginHandler.Build.cs b/Engine/Plugins/Media/ElectraPlayer/Source/ElectraPlayerPluginHandler/ElectraPlayerPluginHandler.Build.cs
--- a/Engine/Plugins/Media/ElectraPlayer/Source/ElectraPlayerPluginHandler/ElectraPlayerPluginHandler.Build.cs
+++ b/Engine/Plugins/Media/ElectraPlayer/Source/ElectraPlayerPluginHandler/ElectraPlayerPluginHandler.Build.cs
@@ -9,7 +9,7 @@ namespace UnrealBuildTool.Rules
 	{ 
 		public ElectraPlayerPluginHandler(ReadOnlyTargetRules Target) : base(Target)
 		{
-			if (DoesPlatformSupportElectra(Target))
+			if (DoesPlatformSupportElectra(Target) && Target.Type != TargetType.Server)
 			{
 				PublicDefinitions.Add("HAVE_ELECTRA=1");
 				PublicDependencyModuleNames.Add("ElectraPlayerRuntime");
