ElectraPlayer plugin: adding back PlatformAllowList

diff --git a/Engine/Plugins/Media/ElectraPlayer/ElectraPlayer.uplugin b/Engine/Plugins/Media/ElectraPlayer/ElectraPlayer.uplugin
--- a/Engine/Plugins/Media/ElectraPlayer/ElectraPlayer.uplugin
+++ b/Engine/Plugins/Media/ElectraPlayer/ElectraPlayer.uplugin
@@ -25,12 +25,14 @@
 			"Name": "ElectraPlayerRuntime",
 			"Type": "RuntimeNoCommandlet",
 			"LoadingPhase": "PreLoadingScreen",
+			"PlatformAllowList": [ "Win64", "Mac", "IOS", "Android", "Linux" ],
 			"TargetDenyList": [ "Server" ]
 		},
 		{
 			"Name": "ElectraPlayerPlugin",
 			"Type": "RuntimeNoCommandlet",
 			"LoadingPhase": "PreLoadingScreen",
+			"PlatformAllowList": [ "Win64", "Mac", "IOS", "Android", "Linux" ],
 			"TargetDenyList": [ "Server" ]
 		},
 		{
