Adding TVOS as supported platform for Electra media player

diff --git a/Engine/Plugins/Media/ElectraPlayer/ElectraPlayer.uplugin b/Engine/Plugins/Media/ElectraPlayer/ElectraPlayer.uplugin
--- a/Engine/Plugins/Media/ElectraPlayer/ElectraPlayer.uplugin
+++ b/Engine/Plugins/Media/ElectraPlayer/ElectraPlayer.uplugin
@@ -25,14 +25,14 @@
 			"Name": "ElectraPlayerRuntime",
 			"Type": "RuntimeNoCommandlet",
 			"LoadingPhase": "PreLoadingScreen",
-			"PlatformAllowList": [ "Win64", "Mac", "IOS", "Android", "Linux" ],
+			"PlatformAllowList": [ "Win64", "Mac", "IOS", "TVOS", "Android", "Linux" ],
 			"TargetDenyList": [ "Server" ]
 		},
 		{
 			"Name": "ElectraPlayerPlugin",
 			"Type": "RuntimeNoCommandlet",
 			"LoadingPhase": "PreLoadingScreen",
-			"PlatformAllowList": [ "Win64", "Mac", "IOS", "Android", "Linux" ],
+			"PlatformAllowList": [ "Win64", "Mac", "IOS", "TVOS", "Android", "Linux" ],
 			"TargetDenyList": [ "Server" ]
 		},
 		{
@@ -46,7 +46,7 @@
 			"Name": "ElectraPlayerFactory",
 			"Type": "RuntimeNoCommandlet",
 			"LoadingPhase": "PostEngineInit",
-			"PlatformAllowList": [ "Win64", "Mac", "IOS", "Android", "Linux" ],
+			"PlatformAllowList": [ "Win64", "Mac", "IOS", "TVOS", "Android", "Linux" ],
 			"TargetDenyList": [ "Server" ]
 		}
 		
@@ -65,4 +65,4 @@
 			"Enabled": true
 		}
 	]
-}
\ No newline at end of file
+}
