Fixing horde issue 250224 by rolling back recent changes to Electra plugin platform allow lilst
#rnx

diff --git a/Engine/Plugins/Media/ElectraPlayer/ElectraPlayer.uplugin b/Engine/Plugins/Media/ElectraPlayer/ElectraPlayer.uplugin
--- a/Engine/Plugins/Media/ElectraPlayer/ElectraPlayer.uplugin
+++ b/Engine/Plugins/Media/ElectraPlayer/ElectraPlayer.uplugin
@@ -25,14 +25,14 @@
 			"Name": "ElectraPlayerRuntime",
 			"Type": "RuntimeNoCommandlet",
 			"LoadingPhase": "PreLoadingScreen",
-			"PlatformAllowList": [ "Win64", "Mac", "IOS", "TVOS", "Android", "Linux" ],
+			"PlatformAllowList": [ "Win64", "Mac", "IOS", "Android", "Linux" ],
 			"TargetDenyList": [ "Server" ]
 		},
 		{
 			"Name": "ElectraPlayerPlugin",
 			"Type": "RuntimeNoCommandlet",
 			"LoadingPhase": "PreLoadingScreen",
-			"PlatformAllowList": [ "Win64", "Mac", "IOS", "TVOS", "Android", "Linux" ],
+			"PlatformAllowList": [ "Win64", "Mac", "IOS", "Android", "Linux" ],
 			"TargetDenyList": [ "Server" ]
 		},
 		{
@@ -46,7 +46,7 @@
 			"Name": "ElectraPlayerFactory",
 			"Type": "RuntimeNoCommandlet",
 			"LoadingPhase": "PostEngineInit",
-			"PlatformAllowList": [ "Win64", "Mac", "IOS", "TVOS", "Android", "Linux" ],
+			"PlatformAllowList": [ "Win64", "Mac", "IOS", "Android", "Linux" ],
 			"TargetDenyList": [ "Server" ]
 		}
 		
