ElectraPlayer: Added custom analytic fields; extended player message history to include relative event times

#rnx

diff --git a/Engine/Plugins/Media/ElectraPlayer/Source/ElectraPlayerPlugin/Private/ElectraPlayerPlugin.cpp b/Engine/Plugins/Media/ElectraPlayer/Source/ElectraPlayerPlugin/Private/ElectraPlayerPlugin.cpp
--- a/Engine/Plugins/Media/ElectraPlayer/Source/ElectraPlayerPlugin/Private/ElectraPlayerPlugin.cpp
+++ b/Engine/Plugins/Media/ElectraPlayer/Source/ElectraPlayerPlugin/Private/ElectraPlayerPlugin.cpp
@@ -264,7 +264,7 @@ Electra::FVariantValue FElectraPlayerPlugin::FPlayerAdapterDelegate::QueryOption
 			IMediaOptions *SafeOptions = SafeOptionInterface->GetMediaOptionInterface();
 			if (SafeOptions)
 			{
-				switch (Type)
+				switch(Type)
 				{
 					case EOptionType::MaxVerticalStreamResolution:
 					{
@@ -312,6 +312,20 @@ Electra::FVariantValue FElectraPlayerPlugin::FPlayerAdapterDelegate::QueryOption
 						break;
 					}
 
+					case EOptionType::CustomAnalyticsMetric:
+					{
+						check(Param.IsType(FVariantValue::EDataType::TypeFString));
+						if (Param.IsType(FVariantValue::EDataType::TypeFString))
+						{
+							FName OptionKey(*Param.GetFString());
+							if (SafeOptions->HasMediaOption(OptionKey))
+							{
+								return FVariantValue(SafeOptions->GetMediaOption(OptionKey, FString()));
+							}
+						}
+						break;
+					}
+
 					case EOptionType::PlaystartPosFromSeekPositions:
 					{
 						static const FName PlaystartOptionKey = TEXT("ElectraGetPlaystartPosFromSeekPositions");
diff --git a/Engine/Plugins/Media/ElectraPlayer/Source/ElectraPlayerRuntime/Private/ElectraPlayer.cpp b/Engine/Plugins/Media/ElectraPlayer/Source/ElectraPlayerRuntime/Private/ElectraPlayer.cpp
--- a/Engine/Plugins/Media/ElectraPlayer/Source/ElectraPlayerRuntime/Private/ElectraPlayer.cpp
+++ b/Engine/Plugins/Media/ElectraPlayer/Source/ElectraPlayerRuntime/Private/ElectraPlayer.cpp
@@ -26,6 +26,12 @@ DECLARE_CYCLE_STAT(TEXT("FElectraPlayer::TickInput"), STAT_ElectraPlayer_Electra
 
 //-----------------------------------------------------------------------------
 
+// Prefix to use in querying for a custom analytic value through QueryOptions()
+#define CUSTOM_ANALYTIC_METRIC_QUERYOPTION_KEY TEXT("ElectraCustomAnalytic")
+// Prefix to use in the metric event to set the custom value.
+#define CUSTOM_ANALYTIC_METRIC_KEYNAME TEXT("Custom")
+
+
 #define USE_INTERNAL_PLAYBACK_STATE 1
 
 //-----------------------------------------------------------------------------
@@ -249,6 +255,7 @@ bool FElectraPlayer::OpenInternal(const FString& Url, const FParamDict& InPlayer
 	NumQueuedAnalyticEvents = 0;
 	// Create a guid string for the analytics. We do this here and not in the constructor in case the same instance is used over again.
 	AnalyticsInstanceGuid = FGuid::NewGuid().ToString(EGuidFormats::Digits);
+	UpdateAnalyticsCustomValues();
 
 	PlaystartOptions = InPlaystartOptions;
 
@@ -1721,6 +1728,33 @@ void FElectraPlayer::AddCommonAnalyticsAttributes(TArray<FAnalyticsEventAttribut
 	InOutParamArray.Add(FAnalyticsEventAttribute(TEXT("Utc"), static_cast<double>(FDateTime::UtcNow().ToUnixTimestamp())));
 	InOutParamArray.Add(FAnalyticsEventAttribute(TEXT("OS"), FString::Printf(TEXT("%s"), *AnalyticsOSVersion)));
 	InOutParamArray.Add(FAnalyticsEventAttribute(TEXT("GPUAdapter"), AnalyticsGPUType));
+	StatisticsLock.Lock();
+	for(int32 nI=0, nIMax=UE_ARRAY_COUNT(AnalyticsCustomValues); nI<nIMax; ++nI)
+	{
+		if (AnalyticsCustomValues[nI].Len())
+		{
+			InOutParamArray.Add(FAnalyticsEventAttribute(FString::Printf(TEXT("%s%d"), CUSTOM_ANALYTIC_METRIC_KEYNAME, nI), AnalyticsCustomValues[nI]));
+		}
+	}
+	StatisticsLock.Unlock();
+}
+
+void FElectraPlayer::UpdateAnalyticsCustomValues()
+{
+	StatisticsLock.Lock();
+	TSharedPtr<IElectraPlayerAdapterDelegate, ESPMode::ThreadSafe> PinnedAdapterDelegate = AdapterDelegate.Pin();
+	if (PinnedAdapterDelegate.IsValid())
+	{
+		for(int32 nI=0, nIMax=UE_ARRAY_COUNT(AnalyticsCustomValues); nI<nIMax; ++nI)
+		{
+			FVariantValue Value = PinnedAdapterDelegate->QueryOptions(IElectraPlayerAdapterDelegate::EOptionType::CustomAnalyticsMetric, FVariantValue(FString::Printf(TEXT("%s%d"), CUSTOM_ANALYTIC_METRIC_QUERYOPTION_KEY, nI)));
+			if (Value.IsValid() && Value.GetDataType() == FVariantValue::EDataType::TypeFString)
+			{
+				AnalyticsCustomValues[nI] = Value.GetFString();
+			}
+		}
+	}
+	StatisticsLock.Unlock();
 }
 
 
@@ -2283,7 +2317,18 @@ void FElectraPlayer::HandlePlayerEventSegmentDownload(const Electra::Metrics::FS
 
 		if (SegmentDownloadStats.FailureReason.Len())
 		{
-			Statistics.AddMessageToHistory(FString::Printf(TEXT("%s segment download issue (%s)"), Electra::Metrics::GetSegmentTypeString(SegmentDownloadStats.SegmentType), *SegmentDownloadStats.FailureReason));
+			FString Msg;
+			if (!SegmentDownloadStats.bWasAborted)
+			{
+				Msg = FString::Printf(TEXT("%s segment download issue on representation %s, bitrate %d, retry %d: %s"), Electra::Metrics::GetSegmentTypeString(SegmentDownloadStats.SegmentType),
+					*SegmentDownloadStats.RepresentationID, SegmentDownloadStats.Bitrate, SegmentDownloadStats.RetryNumber, *SegmentDownloadStats.FailureReason);
+			}
+			else
+			{
+				Msg = FString::Printf(TEXT("%s segment download issue on representation %s, bitrate %d, aborted: %s"), Electra::Metrics::GetSegmentTypeString(SegmentDownloadStats.SegmentType),
+					*SegmentDownloadStats.RepresentationID, SegmentDownloadStats.Bitrate, *SegmentDownloadStats.FailureReason);
+			}
+			Statistics.AddMessageToHistory(Msg);
 		}
 
 		static const FString kEventNameElectraSegmentIssue(TEXT("Electra.SegmentIssue"));
@@ -2319,11 +2364,11 @@ void FElectraPlayer::HandlePlayerEventVideoQualityChange(int32 NewBitrate, int32
 	}
 	else
 	{
-		if(bIsDrasticDownswitch)
+		if (bIsDrasticDownswitch)
 		{
 			++Statistics.NumQualityDrasticDownswitches;
 		}
-		if(NewBitrate > PreviousBitrate)
+		if (NewBitrate > PreviousBitrate)
 		{
 			++Statistics.NumQualityUpswitches;
 		}
@@ -2374,6 +2419,8 @@ void FElectraPlayer::HandlePlayerEventVideoQualityChange(int32 NewBitrate, int32
 		EnqueueAnalyticsEvent(AnalyticEvent);
 	}
 
+	Statistics.AddMessageToHistory(FString::Printf(TEXT("Video bitrate change from %d to %d"), PreviousBitrate, NewBitrate));
+
 	CSV_EVENT(ElectraPlayer, TEXT("QualityChange %d -> %d"), PreviousBitrate, NewBitrate);
 }
 
@@ -2633,8 +2680,8 @@ void FElectraPlayer::HandlePlayerEventError(const FString& ErrorReason)
 	FString MessageHistory;
 	for(auto &msg : Statistics.MessageHistoryBuffer)
 	{
-		MessageHistory.Append(msg);
-		MessageHistory.Append(TEXT("\n"));
+		MessageHistory.Append(FString::Printf(TEXT("%8.3f: %s"), msg.TimeSinceStart, *msg.Message));
+		MessageHistory.Append(TEXT("<br>"));
 	}
 
 	// Enqueue an "Error" event.
@@ -2714,6 +2761,20 @@ void FElectraPlayer::FDroppedFrameStats::AddNewDrop(const FTimespan& InFrameTime
 	PlayerTimeAtLastReport = InPlayerTime;
 }
 
+void FElectraPlayer::FStatistics::AddMessageToHistory(FString InMessage)
+{
+	if (MessageHistoryBuffer.Num() >= 20)
+	{
+		MessageHistoryBuffer.RemoveAt(0);
+	}
+	double Now = FPlatformTime::Seconds();
+	FStatistics::FHistoryEntry he;
+	he.Message = MoveTemp(InMessage);
+	he.TimeSinceStart = TimeAtOpen < 0.0 ? 0.0 : Now - TimeAtOpen;
+	MessageHistoryBuffer.Emplace(MoveTemp(he));
+}
+
+
 void FElectraPlayer::UpdatePlayEndStatistics()
 {
 	TSharedPtr<FInternalPlayerImpl, ESPMode::ThreadSafe> LockedPlayer = CurrentPlayer;
@@ -2788,12 +2849,7 @@ void FElectraPlayer::LogStatistics()
 			"Currently active playlist URL: %s\n"\
 			"Currently active resolution: %d * %d\n" \
 			"Current state: %s\n" \
-			"Number of video frames dropped: %d, worst time delta %.3f ms\n" \
-			"Number of audio frames dropped: %d, worst time delta %.3f ms\n" \
-			"Last error: %s\n" \
-			"Subtitles URL: %s\n" \
-			"Subtitles response time: %.3fs\n" \
-			"Subtitles last error: %s\n"
+			"Last issue: %s\n"
 		),
 			this, CurrentPlayer.Get(),
 			*FString::Printf(TEXT("%s"), *AnalyticsOSVersion),
@@ -2829,14 +2885,7 @@ void FElectraPlayer::LogStatistics()
 			Statistics.CurrentlyActiveResolutionWidth,
 			Statistics.CurrentlyActiveResolutionHeight,
 			*Statistics.LastState,
-			Statistics.DroppedVideoFrames.NumTotalDropped,
-			Statistics.DroppedVideoFrames.WorstDeltaTime.GetTotalMilliseconds(),
-			Statistics.DroppedAudioFrames.NumTotalDropped,
-			Statistics.DroppedAudioFrames.WorstDeltaTime.GetTotalMilliseconds(),
-			*SanitizeMessage(Statistics.LastError),
-			*SanitizeMessage(Statistics.SubtitlesURL),
-			Statistics.SubtitlesResponseTime,
-			*Statistics.SubtitlesLastError
+			*SanitizeMessage(Statistics.LastError)
 		);
 		
 		if (Statistics.LastError.Len())
@@ -2844,7 +2893,7 @@ void FElectraPlayer::LogStatistics()
 			FString MessageHistory;
 			for(auto &msg : Statistics.MessageHistoryBuffer)
 			{
-				MessageHistory.Append(msg);
+				MessageHistory.Append(FString::Printf(TEXT("%8.3f: %s"), msg.TimeSinceStart, *msg.Message));
 				MessageHistory.Append(TEXT("\n"));
 			}
 			UE_LOG(LogElectraPlayer, Log, TEXT("Most recent log messages:\n%s"), *MessageHistory);
@@ -2879,13 +2928,14 @@ void FElectraPlayer::SendAnalyticMetrics(const TSharedPtr<IAnalyticsProviderET>&
 
 
 	TArray<FAnalyticsEventAttribute> ParamArray;
+	UpdateAnalyticsCustomValues();
 	AddCommonAnalyticsAttributes(ParamArray);
 	StatisticsLock.Lock();
 	FString MessageHistory;
 	for(auto &msg : Statistics.MessageHistoryBuffer)
 	{
-		MessageHistory.Append(msg);
-		MessageHistory.Append(TEXT("\n"));
+		MessageHistory.Append(FString::Printf(TEXT("%8.3f: %s"), msg.TimeSinceStart, *msg.Message));
+		MessageHistory.Append(TEXT("<br>"));
 	}
 	ParamArray.Add(FAnalyticsEventAttribute(TEXT("URL"), Statistics.InitialURL));
 	ParamArray.Add(FAnalyticsEventAttribute(TEXT("LastState"), Statistics.LastState));
@@ -2948,6 +2998,7 @@ void FElectraPlayer::SendAnalyticMetricsPerMinute(const TSharedPtr<IAnalyticsPro
 	if (Player.Get() && Player->AdaptivePlayer->IsPlaying())
 	{
 		TArray<FAnalyticsEventAttribute> ParamArray;
+		UpdateAnalyticsCustomValues();
 		AddCommonAnalyticsAttributes(ParamArray);
 		StatisticsLock.Lock();
 		ParamArray.Add(FAnalyticsEventAttribute(TEXT("URL"), Statistics.CurrentlyActivePlaylistURL));
@@ -3086,7 +3137,8 @@ void FElectraPlayer::FAdaptiveStreamingPlayerResourceProvider::ProcessPendingSta
 	TSharedPtr<Electra::IAdaptiveStreamingPlayerResourceRequest, ESPMode::ThreadSafe> InOutRequest;
 	while (PendingStaticResourceRequests.Dequeue(InOutRequest))
 	{
-		check(InOutRequest->GetResourceType() == Electra::IAdaptiveStreamingPlayerResourceRequest::EPlaybackResourceType::Playlist ||
+		check(InOutRequest->GetResourceType() == Electra::IAdaptiveStreamingPlayerResourceRequest::EPlaybackResourceType::Empty ||
+			  InOutRequest->GetResourceType() == Electra::IAdaptiveStreamingPlayerResourceRequest::EPlaybackResourceType::Playlist ||
 			  InOutRequest->GetResourceType() == Electra::IAdaptiveStreamingPlayerResourceRequest::EPlaybackResourceType::LicenseKey);
 		if (InOutRequest->GetResourceType() == Electra::IAdaptiveStreamingPlayerResourceRequest::EPlaybackResourceType::Playlist)
 		{
diff --git a/Engine/Plugins/Media/ElectraPlayer/Source/ElectraPlayerRuntime/Private/ElectraPlayer.h b/Engine/Plugins/Media/ElectraPlayer/Source/ElectraPlayerRuntime/Private/ElectraPlayer.h
--- a/Engine/Plugins/Media/ElectraPlayer/Source/ElectraPlayerRuntime/Private/ElectraPlayer.h
+++ b/Engine/Plugins/Media/ElectraPlayer/Source/ElectraPlayerRuntime/Private/ElectraPlayer.h
@@ -451,7 +451,7 @@ private:
 	/** Metric delegates */
 	FElectraPlayerSendAnalyticMetricsDelegate&			SendAnalyticMetricsDelegate;
 	FElectraPlayerSendAnalyticMetricsPerMinuteDelegate&	SendAnalyticMetricsPerMinuteDelegate;
-	FElectraPlayerReportVideoStreamingErrorDelegate&		ReportVideoStreamingErrorDelegate;
+	FElectraPlayerReportVideoStreamingErrorDelegate&	ReportVideoStreamingErrorDelegate;
 	FElectraPlayerReportSubtitlesMetricsDelegate&		ReportSubtitlesMetricsDelegate;
 
 	/** Option interface **/
@@ -654,6 +654,11 @@ private:
 			FAverageValue	Bandwidth;
 			FAverageValue	Latency;
 		};
+		struct FHistoryEntry
+		{
+			double TimeSinceStart = 0.0;
+			FString Message;
+		};
 		FStatistics()
 		{
 			DroppedVideoFrames.SetFrameType(FDroppedFrameStats::EFrameType::VideoFrame);
@@ -704,14 +709,7 @@ private:
 			MediaDuration = 0.0;
 			MessageHistoryBuffer.Empty();
 		}
-		void AddMessageToHistory(FString InMessage)
-		{
-			if (MessageHistoryBuffer.Num() >= 20)
-			{
-				MessageHistoryBuffer.RemoveAt(0);
-			}
-			MessageHistoryBuffer.Emplace(MoveTemp(InMessage));
-		}
+		void AddMessageToHistory(FString InMessage);
 
 		FString					InitialURL;
 		FString					CurrentlyActivePlaylistURL;
@@ -753,7 +751,7 @@ private:
 		FTimeRange				MediaTimelineAtStart;
 		FTimeRange				MediaTimelineAtEnd;
 		double					MediaDuration;
-		TArray<FString>			MessageHistoryBuffer;
+		TArray<FHistoryEntry>	MessageHistoryBuffer;
 	};
 
 	struct FAnalyticsEvent
@@ -767,6 +765,7 @@ private:
 	void AddCommonAnalyticsAttributes(TArray<FAnalyticsEventAttribute>& InOutParamArray);
 	TSharedPtr<FAnalyticsEvent> CreateAnalyticsEvent(FString InEventName);
 	void EnqueueAnalyticsEvent(TSharedPtr<FAnalyticsEvent> InAnalyticEvent);
+	void UpdateAnalyticsCustomValues();
 
 	FCriticalSection						StatisticsLock;
 	FStatistics								Statistics;
@@ -781,6 +780,9 @@ private:
 	/** Sequential analytics event number. Helps sorting events. **/
 	uint32									AnalyticsInstanceEventCount;
 
+	/** Custom analytics constants. **/
+	FString									AnalyticsCustomValues[8];
+
 	void HandleBlobDownload();
 
 	void HandleDeferredPlayerEvents();
diff --git a/Engine/Plugins/Media/ElectraPlayer/Source/ElectraPlayerRuntime/Public/IElectraPlayerInterface.h b/Engine/Plugins/Media/ElectraPlayer/Source/ElectraPlayerRuntime/Public/IElectraPlayerInterface.h
--- a/Engine/Plugins/Media/ElectraPlayer/Source/ElectraPlayerRuntime/Public/IElectraPlayerInterface.h
+++ b/Engine/Plugins/Media/ElectraPlayer/Source/ElectraPlayerRuntime/Public/IElectraPlayerInterface.h
@@ -40,6 +40,7 @@ public:
 		LicenseKeyData,
 		PlaystartPosFromSeekPositions,
 		MediaMetadataUpdate,
+		CustomAnalyticsMetric,
 		Android_VideoSurface,
 	};
 	virtual Electra::FVariantValue QueryOptions(EOptionType Type, const Electra::FVariantValue& Param = Electra::FVariantValue()) = 0;
