ElectraPlayer: Releasing the seek lock before calling into user code to ensure it won't be held unnecessarily and cause a deadlock in subsequent function calls

#rnx

diff --git a/Engine/Plugins/Media/ElectraPlayer/Source/ElectraPlayerRuntime/Private/Runtime/Player/AdaptiveStreamingPlayer.cpp b/Engine/Plugins/Media/ElectraPlayer/Source/ElectraPlayerRuntime/Private/Runtime/Player/AdaptiveStreamingPlayer.cpp
--- a/Engine/Plugins/Media/ElectraPlayer/Source/ElectraPlayerRuntime/Private/Runtime/Player/AdaptiveStreamingPlayer.cpp
+++ b/Engine/Plugins/Media/ElectraPlayer/Source/ElectraPlayerRuntime/Private/Runtime/Player/AdaptiveStreamingPlayer.cpp
@@ -793,7 +793,10 @@ void FAdaptiveStreamingPlayer::GetLoopState(FLoopState& OutLoopState) const
 
 	// Check if there is a pending seek. A seek resets the loop count, but the seek itself may not be executing
 	// and the count not being reset yet.
+	
+	/* Not needed since we are only accessing a bool
 	FScopeLock lock(&SeekVars.Lock);
+	*/
 	if (SeekVars.PendingRequest.IsSet())
 	{
 		OutLoopState.Count = 0;
@@ -1693,16 +1696,21 @@ void FAdaptiveStreamingPlayer::HandleSeeking()
 		else
 		{
 			// Stop everything.
+			// We need to unlock the seek vars lock here to allow the renderer's call to Flush() in InternalStop()
+			// to call out to the user code for which we must not hold any state locks.
+			SeekVars.Lock.Unlock();
 			InternalStop(PlayerConfig.bHoldLastFrameDuringSeek);
 			if (Manifest.IsValid())
 			{
 				Manifest->UpdateDynamicRefetchCounter();
 			}
 			CurrentState = EPlayerState::eState_Seeking;
+			SeekVars.Lock.Lock();
 
 			SeekVars.bForScrubbing = bIsForScrubbing;
 		}
 		SeekVars.ActiveRequest = SeekParam;
+		lock.Unlock();
 		InternalStartAt(SeekVars.ActiveRequest.GetValue(), &SeekPlayrangeOnRequest);
 	}
 }
@@ -2193,6 +2201,7 @@ void FAdaptiveStreamingPlayer::HandleDeselectedBuffers()
 {
 	if (bIsVideoDeselected || bIsAudioDeselected || bIsTextDeselected)
 	{
+		FAccessUnitBufferInfo BufferStats;
 		// Get the current playback position in case nothing is selected.
 		FTimeValue DiscardUntilTime = PlaybackState.GetPlayPosition();
 
@@ -2208,6 +2217,11 @@ void FAdaptiveStreamingPlayer::HandleDeselectedBuffers()
 			{
 				DiscardUntilTime = v;
 			}
+			VidOutBuffer->GetStats(BufferStats);
+			if (BufferStats.bEndOfData && BufferStats.NumCurrentAccessUnits == 0)
+			{
+				DiscardUntilTime.SetToPositiveInfinity();
+			}
 		}
 		// Get last popped audio PTS if audio track is selected.
 		if (!bIsAudioDeselected && AudOutBuffer.IsValid())
@@ -2217,6 +2231,11 @@ void FAdaptiveStreamingPlayer::HandleDeselectedBuffers()
 			{
 				DiscardUntilTime = v;
 			}
+			AudOutBuffer->GetStats(BufferStats);
+			if (BufferStats.bEndOfData && BufferStats.NumCurrentAccessUnits == 0)
+			{
+				DiscardUntilTime.SetToPositiveInfinity();
+			}
 		}
 		// FIXME: do we need to consider subtitle tracks as being the only selected ones??
 
@@ -3429,6 +3448,18 @@ void FAdaptiveStreamingPlayer::CancelPendingMediaSegmentRequests(EStreamType Str
 			ReadyWaitingSegmentRequests.Enqueue(ReadyReq);
 		}
 	}
+	// Likewise for pending segments that hadn't been made ready yet.
+	TQueue<FPendingSegmentRequest> pendingRequests;
+	Swap(pendingRequests, NextPendingSegmentRequests);
+	while(!pendingRequests.IsEmpty())
+	{
+		FPendingSegmentRequest PendingReq;
+		pendingRequests.Dequeue(PendingReq);
+		if (PendingReq.StreamType != StreamType)
+		{
+			NextPendingSegmentRequests.Enqueue(PendingReq);
+		}
+	}
 }
 
 
@@ -4852,20 +4883,20 @@ void FAdaptiveStreamingPlayer::DebugPrint(void* pPlayer, void (*pPrintFN)(void*
 	{
 		pD = &VideoBufferStats.StreamBuffer;
 		pPrintFN(pPlayer, "Video buffer  : EOT %d; EOS %d; %3u AUs; %8u/%8u bytes in; %#7.4fs", pD->bEndOfTrack, pD->bEndOfData, (uint32)pD->NumCurrentAccessUnits, (uint32)pD->CurrentMemInUse, (uint32)PlayerConfig.StreamBufferConfigVideo.MaxDataSize, pD->PlayableDuration.GetAsSeconds());
-		pPrintFN(pPlayer, "Video decoder : %2u in decoder, %zu total, %s; EOD in %d; EOD out %d", (uint32)VideoBufferStats.DecoderOutputBuffer.NumElementsInDecoder, (uint32)VideoBufferStats.DecoderOutputBuffer.MaxDecodedElementsReady, VideoBufferStats.DecoderOutputBuffer.bOutputStalled?"    stalled":"not stalled", VideoBufferStats.DecoderInputBuffer.bEODReached, VideoBufferStats.DecoderOutputBuffer.bEODreached);
+		pPrintFN(pPlayer, "Video decoder : %2u in decoder, %zu total; EOD in %d; EOD out %d; %s", (uint32)VideoBufferStats.DecoderOutputBuffer.NumElementsInDecoder, (uint32)VideoBufferStats.DecoderOutputBuffer.MaxDecodedElementsReady, VideoBufferStats.DecoderInputBuffer.bEODReached, VideoBufferStats.DecoderOutputBuffer.bEODreached, VideoBufferStats.DecoderOutputBuffer.bOutputStalled?"stalled":"not stalled");
 		pPrintFN(pPlayer, "Video renderer: %#7.4fs enqueued", VideoRender.Renderer.IsValid() ? VideoRender.Renderer->GetEnqueuedSampleDuration().GetAsSeconds() : 0.0);
 	}
 	if(1)
 	{
 		pD = &AudioBufferStats.StreamBuffer;
 		pPrintFN(pPlayer, "Audio buffer  : EOT %d; EOS %d; %3u AUs; %8u/%8u bytes in; %#7.4fs", pD->bEndOfTrack, pD->bEndOfData, (uint32)pD->NumCurrentAccessUnits, (uint32)pD->CurrentMemInUse, (uint32)PlayerConfig.StreamBufferConfigAudio.MaxDataSize, pD->PlayableDuration.GetAsSeconds());
-		pPrintFN(pPlayer, "Audio decoder : %2u in decoder, %2u total, %s; EOD in %d; EOD out %d", (uint32)AudioBufferStats.DecoderOutputBuffer.NumElementsInDecoder, (uint32)AudioBufferStats.DecoderOutputBuffer.MaxDecodedElementsReady, AudioBufferStats.DecoderOutputBuffer.bOutputStalled ? "    stalled" : "not stalled", AudioBufferStats.DecoderInputBuffer.bEODReached, AudioBufferStats.DecoderOutputBuffer.bEODreached);
+		pPrintFN(pPlayer, "Audio decoder : %2u in decoder, %2u total; EOD in %d; EOD out %d; %s", (uint32)AudioBufferStats.DecoderOutputBuffer.NumElementsInDecoder, (uint32)AudioBufferStats.DecoderOutputBuffer.MaxDecodedElementsReady, AudioBufferStats.DecoderInputBuffer.bEODReached, AudioBufferStats.DecoderOutputBuffer.bEODreached, AudioBufferStats.DecoderOutputBuffer.bOutputStalled?"stalled":"not stalled");
 		pPrintFN(pPlayer, "Audio renderer: %#7.4fs enqueued", AudioRender.Renderer.IsValid() ? AudioRender.Renderer->GetEnqueuedSampleDuration().GetAsSeconds() : 0.0);
 	}
 	if(1)
 	{
 		pD = &TextBufferStats.StreamBuffer;
-		pPrintFN(pPlayer, "Text buffer   : EOT %d; EOS %d; %3u AUs; %8u/%8u bytes in; eod %d; %#7.4fs", pD->bEndOfTrack, pD->bEndOfData, (uint32)pD->NumCurrentAccessUnits, (uint32)pD->CurrentMemInUse, (uint32)PlayerConfig.StreamBufferConfigText.MaxDataSize, pD->PlayableDuration.GetAsSeconds());
+		pPrintFN(pPlayer, "Text buffer   : EOT %d; EOS %d; %3u AUs; %8u/%8u bytes in; %#7.4fs", pD->bEndOfTrack, pD->bEndOfData, (uint32)pD->NumCurrentAccessUnits, (uint32)pD->CurrentMemInUse, (uint32)PlayerConfig.StreamBufferConfigText.MaxDataSize, pD->PlayableDuration.GetAsSeconds());
 	}
 	DiagnosticsCriticalSection.Unlock();
 
