ElectraPlayer: Fix for a possible crash when flushing the audio renderer while time stretching is active

#rnx

diff --git a/Engine/Plugins/Media/ElectraPlayer/Source/ElectraPlayerRuntime/Private/Runtime/Player/AdaptiveStreamingPlayerRenderers.cpp b/Engine/Plugins/Media/ElectraPlayer/Source/ElectraPlayerRuntime/Private/Runtime/Player/AdaptiveStreamingPlayerRenderers.cpp
--- a/Engine/Plugins/Media/ElectraPlayer/Source/ElectraPlayerRuntime/Private/Runtime/Player/AdaptiveStreamingPlayerRenderers.cpp
+++ b/Engine/Plugins/Media/ElectraPlayer/Source/ElectraPlayerRuntime/Private/Runtime/Player/AdaptiveStreamingPlayerRenderers.cpp
@@ -507,9 +507,8 @@ UEMediaError FAdaptiveStreamingWrappedRenderer::ReleaseBufferPool()
 	ReturnAllPendingBuffers(false);
 	NumEnqueuedSamples = 0;
 	EnqueuedDuration.SetToZero();
-	Lock.Unlock();
-
 	AudioVars.Reset();
+	Lock.Unlock();
 
 	return WrappedRenderer->ReleaseBufferPool();
 }
@@ -545,9 +544,8 @@ UEMediaError FAdaptiveStreamingWrappedRenderer::Flush(const FParamDict& InOption
 	NumEnqueuedSamples = 0;
 	NumBuffersNotHeldBack = 0;
 	EnqueuedDuration.SetToZero();
-	Lock.Unlock();
-
 	AudioVars.Reset();
+	Lock.Unlock();
 
 	return WrappedRenderer->Flush(InOptions);
 }
@@ -641,6 +639,8 @@ double FAdaptiveStreamingWrappedRenderer::GetPlayRateScale()
 
 bool FAdaptiveStreamingWrappedRenderer::ProcessAudio(bool& bOutNeed2ndBuffer, IBuffer* Buffer, FParamDict& InSampleProperties, double InRate)
 {
+	FScopeLock lock(&Lock);
+
 	SCOPE_CYCLE_COUNTER(STAT_ElectraAudioProcessing_Process);
 
 	bool bGetResiduals = bOutNeed2ndBuffer;
@@ -1061,6 +1061,10 @@ int32 FSoundTouchFilter::DequeueSamples(FTimespan& OutTimestamp, int64& OutSeque
 		{
 			EnqueuedSampleBlockInfos.RemoveAt(0);
 		}
+		if (!EnqueuedSampleBlockInfos.Num())
+		{
+			return 0;
+		}
 		// Adjust the frontmost sample time.
 		EnqueuedSampleBlockInfos[0].NumSamples -= ToGo;
 		FTimespan TimeOffset((int64)ToGo * 10000000 / EnqueuedSampleBlockInfos[0].SampleRate);
@@ -1119,7 +1123,6 @@ int32 FSoundTouchFilter::DequeueSamples(FTimespan& OutTimestamp, int64& OutSeque
 	for(int32 i=0, iMax=numGot*NumChannels; i<iMax; ++i)
 	{
 		int32 S = *Flt++ * 32768.0f;
-//TEMP removal		check(S > -33000 && S < 33000);
 		*OutSamples++ = (int16) Min(Max(-32768, S), 32767);
 	}
 	return (int32) numGot;
