Some prep work for tvOS support in Electra & adding YCbCr output also for H.264 decodes to safe on memory
#rnx

diff --git a/Engine/Plugins/Media/ElectraPlayer/Source/ElectraPlayerPlugin/Private/ElectraPlayerPlugin_Apple.cpp b/Engine/Plugins/Media/ElectraPlayer/Source/ElectraPlayerPlugin/Private/ElectraPlayerPlugin_Apple.cpp
--- a/Engine/Plugins/Media/ElectraPlayer/Source/ElectraPlayerPlugin/Private/ElectraPlayerPlugin_Apple.cpp
+++ b/Engine/Plugins/Media/ElectraPlayer/Source/ElectraPlayerPlugin/Private/ElectraPlayerPlugin_Apple.cpp
@@ -2,7 +2,7 @@
 
 #include "CoreTypes.h"
 
-#if PLATFORM_MAC || PLATFORM_IOS
+#if PLATFORM_MAC || PLATFORM_IOS || PLATFORM_TVOS
 
 #include <ElectraPlayerPlugin.h>
 
@@ -15,4 +15,4 @@ void FElectraPlayerPlugin::PlatformSetupResourceParams(Electra::FParamDict& Para
 {
 }
 
-#endif // PLATFORM_MAC || PLATFORM_IOS
+#endif // PLATFORM_MAC || PLATFORM_IOS || PLATFORM_TVOS
diff --git a/Engine/Plugins/Media/ElectraPlayer/Source/ElectraPlayerRuntime/Private/Runtime/Decoder/Apple/VideoDecoderH264_Apple.cpp b/Engine/Plugins/Media/ElectraPlayer/Source/ElectraPlayerRuntime/Private/Runtime/Decoder/Apple/VideoDecoderH264_Apple.cpp
--- a/Engine/Plugins/Media/ElectraPlayer/Source/ElectraPlayerRuntime/Private/Runtime/Decoder/Apple/VideoDecoderH264_Apple.cpp
+++ b/Engine/Plugins/Media/ElectraPlayer/Source/ElectraPlayerRuntime/Private/Runtime/Decoder/Apple/VideoDecoderH264_Apple.cpp
@@ -737,24 +737,37 @@ bool FVideoDecoderH264::InternalDecoderCreate(CMFormatDescriptionRef InputFormat
 
 	// Output image format configuration
 	CFMutableDictionaryRef OutputImageFormat = CFDictionaryCreateMutable(nullptr, 3, &kCFTypeDictionaryKeyCallBacks, &kCFTypeDictionaryValueCallBacks);
+
+	// Allow for an array of output formats, so the decoder can deliver SDR and HDR content in suitable buffers
+	const int NumFmts = 4;
+	int pxfmt[NumFmts] = {
+		kCVPixelFormatType_420YpCbCr8BiPlanarVideoRange,
+		kCVPixelFormatType_420YpCbCr8BiPlanarFullRange,
+		kCVPixelFormatType_420YpCbCr10BiPlanarVideoRange,
+		kCVPixelFormatType_420YpCbCr10BiPlanarFullRange
+	};
+	CFNumberRef PixelFormats[NumFmts] = {
+		CFNumberCreate(kCFAllocatorDefault, kCFNumberSInt32Type, &pxfmt[0]),
+		CFNumberCreate(kCFAllocatorDefault, kCFNumberSInt32Type, &pxfmt[1]),
+		CFNumberCreate(kCFAllocatorDefault, kCFNumberSInt32Type, &pxfmt[2]),
+		CFNumberCreate(kCFAllocatorDefault, kCFNumberSInt32Type, &pxfmt[3])
+	};
+	CFArrayRef PixFmtArray = CFArrayCreate(kCFAllocatorDefault, (const void**)PixelFormats, NumFmts, &kCFTypeArrayCallBacks);
+
+	CFDictionarySetValue(OutputImageFormat, kCVPixelBufferPixelFormatTypeKey, PixFmtArray);
+
 	// Choice of: kCVPixelBufferOpenGLCompatibilityKey (all)  kCVPixelBufferOpenGLESCompatibilityKey (iOS only)   kCVPixelBufferMetalCompatibilityKey (all)
-#if PLATFORM_MAC
-	int pxfmt = kCVPixelFormatType_32BGRA;
-	CFNumberRef PixelFormat = CFNumberCreate(kCFAllocatorDefault, kCFNumberSInt32Type, &pxfmt);
-	CFDictionarySetValue(OutputImageFormat, kCVPixelBufferPixelFormatTypeKey, PixelFormat);
-	CFRelease(PixelFormat);
 	CFDictionarySetValue(OutputImageFormat, kCVPixelBufferMetalCompatibilityKey, kCFBooleanTrue);
-#elif PLATFORM_IOS
+#if PLATFORM_IOS || PLATFORM_TVOS
 	CFDictionarySetValue(OutputImageFormat, kCVPixelBufferOpenGLESCompatibilityKey, kCFBooleanFalse);
-	CFDictionarySetValue(OutputImageFormat, kCVPixelBufferMetalCompatibilityKey, kCFBooleanTrue);
-	int pxfmt = kCVPixelFormatType_32BGRA;
-	CFNumberRef PixelFormat = CFNumberCreate(kCFAllocatorDefault, kCFNumberSInt32Type, &pxfmt);
-	CFDictionarySetValue(OutputImageFormat, kCVPixelBufferPixelFormatTypeKey, PixelFormat);
-	CFRelease(PixelFormat);
-#else
-	#error "Should not get here. Check platform checks at the top of the file."
 #endif
 
+	CFRelease(PixFmtArray);
+	for (int Idx = 0; Idx < NumFmts; ++Idx)
+	{
+		CFRelease(PixelFormats[Idx]);
+	}
+
 	// Session configuration
 	CFMutableDictionaryRef SessionConfiguration = CFDictionaryCreateMutable(nullptr, 1, &kCFTypeDictionaryKeyCallBacks, &kCFTypeDictionaryValueCallBacks);
 	// Ask for hardware decoding
diff --git a/Engine/Plugins/Media/ElectraPlayer/Source/ElectraPlayerRuntime/Private/Runtime/Decoder/Apple/VideoDecoderH265_Apple.cpp b/Engine/Plugins/Media/ElectraPlayer/Source/ElectraPlayerRuntime/Private/Runtime/Decoder/Apple/VideoDecoderH265_Apple.cpp
--- a/Engine/Plugins/Media/ElectraPlayer/Source/ElectraPlayerRuntime/Private/Runtime/Decoder/Apple/VideoDecoderH265_Apple.cpp
+++ b/Engine/Plugins/Media/ElectraPlayer/Source/ElectraPlayerRuntime/Private/Runtime/Decoder/Apple/VideoDecoderH265_Apple.cpp
@@ -786,7 +786,7 @@ bool FVideoDecoderH265::InternalDecoderCreate(CMFormatDescriptionRef InputFormat
 
 	// Choice of: kCVPixelBufferOpenGLCompatibilityKey (all)  kCVPixelBufferOpenGLESCompatibilityKey (iOS only)   kCVPixelBufferMetalCompatibilityKey (all)
 	CFDictionarySetValue(OutputImageFormat, kCVPixelBufferMetalCompatibilityKey, kCFBooleanTrue);
-#if PLATFORM_IOS
+#if PLATFORM_IOS || PLATFORM_TVOS
 	CFDictionarySetValue(OutputImageFormat, kCVPixelBufferOpenGLESCompatibilityKey, kCFBooleanFalse);
 #endif
 
