ElectraPlayer: Waiting for remote elements to download when retrying on a different quality level.

#rnx

diff --git a/Engine/Plugins/Media/ElectraPlayer/Source/ElectraPlayerRuntime/Private/Runtime/Player/ABRRules/ABRVoDPlus.cpp b/Engine/Plugins/Media/ElectraPlayer/Source/ElectraPlayerRuntime/Private/Runtime/Player/ABRRules/ABRVoDPlus.cpp
--- a/Engine/Plugins/Media/ElectraPlayer/Source/ElectraPlayerRuntime/Private/Runtime/Player/ABRRules/ABRVoDPlus.cpp
+++ b/Engine/Plugins/Media/ElectraPlayer/Source/ElectraPlayerRuntime/Private/Runtime/Player/ABRRules/ABRVoDPlus.cpp
@@ -625,6 +625,11 @@ IAdaptiveStreamSelector::ESegmentAction FABROnDemandPlus::EvaluateForError(TArra
 
 			if (!Stats.bWasSuccessful)
 			{
+				if (Stats.bWaitingForRemoteRetryElement)
+				{
+					return IAdaptiveStreamSelector::ESegmentAction::Retry;
+				}
+
 				// If this is the first failure for this segment (not being retried yet) increase
 				// the number of consecutively failed segments.
 				if (Stats.RetryNumber == 0)
diff --git a/Engine/Plugins/Media/ElectraPlayer/Source/ElectraPlayerRuntime/Private/Runtime/Player/AdaptiveStreamingPlayerMetrics.h b/Engine/Plugins/Media/ElectraPlayer/Source/ElectraPlayerRuntime/Private/Runtime/Player/AdaptiveStreamingPlayerMetrics.h
--- a/Engine/Plugins/Media/ElectraPlayer/Source/ElectraPlayerRuntime/Private/Runtime/Player/AdaptiveStreamingPlayerMetrics.h
+++ b/Engine/Plugins/Media/ElectraPlayer/Source/ElectraPlayerRuntime/Private/Runtime/Player/AdaptiveStreamingPlayerMetrics.h
@@ -147,6 +147,7 @@ namespace Metrics
 		bool			bParseFailure = false;				//!< true if the segment could not be parsed
 		bool			bInsertedFillerData = false;
 		bool			bIsCachedResponse = false;
+		bool			bWaitingForRemoteRetryElement = false;
 
 		// Chunk timing
 		struct FMovieChunkInfo
diff --git a/Engine/Plugins/Media/ElectraPlayer/Source/ElectraPlayerRuntime/Private/Runtime/Player/DASH/ManifestDASH.cpp b/Engine/Plugins/Media/ElectraPlayer/Source/ElectraPlayerRuntime/Private/Runtime/Player/DASH/ManifestDASH.cpp
--- a/Engine/Plugins/Media/ElectraPlayer/Source/ElectraPlayerRuntime/Private/Runtime/Player/DASH/ManifestDASH.cpp
+++ b/Engine/Plugins/Media/ElectraPlayer/Source/ElectraPlayerRuntime/Private/Runtime/Player/DASH/ManifestDASH.cpp
@@ -177,7 +177,7 @@ private:
 		SamePeriodStartOver,
 		NextPeriod,
 	};
-	IManifest::FResult GetNextOrRetrySegment(TSharedPtrTS<IStreamSegment>& OutSegment, TSharedPtrTS<const IStreamSegment> InCurrentSegment, ENextSegType InNextType, const FPlayStartOptions& Options);
+	IManifest::FResult GetNextOrRetrySegment(TSharedPtrTS<IStreamSegment>& OutSegment, bool& OutWaitForRemoteElement, TSharedPtrTS<const IStreamSegment> InCurrentSegment, ENextSegType InNextType, const FPlayStartOptions& Options);
 
 	bool PrepareDRM(const TArray<FManifestDASHInternal::FAdaptationSet::FContentProtection>& InContentProtections);
 
@@ -1308,17 +1308,19 @@ IManifest::FResult FDASHPlayPeriod::GetContinuationSegment(TSharedPtrTS<IStreamS
 {
 	// Create a dummy request we can use to pass into GetNextOrRetrySegment().
 	// Only set the values that that method requires.
+	bool bNeedRemoteElement = false;
 	auto DummyReq = MakeSharedTS<FStreamSegmentRequestFMP4DASH>();
 	DummyReq->StreamType = StreamType;
 	DummyReq->PeriodStart = StartPosition.Time;
 	DummyReq->TimestampSequenceIndex = SequenceState.SequenceIndex;
-	return GetNextOrRetrySegment(OutSegment, DummyReq, ENextSegType::SamePeriodStartOver, StartPosition.Options);
+	return GetNextOrRetrySegment(OutSegment, bNeedRemoteElement, DummyReq, ENextSegType::SamePeriodStartOver, StartPosition.Options);
 }
 
 
 
-IManifest::FResult FDASHPlayPeriod::GetNextOrRetrySegment(TSharedPtrTS<IStreamSegment>& OutSegment, TSharedPtrTS<const IStreamSegment> InCurrentSegment, ENextSegType InNextType, const FPlayStartOptions& Options)
+IManifest::FResult FDASHPlayPeriod::GetNextOrRetrySegment(TSharedPtrTS<IStreamSegment>& OutSegment, bool& OutWaitForRemoteElement, TSharedPtrTS<const IStreamSegment> InCurrentSegment, ENextSegType InNextType, const FPlayStartOptions& Options)
 {
+	OutWaitForRemoteElement = false;
 	TSharedPtrTS<const FStreamSegmentRequestFMP4DASH> Current = StaticCastSharedPtr<const FStreamSegmentRequestFMP4DASH>(InCurrentSegment);
 	if (Current->bIsInitialStartRequest)
 	{
@@ -1440,6 +1442,7 @@ IManifest::FResult FDASHPlayPeriod::GetNextOrRetrySegment(TSharedPtrTS<IStreamSe
 		{
 			return IManifest::FResult(IManifest::FResult::EType::NotFound).SetErrorDetail(FErrorDetail().SetMessage("Entity loader disappeared"));
 		}
+		OutWaitForRemoteElement = true;
 		IPlaylistReaderDASH* Reader = static_cast<IPlaylistReaderDASH*>(ManifestReader.Get());
 		Reader->AddElementLoadRequests(RemoteElementLoadRequests);
 		return IManifest::FResult().RetryAfterMilliseconds(100);
@@ -1543,18 +1546,19 @@ IManifest::FResult FDASHPlayPeriod::GetNextSegment(TSharedPtrTS<IStreamSegment>&
 	const FStreamSegmentRequestFMP4DASH* CurrentRequest = static_cast<const FStreamSegmentRequestFMP4DASH*>(InCurrentSegment.Get());
 
 	// Check if we moved across a period.
+	bool bNeedRemoteElement = false;
 	if (CurrentRequest->Period->GetUniqueIdentifier().Equals(PeriodID))
 	{
 		if (CurrentRequest->Segment.bSawLMSG)
 		{
 			return IManifest::FResult(IManifest::FResult::EType::PastEOS);
 		}
-		return GetNextOrRetrySegment(OutSegment, InCurrentSegment, ENextSegType::SamePeriodNext, Options);
+		return GetNextOrRetrySegment(OutSegment, bNeedRemoteElement, InCurrentSegment, ENextSegType::SamePeriodNext, Options);
 	}
 	else
 	{
 		// Moved into a new period. This here is the new period.
-		return GetNextOrRetrySegment(OutSegment, InCurrentSegment, ENextSegType::NextPeriod, Options);
+		return GetNextOrRetrySegment(OutSegment, bNeedRemoteElement, InCurrentSegment, ENextSegType::NextPeriod, Options);
 	}
 }
 
@@ -1576,7 +1580,12 @@ IManifest::FResult FDASHPlayPeriod::GetRetrySegment(TSharedPtrTS<IStreamSegment>
 		OutSegment = NewRequest;
 		return IManifest::FResult(IManifest::FResult::EType::Found);
 	}
-	return GetNextOrRetrySegment(OutSegment, InCurrentSegment, ENextSegType::SamePeriodRetry, Options);
+	
+	// Pass the download stats bWaitingForRemoteRetryElement to convey if the retry segment needs to wait for a remote element,
+	// which is either some xlink or an index segment.
+	FStreamSegmentRequestFMP4DASH* CurrentRequest = const_cast<FStreamSegmentRequestFMP4DASH*>(static_cast<const FStreamSegmentRequestFMP4DASH*>(InCurrentSegment.Get()));
+	IManifest::FResult Result = GetNextOrRetrySegment(OutSegment, CurrentRequest->DownloadStats.bWaitingForRemoteRetryElement, InCurrentSegment, ENextSegType::SamePeriodRetry, Options);
+	return Result;
 }
 
 
diff --git a/Engine/Plugins/Media/ElectraPlayer/Source/ElectraPlayerRuntime/Private/Runtime/Player/HLS/ManifestHLS.cpp b/Engine/Plugins/Media/ElectraPlayer/Source/ElectraPlayerRuntime/Private/Runtime/Player/HLS/ManifestHLS.cpp
--- a/Engine/Plugins/Media/ElectraPlayer/Source/ElectraPlayerRuntime/Private/Runtime/Player/HLS/ManifestHLS.cpp
+++ b/Engine/Plugins/Media/ElectraPlayer/Source/ElectraPlayerRuntime/Private/Runtime/Player/HLS/ManifestHLS.cpp
@@ -86,7 +86,7 @@ private:
 
 	IManifest::FResult GetMediaStreamForID(TSharedPtrTS<FManifestHLSInternal::FPlaylistBase>& OutPlaylist, TSharedPtrTS<FManifestHLSInternal::FMediaStream>& OutMediaStream, uint32 UniqueID) const;
 
-	IManifest::FResult GetNextOrRetrySegment(TSharedPtrTS<IStreamSegment>& OutSegment, TSharedPtrTS<const IStreamSegment> CurrentSegment, const FPlayStartOptions& Options, bool bRetry);
+	IManifest::FResult GetNextOrRetrySegment(TSharedPtrTS<IStreamSegment>& OutSegment, bool& OutWaitForRemoteElement, TSharedPtrTS<const IStreamSegment> CurrentSegment, const FPlayStartOptions& Options, bool bRetry);
 
 	IManifest::FResult FindSegment(TSharedPtrTS<FStreamSegmentRequestHLSfmp4>& OutRequest, TSharedPtrTS<FManifestHLSInternal::FPlaylistBase> InPlaylist, TSharedPtrTS<FManifestHLSInternal::FMediaStream> InStream, uint32 StreamUniqueID, EStreamType StreamType, const FSegSearchParam& SearchParam, IManifest::ESearchType SearchType);
 
@@ -1056,7 +1056,7 @@ IManifest::FResult FPlayPeriodHLS::GetLoopingSegment(TSharedPtrTS<IStreamSegment
 
 
 
-IManifest::FResult FPlayPeriodHLS::GetNextOrRetrySegment(TSharedPtrTS<IStreamSegment>& OutSegment, TSharedPtrTS<const IStreamSegment> InCurrentSegment, const FPlayStartOptions& Options, bool bRetry)
+IManifest::FResult FPlayPeriodHLS::GetNextOrRetrySegment(TSharedPtrTS<IStreamSegment>& OutSegment, bool& OutWaitForRemoteElement, TSharedPtrTS<const IStreamSegment> InCurrentSegment, const FPlayStartOptions& Options, bool bRetry)
 {
 	// Need to have a current segment to find the next one.
 	if (!InCurrentSegment.IsValid())
@@ -1135,6 +1135,7 @@ IManifest::FResult FPlayPeriodHLS::GetNextOrRetrySegment(TSharedPtrTS<IStreamSeg
 				NextSegmentRequest->NumOverallRetries = CurrentRequest->NumOverallRetries + 1;
 			}
 			OutSegment = NextSegmentRequest;
+			OutWaitForRemoteElement = false;
 			return IManifest::FResult(IManifest::FResult::EType::Found);
 		}
 		if (Result.GetType() == IManifest::FResult::EType::TryAgainLater)
@@ -1155,6 +1156,7 @@ IManifest::FResult FPlayPeriodHLS::GetNextOrRetrySegment(TSharedPtrTS<IStreamSeg
 			check(Playlist.IsValid());
 
 			Playlist->Internal.LoadState = FManifestHLSInternal::FPlaylistBase::FInternal::ELoadState::Pending;
+			OutWaitForRemoteElement = true;
 
 			FURL_RFC3986 UrlBuilder;
 			UrlBuilder.Parse(InternalManifest->MasterPlaylistVars.PlaylistLoadRequest.URL);
@@ -1187,7 +1189,8 @@ IManifest::FResult FPlayPeriodHLS::GetNextOrRetrySegment(TSharedPtrTS<IStreamSeg
  */
 IManifest::FResult FPlayPeriodHLS::GetNextSegment(TSharedPtrTS<IStreamSegment>& OutSegment, TSharedPtrTS<const IStreamSegment> InCurrentSegment, const FPlayStartOptions& Options)
 {
-	return GetNextOrRetrySegment(OutSegment, InCurrentSegment, Options, false);
+	bool bNeedRemoteElement = false;
+	return GetNextOrRetrySegment(OutSegment, bNeedRemoteElement, InCurrentSegment, Options, false);
 }
 
 
@@ -1205,7 +1208,12 @@ IManifest::FResult FPlayPeriodHLS::GetRetrySegment(TSharedPtrTS<IStreamSegment>&
 		OutSegment = NewRequest;
 		return IManifest::FResult(IManifest::FResult::EType::Found);
 	}
-	return GetNextOrRetrySegment(OutSegment, InCurrentSegment, Options, true);
+
+	// Pass the download stats bWaitingForRemoteRetryElement to convey if the retry segment needs to wait for
+	// the variant playlist download.
+	FStreamSegmentRequestHLSfmp4* CurrentRequest = const_cast<FStreamSegmentRequestHLSfmp4*>(static_cast<const FStreamSegmentRequestHLSfmp4*>(InCurrentSegment.Get()));
+	IManifest::FResult Result = GetNextOrRetrySegment(OutSegment, CurrentRequest->DownloadStats.bWaitingForRemoteRetryElement, InCurrentSegment, Options, true);
+	return Result;
 }
 
 
diff --git a/Engine/Plugins/Media/ElectraPlayer/Source/ElectraPlayerRuntime/Private/Runtime/Player/HLS/StreamReaderHLSfmp4.cpp b/Engine/Plugins/Media/ElectraPlayer/Source/ElectraPlayerRuntime/Private/Runtime/Player/HLS/StreamReaderHLSfmp4.cpp
--- a/Engine/Plugins/Media/ElectraPlayer/Source/ElectraPlayerRuntime/Private/Runtime/Player/HLS/StreamReaderHLSfmp4.cpp
+++ b/Engine/Plugins/Media/ElectraPlayer/Source/ElectraPlayerRuntime/Private/Runtime/Player/HLS/StreamReaderHLSfmp4.cpp
@@ -544,6 +544,9 @@ FStreamReaderHLSfmp4::FStreamHandler::EInitSegmentResult FStreamReaderHLSfmp4::F
 				HTTP->Parameters.Range.Start		= Request->InitSegmentInfo->ByteRange.GetStart();
 				HTTP->Parameters.Range.EndIncluding = Request->InitSegmentInfo->ByteRange.GetEnd();
 			}
+			// Set timeouts for init segment retrieval
+			HTTP->Parameters.ConnectTimeout = FTimeValue().SetFromMilliseconds(1000 * 4);
+			HTTP->Parameters.NoDataTimeout = FTimeValue().SetFromMilliseconds(1000 * 2);
 
 			ProgressReportCount = 0;
 			DownloadCompleteSignal.Reset();
@@ -828,6 +831,9 @@ void FStreamReaderHLSfmp4::FStreamHandler::HandleRequest()
 					HTTP->Parameters.Range = Request->Range;
 				}
 				HTTP->ResponseCache = PlayerSessionService->GetHTTPResponseCache();
+				// Set timeouts for media segment retrieval
+				HTTP->Parameters.ConnectTimeout = FTimeValue().SetFromMilliseconds(1000 * 4);
+				HTTP->Parameters.NoDataTimeout = FTimeValue().SetFromMilliseconds(1000 * 2);
 
 				ProgressReportCount = 0;
 				DownloadCompleteSignal.Reset();
