Enabling tvOS support for Electra

diff --git a/Engine/Plugins/Media/ElectraCDM/ElectraCDM.uplugin b/Engine/Plugins/Media/ElectraCDM/ElectraCDM.uplugin
--- a/Engine/Plugins/Media/ElectraCDM/ElectraCDM.uplugin
+++ b/Engine/Plugins/Media/ElectraCDM/ElectraCDM.uplugin
@@ -19,8 +19,8 @@
 			"Name": "ElectraCDM",
 			"Type": "RuntimeNoCommandlet",
 			"LoadingPhase": "PostEngineInit",
-			"PlatformAllowList": [ "Win64", "Mac", "IOS", "Android", "Linux" ],
+			"PlatformAllowList": [ "Win64", "Mac", "IOS", "TVOS", "Android", "Linux" ],
 			"TargetDenyList": [ "Server" ]
 		},
 	]
-}
\ No newline at end of file
+}
diff --git a/Engine/Plugins/Media/ElectraPlayer/ElectraPlayer.uplugin b/Engine/Plugins/Media/ElectraPlayer/ElectraPlayer.uplugin
--- a/Engine/Plugins/Media/ElectraPlayer/ElectraPlayer.uplugin
+++ b/Engine/Plugins/Media/ElectraPlayer/ElectraPlayer.uplugin
@@ -25,14 +25,14 @@
 			"Name": "ElectraPlayerRuntime",
 			"Type": "RuntimeNoCommandlet",
 			"LoadingPhase": "PreLoadingScreen",
-			"PlatformAllowList": [ "Win64", "Mac", "IOS", "Android", "Linux" ],
+			"PlatformAllowList": [ "Win64", "Mac", "IOS", "TVOS", "Android", "Linux" ],
 			"TargetDenyList": [ "Server" ]
 		},
 		{
 			"Name": "ElectraPlayerPlugin",
 			"Type": "RuntimeNoCommandlet",
 			"LoadingPhase": "PreLoadingScreen",
-			"PlatformAllowList": [ "Win64", "Mac", "IOS", "Android", "Linux" ],
+			"PlatformAllowList": [ "Win64", "Mac", "IOS", "TVOS", "Android", "Linux" ],
 			"TargetDenyList": [ "Server" ]
 		},
 		{
@@ -46,7 +46,7 @@
 			"Name": "ElectraPlayerFactory",
 			"Type": "RuntimeNoCommandlet",
 			"LoadingPhase": "PostEngineInit",
-			"PlatformAllowList": [ "Win64", "Mac", "IOS", "Android", "Linux" ],
+			"PlatformAllowList": [ "Win64", "Mac", "IOS", "TVOS", "Android", "Linux" ],
 			"TargetDenyList": [ "Server" ]
 		}
 		
diff --git a/Engine/Plugins/Media/ElectraPlayer/Source/ElectraPlayerFactory/ElectraPlayerFactory.Build.cs b/Engine/Plugins/Media/ElectraPlayer/Source/ElectraPlayerFactory/ElectraPlayerFactory.Build.cs
--- a/Engine/Plugins/Media/ElectraPlayer/Source/ElectraPlayerFactory/ElectraPlayerFactory.Build.cs
+++ b/Engine/Plugins/Media/ElectraPlayer/Source/ElectraPlayerFactory/ElectraPlayerFactory.Build.cs
@@ -63,6 +63,7 @@ namespace UnrealBuildTool.Rules
 			return Target.IsInPlatformGroup(UnrealPlatformGroup.Windows)
 				|| Target.Platform == UnrealTargetPlatform.Mac
 				|| Target.Platform == UnrealTargetPlatform.IOS
+                || Target.Platform == UnrealTargetPlatform.TVOS
 				|| Target.IsInPlatformGroup(UnrealPlatformGroup.Android)
 				|| Target.IsInPlatformGroup(UnrealPlatformGroup.Unix)
 				;
diff --git a/Engine/Plugins/Media/ElectraPlayer/Source/ElectraPlayerFactory/Private/ElectraPlayerFactoryModule.cpp b/Engine/Plugins/Media/ElectraPlayer/Source/ElectraPlayerFactory/Private/ElectraPlayerFactoryModule.cpp
--- a/Engine/Plugins/Media/ElectraPlayer/Source/ElectraPlayerFactory/Private/ElectraPlayerFactoryModule.cpp
+++ b/Engine/Plugins/Media/ElectraPlayer/Source/ElectraPlayerFactory/Private/ElectraPlayerFactoryModule.cpp
@@ -132,6 +132,7 @@ public:
 		AddSupportedPlatform(FGuid(0xb596ce6f, 0xd8324a9c, 0x84e9f880, 0x21322535));
 		AddSupportedPlatform(FGuid(0x115de4fe, 0x241b465b, 0x970a872f, 0x3167492a));
 		AddSupportedPlatform(FGuid(0xc0b45a33, 0x9de340c7, 0xbce24c47, 0x15c3babf));
+        AddSupportedPlatform(FGuid(0xa478294f, 0xbd0d4ec0, 0x8830b6d4, 0xd219c1a4));
 #endif
 
 		// supported schemes
diff --git a/Engine/Plugins/Media/ElectraPlayer/Source/ElectraPlayerPlugin/ElectraPlayerPlugin.Build.cs b/Engine/Plugins/Media/ElectraPlayer/Source/ElectraPlayerPlugin/ElectraPlayerPlugin.Build.cs
--- a/Engine/Plugins/Media/ElectraPlayer/Source/ElectraPlayerPlugin/ElectraPlayerPlugin.Build.cs
+++ b/Engine/Plugins/Media/ElectraPlayer/Source/ElectraPlayerPlugin/ElectraPlayerPlugin.Build.cs
@@ -43,7 +43,7 @@ namespace UnrealBuildTool.Rules
 				PrivateDependencyModuleNames.Add("MetalRHI");
 			}
 
-			if (Target.Platform == UnrealTargetPlatform.IOS )
+			if (Target.Platform == UnrealTargetPlatform.IOS || Target.Platform == UnrealTargetPlatform.TVOS)
 			{
 				PrivateDependencyModuleNames.Add("MetalRHI");
 			}
diff --git a/Engine/Plugins/Media/ElectraPlayer/Source/ElectraPlayerRuntime/ElectraPlayerRuntime.Build.cs b/Engine/Plugins/Media/ElectraPlayer/Source/ElectraPlayerRuntime/ElectraPlayerRuntime.Build.cs
--- a/Engine/Plugins/Media/ElectraPlayer/Source/ElectraPlayerRuntime/ElectraPlayerRuntime.Build.cs
+++ b/Engine/Plugins/Media/ElectraPlayer/Source/ElectraPlayerRuntime/ElectraPlayerRuntime.Build.cs
@@ -106,7 +106,7 @@ namespace UnrealBuildTool.Rules
 
 				PrivateIncludePaths.Add("ElectraPlayerRuntime/Private/Runtime/Decoder/Apple");
 			}
-			else if (Target.Platform == UnrealTargetPlatform.IOS )
+			else if (Target.Platform == UnrealTargetPlatform.IOS || Target.Platform == UnrealTargetPlatform.TVOS)
 			{
 				PublicFrameworks.AddRange(
 				new string[] {
diff --git a/Engine/Plugins/Media/ElectraSubtitles/ElectraSubtitles.uplugin b/Engine/Plugins/Media/ElectraSubtitles/ElectraSubtitles.uplugin
--- a/Engine/Plugins/Media/ElectraSubtitles/ElectraSubtitles.uplugin
+++ b/Engine/Plugins/Media/ElectraSubtitles/ElectraSubtitles.uplugin
@@ -19,7 +19,7 @@
 			"Name": "ElectraSubtitles",
 			"Type": "RuntimeNoCommandlet",
 			"LoadingPhase": "PostEngineInit",
-			"PlatformAllowList": [ "Win64", "Mac", "IOS", "Android", "Linux" ],
+			"PlatformAllowList": [ "Win64", "Mac", "IOS", "TVOS", "Android", "Linux" ],
 			"TargetDenyList": [ "Server" ]
 		}
 	],
@@ -29,4 +29,4 @@
 			"Enabled": true
 		}
 	]
-}
\ No newline at end of file
+}
diff --git a/Engine/Plugins/Media/ElectraUtil/ElectraUtil.uplugin b/Engine/Plugins/Media/ElectraUtil/ElectraUtil.uplugin
--- a/Engine/Plugins/Media/ElectraUtil/ElectraUtil.uplugin
+++ b/Engine/Plugins/Media/ElectraUtil/ElectraUtil.uplugin
@@ -19,22 +19,22 @@
 			"Name": "ElectraBase",
 			"Type": "RuntimeNoCommandlet",
 			"LoadingPhase": "PostEngineInit",
-			"PlatformAllowList": [ "Win64", "Mac", "IOS", "Android", "Linux" ],
+			"PlatformAllowList": [ "Win64", "Mac", "IOS", "TVOS", "Android", "Linux" ],
 			"TargetDenyList": [ "Server" ]
 		},
 		{
 			"Name": "ElectraSamples",
 			"Type": "RuntimeNoCommandlet",
 			"LoadingPhase": "PostEngineInit",
-			"PlatformAllowList": [ "Win64", "Mac", "IOS", "Android", "Linux" ],
+			"PlatformAllowList": [ "Win64", "Mac", "IOS", "TVOS", "Android", "Linux" ],
 			"TargetDenyList": [ "Server" ]
 		},
 		{
 			"Name": "ElectraHTTPStream",
 			"Type": "RuntimeNoCommandlet",
 			"LoadingPhase": "PostEngineInit",
-			"PlatformAllowList": [ "Win64", "Mac", "IOS", "Android", "Linux" ],
+			"PlatformAllowList": [ "Win64", "Mac", "IOS", "TVOS", "Android", "Linux" ],
 			"TargetDenyList": [ "Server" ]
 		}
 	]
-}
\ No newline at end of file
+}
diff --git a/Engine/Plugins/Media/ElectraUtil/Source/ElectraSamples/ElectraSamples.Build.cs b/Engine/Plugins/Media/ElectraUtil/Source/ElectraSamples/ElectraSamples.Build.cs
--- a/Engine/Plugins/Media/ElectraUtil/Source/ElectraSamples/ElectraSamples.Build.cs
+++ b/Engine/Plugins/Media/ElectraUtil/Source/ElectraSamples/ElectraSamples.Build.cs
@@ -64,7 +64,7 @@ namespace UnrealBuildTool.Rules
 
 				PublicIncludePaths.Add("$(ModuleDir)/Public/Windows");
 			}
-			else if (Target.Platform == UnrealTargetPlatform.Mac || Target.Platform == UnrealTargetPlatform.IOS)
+			else if (Target.Platform == UnrealTargetPlatform.Mac || Target.Platform == UnrealTargetPlatform.IOS || Target.Platform == UnrealTargetPlatform.TVOS)
 			{
 				PrivateDependencyModuleNames.Add("MetalRHI");
 				PublicIncludePaths.Add("$(ModuleDir)/Public/Apple");
