ElectraPlayer: Fix for HDR information after seek

diff --git a/Engine/Plugins/Media/ElectraPlayer/Source/ElectraPlayerRuntime/Private/Runtime/Decoder/Android/VideoDecoderH265_Android.cpp b/Engine/Plugins/Media/ElectraPlayer/Source/ElectraPlayerRuntime/Private/Runtime/Decoder/Android/VideoDecoderH265_Android.cpp
--- a/Engine/Plugins/Media/ElectraPlayer/Source/ElectraPlayerRuntime/Private/Runtime/Decoder/Android/VideoDecoderH265_Android.cpp
+++ b/Engine/Plugins/Media/ElectraPlayer/Source/ElectraPlayerRuntime/Private/Runtime/Decoder/Android/VideoDecoderH265_Android.cpp
@@ -1558,6 +1558,32 @@ FVideoDecoderH265::EOutputResult FVideoDecoderH265::ProcessOutput(const FDecoded
 		NotifyReadyBufferListener(bHaveAvailSmpBlk);
 		if (RenderOutputBuffer)
 		{
+			// Set the bit depth and the colorimetry.
+			uint8 colour_primaries = 2, transfer_characteristics = 2, matrix_coeffs = 2;
+			uint8 video_full_range_flag = 0, video_format = 5;
+			uint8 num_bits = 8;
+			if (NextImage.SourceInfo.IsValid() && NextImage.SourceInfo->SPSs.Num())
+			{
+				check(NextImage.SourceInfo->SPSs[0].bit_depth_luma_minus8 == NextImage.SourceInfo->SPSs[0].bit_depth_chroma_minus8);
+				num_bits = NextImage.SourceInfo->SPSs[0].bit_depth_luma_minus8 + 8;
+				if (NextImage.SourceInfo->SPSs[0].colour_description_present_flag)
+				{
+					colour_primaries = NextImage.SourceInfo->SPSs[0].colour_primaries;
+					transfer_characteristics = NextImage.SourceInfo->SPSs[0].transfer_characteristics;
+					matrix_coeffs = NextImage.SourceInfo->SPSs[0].matrix_coeffs;
+				}
+				if (NextImage.SourceInfo->SPSs[0].video_signal_type_present_flag)
+				{
+					video_full_range_flag = NextImage.SourceInfo->SPSs[0].video_full_range_flag;
+					video_format = NextImage.SourceInfo->SPSs[0].video_format;
+				}
+			}
+			Colorimetry.Update(colour_primaries, transfer_characteristics, matrix_coeffs, video_full_range_flag, video_format);
+			if (NextImage.SourceInfo.IsValid())
+			{
+				HDR.Update(num_bits, Colorimetry, NextImage.SourceInfo->CSDPrefixSEIMessages, NextImage.SourceInfo->PrefixSEIMessages, false);
+			}
+
 			TUniquePtr<FParamDict> OutputBufferSampleProperties = MakeUnique<FParamDict>();
 			bool bRender = NextImage.SourceInfo->AdjustedPTS.IsValid();
 			if (bRender)
@@ -1587,37 +1613,9 @@ FVideoDecoderH265::EOutputResult FVideoDecoderH265::ProcessOutput(const FDecoded
 					OutputBufferSampleProperties->Set("aspect_h", FVariantValue(ay));
 					OutputBufferSampleProperties->Set("fps_num", FVariantValue((int64)0));
 					OutputBufferSampleProperties->Set("fps_denom", FVariantValue((int64)0));
-
-					// Set the bit depth and the colorimetry.
-					uint8 colour_primaries=2, transfer_characteristics=2, matrix_coeffs=2;
-					uint8 video_full_range_flag=0, video_format=5;
-					uint8 num_bits = 8;
-					if (NextImage.SourceInfo.IsValid() && NextImage.SourceInfo->SPSs.Num())
-					{
-						check(NextImage.SourceInfo->SPSs[0].bit_depth_luma_minus8 == NextImage.SourceInfo->SPSs[0].bit_depth_chroma_minus8);
-						num_bits = NextImage.SourceInfo->SPSs[0].bit_depth_luma_minus8 + 8;
-						if (NextImage.SourceInfo->SPSs[0].colour_description_present_flag)
-						{
-							colour_primaries = NextImage.SourceInfo->SPSs[0].colour_primaries;
-							transfer_characteristics = NextImage.SourceInfo->SPSs[0].transfer_characteristics;
-							matrix_coeffs = NextImage.SourceInfo->SPSs[0].matrix_coeffs;
-						}
-						if (NextImage.SourceInfo->SPSs[0].video_signal_type_present_flag)
-						{
-							video_full_range_flag = NextImage.SourceInfo->SPSs[0].video_full_range_flag;
-							video_format = NextImage.SourceInfo->SPSs[0].video_format;
-						}
-					}
-					
 					OutputBufferSampleProperties->Set("pixelfmt", FVariantValue((int64)((num_bits > 8)  ? EPixelFormat::PF_A2B10G10R10 : EPixelFormat::PF_B8G8R8A8)));
 					OutputBufferSampleProperties->Set("bits_per", FVariantValue((int64)num_bits));
-					Colorimetry.Update(colour_primaries, transfer_characteristics, matrix_coeffs, video_full_range_flag, video_format);
 					Colorimetry.UpdateParamDict(*OutputBufferSampleProperties);
-					// Extract HDR metadata from SEI messages
-					if (NextImage.SourceInfo.IsValid())
-					{
-						HDR.Update(num_bits, Colorimetry, NextImage.SourceInfo->CSDPrefixSEIMessages, NextImage.SourceInfo->PrefixSEIMessages, false);
-					}
 					HDR.UpdateParamDict(*OutputBufferSampleProperties);
 
 					TSharedPtrTS<FElectraPlayerVideoDecoderOutputAndroid> DecoderOutput = RenderOutputBuffer->GetBufferProperties().GetValue("texture").GetSharedPointer<FElectraPlayerVideoDecoderOutputAndroid>();
diff --git a/Engine/Plugins/Media/ElectraPlayer/Source/ElectraPlayerRuntime/Private/Runtime/Decoder/DX/VideoDecoderH265_DX.cpp b/Engine/Plugins/Media/ElectraPlayer/Source/ElectraPlayerRuntime/Private/Runtime/Decoder/DX/VideoDecoderH265_DX.cpp
--- a/Engine/Plugins/Media/ElectraPlayer/Source/ElectraPlayerRuntime/Private/Runtime/Decoder/DX/VideoDecoderH265_DX.cpp
+++ b/Engine/Plugins/Media/ElectraPlayer/Source/ElectraPlayerRuntime/Private/Runtime/Decoder/DX/VideoDecoderH265_DX.cpp
@@ -827,6 +827,7 @@ bool FVideoDecoderH265::ConvertDecodedImage(const TRefCountPtr<IMFSample>& Decod
 	UINT32			uiPanScanEnabled = 0;
 	FTimeValue		TimeStamp, Duration;
 	bool			bRender = true;
+	uint8			num_bits = 8;
 
 	VERIFY_HR(DecodedOutputSample->GetSampleTime(&llTimeStamp), "Failed to get video decoder output sample timestamp", ERRCODE_INTERNAL_COULD_NOT_GET_OUTPUT_SAMPLE_TIME);
 	TimeStamp.SetFromHNS((int64) llTimeStamp);
@@ -837,6 +838,29 @@ bool FVideoDecoderH265::ConvertDecodedImage(const TRefCountPtr<IMFSample>& Decod
 	TSharedPtrTS<FDecoderInput> MatchingInput;
 	if (FindAndUpdateDecoderInput(MatchingInput, (int64) llTimeStamp))
 	{
+		// Set the bit depth and the colorimetry.
+		uint8 colour_primaries = 2, transfer_characteristics = 2, matrix_coeffs = 2;
+		uint8 video_full_range_flag = 0, video_format = 5;
+		if (MatchingInput->SPSs.Num())
+		{
+			check(MatchingInput->SPSs[0].bit_depth_luma_minus8 == MatchingInput->SPSs[0].bit_depth_chroma_minus8);
+			num_bits = MatchingInput->SPSs[0].bit_depth_luma_minus8 + 8;
+			if (MatchingInput->SPSs[0].colour_description_present_flag)
+			{
+				colour_primaries = MatchingInput->SPSs[0].colour_primaries;
+				transfer_characteristics = MatchingInput->SPSs[0].transfer_characteristics;
+				matrix_coeffs = MatchingInput->SPSs[0].matrix_coeffs;
+			}
+			if (MatchingInput->SPSs[0].video_signal_type_present_flag)
+			{
+				video_full_range_flag = MatchingInput->SPSs[0].video_full_range_flag;
+				video_format = MatchingInput->SPSs[0].video_format;
+			}
+		}
+		Colorimetry.Update(colour_primaries, transfer_characteristics, matrix_coeffs, video_full_range_flag, video_format);
+		// Extract HDR metadata from SEI messages
+		HDR.Update(num_bits, Colorimetry, MatchingInput->CSDPrefixSEIMessages, MatchingInput->PrefixSEIMessages, false);
+
 		// If this image is not for display, return the buffer immediately and bail.
 		if (!MatchingInput->AdjustedPTS.IsValid())
 		{
@@ -939,35 +963,8 @@ bool FVideoDecoderH265::ConvertDecodedImage(const TRefCountPtr<IMFSample>& Decod
 		CurrentVideoFormatGUID = MFVideoFormat_NV12;
 	}
 	OutputBufferSampleProperties->Set("pixelfmt", CurrentVideoFormatGUID == MFVideoFormat_NV12 ? FVariantValue((int64)EPixelFormat::PF_NV12) : FVariantValue((int64)EPixelFormat::PF_P010));
-
-	// Set the bit depth and the colorimetry.
-	uint8 colour_primaries=2, transfer_characteristics=2, matrix_coeffs=2;
-	uint8 video_full_range_flag=0, video_format=5;
-	uint8 num_bits = 8;
-	if (MatchingInput.IsValid() && MatchingInput->SPSs.Num())
-	{
-		check(MatchingInput->SPSs[0].bit_depth_luma_minus8 == MatchingInput->SPSs[0].bit_depth_chroma_minus8);
-		num_bits = MatchingInput->SPSs[0].bit_depth_luma_minus8 + 8;
-		if (MatchingInput->SPSs[0].colour_description_present_flag)
-		{
-			colour_primaries = MatchingInput->SPSs[0].colour_primaries;
-			transfer_characteristics = MatchingInput->SPSs[0].transfer_characteristics;
-			matrix_coeffs = MatchingInput->SPSs[0].matrix_coeffs;
-		}
-		if (MatchingInput->SPSs[0].video_signal_type_present_flag)
-		{
-			video_full_range_flag = MatchingInput->SPSs[0].video_full_range_flag;
-			video_format = MatchingInput->SPSs[0].video_format;
-		}
-	}
 	OutputBufferSampleProperties->Set("bits_per", FVariantValue((int64)num_bits));
-	Colorimetry.Update(colour_primaries, transfer_characteristics, matrix_coeffs, video_full_range_flag, video_format);
 	Colorimetry.UpdateParamDict(*OutputBufferSampleProperties);
-	// Extract HDR metadata from SEI messages
-	if (MatchingInput.IsValid())
-	{
-		HDR.Update(num_bits, Colorimetry, MatchingInput->CSDPrefixSEIMessages, MatchingInput->PrefixSEIMessages, false);
-	}
 	HDR.UpdateParamDict(*OutputBufferSampleProperties);
 
 	if (CurrentRenderOutputBuffer != nullptr)
