ElectraPlayerPlugin: Added additional pointer safety checks

#rnx

diff --git a/Engine/Plugins/Media/ElectraPlayer/Source/ElectraPlayerPlugin/Private/ElectraPlayerPlugin.cpp b/Engine/Plugins/Media/ElectraPlayer/Source/ElectraPlayerPlugin/Private/ElectraPlayerPlugin.cpp
--- a/Engine/Plugins/Media/ElectraPlayer/Source/ElectraPlayerPlugin/Private/ElectraPlayerPlugin.cpp
+++ b/Engine/Plugins/Media/ElectraPlayer/Source/ElectraPlayerPlugin/Private/ElectraPlayerPlugin.cpp
@@ -416,14 +416,19 @@ void FElectraPlayerPlugin::FPlayerAdapterDelegate::OnSubtitleFlush()
 }
 
 
-void FElectraPlayerPlugin::FPlayerAdapterDelegate::PresentVideoFrame(const FVideoDecoderOutputPtr & InVideoFrame)
+void FElectraPlayerPlugin::FPlayerAdapterDelegate::PresentVideoFrame(const FVideoDecoderOutputPtr& InVideoFrame)
 {
 	TSharedPtr<FElectraPlayerPlugin, ESPMode::ThreadSafe> PinnedHost = Host.Pin();
 	if (PinnedHost.IsValid())
 	{
-		FElectraTextureSampleRef TextureSample = PinnedHost->OutputTexturePool->AcquireShared();
-		TextureSample->Initialize(InVideoFrame.Get());
-		PinnedHost->MediaSamples->AddVideo(TextureSample);
+		FVideoDecoderOutputPtr VideoFrame = InVideoFrame;
+		TSharedPtr<FElectraTextureSamplePool, ESPMode::ThreadSafe> TexturePool = PinnedHost->OutputTexturePool;
+		if (VideoFrame.IsValid() && TexturePool.IsValid())
+		{
+			FElectraTextureSampleRef TextureSample = TexturePool->AcquireShared();
+			TextureSample->Initialize(VideoFrame.Get());
+			PinnedHost->MediaSamples->AddVideo(TextureSample);
+		}
 	}
 }
 
@@ -433,9 +438,13 @@ void FElectraPlayerPlugin::FPlayerAdapterDelegate::PresentAudioFrame(const IAudi
 	TSharedPtr<FElectraPlayerPlugin, ESPMode::ThreadSafe> PinnedHost = Host.Pin();
 	if (PinnedHost.IsValid())
 	{
-		TSharedRef<FElectraPlayerAudioSample, ESPMode::ThreadSafe> AudioSample = PinnedHost->OutputAudioPool.AcquireShared();
-		AudioSample->Initialize(InAudioFrame);
-		PinnedHost->MediaSamples->AddAudio(AudioSample);
+		IAudioDecoderOutputPtr AudioFrame = InAudioFrame;
+		if (AudioFrame.IsValid())
+		{
+			TSharedRef<FElectraPlayerAudioSample, ESPMode::ThreadSafe> AudioSample = PinnedHost->OutputAudioPool.AcquireShared();
+			AudioSample->Initialize(InAudioFrame);
+			PinnedHost->MediaSamples->AddAudio(AudioSample);
+		}
 	}
 }
 
@@ -444,9 +453,13 @@ void FElectraPlayerPlugin::FPlayerAdapterDelegate::PresentSubtitleSample(const I
 	TSharedPtr<FElectraPlayerPlugin, ESPMode::ThreadSafe> PinnedHost = Host.Pin();
 	if (PinnedHost.IsValid())
 	{
-		TSharedRef<FElectraSubtitleSample, ESPMode::ThreadSafe> SubtitleSample = MakeShared<FElectraSubtitleSample, ESPMode::ThreadSafe>();
-		SubtitleSample->Subtitle = InSubtitleSample;
-		PinnedHost->MediaSamples->AddSubtitle(SubtitleSample);
+		ISubtitleDecoderOutputPtr Subtitle = InSubtitleSample;
+		if (Subtitle.IsValid())
+		{
+			TSharedRef<FElectraSubtitleSample, ESPMode::ThreadSafe> SubtitleSample = MakeShared<FElectraSubtitleSample, ESPMode::ThreadSafe>();
+			SubtitleSample->Subtitle = InSubtitleSample;
+			PinnedHost->MediaSamples->AddSubtitle(SubtitleSample);
+		}
 	}
 }
 
@@ -455,10 +468,14 @@ void FElectraPlayerPlugin::FPlayerAdapterDelegate::PresentMetadataSample(const I
 	TSharedPtr<FElectraPlayerPlugin, ESPMode::ThreadSafe> PinnedHost = Host.Pin();
 	if (PinnedHost.IsValid())
 	{
-		// Create a binary media sample of our extended format and pass it up.
-		TSharedRef<FElectraBinarySample, ESPMode::ThreadSafe> MetaDataSample = MakeShared<FElectraBinarySample, ESPMode::ThreadSafe>();
-		MetaDataSample->Metadata = InMetadataFrame;
-		PinnedHost->MediaSamples->AddMetadata(MetaDataSample);
+		IMetaDataDecoderOutputPtr MetadataFrame = InMetadataFrame;
+		if (MetadataFrame.IsValid())
+		{
+			// Create a binary media sample of our extended format and pass it up.
+			TSharedRef<FElectraBinarySample, ESPMode::ThreadSafe> MetaDataSample = MakeShared<FElectraBinarySample, ESPMode::ThreadSafe>();
+			MetaDataSample->Metadata = InMetadataFrame;
+			PinnedHost->MediaSamples->AddMetadata(MetaDataSample);
+		}
 	}
 }
 
