Disablement of Electra's BlockOnFetch feature for 5.2 release to enable continued use of the player for VP / Studios
#rnx

diff --git a/Engine/Plugins/Media/ElectraPlayer/Source/ElectraPlayerPlugin/Private/ElectraPlayerPlugin.cpp b/Engine/Plugins/Media/ElectraPlayer/Source/ElectraPlayerPlugin/Private/ElectraPlayerPlugin.cpp
--- a/Engine/Plugins/Media/ElectraPlayer/Source/ElectraPlayerPlugin/Private/ElectraPlayerPlugin.cpp
+++ b/Engine/Plugins/Media/ElectraPlayer/Source/ElectraPlayerPlugin/Private/ElectraPlayerPlugin.cpp
@@ -16,10 +16,22 @@
 #include "MediaSubtitleDecoderOutput.h"
 #include "MediaMetaDataDecoderOutput.h"
 
+#include "HAL/IConsoleManager.h"
+
 using namespace Electra;
 
 //-----------------------------------------------------------------------------
 
+static int GElectraEnableBlockOnFetch = 0;							// Temporarily disable Electra's support for BlockOnFetch by default
+static FAutoConsoleVariableRef CVarElectraEnableBlockOnFetch(
+	TEXT("ElectraPlayer.EnableBlockOnFetch"),
+	GElectraEnableBlockOnFetch,
+	TEXT("Whether 'BlockOnFetch' support is enabled.\n")
+	TEXT(" 0: disabled (default)\n")
+	TEXT(" 1: enabled"));
+
+//-----------------------------------------------------------------------------
+
 FElectraPlayerPlugin::FElectraPlayerPlugin()
 {
 	// Make sure a few assumptions are correct...
@@ -835,6 +847,10 @@ bool FElectraPlayerPlugin::CanControl(EMediaControl Control) const
 	EMediaState CurrentState = GetState();
 	if (Control == EMediaControl::BlockOnFetch)
 	{
+		if (GElectraEnableBlockOnFetch == 0)
+		{
+			return false;
+		}
 		return CurrentState == EMediaState::Playing;
 	}
 	else if (Control == EMediaControl::Pause)
