CIS fix.

diff --git a/Engine/Plugins/Media/ElectraPlayer/Source/ElectraPlayerRuntime/Private/Runtime/Decoder/Android/VideoDecoderH264_Android.cpp b/Engine/Plugins/Media/ElectraPlayer/Source/ElectraPlayerRuntime/Private/Runtime/Decoder/Android/VideoDecoderH264_Android.cpp
--- a/Engine/Plugins/Media/ElectraPlayer/Source/ElectraPlayerRuntime/Private/Runtime/Decoder/Android/VideoDecoderH264_Android.cpp
+++ b/Engine/Plugins/Media/ElectraPlayer/Source/ElectraPlayerRuntime/Private/Runtime/Decoder/Android/VideoDecoderH264_Android.cpp
@@ -1660,7 +1660,7 @@ FVideoDecoderH264::EOutputResult FVideoDecoderH264::ProcessOutput(const FDecoded
 					// (but: this will sync all up as the render command queues are all in order - and hence the async process will be done before MediaTextureResources sees this)
 					Renderer->ReturnBuffer(RenderOutputBuffer, true, *OutputBufferSampleProperties);
 					// The output buffer properties is now owned by the decoder output.
-					OutputBufferSampleProperties.Release();
+					(void)OutputBufferSampleProperties.Release();
 				}
 				else
 				{
diff --git a/Engine/Plugins/Media/ElectraPlayer/Source/ElectraPlayerRuntime/Private/Runtime/Decoder/Android/VideoDecoderH265_Android.cpp b/Engine/Plugins/Media/ElectraPlayer/Source/ElectraPlayerRuntime/Private/Runtime/Decoder/Android/VideoDecoderH265_Android.cpp
--- a/Engine/Plugins/Media/ElectraPlayer/Source/ElectraPlayerRuntime/Private/Runtime/Decoder/Android/VideoDecoderH265_Android.cpp
+++ b/Engine/Plugins/Media/ElectraPlayer/Source/ElectraPlayerRuntime/Private/Runtime/Decoder/Android/VideoDecoderH265_Android.cpp
@@ -1639,7 +1639,7 @@ FVideoDecoderH265::EOutputResult FVideoDecoderH265::ProcessOutput(const FDecoded
 					// (but: this will sync all up as the render command queues are all in order - and hence the async process will be done before MediaTextureResources sees this)
 					Renderer->ReturnBuffer(RenderOutputBuffer, true, *OutputBufferSampleProperties);
 					// The output buffer properties is now owned by the decoder output.
-					OutputBufferSampleProperties.Release();
+					(void)OutputBufferSampleProperties.Release();
 				}
 				else
 				{
