Fix -Wsingle-bit-bitfield-constant-conversion warnings (clang 16)

diff --git a/Engine/Plugins/Media/ElectraPlayer/Source/ElectraPlayerRuntime/Private/Runtime/Player/DASH/ManifestParserDASH.cpp b/Engine/Plugins/Media/ElectraPlayer/Source/ElectraPlayerRuntime/Private/Runtime/Player/DASH/ManifestParserDASH.cpp
--- a/Engine/Plugins/Media/ElectraPlayer/Source/ElectraPlayerRuntime/Private/Runtime/Player/DASH/ManifestParserDASH.cpp
+++ b/Engine/Plugins/Media/ElectraPlayer/Source/ElectraPlayerRuntime/Private/Runtime/Player/DASH/ManifestParserDASH.cpp
@@ -2894,31 +2894,31 @@ bool FTimeline_S_ElementBuilder::ProcessAttribute(FManifestParserDASH* Builder,
 	if (AttrEquals(t))
 	{
 		ParseUIntAttribute(CurrentElement->T, Builder, AttributeName, AttributeValue, GetXMLLineNumber());
-		CurrentElement->HaveT = 1;
+		CurrentElement->HaveT = -1;
 		return true;
 	}
 	else if (AttrEquals(n))
 	{
 		ParseUIntAttribute(CurrentElement->N, Builder, AttributeName, AttributeValue, GetXMLLineNumber());
-		CurrentElement->HaveN = 1;
+		CurrentElement->HaveN = -1;
 		return true;
 	}
 	else if (AttrEquals(d))
 	{
 		ParseUIntAttribute(CurrentElement->D, Builder, AttributeName, AttributeValue, GetXMLLineNumber());
-		CurrentElement->HaveD = 1;
+		CurrentElement->HaveD = -1;
 		return true;
 	}
 	else if (AttrEquals(r))
 	{
 		ParseIntAttribute(CurrentElement->R, Builder, AttributeName, AttributeValue, GetXMLLineNumber());
-		CurrentElement->HaveR = 1;
+		CurrentElement->HaveR = -1;
 		return true;
 	}
 	else if (AttrEquals(k))
 	{
 		ParseUIntAttribute(CurrentElement->K, Builder, AttributeName, AttributeValue, GetXMLLineNumber());
-		CurrentElement->HaveK = 1;
+		CurrentElement->HaveK = -1;
 		return true;
 	}
 	else
@@ -2977,13 +2977,13 @@ bool FFailOver_FCS_ElementBuilder::ProcessAttribute(FManifestParserDASH* Builder
 	if (AttrEquals(t))
 	{
 		ParseUIntAttribute(CurrentElement->T, Builder, AttributeName, AttributeValue, GetXMLLineNumber());
-		CurrentElement->HaveT = 1;
+		CurrentElement->HaveT = -1;
 		return true;
 	}
 	else if (AttrEquals(d))
 	{
 		ParseUIntAttribute(CurrentElement->D, Builder, AttributeName, AttributeValue, GetXMLLineNumber());
-		CurrentElement->HaveD = 1;
+		CurrentElement->HaveD = -1;
 		return true;
 	}
 	else
diff --git a/Engine/Source/Runtime/Engine/Private/Materials/MaterialRenderProxy.cpp b/Engine/Source/Runtime/Engine/Private/Materials/MaterialRenderProxy.cpp
--- a/Engine/Source/Runtime/Engine/Private/Materials/MaterialRenderProxy.cpp
+++ b/Engine/Source/Runtime/Engine/Private/Materials/MaterialRenderProxy.cpp
@@ -555,7 +555,7 @@ FMaterialRenderProxy::~FMaterialRenderProxy()
 		HasVirtualTextureCallbacks = false;
 	}
 
-	DeletedFlag = 1;
+	DeletedFlag = -1;
 }
 
 void FMaterialRenderProxy::InitDynamicRHI()
diff --git a/Engine/Source/Runtime/Engine/Public/Materials/MaterialRenderProxy.h b/Engine/Source/Runtime/Engine/Public/Materials/MaterialRenderProxy.h
--- a/Engine/Source/Runtime/Engine/Public/Materials/MaterialRenderProxy.h
+++ b/Engine/Source/Runtime/Engine/Public/Materials/MaterialRenderProxy.h
@@ -177,7 +177,7 @@ public:
 
 	void MarkForGarbageCollection()
 	{
-		MarkedForGarbageCollection = 1;
+		MarkedForGarbageCollection = -1;
 	}
 
 	bool IsMarkedForGarbageCollection() const
