ElectraPlayer: Fixed an issue introduced in CL 24686373 that destroyed an object held by TUniquePtr<> instead of releasing it after the ownership had been transferred to another object. Since the object was now destroyed even for the new owner things broke down.

#rnx

diff --git a/Engine/Plugins/Media/ElectraPlayer/Source/ElectraPlayerRuntime/Private/Runtime/Decoder/DX/VideoDecoderH264_DX.cpp b/Engine/Plugins/Media/ElectraPlayer/Source/ElectraPlayerRuntime/Private/Runtime/Decoder/DX/VideoDecoderH264_DX.cpp
--- a/Engine/Plugins/Media/ElectraPlayer/Source/ElectraPlayerRuntime/Private/Runtime/Decoder/DX/VideoDecoderH264_DX.cpp
+++ b/Engine/Plugins/Media/ElectraPlayer/Source/ElectraPlayerRuntime/Private/Runtime/Decoder/DX/VideoDecoderH264_DX.cpp
@@ -1066,7 +1066,7 @@ bool FVideoDecoderH264::ConvertDecodedImage(const TRefCountPtr<IMFSample>& Decod
 
 		Renderer->ReturnBuffer(CurrentRenderOutputBuffer, true, *OutputBufferSampleProperties);
 		CurrentRenderOutputBuffer = nullptr;
-		OutputBufferSampleProperties.Reset();
+		(void)OutputBufferSampleProperties.Release();
 	}
 	return true;
 }
diff --git a/Engine/Plugins/Media/ElectraPlayer/Source/ElectraPlayerRuntime/Private/Runtime/Decoder/DX/VideoDecoderH265_DX.cpp b/Engine/Plugins/Media/ElectraPlayer/Source/ElectraPlayerRuntime/Private/Runtime/Decoder/DX/VideoDecoderH265_DX.cpp
--- a/Engine/Plugins/Media/ElectraPlayer/Source/ElectraPlayerRuntime/Private/Runtime/Decoder/DX/VideoDecoderH265_DX.cpp
+++ b/Engine/Plugins/Media/ElectraPlayer/Source/ElectraPlayerRuntime/Private/Runtime/Decoder/DX/VideoDecoderH265_DX.cpp
@@ -976,7 +976,7 @@ bool FVideoDecoderH265::ConvertDecodedImage(const TRefCountPtr<IMFSample>& Decod
 
 		Renderer->ReturnBuffer(CurrentRenderOutputBuffer, bRender, *OutputBufferSampleProperties);
 		CurrentRenderOutputBuffer = nullptr;
-		OutputBufferSampleProperties.Reset();
+		(void)OutputBufferSampleProperties.Release();
 	}
 	return true;
 }
