ElectraPlayer: Fixing incorrect retries of a Live stream

#rnx

diff --git a/Engine/Plugins/Media/ElectraPlayer/Source/ElectraPlayerRuntime/Private/Runtime/Player/ABRRules/ABRLiveStream.cpp b/Engine/Plugins/Media/ElectraPlayer/Source/ElectraPlayerRuntime/Private/Runtime/Player/ABRRules/ABRLiveStream.cpp
--- a/Engine/Plugins/Media/ElectraPlayer/Source/ElectraPlayerRuntime/Private/Runtime/Player/ABRRules/ABRLiveStream.cpp
+++ b/Engine/Plugins/Media/ElectraPlayer/Source/ElectraPlayerRuntime/Private/Runtime/Player/ABRRules/ABRLiveStream.cpp
@@ -906,6 +906,11 @@ IAdaptiveStreamSelector::ESegmentAction FABRLiveStream::EvaluateForError(TArray<
 
 			if (!Stats.bWasSuccessful)
 			{
+				if (Stats.bWaitingForRemoteRetryElement)
+				{
+					return IAdaptiveStreamSelector::ESegmentAction::Retry;
+				}
+
 				// If this is the first failure for this segment (not being retried yet) increase
 				// the number of consecutively failed segments.
 				if (Stats.RetryNumber == 0)
