ElectraPlayer: removed byte size limit from decoder buffers in favor of duration based buffer size only

#rnx

diff --git a/Engine/Plugins/Media/ElectraPlayer/Source/ElectraPlayerRuntime/Private/Runtime/Player/AdaptiveStreamingPlayer.cpp b/Engine/Plugins/Media/ElectraPlayer/Source/ElectraPlayerRuntime/Private/Runtime/Player/AdaptiveStreamingPlayer.cpp
--- a/Engine/Plugins/Media/ElectraPlayer/Source/ElectraPlayerRuntime/Private/Runtime/Player/AdaptiveStreamingPlayer.cpp
+++ b/Engine/Plugins/Media/ElectraPlayer/Source/ElectraPlayerRuntime/Private/Runtime/Player/AdaptiveStreamingPlayer.cpp
@@ -1194,7 +1194,6 @@ bool FAdaptiveStreamingPlayer::OnFragmentAccessUnitReceived(FAccessUnit* InAcces
 		BufferConfig = &PlayerConfig.StreamBufferConfigText;
 	}
 	FAccessUnitBuffer::FExternalBufferInfo TotalUtilization;
-	TotalUtilization.DataSize = bufferInfo.CurrentMemInUse;
 	TotalUtilization.Duration = bufferInfo.PlayableDuration;
 
 	// Try to push the data into the receiving buffer
@@ -4801,21 +4800,21 @@ void FAdaptiveStreamingPlayer::DebugPrint(void* pPlayer, void (*pPrintFN)(void*
 	if(1)
 	{
 		pD = &VideoBufferStats.StreamBuffer;
-		pPrintFN(pPlayer, "Video buffer  : EOT %d; EOS %d; %3u AUs; %8u/%8u bytes in; %#7.4fs", pD->bEndOfTrack, pD->bEndOfData, (uint32)pD->NumCurrentAccessUnits, (uint32)pD->CurrentMemInUse, (uint32)PlayerConfig.StreamBufferConfigVideo.MaxDataSize, pD->PlayableDuration.GetAsSeconds());
+		pPrintFN(pPlayer, "Video buffer  : EOT %d; EOS %d; %3u AUs; %8lld bytes in; %#7.4fs", pD->bEndOfTrack, pD->bEndOfData, (uint32)pD->NumCurrentAccessUnits, (long long int)pD->CurrentMemInUse, pD->PlayableDuration.GetAsSeconds());
 		pPrintFN(pPlayer, "Video decoder : %2u in decoder, %zu total; EOD in %d; EOD out %d; %s", (uint32)VideoBufferStats.DecoderOutputBuffer.NumElementsInDecoder, (uint32)VideoBufferStats.DecoderOutputBuffer.MaxDecodedElementsReady, VideoBufferStats.DecoderInputBuffer.bEODReached, VideoBufferStats.DecoderOutputBuffer.bEODreached, VideoBufferStats.DecoderOutputBuffer.bOutputStalled?"stalled":"not stalled");
 		pPrintFN(pPlayer, "Video renderer: %#7.4fs enqueued", VideoRender.Renderer.IsValid() ? VideoRender.Renderer->GetEnqueuedSampleDuration().GetAsSeconds() : 0.0);
 	}
 	if(1)
 	{
 		pD = &AudioBufferStats.StreamBuffer;
-		pPrintFN(pPlayer, "Audio buffer  : EOT %d; EOS %d; %3u AUs; %8u/%8u bytes in; %#7.4fs", pD->bEndOfTrack, pD->bEndOfData, (uint32)pD->NumCurrentAccessUnits, (uint32)pD->CurrentMemInUse, (uint32)PlayerConfig.StreamBufferConfigAudio.MaxDataSize, pD->PlayableDuration.GetAsSeconds());
+		pPrintFN(pPlayer, "Audio buffer  : EOT %d; EOS %d; %3u AUs; %8lld bytes in; %#7.4fs", pD->bEndOfTrack, pD->bEndOfData, (uint32)pD->NumCurrentAccessUnits, (long long int)pD->CurrentMemInUse, pD->PlayableDuration.GetAsSeconds());
 		pPrintFN(pPlayer, "Audio decoder : %2u in decoder, %2u total; EOD in %d; EOD out %d; %s", (uint32)AudioBufferStats.DecoderOutputBuffer.NumElementsInDecoder, (uint32)AudioBufferStats.DecoderOutputBuffer.MaxDecodedElementsReady, AudioBufferStats.DecoderInputBuffer.bEODReached, AudioBufferStats.DecoderOutputBuffer.bEODreached, AudioBufferStats.DecoderOutputBuffer.bOutputStalled?"stalled":"not stalled");
 		pPrintFN(pPlayer, "Audio renderer: %#7.4fs enqueued", AudioRender.Renderer.IsValid() ? AudioRender.Renderer->GetEnqueuedSampleDuration().GetAsSeconds() : 0.0);
 	}
 	if(1)
 	{
 		pD = &TextBufferStats.StreamBuffer;
-		pPrintFN(pPlayer, "Text buffer   : EOT %d; EOS %d; %3u AUs; %8u/%8u bytes in; %#7.4fs", pD->bEndOfTrack, pD->bEndOfData, (uint32)pD->NumCurrentAccessUnits, (uint32)pD->CurrentMemInUse, (uint32)PlayerConfig.StreamBufferConfigText.MaxDataSize, pD->PlayableDuration.GetAsSeconds());
+		pPrintFN(pPlayer, "Text buffer   : EOT %d; EOS %d; %3u AUs; %8lld bytes in; %#7.4fs", pD->bEndOfTrack, pD->bEndOfData, (uint32)pD->NumCurrentAccessUnits, (long long int)pD->CurrentMemInUse, pD->PlayableDuration.GetAsSeconds());
 	}
 	DiagnosticsCriticalSection.Unlock();
 
diff --git a/Engine/Plugins/Media/ElectraPlayer/Source/ElectraPlayerRuntime/Private/Runtime/Player/AdaptiveStreamingPlayerEvents.cpp b/Engine/Plugins/Media/ElectraPlayer/Source/ElectraPlayerRuntime/Private/Runtime/Player/AdaptiveStreamingPlayerEvents.cpp
--- a/Engine/Plugins/Media/ElectraPlayer/Source/ElectraPlayerRuntime/Private/Runtime/Player/AdaptiveStreamingPlayerEvents.cpp
+++ b/Engine/Plugins/Media/ElectraPlayer/Source/ElectraPlayerRuntime/Private/Runtime/Player/AdaptiveStreamingPlayerEvents.cpp
@@ -72,7 +72,6 @@ void FAdaptiveStreamingPlayer::DispatchBufferUtilizationEvent(EStreamType Buffer
 		stats.BufferType = BufferType;
 		stats.MaxDurationInSeconds = BufferConfig ? BufferConfig->MaxDuration.GetAsSeconds() : 0.0;
 		stats.DurationInUse 	   = BufferStats->PushedDuration.GetAsSeconds();
-		stats.MaxByteCapacity      = BufferConfig ? BufferConfig->MaxDataSize : 0;
 		stats.BytesInUse		   = BufferStats->CurrentMemInUse;
 		DispatchEvent(FMetricEvent::ReportBufferUtilization(stats));
 	}
diff --git a/Engine/Plugins/Media/ElectraPlayer/Source/ElectraPlayerRuntime/Private/Runtime/Player/AdaptiveStreamingPlayerInternalConfig.h b/Engine/Plugins/Media/ElectraPlayer/Source/ElectraPlayerRuntime/Private/Runtime/Player/AdaptiveStreamingPlayerInternalConfig.h
--- a/Engine/Plugins/Media/ElectraPlayer/Source/ElectraPlayerRuntime/Private/Runtime/Player/AdaptiveStreamingPlayerInternalConfig.h
+++ b/Engine/Plugins/Media/ElectraPlayer/Source/ElectraPlayerRuntime/Private/Runtime/Player/AdaptiveStreamingPlayerInternalConfig.h
@@ -13,15 +13,12 @@ namespace Electra
 			FConfiguration()
 			{
 				// Set default values to maximum permitted values.
-				StreamBufferConfigVideo.MaxDataSize = 128 << 20;
 				StreamBufferConfigVideo.MaxDuration.SetFromSeconds(6.0);
 
-				StreamBufferConfigAudio.MaxDataSize = 8 << 20;
 				StreamBufferConfigAudio.MaxDuration.SetFromSeconds(6.0);
 
 				// Subtitle streams tend to be sparse and each AU could potentially have
 				// a huge duration, so we need to ignore durations altogether.
-				StreamBufferConfigText.MaxDataSize = 8 << 20;
 				StreamBufferConfigText.MaxDuration.SetToPositiveInfinity();
 
 				bHoldLastFrameDuringSeek = true;
diff --git a/Engine/Plugins/Media/ElectraPlayer/Source/ElectraPlayerRuntime/Private/Runtime/Player/AdaptiveStreamingPlayerMetrics.h b/Engine/Plugins/Media/ElectraPlayer/Source/ElectraPlayerRuntime/Private/Runtime/Player/AdaptiveStreamingPlayerMetrics.h
--- a/Engine/Plugins/Media/ElectraPlayer/Source/ElectraPlayerRuntime/Private/Runtime/Player/AdaptiveStreamingPlayerMetrics.h
+++ b/Engine/Plugins/Media/ElectraPlayer/Source/ElectraPlayerRuntime/Private/Runtime/Player/AdaptiveStreamingPlayerMetrics.h
@@ -82,13 +82,11 @@ namespace Metrics
 			BufferType  		 = EStreamType::Video;
 			MaxDurationInSeconds = 0.0;
 			DurationInUse   	 = 0.0;
-			MaxByteCapacity 	 = 0;
 			BytesInUse  		 = 0;
 		}
 		EStreamType		BufferType;
 		double			MaxDurationInSeconds;
 		double			DurationInUse;
-		int64			MaxByteCapacity;
 		int64			BytesInUse;
 	};
 
diff --git a/Engine/Plugins/Media/ElectraPlayer/Source/ElectraPlayerRuntime/Private/Runtime/Player/mp4/ManifestMP4.cpp b/Engine/Plugins/Media/ElectraPlayer/Source/ElectraPlayerRuntime/Private/Runtime/Player/mp4/ManifestMP4.cpp
--- a/Engine/Plugins/Media/ElectraPlayer/Source/ElectraPlayerRuntime/Private/Runtime/Player/mp4/ManifestMP4.cpp
+++ b/Engine/Plugins/Media/ElectraPlayer/Source/ElectraPlayerRuntime/Private/Runtime/Player/mp4/ManifestMP4.cpp
@@ -827,10 +827,9 @@ void FManifestMP4Internal::FTimelineAssetMP4::LogMessage(IInfoLog::ELevel Level,
 
 void FManifestMP4Internal::FTimelineAssetMP4::LimitSegmentDownloadSize(TSharedPtrTS<IStreamSegment>& InOutSegment)
 {
-	// Limit the segment download size.
+	// Limit the segment size by content duration.
 	// This helps with downloads that might otherwise take too long or keep the connection open for too long (when downloading a large mp4 from start to finish).
-	const int64 MaxSegmentSize = 4 * 1024 * 1024;
-	const int64 MaxSegmentDurationMSec = 2000;
+	const int64 MaxSegmentDurationMSec = 3000;
 	if (InOutSegment.IsValid())
 	{
 		FStreamSegmentRequestMP4* Request = static_cast<FStreamSegmentRequestMP4*>(InOutSegment.Get());
@@ -870,9 +869,7 @@ void FManifestMP4Internal::FTimelineAssetMP4::LimitSegmentDownloadSize(TSharedPt
 				}
 				int64 CurrentTrackOffset = LastTrackOffset;
 				bool bDurationLimitReached = TrackDurationLimit > 0 && TrackDur > TrackDurationLimit;
-				if (CurrentTrackOffset >= EndOffset ||
-					CurrentTrackOffset - StartOffset >= MaxSegmentSize ||
-					bDurationLimitReached)
+				if (CurrentTrackOffset >= EndOffset || bDurationLimitReached)
 				{
 					// Limit reached.
 					Request->FileEndOffset = CurrentTrackOffset - 1;
diff --git a/Engine/Plugins/Media/ElectraPlayer/Source/ElectraPlayerRuntime/Private/Runtime/StreamAccessUnitBuffer.h b/Engine/Plugins/Media/ElectraPlayer/Source/ElectraPlayerRuntime/Private/Runtime/StreamAccessUnitBuffer.h
--- a/Engine/Plugins/Media/ElectraPlayer/Source/ElectraPlayerRuntime/Private/Runtime/StreamAccessUnitBuffer.h
+++ b/Engine/Plugins/Media/ElectraPlayer/Source/ElectraPlayerRuntime/Private/Runtime/StreamAccessUnitBuffer.h
@@ -231,19 +231,16 @@ namespace Electra
 	public:
 		struct FConfiguration
 		{
-			FConfiguration(int64 InMaxByteSize = 0, double InMaxSeconds = 0.0)
-				: MaxDataSize(InMaxByteSize)
+			FConfiguration(double InMaxSeconds = 0.0)
 			{
 				MaxDuration.SetFromSeconds(InMaxSeconds);
 			}
 			FTimeValue	MaxDuration;
-			int64		MaxDataSize;
 		};
 
 		struct FExternalBufferInfo
 		{
 			FTimeValue	Duration = FTimeValue::GetZero();
-			int64		DataSize = 0;
 		};
 
 		FAccessUnitBuffer()
@@ -624,21 +621,13 @@ namespace Electra
 			}
 			if (ExternalInfo == nullptr)
 			{
-				// Memory ok?
-				if (AU->TotalMemSize() + CurrentMemInUse <= Limit->MaxDataSize)
-				{
-					// Max allowed duration ok?
-					return PlayableDuration.IsValid() && PlayableDuration + AU->Duration > Limit->MaxDuration ? false : true;
-				}
+				// Max allowed duration ok?
+				return PlayableDuration.IsValid() && PlayableDuration + AU->Duration > Limit->MaxDuration ? false : true;
 			}
 			else
 			{
-				// Memory ok?
-				if (AU->TotalMemSize() + ExternalInfo->DataSize <= Limit->MaxDataSize)
-				{
-					// Max allowed duration ok?
-					return AU->Duration + ExternalInfo->Duration > Limit->MaxDuration ? false : true;
-				}
+				// Max allowed duration ok?
+				return AU->Duration + ExternalInfo->Duration > Limit->MaxDuration ? false : true;
 			}
 			return false;
 		}
