Moving Electra decoder plugin output data setup more into the plugins & removing explicit knowledge about the used codec from that code
#rnx

diff --git a/Engine/Plugins/Media/AppleProResDecoderElectra/Source/AppleProResDecoderElectra/Private/ProResDecoder/Mac/ElectraMediaProResDecoder.cpp b/Engine/Plugins/Media/AppleProResDecoderElectra/Source/AppleProResDecoderElectra/Private/ProResDecoder/Mac/ElectraMediaProResDecoder.cpp
--- a/Engine/Plugins/Media/AppleProResDecoderElectra/Source/AppleProResDecoderElectra/Private/ProResDecoder/Mac/ElectraMediaProResDecoder.cpp
+++ b/Engine/Plugins/Media/AppleProResDecoderElectra/Source/AppleProResDecoderElectra/Private/ProResDecoder/Mac/ElectraMediaProResDecoder.cpp
@@ -49,7 +49,7 @@ public:
 };
 
 
-class FVideoDecoderOutputProResElectra : public IElectraDecoderVideoOutput
+class FVideoDecoderOutputProResElectra : public IElectraDecoderVideoOutput, public IElectraDecoderVideoOutputImageBuffers
 {
 public:
 	virtual ~FVideoDecoderOutputProResElectra()
@@ -72,10 +72,6 @@ public:
 	{ return Height; }
 	FElectraVideoDecoderOutputCropValues GetCropValues() const override
 	{ return Crop; }
-	int32 GetFrameWidth() const override
-	{ return Pitch; }
-	int32 GetFrameHeight() const override
-	{ return Height; }
 	int32 GetAspectRatioW() const override
 	{ return AspectW; }
 	int32 GetAspectRatioH() const override
@@ -88,14 +84,11 @@ public:
 	{ return NumBits; }
 	void GetExtraValues(TMap<FString, FVariant>& OutExtraValues) const override
 	{ OutExtraValues = ExtraValues; }
-	int32 GetPixelFormat() const override
-	{ return PixelFormat; }
 	void* GetPlatformOutputHandle(EElectraDecoderPlatformOutputHandleType InTypeOfHandle) const override
 	{
-		check(InTypeOfHandle == EElectraDecoderPlatformOutputHandleType::ImageBufferRef);
-		if (InTypeOfHandle == EElectraDecoderPlatformOutputHandleType::ImageBufferRef)
+		if (InTypeOfHandle == EElectraDecoderPlatformOutputHandleType::ImageBuffers)
 		{
-			return ImageBuffer;
+			return static_cast<IElectraDecoderVideoOutputImageBuffers*>(const_cast<FVideoDecoderOutputProResElectra*>(this));
 		}
 		return nullptr;
 	}
@@ -104,6 +97,36 @@ public:
 	IElectraDecoderVideoOutput::EImageCopyResult CopyPlatformImage(IElectraDecoderVideoOutputCopyResources* InCopyResources) const override
 	{ return IElectraDecoderVideoOutput::EImageCopyResult::NotSupported; }
 
+	// Methods from IElectraDecoderVideoOutputImageBuffers
+	uint32 GetCodec4CC() const override
+	{
+		return Codec4CC;
+	}
+	int32 GetNumberOfBuffers() const override
+	{
+		return 1;
+	}
+	TSharedPtr<TArray<uint8>, ESPMode::ThreadSafe> GetBufferDataByIndex(int32 InBufferIndex) const override
+	{
+		return nullptr;
+	}
+	void* GetBufferTextureByIndex(int32 InBufferIndex) const override
+	{
+		return InBufferIndex == 0 ? ImageBuffer : nullptr;
+	}
+	EElectraDecoderPlatformPixelFormat GetBufferFormatByIndex(int32 InBufferIndex) const override
+	{
+		return PixelFormat;
+	}
+	EElectraDecoderPlatformPixelEncoding GetBufferEncodingByIndex(int32 InBufferIndex) const override
+	{
+		return PixelEncoding;
+	}
+	int32 GetBufferPitchByIndex(int32 InBufferIndex) const override
+	{
+		return Pitch;
+	}
+
 public:
 	void ReleaseOutputBuffer()
 	{
@@ -126,9 +149,11 @@ public:
 	int32 AspectH = 1;
 	int32 FrameRateN = 0;
 	int32 FrameRateD = 0;
-	int32 PixelFormat = 0;
+	EElectraDecoderPlatformPixelFormat PixelFormat = EElectraDecoderPlatformPixelFormat::INVALID;
+	EElectraDecoderPlatformPixelEncoding PixelEncoding = EElectraDecoderPlatformPixelEncoding::Native;
 	TMap<FString, FVariant> ExtraValues;
 
+	uint32 Codec4CC = 0;
 	CVImageBufferRef ImageBuffer = nullptr;
 };
 
@@ -665,44 +690,14 @@ bool FVideoDecoderProResElectra::ParseHeader(const void* InData, int32 InNumData
 
 bool FVideoDecoderProResElectra::ConfigureOutputForInputFormat()
 {
-#if 0
-	OutputPixelFormat = kPRFormat_y416;
-	switch(CurrentHeader.GetSourcePixFormat())
-	{
-		case FProResHeader::ESourcePixFormat::Twovuy:
-			OutputPixelFormat = kPRFormat_2vuy;
-			break;
-		case FProResHeader::ESourcePixFormat::v210:
-			OutputPixelFormat = kPRFormat_v210;
-			break;
-		case FProResHeader::ESourcePixFormat::v216:
-			OutputPixelFormat = kPRFormat_v216;
-			break;
-		case FProResHeader::ESourcePixFormat::r4fl:
-			OutputPixelFormat = kPRFormat_r4fl;
-			break;
-		case FProResHeader::ESourcePixFormat::b64a:
-			OutputPixelFormat = kPRFormat_b64a;
-			break;
-		case FProResHeader::ESourcePixFormat::R10k:
-			OutputPixelFormat = kPRFormat_R10k;
-			break;
-		case FProResHeader::ESourcePixFormat::r408:
-		case FProResHeader::ESourcePixFormat::v408:
-		case FProResHeader::ESourcePixFormat::RGB:
-		case FProResHeader::ESourcePixFormat::BGRA:
-		case FProResHeader::ESourcePixFormat::n302:
-		case FProResHeader::ESourcePixFormat::l302:
-		default:
-			OutputPixelFormat = kPRFormat_y416;
-			break;
-	}
-
-	//FProResHeader::EChromaType = CurrentHeader.GetChromaType();
-	//FProResHeader::ESourceAlpha = CurrentHeader.GetSourceAlphaDepth();
-#endif
+	// Let Apple's decoder choose one of these formats depending on the input...
+	OutputPixelFormats = {
+		kCVPixelFormatType_4444AYpCbCr16,
+		kCVPixelFormatType_420YpCbCr8BiPlanarVideoRange, kCVPixelFormatType_420YpCbCr8BiPlanarFullRange,
+		kCVPixelFormatType_420YpCbCr10BiPlanarVideoRange, kCVPixelFormatType_420YpCbCr10BiPlanarFullRange
+	};
 
-//	OutputPixelFormats = { kCVPixelFormatType_4444AYpCbCr16, kCVPixelFormatType_420YpCbCr10BiPlanarFullRange, kCVPixelFormatType_420YpCbCr8BiPlanarFullRange };
+	// Select the input format from the 4CC code
 	switch(Codec4CC)
 	{
 		default:
@@ -747,11 +742,6 @@ bool FVideoDecoderProResElectra::ConfigureOutputForInputFormat()
 			break;
 		}
 	}
-	// If there is no desired output format yet, use a default one.
-	if (OutputPixelFormats.IsEmpty())
-	{
-		OutputPixelFormats = { kCVPixelFormatType_420YpCbCr8BiPlanarFullRange };
-	}
 
 	return true;
 }
@@ -1113,7 +1103,38 @@ FVideoDecoderProResElectra::EConvertResult FVideoDecoderProResElectra::ConvertDe
 
 	NewOutput->ImageBuffer = NextImage.ReleaseImageBufferRef();
 	int pixelFormat = CVPixelBufferGetPixelFormatType(NewOutput->ImageBuffer);
-	NewOutput->PixelFormat = pixelFormat;
+
+	NewOutput->Codec4CC = Codec4CC;
+
+	switch (pixelFormat)
+	{
+		case	kCVPixelFormatType_4444AYpCbCr16:
+		{
+			NewOutput->PixelFormat = EElectraDecoderPlatformPixelFormat::R16G16B16A16;
+			NewOutput->PixelEncoding = EElectraDecoderPlatformPixelEncoding::YCbCr_Alpha;
+			break;
+		}
+		case	kCVPixelFormatType_420YpCbCr8BiPlanarVideoRange:
+		case	kCVPixelFormatType_420YpCbCr8BiPlanarFullRange:
+		{
+			NewOutput->PixelFormat = EElectraDecoderPlatformPixelFormat::NV12;
+			NewOutput->PixelEncoding = EElectraDecoderPlatformPixelEncoding::Native;
+			break;
+		}
+		case	kCVPixelFormatType_420YpCbCr10BiPlanarVideoRange:
+		case	kCVPixelFormatType_420YpCbCr10BiPlanarFullRange:
+		{
+			NewOutput->PixelFormat = EElectraDecoderPlatformPixelFormat::P010;
+			NewOutput->PixelEncoding = EElectraDecoderPlatformPixelEncoding::Native;
+			break;
+		}
+		default:
+		{
+			check(!"Unexpected output format!");
+			NewOutput->PixelFormat = EElectraDecoderPlatformPixelFormat::INVALID;
+			NewOutput->PixelEncoding = EElectraDecoderPlatformPixelEncoding::Native;
+		}
+	}
 
 	NewOutput->NumBits = ElectraVideoDecoderFormatTypesApple::GetNumComponentBitsForPixelFormat(pixelFormat);
 	if (NewOutput->NumBits < 0)
diff --git a/Engine/Plugins/Media/AppleProResDecoderElectra/Source/AppleProResDecoderElectra/Private/ProResDecoder/Windows/ElectraMediaProResDecoder.cpp b/Engine/Plugins/Media/AppleProResDecoderElectra/Source/AppleProResDecoderElectra/Private/ProResDecoder/Windows/ElectraMediaProResDecoder.cpp
--- a/Engine/Plugins/Media/AppleProResDecoderElectra/Source/AppleProResDecoderElectra/Private/ProResDecoder/Windows/ElectraMediaProResDecoder.cpp
+++ b/Engine/Plugins/Media/AppleProResDecoderElectra/Source/AppleProResDecoderElectra/Private/ProResDecoder/Windows/ElectraMediaProResDecoder.cpp
@@ -56,11 +56,11 @@ public:
 
 	int32 GetWidth() const override
 	{
-		return Width - Crop.Left - Crop.Right;
+		return ImageWidth;
 	}
 	int32 GetHeight() const override
 	{
-		return Height - Crop.Top - Crop.Bottom;
+		return ImageHeight;
 	}
 	int32 GetDecodedWidth() const override
 	{
@@ -74,14 +74,6 @@ public:
 	{
 		return Crop;
 	}
-	int32 GetFrameWidth() const override
-	{
-		return Pitch;
-	}
-	int32 GetFrameHeight() const override
-	{
-		return Height;
-	}
 	int32 GetAspectRatioW() const override
 	{
 		return AspectW;
@@ -106,10 +98,6 @@ public:
 	{
 		OutExtraValues = ExtraValues;
 	}
-	int32 GetPixelFormat() const override
-	{
-		return PixelFormat;
-	}
 	void* GetPlatformOutputHandle(EElectraDecoderPlatformOutputHandleType InTypeOfHandle) const override
 	{
 		if (InTypeOfHandle == EElectraDecoderPlatformOutputHandleType::ImageBuffers)
@@ -136,13 +124,25 @@ public:
 	{
 		return 1;
 	}
-	TSharedPtr<TArray<uint8>, ESPMode::ThreadSafe> GetBufferByIndex(int32 InBufferIndex) const override
+	TSharedPtr<TArray<uint8>, ESPMode::ThreadSafe> GetBufferDataByIndex(int32 InBufferIndex) const override
 	{
 		return InBufferIndex == 0 ? Buffer : nullptr;
 	}
-	uint64 GetBufferFormatByIndex(int32 InBufferIndex) const override
+	void* GetBufferTextureByIndex(int32 InBufferIndex) const override
 	{
-		return static_cast<uint64>(BufferFormat);
+		return nullptr;
+	}
+	EElectraDecoderPlatformPixelFormat GetBufferFormatByIndex(int32 InBufferIndex) const override
+	{
+		return BufferFormat;
+	}
+	EElectraDecoderPlatformPixelEncoding GetBufferEncodingByIndex(int32 InBufferIndex) const override
+	{
+		return BufferEncoding;
+	}
+	int32 GetBufferPitchByIndex(int32 InBufferIndex) const override
+	{
+		return Pitch;
 	}
 
 public:
@@ -150,6 +150,8 @@ public:
 	uint64 UserValue = 0;
 
 	FElectraVideoDecoderOutputCropValues Crop;
+	int32 ImageWidth = 0;
+	int32 ImageHeight = 0;
 	int32 Width = 0;
 	int32 Height = 0;
 	int32 Pitch = 0;
@@ -163,11 +165,11 @@ public:
 
 	uint32 Codec4CC = 0;
 	TSharedPtr<TArray<uint8>, ESPMode::ThreadSafe> Buffer;
-	PRPixelFormat BufferFormat = kPRFormat_2vuy;
+	EElectraDecoderPlatformPixelFormat BufferFormat = EElectraDecoderPlatformPixelFormat::INVALID;
+	EElectraDecoderPlatformPixelEncoding BufferEncoding = EElectraDecoderPlatformPixelEncoding::Native;
 };
 
 
-
 class FVideoDecoderProResElectra : public IElectraDecoder
 {
 public:
@@ -623,7 +625,7 @@ IElectraDecoder::EDecoderError FVideoDecoderProResElectra::DecodeAccessUnit(cons
 
 		NewOutput->Width = DecodedWidth;
 		NewOutput->Height = DecodedHeight;
-		NewOutput->Pitch = NewOutput->Width;
+		NewOutput->Pitch = DecodedWidth;
 		NewOutput->Crop.Right = DecodedWidth - DisplayWidth;
 		NewOutput->Crop.Bottom = DecodedHeight - DisplayHeight;
 		if (AspectW && AspectH)
@@ -632,28 +634,57 @@ IElectraDecoder::EDecoderError FVideoDecoderProResElectra::DecodeAccessUnit(cons
 			NewOutput->AspectH = AspectH;
 		}
 
+		NewOutput->ImageWidth = DecodedWidth - NewOutput->Crop.Left - NewOutput->Crop.Right;
+		NewOutput->ImageHeight = DecodedHeight - NewOutput->Crop.Top - NewOutput->Crop.Bottom;
+
 		NewOutput->Codec4CC = Codec4CC;
 		NewOutput->PixelFormat = static_cast<int32>(OutputPixelFormat);
 		switch(OutputPixelFormat)
 		{
 			case kPRFormat_2vuy:
 				NewOutput->NumBits = 8;
+				NewOutput->BufferFormat = EElectraDecoderPlatformPixelFormat::B8G8R8A8;
+				NewOutput->BufferEncoding = EElectraDecoderPlatformPixelEncoding::CbY0CrY1;
+				NewOutput->Width /= 2;
+				NewOutput->Pitch *= 2;
+				break;
+			case kPRFormat_y416:
+				NewOutput->NumBits = 16;
+				NewOutput->BufferFormat = EElectraDecoderPlatformPixelFormat::A16B16G16R16;
+				NewOutput->BufferEncoding = EElectraDecoderPlatformPixelEncoding::YCbCr_Alpha;
+				NewOutput->Pitch *= 8;
+				break;
+			case kPRFormat_r4fl:
+				NewOutput->NumBits = 32;
+				NewOutput->BufferFormat = EElectraDecoderPlatformPixelFormat::A32B32G32R32F;
+				NewOutput->BufferEncoding = EElectraDecoderPlatformPixelEncoding::YCbCr_Alpha;
+				NewOutput->Pitch *= 16;
 				break;
 			case kPRFormat_v210:
-			case kPRFormat_R10k:
-			case kPRFormat_r210:
 				NewOutput->NumBits = 10;
+				NewOutput->BufferFormat = EElectraDecoderPlatformPixelFormat::A2B10G10R10;
+				NewOutput->BufferEncoding = EElectraDecoderPlatformPixelEncoding::CbY0CrY1;
+				NewOutput->Width = 4 * ((NewOutput->Width + 5) / 6); // each 4 pixel contain 6 horizontally adjacent YCbCr pixels (incl. 4x 2-bit padding)
+				NewOutput->Pitch = NewOutput->Width * 4;
 				break;
-			case kPRFormat_b64a:
 			case kPRFormat_v216:
-			case kPRFormat_y416:
 				NewOutput->NumBits = 16;
+				NewOutput->BufferFormat = EElectraDecoderPlatformPixelFormat::A16B16G16R16;
+				NewOutput->BufferEncoding = EElectraDecoderPlatformPixelEncoding::CbY0CrY1;
+				NewOutput->Width /= 2;
+				NewOutput->Pitch *= 4;
 				break;
-			case kPRFormat_r4fl:
-				NewOutput->NumBits = 32;
+			case kPRFormat_b64a:
+				NewOutput->NumBits = 16;
+				NewOutput->BufferFormat = EElectraDecoderPlatformPixelFormat::A16B16G16R16;
+				NewOutput->BufferEncoding = EElectraDecoderPlatformPixelEncoding::ARGB_BigEndian;
+				NewOutput->Pitch *= 8;
 				break;
+			case kPRFormat_R10k:
+			case kPRFormat_r210:
 			default:
 				NewOutput->NumBits = 0;
+				NewOutput->BufferFormat = EElectraDecoderPlatformPixelFormat::INVALID;
 				break;
 		}
 
@@ -662,7 +693,6 @@ IElectraDecoder::EDecoderError FVideoDecoderProResElectra::DecodeAccessUnit(cons
 
 		NewOutput->Buffer = MakeShared<TArray<uint8>, ESPMode::ThreadSafe>();
 		NewOutput->Buffer->AddUninitialized(BufferAllocationSize);
-		NewOutput->BufferFormat = OutputPixelFormat;
 
 		PRPixelBuffer pb;
 		FMemory::Memzero(pb);
diff --git a/Engine/Plugins/Media/AvidDNxHDDecoderElectra/Source/Private/Windows/ElectraMediaAvidDNxHDDecoder.cpp b/Engine/Plugins/Media/AvidDNxHDDecoderElectra/Source/Private/Windows/ElectraMediaAvidDNxHDDecoder.cpp
--- a/Engine/Plugins/Media/AvidDNxHDDecoderElectra/Source/Private/Windows/ElectraMediaAvidDNxHDDecoder.cpp
+++ b/Engine/Plugins/Media/AvidDNxHDDecoderElectra/Source/Private/Windows/ElectraMediaAvidDNxHDDecoder.cpp
@@ -60,11 +60,11 @@ public:
 
 	int32 GetWidth() const override
 	{
-		return Width - Crop.Left - Crop.Right;
+		return ImageWidth;
 	}
 	int32 GetHeight() const override
 	{
-		return Height - Crop.Top - Crop.Bottom;
+		return ImageHeight;
 	}
 	int32 GetDecodedWidth() const override
 	{
@@ -78,14 +78,6 @@ public:
 	{
 		return Crop;
 	}
-	int32 GetFrameWidth() const override
-	{
-		return Pitch;
-	}
-	int32 GetFrameHeight() const override
-	{
-		return Height;
-	}
 	int32 GetAspectRatioW() const override
 	{
 		return AspectW;
@@ -110,10 +102,6 @@ public:
 	{
 		OutExtraValues = ExtraValues;
 	}
-	int32 GetPixelFormat() const override
-	{
-		return PixelFormat;
-	}
 	void* GetPlatformOutputHandle(EElectraDecoderPlatformOutputHandleType InTypeOfHandle) const override
 	{
 		if (InTypeOfHandle == EElectraDecoderPlatformOutputHandleType::ImageBuffers)
@@ -140,7 +128,7 @@ public:
 	{
 		return 1;
 	}
-	TSharedPtr<TArray<uint8>, ESPMode::ThreadSafe> GetBufferByIndex(int32 InBufferIndex) const override
+	TSharedPtr<TArray<uint8>, ESPMode::ThreadSafe> GetBufferDataByIndex(int32 InBufferIndex) const override
 	{
 		if (InBufferIndex == 0)
 		{
@@ -148,21 +136,42 @@ public:
 		}
 		return nullptr;
 	}
-	uint64 GetBufferFormatByIndex(int32 InBufferIndex) const override
+	void* GetBufferTextureByIndex(int32 InBufferIndex) const override
+	{
+		return nullptr;
+	}
+	EElectraDecoderPlatformPixelFormat GetBufferFormatByIndex(int32 InBufferIndex) const override
 	{
 		if (InBufferIndex == 0)
 		{
 			return BufferFormat;
 		}
+		return EElectraDecoderPlatformPixelFormat::INVALID;
+	}
+	EElectraDecoderPlatformPixelEncoding GetBufferEncodingByIndex(int32 InBufferIndex) const override
+	{
+		if (InBufferIndex == 0)
+		{
+			return BufferEncoding;
+		}
+		return EElectraDecoderPlatformPixelEncoding::Native;
+	}
+	int32 GetBufferPitchByIndex(int32 InBufferIndex) const override
+	{
+		if (InBufferIndex == 0)
+		{
+			return Pitch;
+		}
 		return 0;
 	}
 
-
 public:
 	FTimespan PTS;
 	uint64 UserValue = 0;
 
 	FElectraVideoDecoderOutputCropValues Crop;
+	int32 ImageWidth = 0;
+	int32 ImageHeight = 0;
 	int32 Width = 0;
 	int32 Height = 0;
 	int32 Pitch = 0;
@@ -176,7 +185,8 @@ public:
 
 	uint32 Codec4CC = 0;
 	TSharedPtr<TArray<uint8>, ESPMode::ThreadSafe> Buffer;
-	uint64 BufferFormat = 0;
+	EElectraDecoderPlatformPixelFormat BufferFormat = EElectraDecoderPlatformPixelFormat::INVALID;
+	EElectraDecoderPlatformPixelEncoding BufferEncoding = EElectraDecoderPlatformPixelEncoding::Native;
 };
 
 
@@ -493,8 +503,126 @@ IElectraDecoder::EDecoderError FVideoDecoderAvidDNxHDElectra::DecodeAccessUnit(c
 		}
 		NewOutput->Codec4CC = Codec4CC;
 		NewOutput->NumBits = Decoder.GetNumOutputBits();
+
+		NewOutput->ImageWidth = DecodedWidth - NewOutput->Crop.Left - NewOutput->Crop.Right;
+		NewOutput->ImageHeight = DecodedHeight - NewOutput->Crop.Top - NewOutput->Crop.Bottom;
+
 		NewOutput->PixelFormat = (int32) Decoder.CurrentUncompressedParams.compType;
-		NewOutput->BufferFormat = (uint64) Decoder.CurrentUncompressedParams.compOrder;
+
+		switch (Decoder.CurrentUncompressedParams.compType)
+		{
+			case	DNX_CT_UCHAR:
+			{
+				switch (Decoder.CurrentUncompressedParams.compOrder)
+				{
+					case	DNX_CCO_YCbYCr_NoA:
+					case	DNX_CCO_CbYCrY_NoA:
+					{
+						NewOutput->BufferFormat = EElectraDecoderPlatformPixelFormat::B8G8R8A8;
+						NewOutput->BufferEncoding = (Decoder.CurrentUncompressedParams.compOrder  == DNX_CCO_CbYCrY_NoA) ?  EElectraDecoderPlatformPixelEncoding::CbY0CrY1 : EElectraDecoderPlatformPixelEncoding::Y0CbY1Cr;
+						NewOutput->Width /= 2;
+						NewOutput->Pitch *= 2;
+						break;
+					}
+					case	DNX_CCO_ARGB_Interleaved:
+					{
+						NewOutput->BufferFormat = EElectraDecoderPlatformPixelFormat::A8R8G8B8;
+						NewOutput->BufferEncoding = EElectraDecoderPlatformPixelEncoding::Native;
+						NewOutput->Pitch *= 4;
+						break;
+					}
+					case	DNX_CCO_RGBA_Interleaved:
+					{
+						NewOutput->BufferFormat = EElectraDecoderPlatformPixelFormat::R8G8B8A8;
+						NewOutput->BufferEncoding = EElectraDecoderPlatformPixelEncoding::Native;
+						NewOutput->Pitch *= 4;
+						break;
+					}
+					case	DNX_CCO_CbYCrA_Interleaved:
+					{
+						NewOutput->BufferFormat = EElectraDecoderPlatformPixelFormat::B8G8R8A8;
+						NewOutput->BufferEncoding = EElectraDecoderPlatformPixelEncoding::YCbCr_Alpha;
+						NewOutput->Pitch *= 4;
+						break;
+					}
+					case	DNX_CCO_YCbCr_Planar:
+					{
+						NewOutput->BufferFormat = EElectraDecoderPlatformPixelFormat::NV12;
+						NewOutput->BufferEncoding = EElectraDecoderPlatformPixelEncoding::Native;
+						NewOutput->Height = (NewOutput->Height * 3) / 2;
+						break;
+					}
+					default:
+					{
+						NewOutput->BufferFormat = EElectraDecoderPlatformPixelFormat::INVALID;
+						NewOutput->BufferEncoding = EElectraDecoderPlatformPixelEncoding::Native;
+					}
+				}
+				break;
+			}
+			case	DNX_CT_SHORT:
+			case	DNX_CT_USHORT_10_6:
+			case	DNX_CT_USHORT_12_4:
+			{
+				switch (Decoder.CurrentUncompressedParams.compOrder)
+				{
+					case	DNX_CCO_YCbYCr_NoA:
+					case	DNX_CCO_CbYCrY_NoA:
+					{
+						NewOutput->BufferFormat = EElectraDecoderPlatformPixelFormat::A16B16G16R16;
+						NewOutput->BufferEncoding = (Decoder.CurrentUncompressedParams.compOrder == DNX_CCO_CbYCrY_NoA) ? EElectraDecoderPlatformPixelEncoding::CbY0CrY1 : EElectraDecoderPlatformPixelEncoding::Y0CbY1Cr;
+						NewOutput->Width /= 2;
+						NewOutput->Pitch *= 4;
+						break;
+					}
+					case	DNX_CCO_RGBA_Interleaved:
+					{
+						NewOutput->BufferFormat = EElectraDecoderPlatformPixelFormat::R16G16B16A16;
+						NewOutput->BufferEncoding = EElectraDecoderPlatformPixelEncoding::Native;
+						NewOutput->Pitch *= 8;
+						break;
+					}
+					case	DNX_CCO_ABGR_Interleaved:
+					{
+						NewOutput->BufferFormat = EElectraDecoderPlatformPixelFormat::A16B16G16R16;
+						NewOutput->BufferEncoding = EElectraDecoderPlatformPixelEncoding::Native;
+						NewOutput->Pitch *= 8;
+						break;
+					}
+					case	DNX_CCO_CbYCrA_Interleaved:
+					{
+						NewOutput->BufferFormat = EElectraDecoderPlatformPixelFormat::A16B16G16R16;
+						NewOutput->BufferEncoding = EElectraDecoderPlatformPixelEncoding::YCbCr_Alpha;
+						NewOutput->Pitch *= 8;
+						break;
+					}
+					case	DNX_CCO_YCbCr_Planar:
+					{
+						NewOutput->BufferFormat = EElectraDecoderPlatformPixelFormat::P010;
+						NewOutput->BufferEncoding = EElectraDecoderPlatformPixelEncoding::Native;
+						NewOutput->Height = (NewOutput->Height * 3) / 2;
+						break;
+					}
+					default:
+					{
+						NewOutput->BufferFormat = EElectraDecoderPlatformPixelFormat::INVALID;
+						NewOutput->BufferEncoding = EElectraDecoderPlatformPixelEncoding::Native;
+					}
+				}
+				break;
+			}
+			case	DNX_CT_V210:
+			{
+				NewOutput->BufferFormat = EElectraDecoderPlatformPixelFormat::A2B10G10R10;
+				NewOutput->BufferEncoding = EElectraDecoderPlatformPixelEncoding::CbY0CrY1;
+				break;
+			}
+			default:
+			{
+				NewOutput->BufferFormat = EElectraDecoderPlatformPixelFormat::INVALID;
+				NewOutput->BufferEncoding = EElectraDecoderPlatformPixelEncoding::Native;
+			}
+		}
 
 		NewOutput->ExtraValues.Emplace(TEXT("codec"), FVariant(TEXT("avid")));
 		NewOutput->ExtraValues.Emplace(TEXT("codec_4cc"), FVariant(Codec4CC));
diff --git a/Engine/Plugins/Media/ElectraCodecs/Source/ElectraDecoders/Private/Android/h264/VideoDecoderH264_Android.cpp b/Engine/Plugins/Media/ElectraCodecs/Source/ElectraDecoders/Private/Android/h264/VideoDecoderH264_Android.cpp
--- a/Engine/Plugins/Media/ElectraCodecs/Source/ElectraDecoders/Private/Android/h264/VideoDecoderH264_Android.cpp
+++ b/Engine/Plugins/Media/ElectraCodecs/Source/ElectraDecoders/Private/Android/h264/VideoDecoderH264_Android.cpp
@@ -63,10 +63,6 @@ public:
 	{ return Height; }
 	FElectraVideoDecoderOutputCropValues GetCropValues() const override
 	{ return Crop; }
-	int32 GetFrameWidth() const override
-	{ return Pitch; }
-	int32 GetFrameHeight() const override
-	{ return Height; }
 	int32 GetAspectRatioW() const override
 	{ return AspectW; }
 	int32 GetAspectRatioH() const override
@@ -79,8 +75,6 @@ public:
 	{ return NumBits; }
 	void GetExtraValues(TMap<FString, FVariant>& OutExtraValues) const override
 	{ OutExtraValues = ExtraValues; }
-	int32 GetPixelFormat() const override
-	{ return PixelFormat; }
 	void* GetPlatformOutputHandle(EElectraDecoderPlatformOutputHandleType InTypeOfHandle) const override
 	{ return nullptr; }
 	IElectraDecoderVideoOutputTransferHandle* GetTransferHandle() const override
diff --git a/Engine/Plugins/Media/ElectraCodecs/Source/ElectraDecoders/Private/Android/h265/VideoDecoderH265_Android.cpp b/Engine/Plugins/Media/ElectraCodecs/Source/ElectraDecoders/Private/Android/h265/VideoDecoderH265_Android.cpp
--- a/Engine/Plugins/Media/ElectraCodecs/Source/ElectraDecoders/Private/Android/h265/VideoDecoderH265_Android.cpp
+++ b/Engine/Plugins/Media/ElectraCodecs/Source/ElectraDecoders/Private/Android/h265/VideoDecoderH265_Android.cpp
@@ -63,10 +63,6 @@ public:
 	{ return Height; }
 	FElectraVideoDecoderOutputCropValues GetCropValues() const override
 	{ return Crop; }
-	int32 GetFrameWidth() const override
-	{ return Pitch; }
-	int32 GetFrameHeight() const override
-	{ return Height; }
 	int32 GetAspectRatioW() const override
 	{ return AspectW; }
 	int32 GetAspectRatioH() const override
@@ -79,8 +75,6 @@ public:
 	{ return NumBits; }
 	void GetExtraValues(TMap<FString, FVariant>& OutExtraValues) const override
 	{ OutExtraValues = ExtraValues; }
-	int32 GetPixelFormat() const override
-	{ return PixelFormat; }
 	void* GetPlatformOutputHandle(EElectraDecoderPlatformOutputHandleType InTypeOfHandle) const override
 	{ return nullptr; }
 	IElectraDecoderVideoOutputTransferHandle* GetTransferHandle() const override
diff --git a/Engine/Plugins/Media/ElectraCodecs/Source/ElectraDecoders/Private/Apple/h264/VideoDecoderH264_Apple.cpp b/Engine/Plugins/Media/ElectraCodecs/Source/ElectraDecoders/Private/Apple/h264/VideoDecoderH264_Apple.cpp
--- a/Engine/Plugins/Media/ElectraCodecs/Source/ElectraDecoders/Private/Apple/h264/VideoDecoderH264_Apple.cpp
+++ b/Engine/Plugins/Media/ElectraCodecs/Source/ElectraDecoders/Private/Apple/h264/VideoDecoderH264_Apple.cpp
@@ -31,7 +31,7 @@ public:
 };
 
 
-class FElectraVideoDecoderOutputH264_Apple : public IElectraDecoderVideoOutput
+class FElectraVideoDecoderOutputH264_Apple : public IElectraDecoderVideoOutput, public IElectraDecoderVideoOutputImageBuffers
 {
 public:
 	virtual ~FElectraVideoDecoderOutputH264_Apple()
@@ -54,10 +54,6 @@ public:
 	{ return Height; }
 	FElectraVideoDecoderOutputCropValues GetCropValues() const override
 	{ return Crop; }
-	int32 GetFrameWidth() const override
-	{ return Pitch; }
-	int32 GetFrameHeight() const override
-	{ return Height; }
 	int32 GetAspectRatioW() const override
 	{ return AspectW; }
 	int32 GetAspectRatioH() const override
@@ -70,14 +66,11 @@ public:
 	{ return NumBits; }
 	void GetExtraValues(TMap<FString, FVariant>& OutExtraValues) const override
 	{ OutExtraValues = ExtraValues; }
-	int32 GetPixelFormat() const override
-	{ return PixelFormat; }
 	void* GetPlatformOutputHandle(EElectraDecoderPlatformOutputHandleType InTypeOfHandle) const override
 	{
-		check(InTypeOfHandle == EElectraDecoderPlatformOutputHandleType::ImageBufferRef);
-		if (InTypeOfHandle == EElectraDecoderPlatformOutputHandleType::ImageBufferRef)
+		if (InTypeOfHandle == EElectraDecoderPlatformOutputHandleType::ImageBuffers)
 		{
-			return ImageBuffer;
+			return static_cast<IElectraDecoderVideoOutputImageBuffers*>(const_cast<FElectraVideoDecoderOutputH264_Apple*>(this));
 		}
 		return nullptr;
 	}
@@ -86,6 +79,36 @@ public:
 	IElectraDecoderVideoOutput::EImageCopyResult CopyPlatformImage(IElectraDecoderVideoOutputCopyResources* InCopyResources) const override
 	{ return IElectraDecoderVideoOutput::EImageCopyResult::NotSupported; }
 
+	// Methods from IElectraDecoderVideoOutputImageBuffers
+	uint32 GetCodec4CC() const override
+	{
+		return 0x34363248;
+	}
+	int32 GetNumberOfBuffers() const override
+	{
+		return 1;
+	}
+	TSharedPtr<TArray<uint8>, ESPMode::ThreadSafe> GetBufferDataByIndex(int32 InBufferIndex) const override
+	{
+		return nullptr;
+	}
+	void* GetBufferTextureByIndex(int32 InBufferIndex) const override
+	{
+		return InBufferIndex == 0 ? ImageBuffer : nullptr;
+	}
+	EElectraDecoderPlatformPixelFormat GetBufferFormatByIndex(int32 InBufferIndex) const override
+	{
+		return PixelFormat;
+	}
+	EElectraDecoderPlatformPixelEncoding GetBufferEncodingByIndex(int32 InBufferIndex) const override
+	{
+		return PixelEncoding;
+	}
+	int32 GetBufferPitchByIndex(int32 InBufferIndex) const override
+	{
+		return Pitch;
+	}
+
 public:
 	void ReleaseOutputBuffer()
 	{
@@ -108,7 +131,8 @@ public:
 	int32 AspectH = 1;
 	int32 FrameRateN = 0;
 	int32 FrameRateD = 0;
-	int32 PixelFormat = 0;
+	EElectraDecoderPlatformPixelFormat PixelFormat = EElectraDecoderPlatformPixelFormat::INVALID;
+	EElectraDecoderPlatformPixelEncoding PixelEncoding = EElectraDecoderPlatformPixelEncoding::Native;
 	TMap<FString, FVariant> ExtraValues;
 
 	CVImageBufferRef ImageBuffer = nullptr;
@@ -856,12 +880,33 @@ FElectraVideoDecoderH264_Apple::EConvertResult FElectraVideoDecoderH264_Apple::C
 	NewOutput->Height = In->SPS->GetHeight();
 	NewOutput->Pitch = NewOutput->Width;
 
-	//NewOutput->NumBits = 8 + In->SPS->bit_depth_luma_minus8;
-	//NewOutput->PixelFormat = 0;
-
 	NewOutput->ImageBuffer = NextImage.ReleaseImageBufferRef();
 	int pixelFormat = CVPixelBufferGetPixelFormatType(NewOutput->ImageBuffer);
-	NewOutput->PixelFormat = pixelFormat;
+
+	switch (pixelFormat)
+	{
+		case	kCVPixelFormatType_420YpCbCr8BiPlanarVideoRange:
+		case	kCVPixelFormatType_420YpCbCr8BiPlanarFullRange:
+		{
+			NewOutput->PixelFormat = EElectraDecoderPlatformPixelFormat::NV12;
+			NewOutput->PixelEncoding = EElectraDecoderPlatformPixelEncoding::Native;
+			break;
+		}
+		case	kCVPixelFormatType_420YpCbCr10BiPlanarVideoRange:
+		case	kCVPixelFormatType_420YpCbCr10BiPlanarFullRange:
+		{
+			NewOutput->PixelFormat = EElectraDecoderPlatformPixelFormat::P010;
+			NewOutput->PixelEncoding = EElectraDecoderPlatformPixelEncoding::Native;
+			break;
+		}
+		default:
+		{
+			check(!"Unexpected output format!");
+			NewOutput->PixelFormat = EElectraDecoderPlatformPixelFormat::INVALID;
+			NewOutput->PixelEncoding = EElectraDecoderPlatformPixelEncoding::Native;
+		}
+	}
+
 	NewOutput->NumBits = ElectraVideoDecoderFormatTypesApple::GetNumComponentBitsForPixelFormat(pixelFormat);
 	if (NewOutput->NumBits < 0)
 	{
diff --git a/Engine/Plugins/Media/ElectraCodecs/Source/ElectraDecoders/Private/Apple/h265/VideoDecoderH265_Apple.cpp b/Engine/Plugins/Media/ElectraCodecs/Source/ElectraDecoders/Private/Apple/h265/VideoDecoderH265_Apple.cpp
--- a/Engine/Plugins/Media/ElectraCodecs/Source/ElectraDecoders/Private/Apple/h265/VideoDecoderH265_Apple.cpp
+++ b/Engine/Plugins/Media/ElectraCodecs/Source/ElectraDecoders/Private/Apple/h265/VideoDecoderH265_Apple.cpp
@@ -31,7 +31,7 @@ public:
 };
 
 
-class FElectraVideoDecoderOutputH265_Apple : public IElectraDecoderVideoOutput
+class FElectraVideoDecoderOutputH265_Apple : public IElectraDecoderVideoOutput, public IElectraDecoderVideoOutputImageBuffers
 {
 public:
 	virtual ~FElectraVideoDecoderOutputH265_Apple()
@@ -54,10 +54,6 @@ public:
 	{ return Height; }
 	FElectraVideoDecoderOutputCropValues GetCropValues() const override
 	{ return Crop; }
-	int32 GetFrameWidth() const override
-	{ return Pitch; }
-	int32 GetFrameHeight() const override
-	{ return Height; }
 	int32 GetAspectRatioW() const override
 	{ return AspectW; }
 	int32 GetAspectRatioH() const override
@@ -70,14 +66,11 @@ public:
 	{ return NumBits; }
 	void GetExtraValues(TMap<FString, FVariant>& OutExtraValues) const override
 	{ OutExtraValues = ExtraValues; }
-	int32 GetPixelFormat() const override
-	{ return PixelFormat; }
 	void* GetPlatformOutputHandle(EElectraDecoderPlatformOutputHandleType InTypeOfHandle) const override
 	{
-		check(InTypeOfHandle == EElectraDecoderPlatformOutputHandleType::ImageBufferRef);
-		if (InTypeOfHandle == EElectraDecoderPlatformOutputHandleType::ImageBufferRef)
+		if (InTypeOfHandle == EElectraDecoderPlatformOutputHandleType::ImageBuffers)
 		{
-			return ImageBuffer;
+			return static_cast<IElectraDecoderVideoOutputImageBuffers*>(const_cast<FElectraVideoDecoderOutputH265_Apple*>(this));
 		}
 		return nullptr;
 	}
@@ -86,6 +79,36 @@ public:
 	IElectraDecoderVideoOutput::EImageCopyResult CopyPlatformImage(IElectraDecoderVideoOutputCopyResources* InCopyResources) const override
 	{ return IElectraDecoderVideoOutput::EImageCopyResult::NotSupported; }
 
+	// Methods from IElectraDecoderVideoOutputImageBuffers
+	uint32 GetCodec4CC() const override
+	{
+		return 0x31637668;
+	}
+	int32 GetNumberOfBuffers() const override
+	{
+		return 1;
+	}
+	TSharedPtr<TArray<uint8>, ESPMode::ThreadSafe> GetBufferDataByIndex(int32 InBufferIndex) const override
+	{
+		return nullptr;
+	}
+	void* GetBufferTextureByIndex(int32 InBufferIndex) const override
+	{
+		return InBufferIndex == 0 ? ImageBuffer : nullptr;
+	}
+	EElectraDecoderPlatformPixelFormat GetBufferFormatByIndex(int32 InBufferIndex) const override
+	{
+		return PixelFormat;
+	}
+	EElectraDecoderPlatformPixelEncoding GetBufferEncodingByIndex(int32 InBufferIndex) const override
+	{
+		return PixelEncoding;
+	}
+	int32 GetBufferPitchByIndex(int32 InBufferIndex) const override
+	{
+		return Pitch;
+	}
+
 public:
 	void ReleaseOutputBuffer()
 	{
@@ -108,7 +131,8 @@ public:
 	int32 AspectH = 1;
 	int32 FrameRateN = 0;
 	int32 FrameRateD = 0;
-	int32 PixelFormat = 0;
+	EElectraDecoderPlatformPixelFormat PixelFormat = EElectraDecoderPlatformPixelFormat::INVALID;
+	EElectraDecoderPlatformPixelEncoding PixelEncoding = EElectraDecoderPlatformPixelEncoding::Native;
 	TMap<FString, FVariant> ExtraValues;
 
 	CVImageBufferRef ImageBuffer = nullptr;
@@ -860,12 +884,33 @@ FElectraVideoDecoderH265_Apple::EConvertResult FElectraVideoDecoderH265_Apple::C
 	NewOutput->Height = In->SPS->GetHeight();
 	NewOutput->Pitch = NewOutput->Width;
 
-	//NewOutput->NumBits = 8 + In->SPS->bit_depth_luma_minus8;
-	//NewOutput->PixelFormat = 0;
-
 	NewOutput->ImageBuffer = NextImage.ReleaseImageBufferRef();
 	int pixelFormat = CVPixelBufferGetPixelFormatType(NewOutput->ImageBuffer);
-	NewOutput->PixelFormat = pixelFormat;
+
+	switch (pixelFormat)
+	{
+		case	kCVPixelFormatType_420YpCbCr8BiPlanarVideoRange:
+		case	kCVPixelFormatType_420YpCbCr8BiPlanarFullRange:
+		{
+			NewOutput->PixelFormat = EElectraDecoderPlatformPixelFormat::NV12;
+			NewOutput->PixelEncoding = EElectraDecoderPlatformPixelEncoding::Native;
+			break;
+		}
+		case	kCVPixelFormatType_420YpCbCr10BiPlanarVideoRange:
+		case	kCVPixelFormatType_420YpCbCr10BiPlanarFullRange:
+		{
+			NewOutput->PixelFormat = EElectraDecoderPlatformPixelFormat::P010;
+			NewOutput->PixelEncoding = EElectraDecoderPlatformPixelEncoding::Native;
+			break;
+		}
+		default:
+		{
+			check(!"Unexpected output format!");
+			NewOutput->PixelFormat = EElectraDecoderPlatformPixelFormat::INVALID;
+			NewOutput->PixelEncoding = EElectraDecoderPlatformPixelEncoding::Native;
+		}
+	}
+
 	NewOutput->NumBits = ElectraVideoDecoderFormatTypesApple::GetNumComponentBitsForPixelFormat(pixelFormat);
 	if (NewOutput->NumBits < 0)
 	{
diff --git a/Engine/Plugins/Media/ElectraCodecs/Source/ElectraDecoders/Private/DX/VideoDecoderH264_DX.cpp b/Engine/Plugins/Media/ElectraCodecs/Source/ElectraDecoders/Private/DX/VideoDecoderH264_DX.cpp
--- a/Engine/Plugins/Media/ElectraCodecs/Source/ElectraDecoders/Private/DX/VideoDecoderH264_DX.cpp
+++ b/Engine/Plugins/Media/ElectraCodecs/Source/ElectraDecoders/Private/DX/VideoDecoderH264_DX.cpp
@@ -54,15 +54,11 @@ public:
 	int32 GetHeight() const override
 	{ return Height - Crop.Top - Crop.Bottom; }
 	int32 GetDecodedWidth() const override
-	{ return Width; }
+	{ return DecodedWidth; }
 	int32 GetDecodedHeight() const override
-	{ return Height; }
+	{ return DecodedHeight; }
 	FElectraVideoDecoderOutputCropValues GetCropValues() const override
 	{ return Crop; }
-	int32 GetFrameWidth() const override
-	{ return Pitch; }
-	int32 GetFrameHeight() const override
-	{ return Height; }
 	int32 GetAspectRatioW() const override
 	{ return AspectW; }
 	int32 GetAspectRatioH() const override
@@ -75,8 +71,6 @@ public:
 	{ return NumBits; }
 	void GetExtraValues(TMap<FString, FVariant>& OutExtraValues) const override
 	{ OutExtraValues = ExtraValues; }
-	int32 GetPixelFormat() const override
-	{ return PixelFormat; }
 	void* GetPlatformOutputHandle(EElectraDecoderPlatformOutputHandleType InTypeOfHandle) const override
 	{ 
 		switch(InTypeOfHandle)
@@ -103,6 +97,8 @@ public:
 	FElectraVideoDecoderOutputCropValues Crop;
 	int32 Width = 0;
 	int32 Height = 0;
+	int32 DecodedWidth = 0;
+	int32 DecodedHeight = 0;
 	int32 Pitch = 0;
 	int32 NumBits = 0;
 	int32 AspectW = 1;
@@ -838,11 +834,22 @@ bool FElectraVideoDecoderH264_DX::ConvertDecoderOutput()
 			break;
 		}
 	}
+
+	// If no exact match is found then the decoder transform may have interpolated the PTS given the individual
+	// frame durations or perhaps due to SEI messages.
+	if (!bFoundMatch)
+	{
+		MatchingInput = InDecoderInput[0];
+		bFoundMatch = true;
+		InDecoderInput.RemoveAt(0);
+	}
+	/*
 	if (!bFoundMatch)
 	{
 		PostError(0, TEXT("There is no matching decoder input for the decoded output!"), ERRCODE_INTERNAL_FAILED_TO_CONVERT_OUTPUT_SAMPLE);
 		return false;
 	}
+	*/
 
 	TSharedPtr<FElectraVideoDecoderOutputH264_DX, ESPMode::ThreadSafe> NewOutput = MakeShared<FElectraVideoDecoderOutputH264_DX>();
 	NewOutput->PTS = MatchingInput.AccessUnit.PTS;
@@ -870,9 +877,18 @@ bool FElectraVideoDecoderH264_DX::ConvertDecoderOutput()
 	}
 
 	NewOutput->NumBits = 8;
-	// Width and height are the total pixels without cropping.
-	NewOutput->Width = (int32) dwInputWidth;
-	NewOutput->Height = (int32) dwInputHeight;
+	// Width and height of the final output image
+	NewOutput->Width = (int32)videoArea.Area.cx;
+	NewOutput->Height = (int32)videoArea.Area.cy;
+
+	// Width and height as "decoded dims" (which are featuring format specific scales & no cropping)
+	NewOutput->DecodedWidth = (int32)dwInputWidth;
+	NewOutput->DecodedHeight = (int32)dwInputHeight;
+	if (DecoderPlatformHandle->IsSoftware() || DecoderPlatformHandle->GetDXVersionTimes1000() >= 12000)
+	{
+		// We are returning a CPU side buffer (software decode OR DX12 render device) - adjust the height so it can be interpreted as a single plane texture
+		NewOutput->DecodedHeight = NewOutput->DecodedHeight * 3 / 2;
+	}
 
 	// Note: Decoder crops at the lower right border.
 	NewOutput->Crop.Left = 0;
@@ -906,6 +922,8 @@ bool FElectraVideoDecoderH264_DX::ConvertDecoderOutput()
 	NewOutput->ExtraValues.Emplace(TEXT("dxversion"), FVariant((int64) DecoderPlatformHandle->GetDXVersionTimes1000()));
 	NewOutput->ExtraValues.Emplace(TEXT("sw"), FVariant(DecoderPlatformHandle->IsSoftware()));
 	NewOutput->ExtraValues.Emplace(TEXT("codec"), FVariant(TEXT("avc")));
+	NewOutput->ExtraValues.Emplace(TEXT("pixfmt"), FVariant((int64)EElectraDecoderPlatformPixelFormat::NV12));
+	NewOutput->ExtraValues.Emplace(TEXT("pixenc"), FVariant((int64)EElectraDecoderPlatformPixelEncoding::Native));
 
 	// Retain the IMFSample in the output. It is needed later in converting it for display.
 	NewOutput->MFSample = DecodedOutputSample;
diff --git a/Engine/Plugins/Media/ElectraCodecs/Source/ElectraDecoders/Private/DX/VideoDecoderH265_DX.cpp b/Engine/Plugins/Media/ElectraCodecs/Source/ElectraDecoders/Private/DX/VideoDecoderH265_DX.cpp
--- a/Engine/Plugins/Media/ElectraCodecs/Source/ElectraDecoders/Private/DX/VideoDecoderH265_DX.cpp
+++ b/Engine/Plugins/Media/ElectraCodecs/Source/ElectraDecoders/Private/DX/VideoDecoderH265_DX.cpp
@@ -54,15 +54,11 @@ public:
 	int32 GetHeight() const override
 	{ return Height - Crop.Top - Crop.Bottom; }
 	int32 GetDecodedWidth() const override
-	{ return Width; }
+	{ return DecodedWidth; }
 	int32 GetDecodedHeight() const override
-	{ return Height; }
+	{ return DecodedHeight; }
 	FElectraVideoDecoderOutputCropValues GetCropValues() const override
 	{ return Crop; }
-	int32 GetFrameWidth() const override
-	{ return Pitch; }
-	int32 GetFrameHeight() const override
-	{ return Height; }
 	int32 GetAspectRatioW() const override
 	{ return AspectW; }
 	int32 GetAspectRatioH() const override
@@ -75,8 +71,6 @@ public:
 	{ return NumBits; }
 	void GetExtraValues(TMap<FString, FVariant>& OutExtraValues) const override
 	{ OutExtraValues = ExtraValues; }
-	int32 GetPixelFormat() const override
-	{ return PixelFormat; }
 	void* GetPlatformOutputHandle(EElectraDecoderPlatformOutputHandleType InTypeOfHandle) const override
 	{ 
 		switch(InTypeOfHandle)
@@ -103,6 +97,8 @@ public:
 	FElectraVideoDecoderOutputCropValues Crop;
 	int32 Width = 0;
 	int32 Height = 0;
+	int32 DecodedWidth = 0;
+	int32 DecodedHeight = 0;
 	int32 Pitch = 0;
 	int32 NumBits = 0;
 	int32 AspectW = 1;
@@ -853,11 +849,21 @@ bool FElectraVideoDecoderH265_DX::ConvertDecoderOutput()
 			break;
 		}
 	}
+	// If no exact match is found then the decoder transform may have interpolated the PTS given the individual
+	// frame durations or perhaps due to SEI messages.
+	if (!bFoundMatch)
+	{
+		MatchingInput = InDecoderInput[0];
+		bFoundMatch = true;
+		InDecoderInput.RemoveAt(0);
+	}
+	/*
 	if (!bFoundMatch)
 	{
 		PostError(0, TEXT("There is no matching decoder input for the decoded output!"), ERRCODE_INTERNAL_FAILED_TO_CONVERT_OUTPUT_SAMPLE);
 		return false;
 	}
+	*/
 
 	TSharedPtr<FElectraVideoDecoderOutputH265_DX, ESPMode::ThreadSafe> NewOutput = MakeShared<FElectraVideoDecoderOutputH265_DX>();
 	NewOutput->PTS = MatchingInput.AccessUnit.PTS;
@@ -884,9 +890,18 @@ bool FElectraVideoDecoderH265_DX::ConvertDecoderOutput()
 		}
 	}
 
-	// Width and height are the total pixels without cropping.
-	NewOutput->Width = (int32) dwInputWidth;
-	NewOutput->Height = (int32) dwInputHeight;
+	// Width and height of the final output image
+	NewOutput->Width = (int32)videoArea.Area.cx;
+	NewOutput->Height = (int32)videoArea.Area.cy;
+
+	// Width and height as "decoded dims" (which are featuring format specific scales & no cropping)
+	NewOutput->DecodedWidth = (int32)dwInputWidth;
+	NewOutput->DecodedHeight = (int32)dwInputHeight;
+	if (DecoderPlatformHandle->IsSoftware() || DecoderPlatformHandle->GetDXVersionTimes1000() >= 12000)
+	{
+		// We are returning a CPU side buffer (software decode OR DX12 render device) - adjust the height so it can be interpreted as a single plane texture
+		NewOutput->DecodedHeight = NewOutput->DecodedHeight * 3 / 2;
+	}
 
 	// Note: Decoder crops at the lower right border.
 	NewOutput->Crop.Left = 0;
@@ -917,6 +932,7 @@ bool FElectraVideoDecoderH265_DX::ConvertDecoderOutput()
 
 	// Get actual output format used right now.
 	GUID CurrentVideoFormatGUID;
+	EElectraDecoderPlatformPixelFormat PixFmt;
 	if (SUCCEEDED(CurrentOutputMediaType->GetGUID(MF_MT_SUBTYPE, &CurrentVideoFormatGUID)))
 	{
 		check(CurrentVideoFormatGUID == MFVideoFormat_NV12 || CurrentVideoFormatGUID == MFVideoFormat_P010);
@@ -929,16 +945,20 @@ bool FElectraVideoDecoderH265_DX::ConvertDecoderOutput()
 	{
 		NewOutput->NumBits = 8;
 		NewOutput->PixelFormat = 0;
+		PixFmt = EElectraDecoderPlatformPixelFormat::NV12;
 	}
 	else
 	{
 		NewOutput->NumBits = 10;
 		NewOutput->PixelFormat = 1;
+		PixFmt = EElectraDecoderPlatformPixelFormat::P010;
 	}
 	NewOutput->ExtraValues.Emplace(TEXT("platform"), FVariant(TEXT("dx")));
 	NewOutput->ExtraValues.Emplace(TEXT("dxversion"), FVariant((int64) DecoderPlatformHandle->GetDXVersionTimes1000()));
 	NewOutput->ExtraValues.Emplace(TEXT("sw"), FVariant(DecoderPlatformHandle->IsSoftware()));
 	NewOutput->ExtraValues.Emplace(TEXT("codec"), FVariant(TEXT("hevc")));
+	NewOutput->ExtraValues.Emplace(TEXT("pixfmt"), FVariant((int64)PixFmt));
+	NewOutput->ExtraValues.Emplace(TEXT("pixenc"), FVariant((int64)EElectraDecoderPlatformPixelEncoding::Native));
 
 	// Retain the IMFSample in the output. It is needed later in converting it for display.
 	NewOutput->MFSample = DecodedOutputSample;
diff --git a/Engine/Plugins/Media/ElectraCodecs/Source/ElectraDecoders/Private/Linux/h264/VideoDecoderH264_Linux.cpp b/Engine/Plugins/Media/ElectraCodecs/Source/ElectraDecoders/Private/Linux/h264/VideoDecoderH264_Linux.cpp
--- a/Engine/Plugins/Media/ElectraCodecs/Source/ElectraDecoders/Private/Linux/h264/VideoDecoderH264_Linux.cpp
+++ b/Engine/Plugins/Media/ElectraCodecs/Source/ElectraDecoders/Private/Linux/h264/VideoDecoderH264_Linux.cpp
@@ -51,10 +51,6 @@ public:
 	{ return Height; }
 	FElectraVideoDecoderOutputCropValues GetCropValues() const override
 	{ return Crop; }
-	int32 GetFrameWidth() const override
-	{ return Pitch; }
-	int32 GetFrameHeight() const override
-	{ return Height; }
 	int32 GetAspectRatioW() const override
 	{ return AspectW; }
 	int32 GetAspectRatioH() const override
@@ -67,8 +63,6 @@ public:
 	{ return NumBits; }
 	void GetExtraValues(TMap<FString, FVariant>& OutExtraValues) const override
 	{ OutExtraValues = ExtraValues; }
-	int32 GetPixelFormat() const override
-	{ return PixelFormat; }
 	void* GetPlatformOutputHandle(EElectraDecoderPlatformOutputHandleType InTypeOfHandle) const override
 	{
 		if (InTypeOfHandle == EElectraDecoderPlatformOutputHandleType::LibavDecoderDecodedImage)
diff --git a/Engine/Plugins/Media/ElectraCodecs/Source/ElectraDecoders/Private/Linux/h265/VideoDecoderH265_Linux.cpp b/Engine/Plugins/Media/ElectraCodecs/Source/ElectraDecoders/Private/Linux/h265/VideoDecoderH265_Linux.cpp
--- a/Engine/Plugins/Media/ElectraCodecs/Source/ElectraDecoders/Private/Linux/h265/VideoDecoderH265_Linux.cpp
+++ b/Engine/Plugins/Media/ElectraCodecs/Source/ElectraDecoders/Private/Linux/h265/VideoDecoderH265_Linux.cpp
@@ -51,10 +51,6 @@ public:
 	{ return Height; }
 	FElectraVideoDecoderOutputCropValues GetCropValues() const override
 	{ return Crop; }
-	int32 GetFrameWidth() const override
-	{ return Pitch; }
-	int32 GetFrameHeight() const override
-	{ return Height; }
 	int32 GetAspectRatioW() const override
 	{ return AspectW; }
 	int32 GetAspectRatioH() const override
@@ -67,8 +63,6 @@ public:
 	{ return NumBits; }
 	void GetExtraValues(TMap<FString, FVariant>& OutExtraValues) const override
 	{ OutExtraValues = ExtraValues; }
-	int32 GetPixelFormat() const override
-	{ return PixelFormat; }
 	void* GetPlatformOutputHandle(EElectraDecoderPlatformOutputHandleType InTypeOfHandle) const override
 	{
 		if (InTypeOfHandle == EElectraDecoderPlatformOutputHandleType::LibavDecoderDecodedImage)
diff --git a/Engine/Plugins/Media/ElectraCodecs/Source/ElectraDecoders/Public/Apple/AppleElectraDecoderPlatformOutputHandleTypes.h b/Engine/Plugins/Media/ElectraCodecs/Source/ElectraDecoders/Public/Apple/AppleElectraDecoderPlatformOutputHandleTypes.h
--- a/Engine/Plugins/Media/ElectraCodecs/Source/ElectraDecoders/Public/Apple/AppleElectraDecoderPlatformOutputHandleTypes.h
+++ b/Engine/Plugins/Media/ElectraCodecs/Source/ElectraDecoders/Public/Apple/AppleElectraDecoderPlatformOutputHandleTypes.h
@@ -6,10 +6,39 @@
 
 enum class EElectraDecoderPlatformOutputHandleType
 {
-	ImageBufferRef,			// CVImageBufferRef
 	ImageBuffers			// IElectraDecoderVideoOutputImageBuffers interface
 };
 
+enum class EElectraDecoderPlatformPixelFormat
+{
+	INVALID = 0,
+	R8G8B8A8,
+	A8R8G8B8,
+	B8G8R8A8,
+	R16G16B16A16,
+	A16B16G16R16,
+	A32B32G32R32F,
+	A2B10G10R10,
+	DXT1,
+	DXT5,
+	BC4,
+	NV12,
+	P010,
+};
+
+enum class EElectraDecoderPlatformPixelEncoding
+{
+	Native = 0,		//!< Pixel formats native representation
+	RGB,			//!< Interpret as RGB
+	RGBA,			//!< Interpret as RGBA
+	YCbCr,			//!< Interpret as YCbCR
+	YCbCr_Alpha,	//!< Interpret as YCbCR with alpha
+	YCoCg,			//!< Interpret as scaled YCoCg
+	YCoCg_Alpha,	//!< Interpret as scaled YCoCg with trailing BC4 alpha data
+	CbY0CrY1,		//!< Interpret as CbY0CrY1
+	Y0CbY1Cr,		//!< Interpret as Y0CbY1Cr
+	ARGB_BigEndian,	//!< Interpret as ARGB, big endian
+};
 
 class IElectraDecoderVideoOutputImageBuffers
 {
@@ -20,9 +49,18 @@ public:
 	// Returns the number of separate image buffers making up the frame.
 	virtual int32 GetNumberOfBuffers() const = 0;
 
-	// Returns the n'th image buffer.
-	virtual TSharedPtr<TArray<uint8>, ESPMode::ThreadSafe> GetBufferByIndex(int32 InBufferIndex) const = 0;
+	// Returns the n'th image buffer data address.
+	virtual TSharedPtr<TArray<uint8>, ESPMode::ThreadSafe> GetBufferDataByIndex(int32 InBufferIndex) const = 0;
+	
+	// Returns the n'th CFImageBuffer ("texture") reference.
+	virtual void* GetBufferTextureByIndex(int32 InBufferIndex) const = 0;
+
+	// Returns the n'th image buffer format
+	virtual EElectraDecoderPlatformPixelFormat GetBufferFormatByIndex(int32 InBufferIndex) const = 0;
+
+	// Returns the n'th image buffer encoding
+	virtual EElectraDecoderPlatformPixelEncoding GetBufferEncodingByIndex(int32 InBufferIndex) const = 0;
 
-	// Returns the n'th image buffer format, which is a specific value of the codec.
-	virtual uint64 GetBufferFormatByIndex(int32 InBufferIndex) const = 0;
+	// Returns the n'th image buffer pitch
+	virtual int32 GetBufferPitchByIndex(int32 InBufferIndex) const = 0;
 };
diff --git a/Engine/Plugins/Media/ElectraCodecs/Source/ElectraDecoders/Public/IElectraDecoderOutputVideo.h b/Engine/Plugins/Media/ElectraCodecs/Source/ElectraDecoders/Public/IElectraDecoderOutputVideo.h
--- a/Engine/Plugins/Media/ElectraCodecs/Source/ElectraDecoders/Public/IElectraDecoderOutputVideo.h
+++ b/Engine/Plugins/Media/ElectraCodecs/Source/ElectraDecoders/Public/IElectraDecoderOutputVideo.h
@@ -78,8 +78,11 @@ public:
 	
 	/**
 	 * Returns the number of decoded horizontal pixels.
-	 * This value may be greater than GetWidth() since it does not take any
+	 * This value may be different from GetWidth() since it does not take any
 	 * cropping values into account.
+	 * Note that this represents the actual dimension of the data returned, which
+	 * may differ from GetWidth() + cropping offsets!
+	 * (e.g. if the data is returned as RGBA pixels, but represents a YUYV 4:2:2 format)
 	 */
 	virtual int32 GetDecodedWidth() const = 0;
 
@@ -87,6 +90,9 @@ public:
 	 * Returns the number of decoded vertical pixels.
 	 * This value may be greater than GetHeight() since it does not take any
 	 * cropping values into account.
+	 * Note that this represents the actual dimension of the data returned, which
+	 * may differ from GetWidth() + cropping offsets!
+	 * (e.g. if the data is returned as RGBA pixels, but represents a YUYV 4:2:2 format)
 	 */
 	virtual int32 GetDecodedHeight() const = 0;
 	
@@ -97,18 +103,6 @@ public:
 	 */
 	virtual FElectraVideoDecoderOutputCropValues GetCropValues() const = 0;
 
-	/**
-	 * Returns the horizontal size of the frame buffer.
-	 * This includes any required horizontal pitch.
-	 */
-	virtual int32 GetFrameWidth() const = 0;
-
-	/**
-	 * Returns the vertical size of the frame buffer.
-	 * This includes any required vertical pitch.
-	 */
-	virtual int32 GetFrameHeight() const = 0;
-
 	/**
 	 * Returns the aspect ratio w value.
 	 */
@@ -142,13 +136,6 @@ public:
 	 */
 	virtual void GetExtraValues(TMap<FString, FVariant>& OutExtraValues) const = 0;
 
-	/**
-	 * Returns the internal pixel format of the decoded image.
-	 * This is a value particular to the target hardware.
-	 * It usually describes the internal texture format.
-	 */
-	virtual int32 GetPixelFormat() const = 0;
-
 	/**
 	 * Returns the platform specific decoder output handle, if any.
 	 * This value is specific to the decoder being used.
diff --git a/Engine/Plugins/Media/ElectraCodecs/Source/ElectraDecoders/Public/Windows/WindowsElectraDecoderPlatformOutputHandleTypes.h b/Engine/Plugins/Media/ElectraCodecs/Source/ElectraDecoders/Public/Windows/WindowsElectraDecoderPlatformOutputHandleTypes.h
--- a/Engine/Plugins/Media/ElectraCodecs/Source/ElectraDecoders/Public/Windows/WindowsElectraDecoderPlatformOutputHandleTypes.h
+++ b/Engine/Plugins/Media/ElectraCodecs/Source/ElectraDecoders/Public/Windows/WindowsElectraDecoderPlatformOutputHandleTypes.h
@@ -9,22 +9,61 @@ enum class EElectraDecoderPlatformOutputHandleType
 	MFSample,
 	DXDevice,
 	DXDeviceContext,
-	DXTexture,
 	ImageBuffers			// IElectraDecoderVideoOutputImageBuffers interface
 };
 
+enum class EElectraDecoderPlatformPixelFormat
+{
+	INVALID = 0,
+	R8G8B8A8,
+	A8R8G8B8,
+	B8G8R8A8,
+	R16G16B16A16,
+	A16B16G16R16,
+	A32B32G32R32F,
+	A2B10G10R10,
+	DXT1,
+	DXT5,
+	BC4,
+	NV12,
+	P010,
+};
+
+enum class EElectraDecoderPlatformPixelEncoding
+{
+	Native = 0,		//!< Pixel formats native representation
+	RGB,			//!< Interpret as RGB
+	RGBA,			//!< Interpret as RGBA
+	YCbCr,			//!< Interpret as YCbCR
+	YCbCr_Alpha,	//!< Interpret as YCbCR with alpha
+	YCoCg,			//!< Interpret as scaled YCoCg
+	YCoCg_Alpha,	//!< Interpret as scaled YCoCg with trailing BC4 alpha data
+	CbY0CrY1,		//!< Interpret as CbY0CrY1
+	Y0CbY1Cr,		//!< Interpret as Y0CbY1Cr
+	ARGB_BigEndian,	//!< Interpret as ARGB, big endian
+};
+
 class IElectraDecoderVideoOutputImageBuffers
 {
 public:
-	// Return the 4cc of the codec. This determines how the rest of the information must be interpreted.
+	// Return the 4cc of the codec.
 	virtual uint32 GetCodec4CC() const = 0;
 
 	// Returns the number of separate image buffers making up the frame.
 	virtual int32 GetNumberOfBuffers() const = 0;
 
-	// Returns the n'th image buffer.
-	virtual TSharedPtr<TArray<uint8>, ESPMode::ThreadSafe> GetBufferByIndex(int32 InBufferIndex) const = 0;
+	// Returns the n'th image buffer as CPU data buffer.
+	virtual TSharedPtr<TArray<uint8>, ESPMode::ThreadSafe> GetBufferDataByIndex(int32 InBufferIndex) const = 0;
+
+	// Returns the n'th image buffer as GPU texture resource.
+	virtual void* GetBufferTextureByIndex(int32 InBufferIndex) const = 0;
+
+	// Returns the n'th image buffer format
+	virtual EElectraDecoderPlatformPixelFormat GetBufferFormatByIndex(int32 InBufferIndex) const = 0;
+
+	// Returns the n'th image buffer encoding
+	virtual EElectraDecoderPlatformPixelEncoding GetBufferEncodingByIndex(int32 InBufferIndex) const = 0;
 
-	// Returns the n'th image buffer format, which is a specific value of the codec.
-	virtual uint64 GetBufferFormatByIndex(int32 InBufferIndex) const = 0;
+	// Returns the n'th image buffer pitch
+	virtual int32 GetBufferPitchByIndex(int32 InBufferIndex) const = 0;
 };
diff --git a/Engine/Plugins/Media/ElectraPlayer/Source/ElectraPlayerRuntime/Private/Runtime/Decoder/Android/AndroidElectraDecoderResourceManager.cpp b/Engine/Plugins/Media/ElectraPlayer/Source/ElectraPlayerRuntime/Private/Runtime/Decoder/Android/AndroidElectraDecoderResourceManager.cpp
--- a/Engine/Plugins/Media/ElectraPlayer/Source/ElectraPlayerRuntime/Private/Runtime/Decoder/Android/AndroidElectraDecoderResourceManager.cpp
+++ b/Engine/Plugins/Media/ElectraPlayer/Source/ElectraPlayerRuntime/Private/Runtime/Decoder/Android/AndroidElectraDecoderResourceManager.cpp
@@ -190,7 +190,6 @@ bool FElectraDecoderResourceManagerAndroid::SetupRenderBufferFromDecoderOutput(I
 		FElectraVideoDecoderOutputCropValues Crop = InDecoderOutput->GetCropValues();
 		InOutBufferPropertes->Set(TEXT("width"), FVariantValue((int64) InDecoderOutput->GetWidth()));
 		InOutBufferPropertes->Set(TEXT("height"), FVariantValue((int64) InDecoderOutput->GetHeight()));
-		InOutBufferPropertes->Set(TEXT("pitch"), FVariantValue((int64) InDecoderOutput->GetFrameWidth()));
 		InOutBufferPropertes->Set(TEXT("crop_left"), FVariantValue((int64) Crop.Left));
 		InOutBufferPropertes->Set(TEXT("crop_right"), FVariantValue((int64) Crop.Right));
 		InOutBufferPropertes->Set(TEXT("crop_top"), FVariantValue((int64) Crop.Top));
@@ -200,11 +199,12 @@ bool FElectraDecoderResourceManagerAndroid::SetupRenderBufferFromDecoderOutput(I
 		InOutBufferPropertes->Set(TEXT("aspect_h"), FVariantValue((int64) InDecoderOutput->GetAspectRatioH()));
 		InOutBufferPropertes->Set(TEXT("fps_num"), FVariantValue((int64) InDecoderOutput->GetFrameRateNumerator()));
 		InOutBufferPropertes->Set(TEXT("fps_denom"), FVariantValue((int64) InDecoderOutput->GetFrameRateDenominator()));
-		check(InDecoderOutput->GetPixelFormat() == 0);
 		int32 NumBits = InDecoderOutput->GetNumberOfBits();
 		InOutBufferPropertes->Set(TEXT("pixelfmt"), FVariantValue((int64)((NumBits > 8)  ? EPixelFormat::PF_A2B10G10R10 : EPixelFormat::PF_B8G8R8A8)));
 		InOutBufferPropertes->Set(TEXT("bits_per"), FVariantValue((int64) NumBits));
 
+		InOutBufferPropertes->Set(TEXT("pitch"), FVariantValue((int64)InDecoderOutput->GetDecodedWidth() * ((NumBits > 8) ? 2 : 1)));
+
 		FElectraDecoderVideoOutputCopyResources cr;
 		IElectraDecoderVideoOutput::EImageCopyResult CopyResult = InDecoderOutput->CopyPlatformImage(&cr);
 		if (CopyResult != IElectraDecoderVideoOutput::EImageCopyResult::Succeeded)
diff --git a/Engine/Plugins/Media/ElectraPlayer/Source/ElectraPlayerRuntime/Private/Runtime/Decoder/Apple/AppleElectraDecoderResourceManager.cpp b/Engine/Plugins/Media/ElectraPlayer/Source/ElectraPlayerRuntime/Private/Runtime/Decoder/Apple/AppleElectraDecoderResourceManager.cpp
--- a/Engine/Plugins/Media/ElectraPlayer/Source/ElectraPlayerRuntime/Private/Runtime/Decoder/Apple/AppleElectraDecoderResourceManager.cpp
+++ b/Engine/Plugins/Media/ElectraPlayer/Source/ElectraPlayerRuntime/Private/Runtime/Decoder/Apple/AppleElectraDecoderResourceManager.cpp
@@ -19,32 +19,6 @@ namespace Electra
 
 namespace AppleDecoderResources
 {
-	static constexpr uint32 Make4CC(const uint8 A, const uint8 B, const uint8 C, const uint8 D)
-	{
-		return (static_cast<uint32>(A) << 24) | (static_cast<uint32>(B) << 16) | (static_cast<uint32>(C) << 8) | static_cast<uint32>(D);
-	}
-
-	static bool IsHAP(uint32 Codec4CC)
-	{
-		return Codec4CC == AppleDecoderResources::Make4CC('H','a','p','1') ||
-			   Codec4CC == AppleDecoderResources::Make4CC('H','a','p','5') ||
-			   Codec4CC == AppleDecoderResources::Make4CC('H','a','p','Y') ||
-			   Codec4CC == AppleDecoderResources::Make4CC('H','a','p','M') ||
-			   Codec4CC == AppleDecoderResources::Make4CC('H','a','p','A') ||
-			   Codec4CC == AppleDecoderResources::Make4CC('H','a','p','7') ||
-			   Codec4CC == AppleDecoderResources::Make4CC('H','a','p','H');
-	}
-
-	static bool IsProRes(uint32 Codec4CC)
-	{
-		return Codec4CC == AppleDecoderResources::Make4CC('a','p','c','h') ||
-			   Codec4CC == AppleDecoderResources::Make4CC('a','p','c','n') ||
-			   Codec4CC == AppleDecoderResources::Make4CC('a','p','c','s') ||
-			   Codec4CC == AppleDecoderResources::Make4CC('a','p','c','o') ||
-			   Codec4CC == AppleDecoderResources::Make4CC('a','p','4','h') ||
-			   Codec4CC == AppleDecoderResources::Make4CC('a','p','4','x');
-	}
-
 	static FElectraDecoderResourceManagerApple::FCallbacks Callbacks;
 }
 
@@ -73,7 +47,6 @@ FElectraDecoderResourceManagerApple::~FElectraDecoderResourceManagerApple()
 class FElectraDecoderResourceManagerApple::FInstanceVars : public FElectraDecoderResourceManagerApple::IDecoderPlatformResource
 {
 public:
-	bool (*SetupRenderBufferFromDecoderOutput)(IMediaRenderer::IBuffer* /*InOutBufferToSetup*/, TSharedPtr<FParamDict, ESPMode::ThreadSafe> /*InOutBufferPropertes*/, TSharedPtr<IElectraDecoderVideoOutput, ESPMode::ThreadSafe> /*InDecoderOutput*/, IDecoderPlatformResource* /*InPlatformSpecificResource*/) = nullptr;
 	uint32 Codec4CC = 0;
 };
 
@@ -82,20 +55,11 @@ public:
 IElectraDecoderResourceDelegateApple::IDecoderPlatformResource* FElectraDecoderResourceManagerApple::CreatePlatformResource(void* InOwnerHandle, EDecoderPlatformResourceType InDecoderResourceType, const TMap<FString, FVariant> InOptions)
 {
 	FInstanceVars* Vars = new FInstanceVars;
-
-	/*
-		We can't be as codec agnostic here as we would like to be.
-		Image decoders return specific texture formats in a CPU buffer that one must be aware of.
-	*/
-	uint32 Codec4CC = (uint32)ElectraDecodersUtil::GetVariantValueSafeU64(InOptions, TEXT("codec_4cc"), 0);
-	if (AppleDecoderResources::IsHAP(Codec4CC))
+	check(Vars);
+	if (Vars)
 	{
+		uint32 Codec4CC = (uint32)ElectraDecodersUtil::GetVariantValueSafeU64(InOptions, TEXT("codec_4cc"), 0);
 		Vars->Codec4CC = Codec4CC;
-		Vars->SetupRenderBufferFromDecoderOutput = FElectraDecoderResourceManagerApple::SetupRenderBufferFromDecoderOutput_Images;
-	}
-	else
-	{
-		Vars->SetupRenderBufferFromDecoderOutput = FElectraDecoderResourceManagerApple::SetupRenderBufferFromDecoderOutput_Default;
 	}
 	return Vars;
 }
@@ -121,169 +85,119 @@ bool FElectraDecoderResourceManagerApple::GetMetalDevice(void **OutMetalDevice)
 }
 
 
-bool FElectraDecoderResourceManagerApple::SetupRenderBufferFromDecoderOutput(IMediaRenderer::IBuffer* InOutBufferToSetup, TSharedPtr<FParamDict, ESPMode::ThreadSafe> InOutBufferPropertes, TSharedPtr<IElectraDecoderVideoOutput, ESPMode::ThreadSafe> InDecoderOutput, IDecoderPlatformResource* InPlatformSpecificResource)
+bool FElectraDecoderResourceManagerApple::SetupRenderBufferFromDecoderOutput(IMediaRenderer::IBuffer* InOutBufferToSetup, TSharedPtr<FParamDict, ESPMode::ThreadSafe> InOutBufferPropertes, TSharedPtr<IElectraDecoderVideoOutput, ESPMode::ThreadSafe> InDecoderOutput, FElectraDecoderResourceManagerApple::IDecoderPlatformResource* InPlatformSpecificResource)
 {
 	check(InOutBufferToSetup);
 	check(InOutBufferPropertes.IsValid());
 	check(InDecoderOutput.IsValid());
 
-	if (!InPlatformSpecificResource)
-	{
-		return SetupRenderBufferFromDecoderOutput_Default(InOutBufferToSetup, InOutBufferPropertes, InDecoderOutput, InPlatformSpecificResource);
-	}
-	else
-	{
-		FInstanceVars* Vars = static_cast<FInstanceVars*>(InPlatformSpecificResource);
-		return Vars->SetupRenderBufferFromDecoderOutput(InOutBufferToSetup, InOutBufferPropertes, InDecoderOutput, InPlatformSpecificResource);
-	}
-}
-
-bool FElectraDecoderResourceManagerApple::SetupRenderBufferFromDecoderOutput_Default(IMediaRenderer::IBuffer* InOutBufferToSetup, TSharedPtr<FParamDict, ESPMode::ThreadSafe> InOutBufferPropertes, TSharedPtr<IElectraDecoderVideoOutput, ESPMode::ThreadSafe> InDecoderOutput, IDecoderPlatformResource* InPlatformSpecificResource)
-{
 	TSharedPtr<FElectraPlayerVideoDecoderOutputApple, ESPMode::ThreadSafe> DecoderOutput = InOutBufferToSetup->GetBufferProperties().GetValue("texture").GetSharedPointer<FElectraPlayerVideoDecoderOutputApple>();
-	if (DecoderOutput.IsValid())
-	{
-		CVImageBufferRef ImageBuffer = reinterpret_cast<CVImageBufferRef>(InDecoderOutput->GetPlatformOutputHandle(EElectraDecoderPlatformOutputHandleType::ImageBufferRef));
-		if (ImageBuffer)
-		{
-			FElectraVideoDecoderOutputCropValues Crop = InDecoderOutput->GetCropValues();
-			InOutBufferPropertes->Set(TEXT("width"), FVariantValue((int64) InDecoderOutput->GetWidth()));
-			InOutBufferPropertes->Set(TEXT("height"), FVariantValue((int64) InDecoderOutput->GetHeight()));
-			InOutBufferPropertes->Set(TEXT("pitch"), FVariantValue((int64) InDecoderOutput->GetFrameWidth()));
-			InOutBufferPropertes->Set(TEXT("crop_left"), FVariantValue((int64) Crop.Left));
-			InOutBufferPropertes->Set(TEXT("crop_right"), FVariantValue((int64) Crop.Right));
-			InOutBufferPropertes->Set(TEXT("crop_top"), FVariantValue((int64) Crop.Top));
-			InOutBufferPropertes->Set(TEXT("crop_bottom"), FVariantValue((int64) Crop.Bottom));
-			InOutBufferPropertes->Set(TEXT("aspect_ratio"), FVariantValue((double) InDecoderOutput->GetAspectRatioW() / (double) InDecoderOutput->GetAspectRatioH()));
-			InOutBufferPropertes->Set(TEXT("aspect_w"), FVariantValue((int64) InDecoderOutput->GetAspectRatioW()));
-			InOutBufferPropertes->Set(TEXT("aspect_h"), FVariantValue((int64) InDecoderOutput->GetAspectRatioH()));
-			InOutBufferPropertes->Set(TEXT("fps_num"), FVariantValue((int64) InDecoderOutput->GetFrameRateNumerator()));
-			InOutBufferPropertes->Set(TEXT("fps_denom"), FVariantValue((int64) InDecoderOutput->GetFrameRateDenominator()));
-			int32 PixelFormat =	InDecoderOutput->GetPixelFormat();
-			if (PixelFormat == kCVPixelFormatType_4444AYpCbCr16)
-			{
-				InOutBufferPropertes->Set(TEXT("pixelfmt"), FVariantValue((int64)PF_R16G16B16A16_UNORM));
-			}
-			else if ((PixelFormat == kCVPixelFormatType_420YpCbCr8BiPlanarVideoRange) || (PixelFormat == kCVPixelFormatType_420YpCbCr8BiPlanarFullRange))
-			{
-				InOutBufferPropertes->Set(TEXT("pixelfmt"), FVariantValue((int64)PF_NV12));
-			}
-			else
-			{
-				InOutBufferPropertes->Set(TEXT("pixelfmt"), FVariantValue((int64)PF_P010));
-			}
-			InOutBufferPropertes->Set(TEXT("bits_per"), FVariantValue((int64) InDecoderOutput->GetNumberOfBits()));
-#if 0
-			// What type of decoder output do we have here?
-			TMap<FString, FVariant> ExtraValues;
-			InDecoderOutput->GetExtraValues(ExtraValues);
-			if (ElectraDecodersUtil::GetVariantValueFString(ExtraValues, TEXT("codec")).Equals(TEXT("prores")))
-			{
-				uint32 Codec4CC = (uint32)ElectraDecodersUtil::GetVariantValueSafeI64(ExtraValues, TEXT("codec_4cc"), 0);
-				// ...
-			}
-#endif
-			DecoderOutput->Initialize(ImageBuffer, InOutBufferPropertes);
-			return true;
-		}
-	}
-	return false;
-}
-
-
-bool FElectraDecoderResourceManagerApple::SetupRenderBufferFromDecoderOutput_Images(IMediaRenderer::IBuffer* InOutBufferToSetup, TSharedPtr<FParamDict, ESPMode::ThreadSafe> InOutBufferPropertes, TSharedPtr<IElectraDecoderVideoOutput, ESPMode::ThreadSafe> InDecoderOutput, IDecoderPlatformResource* InPlatformSpecificResource)
-{
-	TSharedPtr<FElectraPlayerVideoDecoderOutputApple, ESPMode::ThreadSafe> DecoderOutput = InOutBufferToSetup->GetBufferProperties().GetValue("texture").GetSharedPointer<FElectraPlayerVideoDecoderOutputApple>();
-	IElectraDecoderVideoOutputImageBuffers* ImageBuffers = reinterpret_cast<IElectraDecoderVideoOutputImageBuffers*>(InDecoderOutput->GetPlatformOutputHandle(EElectraDecoderPlatformOutputHandleType::ImageBuffers));
 	FInstanceVars* Vars = static_cast<FInstanceVars*>(InPlatformSpecificResource);
-	if (DecoderOutput.IsValid() && Vars)
+
+	if (DecoderOutput.IsValid())
 	{
 		FElectraVideoDecoderOutputCropValues Crop = InDecoderOutput->GetCropValues();
-		InOutBufferPropertes->Set(TEXT("width"), FVariantValue((int64) InDecoderOutput->GetWidth()));
-		InOutBufferPropertes->Set(TEXT("height"), FVariantValue((int64) InDecoderOutput->GetHeight()));
-		InOutBufferPropertes->Set(TEXT("pitch"), FVariantValue((int64) InDecoderOutput->GetFrameWidth()));
-		InOutBufferPropertes->Set(TEXT("crop_left"), FVariantValue((int64) Crop.Left));
-		InOutBufferPropertes->Set(TEXT("crop_right"), FVariantValue((int64) Crop.Right));
-		InOutBufferPropertes->Set(TEXT("crop_top"), FVariantValue((int64) Crop.Top));
-		InOutBufferPropertes->Set(TEXT("crop_bottom"), FVariantValue((int64) Crop.Bottom));
-		InOutBufferPropertes->Set(TEXT("aspect_ratio"), FVariantValue((double) InDecoderOutput->GetAspectRatioW() / (double) InDecoderOutput->GetAspectRatioH()));
-		InOutBufferPropertes->Set(TEXT("aspect_w"), FVariantValue((int64) InDecoderOutput->GetAspectRatioW()));
-		InOutBufferPropertes->Set(TEXT("aspect_h"), FVariantValue((int64) InDecoderOutput->GetAspectRatioH()));
-		InOutBufferPropertes->Set(TEXT("fps_num"), FVariantValue((int64) InDecoderOutput->GetFrameRateNumerator()));
-		InOutBufferPropertes->Set(TEXT("fps_denom"), FVariantValue((int64) InDecoderOutput->GetFrameRateDenominator()));
+		InOutBufferPropertes->Set(TEXT("width"), FVariantValue((int64)InDecoderOutput->GetWidth()));
+		InOutBufferPropertes->Set(TEXT("height"), FVariantValue((int64)InDecoderOutput->GetHeight()));
+		InOutBufferPropertes->Set(TEXT("crop_left"), FVariantValue((int64)Crop.Left));
+		InOutBufferPropertes->Set(TEXT("crop_right"), FVariantValue((int64)Crop.Right));
+		InOutBufferPropertes->Set(TEXT("crop_top"), FVariantValue((int64)Crop.Top));
+		InOutBufferPropertes->Set(TEXT("crop_bottom"), FVariantValue((int64)Crop.Bottom));
+		InOutBufferPropertes->Set(TEXT("aspect_ratio"), FVariantValue((double)InDecoderOutput->GetAspectRatioW() / (double)InDecoderOutput->GetAspectRatioH()));
+		InOutBufferPropertes->Set(TEXT("aspect_w"), FVariantValue((int64)InDecoderOutput->GetAspectRatioW()));
+		InOutBufferPropertes->Set(TEXT("aspect_h"), FVariantValue((int64)InDecoderOutput->GetAspectRatioH()));
+		InOutBufferPropertes->Set(TEXT("fps_num"), FVariantValue((int64)InDecoderOutput->GetFrameRateNumerator()));
+		InOutBufferPropertes->Set(TEXT("fps_denom"), FVariantValue((int64)InDecoderOutput->GetFrameRateDenominator()));
 
 		int32 Width = InDecoderOutput->GetDecodedWidth();
 		int32 Height = InDecoderOutput->GetDecodedHeight();
-		int32 Pitch = InDecoderOutput->GetFrameWidth();
 
-		if (AppleDecoderResources::IsHAP(Vars->Codec4CC))
-		{
-			if (!ImageBuffers)
-			{
-				UE_LOG(LogElectraPlayer, Log, TEXT("Did not receive image buffers for HAP output"));
-				return false;
-			}
+		InOutBufferPropertes->Set(TEXT("bits_per"), FVariantValue((int64)InDecoderOutput->GetNumberOfBits()));
 
-			// One or two buffers (color or color+alpha)
+		//
+		// Image buffers?
+		//
+		IElectraDecoderVideoOutputImageBuffers* ImageBuffers = reinterpret_cast<IElectraDecoderVideoOutputImageBuffers*>(InDecoderOutput->GetPlatformOutputHandle(EElectraDecoderPlatformOutputHandleType::ImageBuffers));
+		if (ImageBuffers != nullptr)
+		{
 			int32 NumImageBuffers = ImageBuffers->GetNumberOfBuffers();
 			check(NumImageBuffers == 1 || NumImageBuffers == 2);
 
-			uint32 ColorBufferFormat = (uint32) ImageBuffers->GetBufferFormatByIndex(0);
-			switch(ColorBufferFormat)
+			// Color buffer
+
+			EElectraDecoderPlatformPixelFormat PixFmt = ImageBuffers->GetBufferFormatByIndex(0);
+			EElectraDecoderPlatformPixelEncoding PixEnc = ImageBuffers->GetBufferEncodingByIndex(0);
+
+			EPixelFormat RHIPixFmt;
+			switch (PixFmt)
 			{
-				case 0x01:		// HapTextureFormat_YCoCg_DXT5
-				{
-					InOutBufferPropertes->Set(TEXT("pixelfmt"), FVariantValue((int64)EPixelFormat::PF_DXT5));
-					InOutBufferPropertes->Set(TEXT("bits_per"), FVariantValue((int64) 8));
-					InOutBufferPropertes->Set(TEXT("pixelenc"), FVariantValue((int64)EVideoDecoderPixelEncoding::YCoCg));
-					Pitch = ((Pitch + 3) / 4) * 16;		// 4 pixel wide blocks with 16 bytes
-					break;
-				}
-				case 0x83F0:	// HapTextureFormat_RGB_DXT1
-				{
-					InOutBufferPropertes->Set(TEXT("pixelfmt"), FVariantValue((int64)EPixelFormat::PF_DXT1));
-					InOutBufferPropertes->Set(TEXT("bits_per"), FVariantValue((int64) 8));
-					Pitch = ((Pitch + 3) / 4) * 8;		// 4 pixel wide blocks with 8 bytes
-					break;
-				}
-				case 0x83F3:	// HapTextureFormat_RGBA_DXT5
-				{
-					InOutBufferPropertes->Set(TEXT("pixelfmt"), FVariantValue((int64)EPixelFormat::PF_DXT5));
-					InOutBufferPropertes->Set(TEXT("bits_per"), FVariantValue((int64) 8));
-					Pitch = ((Pitch + 3) / 4) * 16;		// 4 pixel wide blocks with 16 bytes
-					break;
-				}
-				case 0x8DBB:	// HapTextureFormat_A_RGTC1
-				{
-					InOutBufferPropertes->Set(TEXT("pixelfmt"), FVariantValue((int64)EPixelFormat::PF_BC4));
-					InOutBufferPropertes->Set(TEXT("bits_per"), FVariantValue((int64) 8));
-					Pitch = ((Pitch + 3) / 4) * 8;		// 4 pixel wide blocks with 8 bytes
-					break;
-				}
-				case 0x8E8C:	// HapTextureFormat_RGBA_BPTC_UNORM
-				case 0x8E8F:	// HapTextureFormat_RGB_BPTC_UNSIGNED_FLOAT
-				case 0x8E8E:	// HapTextureFormat_RGB_BPTC_SIGNED_FLOAT
-				default:
-				{
-					UE_LOG(LogElectraPlayer, Log, TEXT("Unsupported HAP texture format 0x%04x"), ColorBufferFormat);
-					return false;
-				}
+				case EElectraDecoderPlatformPixelFormat::R8G8B8A8:		RHIPixFmt = EPixelFormat::PF_R8G8B8A8; break;
+				case EElectraDecoderPlatformPixelFormat::A8R8G8B8:		RHIPixFmt = EPixelFormat::PF_A8R8G8B8; break;
+				case EElectraDecoderPlatformPixelFormat::B8G8R8A8:		RHIPixFmt = EPixelFormat::PF_B8G8R8A8; break;
+				case EElectraDecoderPlatformPixelFormat::R16G16B16A16:	RHIPixFmt = EPixelFormat::PF_R16G16B16A16_UNORM; break;
+				case EElectraDecoderPlatformPixelFormat::A16B16G16R16:	RHIPixFmt = EPixelFormat::PF_A16B16G16R16; break;
+				case EElectraDecoderPlatformPixelFormat::A32B32G32R32F:	RHIPixFmt = EPixelFormat::PF_A32B32G32R32F; break;
+				case EElectraDecoderPlatformPixelFormat::A2B10G10R10:	RHIPixFmt = EPixelFormat::PF_A2B10G10R10; break;
+				case EElectraDecoderPlatformPixelFormat::DXT1:			RHIPixFmt = EPixelFormat::PF_DXT1; break;
+				case EElectraDecoderPlatformPixelFormat::DXT5:			RHIPixFmt = EPixelFormat::PF_DXT5; break;
+				case EElectraDecoderPlatformPixelFormat::BC4:			RHIPixFmt = EPixelFormat::PF_BC4; break;
+				case EElectraDecoderPlatformPixelFormat::NV12:			RHIPixFmt = EPixelFormat::PF_NV12; break;
+				case EElectraDecoderPlatformPixelFormat::P010:			RHIPixFmt = EPixelFormat::PF_P010; break;
+				default: RHIPixFmt = EPixelFormat::PF_Unknown; break;
 			}
+			check(RHIPixFmt != EPixelFormat::PF_Unknown);
+
+			EVideoDecoderPixelEncoding DecPixEnc;
+			switch (PixEnc)
+			{
+				case EElectraDecoderPlatformPixelEncoding::Native:			DecPixEnc = EVideoDecoderPixelEncoding::Native; break;
+				case EElectraDecoderPlatformPixelEncoding::RGB:				DecPixEnc = EVideoDecoderPixelEncoding::RGB; break;
+				case EElectraDecoderPlatformPixelEncoding::RGBA:			DecPixEnc = EVideoDecoderPixelEncoding::RGBA; break;
+				case EElectraDecoderPlatformPixelEncoding::YCbCr:			DecPixEnc = EVideoDecoderPixelEncoding::YCbCr; break;
+				case EElectraDecoderPlatformPixelEncoding::YCbCr_Alpha:		DecPixEnc = EVideoDecoderPixelEncoding::YCbCr_Alpha; break;
+				case EElectraDecoderPlatformPixelEncoding::YCoCg:			DecPixEnc = EVideoDecoderPixelEncoding::YCoCg; break;
+				case EElectraDecoderPlatformPixelEncoding::YCoCg_Alpha:		DecPixEnc = EVideoDecoderPixelEncoding::YCoCg_Alpha; break;
+				case EElectraDecoderPlatformPixelEncoding::CbY0CrY1:		DecPixEnc = EVideoDecoderPixelEncoding::CbY0CrY1; break;
+				case EElectraDecoderPlatformPixelEncoding::Y0CbY1Cr:		DecPixEnc = EVideoDecoderPixelEncoding::Y0CbY1Cr; break;
+				case EElectraDecoderPlatformPixelEncoding::ARGB_BigEndian:	DecPixEnc = EVideoDecoderPixelEncoding::ARGB_BigEndian; break;
+				default: DecPixEnc = DecPixEnc = EVideoDecoderPixelEncoding::Native; break;
+			}
+
+			int32 Pitch = ImageBuffers->GetBufferPitchByIndex(0);
+
+			InOutBufferPropertes->Set(TEXT("pixelfmt"), FVariantValue((int64)RHIPixFmt));
+			InOutBufferPropertes->Set(TEXT("pixelenc"), FVariantValue((int64)DecPixEnc));
+			InOutBufferPropertes->Set(TEXT("pitch"), FVariantValue((int64)Pitch));
 
-			TSharedPtr<TArray<uint8>, ESPMode::ThreadSafe> ColorBuffer = ImageBuffers->GetBufferByIndex(0);
+// [...] ALPHA BUFFER -- HOW DO WE PASS IT ON!?
+// [...] ANY CS/HDR INFO FROM THE DECODER? -- PASSING IT ON WOULD BE EASY, BUT RIGHT NOW WE ALWAYS READ IT FROM THE CONTAINER (in code further up the chain)
+
+			TSharedPtr<TArray<uint8>, ESPMode::ThreadSafe> ColorBuffer = ImageBuffers->GetBufferDataByIndex(0);
 			if (ColorBuffer.IsValid() && ColorBuffer->Num())
 			{
+				//
+				// CPU side buffer
+				//
 				DecoderOutput->InitializeWithBuffer(ColorBuffer,
 					Pitch,							// Buffer stride
 					FIntPoint(Width, Height),		// Buffer dimensions
 					InOutBufferPropertes);
 				return true;
 			}
-		}
-		else
-		{
-			UE_LOG(LogElectraPlayer, Log, TEXT("Unsupported decoded 4CC 0x%08x"), Vars->Codec4CC);
-			return false;
+			else
+			{
+				// Note: we assume a DX11 texture at this point, but this could also be interpreted as IUknown and then either DX11 or DX12 resources being derived from it...
+				CVImageBufferRef ImageBuffer = static_cast<CVImageBufferRef>(ImageBuffers->GetBufferTextureByIndex(0));
+				if (ImageBuffer != nullptr)
+				{
+					//
+					// GPU texture (CVImageBuffer)
+					//
+					DecoderOutput->Initialize(ImageBuffer, InOutBufferPropertes);
+					return true;
+				}
+			}
+
 		}
 	}
 	return false;
diff --git a/Engine/Plugins/Media/ElectraPlayer/Source/ElectraPlayerRuntime/Private/Runtime/Decoder/Apple/AppleElectraDecoderResourceManager.h b/Engine/Plugins/Media/ElectraPlayer/Source/ElectraPlayerRuntime/Private/Runtime/Decoder/Apple/AppleElectraDecoderResourceManager.h
--- a/Engine/Plugins/Media/ElectraPlayer/Source/ElectraPlayerRuntime/Private/Runtime/Decoder/Apple/AppleElectraDecoderResourceManager.h
+++ b/Engine/Plugins/Media/ElectraPlayer/Source/ElectraPlayerRuntime/Private/Runtime/Decoder/Apple/AppleElectraDecoderResourceManager.h
@@ -40,8 +40,6 @@ public:
 
 private:
 	class FInstanceVars;
-	static bool SetupRenderBufferFromDecoderOutput_Default(IMediaRenderer::IBuffer* InOutBufferToSetup, TSharedPtr<FParamDict, ESPMode::ThreadSafe> InOutBufferPropertes, TSharedPtr<IElectraDecoderVideoOutput, ESPMode::ThreadSafe> InDecoderOutput, IDecoderPlatformResource* InPlatformSpecificResource);
-	static bool SetupRenderBufferFromDecoderOutput_Images(IMediaRenderer::IBuffer* InOutBufferToSetup, TSharedPtr<FParamDict, ESPMode::ThreadSafe> InOutBufferPropertes, TSharedPtr<IElectraDecoderVideoOutput, ESPMode::ThreadSafe> InDecoderOutput, IDecoderPlatformResource* InPlatformSpecificResource);
 };
 
 using FPlatformElectraDecoderResourceManager = FElectraDecoderResourceManagerApple;
diff --git a/Engine/Plugins/Media/ElectraPlayer/Source/ElectraPlayerRuntime/Private/Runtime/Decoder/Linux/LinuxElectraDecoderResourceManager.cpp b/Engine/Plugins/Media/ElectraPlayer/Source/ElectraPlayerRuntime/Private/Runtime/Decoder/Linux/LinuxElectraDecoderResourceManager.cpp
--- a/Engine/Plugins/Media/ElectraPlayer/Source/ElectraPlayerRuntime/Private/Runtime/Decoder/Linux/LinuxElectraDecoderResourceManager.cpp
+++ b/Engine/Plugins/Media/ElectraPlayer/Source/ElectraPlayerRuntime/Private/Runtime/Decoder/Linux/LinuxElectraDecoderResourceManager.cpp
@@ -268,7 +268,6 @@ bool FElectraDecoderResourceManagerLinux::SetupRenderBufferFromDecoderOutput(IMe
 			FElectraVideoDecoderOutputCropValues Crop = InDecoderOutput->GetCropValues();
 			InOutBufferPropertes->Set(TEXT("width"), FVariantValue((int64) InDecoderOutput->GetWidth()));
 			InOutBufferPropertes->Set(TEXT("height"), FVariantValue((int64) InDecoderOutput->GetHeight()));
-			InOutBufferPropertes->Set(TEXT("pitch"), FVariantValue((int64) InDecoderOutput->GetFrameWidth()));
 			InOutBufferPropertes->Set(TEXT("crop_left"), FVariantValue((int64) Crop.Left));
 			InOutBufferPropertes->Set(TEXT("crop_right"), FVariantValue((int64) Crop.Right));
 			InOutBufferPropertes->Set(TEXT("crop_top"), FVariantValue((int64) Crop.Top));
@@ -286,6 +285,8 @@ bool FElectraDecoderResourceManagerLinux::SetupRenderBufferFromDecoderOutput(IMe
 			InOutBufferPropertes->Set(TEXT("bits_per"), FVariantValue((int64) num_bits));
 			InOutBufferPropertes->Set(TEXT("pixelfmt"), FVariantValue((int64)((num_bits <= 8) ? EPixelFormat::PF_NV12 : EPixelFormat::PF_P010)));
 
+			InOutBufferPropertes->Set(TEXT("pitch"), FVariantValue((int64)InDecoderOutput->GetDecodedWidth() * ((num_bits <= 8) ? 1 : 2)));
+
 #if WITH_LIBAV
 			const ILibavDecoderVideoCommon::FOutputInfo& DecodedImageInfo = DecodedImage->GetOutputInfo();
 			FIntPoint BufferDim(DecodedImageInfo.Planes[0].Width, DecodedImageInfo.Planes[0].Height);
diff --git a/Engine/Plugins/Media/ElectraPlayer/Source/ElectraPlayerRuntime/Private/Runtime/Decoder/Windows/WindowsElectraDecoderResourceManager.cpp b/Engine/Plugins/Media/ElectraPlayer/Source/ElectraPlayerRuntime/Private/Runtime/Decoder/Windows/WindowsElectraDecoderResourceManager.cpp
--- a/Engine/Plugins/Media/ElectraPlayer/Source/ElectraPlayerRuntime/Private/Runtime/Decoder/Windows/WindowsElectraDecoderResourceManager.cpp
+++ b/Engine/Plugins/Media/ElectraPlayer/Source/ElectraPlayerRuntime/Private/Runtime/Decoder/Windows/WindowsElectraDecoderResourceManager.cpp
@@ -20,37 +20,6 @@ namespace Electra
 
 namespace WindowsDecoderResources
 {
-	static constexpr uint32 Make4CC(const uint8 A, const uint8 B, const uint8 C, const uint8 D)
-	{
-		return (static_cast<uint32>(A) << 24) | (static_cast<uint32>(B) << 16) | (static_cast<uint32>(C) << 8) | static_cast<uint32>(D);
-	}
-
-	static bool IsHAP(uint32 Codec4CC)
-	{
-		return Codec4CC == WindowsDecoderResources::Make4CC('H','a','p','1') ||
-			   Codec4CC == WindowsDecoderResources::Make4CC('H','a','p','5') ||
-			   Codec4CC == WindowsDecoderResources::Make4CC('H','a','p','Y') ||
-			   Codec4CC == WindowsDecoderResources::Make4CC('H','a','p','M') ||
-			   Codec4CC == WindowsDecoderResources::Make4CC('H','a','p','A') ||
-			   Codec4CC == WindowsDecoderResources::Make4CC('H','a','p','7') ||
-			   Codec4CC == WindowsDecoderResources::Make4CC('H','a','p','H');
-	}
-
-	static bool IsProRes(uint32 Codec4CC)
-	{
-		return Codec4CC == WindowsDecoderResources::Make4CC('a','p','c','h') ||
-			   Codec4CC == WindowsDecoderResources::Make4CC('a','p','c','n') ||
-			   Codec4CC == WindowsDecoderResources::Make4CC('a','p','c','s') ||
-			   Codec4CC == WindowsDecoderResources::Make4CC('a','p','c','o') ||
-			   Codec4CC == WindowsDecoderResources::Make4CC('a','p','4','h') ||
-			   Codec4CC == WindowsDecoderResources::Make4CC('a','p','4','x');
-	}
-
-	static bool IsAvidDNxHD(uint32 Codec4CC)
-	{
-		return Codec4CC == WindowsDecoderResources::Make4CC('A', 'V', 'd', 'h');
-	}
-
 	static FElectraDecoderResourceManagerWindows::FCallbacks Callbacks;
 	static bool bDidInitializeMF = false;
 }
@@ -106,7 +75,6 @@ FElectraDecoderResourceManagerWindows::~FElectraDecoderResourceManagerWindows()
 class FElectraDecoderResourceManagerWindows::FInstanceVars : public IElectraDecoderResourceDelegateWindows::IDecoderPlatformResource
 {
 public:
-	bool (*SetupRenderBufferFromDecoderOutput)(IMediaRenderer::IBuffer* /*InOutBufferToSetup*/, TSharedPtr<FParamDict, ESPMode::ThreadSafe> /*InOutBufferPropertes*/, TSharedPtr<IElectraDecoderVideoOutput, ESPMode::ThreadSafe> /*InDecoderOutput*/, IDecoderPlatformResource* /*InPlatformSpecificResource*/) = nullptr;
 	uint32 Codec4CC = 0;
 };
 
@@ -115,24 +83,15 @@ public:
 IElectraDecoderResourceDelegateWindows::IDecoderPlatformResource* FElectraDecoderResourceManagerWindows::CreatePlatformResource(void* InOwnerHandle, EDecoderPlatformResourceType InDecoderResourceType, const TMap<FString, FVariant> InOptions)
 {
 	FInstanceVars* Vars = new FInstanceVars;
-
-	/*
-		We can't be as codec agnostic here as we would like to be.
-		Image decoders return specific texture formats in a CPU buffer that one must be aware of.
-	*/
-	uint32 Codec4CC = (uint32)ElectraDecodersUtil::GetVariantValueSafeU64(InOptions, TEXT("codec_4cc"), 0);
-	if (WindowsDecoderResources::IsHAP(Codec4CC) || WindowsDecoderResources::IsProRes(Codec4CC) || WindowsDecoderResources::IsAvidDNxHD(Codec4CC))
-	{
-		Vars->Codec4CC = Codec4CC;
-		Vars->SetupRenderBufferFromDecoderOutput = FElectraDecoderResourceManagerWindows::SetupRenderBufferFromDecoderOutput_Images;
-	}
-	else
+	check(Vars);
+	if (Vars)
 	{
-		Vars->SetupRenderBufferFromDecoderOutput = FElectraDecoderResourceManagerWindows::SetupRenderBufferFromDecoderOutput_Default;
+		Vars->Codec4CC = (uint32)ElectraDecodersUtil::GetVariantValueSafeU64(InOptions, TEXT("codec_4cc"), 0);
 	}
 	return Vars;
 }
 
+
 void FElectraDecoderResourceManagerWindows::ReleasePlatformResource(void* InOwnerHandle, IDecoderPlatformResource* InHandleToDestroy)
 {
 	FInstanceVars* Vars = static_cast<FInstanceVars*>(InHandleToDestroy);
@@ -143,23 +102,6 @@ void FElectraDecoderResourceManagerWindows::ReleasePlatformResource(void* InOwne
 }
 
 
-bool FElectraDecoderResourceManagerWindows::SetupRenderBufferFromDecoderOutput(IMediaRenderer::IBuffer* InOutBufferToSetup, TSharedPtr<FParamDict, ESPMode::ThreadSafe> InOutBufferPropertes, TSharedPtr<IElectraDecoderVideoOutput, ESPMode::ThreadSafe> InDecoderOutput, IDecoderPlatformResource* InPlatformSpecificResource)
-{
-	check(InOutBufferToSetup);
-	check(InOutBufferPropertes.IsValid());
-	check(InDecoderOutput.IsValid());
-
-	if (!InPlatformSpecificResource)
-	{
-		return SetupRenderBufferFromDecoderOutput_Default(InOutBufferToSetup, InOutBufferPropertes, InDecoderOutput, InPlatformSpecificResource);
-	}
-	else
-	{
-		FInstanceVars* Vars = static_cast<FInstanceVars*>(InPlatformSpecificResource);
-		return Vars->SetupRenderBufferFromDecoderOutput(InOutBufferToSetup, InOutBufferPropertes, InDecoderOutput, InPlatformSpecificResource);
-	}
-}
-
 bool FElectraDecoderResourceManagerWindows::GetD3DDevice(void **OutD3DDevice, int32* OutD3DVersionTimes1000)
 {
 	check(OutD3DDevice);
@@ -172,10 +114,7 @@ bool FElectraDecoderResourceManagerWindows::GetD3DDevice(void **OutD3DDevice, in
 }
 
 
-
-
-
-bool FElectraDecoderResourceManagerWindows::SetupRenderBufferFromDecoderOutput_Default(IMediaRenderer::IBuffer* InOutBufferToSetup, TSharedPtr<FParamDict, ESPMode::ThreadSafe> InOutBufferPropertes, TSharedPtr<IElectraDecoderVideoOutput, ESPMode::ThreadSafe> InDecoderOutput, FElectraDecoderResourceManagerWindows::IDecoderPlatformResource* InPlatformSpecificResource)
+bool FElectraDecoderResourceManagerWindows::SetupRenderBufferFromDecoderOutputFromMFSample(IMediaRenderer::IBuffer* InOutBufferToSetup, TSharedPtr<FParamDict, ESPMode::ThreadSafe> InOutBufferPropertes, TSharedPtr<IElectraDecoderVideoOutput, ESPMode::ThreadSafe> InDecoderOutput, FElectraDecoderResourceManagerWindows::IDecoderPlatformResource* InPlatformSpecificResource)
 {
 	TSharedPtr<FElectraPlayerVideoDecoderOutputPC, ESPMode::ThreadSafe> DecoderOutput = InOutBufferToSetup->GetBufferProperties().GetValue("texture").GetSharedPointer<FElectraPlayerVideoDecoderOutputPC>();
 	if (DecoderOutput.IsValid())
@@ -183,7 +122,6 @@ bool FElectraDecoderResourceManagerWindows::SetupRenderBufferFromDecoderOutput_D
 		FElectraVideoDecoderOutputCropValues Crop = InDecoderOutput->GetCropValues();
 		InOutBufferPropertes->Set(TEXT("width"), FVariantValue((int64)InDecoderOutput->GetWidth()));
 		InOutBufferPropertes->Set(TEXT("height"), FVariantValue((int64)InDecoderOutput->GetHeight()));
-		InOutBufferPropertes->Set(TEXT("pitch"), FVariantValue((int64)InDecoderOutput->GetFrameWidth()));
 		InOutBufferPropertes->Set(TEXT("crop_left"), FVariantValue((int64)Crop.Left));
 		InOutBufferPropertes->Set(TEXT("crop_right"), FVariantValue((int64)Crop.Right));
 		InOutBufferPropertes->Set(TEXT("crop_top"), FVariantValue((int64)Crop.Top));
@@ -193,8 +131,6 @@ bool FElectraDecoderResourceManagerWindows::SetupRenderBufferFromDecoderOutput_D
 		InOutBufferPropertes->Set(TEXT("aspect_h"), FVariantValue((int64)InDecoderOutput->GetAspectRatioH()));
 		InOutBufferPropertes->Set(TEXT("fps_num"), FVariantValue((int64)InDecoderOutput->GetFrameRateNumerator()));
 		InOutBufferPropertes->Set(TEXT("fps_denom"), FVariantValue((int64)InDecoderOutput->GetFrameRateDenominator()));
-		InOutBufferPropertes->Set(TEXT("pixelfmt"), InDecoderOutput->GetPixelFormat() == 0 ? FVariantValue((int64)EPixelFormat::PF_NV12) : FVariantValue((int64)EPixelFormat::PF_P010));
-		InOutBufferPropertes->Set(TEXT("bits_per"), FVariantValue((int64)InDecoderOutput->GetNumberOfBits()));
 
 		// What type of decoder output do we have here?
 		TMap<FString, FVariant> ExtraValues;
@@ -207,8 +143,8 @@ bool FElectraDecoderResourceManagerWindows::SetupRenderBufferFromDecoderOutput_D
 		bool bIsSW = !!ElectraDecodersUtil::GetVariantValueSafeI64(ExtraValues, TEXT("sw"), 0);
 
 		HRESULT Result;
-		// DX12 or non-DX, but HW accelerated?
-		if ((DXVersion == 0 || DXVersion >= 12000) && !bIsSW)
+		// DX11, DX12 or non-DX & HW accelerated?
+		if ((DXVersion == 0 || DXVersion >= 11000) && !bIsSW)
 		{
 			TRefCountPtr<IMFSample> DecodedOutputSample = reinterpret_cast<IMFSample*>(InDecoderOutput->GetPlatformOutputHandle(EElectraDecoderPlatformOutputHandleType::MFSample));
 
@@ -252,64 +188,79 @@ bool FElectraDecoderResourceManagerWindows::SetupRenderBufferFromDecoderOutput_D
 				return false;
 			}
 
-			TRefCountPtr<IMF2DBuffer> Buffer2D;
-			if ((Result = Buffer->QueryInterface(__uuidof(IMF2DBuffer), (void**)Buffer2D.GetInitReference())) != S_OK)
-			{
-				UE_LOG(LogElectraPlayer, Error, TEXT("IMFDXGIBuffer::QueryInterface(IMF2DBuffer) failed with %08x"), Result);
-				return false;
-			}
+			InOutBufferPropertes->Set(TEXT("pixelfmt"), FVariantValue((int64)(TextureDesc.Format == DXGI_FORMAT_NV12 ? EPixelFormat::PF_NV12 : EPixelFormat::PF_P010)));
+			InOutBufferPropertes->Set(TEXT("bits_per"), FVariantValue((int64)(TextureDesc.Format == DXGI_FORMAT_NV12 ? 8 : 10)));
 
-			uint8* Data = nullptr;
-			LONG Pitch = 0;
-			if ((Result = Buffer2D->Lock2D(&Data, &Pitch)) != S_OK)
+			if (DXVersion == 0 || DXVersion >= 12000)
 			{
-				UE_LOG(LogElectraPlayer, Error, TEXT("IMF2DBuffer::Lock2D() failed with %08x"), Result);
-				return false;
-			}
+				//
+				// DX12 & non-DX
+				// 
+				// (access buffer for CPU use)
+				//
+				TRefCountPtr<IMF2DBuffer> Buffer2D;
+				if ((Result = Buffer->QueryInterface(__uuidof(IMF2DBuffer), (void**)Buffer2D.GetInitReference())) != S_OK)
+				{
+					UE_LOG(LogElectraPlayer, Error, TEXT("IMFDXGIBuffer::QueryInterface(IMF2DBuffer) failed with %08x"), Result);
+					return false;
+				}
 
-			DecoderOutput->InitializeWithBuffer(Data, Pitch * (TextureDesc.Height * 3 / 2),
-				Pitch,															// Buffer stride
-				FIntPoint(TextureDesc.Width, TextureDesc.Height * 3 / 2),		// Buffer dimensions
-				InOutBufferPropertes);
+				uint8* Data = nullptr;
+				LONG Pitch = 0;
+				if ((Result = Buffer2D->Lock2D(&Data, &Pitch)) != S_OK)
+				{
+					UE_LOG(LogElectraPlayer, Error, TEXT("IMF2DBuffer::Lock2D() failed with %08x"), Result);
+					return false;
+				}
 
-			if ((Result = Buffer2D->Unlock2D()) != S_OK)
-			{
-				UE_LOG(LogElectraPlayer, Error, TEXT("IMF2DBuffer::Unlock2D() failed with %08x"), Result);
-				return false;
-			}
-			return true;
-		}
-		if (DXVersion >= 11000 && !bIsSW)
-		{
-			TRefCountPtr<IMFSample> DecodedOutputSample = reinterpret_cast<IMFSample*>(InDecoderOutput->GetPlatformOutputHandle(EElectraDecoderPlatformOutputHandleType::MFSample));
+				InOutBufferPropertes->Set(TEXT("pitch"), FVariantValue((int64)Pitch));
 
-			TRefCountPtr<ID3D11Device> DxDevice = reinterpret_cast<ID3D11Device*>(InDecoderOutput->GetPlatformOutputHandle(EElectraDecoderPlatformOutputHandleType::DXDevice));
-			TRefCountPtr<ID3D11DeviceContext> DxDeviceContext = reinterpret_cast<ID3D11DeviceContext*>(InDecoderOutput->GetPlatformOutputHandle(EElectraDecoderPlatformOutputHandleType::DXDeviceContext));
+				// Get decoded with (e.g. featuring any height adjustments for CPU buffer usage of NV12 etc.)
+				int32 Width = InDecoderOutput->GetDecodedWidth();
+				int32 Height = InDecoderOutput->GetDecodedHeight();
 
-			DWORD BuffersNum = 0;
-			if ((Result = DecodedOutputSample->GetBufferCount(&BuffersNum)) != S_OK)
-			{
-				UE_LOG(LogElectraPlayer, Error, TEXT("IMFSample::GetBufferCount() failed with %08x"), Result);
-				return false;
-			}
-			if (BuffersNum != 1)
-			{
-				UE_LOG(LogElectraPlayer, Error, TEXT("IMFSample::GetBufferCount() returned %u buffers instead of 1"), Result);
-				return false;
-			}
+				DecoderOutput->InitializeWithBuffer(Data, Pitch * Height,
+					Pitch,						// Buffer stride
+					FIntPoint(Width, Height),	// Buffer dimensions
+					InOutBufferPropertes);
 
-			int32 Width = InDecoderOutput->GetWidth();
-			int32 Height = InDecoderOutput->GetHeight();
-			DecoderOutput->InitializeWithSharedTexture(DxDevice, DxDeviceContext, DecodedOutputSample, FIntPoint(Width, Height), InOutBufferPropertes);
-			if (!DecoderOutput->GetTexture())
+				if ((Result = Buffer2D->Unlock2D()) != S_OK)
+				{
+					UE_LOG(LogElectraPlayer, Error, TEXT("IMF2DBuffer::Unlock2D() failed with %08x"), Result);
+					return false;
+				}
+			}
+			else
 			{
-				UE_LOG(LogElectraPlayer, Error, TEXT("ID3D11Device::CreateTexture2D() failed!"));
-				return false;
+				//
+				// DX11
+				//
+				check(DXVersion >= 11000);
+
+				// Notes:
+				// - No need to apply a *1.5 factor to the height in this case. DX11 can access sub-resources in both NV12 & P010 to get to the CbCr data
+				// - No need to specify a pitch as this is all directly handled on the GPU side of things
+				//
+				ID3D11Device* Device = reinterpret_cast<ID3D11Device*>(InDecoderOutput->GetPlatformOutputHandle(EElectraDecoderPlatformOutputHandleType::DXDevice));
+				ID3D11DeviceContext* DeviceContext = reinterpret_cast<ID3D11DeviceContext*>(InDecoderOutput->GetPlatformOutputHandle(EElectraDecoderPlatformOutputHandleType::DXDeviceContext));
+				if (!Device || !DeviceContext)
+				{
+					return false;
+				}
+				DecoderOutput->InitializeWithSharedTexture(Device, DeviceContext, Texture2D, FIntPoint(InDecoderOutput->GetWidth(), InDecoderOutput->GetHeight()), InOutBufferPropertes);
+				if (!DecoderOutput->GetTexture())
+				{
+					UE_LOG(LogElectraPlayer, Error, TEXT("ID3D11Device::CreateTexture2D() failed!"));
+					return false;
+				}
 			}
 			return true;
 		}
 		else if (bIsSW)
 		{
+			//
+			// "Software" case
+			//
 			TRefCountPtr<IMFSample> DecodedOutputSample = reinterpret_cast<IMFSample*>(InDecoderOutput->GetPlatformOutputHandle(EElectraDecoderPlatformOutputHandleType::MFSample));
 
 			DWORD BuffersNum = 0;
@@ -323,7 +274,6 @@ bool FElectraDecoderResourceManagerWindows::SetupRenderBufferFromDecoderOutput_D
 				UE_LOG(LogElectraPlayer, Error, TEXT("IMFSample::GetBufferCount() returned %u buffers instead of 1"), Result);
 				return false;
 			}
-
 			TRefCountPtr<IMFMediaBuffer> Buffer;
 			if ((Result = DecodedOutputSample->GetBufferByIndex(0, Buffer.GetInitReference())) != S_OK)
 			{
@@ -331,31 +281,70 @@ bool FElectraDecoderResourceManagerWindows::SetupRenderBufferFromDecoderOutput_D
 				return false;
 			}
 
-			DWORD BufferSize = 0;
-			uint8* Data = nullptr;
-			if ((Result = Buffer->GetCurrentLength(&BufferSize)) != S_OK)
+			int32 Width = InDecoderOutput->GetDecodedWidth();
+			int32 Height = InDecoderOutput->GetDecodedHeight();
+
+			// With software decode we cannot query any DXGI/DirectX data types, so we query the decoder extra data...
+			EElectraDecoderPlatformPixelFormat SamplePixFmt = static_cast<EElectraDecoderPlatformPixelFormat>(ElectraDecodersUtil::GetVariantValueSafeI64(ExtraValues, TEXT("pixfmt"), (int64)EElectraDecoderPlatformPixelFormat::NV12));
+
+			EPixelFormat PixFmt = (SamplePixFmt == EElectraDecoderPlatformPixelFormat::NV12) ? EPixelFormat::PF_NV12 : EPixelFormat::PF_P010;
+
+			InOutBufferPropertes->Set(TEXT("pixelfmt"), FVariantValue((int64)PixFmt));
+			InOutBufferPropertes->Set(TEXT("bits_per"), FVariantValue((int64)((PixFmt == EPixelFormat::PF_NV12) ? 8 : 10)));
+
+			TRefCountPtr<IMF2DBuffer> Buffer2D;
+			if ((Result = Buffer->QueryInterface(__uuidof(IMF2DBuffer), (void**)Buffer2D.GetInitReference())) == S_OK)
 			{
-				UE_LOG(LogElectraPlayer, Error, TEXT("IMFMediaBuffer::GetCurrentLength() failed with %08x"), Result);
-				return false;
+				uint8* Data = nullptr;
+				LONG Pitch = 0;
+				if ((Result = Buffer2D->Lock2D(&Data, &Pitch)) != S_OK)
+				{
+					UE_LOG(LogElectraPlayer, Error, TEXT("IMF2DBuffer::Lock2D() failed with %08x"), Result);
+					return false;
+				}
+
+				InOutBufferPropertes->Set(TEXT("pitch"), FVariantValue((int64)Pitch));
+
+				DecoderOutput->InitializeWithBuffer(Data, Pitch * Height,
+					Pitch,						// Buffer stride
+					FIntPoint(Width, Height),	// Buffer dimensions
+					InOutBufferPropertes);
+
+				if ((Result = Buffer2D->Unlock2D()) != S_OK)
+				{
+					UE_LOG(LogElectraPlayer, Error, TEXT("IMF2DBuffer::Unlock2D() failed with %08x"), Result);
+					return false;
+				}
 			}
-			if ((Result = Buffer->Lock(&Data, NULL, NULL)) != S_OK)
+			else
 			{
-				UE_LOG(LogElectraPlayer, Error, TEXT("IMFMediaBuffer::Lock() failed with %08x"), Result);
-				return false;
-			}
+				DWORD BufferSize = 0;
+				uint8* Data = nullptr;
+				if ((Result = Buffer->GetCurrentLength(&BufferSize)) != S_OK)
+				{
+					UE_LOG(LogElectraPlayer, Error, TEXT("IMFMediaBuffer::GetCurrentLength() failed with %08x"), Result);
+					return false;
+				}
+				if ((Result = Buffer->Lock(&Data, NULL, NULL)) != S_OK)
+				{
+					UE_LOG(LogElectraPlayer, Error, TEXT("IMFMediaBuffer::Lock() failed with %08x"), Result);
+					return false;
+				}
 
-			int32 Width = InDecoderOutput->GetDecodedWidth();
-			int32 Height = InDecoderOutput->GetDecodedHeight();
-			int32 Pitch = InDecoderOutput->GetFrameWidth();
-			DecoderOutput->InitializeWithBuffer(Data, BufferSize,
-				Pitch,									// Buffer stride
-				FIntPoint(Width, Height * 3 / 2),		// Buffer dimensions
-				InOutBufferPropertes);
+				int32 Pitch = Width * ((PixFmt == EPixelFormat::PF_NV12) ? 1 : 2);
 
-			if ((Result = Buffer->Unlock()) != S_OK)
-			{
-				UE_LOG(LogElectraPlayer, Error, TEXT("IMFMediaBuffer::Unlock() failed with %08x"), Result);
-				return false;
+				InOutBufferPropertes->Set(TEXT("pitch"), FVariantValue((int64)Pitch));
+
+				DecoderOutput->InitializeWithBuffer(Data, BufferSize,
+					Pitch,						// Buffer stride
+					FIntPoint(Width, Height),	// Buffer dimensions
+					InOutBufferPropertes);
+
+				if ((Result = Buffer->Unlock()) != S_OK)
+				{
+					UE_LOG(LogElectraPlayer, Error, TEXT("IMFMediaBuffer::Unlock() failed with %08x"), Result);
+					return false;
+				}
 			}
 			return true;
 		}
@@ -363,17 +352,21 @@ bool FElectraDecoderResourceManagerWindows::SetupRenderBufferFromDecoderOutput_D
 	return false;
 }
 
-bool FElectraDecoderResourceManagerWindows::SetupRenderBufferFromDecoderOutput_Images(IMediaRenderer::IBuffer* InOutBufferToSetup, TSharedPtr<FParamDict, ESPMode::ThreadSafe> InOutBufferPropertes, TSharedPtr<IElectraDecoderVideoOutput, ESPMode::ThreadSafe> InDecoderOutput, FElectraDecoderResourceManagerWindows::IDecoderPlatformResource* InPlatformSpecificResource)
+
+bool FElectraDecoderResourceManagerWindows::SetupRenderBufferFromDecoderOutput(IMediaRenderer::IBuffer* InOutBufferToSetup, TSharedPtr<FParamDict, ESPMode::ThreadSafe> InOutBufferPropertes, TSharedPtr<IElectraDecoderVideoOutput, ESPMode::ThreadSafe> InDecoderOutput, FElectraDecoderResourceManagerWindows::IDecoderPlatformResource* InPlatformSpecificResource)
 {
+	check(InOutBufferToSetup);
+	check(InOutBufferPropertes.IsValid());
+	check(InDecoderOutput.IsValid());
+
 	TSharedPtr<FElectraPlayerVideoDecoderOutputPC, ESPMode::ThreadSafe> DecoderOutput = InOutBufferToSetup->GetBufferProperties().GetValue("texture").GetSharedPointer<FElectraPlayerVideoDecoderOutputPC>();
-	IElectraDecoderVideoOutputImageBuffers* ImageBuffers = reinterpret_cast<IElectraDecoderVideoOutputImageBuffers*>(InDecoderOutput->GetPlatformOutputHandle(EElectraDecoderPlatformOutputHandleType::ImageBuffers));
 	FInstanceVars* Vars = static_cast<FInstanceVars*>(InPlatformSpecificResource);
-	if (DecoderOutput.IsValid() && Vars)
+
+	if (DecoderOutput.IsValid())
 	{
 		FElectraVideoDecoderOutputCropValues Crop = InDecoderOutput->GetCropValues();
 		InOutBufferPropertes->Set(TEXT("width"), FVariantValue((int64)InDecoderOutput->GetWidth()));
 		InOutBufferPropertes->Set(TEXT("height"), FVariantValue((int64)InDecoderOutput->GetHeight()));
-		InOutBufferPropertes->Set(TEXT("pitch"), FVariantValue((int64)InDecoderOutput->GetFrameWidth()));
 		InOutBufferPropertes->Set(TEXT("crop_left"), FVariantValue((int64)Crop.Left));
 		InOutBufferPropertes->Set(TEXT("crop_right"), FVariantValue((int64)Crop.Right));
 		InOutBufferPropertes->Set(TEXT("crop_top"), FVariantValue((int64)Crop.Top));
@@ -386,330 +379,108 @@ bool FElectraDecoderResourceManagerWindows::SetupRenderBufferFromDecoderOutput_I
 
 		int32 Width = InDecoderOutput->GetDecodedWidth();
 		int32 Height = InDecoderOutput->GetDecodedHeight();
-		int32 Pitch = InDecoderOutput->GetFrameWidth();
 
-		TMap<FString, FVariant> ExtraValues;
-		InDecoderOutput->GetExtraValues(ExtraValues);
+		InOutBufferPropertes->Set(TEXT("bits_per"), FVariantValue((int64)InDecoderOutput->GetNumberOfBits()));
 
-		if (WindowsDecoderResources::IsHAP(Vars->Codec4CC))
+		//
+		// Image buffers?
+		//
+		IElectraDecoderVideoOutputImageBuffers* ImageBuffers = reinterpret_cast<IElectraDecoderVideoOutputImageBuffers*>(InDecoderOutput->GetPlatformOutputHandle(EElectraDecoderPlatformOutputHandleType::ImageBuffers));
+		if (ImageBuffers != nullptr)
 		{
-			if (!ImageBuffers)
-			{
-				UE_LOG(LogElectraPlayer, Log, TEXT("Did not receive image buffers for HAP output"));
-				return false;
-			}
-
-			// One or two buffers (color or color+alpha)
 			int32 NumImageBuffers = ImageBuffers->GetNumberOfBuffers();
 			check(NumImageBuffers == 1 || NumImageBuffers == 2);
 
-			uint32 ColorBufferFormat = (uint32) ImageBuffers->GetBufferFormatByIndex(0);
-			switch(ColorBufferFormat)
-			{
-				case 0x01:		// HapTextureFormat_YCoCg_DXT5
-				{
-					InOutBufferPropertes->Set(TEXT("pixelfmt"), FVariantValue((int64)EPixelFormat::PF_DXT5));
-					InOutBufferPropertes->Set(TEXT("bits_per"), FVariantValue((int64) 8));
-					InOutBufferPropertes->Set(TEXT("pixelenc"), FVariantValue((int64)EVideoDecoderPixelEncoding::YCoCg));
-					Pitch = ((Pitch + 3) / 4) * 16;		// 4 pixel wide blocks with 16 bytes
-					break;
-				}
-				case 0x83F0:	// HapTextureFormat_RGB_DXT1
-				{
-					InOutBufferPropertes->Set(TEXT("pixelfmt"), FVariantValue((int64)EPixelFormat::PF_DXT1));
-					InOutBufferPropertes->Set(TEXT("bits_per"), FVariantValue((int64) 8));
-					Pitch = ((Pitch + 3) / 4) * 8;		// 4 pixel wide blocks with 8 bytes
-					break;
-				}
-				case 0x83F3:	// HapTextureFormat_RGBA_DXT5
-				{
-					InOutBufferPropertes->Set(TEXT("pixelfmt"), FVariantValue((int64)EPixelFormat::PF_DXT5));
-					InOutBufferPropertes->Set(TEXT("bits_per"), FVariantValue((int64) 8));
-					Pitch = ((Pitch + 3) / 4) * 16;		// 4 pixel wide blocks with 16 bytes
-					break;
-				}
-				case 0x8DBB:	// HapTextureFormat_A_RGTC1
-				{
-					InOutBufferPropertes->Set(TEXT("pixelfmt"), FVariantValue((int64)EPixelFormat::PF_BC4));
-					InOutBufferPropertes->Set(TEXT("bits_per"), FVariantValue((int64) 8));
-					Pitch = ((Pitch + 3) / 4) * 8;		// 4 pixel wide blocks with 8 bytes
-					break;
-				}
-				case 0x8E8C:	// HapTextureFormat_RGBA_BPTC_UNORM
-				case 0x8E8F:	// HapTextureFormat_RGB_BPTC_UNSIGNED_FLOAT
-				case 0x8E8E:	// HapTextureFormat_RGB_BPTC_SIGNED_FLOAT
-				default:
-				{
-					UE_LOG(LogElectraPlayer, Log, TEXT("Unsupported HAP texture format 0x%04x"), ColorBufferFormat);
-					return false;
-				}
-			}
+			// Color buffer
 
-			TSharedPtr<TArray<uint8>, ESPMode::ThreadSafe> ColorBuffer = ImageBuffers->GetBufferByIndex(0);
-			if (ColorBuffer.IsValid() && ColorBuffer->Num())
-			{
-				DecoderOutput->InitializeWithBuffer(ColorBuffer,
-					Pitch,							// Buffer stride
-					FIntPoint(Width, Height),		// Buffer dimensions
-					InOutBufferPropertes);
-				return true;
-			}
-		}
-		else if (WindowsDecoderResources::IsProRes(Vars->Codec4CC))
-		{
-			if (!ImageBuffers)
+			EElectraDecoderPlatformPixelFormat PixFmt = ImageBuffers->GetBufferFormatByIndex(0);
+			EElectraDecoderPlatformPixelEncoding PixEnc = ImageBuffers->GetBufferEncodingByIndex(0);
+
+			EPixelFormat RHIPixFmt;
+			switch (PixFmt)
 			{
-				UE_LOG(LogElectraPlayer, Log, TEXT("Did not receive image buffers for ProRes output"));
-				return false;
+				case EElectraDecoderPlatformPixelFormat::R8G8B8A8:		RHIPixFmt = EPixelFormat::PF_R8G8B8A8; break;
+				case EElectraDecoderPlatformPixelFormat::A8R8G8B8:		RHIPixFmt = EPixelFormat::PF_A8R8G8B8; break;
+				case EElectraDecoderPlatformPixelFormat::B8G8R8A8:		RHIPixFmt = EPixelFormat::PF_B8G8R8A8; break;
+				case EElectraDecoderPlatformPixelFormat::R16G16B16A16:	RHIPixFmt = EPixelFormat::PF_R16G16B16A16_UNORM; break;
+				case EElectraDecoderPlatformPixelFormat::A16B16G16R16:	RHIPixFmt = EPixelFormat::PF_A16B16G16R16; break;
+				case EElectraDecoderPlatformPixelFormat::A32B32G32R32F:	RHIPixFmt = EPixelFormat::PF_A32B32G32R32F; break;
+				case EElectraDecoderPlatformPixelFormat::A2B10G10R10:	RHIPixFmt = EPixelFormat::PF_A2B10G10R10; break;
+				case EElectraDecoderPlatformPixelFormat::DXT1:			RHIPixFmt = EPixelFormat::PF_DXT1; break;
+				case EElectraDecoderPlatformPixelFormat::DXT5:			RHIPixFmt = EPixelFormat::PF_DXT5; break;
+				case EElectraDecoderPlatformPixelFormat::BC4:			RHIPixFmt = EPixelFormat::PF_BC4; break;
+				case EElectraDecoderPlatformPixelFormat::NV12:			RHIPixFmt = EPixelFormat::PF_NV12; break;
+				case EElectraDecoderPlatformPixelFormat::P010:			RHIPixFmt = EPixelFormat::PF_P010; break;
+				default: RHIPixFmt = EPixelFormat::PF_Unknown; break;
 			}
-			check(ImageBuffers->GetNumberOfBuffers() == 1);
-			uint32 ColorBufferFormat = (uint32) ImageBuffers->GetBufferFormatByIndex(0);
-			switch(ColorBufferFormat)
+			check(RHIPixFmt != EPixelFormat::PF_Unknown);
+
+			EVideoDecoderPixelEncoding DecPixEnc;
+			switch (PixEnc)
 			{
-				case 0x32767579:	// '2vuy'; 4:2:2   Y'CbCr  8-bit video range
-				{
-					InOutBufferPropertes->Set(TEXT("pixelfmt"), FVariantValue((int64)EPixelFormat::PF_B8G8R8A8));
-					InOutBufferPropertes->Set(TEXT("bits_per"), FVariantValue((int64)8));
-					InOutBufferPropertes->Set(TEXT("pixelenc"), FVariantValue((int64)EVideoDecoderPixelEncoding::CbY0CrY1));
-					Width = (Width + 1) / 2; // each pixel contains TWO horizontally adjacent output pixels)
-					Pitch *= 2;
-					break;
-				}
-				case 0x79343136:	// 'y416'; 4:4:4:4 AY'CbCr 16-bit little endian full range alpha, video range Y'CbCr
-				{
-					InOutBufferPropertes->Set(TEXT("pixelfmt"), FVariantValue((int64)EPixelFormat::PF_A16B16G16R16));
-					InOutBufferPropertes->Set(TEXT("bits_per"), FVariantValue((int64) 16));
-					InOutBufferPropertes->Set(TEXT("pixelenc"), FVariantValue((int64)EVideoDecoderPixelEncoding::YCbCr_Alpha));
-					Pitch *= 8;
-					break;
-				}
-				case 0x7234666c:	// 'r4fl'; 4:4:4:4 AY'CbCr 32-bit float
-				{
-					InOutBufferPropertes->Set(TEXT("pixelfmt"), FVariantValue((int64)EPixelFormat::PF_A32B32G32R32F));
-					InOutBufferPropertes->Set(TEXT("bits_per"), FVariantValue((int64)32));
-					InOutBufferPropertes->Set(TEXT("pixelenc"), FVariantValue((int64)EVideoDecoderPixelEncoding::YCbCr_Alpha));
-					Pitch *= 16;
-					break;
-				}
-				case 0x76323130:	// 'v210'; 4:2:2   Y'CbCr 10-bit video range
-				{
-					InOutBufferPropertes->Set(TEXT("pixelfmt"), FVariantValue((int64)EPixelFormat::PF_A2B10G10R10));
-					InOutBufferPropertes->Set(TEXT("bits_per"), FVariantValue((int64)10));
-					InOutBufferPropertes->Set(TEXT("pixelenc"), FVariantValue((int64)EVideoDecoderPixelEncoding::CbY0CrY1));
-					Width = 4 * ((Width + 5) / 6); // each 4 pixel contain 6 horizontally adjacent YCbCr pixels (incl. 4x 2-bit padding)
-					Pitch = Width * 4;
-					break;
-				}
-				case 0x76323136:	// 'v216'; 4:2:2   Y'CbCr 16-bit little endian video range (Cb:Y0:Cr:Y1)
-				{
-					InOutBufferPropertes->Set(TEXT("pixelfmt"), FVariantValue((int64)EPixelFormat::PF_A16B16G16R16));
-					InOutBufferPropertes->Set(TEXT("bits_per"), FVariantValue((int64)16));
-					InOutBufferPropertes->Set(TEXT("pixelenc"), FVariantValue((int64)EVideoDecoderPixelEncoding::CbY0CrY1));
-					Width = (Width + 1) / 2; // each pixel contains TWO horizontally adjacent output pixels)
-					Pitch *= 4;
-					break;
-				}
-				case 0x62363461:	// 'b64a'; 4:4:4:4 Full-range (0-65535) ARGB  16-bit big endian per component
-				{
-					InOutBufferPropertes->Set(TEXT("pixelfmt"), FVariantValue((int64)EPixelFormat::PF_A16B16G16R16));
-					InOutBufferPropertes->Set(TEXT("bits_per"), FVariantValue((int64)16));
-					InOutBufferPropertes->Set(TEXT("pixelenc"), FVariantValue((int64)EVideoDecoderPixelEncoding::ARGB_BigEndian));
-					Pitch *= 8;
-					break;
-				}
-				default:
-				{
-					UE_LOG(LogElectraPlayer, Log, TEXT("Unsupported ProRes texture format 0x%08x"), ColorBufferFormat);
-					return false;
-				}
+				case EElectraDecoderPlatformPixelEncoding::Native:			DecPixEnc = EVideoDecoderPixelEncoding::Native; break;
+				case EElectraDecoderPlatformPixelEncoding::RGB:				DecPixEnc = EVideoDecoderPixelEncoding::RGB; break;
+				case EElectraDecoderPlatformPixelEncoding::RGBA:			DecPixEnc = EVideoDecoderPixelEncoding::RGBA; break;
+				case EElectraDecoderPlatformPixelEncoding::YCbCr:			DecPixEnc = EVideoDecoderPixelEncoding::YCbCr; break;
+				case EElectraDecoderPlatformPixelEncoding::YCbCr_Alpha:		DecPixEnc = EVideoDecoderPixelEncoding::YCbCr_Alpha; break;
+				case EElectraDecoderPlatformPixelEncoding::YCoCg:			DecPixEnc = EVideoDecoderPixelEncoding::YCoCg; break;
+				case EElectraDecoderPlatformPixelEncoding::YCoCg_Alpha:		DecPixEnc = EVideoDecoderPixelEncoding::YCoCg_Alpha; break;
+				case EElectraDecoderPlatformPixelEncoding::CbY0CrY1:		DecPixEnc = EVideoDecoderPixelEncoding::CbY0CrY1; break;
+				case EElectraDecoderPlatformPixelEncoding::Y0CbY1Cr:		DecPixEnc = EVideoDecoderPixelEncoding::Y0CbY1Cr; break;
+				case EElectraDecoderPlatformPixelEncoding::ARGB_BigEndian:	DecPixEnc = EVideoDecoderPixelEncoding::ARGB_BigEndian; break;
+				default: DecPixEnc = DecPixEnc = EVideoDecoderPixelEncoding::Native; break;
 			}
+			
+			int32 Pitch = ImageBuffers->GetBufferPitchByIndex(0);
+
+			InOutBufferPropertes->Set(TEXT("pixelfmt"), FVariantValue((int64)RHIPixFmt));
+			InOutBufferPropertes->Set(TEXT("pixelenc"), FVariantValue((int64)DecPixEnc));
+			InOutBufferPropertes->Set(TEXT("pitch"), FVariantValue((int64)Pitch));
 
-			TSharedPtr<TArray<uint8>, ESPMode::ThreadSafe> ColorBuffer = ImageBuffers->GetBufferByIndex(0);
+// [...] ALPHA BUFFER -- HOW DO WE PASS IT ON!?
+// [...] ANY CS/HDR INFO FROM THE DECODER? -- PASSING IT ON WOULD BE EASY, BUT RIGHT NOW WE ALWAYS READ IT FROM THE CONTAINER (in code further up the chain)
+
+			TSharedPtr<TArray<uint8>, ESPMode::ThreadSafe> ColorBuffer = ImageBuffers->GetBufferDataByIndex(0);
 			if (ColorBuffer.IsValid() && ColorBuffer->Num())
 			{
+				//
+				// CPU side buffer
+				//
 				DecoderOutput->InitializeWithBuffer(ColorBuffer,
 					Pitch,							// Buffer stride
 					FIntPoint(Width, Height),		// Buffer dimensions
 					InOutBufferPropertes);
 				return true;
 			}
-		}
-		else if (WindowsDecoderResources::IsAvidDNxHD(Vars->Codec4CC))
-		{
-			if (!ImageBuffers)
-			{
-				UE_LOG(LogElectraPlayer, Log, TEXT("Did not receive image buffers for Avid output"));
-				return false;
-			}
-
-			int32 AvidComponentType = InDecoderOutput->GetPixelFormat();
-				/*
-					,DNX_CT_UCHAR        = 0x001   ///< 8 bit
-					,DNX_CT_USHORT_10_6  = 0x004   ///< 10 bit
-					,DNX_CT_SHORT_2_14   = 0x008   ///< Fixed point
-					,DNX_CT_SHORT        = 0x010   ///< 16 bit. Premultiplied by 257. Byte ordering is machine dependent.
-					,DNX_CT_10Bit_2_8    = 0x040   ///< 10 bit in 2_8 format. Byte ordering is fixed. This is to be used with 10-bit 4:2:2 YCbCr components.
-					,DNX_CT_V210         = 0x400   ///< Apple's V210
-					,DNX_CT_USHORT_12_4  = 0x20000 ///< 12 bit
-				*/
-			uint32 AvidComponentOrder = (uint32)ImageBuffers->GetBufferFormatByIndex(0);
-				/*
-					,DNX_CCO_YCbYCr_NoA          = 0x00000001    ///< Y0CbY1Cr
-					,DNX_CCO_CbYCrY_NoA          = 0x00000002    ///< CbY0CrY1
-					,DNX_CCO_ARGB_Interleaved    = 0x00000004    ///< ARGB
-					,DNX_CCO_BGRA_Interleaved    = 0x00000008    ///< BGRA
-					,DNX_CCO_RGB_NoA             = 0x00000040    ///< RGB
-					,DNX_CCO_BGR_NoA             = 0x00000080    ///< BGR
-					,DNX_CCO_RGBA_Interleaved    = 0x00000800    ///< RGBA
-					,DNX_CCO_ABGR_Interleaved    = 0x00001000    ///< ABGR
-					,DNX_CCO_YCbCr_Interleaved   = 0x00002000    ///< YCbCr 444
-					,DNX_CCO_Ch1Ch2Ch3           = 0x00004000    ///< Arbitrary 444 subsampled color components for any colorspace other than RGB and YCbCr
-					,DNX_CCO_Ch1Ch2Ch1Ch3        = 0x00008000    ///< Arbitrary 422 subsampled color components for any colorspace other than RGB and YCbCr with not subsampled channel first
-					,DNX_CCO_Ch2Ch1Ch3Ch1        = 0x00010000    ///< Arbitrary 422 subsampled color components for any colorspace other than RGB and YCbCr with subsampled channel first
-					,DNX_CCO_YCbCr_Planar        = 0x00020000    ///< YCbCr 420 stored in planar form. Y followed by Cb followed by Cr
-					,DNX_CCO_CbYACrYA_Interleaved= 0x00040000    ///< YCbCrA 4224
-					,DNX_CCO_CbYCrA_Interleaved  = 0x00080000    ///< YCbCrA 4444
-					,DNX_CCO_YCbCrA_Planar       = 0x00100000    ///< YCbCrA 4204 stored in planar form. Y followed by Cb followed by Cr followed by A
-					,DNX_CCO_Ch1Ch2Ch3A          = 0x00200000    ///< Arbitrary 4444 subsampled color components for any colorspace other than RGB and YCbCr with Alpha (Alpha channel last)
-					,DNX_CCO_Ch3Ch2Ch1A          = 0x00400000    ///< Arbitrary 4444 subsampled color components for any colorspace other than RGB and YCbCr with Alpha with reversed channel order (Alpha channel last)
-					,DNX_CCO_ACh1Ch2Ch3          = 0x00800000    ///< Arbitrary 4444 subsampled color components for any colorspace other than RGB and YCbCr with Alpha (Alpha channel first)
-					,DNX_CCO_Ch2Ch1ACh3Ch1A      = 0x01000000    ///< Arbitrary 4224 subsampled color components for any colorspace other than RGB and YCbCr with not subsampled channel first with Alpha
-					,DNX_CCO_CbYCrY_A            = 0x02000000    ///< YCbCrA 4224 mixed-planar (separate Alpha plane)
-				*/
-
-			int32 AvidColorVolume = (int32)ElectraDecodersUtil::GetVariantValueSafeI64(ExtraValues, TEXT("avid_color_volume"), 0);
-				/*
-					,DNX_CV_709             = 0x01     ///< Rec. 709
-					,DNX_CV_2020            = 0x02     ///< Non-constant luminance Rec. 2020
-					,DNX_CV_2020c           = 0x04     ///< Constant luminance Rec. 2020
-					,DNX_CV_OutOfBand       = 0x08     ///< Any other
-				*/
-
-			switch (AvidComponentType)
+			else
 			{
-				case	0x001:	// DNX_CT_UCHAR
-				{
-					switch (AvidComponentOrder)
-					{
-						case	0x00000001:	// DNX_CCO_YCbYCr_NoA
-						case	0x00000002:	// DNX_CCO_CbYCrY_NoA
-						{
-							InOutBufferPropertes->Set(TEXT("pixelfmt"), FVariantValue((int64)EPixelFormat::PF_B8G8R8A8));
-							InOutBufferPropertes->Set(TEXT("pixelenc"), FVariantValue((int64)((AvidComponentOrder == 0x00000002) ? EVideoDecoderPixelEncoding::CbY0CrY1 : EVideoDecoderPixelEncoding::Y0CbY1Cr)));
-							Width = (Width + 1) / 2; // each pixel contains TWO horizontally adjacent output pixels)
-							Pitch *= 2;
-							break;
-						}
-						case	0x00000004:	// DNX_CCO_ARGB_Interleaved
-						{
-							InOutBufferPropertes->Set(TEXT("pixelfmt"), FVariantValue((int64)EPixelFormat::PF_A8R8G8B8));
-							Pitch *= 4;
-							break;
-						}
-						case	0x00000800:	// DNX_CCO_RGBA_Interleaved
-						{
-							InOutBufferPropertes->Set(TEXT("pixelfmt"), FVariantValue((int64)EPixelFormat::PF_R8G8B8A8));
-							Pitch *= 4;
-							break;
-						}
-						case	0x00080000:	// DNX_CCO_CbYCrA_Interleaved
-						{
-							InOutBufferPropertes->Set(TEXT("pixelfmt"), FVariantValue((int64)EPixelFormat::PF_B8G8R8A8));
-							InOutBufferPropertes->Set(TEXT("pixelenc"), FVariantValue((int64)EVideoDecoderPixelEncoding::YCbCr_Alpha));
-							Pitch *= 4;
-							break;
-						}
-						case	0x00020000:	// DNX_CCO_YCbCr_Planar
-						{
-							InOutBufferPropertes->Set(TEXT("pixelfmt"), FVariantValue((int64)EPixelFormat::PF_NV12));
-							Height = (Height * 3) / 2;
-							break;
-						}
-						default:
-						{
-							return false;
-						}
-					}
-					break;
-				}
-				case	0x010:	// DNX_CT_USHORT
-				case	0x004:	// DNX_CT_USHORT_10_6
-				case	0x20000: // DNX_CT_USHORT_12_4
+				// Note: we assume a DX11 texture at this point, but this could also be interpreted as IUknown and then either DX11 or DX12 resources being derived from it...
+				ID3D11Texture2D* Texture = static_cast<ID3D11Texture2D*>(ImageBuffers->GetBufferTextureByIndex(0));
+				if (Texture != nullptr)
 				{
-					switch (AvidComponentOrder)
+					//
+					// GPU texture
+					//
+					ID3D11Device* Device = reinterpret_cast<ID3D11Device*>(InDecoderOutput->GetPlatformOutputHandle(EElectraDecoderPlatformOutputHandleType::DXDevice));
+					ID3D11DeviceContext* DeviceContext = reinterpret_cast<ID3D11DeviceContext*>(InDecoderOutput->GetPlatformOutputHandle(EElectraDecoderPlatformOutputHandleType::DXDeviceContext));
+					if (Device && DeviceContext)
 					{
-						case	0x00000001:	// DNX_CCO_YCbYCr_NoA
-						case	0x00000002:	// DNX_CCO_CbYCrY_NoA
-						{
-							InOutBufferPropertes->Set(TEXT("pixelfmt"), FVariantValue((int64)EPixelFormat::PF_A16B16G16R16));
-							InOutBufferPropertes->Set(TEXT("pixelenc"), FVariantValue((int64)((AvidComponentOrder == 0x00000002) ? EVideoDecoderPixelEncoding::CbY0CrY1 : EVideoDecoderPixelEncoding::Y0CbY1Cr)));
-							Width = (Width + 1) / 2; // each pixel contains TWO horizontally adjacent output pixels)
-							Pitch *= 2;
-							break;
-						}
-						case	0x00000800:	// DNX_CCO_RGBA_Interleaved
-						{
-							InOutBufferPropertes->Set(TEXT("pixelfmt"), FVariantValue((int64)EPixelFormat::PF_R16G16B16A16_UNORM));
-							Pitch *= 8;
-							break;
-						}
-						case	0x00001000:	// DNX_CCO_ABGR_Interleaved
-						{
-							InOutBufferPropertes->Set(TEXT("pixelfmt"), FVariantValue((int64)EPixelFormat::PF_A16B16G16R16));
-							Pitch *= 8;
-							break;
-						}
-						case	0x00080000:	// DNX_CCO_CbYCrA_Interleaved
-						{
-							InOutBufferPropertes->Set(TEXT("pixelfmt"), FVariantValue((int64)EPixelFormat::PF_A16B16G16R16));
-							InOutBufferPropertes->Set(TEXT("pixelenc"), FVariantValue((int64)EVideoDecoderPixelEncoding::YCbCr_Alpha));
-							Pitch *= 8;
-							break;
-						}
-						case	0x00020000:	// DNX_CCO_YCbCr_Planar
-						{
-							InOutBufferPropertes->Set(TEXT("pixelfmt"), FVariantValue((int64)EPixelFormat::PF_P010));
-							Height = (Height * 3) / 2;
-							break;
-						}
-						default:
-						{
-							return false;
-						}
+						DecoderOutput->InitializeWithSharedTexture(Device, DeviceContext, Texture, FIntPoint(InDecoderOutput->GetWidth(), InDecoderOutput->GetHeight()), InOutBufferPropertes);
+						return true;
 					}
-					break;
-				}
-				case	0x400:	// DNX_CT_USHORT_V210
-				{
-					InOutBufferPropertes->Set(TEXT("pixelfmt"), FVariantValue((int64)EPixelFormat::PF_A2B10G10R10));
-					InOutBufferPropertes->Set(TEXT("pixelenc"), FVariantValue((int64)EVideoDecoderPixelEncoding::CbY0CrY1));
-					Width = 4 * ((Width + 5) / 6); // each 4 pixel contain 6 horizontally adjacent YCbCr pixels (incl. 4x 2-bit padding)
-					Pitch = Width * 4;
-					break;
-				}
-				default:
-				{
 				}
 			}
 
-//TODO >>> AvidColorVolume: right now we only deliver MP4-based HDR information
-
-			InOutBufferPropertes->Set(TEXT("bits_per"), FVariantValue((int64)InDecoderOutput->GetNumberOfBits()));
-
-			TSharedPtr<TArray<uint8>, ESPMode::ThreadSafe> ColorBuffer = ImageBuffers->GetBufferByIndex(0);
-			if (ColorBuffer.IsValid() && ColorBuffer->Num())
-			{
-				DecoderOutput->InitializeWithBuffer(ColorBuffer,
-					Pitch,							// Buffer stride
-					FIntPoint(Width, Height),		// Buffer dimensions
-					InOutBufferPropertes);
-				return true;
-			}
 		}
 		else
 		{
-			UE_LOG(LogElectraPlayer, Log, TEXT("Unsupported decoded 4CC 0x%08x"), Vars->Codec4CC);
-			return false;
+			//
+			// IMF sample?
+			//
+			if (InDecoderOutput->GetPlatformOutputHandle(EElectraDecoderPlatformOutputHandleType::MFSample) != nullptr)
+			{
+				return SetupRenderBufferFromDecoderOutputFromMFSample(InOutBufferToSetup, InOutBufferPropertes, InDecoderOutput, InPlatformSpecificResource);
+			}
 		}
 
 	}
diff --git a/Engine/Plugins/Media/ElectraPlayer/Source/ElectraPlayerRuntime/Private/Runtime/Decoder/Windows/WindowsElectraDecoderResourceManager.h b/Engine/Plugins/Media/ElectraPlayer/Source/ElectraPlayerRuntime/Private/Runtime/Decoder/Windows/WindowsElectraDecoderResourceManager.h
--- a/Engine/Plugins/Media/ElectraPlayer/Source/ElectraPlayerRuntime/Private/Runtime/Decoder/Windows/WindowsElectraDecoderResourceManager.h
+++ b/Engine/Plugins/Media/ElectraPlayer/Source/ElectraPlayerRuntime/Private/Runtime/Decoder/Windows/WindowsElectraDecoderResourceManager.h
@@ -36,8 +36,7 @@ public:
 
 private:
 	class FInstanceVars;
-	static bool SetupRenderBufferFromDecoderOutput_Default(IMediaRenderer::IBuffer* InOutBufferToSetup, TSharedPtr<FParamDict, ESPMode::ThreadSafe> InOutBufferPropertes, TSharedPtr<IElectraDecoderVideoOutput, ESPMode::ThreadSafe> InDecoderOutput, IDecoderPlatformResource* InPlatformSpecificResource);
-	static bool SetupRenderBufferFromDecoderOutput_Images(IMediaRenderer::IBuffer* InOutBufferToSetup, TSharedPtr<FParamDict, ESPMode::ThreadSafe> InOutBufferPropertes, TSharedPtr<IElectraDecoderVideoOutput, ESPMode::ThreadSafe> InDecoderOutput, IDecoderPlatformResource* InPlatformSpecificResource);
+	static bool SetupRenderBufferFromDecoderOutputFromMFSample(IMediaRenderer::IBuffer* InOutBufferToSetup, TSharedPtr<FParamDict, ESPMode::ThreadSafe> InOutBufferPropertes, TSharedPtr<IElectraDecoderVideoOutput, ESPMode::ThreadSafe> InDecoderOutput, IDecoderPlatformResource* InPlatformSpecificResource);
 };
 
 using FPlatformElectraDecoderResourceManager = FElectraDecoderResourceManagerWindows;
diff --git a/Engine/Plugins/Media/HAPDecoderElectra/Source/HAPDecoderElectra/Private/HAPDecoder/ElectraMediaHAPDecoder.cpp b/Engine/Plugins/Media/HAPDecoderElectra/Source/HAPDecoderElectra/Private/HAPDecoder/ElectraMediaHAPDecoder.cpp
--- a/Engine/Plugins/Media/HAPDecoderElectra/Source/HAPDecoderElectra/Private/HAPDecoder/ElectraMediaHAPDecoder.cpp
+++ b/Engine/Plugins/Media/HAPDecoderElectra/Source/HAPDecoderElectra/Private/HAPDecoder/ElectraMediaHAPDecoder.cpp
@@ -73,14 +73,6 @@ public:
 	{
 		return Crop;
 	}
-	int32 GetFrameWidth() const override
-	{
-		return Pitch;
-	}
-	int32 GetFrameHeight() const override
-	{
-		return Height;
-	}
 	int32 GetAspectRatioW() const override
 	{
 		return AspectW;
@@ -105,10 +97,6 @@ public:
 	{
 		OutExtraValues = ExtraValues;
 	}
-	int32 GetPixelFormat() const override
-	{
-		return PixelFormat;
-	}
 	void* GetPlatformOutputHandle(EElectraDecoderPlatformOutputHandleType InTypeOfHandle) const override
 	{
 		if (InTypeOfHandle == EElectraDecoderPlatformOutputHandleType::ImageBuffers)
@@ -135,7 +123,7 @@ public:
 	{
 		return NumBuffers;
 	}
-	TSharedPtr<TArray<uint8>, ESPMode::ThreadSafe> GetBufferByIndex(int32 InBufferIndex) const override
+	TSharedPtr<TArray<uint8>, ESPMode::ThreadSafe> GetBufferDataByIndex(int32 InBufferIndex) const override
 	{
 		if (InBufferIndex == 0)
 		{
@@ -147,7 +135,11 @@ public:
 		}
 		return nullptr;
 	}
-	uint64 GetBufferFormatByIndex(int32 InBufferIndex) const override
+	void* GetBufferTextureByIndex(int32 InBufferIndex) const override
+	{
+		return nullptr;
+	}
+	EElectraDecoderPlatformPixelFormat GetBufferFormatByIndex(int32 InBufferIndex) const override
 	{
 		if (InBufferIndex == 0)
 		{
@@ -157,9 +149,32 @@ public:
 		{
 			return AlphaBufferFormat;
 		}
-	return 0;
+	return EElectraDecoderPlatformPixelFormat::INVALID;
+	}
+	EElectraDecoderPlatformPixelEncoding GetBufferEncodingByIndex(int32 InBufferIndex) const override
+	{
+		if (InBufferIndex == 0)
+		{
+			return ColorBufferEncoding;
+		}
+		else if (InBufferIndex == 1)
+		{
+			return AlphaBufferEncoding;
+		}
+		return EElectraDecoderPlatformPixelEncoding::Native;
+	}
+	int32 GetBufferPitchByIndex(int32 InBufferIndex) const override
+	{
+		if (InBufferIndex == 0)
+		{
+			return ColorPitch;
+		}
+		else if (InBufferIndex == 1)
+		{
+			return AlphaPitch;
+		}
+		return 0;
 	}
-
 
 public:
 	FTimespan PTS;
@@ -168,7 +183,8 @@ public:
 	FElectraVideoDecoderOutputCropValues Crop;
 	int32 Width = 0;
 	int32 Height = 0;
-	int32 Pitch = 0;
+	int32 ColorPitch = 0;
+	int32 AlphaPitch = 0;
 	int32 NumBits = 0;
 	int32 AspectW = 1;
 	int32 AspectH = 1;
@@ -181,12 +197,13 @@ public:
 	int32 NumBuffers = 0;
 	TSharedPtr<TArray<uint8>, ESPMode::ThreadSafe> ColorBuffer;
 	TSharedPtr<TArray<uint8>, ESPMode::ThreadSafe> AlphaBuffer;
-	HapTextureFormat ColorBufferFormat;
-	HapTextureFormat AlphaBufferFormat;
+	EElectraDecoderPlatformPixelFormat ColorBufferFormat;
+	EElectraDecoderPlatformPixelEncoding ColorBufferEncoding;
+	EElectraDecoderPlatformPixelFormat AlphaBufferFormat;
+	EElectraDecoderPlatformPixelEncoding AlphaBufferEncoding;
 };
 
 
-
 class FVideoDecoderHAPElectra : public IElectraDecoder
 {
 public:
@@ -408,6 +425,59 @@ TSharedPtr<IElectraDecoderDefaultOutputFormat, ESPMode::ThreadSafe> FVideoDecode
 	return nullptr;
 }
 
+static bool ConvertHapTextureFormat(HapTextureFormat HapFmt, EElectraDecoderPlatformPixelFormat& Format, EElectraDecoderPlatformPixelEncoding& Encoding)
+{
+	switch (HapFmt)
+	{
+		case HapTextureFormat_RGB_DXT1:
+			Format = EElectraDecoderPlatformPixelFormat::DXT1;
+			Encoding = EElectraDecoderPlatformPixelEncoding::Native;
+			break;
+		case HapTextureFormat_RGBA_DXT5:
+			Format = EElectraDecoderPlatformPixelFormat::DXT5;
+			Encoding = EElectraDecoderPlatformPixelEncoding::Native;
+			break;
+		case HapTextureFormat_YCoCg_DXT5:
+			Format = EElectraDecoderPlatformPixelFormat::DXT5;
+			Encoding = EElectraDecoderPlatformPixelEncoding::YCoCg;
+			break;
+		case HapTextureFormat_A_RGTC1:
+			Format = EElectraDecoderPlatformPixelFormat::BC4;
+			Encoding = EElectraDecoderPlatformPixelEncoding::Native;
+			break;
+		case HapTextureFormat_RGBA_BPTC_UNORM:
+		case HapTextureFormat_RGB_BPTC_UNSIGNED_FLOAT:
+		case HapTextureFormat_RGB_BPTC_SIGNED_FLOAT:
+		default:
+			Format = EElectraDecoderPlatformPixelFormat::INVALID;
+			Encoding = EElectraDecoderPlatformPixelEncoding::Native;
+			break;
+	}
+	return Format != EElectraDecoderPlatformPixelFormat::INVALID;
+}
+
+static uint32 GetImageBufferPitch(HapTextureFormat HapFmt, uint32 Width)
+{
+	uint32 Pitch = 0;
+	switch (HapFmt)
+	{
+		case HapTextureFormat_RGB_DXT1:
+		case HapTextureFormat_A_RGTC1:
+			Pitch = ((Width + 3) / 4) * 8;		// 4 pixel wide blocks with 8 bytes
+			break;
+		case HapTextureFormat_RGBA_DXT5:
+		case HapTextureFormat_YCoCg_DXT5:
+			Pitch = ((Width + 3) / 4) * 16;		// 4 pixel wide blocks with 16 bytes
+			break;
+		case HapTextureFormat_RGBA_BPTC_UNORM:
+		case HapTextureFormat_RGB_BPTC_UNSIGNED_FLOAT:
+		case HapTextureFormat_RGB_BPTC_SIGNED_FLOAT:
+		default:
+			break;
+	}
+	return Pitch;
+}
+
 IElectraDecoder::EDecoderError FVideoDecoderHAPElectra::DecodeAccessUnit(const FInputAccessUnit& InInputAccessUnit, const TMap<FString, FVariant>& InAdditionalOptions)
 {
 	// If already in error do nothing!
@@ -446,7 +516,6 @@ IElectraDecoder::EDecoderError FVideoDecoderHAPElectra::DecodeAccessUnit(const F
 
 		NewOutput->Width = DecodedWidth;
 		NewOutput->Height = DecodedHeight;
-		NewOutput->Pitch = NewOutput->Width;
 		NewOutput->Crop.Right = DecodedWidth - DisplayWidth;
 		NewOutput->Crop.Bottom = DecodedHeight - DisplayHeight;
 		if (AspectW && AspectH)
@@ -475,7 +544,9 @@ IElectraDecoder::EDecoderError FVideoDecoderHAPElectra::DecodeAccessUnit(const F
 			return IElectraDecoder::EDecoderError::Error;
 		}
 		NewOutput->ColorBuffer->SetNumUnsafeInternal((int32) ColorBufferBytesUsed);
-		NewOutput->ColorBufferFormat = static_cast<HapTextureFormat>(ColorBufferTextureFormat);
+		ConvertHapTextureFormat(static_cast<HapTextureFormat>(ColorBufferTextureFormat), NewOutput->ColorBufferFormat, NewOutput->ColorBufferEncoding);
+
+		NewOutput->ColorPitch = GetImageBufferPitch(static_cast<HapTextureFormat>(ColorBufferTextureFormat), NewOutput->Width);
 
 		NewOutput->NumBuffers = (int32)TextureCount;
 		NewOutput->Codec4CC = Codec4CC;
@@ -497,7 +568,9 @@ IElectraDecoder::EDecoderError FVideoDecoderHAPElectra::DecodeAccessUnit(const F
 				return IElectraDecoder::EDecoderError::Error;
 			}
 			NewOutput->AlphaBuffer->SetNumUnsafeInternal((int32) AlphaBufferBytesUsed);
-			NewOutput->AlphaBufferFormat = static_cast<HapTextureFormat>(AlphaBufferTextureFormat);
+			ConvertHapTextureFormat(static_cast<HapTextureFormat>(AlphaBufferTextureFormat), NewOutput->AlphaBufferFormat, NewOutput->AlphaBufferEncoding);
+
+			NewOutput->AlphaPitch = GetImageBufferPitch(static_cast<HapTextureFormat>(AlphaBufferTextureFormat), NewOutput->Width);
 		}
 
 		CurrentOutput = MoveTemp(NewOutput);
