Electra: Fixed incorrect DX11 and Vulkan video output
#rnx

diff --git a/Engine/Plugins/Media/ElectraCodecs/Source/ElectraDecoders/Private/DX/VideoDecoderH264_DX.cpp b/Engine/Plugins/Media/ElectraCodecs/Source/ElectraDecoders/Private/DX/VideoDecoderH264_DX.cpp
--- a/Engine/Plugins/Media/ElectraCodecs/Source/ElectraDecoders/Private/DX/VideoDecoderH264_DX.cpp
+++ b/Engine/Plugins/Media/ElectraCodecs/Source/ElectraDecoders/Private/DX/VideoDecoderH264_DX.cpp
@@ -50,9 +50,9 @@ public:
 	{ return UserValue; }
 
 	int32 GetWidth() const override
-	{ return Width - Crop.Left - Crop.Right; }
+	{ return Width; }
 	int32 GetHeight() const override
-	{ return Height - Crop.Top - Crop.Bottom; }
+	{ return Height; }
 	int32 GetDecodedWidth() const override
 	{ return DecodedWidth; }
 	int32 GetDecodedHeight() const override
@@ -884,7 +884,9 @@ bool FElectraVideoDecoderH264_DX::ConvertDecoderOutput()
 	// Width and height as "decoded dims" (which are featuring format specific scales & no cropping)
 	NewOutput->DecodedWidth = (int32)dwInputWidth;
 	NewOutput->DecodedHeight = (int32)dwInputHeight;
-	if (DecoderPlatformHandle->IsSoftware() || DecoderPlatformHandle->GetDXVersionTimes1000() >= 12000)
+	// If we deliver the data in a CPU side buffer, we need to patch the decoded height to reflect the oddness of these planar formats
+	// (DirectX12 and higher, SW decoder or a non-DirectX API will trigger this)
+	if (DecoderPlatformHandle->IsSoftware() || DecoderPlatformHandle->GetDXVersionTimes1000() >= 12000 || DecoderPlatformHandle->GetDXVersionTimes1000() == 0)
 	{
 		// We are returning a CPU side buffer (software decode OR DX12 render device) - adjust the height so it can be interpreted as a single plane texture
 		NewOutput->DecodedHeight = NewOutput->DecodedHeight * 3 / 2;
diff --git a/Engine/Plugins/Media/ElectraCodecs/Source/ElectraDecoders/Private/DX/VideoDecoderH265_DX.cpp b/Engine/Plugins/Media/ElectraCodecs/Source/ElectraDecoders/Private/DX/VideoDecoderH265_DX.cpp
--- a/Engine/Plugins/Media/ElectraCodecs/Source/ElectraDecoders/Private/DX/VideoDecoderH265_DX.cpp
+++ b/Engine/Plugins/Media/ElectraCodecs/Source/ElectraDecoders/Private/DX/VideoDecoderH265_DX.cpp
@@ -50,9 +50,9 @@ public:
 	{ return UserValue; }
 
 	int32 GetWidth() const override
-	{ return Width - Crop.Left - Crop.Right; }
+	{ return Width; }
 	int32 GetHeight() const override
-	{ return Height - Crop.Top - Crop.Bottom; }
+	{ return Height; }
 	int32 GetDecodedWidth() const override
 	{ return DecodedWidth; }
 	int32 GetDecodedHeight() const override
@@ -897,7 +897,7 @@ bool FElectraVideoDecoderH265_DX::ConvertDecoderOutput()
 	// Width and height as "decoded dims" (which are featuring format specific scales & no cropping)
 	NewOutput->DecodedWidth = (int32)dwInputWidth;
 	NewOutput->DecodedHeight = (int32)dwInputHeight;
-	if (DecoderPlatformHandle->IsSoftware() || DecoderPlatformHandle->GetDXVersionTimes1000() >= 12000)
+	if (DecoderPlatformHandle->IsSoftware() || DecoderPlatformHandle->GetDXVersionTimes1000() >= 12000 || DecoderPlatformHandle->GetDXVersionTimes1000() == 0)
 	{
 		// We are returning a CPU side buffer (software decode OR DX12 render device) - adjust the height so it can be interpreted as a single plane texture
 		NewOutput->DecodedHeight = NewOutput->DecodedHeight * 3 / 2;
diff --git a/Engine/Plugins/Media/ElectraPlayer/Source/ElectraPlayerRuntime/Private/Runtime/Decoder/Windows/ElectraVideoDecoder_PC.cpp b/Engine/Plugins/Media/ElectraPlayer/Source/ElectraPlayerRuntime/Private/Runtime/Decoder/Windows/ElectraVideoDecoder_PC.cpp
--- a/Engine/Plugins/Media/ElectraPlayer/Source/ElectraPlayerRuntime/Private/Runtime/Decoder/Windows/ElectraVideoDecoder_PC.cpp
+++ b/Engine/Plugins/Media/ElectraPlayer/Source/ElectraPlayerRuntime/Private/Runtime/Decoder/Windows/ElectraVideoDecoder_PC.cpp
@@ -111,15 +111,10 @@ void FElectraPlayerVideoDecoderOutputPC::InitializeWithSharedTexture(const TRefC
 		return;
 	}
 
-	_InitializeWithSharedTexture(InD3D11Device, InDeviceContext, DecoderTexture, ViewIndex, OutputDim, InParamDict);
+	InitializeWithSharedTexture(InD3D11Device, InDeviceContext, DecoderTexture, ViewIndex, OutputDim, InParamDict);
 }
 
-void FElectraPlayerVideoDecoderOutputPC::InitializeWithSharedTexture(const TRefCountPtr<ID3D11Device>& InD3D11Device, const TRefCountPtr<ID3D11DeviceContext> InDeviceContext, const TRefCountPtr<ID3D11Texture2D> DecoderTexture, const FIntPoint& OutputDim, TSharedPtr<Electra::FParamDict, ESPMode::ThreadSafe> InParamDict)
-{
-	_InitializeWithSharedTexture(InD3D11Device, InDeviceContext, DecoderTexture, 0, OutputDim, InParamDict);
-}
-
-void FElectraPlayerVideoDecoderOutputPC::_InitializeWithSharedTexture(const TRefCountPtr<ID3D11Device>& InD3D11Device, const TRefCountPtr<ID3D11DeviceContext> InDeviceContext, const TRefCountPtr<ID3D11Texture2D> DecoderTexture, uint32 ViewIndex, const FIntPoint& OutputDim, TSharedPtr<Electra::FParamDict, ESPMode::ThreadSafe> InParamDict)
+void FElectraPlayerVideoDecoderOutputPC::InitializeWithSharedTexture(const TRefCountPtr<ID3D11Device>& InD3D11Device, const TRefCountPtr<ID3D11DeviceContext> InDeviceContext, const TRefCountPtr<ID3D11Texture2D> DecoderTexture, uint32 ViewIndex, const FIntPoint& OutputDim, TSharedPtr<Electra::FParamDict, ESPMode::ThreadSafe> InParamDict)
 {
 	FVideoDecoderOutput::Initialize(InParamDict);
 
diff --git a/Engine/Plugins/Media/ElectraPlayer/Source/ElectraPlayerRuntime/Private/Runtime/Decoder/Windows/ElectraVideoDecoder_PC.h b/Engine/Plugins/Media/ElectraPlayer/Source/ElectraPlayerRuntime/Private/Runtime/Decoder/Windows/ElectraVideoDecoder_PC.h
--- a/Engine/Plugins/Media/ElectraPlayer/Source/ElectraPlayerRuntime/Private/Runtime/Decoder/Windows/ElectraVideoDecoder_PC.h
+++ b/Engine/Plugins/Media/ElectraPlayer/Source/ElectraPlayerRuntime/Private/Runtime/Decoder/Windows/ElectraVideoDecoder_PC.h
@@ -43,7 +43,7 @@ public:
 	void InitializeWithSharedTexture(const TRefCountPtr<ID3D11Device>& InD3D11Device, const TRefCountPtr<ID3D11DeviceContext> InDeviceContext, const TRefCountPtr<IMFSample> MFSample, const FIntPoint& OutputDim, TSharedPtr<Electra::FParamDict, ESPMode::ThreadSafe> InParamDict);
 
 	// Hardware decode to shared DX11 texture (Win8+) from DX11 texture
-	void InitializeWithSharedTexture(const TRefCountPtr<ID3D11Device>& InD3D11Device, const TRefCountPtr<ID3D11DeviceContext> InDeviceContext, const TRefCountPtr<ID3D11Texture2D> DecoderTexture, const FIntPoint& OutputDim, TSharedPtr<Electra::FParamDict, ESPMode::ThreadSafe> InParamDict);
+	void InitializeWithSharedTexture(const TRefCountPtr<ID3D11Device>& InD3D11Device, const TRefCountPtr<ID3D11DeviceContext> InDeviceContext, const TRefCountPtr<ID3D11Texture2D> DecoderTexture, uint32 ViewIndex, const FIntPoint& OutputDim, TSharedPtr<Electra::FParamDict, ESPMode::ThreadSafe> InParamDict);
 
 	void SetOwner(const TSharedPtr<IDecoderOutputOwner, ESPMode::ThreadSafe>& InOwningRenderer) override;
 	void ShutdownPoolable() override;
@@ -57,8 +57,6 @@ public:
 	TRefCountPtr<IUnknown> GetSync(uint64& SyncValue) const override;
 
 private:
-	void _InitializeWithSharedTexture(const TRefCountPtr<ID3D11Device>& InD3D11Device, const TRefCountPtr<ID3D11DeviceContext> InDeviceContext, const TRefCountPtr<ID3D11Texture2D> DecoderTexture, uint32 ViewIndex, const FIntPoint& OutputDim, TSharedPtr<Electra::FParamDict, ESPMode::ThreadSafe> InParamDict);
-
 	// Decoder output type
 	EOutputType OutputType = EOutputType::Unknown;
 
diff --git a/Engine/Plugins/Media/ElectraPlayer/Source/ElectraPlayerRuntime/Private/Runtime/Decoder/Windows/WindowsElectraDecoderResourceManager.cpp b/Engine/Plugins/Media/ElectraPlayer/Source/ElectraPlayerRuntime/Private/Runtime/Decoder/Windows/WindowsElectraDecoderResourceManager.cpp
--- a/Engine/Plugins/Media/ElectraPlayer/Source/ElectraPlayerRuntime/Private/Runtime/Decoder/Windows/WindowsElectraDecoderResourceManager.cpp
+++ b/Engine/Plugins/Media/ElectraPlayer/Source/ElectraPlayerRuntime/Private/Runtime/Decoder/Windows/WindowsElectraDecoderResourceManager.cpp
@@ -247,7 +247,13 @@ bool FElectraDecoderResourceManagerWindows::SetupRenderBufferFromDecoderOutputFr
 				{
 					return false;
 				}
-				DecoderOutput->InitializeWithSharedTexture(Device, DeviceContext, Texture2D, FIntPoint(InDecoderOutput->GetWidth(), InDecoderOutput->GetHeight()), InOutBufferPropertes);
+
+				uint32 ViewIndex = 0;
+				if (DXGIBuffer->GetSubresourceIndex(&ViewIndex) != S_OK)
+				{
+					return false;
+				}
+				DecoderOutput->InitializeWithSharedTexture(Device, DeviceContext, Texture2D, ViewIndex, FIntPoint(InDecoderOutput->GetWidth(), InDecoderOutput->GetHeight()), InOutBufferPropertes);
 				if (!DecoderOutput->GetTexture())
 				{
 					UE_LOG(LogElectraPlayer, Error, TEXT("ID3D11Device::CreateTexture2D() failed!"));
@@ -364,35 +370,34 @@ bool FElectraDecoderResourceManagerWindows::SetupRenderBufferFromDecoderOutput(I
 
 	if (DecoderOutput.IsValid())
 	{
-		FElectraVideoDecoderOutputCropValues Crop = InDecoderOutput->GetCropValues();
-		InOutBufferPropertes->Set(TEXT("width"), FVariantValue((int64)InDecoderOutput->GetWidth()));
-		InOutBufferPropertes->Set(TEXT("height"), FVariantValue((int64)InDecoderOutput->GetHeight()));
-		InOutBufferPropertes->Set(TEXT("crop_left"), FVariantValue((int64)Crop.Left));
-		InOutBufferPropertes->Set(TEXT("crop_right"), FVariantValue((int64)Crop.Right));
-		InOutBufferPropertes->Set(TEXT("crop_top"), FVariantValue((int64)Crop.Top));
-		InOutBufferPropertes->Set(TEXT("crop_bottom"), FVariantValue((int64)Crop.Bottom));
-		InOutBufferPropertes->Set(TEXT("aspect_ratio"), FVariantValue((double)InDecoderOutput->GetAspectRatioW() / (double)InDecoderOutput->GetAspectRatioH()));
-		InOutBufferPropertes->Set(TEXT("aspect_w"), FVariantValue((int64)InDecoderOutput->GetAspectRatioW()));
-		InOutBufferPropertes->Set(TEXT("aspect_h"), FVariantValue((int64)InDecoderOutput->GetAspectRatioH()));
-		InOutBufferPropertes->Set(TEXT("fps_num"), FVariantValue((int64)InDecoderOutput->GetFrameRateNumerator()));
-		InOutBufferPropertes->Set(TEXT("fps_denom"), FVariantValue((int64)InDecoderOutput->GetFrameRateDenominator()));
-
-		int32 Width = InDecoderOutput->GetDecodedWidth();
-		int32 Height = InDecoderOutput->GetDecodedHeight();
-
-		InOutBufferPropertes->Set(TEXT("bits_per"), FVariantValue((int64)InDecoderOutput->GetNumberOfBits()));
-
 		//
 		// Image buffers?
 		//
 		IElectraDecoderVideoOutputImageBuffers* ImageBuffers = reinterpret_cast<IElectraDecoderVideoOutputImageBuffers*>(InDecoderOutput->GetPlatformOutputHandle(EElectraDecoderPlatformOutputHandleType::ImageBuffers));
 		if (ImageBuffers != nullptr)
 		{
+			FElectraVideoDecoderOutputCropValues Crop = InDecoderOutput->GetCropValues();
+			InOutBufferPropertes->Set(TEXT("width"), FVariantValue((int64)InDecoderOutput->GetWidth()));
+			InOutBufferPropertes->Set(TEXT("height"), FVariantValue((int64)InDecoderOutput->GetHeight()));
+			InOutBufferPropertes->Set(TEXT("crop_left"), FVariantValue((int64)Crop.Left));
+			InOutBufferPropertes->Set(TEXT("crop_right"), FVariantValue((int64)Crop.Right));
+			InOutBufferPropertes->Set(TEXT("crop_top"), FVariantValue((int64)Crop.Top));
+			InOutBufferPropertes->Set(TEXT("crop_bottom"), FVariantValue((int64)Crop.Bottom));
+			InOutBufferPropertes->Set(TEXT("aspect_ratio"), FVariantValue((double)InDecoderOutput->GetAspectRatioW() / (double)InDecoderOutput->GetAspectRatioH()));
+			InOutBufferPropertes->Set(TEXT("aspect_w"), FVariantValue((int64)InDecoderOutput->GetAspectRatioW()));
+			InOutBufferPropertes->Set(TEXT("aspect_h"), FVariantValue((int64)InDecoderOutput->GetAspectRatioH()));
+			InOutBufferPropertes->Set(TEXT("fps_num"), FVariantValue((int64)InDecoderOutput->GetFrameRateNumerator()));
+			InOutBufferPropertes->Set(TEXT("fps_denom"), FVariantValue((int64)InDecoderOutput->GetFrameRateDenominator()));
+
+			int32 Width = InDecoderOutput->GetDecodedWidth();
+			int32 Height = InDecoderOutput->GetDecodedHeight();
+
+			InOutBufferPropertes->Set(TEXT("bits_per"), FVariantValue((int64)InDecoderOutput->GetNumberOfBits()));
+
 			int32 NumImageBuffers = ImageBuffers->GetNumberOfBuffers();
 			check(NumImageBuffers == 1 || NumImageBuffers == 2);
 
 			// Color buffer
-
 			EElectraDecoderPlatformPixelFormat PixFmt = ImageBuffers->GetBufferFormatByIndex(0);
 			EElectraDecoderPlatformPixelEncoding PixEnc = ImageBuffers->GetBufferEncodingByIndex(0);
 
@@ -454,7 +459,7 @@ bool FElectraDecoderResourceManagerWindows::SetupRenderBufferFromDecoderOutput(I
 			}
 			else
 			{
-				// Note: we assume a DX11 texture at this point, but this could also be interpreted as IUknown and then either DX11 or DX12 resources being derived from it...
+				// Note: we assume a DX11 texture at this point, but this could also be interpreted as IUnknown and then either DX11 or DX12 resources being derived from it...
 				ID3D11Texture2D* Texture = static_cast<ID3D11Texture2D*>(ImageBuffers->GetBufferTextureByIndex(0));
 				if (Texture != nullptr)
 				{
@@ -465,7 +470,8 @@ bool FElectraDecoderResourceManagerWindows::SetupRenderBufferFromDecoderOutput(I
 					ID3D11DeviceContext* DeviceContext = reinterpret_cast<ID3D11DeviceContext*>(InDecoderOutput->GetPlatformOutputHandle(EElectraDecoderPlatformOutputHandleType::DXDeviceContext));
 					if (Device && DeviceContext)
 					{
-						DecoderOutput->InitializeWithSharedTexture(Device, DeviceContext, Texture, FIntPoint(InDecoderOutput->GetWidth(), InDecoderOutput->GetHeight()), InOutBufferPropertes);
+						uint32 ViewIndex = 0;
+						DecoderOutput->InitializeWithSharedTexture(Device, DeviceContext, Texture, ViewIndex, FIntPoint(InDecoderOutput->GetWidth(), InDecoderOutput->GetHeight()), InOutBufferPropertes);
 						return true;
 					}
 				}
@@ -482,7 +488,6 @@ bool FElectraDecoderResourceManagerWindows::SetupRenderBufferFromDecoderOutput(I
 				return SetupRenderBufferFromDecoderOutputFromMFSample(InOutBufferToSetup, InOutBufferPropertes, InDecoderOutput, InPlatformSpecificResource);
 			}
 		}
-
 	}
 	return false;
 }
