Electra: Fixed PVS warnings

#rnx

diff --git a/Engine/Plugins/Media/ElectraCodecs/Source/ElectraDecoders/Private/DX/AudioDecoderAAC_DX.cpp b/Engine/Plugins/Media/ElectraCodecs/Source/ElectraDecoders/Private/DX/AudioDecoderAAC_DX.cpp
--- a/Engine/Plugins/Media/ElectraCodecs/Source/ElectraDecoders/Private/DX/AudioDecoderAAC_DX.cpp
+++ b/Engine/Plugins/Media/ElectraCodecs/Source/ElectraDecoders/Private/DX/AudioDecoderAAC_DX.cpp
@@ -874,7 +874,7 @@ bool FElectraAudioDecoderAAC_DX::ConvertDecoderOutput()
 		PostError(0, TEXT("Converted wave format struct is not at least a WAVEFORMATEXTENSIBLE"), ERRCODE_INTERNAL_FAILED_TO_CONVERT_OUTPUT_SAMPLE);
 		return false;
 	}
-	FMemory::Memcpy(&OutputWaveFormat, OutputWaveFormatPtr, sizeof(OutputWaveFormat));
+	FMemory::Memcpy(&OutputWaveFormat, OutputWaveFormatPtr, sizeof(OutputWaveFormat));	//-V512
 	if (OutputWaveFormatPtr->wFormatTag != WAVE_FORMAT_EXTENSIBLE)
 	{
 		PostError(0, TEXT("Converted wave format struct is not a WAVEFORMATEXTENSIBLE!"), ERRCODE_INTERNAL_FAILED_TO_CONVERT_OUTPUT_SAMPLE);
diff --git a/Engine/Plugins/Media/ElectraPlayer/Source/ElectraPlayerRuntime/Private/Runtime/Decoder/AudioDecoder.cpp b/Engine/Plugins/Media/ElectraPlayer/Source/ElectraPlayerRuntime/Private/Runtime/Decoder/AudioDecoder.cpp
--- a/Engine/Plugins/Media/ElectraPlayer/Source/ElectraPlayerRuntime/Private/Runtime/Decoder/AudioDecoder.cpp
+++ b/Engine/Plugins/Media/ElectraPlayer/Source/ElectraPlayerRuntime/Private/Runtime/Decoder/AudioDecoder.cpp
@@ -804,7 +804,7 @@ IElectraDecoder::EOutputStatus FAudioDecoderImpl::HandleOutput()
 
 					OutputBufferSampleProperties.Clear();
 					OutputBufferSampleProperties.Set(TEXT("num_channels"), FVariantValue((int64)ChannelMapper.GetNumTargetChannels()));
-					OutputBufferSampleProperties.Set(TEXT("byte_size"), FVariantValue((int64)(ChannelMapper.GetNumTargetChannels() * NumDecodedSamplesPerChannel * NumBytesPerSample)));
+					OutputBufferSampleProperties.Set(TEXT("byte_size"), FVariantValue((int64)((int64)NumDecodedSamplesPerChannel * ChannelMapper.GetNumTargetChannels() * NumBytesPerSample)));
 					OutputBufferSampleProperties.Set(TEXT("sample_rate"), FVariantValue((int64) SamplingRate));
 					OutputBufferSampleProperties.Set(TEXT("duration"), FVariantValue(Duration));
 					OutputBufferSampleProperties.Set(TEXT("pts"), FVariantValue(MatchingInput->AdjustedPTS));
diff --git a/Engine/Plugins/Media/ElectraPlayer/Source/ElectraPlayerRuntime/Private/Runtime/Decoder/Windows/WindowsElectraDecoderResourceManager.cpp b/Engine/Plugins/Media/ElectraPlayer/Source/ElectraPlayerRuntime/Private/Runtime/Decoder/Windows/WindowsElectraDecoderResourceManager.cpp
--- a/Engine/Plugins/Media/ElectraPlayer/Source/ElectraPlayerRuntime/Private/Runtime/Decoder/Windows/WindowsElectraDecoderResourceManager.cpp
+++ b/Engine/Plugins/Media/ElectraPlayer/Source/ElectraPlayerRuntime/Private/Runtime/Decoder/Windows/WindowsElectraDecoderResourceManager.cpp
@@ -433,7 +433,7 @@ bool FElectraDecoderResourceManagerWindows::SetupRenderBufferFromDecoderOutput(I
 				case EElectraDecoderPlatformPixelEncoding::CbY0CrY1:		DecPixEnc = EVideoDecoderPixelEncoding::CbY0CrY1; break;
 				case EElectraDecoderPlatformPixelEncoding::Y0CbY1Cr:		DecPixEnc = EVideoDecoderPixelEncoding::Y0CbY1Cr; break;
 				case EElectraDecoderPlatformPixelEncoding::ARGB_BigEndian:	DecPixEnc = EVideoDecoderPixelEncoding::ARGB_BigEndian; break;
-				default: DecPixEnc = DecPixEnc = EVideoDecoderPixelEncoding::Native; break;
+				default: DecPixEnc = EVideoDecoderPixelEncoding::Native; break;
 			}
 			
 			int32 Pitch = ImageBuffers->GetBufferPitchByIndex(0);
diff --git a/Engine/Plugins/Media/ElectraPlayer/Source/ElectraPlayerRuntime/Private/Runtime/Demuxer/ParserISO14496-12.cpp b/Engine/Plugins/Media/ElectraPlayer/Source/ElectraPlayerRuntime/Private/Runtime/Demuxer/ParserISO14496-12.cpp
--- a/Engine/Plugins/Media/ElectraPlayer/Source/ElectraPlayerRuntime/Private/Runtime/Demuxer/ParserISO14496-12.cpp
+++ b/Engine/Plugins/Media/ElectraPlayer/Source/ElectraPlayerRuntime/Private/Runtime/Demuxer/ParserISO14496-12.cpp
@@ -7742,7 +7742,7 @@ private:
 									{
 										bIsSupported = false;
 									}
-									else if (Error != UEMEDIA_ERROR_OK)
+									else if (Error != UEMEDIA_ERROR_OK) //-V547
 									{
 										return Error;
 									}
diff --git a/Engine/Plugins/Media/ElectraPlayer/Source/ElectraPlayerRuntime/Private/Runtime/Player/AdaptiveStreamingPlayerDecoders.cpp b/Engine/Plugins/Media/ElectraPlayer/Source/ElectraPlayerRuntime/Private/Runtime/Player/AdaptiveStreamingPlayerDecoders.cpp
--- a/Engine/Plugins/Media/ElectraPlayer/Source/ElectraPlayerRuntime/Private/Runtime/Player/AdaptiveStreamingPlayerDecoders.cpp
+++ b/Engine/Plugins/Media/ElectraPlayer/Source/ElectraPlayerRuntime/Private/Runtime/Player/AdaptiveStreamingPlayerDecoders.cpp
@@ -213,7 +213,7 @@ bool FAdaptiveStreamingPlayer::FindMatchingStreamInfo(FStreamCodecInformation& O
 				VideoCodecInfos.Sort([](const FStreamCodecInformation& a, const FStreamCodecInformation& b) { return a.GetProfileLevel() > b.GetProfileLevel(); });
 				VideoCodecInfos = VideoCodecInfos.FilterByPredicate([BestLevel=VideoCodecInfos[0].GetProfileLevel()](const FStreamCodecInformation& e) { return e.GetProfileLevel() >= BestLevel; });
 				// Sort by resolution
-				VideoCodecInfos.Sort([](const FStreamCodecInformation& a, const FStreamCodecInformation& b) { return a.GetResolution().Width*a.GetResolution().Height > a.GetResolution().Width*a.GetResolution().Height; });
+				VideoCodecInfos.Sort([](const FStreamCodecInformation& a, const FStreamCodecInformation& b) { return a.GetResolution().Width*a.GetResolution().Height > b.GetResolution().Width*b.GetResolution().Height; });
 
 				OutStreamInfo = VideoCodecInfos[0];
 				return true;
