Adding access APIs to retrieve stream meta data without the need to access the IMediaPlugin
#rnx

diff --git a/Engine/Plugins/Media/ElectraPlayer/Source/ElectraPlayerPlugin/Private/ElectraPlayerPlugin.cpp b/Engine/Plugins/Media/ElectraPlayer/Source/ElectraPlayerPlugin/Private/ElectraPlayerPlugin.cpp
--- a/Engine/Plugins/Media/ElectraPlayer/Source/ElectraPlayerPlugin/Private/ElectraPlayerPlugin.cpp
+++ b/Engine/Plugins/Media/ElectraPlayer/Source/ElectraPlayerPlugin/Private/ElectraPlayerPlugin.cpp
@@ -5,6 +5,7 @@
 #include "Misc/Optional.h"
 #include "IMediaEventSink.h"
 #include "IMediaOptions.h"
+#include "IMediaMetadataItem.h"
 #include "MediaSamples.h"
 #include "MediaPlayerOptions.h"
 
@@ -227,7 +228,7 @@ public:
 
 //-----------------------------------------------------------------------------
 
-class FStreamMetadataItem : public IMediaPlayer::IMetadataItem
+class FStreamMetadataItem : public IMediaMetadataItem
 {
 public:
 	FStreamMetadataItem(const TSharedPtr<Electra::IMediaStreamMetadata::IItem, ESPMode::ThreadSafe>& InItem) : Item(InItem.ToSharedRef())
@@ -737,17 +738,17 @@ void FElectraPlayerPlugin::TickInput(FTimespan DeltaTime, FTimespan Timecode)
 /**
 	Returns the current metadata, if any.
 */
-TSharedPtr<TMap<FString, TArray<TUniquePtr<IMediaPlayer::IMetadataItem>>>, ESPMode::ThreadSafe> FElectraPlayerPlugin::GetMediaMetadata() const
+TSharedPtr<TMap<FString, TArray<TUniquePtr<IMediaMetadataItem>>>, ESPMode::ThreadSafe> FElectraPlayerPlugin::GetMediaMetadata() const
 {
 	if (bMetadataChanged && Player.IsValid())
 	{
 		TSharedPtr<TMap<FString, TArray<TSharedPtr<Electra::IMediaStreamMetadata::IItem, ESPMode::ThreadSafe>>>, ESPMode::ThreadSafe> PlayerMeta = Player->GetMediaMetadata();
 		if (PlayerMeta.IsValid())
 		{
-			TSharedPtr<TMap<FString, TArray<TUniquePtr<IMediaPlayer::IMetadataItem>>>, ESPMode::ThreadSafe> NewMeta(new TMap<FString, TArray<TUniquePtr<IMediaPlayer::IMetadataItem>>>);
+			TSharedPtr<TMap<FString, TArray<TUniquePtr<IMediaMetadataItem>>>, ESPMode::ThreadSafe> NewMeta(new TMap<FString, TArray<TUniquePtr<IMediaMetadataItem>>>);
 			for(auto& PlayerMetaItem : *PlayerMeta)
 			{
-				TArray<TUniquePtr<IMediaPlayer::IMetadataItem>>& NewItemList = NewMeta->Emplace(PlayerMetaItem.Key);
+				TArray<TUniquePtr<IMediaMetadataItem>>& NewItemList = NewMeta->Emplace(PlayerMetaItem.Key);
 				for(auto& PlayerMetaListItem : PlayerMetaItem.Value)
 				{
 					if (PlayerMetaListItem.IsValid())
diff --git a/Engine/Plugins/Media/ElectraPlayer/Source/ElectraPlayerPlugin/Private/ElectraPlayerPlugin.h b/Engine/Plugins/Media/ElectraPlayer/Source/ElectraPlayerPlugin/Private/ElectraPlayerPlugin.h
--- a/Engine/Plugins/Media/ElectraPlayer/Source/ElectraPlayerPlugin/Private/ElectraPlayerPlugin.h
+++ b/Engine/Plugins/Media/ElectraPlayer/Source/ElectraPlayerPlugin/Private/ElectraPlayerPlugin.h
@@ -98,7 +98,7 @@ public:
 	void TickInput(FTimespan DeltaTime, FTimespan Timecode) override;
 	void SetLastAudioRenderedSampleTime(FTimespan SampleTime) override;
 
-	TSharedPtr<TMap<FString, TArray<TUniquePtr<IMediaPlayer::IMetadataItem>>>, ESPMode::ThreadSafe> GetMediaMetadata() const override;
+	TSharedPtr<TMap<FString, TArray<TUniquePtr<IMediaMetadataItem>>>, ESPMode::ThreadSafe> GetMediaMetadata() const override;
 
 	bool GetPlayerFeatureFlag(EFeatureFlag flag) const override;
 
@@ -191,7 +191,7 @@ private:
 
 	/** Current player stream metadata */
 	mutable bool bMetadataChanged = false;
-	mutable TSharedPtr<TMap<FString, TArray<TUniquePtr<IMediaPlayer::IMetadataItem>>>, ESPMode::ThreadSafe> CurrentMetadata;
+	mutable TSharedPtr<TMap<FString, TArray<TUniquePtr<IMediaMetadataItem>>>, ESPMode::ThreadSafe> CurrentMetadata;
 
 
 	/** Output sample pools */
diff --git a/Engine/Source/Runtime/Media/Private/MediaModule.cpp b/Engine/Source/Runtime/Media/Private/MediaModule.cpp
--- a/Engine/Source/Runtime/Media/Private/MediaModule.cpp
+++ b/Engine/Source/Runtime/Media/Private/MediaModule.cpp
@@ -12,7 +12,6 @@
 
 #include "IMediaCaptureSupport.h"
 #include "IMediaPlayerFactory.h"
-#include "IMediaInfo.h"
 #include "IMediaTimeSource.h"
 #include "MediaClock.h"
 #include "MediaTicker.h"
diff --git a/Engine/Source/Runtime/Media/Public/IMediaInfo.h b/Engine/Source/Runtime/Media/Public/IMediaInfo.h
deleted file mode 100644
diff --git a/Engine/Source/Runtime/Media/Public/IMediaMetadataItem.h b/Engine/Source/Runtime/Media/Public/IMediaMetadataItem.h
new file mode 100644
--- /dev/null
+++ b/Engine/Source/Runtime/Media/Public/IMediaMetadataItem.h
@@ -0,0 +1,35 @@
+// Copyright Epic Games, Inc. All Rights Reserved.
+
+#pragma once
+
+#include "Containers/UnrealString.h"
+#include "Misc/Variant.h"
+
+/**
+ * This structure describes a metadata item.
+ */
+class IMediaMetadataItem
+{
+public:
+	virtual ~IMediaMetadataItem() = default;
+
+	/*
+	  Language this item is described in.
+	  This should a three letter ISO 639-2 code. If no language is specified this item is intended
+	  to be used when no other item exists in the user's preferred language.
+	*/
+	virtual const FString& GetLanguageCode() const = 0;
+
+	/*
+	  Optional mime type of the item described by Value.
+	  This should be set when Value is of type EVariantTypes::ByteArray
+	  For example, if this item is a binary JPEG image the mime type should be set as "image/jpeg"
+	  to identify the type of the payload.
+	*/
+	virtual const FString& GetMimeType() const = 0;
+
+	/*
+	  The value of the metadata item.
+	*/
+	virtual const FVariant& GetValue() const = 0;
+};
diff --git a/Engine/Source/Runtime/Media/Public/IMediaModule.h b/Engine/Source/Runtime/Media/Public/IMediaModule.h
--- a/Engine/Source/Runtime/Media/Public/IMediaModule.h
+++ b/Engine/Source/Runtime/Media/Public/IMediaModule.h
@@ -14,7 +14,6 @@ class IMediaClock;
 class IMediaPlayerFactory;
 class IMediaTicker;
 class IMediaTimeSource;
-class IMediaInfo;
 class IMediaPlayerLifecycleManagerDelegate;
 
 /**
diff --git a/Engine/Source/Runtime/Media/Public/IMediaPlayer.h b/Engine/Source/Runtime/Media/Public/IMediaPlayer.h
--- a/Engine/Source/Runtime/Media/Public/IMediaPlayer.h
+++ b/Engine/Source/Runtime/Media/Public/IMediaPlayer.h
@@ -21,6 +21,7 @@ class IMediaOptions;
 class IMediaSamples;
 class IMediaTracks;
 class IMediaView;
+class IMediaMetadataItem;
 
 struct FGuid;
 struct FMediaPlayerOptions;
@@ -196,35 +197,6 @@ public:
 		return FVariant();
 	}
 
-	/**
-	 * This structure describes a metadata item.
-	 */
-	class IMetadataItem
-	{
-	public:
-		virtual ~IMetadataItem() = default;
-
-		/*
-		  Language this item is described in.
-		  This should a three letter ISO 639-2 code. If no language is specified this item is intended
-		  to be used when no other item exists in the user's preferred language.
-		*/
-		virtual const FString& GetLanguageCode() const = 0;
-		
-		/*
-		  Optional mime type of the item described by Value.
-		  This should be set when Value is of type EVariantTypes::ByteArray
-		  For example, if this item is a binary JPEG image the mime type should be set as "image/jpeg"
-		  to identify the type of the payload.
-		*/
-		virtual const FString& GetMimeType() const = 0;
-
-		/*
-		  The value of the metadata item.
-		*/
-		virtual const FVariant& GetValue() const = 0;
-	};
-
 	/**
 	 * Gets the current metadata of the media source.
 	 * 
@@ -237,7 +209,7 @@ public:
 	 * 
 	 * Metadata may change over time. Its presence or change is reported by a MetadataChanged event.
 	 */
-	virtual TSharedPtr<TMap<FString, TArray<TUniquePtr<IMetadataItem>>>, ESPMode::ThreadSafe> GetMediaMetadata() const
+	virtual TSharedPtr<TMap<FString, TArray<TUniquePtr<IMediaMetadataItem>>>, ESPMode::ThreadSafe> GetMediaMetadata() const
 	{
 		return nullptr;
 	}
diff --git a/Engine/Source/Runtime/MediaAssets/Private/Assets/MediaPlayer.cpp b/Engine/Source/Runtime/MediaAssets/Private/Assets/MediaPlayer.cpp
--- a/Engine/Source/Runtime/MediaAssets/Private/Assets/MediaPlayer.cpp
+++ b/Engine/Source/Runtime/MediaAssets/Private/Assets/MediaPlayer.cpp
@@ -228,6 +228,10 @@ FVariant UMediaPlayer::GetMediaInfo(FName InfoName) const
 	return PlayerFacade->GetMediaInfo(InfoName);
 }
 
+TSharedPtr<TMap<FString, TArray<TUniquePtr<IMediaMetadataItem>>>, ESPMode::ThreadSafe> UMediaPlayer::GetMediaMetadata() const
+{
+	return PlayerFacade->GetMediaMetadata();
+}
 
 TSharedRef<FMediaPlayerFacade, ESPMode::ThreadSafe> UMediaPlayer::GetPlayerFacade() const
 {
diff --git a/Engine/Source/Runtime/MediaAssets/Public/MediaPlayer.h b/Engine/Source/Runtime/MediaAssets/Public/MediaPlayer.h
--- a/Engine/Source/Runtime/MediaAssets/Public/MediaPlayer.h
+++ b/Engine/Source/Runtime/MediaAssets/Public/MediaPlayer.h
@@ -23,6 +23,7 @@
 class FMediaPlayerFacade;
 class IMediaPlayer;
 class IMediaModule;
+class IMediaMetadataItem;
 class UMediaPlaylist;
 class UMediaSource;
 
@@ -1031,6 +1032,14 @@ public:
 		return false;
 	}
 
+	/**
+	 * Get meta data contained in the current stream
+	 *
+	 * @return Map with arrays of IMediaMetaDataItem entries describing any metadata found in the current stream
+	 * @note Listen to EMediaEvent::MetadataChanged to catch updates to this data
+	 */
+	TSharedPtr<TMap<FString, TArray<TUniquePtr<IMediaMetadataItem>>>, ESPMode::ThreadSafe> GetMediaMetadata() const;
+
 	/**
 	 * Get the media player facade that manages low-level media players
 	 *
diff --git a/Engine/Source/Runtime/MediaUtils/Private/MediaPlayerFacade.cpp b/Engine/Source/Runtime/MediaUtils/Private/MediaPlayerFacade.cpp
--- a/Engine/Source/Runtime/MediaUtils/Private/MediaPlayerFacade.cpp
+++ b/Engine/Source/Runtime/MediaUtils/Private/MediaPlayerFacade.cpp
@@ -356,6 +356,17 @@ FText FMediaPlayerFacade::GetMediaName() const
 }
 
 
+TSharedPtr<TMap<FString, TArray<TUniquePtr<IMediaMetadataItem>>>, ESPMode::ThreadSafe> FMediaPlayerFacade::GetMediaMetadata() const
+{
+	TSharedPtr<IMediaPlayer, ESPMode::ThreadSafe> CurrentPlayer(Player);
+	if (!CurrentPlayer.IsValid())
+	{
+		return nullptr;
+	}
+	return CurrentPlayer->GetMediaMetadata();
+}
+
+
 int32 FMediaPlayerFacade::GetNumTracks(EMediaTrackType TrackType) const
 {
 	TSharedPtr<IMediaPlayer, ESPMode::ThreadSafe> CurrentPlayer(Player);
diff --git a/Engine/Source/Runtime/MediaUtils/Public/MediaPlayerFacade.h b/Engine/Source/Runtime/MediaUtils/Public/MediaPlayerFacade.h
--- a/Engine/Source/Runtime/MediaUtils/Public/MediaPlayerFacade.h
+++ b/Engine/Source/Runtime/MediaUtils/Public/MediaPlayerFacade.h
@@ -36,6 +36,7 @@ class IMediaOptions;
 class IMediaPlayer;
 class IMediaPlayerFactory;
 class IMediaSamples;
+class IMediaMetadataItem;
 
 enum class EMediaEvent;
 enum class EMediaCacheState;
@@ -252,6 +253,14 @@ public:
 	 */
 	FText GetMediaName() const;
 
+	/**
+	 * Get meta data contained in the current stream
+	 *
+	 * @return Map with arrays of IMediaMetaDataItem entries describing any metadata found in the current stream
+	 * @note Listen to EMediaEvent::MetadataChanged to catch updates to this data
+	 */
+	TSharedPtr<TMap<FString, TArray<TUniquePtr<IMediaMetadataItem>>>, ESPMode::ThreadSafe> GetMediaMetadata() const;
+
 	/**
 	 * Get the number of tracks of the given type.
 	 *
