ElectraPlayer: Permitting playback rates other than 0.0 and 1.0

#rnx

diff --git a/Engine/Plugins/Media/ElectraPlayer/Source/ElectraPlayerPlugin/Private/ElectraPlayerPlugin.cpp b/Engine/Plugins/Media/ElectraPlayer/Source/ElectraPlayerPlugin/Private/ElectraPlayerPlugin.cpp
--- a/Engine/Plugins/Media/ElectraPlayer/Source/ElectraPlayerPlugin/Private/ElectraPlayerPlugin.cpp
+++ b/Engine/Plugins/Media/ElectraPlayer/Source/ElectraPlayerPlugin/Private/ElectraPlayerPlugin.cpp
@@ -58,6 +58,9 @@ FElectraPlayerPlugin::FElectraPlayerPlugin()
 	static_assert((int32)EMediaTrackType::Audio == (int32)IElectraPlayerInterface::EPlayerTrackType::Audio, "check alignment of both enums");
 	static_assert((int32)EMediaTrackType::Video == (int32)IElectraPlayerInterface::EPlayerTrackType::Video, "check alignment of both enums");
 
+	static_assert((int32)EMediaRateThinning::Unthinned == (int32)IElectraPlayerInterface::EPlayRateType::Unthinned, "check alignment of both enums");
+	static_assert((int32)EMediaRateThinning::Thinned == (int32)IElectraPlayerInterface::EPlayRateType::Thinned, "check alignment of both enums");
+
 	static_assert(IMediaPlayerLifecycleManagerDelegate::ResourceFlags_Decoder == IElectraPlayerInterface::ResourceFlags_Decoder, "check alignment of both enums");
 	static_assert(IMediaPlayerLifecycleManagerDelegate::ResourceFlags_OutputBuffers == IElectraPlayerInterface::ResourceFlags_OutputBuffers, "check alignment of both enums");
 	static_assert(IMediaPlayerLifecycleManagerDelegate::ResourceFlags_Any == IElectraPlayerInterface::ResourceFlags_Any, "check alignment of both enums");
@@ -892,10 +895,7 @@ bool FElectraPlayerPlugin::SetLooping(bool bLooping)
  */
 TRangeSet<float> FElectraPlayerPlugin::GetSupportedRates(EMediaRateThinning Thinning) const
 {
-	TRangeSet<float> Res;
-	Res.Add(TRange<float>{1.0f}); // only normal (real-time) playback rate
-	Res.Add(TRange<float>{0.0f}); // and pause
-	return Res;
+	return Player->GetSupportedRates(Thinning == EMediaRateThinning::Thinned ? IElectraPlayerInterface::EPlayRateType::Thinned : IElectraPlayerInterface::EPlayRateType::Unthinned);
 }
 
 
diff --git a/Engine/Plugins/Media/ElectraPlayer/Source/ElectraPlayerRuntime/Private/ElectraPlayer.cpp b/Engine/Plugins/Media/ElectraPlayer/Source/ElectraPlayerRuntime/Private/ElectraPlayer.cpp
--- a/Engine/Plugins/Media/ElectraPlayer/Source/ElectraPlayerRuntime/Private/ElectraPlayer.cpp
+++ b/Engine/Plugins/Media/ElectraPlayer/Source/ElectraPlayerRuntime/Private/ElectraPlayer.cpp
@@ -936,20 +936,6 @@ bool FElectraPlayer::CanPresentAudioFrames(uint64 NumFrames)
 	return false;
 }
 
-float FElectraPlayer::GetRate() const
-{
-	return PlayerState.GetRate();
-/*
-	if (PlayerState.bUseInternal)
-	{
-		return 0.0f;
-	}
-	else
-	{
-		return CurrentPlayer.Get() && CurrentPlayer->AdaptivePlayer->IsPlaying() ? 1.0f : 0.0f;
-	}
-*/
-}
 
 FElectraPlayer::EPlayerState FElectraPlayer::GetState() const
 {
@@ -1055,6 +1041,31 @@ FTimespan FElectraPlayer::GetSeekableDuration() const
 }
 
 
+TRangeSet<float> FElectraPlayer::GetSupportedRates(EPlayRateType InPlayRateType) const
+{
+	TRangeSet<float> Res;
+	TSharedPtr<FInternalPlayerImpl, ESPMode::ThreadSafe> LockedPlayer = CurrentPlayer;
+	if (LockedPlayer.IsValid() && LockedPlayer->AdaptivePlayer.IsValid())
+	{
+		TArray<TRange<double>> SupportedRanges;
+		LockedPlayer->AdaptivePlayer->GetSupportedRates(InPlayRateType == IElectraPlayerInterface::EPlayRateType::Unthinned ? IAdaptiveStreamingPlayer::EPlaybackRateType::Unthinned : IAdaptiveStreamingPlayer::EPlaybackRateType::Thinned).GetRanges(SupportedRanges);
+		for(auto &Rate : SupportedRanges)
+		{
+			TRange<float> r;
+			if (Rate.HasLowerBound())
+			{
+				r.SetLowerBound(TRange<float>::BoundsType::Inclusive((float) Rate.GetLowerBoundValue()));
+			}
+			if (Rate.HasUpperBound())
+			{
+				r.SetUpperBound(TRange<float>::BoundsType::Inclusive((float) Rate.GetUpperBoundValue()));
+			}
+			Res.Add(r);
+		}
+	}
+	return Res;
+}
+
 bool FElectraPlayer::SetRate(float Rate)
 {
 	//UE_LOG(LogElectraPlayer, Verbose, TEXT("[%p][%p] IMediaPlayer::SetRate(%.3f)"), this, CurrentPlayer.Get(), Rate);
@@ -1074,12 +1085,30 @@ bool FElectraPlayer::SetRate(float Rate)
 				TriggerFirstSeekIfNecessary();
 				CurrentPlayer->AdaptivePlayer->Resume();
 			}
+			IAdaptiveStreamingPlayer::FTrickplayParams Params;
+			CurrentPlayer->AdaptivePlayer->SetPlayRate((double) Rate, Params);
 		}
 		return true;
 	}
 	return false;
 }
 
+float FElectraPlayer::GetRate() const
+{
+	return PlayerState.GetRate();
+/*
+	if (PlayerState.bUseInternal)
+	{
+		return 0.0f;
+	}
+	else
+	{
+		return CurrentPlayer.Get() && CurrentPlayer->AdaptivePlayer->IsPlaying() ? 1.0f : 0.0f;
+	}
+*/
+}
+
+
 void FElectraPlayer::TriggerFirstSeekIfNecessary()
 {
 	if (!bInitialSeekPerformed)
@@ -3081,8 +3110,15 @@ bool FElectraPlayer::MediaStateOnPlay()
 
 	CSV_EVENT(ElectraPlayer, TEXT("MediaStateOnPlay"));
 
+	double CurrentRate = 1.0;
+	TSharedPtr<FInternalPlayerImpl, ESPMode::ThreadSafe> LockedPlayer = CurrentPlayer;
+	if (LockedPlayer.IsValid() && LockedPlayer->AdaptivePlayer.IsValid())
+	{
+		CurrentRate = LockedPlayer->AdaptivePlayer->GetPlayRate();
+	}
+
 	PlayerState.State = EPlayerState::Playing;
-	PlayerState.SetPlayRateFromPlayer(1.0f);
+	PlayerState.SetPlayRateFromPlayer(CurrentRate);
 
 	DeferredEvents.Enqueue(IElectraPlayerAdapterDelegate::EPlayerEvent::PlaybackResumed);
 	return true;
diff --git a/Engine/Plugins/Media/ElectraPlayer/Source/ElectraPlayerRuntime/Private/ElectraPlayer.h b/Engine/Plugins/Media/ElectraPlayer/Source/ElectraPlayerRuntime/Private/ElectraPlayer.h
--- a/Engine/Plugins/Media/ElectraPlayer/Source/ElectraPlayerRuntime/Private/ElectraPlayer.h
+++ b/Engine/Plugins/Media/ElectraPlayer/Source/ElectraPlayerRuntime/Private/ElectraPlayer.h
@@ -111,6 +111,7 @@ public:
 	void SetPlaybackRange(const FPlaybackRange& InPlaybackRange) override;
 	void GetPlaybackRange(FPlaybackRange& OutPlaybackRange) const override;
 
+	TRangeSet<float> GetSupportedRates(EPlayRateType InPlayRateType) const override;
 	float GetRate() const override;
 	bool SetRate(float Rate) override;
 
diff --git a/Engine/Plugins/Media/ElectraPlayer/Source/ElectraPlayerRuntime/Private/Runtime/Player/AdaptiveStreamingPlayer.cpp b/Engine/Plugins/Media/ElectraPlayer/Source/ElectraPlayerRuntime/Private/Runtime/Player/AdaptiveStreamingPlayer.cpp
--- a/Engine/Plugins/Media/ElectraPlayer/Source/ElectraPlayerRuntime/Private/Runtime/Player/AdaptiveStreamingPlayer.cpp
+++ b/Engine/Plugins/Media/ElectraPlayer/Source/ElectraPlayerRuntime/Private/Runtime/Player/AdaptiveStreamingPlayer.cpp
@@ -75,7 +75,6 @@ FAdaptiveStreamingPlayer::FAdaptiveStreamingPlayer(const IAdaptiveStreamingPlaye
 	PipelineState  		    = EPipelineState::ePipeline_Stopped;
 	DecoderState   		    = EDecoderState::eDecoder_Paused;
 	StreamState			    = EStreamState::eStream_Running;
-	PlaybackRate   		    = 0.0;
 	RenderRateScale			= 1.0;
 	StreamReaderHandler     = nullptr;
 	bStreamingHasStarted	= false;
@@ -568,6 +567,23 @@ void FAdaptiveStreamingPlayer::GetPlaybackRange(FPlaybackRange& OutPlaybackRange
 }
 
 
+TRangeSet<double> FAdaptiveStreamingPlayer::GetSupportedRates(EPlaybackRateType InForPlayRateType)
+{
+	return PlaybackState.GetPlaybackRates(InForPlayRateType);
+}
+
+void FAdaptiveStreamingPlayer::SetPlayRate(double InDesiredPlayRate, const FTrickplayParams& InParameters)
+{
+	PlaybackState.SetDesiredPlayRate(InDesiredPlayRate, InParameters);
+}
+
+double FAdaptiveStreamingPlayer::GetPlayRate() const
+{
+	// Return the *desired* play rate.
+	return PlaybackState.GetDesiredPlayRate();
+}
+
+
 //-----------------------------------------------------------------------------
 /**
  * Starts loading the manifest / main playlist.
@@ -2341,7 +2357,6 @@ void FAdaptiveStreamingPlayer::HandleNewOutputData()
 				{
 					PipelineState = EPipelineState::ePipeline_Stopped;
 					CurrentState = EPlayerState::eState_Paused;
-					PlaybackRate = 0.0;
 					PrerollVars.Clear();
 					SeekVars.ClearWorkVars();
 					SeekVars.SetFinished();
@@ -2381,7 +2396,6 @@ void FAdaptiveStreamingPlayer::HandleNewOutputData()
 					PipelineState = EPipelineState::ePipeline_Stopped;
 					CurrentState = EPlayerState::eState_Paused;
 					DecoderState = EDecoderState::eDecoder_Running;
-					PlaybackRate = 0.0;
 					PrerollVars.Clear();
 					DispatchBufferingEvent(false, LastBufferingState);
 					DispatchEvent(FMetricEvent::ReportPrerollEnd());
@@ -4262,7 +4276,6 @@ void FAdaptiveStreamingPlayer::InternalPause()
 	// Pause playing.
 	if (CurrentState != EPlayerState::eState_Error)
 	{
-		PlaybackRate = 0.0;
 		CurrentState = EPlayerState::eState_Paused;
 		PlaybackState.SetPausedAndPlaying(true, false);
 		DispatchEvent(FMetricEvent::ReportPlaybackPaused());
@@ -4297,7 +4310,6 @@ void FAdaptiveStreamingPlayer::InternalResume()
 
 		// Resume playing.
 		CurrentState = EPlayerState::eState_Playing;
-		PlaybackRate = 1.0;
 		PlaybackState.SetPausedAndPlaying(false, true);
 		PostrollVars.Clear();
 		DispatchEvent(FMetricEvent::ReportPlaybackResumed());
@@ -4353,7 +4365,6 @@ void FAdaptiveStreamingPlayer::InternalStop(bool bHoldCurrentFrame)
 	NextLoopStates.Empty();
 	SeekVars.Reset();
 
-	PlaybackRate = 0.0;
 	CurrentState = EPlayerState::eState_Paused;
 	PlaybackState.SetPausedAndPlaying(false, false);
 	PlaybackState.SetPlaybackEndAtTime(FTimeValue::GetInvalid());
@@ -4491,9 +4502,8 @@ void FAdaptiveStreamingPlayer::InternalClose()
 	ExternalDataReader.Reset();
 
 	// Reset remaining internal state
-	CurrentState   	= EPlayerState::eState_Idle;
-	PlaybackRate   	= 0.0;
-	StreamState		= EStreamState::eStream_Running;
+	CurrentState = EPlayerState::eState_Idle;
+	StreamState = EStreamState::eStream_Running;
 	RebufferCause = ERebufferCause::None;
 	LastBufferingState = EPlayerState::eState_Buffering;
 	RebufferDetectedAtPlayPos.SetToInvalid();
diff --git a/Engine/Plugins/Media/ElectraPlayer/Source/ElectraPlayerRuntime/Private/Runtime/Player/AdaptiveStreamingPlayer.h b/Engine/Plugins/Media/ElectraPlayer/Source/ElectraPlayerRuntime/Private/Runtime/Player/AdaptiveStreamingPlayer.h
--- a/Engine/Plugins/Media/ElectraPlayer/Source/ElectraPlayerRuntime/Private/Runtime/Player/AdaptiveStreamingPlayer.h
+++ b/Engine/Plugins/Media/ElectraPlayer/Source/ElectraPlayerRuntime/Private/Runtime/Player/AdaptiveStreamingPlayer.h
@@ -214,6 +214,43 @@ public:
 	virtual void SetLooping(const FLoopParam& InLoopParams) = 0;
 
 
+	//-------------------------------------------------------------------------
+	// "Trickplay" related functions
+	//
+	enum class EPlaybackRateType
+	{
+		// Smooth playback, not dropping any frames.
+		Unthinned,
+		// Dropping frames
+		Thinned
+	};
+	/**
+	 * Returns ranges the playback rate can be set to.
+	 * There are two types, one `Unthinned`, where (ideally) no frames will be dropped and
+	 * `Thinned`, where frames will be dropped to maintain the playback rate.
+	 * 
+	 * Either way, the player will not resample audio or generate interpolated frames of video.
+	 * Decoded samples will be delivered to the renderer as they are. It is also up to the
+	 * renderers to consume the data at a faster or slower rate to actually realize the desired
+	 * playback rate.
+	 * 
+	 * The range of supported rates depends on the type of media and the decoder capability to
+	 * decode faster than realtime. If the media allows for adaptive bitrate selection the player
+	 * will choose the stream to play back accordingly.
+	 * In thinned mode it will need to drop and not decode samples, so the renderer will not
+	 * receive all the possible data.
+	 * 
+	 * Live streams will only allow for rates of 0.0 (pause) and 1.0 (real time play forward).
+	 */
+	virtual TRangeSet<double> GetSupportedRates(EPlaybackRateType InForPlayRateType) = 0;
+
+	struct FTrickplayParams
+	{
+	};
+	virtual void SetPlayRate(double InDesiredPlayRate, const FTrickplayParams& InParameters) = 0;
+	virtual double GetPlayRate() const = 0;
+
+
 	//-------------------------------------------------------------------------
 	// Error related functions
 	//
diff --git a/Engine/Plugins/Media/ElectraPlayer/Source/ElectraPlayerRuntime/Private/Runtime/Player/AdaptiveStreamingPlayerInternal.h b/Engine/Plugins/Media/ElectraPlayer/Source/ElectraPlayerRuntime/Private/Runtime/Player/AdaptiveStreamingPlayerInternal.h
--- a/Engine/Plugins/Media/ElectraPlayer/Source/ElectraPlayerRuntime/Private/Runtime/Player/AdaptiveStreamingPlayerInternal.h
+++ b/Engine/Plugins/Media/ElectraPlayer/Source/ElectraPlayerRuntime/Private/Runtime/Player/AdaptiveStreamingPlayerInternal.h
@@ -201,6 +201,10 @@ struct FPlaybackState
 		LoopState = {};
 		ActivePlaybackRange.Reset();
 		NewPlaybackRange.Reset();
+		UnthinnedPlaybackRates.Empty();
+		ThinnedPlaybackRates.Empty();
+		CurrentPlaybackRate = 0.0;
+		DesiredPlaybackRate = 0.0;
 		bHaveMetadata = false;
 		bHasEnded = false;
 		bIsSeeking = false;
@@ -223,6 +227,10 @@ struct FPlaybackState
 	IAdaptiveStreamingPlayer::FLoopState	LoopState;
 	FTimeRange								ActivePlaybackRange;
 	FTimeRange								NewPlaybackRange;
+	TRangeSet<double>						UnthinnedPlaybackRates;
+	TRangeSet<double>						ThinnedPlaybackRates;
+	double									CurrentPlaybackRate;
+	double									DesiredPlaybackRate;
 	bool									bHaveMetadata;
 	bool									bHasEnded;
 	bool									bIsSeeking;
@@ -397,6 +405,13 @@ struct FPlaybackState
 		FScopeLock lock(&Lock);
 		bIsPaused  = bInIsPaused;
 		bIsPlaying = bInIsPlaying;
+		// If asked to play, but no desired playback rate has been set yet, assume 1.0.
+		// It is thought to be a user error to resume playback with no play rate.
+		if (bIsPlaying && DesiredPlaybackRate == 0.0)
+		{
+			DesiredPlaybackRate = 1.0;
+		}
+		CurrentPlaybackRate = bIsPaused ? 0.0 : bIsPlaying ? DesiredPlaybackRate : 0.0;
 	}
 
 	bool SetTrackMetadata(const TArray<FTrackMetadata> &InVideoTracks, const TArray<FTrackMetadata>& InAudioTracks, const TArray<FTrackMetadata>& InSubtitleTracks)
@@ -513,6 +528,45 @@ struct FPlaybackState
 		return bPlayrangeHasChanged;
 	}
 
+	void SetPlaybackRates(IAdaptiveStreamingPlayer::EPlaybackRateType InForType, const TRangeSet<double>& InRates)
+	{
+		FScopeLock lock(&Lock);
+		if (InForType == IAdaptiveStreamingPlayer::EPlaybackRateType::Unthinned)
+		{
+			UnthinnedPlaybackRates = InRates;
+		}
+		else
+		{
+			ThinnedPlaybackRates = InRates;
+		}
+	}
+	TRangeSet<double> GetPlaybackRates(IAdaptiveStreamingPlayer::EPlaybackRateType InForType)
+	{
+		FScopeLock lock(&Lock);
+		return InForType == IAdaptiveStreamingPlayer::EPlaybackRateType::Unthinned ? UnthinnedPlaybackRates : ThinnedPlaybackRates;
+	}
+
+	void SetDesiredPlayRate(double InDesiredPlayRate, const IAdaptiveStreamingPlayer::FTrickplayParams& /*InParameters*/)
+	{
+		FScopeLock lock(&Lock);
+		DesiredPlaybackRate = InDesiredPlayRate;
+	}
+	double GetDesiredPlayRate() const
+	{
+		FScopeLock lock(&Lock);
+		return DesiredPlaybackRate;
+	}
+	void SetCurrentPlayRate(double InRate)
+	{
+		FScopeLock lock(&Lock);
+		CurrentPlaybackRate = InRate;
+	}
+	double GetCurrentPlayRate() const
+	{
+		FScopeLock lock(&Lock);
+		return CurrentPlaybackRate;
+	}
+
 
 	void SetLoopStateHasChanged(bool bWasChanged)
 	{
@@ -954,6 +1008,9 @@ public:
 	void SetPlaybackRange(const FPlaybackRange& InPlaybackRange) override;
 	void GetPlaybackRange(FPlaybackRange& OutPlaybackRange) override;
 	void SetLooping(const FLoopParam& InLoopParams) override;
+	TRangeSet<double> GetSupportedRates(EPlaybackRateType InForPlayRateType) override;
+	void SetPlayRate(double InDesiredPlayRate, const FTrickplayParams& InParameters) override;
+	double GetPlayRate() const override;
 
 	FErrorDetail GetError() const override;
 
@@ -2163,7 +2220,6 @@ private:
 	FPostrollVars														PostrollVars;
 	FSeekVars															SeekVars;
 	EPlayerState														LastBufferingState;
-	double																PlaybackRate;
 	double																RenderRateScale;
 	FTimeValue															RebufferDetectedAtPlayPos;
 	ERebufferCause														RebufferCause;
diff --git a/Engine/Plugins/Media/ElectraPlayer/Source/ElectraPlayerRuntime/Private/Runtime/Player/AdaptiveStreamingPlayerManifest.cpp b/Engine/Plugins/Media/ElectraPlayer/Source/ElectraPlayerRuntime/Private/Runtime/Player/AdaptiveStreamingPlayerManifest.cpp
--- a/Engine/Plugins/Media/ElectraPlayer/Source/ElectraPlayerRuntime/Private/Runtime/Player/AdaptiveStreamingPlayerManifest.cpp
+++ b/Engine/Plugins/Media/ElectraPlayer/Source/ElectraPlayerRuntime/Private/Runtime/Player/AdaptiveStreamingPlayerManifest.cpp
@@ -301,6 +301,9 @@ bool FAdaptiveStreamingPlayer::SelectManifest()
 			NewPresentation->GetTrackMetadata(SubtitleTrackMetadata, EStreamType::Subtitle);
 			PlaybackState.SetTrackMetadata(VideoTrackMetadata, AudioTrackMetadata, SubtitleTrackMetadata);
 			PlaybackState.SetHaveMetadata(true);
+			// Get the supported playback rates for thinned and unthinned playback.
+			PlaybackState.SetPlaybackRates(EPlaybackRateType::Unthinned, NewPresentation->GetPossiblePlaybackRates(IManifest::EPlayRateType::UnthinnedRate));
+			PlaybackState.SetPlaybackRates(EPlaybackRateType::Thinned, NewPresentation->GetPossiblePlaybackRates(IManifest::EPlayRateType::ThinnedRate));
 
 			Manifest = NewPresentation;
 
diff --git a/Engine/Plugins/Media/ElectraPlayer/Source/ElectraPlayerRuntime/Private/Runtime/Player/DASH/ManifestDASH.cpp b/Engine/Plugins/Media/ElectraPlayer/Source/ElectraPlayerRuntime/Private/Runtime/Player/DASH/ManifestDASH.cpp
--- a/Engine/Plugins/Media/ElectraPlayer/Source/ElectraPlayerRuntime/Private/Runtime/Player/DASH/ManifestDASH.cpp
+++ b/Engine/Plugins/Media/ElectraPlayer/Source/ElectraPlayerRuntime/Private/Runtime/Player/DASH/ManifestDASH.cpp
@@ -333,6 +333,14 @@ TSharedPtrTS<IProducerReferenceTimeInfo> FManifestDASH::GetProducerReferenceTime
 	return Manifest.IsValid() ? Manifest->GetProducerReferenceTimeElement(ID) : nullptr;
 }
 
+TRangeSet<double> FManifestDASH::GetPossiblePlaybackRates(EPlayRateType InForType) const
+{
+	TRangeSet<double> Ranges;
+	Ranges.Add(TRange<double>{1.0}); // normal (real-time) playback rate
+	Ranges.Add(TRange<double>{0.0}); // and pause
+	return Ranges;
+}
+
 
 void FManifestDASH::GetTrackMetadata(TArray<FTrackMetadata>& OutMetadata, EStreamType StreamType) const
 {
diff --git a/Engine/Plugins/Media/ElectraPlayer/Source/ElectraPlayerRuntime/Private/Runtime/Player/DASH/ManifestDASH.h b/Engine/Plugins/Media/ElectraPlayer/Source/ElectraPlayerRuntime/Private/Runtime/Player/DASH/ManifestDASH.h
--- a/Engine/Plugins/Media/ElectraPlayer/Source/ElectraPlayerRuntime/Private/Runtime/Player/DASH/ManifestDASH.h
+++ b/Engine/Plugins/Media/ElectraPlayer/Source/ElectraPlayerRuntime/Private/Runtime/Player/DASH/ManifestDASH.h
@@ -30,6 +30,7 @@ public:
 	void ClearDefaultStartTime() override;
 	FTimeValue GetMinBufferTime() const override;
 	FTimeValue GetDesiredLiveLatency() const override;
+	TRangeSet<double> GetPossiblePlaybackRates(EPlayRateType InForType) const override;
 	TSharedPtrTS<IProducerReferenceTimeInfo> GetProducerReferenceTimeInfo(int64 ID) const override;
 	void GetTrackMetadata(TArray<FTrackMetadata>& OutMetadata, EStreamType StreamType) const override;
 	void UpdateDynamicRefetchCounter() override;
diff --git a/Engine/Plugins/Media/ElectraPlayer/Source/ElectraPlayerRuntime/Private/Runtime/Player/HLS/ManifestHLS.cpp b/Engine/Plugins/Media/ElectraPlayer/Source/ElectraPlayerRuntime/Private/Runtime/Player/HLS/ManifestHLS.cpp
--- a/Engine/Plugins/Media/ElectraPlayer/Source/ElectraPlayerRuntime/Private/Runtime/Player/HLS/ManifestHLS.cpp
+++ b/Engine/Plugins/Media/ElectraPlayer/Source/ElectraPlayerRuntime/Private/Runtime/Player/HLS/ManifestHLS.cpp
@@ -242,6 +242,13 @@ FTimeValue FManifestHLS::GetDesiredLiveLatency() const
 	return ll;
 }
 
+TRangeSet<double> FManifestHLS::GetPossiblePlaybackRates(EPlayRateType InForType) const
+{
+	TRangeSet<double> Ranges;
+	Ranges.Add(TRange<double>{1.0}); // normal (real-time) playback rate
+	Ranges.Add(TRange<double>{0.0}); // and pause
+	return Ranges;
+}
 
 TSharedPtrTS<IProducerReferenceTimeInfo> FManifestHLS::GetProducerReferenceTimeInfo(int64 ID) const
 {
diff --git a/Engine/Plugins/Media/ElectraPlayer/Source/ElectraPlayerRuntime/Private/Runtime/Player/HLS/ManifestHLS.h b/Engine/Plugins/Media/ElectraPlayer/Source/ElectraPlayerRuntime/Private/Runtime/Player/HLS/ManifestHLS.h
--- a/Engine/Plugins/Media/ElectraPlayer/Source/ElectraPlayerRuntime/Private/Runtime/Player/HLS/ManifestHLS.h
+++ b/Engine/Plugins/Media/ElectraPlayer/Source/ElectraPlayerRuntime/Private/Runtime/Player/HLS/ManifestHLS.h
@@ -28,6 +28,7 @@ public:
 	void ClearDefaultStartTime() override;
 	FTimeValue GetMinBufferTime() const override;
 	FTimeValue GetDesiredLiveLatency() const override;
+	TRangeSet<double> GetPossiblePlaybackRates(EPlayRateType InForType) const override;
 	TSharedPtrTS<IProducerReferenceTimeInfo> GetProducerReferenceTimeInfo(int64 ID) const override;
 	void GetTrackMetadata(TArray<FTrackMetadata>& OutMetadata, EStreamType StreamType) const override;
 	void UpdateDynamicRefetchCounter() override;
diff --git a/Engine/Plugins/Media/ElectraPlayer/Source/ElectraPlayerRuntime/Private/Runtime/Player/Manifest.h b/Engine/Plugins/Media/ElectraPlayer/Source/ElectraPlayerRuntime/Private/Runtime/Player/Manifest.h
--- a/Engine/Plugins/Media/ElectraPlayer/Source/ElectraPlayerRuntime/Private/Runtime/Player/Manifest.h
+++ b/Engine/Plugins/Media/ElectraPlayer/Source/ElectraPlayerRuntime/Private/Runtime/Player/Manifest.h
@@ -318,12 +318,18 @@ namespace Electra
 
 		virtual FTimeValue GetDesiredLiveLatency() const = 0;
 
+		enum class EPlayRateType
+		{
+			UnthinnedRate,
+			ThinnedRate
+		};
+		virtual TRangeSet<double> GetPossiblePlaybackRates(EPlayRateType InForType) const = 0;
+
 		//! Needs to be called when the user has explicitly triggered a seek, including a programmatic loop back to the beginning.
 		//! For presentations with dynamic content changes (eg. DASH xlink:onRequest Periods) the content may need to be updated
 		//! again. This is different to internal seeking for retry purposes where content will not be re-resolved.
 		virtual void UpdateDynamicRefetchCounter() = 0;
 
-
 		enum class EClockSyncType
 		{
 			Recommended,
diff --git a/Engine/Plugins/Media/ElectraPlayer/Source/ElectraPlayerRuntime/Private/Runtime/Player/mp4/ManifestMP4.cpp b/Engine/Plugins/Media/ElectraPlayer/Source/ElectraPlayerRuntime/Private/Runtime/Player/mp4/ManifestMP4.cpp
--- a/Engine/Plugins/Media/ElectraPlayer/Source/ElectraPlayerRuntime/Private/Runtime/Player/mp4/ManifestMP4.cpp
+++ b/Engine/Plugins/Media/ElectraPlayer/Source/ElectraPlayerRuntime/Private/Runtime/Player/mp4/ManifestMP4.cpp
@@ -177,7 +177,7 @@ FTimeRange FManifestMP4Internal::GetPlaybackRange() const
 FTimeValue FManifestMP4Internal::GetMinBufferTime() const
 {
 	// NOTE: This could come from a 'pdin' (progressive download information) box, but those are rarely, if ever, set by any tool.
-	return FTimeValue().SetFromSeconds(2.0);
+	return FTimeValue().SetFromSeconds(1.0);
 }
 
 TSharedPtrTS<IProducerReferenceTimeInfo> FManifestMP4Internal::GetProducerReferenceTimeInfo(int64 ID) const
@@ -185,6 +185,15 @@ TSharedPtrTS<IProducerReferenceTimeInfo> FManifestMP4Internal::GetProducerRefere
 	return nullptr;
 }
 
+TRangeSet<double> FManifestMP4Internal::GetPossiblePlaybackRates(EPlayRateType InForType) const
+{
+	TRangeSet<double> Ranges;
+//	Ranges.Add(TRange<double>{1.0}); // normal (real-time) playback rate
+	Ranges.Add(TRange<double>::Inclusive(0.5, 4.0));
+	Ranges.Add(TRange<double>{0.0}); // and pause
+	return Ranges;
+}
+
 
 void FManifestMP4Internal::UpdateDynamicRefetchCounter()
 {
diff --git a/Engine/Plugins/Media/ElectraPlayer/Source/ElectraPlayerRuntime/Private/Runtime/Player/mp4/ManifestMP4.h b/Engine/Plugins/Media/ElectraPlayer/Source/ElectraPlayerRuntime/Private/Runtime/Player/mp4/ManifestMP4.h
--- a/Engine/Plugins/Media/ElectraPlayer/Source/ElectraPlayerRuntime/Private/Runtime/Player/mp4/ManifestMP4.h
+++ b/Engine/Plugins/Media/ElectraPlayer/Source/ElectraPlayerRuntime/Private/Runtime/Player/mp4/ManifestMP4.h
@@ -57,6 +57,7 @@ public:
 	FTimeValue GetMinBufferTime() const override;
 	FTimeValue GetDesiredLiveLatency() const override
 	{ return FTimeValue(); }
+	TRangeSet<double> GetPossiblePlaybackRates(EPlayRateType InForType) const override;
 	TSharedPtrTS<IProducerReferenceTimeInfo> GetProducerReferenceTimeInfo(int64 ID) const override;
 	void UpdateDynamicRefetchCounter() override;
 	void TriggerClockSync(EClockSyncType InClockSyncType) override;
diff --git a/Engine/Plugins/Media/ElectraPlayer/Source/ElectraPlayerRuntime/Public/IElectraPlayerInterface.h b/Engine/Plugins/Media/ElectraPlayer/Source/ElectraPlayerRuntime/Public/IElectraPlayerInterface.h
--- a/Engine/Plugins/Media/ElectraPlayer/Source/ElectraPlayerRuntime/Public/IElectraPlayerInterface.h
+++ b/Engine/Plugins/Media/ElectraPlayer/Source/ElectraPlayerRuntime/Public/IElectraPlayerInterface.h
@@ -1,9 +1,11 @@
 // Copyright Epic Games, Inc. All Rights Reserved.
 #pragma once
 
+#include "CoreMinimal.h"
 #include "MediaDecoderOutput.h"
 #include "IAnalyticsProviderET.h"
 #include "ParameterDictionary.h"
+#include "Math/RangeSet.h"
 
 #include "ElectraPlayerMisc.h"
 #include "ElectraPlayerResourceDelegate.h"
@@ -233,6 +235,15 @@ public:
 	virtual void SetPlaybackRange(const FPlaybackRange& InPlaybackRange) = 0;
 	virtual void GetPlaybackRange(FPlaybackRange& OutPlaybackRange) const = 0;
 
+	enum class EPlayRateType
+	{
+		// Playback rate with frames being dropped.
+		Thinned,
+		// Playback rate without dropping frames.
+		Unthinned
+	};
+	virtual TRangeSet<float> GetSupportedRates(EPlayRateType InPlayRateType) const = 0;
+
 	virtual float GetRate() const = 0;
 	virtual bool SetRate(float Rate) = 0;
 
