Fixed mismatched copy constructor/assignment operator in TMediaOptionalValue.

diff --git a/Engine/Plugins/Media/ElectraPlayer/Source/ElectraPlayerRuntime/Private/Runtime/OptionalValue.h b/Engine/Plugins/Media/ElectraPlayer/Source/ElectraPlayerRuntime/Private/Runtime/OptionalValue.h
--- a/Engine/Plugins/Media/ElectraPlayer/Source/ElectraPlayerRuntime/Private/Runtime/OptionalValue.h
+++ b/Engine/Plugins/Media/ElectraPlayer/Source/ElectraPlayerRuntime/Private/Runtime/OptionalValue.h
@@ -13,7 +13,6 @@ struct TMediaOptionalValue
 {
 	TMediaOptionalValue() : bValueIsSet(false) {}
 	TMediaOptionalValue(const T& v) : OptionalValue(v), bValueIsSet(true) {}
-	TMediaOptionalValue(const TMediaOptionalValue& rhs) : OptionalValue(rhs.OptionalValue), bValueIsSet(rhs.bValueIsSet) {}
 	void Set(const T& v) { OptionalValue = v; bValueIsSet = true; }
 	void SetIfNot(const T& v) { if (!IsSet()) Set(v); }
 	bool IsSet() const { return bValueIsSet; }
