ElectraPlayer: Supporting the mkv/webm container format

#rnx

diff --git a/Engine/Plugins/Media/ElectraPlayer/Source/ElectraPlayerRuntime/Private/Runtime/Demuxer/ParserMKV.h b/Engine/Plugins/Media/ElectraPlayer/Source/ElectraPlayerRuntime/Private/Runtime/Demuxer/ParserMKV.h
new file mode 100644
--- /dev/null
+++ b/Engine/Plugins/Media/ElectraPlayer/Source/ElectraPlayerRuntime/Private/Runtime/Demuxer/ParserMKV.h
@@ -0,0 +1,236 @@
+// Copyright Epic Games, Inc. All Rights Reserved.
+
+#pragma once
+
+#include "PlayerCore.h"
+#include "PlayerTime.h"
+#include "StreamTypes.h"
+#include "ParameterDictionary.h"
+#include "ErrorDetail.h"
+#include "ElectraEncryptedSampleInfo.h"
+
+
+namespace Electra
+{
+	//
+	// Forward declarations
+	//
+	class IPlayerSessionServices;
+
+
+	/**
+	 * Interface for parsing a Matroska or WebM file.
+	 */
+	class IParserMKV
+	{
+	public:
+		virtual ~IParserMKV() = default;
+
+		/**
+		 * Interface for reading data from a source.
+		 */
+		class IReader
+		{
+		public:
+			virtual ~IReader() = default;
+			/**
+			 * Read n bytes of data starting at offset o into the provided buffer.
+			 *
+			 * Reading must return the number of bytes asked to get, if necessary by blocking.
+			 * If a read error prevents reading the number of bytes -1 must be returned.
+			 *
+			 * @param IntoBuffer Buffer into which to store the data bytes. If nullptr is passed the data must be skipped over.
+			 * @param NumBytesToRead The number of bytes to read. Must not read more bytes and no less than requested.
+			 * @return The number of bytes read or -1 on a read error. If the read would go beyond the size of the file then
+			 *         returning fewer bytes than requested is permitted in this case ONLY.
+			 */
+			virtual int64 ReadData(void* InDestinationBuffer, int64 InNumBytesToRead, int64 InFromOffset) = 0;
+
+			virtual int64 GetCurrentFileOffset() const = 0;
+
+			/**
+			 * Returns the total size of the file.
+			 * This should be possible after performing the first ReadData().
+			 * If the total length is not known, return -1.
+			 */
+			virtual int64 GetTotalSize() = 0;
+
+			/**
+			 * Checks if reading of the file and therefor parsing has been aborted.
+			 *
+			 * @return true if reading/parsing has been aborted, false otherwise.
+			 */
+			virtual bool HasReadBeenAborted() const = 0;
+		};
+
+
+		static TSharedPtrTS<IParserMKV> CreateParser(IPlayerSessionServices* PlayerSession);
+
+		/**
+		 * Parses the header boxes.
+		 */
+		virtual FErrorDetail ParseHeader(IReader* DataReader) = 0;
+
+
+		/*******************************************************************************************************************/
+		class ITrack;
+
+		virtual FErrorDetail PrepareTracks() = 0;
+
+		virtual FTimeValue GetDuration() const = 0;
+
+		virtual int32 GetNumberOfTracks() const = 0;
+
+		class ICueIterator
+		{
+		public:
+			virtual ~ICueIterator() = default;
+
+			enum class ESearchMode
+			{
+				Before,
+				After,
+				Closest
+			};
+
+			virtual UEMediaError StartAtTime(const FTimeValue& AtTime, ESearchMode SearchMode) = 0;
+			virtual UEMediaError StartAtFirst() = 0;
+			virtual UEMediaError StartAtUniqueID(uint32 CueUniqueID) = 0;
+			virtual UEMediaError Next() = 0;
+			virtual bool IsAtEOS() const = 0;
+
+			virtual const ITrack* GetTrack() const = 0;
+			virtual FTimeValue GetTimestamp() const = 0;
+			virtual int64 GetClusterFileOffset() const = 0;
+			virtual int64 GetClusterFileSize() const = 0;
+			virtual FTimeValue GetClusterDuration() const = 0;
+			virtual bool IsLastCluster() const = 0;
+
+			virtual uint32 GetUniqueID() const = 0;
+			virtual uint32 GetNextUniqueID() const = 0;
+		};
+
+		class IClusterParser
+		{
+		public:
+			virtual ~IClusterParser() = default;
+
+			enum class EParseAction
+			{
+				SkipOver,		// Skip over the next n bytes. GetAction() returns an IActionSkipOver action.
+				PrependData,	// Prepend the frame data with constant data. GetAction() returns an IActionPrependData action.
+				ReadFrameData,	// Read the next n bytes as frame data. GetAction() returns an IActionReadFrameData action.
+				DecryptData,	// Decrypt the frame data. GetAction() returns an IActionDecryptData action.
+				FrameDone,		// Finished this frame data, continue with the next frame. GetAction() returns an IActionFrameDone action. All values (ie timestamps) are valid only now.
+				EndOfData,		// All cluster input consumed or skipped over. GetAction() returns nullptr.
+				Failure			// An error occurred. GetAction() returns nullptr.
+			};
+
+			class IAction
+			{
+			public:
+				virtual ~IAction() = default;
+				virtual uint64 GetTrackID() const = 0;
+				virtual FTimeValue GetPTS() const = 0;
+				virtual FTimeValue GetDTS() const = 0;
+				virtual FTimeValue GetDuration() const = 0;
+				virtual bool IsKeyFrame() const = 0;
+				virtual int64 GetTimestamp() const = 0;
+				virtual int64 GetSegmentRelativePosition() const = 0;
+				virtual int64 GetClusterPosition() const = 0;
+			};
+
+			class IActionSkipOver : public IAction
+			{
+			public:
+				virtual int64 GetNumBytesToSkip() const = 0;
+			};
+
+			class IActionPrependData : public IAction
+			{
+			public:
+				virtual const TArray<uint8>& GetPrependData() const = 0;
+			};
+
+			class IActionReadFrameData : public IAction
+			{
+			public:
+				virtual int64 GetNumBytesToRead() const = 0;
+			};
+
+			class IActionDecryptData : public IAction
+			{
+			public:
+				virtual void GetDecryptionInfo(ElectraCDM::FMediaCDMSampleInfo& OutSampleDecryptionInfo) const = 0;
+			};
+
+			class IActionFrameDone : public IAction
+			{
+			public:
+			};
+
+			/**
+			 * Performs parsing the cluster content or current frame, returning the next action to take.
+			 * 
+			 * Returns the action to perform next.
+			 */
+			virtual EParseAction NextParseAction() = 0;
+
+			/**
+			 * Returns the error which resulted in returning the next action `Failure`
+			 */
+			virtual FErrorDetail GetLastError() const = 0;
+			
+			/**
+			 * Returns the base class pointer of the next action to perform.
+			 * Must be cast to the action indicated by the next action.
+			 * 
+			 * The action is owned by the parser and must not be destroyed!
+			 */
+			virtual const IAction* GetAction() const = 0;
+
+			/**
+			 * Returns the start offset of the current cluster.
+			 * This is needed for retries.
+			 */
+			virtual int64 GetClusterPosition() const = 0;
+			
+			/**
+			 * Returns the offset of the current block (simple or group) in the cluster.
+			 * This is needed for retries.
+			 */
+			virtual int64 GetClusterBlockPosition() const = 0;
+		};
+
+
+		class ITrack
+		{
+		public:
+			virtual ~ITrack() = default;
+
+			virtual uint64 GetID() const = 0;
+			virtual FString GetName() const = 0;
+			virtual const TArray<uint8>& GetCodecSpecificData() const = 0;
+			virtual const FStreamCodecInformation& GetCodecInformation() const = 0;
+			virtual const FString GetLanguage() const = 0;
+			
+			virtual ICueIterator* CreateCueIterator() const = 0;
+		};
+
+
+		virtual const ITrack* GetTrackByIndex(int32 Index) const = 0;
+		virtual const ITrack* GetTrackByTrackID(uint64 TrackID) const = 0;
+
+		/**
+		 * Create a cluster parser.
+		 * The data reader MUST start reading on a Matroska cluster.
+		 */
+		virtual TSharedPtrTS<IClusterParser> CreateClusterParser(IReader* DataReader, const TArray<uint64>& TrackIDsToParse) = 0;
+
+		// Adds a cue of it does not exist yet. This may be called during cluster parsing for sync samples since not all
+		// sync samples may have been added as cues in the multiplexing process.
+		// NOTE: This should only be called for video samples.
+		virtual void AddCue(int64 InCueTimestamp, uint64 InTrackID, int64 InCueRelativePosition, uint64 InCueBlockNumber, int64 InClusterPosition) = 0;
+	};
+
+} // namespace Electra
diff --git a/Engine/Plugins/Media/ElectraPlayer/Source/ElectraPlayerRuntime/Private/Runtime/Player/AdaptiveStreamingPlayer.cpp b/Engine/Plugins/Media/ElectraPlayer/Source/ElectraPlayerRuntime/Private/Runtime/Player/AdaptiveStreamingPlayer.cpp
--- a/Engine/Plugins/Media/ElectraPlayer/Source/ElectraPlayerRuntime/Private/Runtime/Player/AdaptiveStreamingPlayer.cpp
+++ b/Engine/Plugins/Media/ElectraPlayer/Source/ElectraPlayerRuntime/Private/Runtime/Player/AdaptiveStreamingPlayer.cpp
@@ -2085,7 +2085,7 @@ bool FAdaptiveStreamingPlayer::HaveEnoughBufferedDataToStartPlayback()
 	// When we are dealing with a single multiplexed stream and one buffer is blocked then essentially all buffers
 	// must be considered blocked since demuxing cannot continue.
 	// FIXME: do this more elegant somehow
-	if (ManifestType == EMediaFormatType::ISOBMFF)
+	if (ManifestType == EMediaFormatType::ISOBMFF || ManifestType == EMediaFormatType::MKV)
 	{
 		if (VideoBufferStats.StreamBuffer.bLastPushWasBlocked ||
 			AudioBufferStats.StreamBuffer.bLastPushWasBlocked ||
@@ -2644,7 +2644,7 @@ void FAdaptiveStreamingPlayer::InternalHandlePendingStartRequest(const FTimeValu
 						if (BufferSourceInfoVid.IsValid())
 						{
 							SelectedStreamAttributesVid.UpdateWith(BufferSourceInfoVid->Kind, BufferSourceInfoVid->Language, BufferSourceInfoVid->Codec, BufferSourceInfoVid->HardIndex);
-							if (ManifestType != EMediaFormatType::ISOBMFF)
+							if (ManifestType != EMediaFormatType::ISOBMFF && ManifestType != EMediaFormatType::MKV)
 							{
 								StreamSelectionAttributesVid.UpdateIfOverrideSet(BufferSourceInfoVid->Kind, BufferSourceInfoVid->Language, BufferSourceInfoVid->Codec);
 							}
@@ -2652,7 +2652,7 @@ void FAdaptiveStreamingPlayer::InternalHandlePendingStartRequest(const FTimeValu
 						if (BufferSourceInfoAud.IsValid())
 						{
 							SelectedStreamAttributesAud.UpdateWith(BufferSourceInfoAud->Kind, BufferSourceInfoAud->Language, BufferSourceInfoAud->Codec, BufferSourceInfoAud->HardIndex);
-							if (ManifestType != EMediaFormatType::ISOBMFF)
+							if (ManifestType != EMediaFormatType::ISOBMFF && ManifestType != EMediaFormatType::MKV)
 							{
 								StreamSelectionAttributesAud.UpdateIfOverrideSet(BufferSourceInfoAud->Kind, BufferSourceInfoAud->Language, BufferSourceInfoAud->Codec);
 							}
@@ -2660,7 +2660,7 @@ void FAdaptiveStreamingPlayer::InternalHandlePendingStartRequest(const FTimeValu
 						if (BufferSourceInfoTxt.IsValid())
 						{
 							SelectedStreamAttributesTxt.UpdateWith(BufferSourceInfoTxt->Kind, BufferSourceInfoTxt->Language, BufferSourceInfoTxt->Codec, BufferSourceInfoTxt->HardIndex);
-							if (ManifestType != EMediaFormatType::ISOBMFF)
+							if (ManifestType != EMediaFormatType::ISOBMFF && ManifestType != EMediaFormatType::MKV)
 							{
 								StreamSelectionAttributesTxt.UpdateIfOverrideSet(BufferSourceInfoTxt->Kind, BufferSourceInfoTxt->Language, BufferSourceInfoTxt->Codec);
 							}
@@ -2845,9 +2845,9 @@ void FAdaptiveStreamingPlayer::InternalHandlePendingFirstSegmentRequest(const FT
 			NewReceiveBuffers->AudBuffer = MakeSharedTS<FMultiTrackAccessUnitBuffer>(EStreamType::Audio);
 			NewReceiveBuffers->TxtBuffer = MakeSharedTS<FMultiTrackAccessUnitBuffer>(EStreamType::Subtitle);
 
-			if (ManifestType == EMediaFormatType::ISOBMFF)
+			if (ManifestType == EMediaFormatType::ISOBMFF || ManifestType == EMediaFormatType::MKV)
 			{
-				// mp4 has multiple tracks in a single stream that provide data all at the same time.
+				// mp4 and mkv have multiple tracks in a single stream that provide data all at the same time.
 				NewReceiveBuffers->VidBuffer->SetParallelTrackMode();
 				NewReceiveBuffers->AudBuffer->SetParallelTrackMode();
 				NewReceiveBuffers->TxtBuffer->SetParallelTrackMode();
@@ -3274,7 +3274,7 @@ void FAdaptiveStreamingPlayer::HandlePendingMediaSegmentRequests()
 								case EStreamType::Video:
 								{
 									SelectedStreamAttributesVid.UpdateWith(BufferSourceInfo->Kind, BufferSourceInfo->Language, BufferSourceInfo->Codec, BufferSourceInfo->HardIndex);
-									if (ManifestType != EMediaFormatType::ISOBMFF)
+									if (ManifestType != EMediaFormatType::ISOBMFF && ManifestType != EMediaFormatType::MKV)
 									{
 										StreamSelectionAttributesVid.ClearOverrideIndex();
 									}
@@ -3283,7 +3283,7 @@ void FAdaptiveStreamingPlayer::HandlePendingMediaSegmentRequests()
 								case EStreamType::Audio:
 								{
 									SelectedStreamAttributesAud.UpdateWith(BufferSourceInfo->Kind, BufferSourceInfo->Language, BufferSourceInfo->Codec, BufferSourceInfo->HardIndex);
-									if (ManifestType != EMediaFormatType::ISOBMFF)
+									if (ManifestType != EMediaFormatType::ISOBMFF && ManifestType != EMediaFormatType::MKV)
 									{
 										StreamSelectionAttributesAud.ClearOverrideIndex();
 									}
@@ -3292,7 +3292,7 @@ void FAdaptiveStreamingPlayer::HandlePendingMediaSegmentRequests()
 								case EStreamType::Subtitle:
 								{
 									SelectedStreamAttributesTxt.UpdateWith(BufferSourceInfo->Kind, BufferSourceInfo->Language, BufferSourceInfo->Codec, BufferSourceInfo->HardIndex);
-									if (ManifestType != EMediaFormatType::ISOBMFF)
+									if (ManifestType != EMediaFormatType::ISOBMFF && ManifestType != EMediaFormatType::MKV)
 									{
 										StreamSelectionAttributesTxt.ClearOverrideIndex();
 									}
@@ -3549,7 +3549,7 @@ void FAdaptiveStreamingPlayer::InternalHandleSegmentTrackChanges(const FTimeValu
 {
 	bool bChangeMade = false;
 
-	if (ManifestType == EMediaFormatType::ISOBMFF)
+	if (ManifestType == EMediaFormatType::ISOBMFF || ManifestType == EMediaFormatType::MKV)
 	{
 		auto HandleISOBMFFTrackSelection = [&bChangeMade, this](FStreamSelectionAttributes& OutSelectionAttributes, FInternalStreamSelectionAttributes& InOutSelectedAttributes, EStreamType InType,
 																TSharedPtrTS<FStreamSelectionAttributes>& InPendingSelection, TSharedPtrTS<IManifest::IPlayPeriod> InCurrentPeriod) -> void
@@ -4502,8 +4502,8 @@ void FAdaptiveStreamingPlayer::InternalClose()
 	ExternalDataReader.Reset();
 
 	// Reset remaining internal state
-	CurrentState = EPlayerState::eState_Idle;
-	StreamState = EStreamState::eStream_Running;
+	CurrentState   	= EPlayerState::eState_Idle;
+	StreamState		= EStreamState::eStream_Running;
 	RebufferCause = ERebufferCause::None;
 	LastBufferingState = EPlayerState::eState_Buffering;
 	RebufferDetectedAtPlayPos.SetToInvalid();
diff --git a/Engine/Plugins/Media/ElectraPlayer/Source/ElectraPlayerRuntime/Private/Runtime/Player/AdaptiveStreamingPlayerABR.cpp b/Engine/Plugins/Media/ElectraPlayer/Source/ElectraPlayerRuntime/Private/Runtime/Player/AdaptiveStreamingPlayerABR.cpp
--- a/Engine/Plugins/Media/ElectraPlayer/Source/ElectraPlayerRuntime/Private/Runtime/Player/AdaptiveStreamingPlayerABR.cpp
+++ b/Engine/Plugins/Media/ElectraPlayer/Source/ElectraPlayerRuntime/Private/Runtime/Player/AdaptiveStreamingPlayerABR.cpp
@@ -190,7 +190,7 @@ namespace Electra
 	void FAdaptiveStreamSelector::SetFormatType(EMediaFormatType InFormatType)
 	{
 		// Multiplexed stream with no bitrate switching alternatives?
-		if (InFormatType == EMediaFormatType::ISOBMFF)
+		if (InFormatType == EMediaFormatType::ISOBMFF || InFormatType == EMediaFormatType::MKV)
 		{
 			ABRMethod.Reset(IABRFixedStream::Create(this, InFormatType, EABRPresentationType::OnDemand));
 		}
diff --git a/Engine/Plugins/Media/ElectraPlayer/Source/ElectraPlayerRuntime/Private/Runtime/Player/AdaptiveStreamingPlayerManifest.cpp b/Engine/Plugins/Media/ElectraPlayer/Source/ElectraPlayerRuntime/Private/Runtime/Player/AdaptiveStreamingPlayerManifest.cpp
--- a/Engine/Plugins/Media/ElectraPlayer/Source/ElectraPlayerRuntime/Private/Runtime/Player/AdaptiveStreamingPlayerManifest.cpp
+++ b/Engine/Plugins/Media/ElectraPlayer/Source/ElectraPlayerRuntime/Private/Runtime/Player/AdaptiveStreamingPlayerManifest.cpp
@@ -7,6 +7,7 @@
 #include "Player/HLS/PlaylistReaderHLS.h"
 #include "Player/mp4/PlaylistReaderMP4.h"
 #include "Player/DASH/PlaylistReaderDASH.h"
+#include "Player/mkv/PlaylistReaderMKV.h"
 #include "Utilities/Utilities.h"
 #include "Utilities/StringHelpers.h"
 #include "Utilities/URLParser.h"
@@ -23,6 +24,9 @@ static const FString MIMETypeMP4A(TEXT("audio/mp4"));
 static const FString MIMETypeQuickTime(TEXT("video/quicktime"));
 static const FString MIMETypeHLS(TEXT("application/vnd.apple.mpegURL"));
 static const FString MIMETypeDASH(TEXT("application/dash+xml"));
+static const FString MIMETypeMKV(TEXT("video/x-matroska"));
+static const FString MIMETypeMKA(TEXT("audio/x-matroska"));
+
 
 
 /**
@@ -72,6 +76,9 @@ FString GetMIMETypeForURL(const FString& URL)
 		static const FString kTextMOV(TEXT("mov"));
 		static const FString kTextMPD(TEXT("mpd"));
 		static const FString kTextM3U8(TEXT("m3u8"));
+		static const FString kTextMKV(TEXT("mkv"));
+		static const FString kTextMKA(TEXT("mka"));
+		static const FString kTextWEBM(TEXT("webm"));
 		if (LowerCaseExtension == kTextMP4 || LowerCaseExtension == kTextMP4V)
 		{
 			MimeType = MIMETypeMP4;
@@ -92,6 +99,14 @@ FString GetMIMETypeForURL(const FString& URL)
 		{
 			MimeType = MIMETypeHLS;
 		}
+		else if (LowerCaseExtension == kTextMKV || LowerCaseExtension == kTextWEBM)
+		{
+			MimeType = MIMETypeMKV;
+		}
+		else if (LowerCaseExtension == kTextMKA)
+		{
+			MimeType = MIMETypeMKA;
+		}
 	}
 
 	return MimeType;
@@ -210,6 +225,11 @@ void FAdaptiveStreamingPlayer::InternalLoadManifest(const FString& InURL, const
 				ManifestReader = IPlaylistReaderDASH::Create(this);
 				ManifestType = EMediaFormatType::DASH;
 			}
+			else if (mimeType == Playlist::MIMETypeMKV || mimeType == Playlist::MIMETypeMKA)
+			{
+				ManifestReader = IPlaylistReaderMKV::Create(this);
+				ManifestType = EMediaFormatType::MKV;
+			}
 			else
 			{
 				FErrorDetail err;
@@ -314,8 +334,8 @@ bool FAdaptiveStreamingPlayer::SelectManifest()
 			PlayerConfig.SeekBufferMinTimeAvailBeforePlayback    = Utils::Min(minBufTimeMPD, PlayerConfig.SeekBufferMinTimeAvailBeforePlayback);
 			PlayerConfig.RebufferMinTimeAvailBeforePlayback 	 = Utils::Min(minBufTimeMPD, PlayerConfig.RebufferMinTimeAvailBeforePlayback);
 
-			// For an mp4 stream we can now get rid of the manifest reader. It is no longer needed and we don't need to have it linger.
-			if (ManifestType == EMediaFormatType::ISOBMFF)
+			// For an mp4 or mkv stream we can now get rid of the manifest reader. It is no longer needed and we don't need to have it linger.
+			if (ManifestType == EMediaFormatType::ISOBMFF || ManifestType == EMediaFormatType::MKV)
 			{
 				InternalCloseManifestReader();
 			}
diff --git a/Engine/Plugins/Media/ElectraPlayer/Source/ElectraPlayerRuntime/Private/Runtime/Player/Manifest.h b/Engine/Plugins/Media/ElectraPlayer/Source/ElectraPlayerRuntime/Private/Runtime/Player/Manifest.h
--- a/Engine/Plugins/Media/ElectraPlayer/Source/ElectraPlayerRuntime/Private/Runtime/Player/Manifest.h
+++ b/Engine/Plugins/Media/ElectraPlayer/Source/ElectraPlayerRuntime/Private/Runtime/Player/Manifest.h
@@ -22,7 +22,8 @@ namespace Electra
 		Unknown,
 		ISOBMFF,					// mp4
 		HLS,						// Apple HLS (HTTP Live Streaming)
-		DASH						// MPEG DASH
+		DASH,						// MPEG DASH
+		MKV							// Matroska / WebM
 	};
 
 
diff --git a/Engine/Plugins/Media/ElectraPlayer/Source/ElectraPlayerRuntime/Private/Runtime/Player/PlaylistReader.h b/Engine/Plugins/Media/ElectraPlayer/Source/ElectraPlayerRuntime/Private/Runtime/Player/PlaylistReader.h
--- a/Engine/Plugins/Media/ElectraPlayer/Source/ElectraPlayerRuntime/Private/Runtime/Player/PlaylistReader.h
+++ b/Engine/Plugins/Media/ElectraPlayer/Source/ElectraPlayerRuntime/Private/Runtime/Player/PlaylistReader.h
@@ -7,6 +7,7 @@
 #include "StreamTypes.h"
 #include "Player/Playlist.h"
 #include "ParameterDictionary.h"
+#include "HTTP/HTTPManager.h"
 
 
 
diff --git a/Engine/Plugins/Media/ElectraPlayer/Source/ElectraPlayerRuntime/Private/Runtime/Player/mkv/OptionKeynamesMKV.h b/Engine/Plugins/Media/ElectraPlayer/Source/ElectraPlayerRuntime/Private/Runtime/Player/mkv/OptionKeynamesMKV.h
new file mode 100644
--- /dev/null
+++ b/Engine/Plugins/Media/ElectraPlayer/Source/ElectraPlayerRuntime/Private/Runtime/Player/mkv/OptionKeynamesMKV.h
@@ -0,0 +1,19 @@
+// Copyright Epic Games, Inc. All Rights Reserved.
+
+#pragma once
+
+#include "PlayerCore.h"
+
+namespace Electra
+{
+namespace MKV
+{
+
+const TCHAR* const OptionKeyMKVLoadConnectTimeout = TEXT("mkv_connection_timeout");				//!< (FTimeValue) value specifying connection timeout fetching the MKV
+const TCHAR* const OptionKeyMKVLoadNoDataTimeout = TEXT("mkv_nodata_timeout");					//!< (FTimeValue) value specifying no-data timeout fetching the MKV
+
+}
+
+} // namespace Electra
+
+
diff --git a/Engine/Plugins/Media/ElectraPlayer/Source/ElectraPlayerRuntime/Private/Runtime/Player/mp4/StreamReaderMP4.h b/Engine/Plugins/Media/ElectraPlayer/Source/ElectraPlayerRuntime/Private/Runtime/Player/mp4/StreamReaderMP4.h
--- a/Engine/Plugins/Media/ElectraPlayer/Source/ElectraPlayerRuntime/Private/Runtime/Player/mp4/StreamReaderMP4.h
+++ b/Engine/Plugins/Media/ElectraPlayer/Source/ElectraPlayerRuntime/Private/Runtime/Player/mp4/StreamReaderMP4.h
@@ -160,7 +160,7 @@ private:
 	struct FSelectedTrackData
 	{
 		FSelectedTrackData()
-			: StreamType(EStreamType::Video), Bitrate(0), bGotKeyframe(false), bIsSelectedTrack(false), bIsFirstInSequence(true), bReadPastLastPTS(false)
+			: StreamType(EStreamType::Video), Bitrate(0), bIsSelectedTrack(false), bIsFirstInSequence(true), bReadPastLastPTS(false)
 		{
 			DurationSuccessfullyRead.SetToZero();
 			DurationSuccessfullyDelivered.SetToZero();
@@ -169,7 +169,6 @@ private:
 		TSharedPtrTS<FBufferSourceInfo>			BufferSourceInfo;
 		EStreamType								StreamType;
 		int32									Bitrate;
-		bool									bGotKeyframe;
 		bool									bIsSelectedTrack;
 		bool									bIsFirstInSequence;
 		bool									bReadPastLastPTS;
diff --git a/Engine/Plugins/Media/ElectraPlayer/Source/ElectraPlayerRuntime/Private/Runtime/PlayerFacility.h b/Engine/Plugins/Media/ElectraPlayer/Source/ElectraPlayerRuntime/Private/Runtime/PlayerFacility.h
--- a/Engine/Plugins/Media/ElectraPlayer/Source/ElectraPlayerRuntime/Private/Runtime/PlayerFacility.h
+++ b/Engine/Plugins/Media/ElectraPlayer/Source/ElectraPlayerRuntime/Private/Runtime/PlayerFacility.h
@@ -13,12 +13,12 @@ namespace Electra
 
 		enum class EFacility
 		{
-			// NOTE: Never change this order even if adding new elements mean they go at the end!
 			Unknown,
 			Utility,
 			Player,
 			ABR,
 			HTTPReader,
+			BufferedDataReader,
 			HLSPlaylistReader,
 			HLSPlaylistBuilder,
 			HLSManifest,
@@ -30,6 +30,10 @@ namespace Electra
 			MP4PlaylistReader,
 			MP4Playlist,
 			MP4StreamReader,
+			MKVPlaylistReader,
+			MKVPlaylist,
+			MKVStreamReader,
+			MKVParser,
 			LicenseKey,
 			DASHMPDReader,
 			DASHMPDParser,
@@ -60,6 +64,8 @@ namespace Electra
 					return TEXT("ABR");
 				case EFacility::HTTPReader:
 					return TEXT("HTTP reader");
+				case EFacility::BufferedDataReader:
+					return TEXT("Buffered data reader");
 				case EFacility::HLSPlaylistReader:
 					return TEXT("HLS playlist reader");
 				case EFacility::HLSPlaylistBuilder:
@@ -104,6 +110,14 @@ namespace Electra
 					return TEXT("Audio decoder");
 				case EFacility::VideoDecoder:
 					return TEXT("Video decoder");
+				case EFacility::MKVPlaylistReader:
+					return TEXT("MKV playlist reader");
+				case EFacility::MKVPlaylist:
+					return TEXT("MKV playlist");
+				case EFacility::MKVStreamReader:
+					return TEXT("MKV reader");
+				case EFacility::MKVParser:
+					return TEXT("MKV parser");
 			}
 			return TEXT("???");
 		}
diff --git a/Engine/Plugins/Media/ElectraPlayer/Source/ElectraPlayerRuntime/Private/Runtime/StreamDataBuffer.h b/Engine/Plugins/Media/ElectraPlayer/Source/ElectraPlayerRuntime/Private/Runtime/StreamDataBuffer.h
--- a/Engine/Plugins/Media/ElectraPlayer/Source/ElectraPlayerRuntime/Private/Runtime/StreamDataBuffer.h
+++ b/Engine/Plugins/Media/ElectraPlayer/Source/ElectraPlayerRuntime/Private/Runtime/StreamDataBuffer.h
@@ -4,7 +4,6 @@
 
 #include "PlayerCore.h"
 
-
 namespace Electra
 {
 	/**
@@ -23,10 +22,22 @@ namespace Electra
 			Deallocate();
 		}
 
+		void SetExternalBuffer(uint8* InExternalBuffer, uint64 InExternalBufferSize)
+		{
+			FScopeLock Lock(&AccessLock);
+			Deallocate();
+			ExternalBuffer = InExternalBuffer;
+			DataSize = InExternalBufferSize;
+		}
+
 		// Allocates buffer of the specified capacity, destroying any previous buffer.
 		bool Reserve(int64 InNumBytes)
 		{
 			FScopeLock Lock(&AccessLock);
+			if (ExternalBuffer)
+			{
+				return false;
+			}
 			Deallocate();
 			if (Allocate(InNumBytes))
 			{
@@ -40,6 +51,10 @@ namespace Electra
 		bool EnlargeTo(int64 InNewNumBytes)
 		{
 			FScopeLock Lock(&AccessLock);
+			if (ExternalBuffer)
+			{
+				return InNewNumBytes <= Capacity();
+			}
 			// Do we need a bigger capacity?
 			if (InNewNumBytes > Capacity())
 			{
@@ -151,6 +166,7 @@ namespace Electra
 		{
 			FScopeLock Lock(&AccessLock);
 			check(!bEOD);
+
 			// Zero elements can always be pushed...
 			if (NumElements == 0)
 			{
@@ -158,7 +174,10 @@ namespace Electra
 			}
 			if (Avail() >= NumElements)
 			{
-				CopyData(Buffer->GetData() + WritePos, InData, NumElements);
+				if (InData)
+				{
+					CopyData(GetBufferBase() + WritePos, InData, NumElements);
+				}
 				WritePos += NumElements;
 				if (WritePos >= WaitingForSize)
 				{
@@ -194,7 +213,7 @@ namespace Electra
 			// Copy out or skip over?
 			if (OutData)
 			{
-				CopyData(OutData, Buffer->GetData() + ReadPos, MaxElementsWanted);
+				CopyData(OutData, GetBufferBase() + ReadPos, MaxElementsWanted);
 			}
 			ReadPos += MaxElementsWanted;
 			return MaxElementsWanted;
@@ -219,11 +238,11 @@ namespace Electra
 		// Must control Lock()/Unlock() externally!
 		const uint8* GetLinearReadData() const
 		{
-			return Buffer.IsValid() ? (Buffer->GetData() + ReadPos) : nullptr;
+			return GetBufferBase() ? GetBufferBase() + ReadPos : nullptr;
 		}
 		uint8* GetLinearReadData()
 		{
-			return Buffer.IsValid() ? (Buffer->GetData() + ReadPos) : nullptr;
+			return GetBufferBase() ? GetBufferBase() + ReadPos : nullptr;
 		}
 
 		uint8* GetLinearWriteData(int64 InNumBytesToAppend)
@@ -235,7 +254,7 @@ namespace Electra
 				bool bOk = IsEmpty() ? Allocate(InNumBytesToAppend) : InternalGrowTo(DataSize + InNumBytesToAppend - Av);
 				check(bOk); (void)bOk;
 			}
-			return Buffer.IsValid() ? (Buffer->GetData() + WritePos) : nullptr;
+			return GetBufferBase() ? GetBufferBase() + WritePos : nullptr;
 		}
 		void AppendedNewData(int64 InNumAppended)
 		{
@@ -280,8 +299,17 @@ namespace Electra
 		}
 
 	protected:
+		uint8* GetBufferBase() const
+		{
+			return Buffer.IsValid() ? Buffer->GetData() : ExternalBuffer ? ExternalBuffer : nullptr;
+		}
+
 		bool Allocate(int64 InSize)
 		{
+			if (ExternalBuffer)
+			{
+				return false;
+			}
 			if (InSize)
 			{
 				DataSize = InSize;
@@ -300,6 +328,7 @@ namespace Electra
 			DataSize = 0;
 			WritePos = 0;
 			ReadPos = 0;
+			ExternalBuffer = nullptr;
 		}
 
 		bool InternalGrowTo(int64 InNewNumBytes)
@@ -333,6 +362,8 @@ namespace Electra
 		uint64 ReadPos = 0;
 		// Amount of data necessary to be present for `SizeAvailableSignal` to get set.
 		uint64 WaitingForSize = 0;
+		// If set a buffer is provided externally to read into directly.
+		uint8* ExternalBuffer = nullptr;
 		// Flag indicating that no additional data will be added to the buffer.
 		volatile bool bEOD = false;	
 		// Flag indicating that reading into the buffer has been aborted.
diff --git a/Engine/Plugins/Media/ElectraUtil/Source/ElectraBase/Public/Utilities/Utilities.h b/Engine/Plugins/Media/ElectraUtil/Source/ElectraBase/Public/Utilities/Utilities.h
--- a/Engine/Plugins/Media/ElectraUtil/Source/ElectraBase/Public/Utilities/Utilities.h
+++ b/Engine/Plugins/Media/ElectraUtil/Source/ElectraBase/Public/Utilities/Utilities.h
@@ -44,5 +44,35 @@ namespace Electra
 			return rev;
 		}
 
+		inline uint16 EndianSwap(uint16 value)
+		{ 
+			return (value >> 8) | (value << 8);
+		}
+		inline int16 EndianSwap(int16 value)
+		{
+			return int16(EndianSwap(uint16(value)));
+		}
+		inline uint32 EndianSwap(uint32 value)
+		{ 
+			return (value << 24) | ((value & 0xff00) << 8) | ((value >> 8) & 0xff00) | (value >> 24);
+		}
+		inline int32 EndianSwap(int32 value)
+		{ 
+			return int32(EndianSwap(uint32(value)));
+		}
+		inline uint64 EndianSwap(uint64 value)
+		{ 
+			return (uint64(EndianSwap(uint32(value & 0xffffffffU))) << 32) | uint64(EndianSwap(uint32(value >> 32)));
+		}
+		inline int64 EndianSwap(int64 value)
+		{
+			return int64(EndianSwap(uint64(value)));
+		}
+
+		inline constexpr uint32 Make4CC(const uint8 A, const uint8 B, const uint8 C, const uint8 D)
+		{
+			return (static_cast<uint32>(A) << 24) | (static_cast<uint32>(B) << 16) | (static_cast<uint32>(C) << 8) | static_cast<uint32>(D);
+		}
+
 	} // namespace Utils
 } // namespace Electra
