ElectraPlayer: changed log verbosity of data reader

#rnx

diff --git a/Engine/Plugins/Media/ElectraPlayer/Source/ElectraPlayerRuntime/Private/Runtime/BufferedDataReader.cpp b/Engine/Plugins/Media/ElectraPlayer/Source/ElectraPlayerRuntime/Private/Runtime/BufferedDataReader.cpp
--- a/Engine/Plugins/Media/ElectraPlayer/Source/ElectraPlayerRuntime/Private/Runtime/BufferedDataReader.cpp
+++ b/Engine/Plugins/Media/ElectraPlayer/Source/ElectraPlayerRuntime/Private/Runtime/BufferedDataReader.cpp
@@ -12,13 +12,14 @@ namespace Electra
 	enum EBufferedDataReaderError
 	{
 		FileTooShort = 1,
+		ReadAssetDataFailed = 2,
 	};
 
 	void FBufferedDataReader::CreateNewArea(int64 NumBytes, int64 FromOffset)
 	{
 		if (TotalDataSize >= 0 && FromOffset + NumBytes > TotalDataSize)
 		{
-			UE_LOG(LogElectraBufferReader, Log, TEXT("CreateNewArea(): clamp request to end of file size"));
+			UE_LOG(LogElectraBufferReader, VeryVerbose, TEXT("CreateNewArea(): clamp request to end of file size"));
 			NumBytes = TotalDataSize - FromOffset;
 		}
 
@@ -30,7 +31,7 @@ namespace Electra
 			NewArea->StartOffset = FromOffset;
 			int64 TotalSize = -1;
 			int64 NumRead = DataProvider->OnReadAssetData((void*) NewArea->Data, NumBytes, FromOffset, &TotalSize);
-			UE_LOG(LogElectraBufferReader, Log, TEXT("CreateNewArea(): read %lld bytes from offset %lld; total file size = %lld, received %lld bytes"), (long long int)NumBytes, (long long int)FromOffset, (long long int)TotalSize, (long long int)NumRead);
+			UE_LOG(LogElectraBufferReader, VeryVerbose, TEXT("CreateNewArea(): read %lld bytes from offset %lld; total file size = %lld, received %lld bytes"), (long long int)NumBytes, (long long int)FromOffset, (long long int)TotalSize, (long long int)NumRead);
 			if (TotalSize >= 0)
 			{
 				TotalDataSize = TotalSize;
@@ -55,8 +56,7 @@ namespace Electra
 				IDataProvider::EError Error = static_cast<IDataProvider::EError>(NumRead);
 				if (Error == IDataProvider::EError::Failed)
 				{
-// TODO: Facility & code
-					LastError.SetError(UEMEDIA_ERROR_READ_ERROR).SetFacility(Facility::EFacility::HTTPReader).SetCode(4711).SetMessage(TEXT("OnReadAssetData() failed"));
+					LastError.SetError(UEMEDIA_ERROR_READ_ERROR).SetFacility(Facility::EFacility::BufferedDataReader).SetCode(EBufferedDataReaderError::ReadAssetDataFailed).SetMessage(TEXT("OnReadAssetData() failed"));
 				}
 			}
 		}
@@ -68,7 +68,7 @@ namespace Electra
 		// Reading past the known end of the file?
 		if (TotalDataSize >= 0 && CurrentArea->StartOffset + CurrentArea->Size + NumBytesToAdd > TotalDataSize)
 		{
-			UE_LOG(LogElectraBufferReader, Log, TEXT("EnlargeCurrentAreaBy(): clamp request of %lld to %lld at end of file"), (long long int)NumBytesToAdd, (long long int)(TotalDataSize - (CurrentArea->StartOffset + CurrentArea->Size)));
+			UE_LOG(LogElectraBufferReader, VeryVerbose, TEXT("EnlargeCurrentAreaBy(): clamp request of %lld to %lld at end of file"), (long long int)NumBytesToAdd, (long long int)(TotalDataSize - (CurrentArea->StartOffset + CurrentArea->Size)));
 			NumBytesToAdd = TotalDataSize - (CurrentArea->StartOffset + CurrentArea->Size);
 		}
 
@@ -84,7 +84,7 @@ namespace Electra
 		// Will we reach into the next area if we enlarge this one by the given amount?
 		if (Next && CurrentArea->StartOffset + CurrentArea->Size + NumBytesToAdd >= Next->StartOffset)
 		{
-			UE_LOG(LogElectraBufferReader, Log, TEXT("EnlargeCurrentAreaBy(): reaching into area of offset %lld, clamping"), (long long int)Next->StartOffset);
+			UE_LOG(LogElectraBufferReader, VeryVerbose, TEXT("EnlargeCurrentAreaBy(): reaching into area of offset %lld, clamping"), (long long int)Next->StartOffset);
 
 			// Yes. Trim the amount to read to avoid overlap.
 			NumBytesToAdd = Next->StartOffset - (CurrentArea->StartOffset + CurrentArea->Size);
@@ -102,7 +102,7 @@ namespace Electra
 		{
 			int64 TotalSize = -1;
 			int64 NumRead = DataProvider->OnReadAssetData((void*)(CurrentArea->Data + CurrentArea->Size), NumBytesToAdd, CurrentArea->StartOffset + CurrentArea->Size, &TotalSize);
-			UE_LOG(LogElectraBufferReader, Log, TEXT("CreateNewArea(): read %lld bytes from offset %lld; total file size = %lld, received %lld bytes"), (long long int)NumBytesToAdd, (long long int)(CurrentArea->StartOffset + CurrentArea->Size), (long long int)TotalSize, (long long int)NumRead);
+			UE_LOG(LogElectraBufferReader, VeryVerbose, TEXT("CreateNewArea(): read %lld bytes from offset %lld; total file size = %lld, received %lld bytes"), (long long int)NumBytesToAdd, (long long int)(CurrentArea->StartOffset + CurrentArea->Size), (long long int)TotalSize, (long long int)NumRead);
 			if (TotalSize >= 0)
 			{
 				TotalDataSize = TotalSize;
@@ -111,7 +111,7 @@ namespace Electra
 			{
 				if (NumRead < NumBytesToAdd || (TotalSize >= 0 && CurrentArea->StartOffset + CurrentArea->Size + NumRead >= TotalSize))
 				{
-					UE_LOG(LogElectraBufferReader, Log, TEXT("EnlargeCurrentAreaBy(): reached the end of the file"));
+					UE_LOG(LogElectraBufferReader, VeryVerbose, TEXT("EnlargeCurrentAreaBy(): reached the end of the file"));
 					CurrentArea->bEOS = true;
 				}
 				CurrentArea->Size += NumRead;
@@ -119,7 +119,7 @@ namespace Electra
 				// Consolidate forward?
 				if (Next)
 				{
-					UE_LOG(LogElectraBufferReader, Log, TEXT("EnlargeCurrentAreaBy(): merging blocks from offsets %lld and %lld into one"), (long long int)CurrentArea->StartOffset, (long long int)Next->StartOffset);
+					UE_LOG(LogElectraBufferReader, VeryVerbose, TEXT("EnlargeCurrentAreaBy(): merging blocks from offsets %lld and %lld into one"), (long long int)CurrentArea->StartOffset, (long long int)Next->StartOffset);
 					FMemory::Memcpy((void*)(CurrentArea->Data + CurrentArea->Size), Next->Data, Next->Size);
 					CurrentArea->bEOS = Next->bEOS;
 					CurrentArea->Size += Next->Size;
@@ -130,12 +130,10 @@ namespace Electra
 			}
 			else
 			{
-check(!"TODO");
 				IDataProvider::EError Error = static_cast<IDataProvider::EError>(NumRead);
 				if (Error == IDataProvider::EError::Failed)
 				{
-					// TODO: Facility & code
-					LastError.SetError(UEMEDIA_ERROR_READ_ERROR).SetFacility(Facility::EFacility::HTTPReader).SetCode(4711).SetMessage(TEXT("OnReadAssetData() failed"));
+					LastError.SetError(UEMEDIA_ERROR_READ_ERROR).SetFacility(Facility::EFacility::BufferedDataReader).SetCode(EBufferedDataReaderError::ReadAssetDataFailed).SetMessage(TEXT("OnReadAssetData() failed"));
 				}
 			}
 		}
@@ -369,7 +367,7 @@ check(!"TODO");
 	{
 		if (TotalDataSize >= 0 && AbsolutePosition > TotalDataSize)
 		{
-			UE_LOG(LogElectraBufferReader, Log, TEXT("SeekTo(): Seeking beyond the end of the file"));
+			UE_LOG(LogElectraBufferReader, VeryVerbose, TEXT("SeekTo(): Seeking beyond the end of the file"));
 			return false;
 		}
 
