Fix ElectraPlayer RFC6381 codec name for opus

diff --git a/Engine/Plugins/Media/ElectraPlayer/Source/ElectraPlayerRuntime/Private/Runtime/Player/HLS/ManifestBuilderHLS.cpp b/Engine/Plugins/Media/ElectraPlayer/Source/ElectraPlayerRuntime/Private/Runtime/Player/HLS/ManifestBuilderHLS.cpp
--- a/Engine/Plugins/Media/ElectraPlayer/Source/ElectraPlayerRuntime/Private/Runtime/Player/HLS/ManifestBuilderHLS.cpp
+++ b/Engine/Plugins/Media/ElectraPlayer/Source/ElectraPlayerRuntime/Private/Runtime/Player/HLS/ManifestBuilderHLS.cpp
@@ -600,7 +600,7 @@ FErrorDetail FManifestBuilderHLS::SetupVariants(FManifestHLSInternal* Manifest,
 				{
 					si.SetStreamType(EStreamType::Audio);
 					si.SetCodec(FStreamCodecInformation::ECodec::Audio4CC);
-					si.SetCodec4CC(0x4f707573);	// Make4CC('O','p','u','s'));
+					si.SetCodec4CC(0x6f707573);	// Make4CC('o','p','u','s'));
 					si.SetCodecSpecifierRFC6381(CodecList[nCodec]);
 					// For lack of knowledge pretend this is stereo.
 					si.SetChannelConfiguration(2);
diff --git a/Engine/Plugins/Media/ElectraPlayer/Source/ElectraPlayerRuntime/Private/Runtime/StreamTypes.cpp b/Engine/Plugins/Media/ElectraPlayer/Source/ElectraPlayerRuntime/Private/Runtime/StreamTypes.cpp
--- a/Engine/Plugins/Media/ElectraPlayer/Source/ElectraPlayerRuntime/Private/Runtime/StreamTypes.cpp
+++ b/Engine/Plugins/Media/ElectraPlayer/Source/ElectraPlayerRuntime/Private/Runtime/StreamTypes.cpp
@@ -246,12 +246,12 @@ namespace Electra
 			// Presently not supported.
 			return false;
 		}
-		else if (CodecOTI.Equals("Opus"))
+		else if (CodecOTI.Equals("opus"))
 		{
 			StreamType = EStreamType::Audio;
 			CodecSpecifier = CodecOTI;
 			Codec = ECodec::Audio4CC;
-			Codec4CC = Make4CC('O','p','u','s');
+			Codec4CC = Make4CC('o','p','u','s');
 			return true;
 		}
 		else if (CodecOTI.Equals(TEXT("wvtt")))
