ElectraPlayer: Changed Opus codec 4cc back to what is required by the decoder factory.

#rnx

diff --git a/Engine/Plugins/Media/ElectraPlayer/Source/ElectraPlayerRuntime/Private/Runtime/Player/HLS/ManifestBuilderHLS.cpp b/Engine/Plugins/Media/ElectraPlayer/Source/ElectraPlayerRuntime/Private/Runtime/Player/HLS/ManifestBuilderHLS.cpp
--- a/Engine/Plugins/Media/ElectraPlayer/Source/ElectraPlayerRuntime/Private/Runtime/Player/HLS/ManifestBuilderHLS.cpp
+++ b/Engine/Plugins/Media/ElectraPlayer/Source/ElectraPlayerRuntime/Private/Runtime/Player/HLS/ManifestBuilderHLS.cpp
@@ -600,7 +600,7 @@ FErrorDetail FManifestBuilderHLS::SetupVariants(FManifestHLSInternal* Manifest,
 				{
 					si.SetStreamType(EStreamType::Audio);
 					si.SetCodec(FStreamCodecInformation::ECodec::Audio4CC);
-					si.SetCodec4CC(0x6f707573);	// Make4CC('o','p','u','s'));
+					si.SetCodec4CC(0x4f707573);	// Make4CC('O','p','u','s'));
 					si.SetCodecSpecifierRFC6381(CodecList[nCodec]);
 					// For lack of knowledge pretend this is stereo.
 					si.SetChannelConfiguration(2);
diff --git a/Engine/Plugins/Media/ElectraPlayer/Source/ElectraPlayerRuntime/Private/Runtime/StreamTypes.cpp b/Engine/Plugins/Media/ElectraPlayer/Source/ElectraPlayerRuntime/Private/Runtime/StreamTypes.cpp
--- a/Engine/Plugins/Media/ElectraPlayer/Source/ElectraPlayerRuntime/Private/Runtime/StreamTypes.cpp
+++ b/Engine/Plugins/Media/ElectraPlayer/Source/ElectraPlayerRuntime/Private/Runtime/StreamTypes.cpp
@@ -83,7 +83,7 @@ namespace Electra
 
 	bool FStreamCodecInformation::ParseFromRFC6381(const FString& CodecOTI)
 	{
-		if (CodecOTI.StartsWith("avc"))
+		if (CodecOTI.StartsWith(TEXT("avc")))
 		{
 			// avc1 and avc3 (inband SPS/PPS) are recognized.
 			StreamType = EStreamType::Video;
@@ -133,7 +133,7 @@ namespace Electra
 			}
 			return true;
 		}
-		else if (CodecOTI.StartsWith("hvc") || CodecOTI.StartsWith("hev"))
+		else if (CodecOTI.StartsWith(TEXT("hvc")) || CodecOTI.StartsWith(TEXT("hev")))
 		{
 			FString oti = CodecOTI;
 			FString Temp;
@@ -210,7 +210,7 @@ namespace Electra
 			}
 			return false;
 		}
-		else if (CodecOTI.StartsWith("mp4a"))
+		else if (CodecOTI.StartsWith(TEXT("mp4a")))
 		{
 			StreamType = EStreamType::Audio;
 			CodecSpecifier = CodecOTI;
@@ -238,7 +238,7 @@ namespace Electra
 			}
 			return true;
 		}
-		else if (CodecOTI.StartsWith("ec-3") || CodecOTI.StartsWith("ec+3") || CodecOTI.StartsWith("ec3") || CodecOTI.StartsWith("eac3"))
+		else if (CodecOTI.StartsWith(TEXT("ec-3")) || CodecOTI.StartsWith(TEXT("ec+3")) || CodecOTI.StartsWith(TEXT("ec3")) || CodecOTI.StartsWith(TEXT("eac3")))
 		{
 			StreamType = EStreamType::Audio;
 			CodecSpecifier = CodecOTI;
@@ -246,12 +246,12 @@ namespace Electra
 			// Presently not supported.
 			return false;
 		}
-		else if (CodecOTI.Equals("opus"))
+		else if (CodecOTI.Equals(TEXT("opus"), ESearchCase::IgnoreCase))
 		{
 			StreamType = EStreamType::Audio;
 			CodecSpecifier = CodecOTI;
 			Codec = ECodec::Audio4CC;
-			Codec4CC = Make4CC('o','p','u','s');
+			Codec4CC = Make4CC('O','p','u','s');
 			return true;
 		}
 		else if (CodecOTI.Equals(TEXT("wvtt")))
