ElectraPlayer: Fixed missing frames at the end of a DASH period

#rnx

diff --git a/Engine/Plugins/Media/ElectraPlayer/Source/ElectraPlayerRuntime/Private/Runtime/Player/DASH/StreamReaderFMP4DASH.cpp b/Engine/Plugins/Media/ElectraPlayer/Source/ElectraPlayerRuntime/Private/Runtime/Player/DASH/StreamReaderFMP4DASH.cpp
--- a/Engine/Plugins/Media/ElectraPlayer/Source/ElectraPlayerRuntime/Private/Runtime/Player/DASH/StreamReaderFMP4DASH.cpp
+++ b/Engine/Plugins/Media/ElectraPlayer/Source/ElectraPlayerRuntime/Private/Runtime/Player/DASH/StreamReaderFMP4DASH.cpp
@@ -1180,7 +1180,7 @@ void FStreamReaderFMP4DASH::FStreamHandler::HandleRequest()
 									// Entirely past the time allowed?
 									if (bIsLastSegment || Request->Segment.bFrameAccuracyRequired)
 									{
-										if (AUPTS >= MediaLocalLastAUTime)
+										if (AUDTS >= MediaLocalLastAUTime && AUPTS >= MediaLocalLastAUTime)
 										{
 											AccessUnit->DropState |= FAccessUnit::EDropState::TooLate;
 										}
