ElectraPlayer: Added a missing return statement in the case where the resource delegate is no longer valid.

#rnx

diff --git a/Engine/Plugins/Media/ElectraPlayer/Source/ElectraPlayerRuntime/Private/Runtime/Decoder/Android/AndroidElectraDecoderResourceManager.cpp b/Engine/Plugins/Media/ElectraPlayer/Source/ElectraPlayerRuntime/Private/Runtime/Decoder/Android/AndroidElectraDecoderResourceManager.cpp
--- a/Engine/Plugins/Media/ElectraPlayer/Source/ElectraPlayerRuntime/Private/Runtime/Decoder/Android/AndroidElectraDecoderResourceManager.cpp
+++ b/Engine/Plugins/Media/ElectraPlayer/Source/ElectraPlayerRuntime/Private/Runtime/Decoder/Android/AndroidElectraDecoderResourceManager.cpp
@@ -76,6 +76,7 @@ void FElectraDecoderResourceManagerAndroid::FInstanceVars::RequestSurface(TWeakP
 	if (!rd.IsValid())
 	{
 		Notify(IDecoderPlatformResourceAndroid::ISurfaceRequestCallback::ESurfaceType::Error, nullptr);
+		return;
 	}
 
 	void* Surface = rd->VideoDecoderResourceDelegate_GetCodecSurface();
