Fixing up CPU-buffer NV12 data representation for texture samples / decoder output on Apple devices
#rnx

diff --git a/Engine/Plugins/Media/ElectraPlayer/Source/ElectraPlayerRuntime/Private/Runtime/Decoder/Apple/ElectraVideoDecoder_Apple.cpp b/Engine/Plugins/Media/ElectraPlayer/Source/ElectraPlayerRuntime/Private/Runtime/Decoder/Apple/ElectraVideoDecoder_Apple.cpp
--- a/Engine/Plugins/Media/ElectraPlayer/Source/ElectraPlayerRuntime/Private/Runtime/Decoder/Apple/ElectraVideoDecoder_Apple.cpp
+++ b/Engine/Plugins/Media/ElectraPlayer/Source/ElectraPlayerRuntime/Private/Runtime/Decoder/Apple/ElectraVideoDecoder_Apple.cpp
@@ -35,6 +35,7 @@ void FElectraPlayerVideoDecoderOutputApple::Initialize(CVImageBufferRef InImageB
 
 		Stride = CVPixelBufferGetBytesPerRow(ImageBufferRef);
 	}
+	SampleDim = GetOutputDim();
 }
 
 void FElectraPlayerVideoDecoderOutputApple::InitializeWithBuffer(const void* InBuffer, uint32 InSize, uint32 InStride, FIntPoint Dim, TSharedPtr<Electra::FParamDict, ESPMode::ThreadSafe> InParamDict)
@@ -45,6 +46,7 @@ void FElectraPlayerVideoDecoderOutputApple::InitializeWithBuffer(const void* InB
 	Buffer->Append((uint8*)InBuffer, InSize);
 
 	Stride = InStride;
+	SampleDim = Dim;
 }
 
 void FElectraPlayerVideoDecoderOutputApple::InitializeWithBuffer(TSharedPtr<TArray<uint8>, ESPMode::ThreadSafe> InBuffer, uint32 InStride, FIntPoint Dim, TSharedPtr<Electra::FParamDict, ESPMode::ThreadSafe> InParamDict)
@@ -54,6 +56,7 @@ void FElectraPlayerVideoDecoderOutputApple::InitializeWithBuffer(TSharedPtr<TArr
 	Buffer = MoveTemp(InBuffer);
 
 	Stride = InStride;
+	SampleDim = Dim;
 }
 
 void FElectraPlayerVideoDecoderOutputApple::SetOwner(const TSharedPtr<IDecoderOutputOwner, ESPMode::ThreadSafe>& InOwningRenderer)
@@ -95,6 +98,11 @@ uint32 FElectraPlayerVideoDecoderOutputApple::GetStride() const
 	return Stride;
 }
 
+FIntPoint FElectraPlayerVideoDecoderOutputApple::GetDim() const
+{
+	return SampleDim;
+}
+
 CVImageBufferRef FElectraPlayerVideoDecoderOutputApple::GetImageBuffer() const
 {
 	return ImageBufferRef;
diff --git a/Engine/Plugins/Media/ElectraPlayer/Source/ElectraPlayerRuntime/Private/Runtime/Decoder/Apple/ElectraVideoDecoder_Apple.h b/Engine/Plugins/Media/ElectraPlayer/Source/ElectraPlayerRuntime/Private/Runtime/Decoder/Apple/ElectraVideoDecoder_Apple.h
--- a/Engine/Plugins/Media/ElectraPlayer/Source/ElectraPlayerRuntime/Private/Runtime/Decoder/Apple/ElectraVideoDecoder_Apple.h
+++ b/Engine/Plugins/Media/ElectraPlayer/Source/ElectraPlayerRuntime/Private/Runtime/Decoder/Apple/ElectraVideoDecoder_Apple.h
@@ -30,6 +30,7 @@ public:
 	void ShutdownPoolable() override;
 
 	uint32 GetStride() const override;
+	FIntPoint GetDim() const override;
 
 	CVImageBufferRef GetImageBuffer() const override;
 
@@ -37,6 +38,7 @@ public:
 
 private:
 	uint32 Stride;
+	FIntPoint SampleDim;
 
 	TSharedPtr<TArray<uint8>, ESPMode::ThreadSafe> Buffer;
 	CVImageBufferRef ImageBufferRef;
