ElectraPlayer: Added provisions for an external data cache

#rnx

diff --git a/Engine/Plugins/Media/ElectraPlayer/Source/ElectraPlayerRuntime/Private/ElectraPlayer.cpp b/Engine/Plugins/Media/ElectraPlayer/Source/ElectraPlayerRuntime/Private/ElectraPlayer.cpp
--- a/Engine/Plugins/Media/ElectraPlayer/Source/ElectraPlayerRuntime/Private/ElectraPlayer.cpp
+++ b/Engine/Plugins/Media/ElectraPlayer/Source/ElectraPlayerRuntime/Private/ElectraPlayer.cpp
@@ -290,6 +290,10 @@ bool FElectraPlayer::OpenInternal(const FString& Url, const FParamDict& InPlayer
 		NewPlayer->AdaptivePlayer->AddMetricsReceiver(this);
 		NewPlayer->AdaptivePlayer->SetStaticResourceProviderCallback(StaticResourceProvider);
 		NewPlayer->AdaptivePlayer->SetVideoDecoderResourceDelegate(VideoDecoderResourceDelegate);
+		if (PlaystartOptions.ExternalDataCache.IsValid())
+		{
+			NewPlayer->AdaptivePlayer->SetPlayerDataCache(PlaystartOptions.ExternalDataCache);
+		}
 
 		// Create the subtitle receiver and register it with the player.
 		MediaPlayerSubtitleReceiver = MakeSharedTS<FSubtitleEventReceiver>();
diff --git a/Engine/Plugins/Media/ElectraPlayer/Source/ElectraPlayerRuntime/Private/Runtime/HTTP/HTTPManager.cpp b/Engine/Plugins/Media/ElectraPlayer/Source/ElectraPlayerRuntime/Private/Runtime/HTTP/HTTPManager.cpp
--- a/Engine/Plugins/Media/ElectraPlayer/Source/ElectraPlayerRuntime/Private/Runtime/HTTP/HTTPManager.cpp
+++ b/Engine/Plugins/Media/ElectraPlayer/Source/ElectraPlayerRuntime/Private/Runtime/HTTP/HTTPManager.cpp
@@ -195,7 +195,7 @@ namespace Electra
 				{
 					TSharedPtrTS<IHTTPResponseCache::FCacheItem> CacheItem;
 					IHTTPResponseCache::EScatterResult Result;
-					Result = HttpResponseCache->GetScatteredCacheEntity(CacheItem, ActiveResponse.URL, ActiveResponse.Range);
+					Result = HttpResponseCache->GetScatteredCacheEntity(CacheItem, ActiveResponse.URL, ActiveResponse.Range, ActiveResponse.Quality);
 					if (Result == IHTTPResponseCache::EScatterResult::FullHit)
 					{
 						check(CacheItem.IsValid());
@@ -251,6 +251,8 @@ namespace Electra
 				ElectraHTTPStream::FHttpRange OriginalRange;
 				ElectraHTTPStream::FHttpRange ReceivedContentRange;
 				//
+				IHTTPResponseCache::FQualityInfo Quality;
+				//
 				bool bIsSubRangeRequest = false;
 				int32 NumSubRangeRequest = 0;
 				//
@@ -610,6 +612,9 @@ namespace Electra
 		Handle->bHttpRequestFirstEvent = true;
 		Handle->ActiveResponse.NumSubRangeRequest = 0;
 		Handle->ActiveResponse.OriginalRange = Request->Parameters.Range;
+		Handle->ActiveResponse.Quality.StreamType = Request->Parameters.StreamType;
+		Handle->ActiveResponse.Quality.QualityIndex = Request->Parameters.QualityIndex;
+		Handle->ActiveResponse.Quality.MaxQualityIndex = Request->Parameters.MaxQualityIndex;
 
 		if (PrepareHTTPHandle(Now, Handle.Get(), Request, true))
 		{
@@ -1493,6 +1498,9 @@ namespace Electra
 									CacheItem->Range.SetEndIncluding(CacheItem->Range.DocumentSize - 1);
 								}
 								CacheItem->Response = Response;
+								CacheItem->Quality.QualityIndex = Request->Parameters.QualityIndex;
+								CacheItem->Quality.MaxQualityIndex = Request->Parameters.MaxQualityIndex;
+								CacheItem->Quality.StreamType = Request->Parameters.StreamType;
 								Request->ResponseCache->CacheEntity(CacheItem);
 							}
  
diff --git a/Engine/Plugins/Media/ElectraPlayer/Source/ElectraPlayerRuntime/Private/Runtime/HTTP/HTTPManager.h b/Engine/Plugins/Media/ElectraPlayer/Source/ElectraPlayerRuntime/Private/Runtime/HTTP/HTTPManager.h
--- a/Engine/Plugins/Media/ElectraPlayer/Source/ElectraPlayerRuntime/Private/Runtime/HTTP/HTTPManager.h
+++ b/Engine/Plugins/Media/ElectraPlayer/Source/ElectraPlayerRuntime/Private/Runtime/HTTP/HTTPManager.h
@@ -6,6 +6,7 @@
 #include "Delegates/Delegate.h"
 
 #include "StreamDataBuffer.h"
+#include "StreamTypes.h"
 #include "OptionalValue.h"
 
 #include "ParameterDictionary.h"
@@ -246,6 +247,10 @@ namespace Electra
 			FTimeValue							NoDataTimeout;					//!< Optional timeout when no data is being received
 			TArray<uint8>						PostData;						//!< Data for POST
 			bool								bCollectTimingTraces = false;	//!< Whether or not to collect download timing traces.
+
+			EStreamType							StreamType = EStreamType::Unsupported;
+			int32								QualityIndex = 0;
+			int32								MaxQualityIndex = 0;
 		};
 
 
diff --git a/Engine/Plugins/Media/ElectraPlayer/Source/ElectraPlayerRuntime/Private/Runtime/HTTP/HTTPResponseCache.cpp b/Engine/Plugins/Media/ElectraPlayer/Source/ElectraPlayerRuntime/Private/Runtime/HTTP/HTTPResponseCache.cpp
--- a/Engine/Plugins/Media/ElectraPlayer/Source/ElectraPlayerRuntime/Private/Runtime/HTTP/HTTPResponseCache.cpp
+++ b/Engine/Plugins/Media/ElectraPlayer/Source/ElectraPlayerRuntime/Private/Runtime/HTTP/HTTPResponseCache.cpp
@@ -14,12 +14,13 @@ namespace Electra
 class FHTTPResponseCache : public IHTTPResponseCache
 {
 public:
-	FHTTPResponseCache(IPlayerSessionServices* SessionServices, const FParamDict& Options);
+	FHTTPResponseCache(IPlayerSessionServices* SessionServices, const FParamDict& Options, TSharedPtr<IElectraPlayerDataCache, ESPMode::ThreadSafe> ExternalCache);
 	virtual ~FHTTPResponseCache();
 
-	virtual void HandleEntityExpiration() override;
-	virtual void CacheEntity(TSharedPtrTS<FCacheItem> EntityToAdd) override;
-	virtual EScatterResult GetScatteredCacheEntity(TSharedPtrTS<FCacheItem>& OutScatteredCachedEntity, const FString& URL, const ElectraHTTPStream::FHttpRange& Range) override;
+	void Disable() override;
+	void HandleEntityExpiration() override;
+	void CacheEntity(TSharedPtrTS<FCacheItem> EntityToAdd) override;
+	EScatterResult GetScatteredCacheEntity(TSharedPtrTS<FCacheItem>& OutScatteredCachedEntity, const FString& URL, const ElectraHTTPStream::FHttpRange& Range, const FQualityInfo& ForQuality) override;
 
 private:
 	/**
@@ -343,6 +344,7 @@ private:
 
 	void EvictToAddSize(int64 ResponseSize);
 
+	TSharedPtr<IElectraPlayerDataCache, ESPMode::ThreadSafe> ExternalCache;
 	IPlayerSessionServices* SessionServices = nullptr;
 	int64 MaxElementSize = 0;
 	int32 MaxNumElements = 0;
@@ -357,30 +359,34 @@ private:
 /***************************************************************************************************************************************************/
 /***************************************************************************************************************************************************/
 
-TSharedPtrTS<IHTTPResponseCache> IHTTPResponseCache::Create(IPlayerSessionServices* SessionServices, const FParamDict& Options)
+TSharedPtrTS<IHTTPResponseCache> IHTTPResponseCache::Create(IPlayerSessionServices* SessionServices, const FParamDict& Options, TSharedPtr<IElectraPlayerDataCache, ESPMode::ThreadSafe> ExternalCache)
 {
-	return MakeSharedTS<FHTTPResponseCache>(SessionServices, Options);
+	return MakeSharedTS<FHTTPResponseCache>(SessionServices, Options, ExternalCache);
 }
 
 /***************************************************************************************************************************************************/
 /***************************************************************************************************************************************************/
 /***************************************************************************************************************************************************/
 
-FHTTPResponseCache::FHTTPResponseCache(IPlayerSessionServices* InSessionServices, const FParamDict& InOptions)
-	: SessionServices(InSessionServices)
+FHTTPResponseCache::FHTTPResponseCache(IPlayerSessionServices* InSessionServices, const FParamDict& InOptions, TSharedPtr<IElectraPlayerDataCache, ESPMode::ThreadSafe> InExternalCache)
+	: ExternalCache(InExternalCache), SessionServices(InSessionServices)
 {
-	int64 v;
-	v = InOptions.GetValue(OptionKeyResponseCacheMaxByteSize).SafeGetInt64(0);
-	if (v >= 0)
+	// Only configure our cache when there is no external one!
+	if (!ExternalCache.IsValid())
 	{
-		MaxElementSize = (int64)v;
-	}
-	// Max number of elements is probably not used a lot, if at all.
-	// If not specified we allow for some reasonable number.
-	v = InOptions.GetValue(OptionKeyResponseCacheMaxEntries).SafeGetInt64(8192);
-	if (v >= 0)
-	{
-		MaxNumElements = (int32) v;
+		int64 v;
+		v = InOptions.GetValue(OptionKeyResponseCacheMaxByteSize).SafeGetInt64(0);
+		if (v >= 0)
+		{
+			MaxElementSize = (int64)v;
+		}
+		// Max number of elements is probably not used a lot, if at all.
+		// If not specified we allow for some reasonable number.
+		v = InOptions.GetValue(OptionKeyResponseCacheMaxEntries).SafeGetInt64(8192);
+		if (v >= 0)
+		{
+			MaxNumElements = (int32) v;
+		}
 	}
 }
 
@@ -402,6 +408,17 @@ void FHTTPResponseCache::EvictToAddSize(int64 ResponseSize)
 }
 
 
+void FHTTPResponseCache::Disable()
+{
+	FScopeLock ScopeLock(&Lock);
+	Cache.Empty();
+	ExternalCache.Reset();
+	MaxElementSize = 0;
+	MaxNumElements = 0;
+	SizeInUse = 0;
+}
+
+
 void FHTTPResponseCache::HandleEntityExpiration()
 {
 }
@@ -413,183 +430,315 @@ void FHTTPResponseCache::CacheEntity(TSharedPtrTS<FCacheItem> EntityToAdd)
 	{
 		if (EntityToAdd->EffectiveURL.Len())
 		{
-			int64 SizeRequired = EntityToAdd->Response->GetResponseData().GetNumTotalBytesAdded();
-			EvictToAddSize(SizeRequired);
-			FScopeLock ScopeLock(&Lock);
-			if (SizeInUse + SizeRequired <= MaxElementSize && Cache.Num() < MaxNumElements)
+			if (!ExternalCache.IsValid())
 			{
-				Cache.Insert(MoveTemp(EntityToAdd), 0);
-				SizeInUse += SizeRequired;
+				int64 SizeRequired = EntityToAdd->Response->GetResponseData().GetNumTotalBytesAdded();
+				EvictToAddSize(SizeRequired);
+				FScopeLock ScopeLock(&Lock);
+				if (SizeInUse + SizeRequired <= MaxElementSize && Cache.Num() < MaxNumElements)
+				{
+					Cache.Insert(MoveTemp(EntityToAdd), 0);
+					SizeInUse += SizeRequired;
+				}
+			}
+			else
+			{
+				IElectraPlayerDataCache::FItemInfo Info;
+				Info.URI = EntityToAdd->RequestedURL;
+				switch(EntityToAdd->Quality.StreamType)
+				{
+					case EStreamType::Video:
+						Info.StreamType = IElectraPlayerDataCache::FItemInfo::EStreamType::Video;
+						break;
+					case EStreamType::Audio:
+						Info.StreamType = IElectraPlayerDataCache::FItemInfo::EStreamType::Audio;
+						break;
+					default:
+						Info.StreamType = IElectraPlayerDataCache::FItemInfo::EStreamType::Other;
+						break;
+				}
+				Info.QualityIndex = EntityToAdd->Quality.QualityIndex;
+				Info.MaxQualityIndex = EntityToAdd->Quality.MaxQualityIndex;
+
+				if (EntityToAdd->Range.GetStart() == 0 && EntityToAdd->Range.GetEndIncluding() + 1 == EntityToAdd->Range.GetDocumentSize())
+				{
+					Info.Range.TotalSize = EntityToAdd->Range.GetDocumentSize();
+				}
+				else
+				{
+					Info.Range.Start = EntityToAdd->Range.GetStart();
+					Info.Range.EndIncluding = EntityToAdd->Range.GetEndIncluding();
+					Info.Range.TotalSize = EntityToAdd->Range.GetDocumentSize();
+				}
+				const uint8* Data = nullptr;
+				int64 DataSize = 0;
+				EntityToAdd->Response->GetResponseData().GetBaseBuffer(Data, DataSize);
+				TSharedPtr<TArray<uint8>, ESPMode::ThreadSafe> Buf = MakeShared<TArray<uint8>, ESPMode::ThreadSafe>();
+				Buf->AddUninitialized(DataSize);
+				FMemory::Memcpy(Buf->GetData(), Data, DataSize);
+				ExternalCache->AddElementToCache(Info, Buf);
 			}
 		}
 	}
 }
 
-FHTTPResponseCache::EScatterResult FHTTPResponseCache::GetScatteredCacheEntity(TSharedPtrTS<FCacheItem>& OutScatteredCachedEntity, const FString& URL, const ElectraHTTPStream::FHttpRange& InRange)
+FHTTPResponseCache::EScatterResult FHTTPResponseCache::GetScatteredCacheEntity(TSharedPtrTS<FCacheItem>& OutScatteredCachedEntity, const FString& URL, const ElectraHTTPStream::FHttpRange& InRange, const FQualityInfo& InForQuality)
 {
-	FScopeLock lock(&Lock);
-	// Get a list of all cached blocks for this URL.
-	TArray<TSharedPtrTS<FCacheItem>> CachedBlocks;
-	for(auto &Entry : Cache)
+	if (!ExternalCache.IsValid())
 	{
-		if (Entry->RequestedURL.Equals(URL) || Entry->EffectiveURL.Equals(URL))
+		FScopeLock lock(&Lock);
+		// Get a list of all cached blocks for this URL.
+		TArray<TSharedPtrTS<FCacheItem>> CachedBlocks;
+		for(auto &Entry : Cache)
 		{
-			CachedBlocks.Add(Entry);
+			if (Entry->RequestedURL.Equals(URL) || Entry->EffectiveURL.Equals(URL))
+			{
+				CachedBlocks.Add(Entry);
+			}
 		}
-	}
 
-	ElectraHTTPStream::FHttpRange Range(InRange);
+		ElectraHTTPStream::FHttpRange Range(InRange);
 
-	// Set up the default output that encompasses the request.
-	// This is what needs to be fetched for cache misses.
-	OutScatteredCachedEntity = MakeSharedTS<IHTTPResponseCache::FCacheItem>();
-	OutScatteredCachedEntity->RequestedURL = URL;
-	OutScatteredCachedEntity->EffectiveURL = URL;
-	OutScatteredCachedEntity->Range = Range;
+		// Set up the default output that encompasses the request.
+		// This is what needs to be fetched for cache misses.
+		OutScatteredCachedEntity = MakeSharedTS<IHTTPResponseCache::FCacheItem>();
+		OutScatteredCachedEntity->RequestedURL = URL;
+		OutScatteredCachedEntity->EffectiveURL = URL;
+		OutScatteredCachedEntity->Range = Range;
 
-	// Quick out if there are no cached blocks at all.
-	if (CachedBlocks.Num() == 0)
-	{
-		return EScatterResult::Miss;
-	}
+		// Quick out if there are no cached blocks at all.
+		if (CachedBlocks.Num() == 0)
+		{
+			return EScatterResult::Miss;
+		}
 
-	// Adjust the input range to valid values.
-	Range.DocumentSize = CachedBlocks[0]->Range.GetDocumentSize();
-	if (Range.GetStart() < 0)
-	{
-		Range.SetStart(0);
-	}
-	if (Range.GetEndIncluding() < 0)
-	{
-		Range.SetEndIncluding(Range.GetDocumentSize() - 1);
-	}
+		// Adjust the input range to valid values.
+		Range.DocumentSize = CachedBlocks[0]->Range.GetDocumentSize();
+		if (Range.GetStart() < 0)
+		{
+			Range.SetStart(0);
+		}
+		if (Range.GetEndIncluding() < 0)
+		{
+			Range.SetEndIncluding(Range.GetDocumentSize() - 1);
+		}
 
-	// Sort the blocks by ascending range.
-	CachedBlocks.Sort([](const TSharedPtrTS<FCacheItem>& e1, const TSharedPtrTS<FCacheItem>& e2)
-	{
-		return e1->Range.GetStart() < e2->Range.GetStart();
-	});
-
-	// Find a cached block that overlaps the requested range.
-	int64 rs = Range.GetStart();
-	int64 re = Range.GetEndIncluding();
-	int32 firstIndex;
-	for(firstIndex=0; firstIndex<CachedBlocks.Num(); ++firstIndex)
-	{
-		// Any overlap?
-		const int64 cbstart = CachedBlocks[firstIndex]->Range.GetStart();
-		const int64 cbend = CachedBlocks[firstIndex]->Range.GetEndIncluding();
-		if (re >= cbstart && rs <= cbend)
+		// Sort the blocks by ascending range.
+		CachedBlocks.Sort([](const TSharedPtrTS<FCacheItem>& e1, const TSharedPtrTS<FCacheItem>& e2)
 		{
-			// Is the request start before the cached block, ie we need to get additional data at the beginning?
-			if (rs < cbstart)
-			{
-				// The end of the requested range may also not be covered by the cache, but that is not relevant
-				// here as we move forward sequentially, one overlap at a time.
-				// Data up to the beginning of this cached block must be fetched first.
-				OutScatteredCachedEntity->Range = Range;
-				OutScatteredCachedEntity->Range.SetEndIncluding(cbstart - 1);
-				OutScatteredCachedEntity->EffectiveURL = CachedBlocks[firstIndex]->EffectiveURL;
-				return EScatterResult::PartialHit;
-			}
-			else
+			return e1->Range.GetStart() < e2->Range.GetStart();
+		});
+
+		// Find a cached block that overlaps the requested range.
+		int64 rs = Range.GetStart();
+		int64 re = Range.GetEndIncluding();
+		int32 firstIndex;
+		for(firstIndex=0; firstIndex<CachedBlocks.Num(); ++firstIndex)
+		{
+			// Any overlap?
+			const int64 cbstart = CachedBlocks[firstIndex]->Range.GetStart();
+			const int64 cbend = CachedBlocks[firstIndex]->Range.GetEndIncluding();
+			if (re >= cbstart && rs <= cbend)
 			{
-				break;
+				// Is the request start before the cached block, ie we need to get additional data at the beginning?
+				if (rs < cbstart)
+				{
+					// The end of the requested range may also not be covered by the cache, but that is not relevant
+					// here as we move forward sequentially, one overlap at a time.
+					// Data up to the beginning of this cached block must be fetched first.
+					OutScatteredCachedEntity->Range = Range;
+					OutScatteredCachedEntity->Range.SetEndIncluding(cbstart - 1);
+					OutScatteredCachedEntity->EffectiveURL = CachedBlocks[firstIndex]->EffectiveURL;
+					return EScatterResult::PartialHit;
+				}
+				else
+				{
+					break;
+				}
 			}
 		}
-	}
-	// No overlap anywhere?
-	if (firstIndex == CachedBlocks.Num())
-	{
-		return EScatterResult::Miss;
-	}
+		// No overlap anywhere?
+		if (firstIndex == CachedBlocks.Num())
+		{
+			return EScatterResult::Miss;
+		}
 
-	// Exact match?
-	if (rs == CachedBlocks[firstIndex]->Range.GetStart() && re == CachedBlocks[firstIndex]->Range.GetEndIncluding())
-	{
-		OutScatteredCachedEntity->Response = MakeShared<FWrappedCacheResponse, ESPMode::ThreadSafe>(CachedBlocks[firstIndex]->Response);
+		// Exact match?
+		if (rs == CachedBlocks[firstIndex]->Range.GetStart() && re == CachedBlocks[firstIndex]->Range.GetEndIncluding())
+		{
+			OutScatteredCachedEntity->Response = MakeShared<FWrappedCacheResponse, ESPMode::ThreadSafe>(CachedBlocks[firstIndex]->Response);
+			OutScatteredCachedEntity->EffectiveURL = CachedBlocks[firstIndex]->EffectiveURL;
+			Cache.Remove(CachedBlocks[firstIndex]);
+			Cache.Insert(CachedBlocks[firstIndex], 0);
+			return EScatterResult::FullHit;
+		}
+
+		// Set up the partial cache response.
+		TSharedPtr<FSynthesizedCacheResponse, ESPMode::ThreadSafe> Response = MakeShared<FSynthesizedCacheResponse, ESPMode::ThreadSafe>();
+		Response->SetURL(CachedBlocks[firstIndex]->EffectiveURL);
 		OutScatteredCachedEntity->EffectiveURL = CachedBlocks[firstIndex]->EffectiveURL;
-		Cache.Remove(CachedBlocks[firstIndex]);
-		Cache.Insert(CachedBlocks[firstIndex], 0);
-		return EScatterResult::FullHit;
-	}
+		Response->CopyHeadersExceptContentSizesAndDateFrom(CachedBlocks[firstIndex]->Response);
+		Response->SetHTTPStatusLine(CachedBlocks[firstIndex]->Response->GetHTTPStatusLine());
 
-	// Set up the partial cache response.
-	TSharedPtr<FSynthesizedCacheResponse, ESPMode::ThreadSafe> Response = MakeShared<FSynthesizedCacheResponse, ESPMode::ThreadSafe>();
-	Response->SetURL(CachedBlocks[firstIndex]->EffectiveURL);
-	OutScatteredCachedEntity->EffectiveURL = CachedBlocks[firstIndex]->EffectiveURL;
-	Response->CopyHeadersExceptContentSizesAndDateFrom(CachedBlocks[firstIndex]->Response);
-	Response->SetHTTPStatusLine(CachedBlocks[firstIndex]->Response->GetHTTPStatusLine());
+		// There is some overlap. Find how many consecutive bytes we can get from the cached blocks.
+		TArray<TSharedPtrTS<FCacheItem>> TouchedBlocks;
+		while(1)
+		{
+			int64 cbstart = CachedBlocks[firstIndex]->Range.GetStart();
+			int64 cbend = CachedBlocks[firstIndex]->Range.GetEndIncluding();
+			int64 last = re < cbend ? re : cbend;
+			int64 offset = rs - cbstart;
+			int64 numAvail = last + 1 - cbstart - offset;
 
-	// There is some overlap. Find how many consecutive bytes we can get from the cached blocks.
-	TArray<TSharedPtrTS<FCacheItem>> TouchedBlocks;
-	while(1)
-	{
-		int64 cbstart = CachedBlocks[firstIndex]->Range.GetStart();
-		int64 cbend = CachedBlocks[firstIndex]->Range.GetEndIncluding();
-		int64 last = re < cbend ? re : cbend;
-		int64 offset = rs - cbstart;
-		int64 numAvail = last + 1 - cbstart - offset;
+			// Add the bytes from the cache to the new response.
+			Response->GetResponseData().AddData(CachedBlocks[firstIndex]->Response->GetResponseData(), offset, numAvail);
 
-		// Add the bytes from the cache to the new response.
-		Response->GetResponseData().AddData(CachedBlocks[firstIndex]->Response->GetResponseData(), offset, numAvail);
+			// Remember which blocks we just touched.
+			TouchedBlocks.Add(CachedBlocks[firstIndex]);
 
-		// Remember which blocks we just touched.
-		TouchedBlocks.Add(CachedBlocks[firstIndex]);
+			// Move up the start to the end of this cached block.
+			rs = last + 1;
 
-		// Move up the start to the end of this cached block.
-		rs = last + 1;
+			// End fully included?
+			if (re <= cbend)
+			{
+				break;
+			}
+			// Is there another cached block?
+			if (++firstIndex >= CachedBlocks.Num())
+			{
+				break;
+			}
+			// Does its start coincide with what we need next?
+			if (CachedBlocks[firstIndex]->Range.GetStart() != rs)
+			{
+				break;
+			}
+		}
 
-		// End fully included?
-		if (re <= cbend)
+		// Move the blocks we touched to the front of the cached blocks (LRU).
+		for(auto &Touched : TouchedBlocks)
 		{
-			break;
+			Cache.Remove(Touched);
+			Cache.Insert(Touched, 0);
 		}
-		// Is there another cached block?
-		if (++firstIndex >= CachedBlocks.Num())
+
+		Range.SetEndIncluding(rs - 1);
+		if (InRange.IsSet())
 		{
-			break;
+			Response->AddHeader(TEXT("Content-Range"), FString::Printf(TEXT("bytes %lld-%lld/%lld"), (long long int)Range.GetStart(), (long long int)Range.GetEndIncluding(), (long long int)Range.GetDocumentSize()));
+			Response->SetHTTPResponseCode(206);
+			FString StatusLine = Response->GetHTTPStatusLine();
+			StatusLine.ReplaceInline(TEXT("200"), TEXT("206"));
+			Response->SetHTTPStatusLine(StatusLine);
 		}
-		// Does its start coincide with what we need next?
-		if (CachedBlocks[firstIndex]->Range.GetStart() != rs)
+		Response->AddHeader(TEXT("Content-Length"), FString::Printf(TEXT("%lld"), (long long int)Response->GetResponseData().GetNumTotalBytesAdded()));
+		Response->GetResponseData().SetLengthFromResponseHeader(Response->GetResponseData().GetNumTotalBytesAdded());
+		Response->GetResponseData().SetEOS();
+
+		/*
+		Response->SetTimeUntilNameResolved(double InTimeUntilNameResolved);
+		Response->SetTimeUntilConnected(double InTimeUntilConnected);
+		Response->SetTimeUntilRequestSent(double InTimeUntilRequestSent);
+		Response->SetTimeUntilHeadersAvailable(double InTimeUntilHeadersAvailable);
+		Response->SetTimeUntilFirstByte(double InTimeUntilFirstByte);
+		Response->SetTimeUntilFinished(double InTimeUntilFinished);
+		*/
+
+		OutScatteredCachedEntity->Response = Response;
+		OutScatteredCachedEntity->Range = Range;
+
+		return EScatterResult::PartialHit;
+	}
+	else
+	{
+		IElectraPlayerDataCache::FItemInfo Info;
+		Info.URI = URL;
+		switch(InForQuality.StreamType)
 		{
-			break;
+			case EStreamType::Video:
+				Info.StreamType = IElectraPlayerDataCache::FItemInfo::EStreamType::Video;
+				break;
+			case EStreamType::Audio:
+				Info.StreamType = IElectraPlayerDataCache::FItemInfo::EStreamType::Audio;
+				break;
+			default:
+				Info.StreamType = IElectraPlayerDataCache::FItemInfo::EStreamType::Other;
+				break;
 		}
-	}
+		Info.QualityIndex = InForQuality.QualityIndex;
+		Info.MaxQualityIndex = InForQuality.MaxQualityIndex;
+		Info.Range.Start = InRange.GetStart();
+		Info.Range.EndIncluding = InRange.GetEndIncluding();
+		Info.Range.TotalSize = -1;
 
-	// Move the blocks we touched to the front of the cached blocks (LRU).
-	for(auto &Touched : TouchedBlocks)
-	{
-		Cache.Remove(Touched);
-		Cache.Insert(Touched, 0);
-	}
+		struct FResult
+		{
+			IElectraPlayerDataCache::ECacheResult Result;
+			IElectraPlayerDataCache::FItemInfo ItemInfoFromCache;
+			IElectraPlayerDataCache::FCacheDataPtr DataFromCache;
+			FMediaEvent Event;
+		};
+		TSharedPtrTS<FResult> CacheResult = MakeSharedTS<FResult>();
+
+		IElectraPlayerDataCache::FCachedDataReadCompleted FinishedDelegate;
+		FinishedDelegate.BindLambda([CacheResult](IElectraPlayerDataCache::ECacheResult InResult, IElectraPlayerDataCache::FItemInfo InItemInfoFromCache, IElectraPlayerDataCache::FCacheDataPtr InDataFromCache)
+		{
+			CacheResult->Result = InResult;
+			CacheResult->ItemInfoFromCache = InItemInfoFromCache;
+			CacheResult->DataFromCache = InDataFromCache;
+			CacheResult->Event.Signal();
+		});
+		ExternalCache->GetElementFromCache(Info, FinishedDelegate);
+		CacheResult->Event.Wait();
+
+		OutScatteredCachedEntity = MakeSharedTS<IHTTPResponseCache::FCacheItem>();
+		OutScatteredCachedEntity->RequestedURL = URL;
+		OutScatteredCachedEntity->EffectiveURL = URL;
+		OutScatteredCachedEntity->Range = InRange;
+
+		if (CacheResult->Result == IElectraPlayerDataCache::ECacheResult::Miss)
+		{
+			// Set up the default output that encompasses the request.
+			// This is what needs to be fetched for cache misses.
+			return EScatterResult::Miss;
+		}
+		else
+		{
+			TSharedPtr<FSynthesizedCacheResponse, ESPMode::ThreadSafe> Response = MakeShared<FSynthesizedCacheResponse, ESPMode::ThreadSafe>();
+			Response->SetURL(URL);
 
-	Range.SetEndIncluding(rs - 1);
-	if (InRange.IsSet())
-	{
-		Response->AddHeader(TEXT("Content-Range"), FString::Printf(TEXT("bytes %lld-%lld/%lld"), (long long int)Range.GetStart(), (long long int)Range.GetEndIncluding(), (long long int)Range.GetDocumentSize()));
-		Response->SetHTTPResponseCode(206);
-		FString StatusLine = Response->GetHTTPStatusLine();
-		StatusLine.ReplaceInline(TEXT("200"), TEXT("206"));
-		Response->SetHTTPStatusLine(StatusLine);
+			// Add the bytes from the cache to the new response.
+			if (CacheResult->DataFromCache.IsValid())
+			{
+				Response->GetResponseData().AddData(*CacheResult->DataFromCache.Get());
+			}
+
+			ElectraHTTPStream::FHttpRange Range(InRange);
+			Range.SetStart(CacheResult->ItemInfoFromCache.Range.Start);
+			Range.SetEndIncluding(CacheResult->ItemInfoFromCache.Range.EndIncluding);
+			Range.DocumentSize = CacheResult->ItemInfoFromCache.Range.TotalSize;
+			if (Range.IsSet())
+			{
+				Response->AddHeader(TEXT("Content-Range"), FString::Printf(TEXT("bytes %lld-%lld/%lld"), (long long int)Range.GetStart(), (long long int)Range.GetEndIncluding(), (long long int)Range.GetDocumentSize()));
+				Response->SetHTTPResponseCode(206);
+				Response->SetHTTPStatusLine(TEXT("HTTP/1.1 206 Partial Content"));
+			}
+			else
+			{
+				Response->SetHTTPResponseCode(200);
+				Response->SetHTTPStatusLine(TEXT("HTTP/1.1 200 OK"));
+			}
+			Response->AddHeader(TEXT("Content-Length"), FString::Printf(TEXT("%lld"), (long long int)Response->GetResponseData().GetNumTotalBytesAdded()));
+			Response->GetResponseData().SetLengthFromResponseHeader(Response->GetResponseData().GetNumTotalBytesAdded());
+			Response->GetResponseData().SetEOS();
+
+			OutScatteredCachedEntity->Response = Response;
+			OutScatteredCachedEntity->Range = Range;
+
+			return EScatterResult::FullHit;
+		}
 	}
-	Response->AddHeader(TEXT("Content-Length"), FString::Printf(TEXT("%lld"), (long long int)Response->GetResponseData().GetNumTotalBytesAdded()));
-	Response->GetResponseData().SetLengthFromResponseHeader(Response->GetResponseData().GetNumTotalBytesAdded());
-	Response->GetResponseData().SetEOS();
-
-/*
-	Response->SetTimeUntilNameResolved(double InTimeUntilNameResolved);
-	Response->SetTimeUntilConnected(double InTimeUntilConnected);
-	Response->SetTimeUntilRequestSent(double InTimeUntilRequestSent);
-	Response->SetTimeUntilHeadersAvailable(double InTimeUntilHeadersAvailable);
-	Response->SetTimeUntilFirstByte(double InTimeUntilFirstByte);
-	Response->SetTimeUntilFinished(double InTimeUntilFinished);
-*/
-
-	OutScatteredCachedEntity->Response = Response;
-	OutScatteredCachedEntity->Range = Range;
-
-	return EScatterResult::PartialHit;
 }
 
 
diff --git a/Engine/Plugins/Media/ElectraPlayer/Source/ElectraPlayerRuntime/Private/Runtime/HTTP/HTTPResponseCache.h b/Engine/Plugins/Media/ElectraPlayer/Source/ElectraPlayerRuntime/Private/Runtime/HTTP/HTTPResponseCache.h
--- a/Engine/Plugins/Media/ElectraPlayer/Source/ElectraPlayerRuntime/Private/Runtime/HTTP/HTTPResponseCache.h
+++ b/Engine/Plugins/Media/ElectraPlayer/Source/ElectraPlayerRuntime/Private/Runtime/HTTP/HTTPResponseCache.h
@@ -7,6 +7,8 @@
 #include "ParameterDictionary.h"
 #include "ElectraHTTPStream.h"
 #include "Utilities/HttpRangeHeader.h"
+#include "IElectraPlayerDataCache.h"
+#include "StreamTypes.h"
 
 namespace Electra
 {
@@ -15,6 +17,12 @@ class IPlayerSessionServices;
 class IHTTPResponseCache
 {
 public:
+	struct FQualityInfo
+	{
+		EStreamType StreamType = EStreamType::Unsupported;
+		int32 QualityIndex = 0;
+		int32 MaxQualityIndex = 0;
+	};
 	struct FCacheItem
 	{
 		FString RequestedURL;
@@ -22,12 +30,18 @@ public:
 		ElectraHTTPStream::FHttpRange Range;
 		FTimeValue ExpiresAtUTC;
 		IElectraHTTPStreamResponsePtr Response;
+		FQualityInfo Quality;
 	};
 
-	static TSharedPtrTS<IHTTPResponseCache> Create(IPlayerSessionServices* SessionServices, const FParamDict& Options);
+	static TSharedPtrTS<IHTTPResponseCache> Create(IPlayerSessionServices* SessionServices, const FParamDict& Options, TSharedPtr<IElectraPlayerDataCache, ESPMode::ThreadSafe> ExternalCache);
 
 	virtual ~IHTTPResponseCache() = default;
 
+	/**
+	 * Call this to disable the cache. Caching makes no sense for Live video streams.
+	 */
+	virtual void Disable() = 0;
+
 	/**
 	 * Call this periodically to handle expiration times of cached entities.
 	 */
@@ -54,7 +68,7 @@ public:
 	 * Returns false to indicate that the range specified in the (otherwise empty) cached response must be requested from the origin.
 	 * Returns true if a - possibly partial - cached response exists. The range that is cached is given in the cache item.
 	 */
-	virtual EScatterResult GetScatteredCacheEntity(TSharedPtrTS<FCacheItem>& OutScatteredCachedEntity, const FString& URL, const ElectraHTTPStream::FHttpRange& Range) = 0;
+	virtual EScatterResult GetScatteredCacheEntity(TSharedPtrTS<FCacheItem>& OutScatteredCachedEntity, const FString& URL, const ElectraHTTPStream::FHttpRange& Range, const FQualityInfo& ForQuality) = 0;
 
 protected:
 	IHTTPResponseCache() = default;
diff --git a/Engine/Plugins/Media/ElectraPlayer/Source/ElectraPlayerRuntime/Private/Runtime/Player/ABRRules/ABRLiveStream.cpp b/Engine/Plugins/Media/ElectraPlayer/Source/ElectraPlayerRuntime/Private/Runtime/Player/ABRRules/ABRLiveStream.cpp
--- a/Engine/Plugins/Media/ElectraPlayer/Source/ElectraPlayerRuntime/Private/Runtime/Player/ABRRules/ABRLiveStream.cpp
+++ b/Engine/Plugins/Media/ElectraPlayer/Source/ElectraPlayerRuntime/Private/Runtime/Player/ABRRules/ABRLiveStream.cpp
@@ -1417,7 +1417,7 @@ IAdaptiveStreamSelector::ESegmentAction FABRLiveStream::PerformSelection(const T
 				}
 			}
 
-			CurrentPlayPeriod->SelectStream(Candidate->AdaptationSetUniqueID, Candidate->RepresentationUniqueID);
+			CurrentPlayPeriod->SelectStream(Candidate->AdaptationSetUniqueID, Candidate->RepresentationUniqueID, Candidate->QualityIndex, Info->GetStreamInformations(StreamType).Num()-1);
 			if (bRetryIfPossible)
 			{
 				return IAdaptiveStreamSelector::ESegmentAction::Retry;
diff --git a/Engine/Plugins/Media/ElectraPlayer/Source/ElectraPlayerRuntime/Private/Runtime/Player/ABRRules/ABRVoDPlus.cpp b/Engine/Plugins/Media/ElectraPlayer/Source/ElectraPlayerRuntime/Private/Runtime/Player/ABRRules/ABRVoDPlus.cpp
--- a/Engine/Plugins/Media/ElectraPlayer/Source/ElectraPlayerRuntime/Private/Runtime/Player/ABRRules/ABRVoDPlus.cpp
+++ b/Engine/Plugins/Media/ElectraPlayer/Source/ElectraPlayerRuntime/Private/Runtime/Player/ABRRules/ABRVoDPlus.cpp
@@ -935,7 +935,7 @@ IAdaptiveStreamSelector::ESegmentAction FABROnDemandPlus::PerformSelection(const
 				WorkVars->BitrateDownloading = Candidate->Bitrate;
 			}
 
-			CurrentPlayPeriod->SelectStream(Candidate->AdaptationSetUniqueID, Candidate->RepresentationUniqueID);
+			CurrentPlayPeriod->SelectStream(Candidate->AdaptationSetUniqueID, Candidate->RepresentationUniqueID, NewQualityIndex, Info->GetStreamInformations(StreamType).Num()-1);
 			if (bRetryIfPossible)
 			{
 				return IAdaptiveStreamSelector::ESegmentAction::Retry;
diff --git a/Engine/Plugins/Media/ElectraPlayer/Source/ElectraPlayerRuntime/Private/Runtime/Player/AdaptiveStreamingPlayer.cpp b/Engine/Plugins/Media/ElectraPlayer/Source/ElectraPlayerRuntime/Private/Runtime/Player/AdaptiveStreamingPlayer.cpp
--- a/Engine/Plugins/Media/ElectraPlayer/Source/ElectraPlayerRuntime/Private/Runtime/Player/AdaptiveStreamingPlayer.cpp
+++ b/Engine/Plugins/Media/ElectraPlayer/Source/ElectraPlayerRuntime/Private/Runtime/Player/AdaptiveStreamingPlayer.cpp
@@ -304,6 +304,11 @@ void FAdaptiveStreamingPlayer::SetVideoDecoderResourceDelegate(const TSharedPtr<
 	VideoDecoderResourceDelegate = InResourceDelegate;
 }
 
+void FAdaptiveStreamingPlayer::SetPlayerDataCache(const TSharedPtr<IElectraPlayerDataCache, ESPMode::ThreadSafe>& InPlayerDataCache)
+{
+	ExternalCache = InPlayerDataCache;
+}
+
 
 void FAdaptiveStreamingPlayer::AddMetricsReceiver(IAdaptiveStreamingPlayerMetrics* InMetricsReceiver)
 {
@@ -4409,8 +4414,8 @@ void FAdaptiveStreamingPlayer::InternalInitialize()
 	// Create an entity cache.
 	EntityCache = IPlayerEntityCache::Create(this, PlayerOptions);
 
-	// Create an HTTP response cache.
-	HttpResponseCache = IHTTPResponseCache::Create(this, PlayerOptions);
+	// Create an HTTP response cache. Hand over any externally set cache. We do not need it in here any further.
+	HttpResponseCache = IHTTPResponseCache::Create(this, PlayerOptions, MoveTemp(ExternalCache));
 
 	// If all read requests for this player are to be routed to an external reader we need to create a wrapper for it.
 	if (PlayerOptions.GetValue(OptionKeyUseExternalDataReader).SafeGetBool(false))
diff --git a/Engine/Plugins/Media/ElectraPlayer/Source/ElectraPlayerRuntime/Private/Runtime/Player/AdaptiveStreamingPlayer.h b/Engine/Plugins/Media/ElectraPlayer/Source/ElectraPlayerRuntime/Private/Runtime/Player/AdaptiveStreamingPlayer.h
--- a/Engine/Plugins/Media/ElectraPlayer/Source/ElectraPlayerRuntime/Private/Runtime/Player/AdaptiveStreamingPlayer.h
+++ b/Engine/Plugins/Media/ElectraPlayer/Source/ElectraPlayerRuntime/Private/Runtime/Player/AdaptiveStreamingPlayer.h
@@ -13,6 +13,7 @@
 #include "AdaptiveStreamingPlayerEvents.h"
 #include "AdaptiveStreamingPlayerSubtitles.h"
 #include "AdaptiveStreamingPlayerResourceRequest.h"
+#include "IElectraPlayerDataCache.h"
 
 class IOptionPointerValueContainer;
 
@@ -38,9 +39,12 @@ public:
 
 	//-------------------------------------------------------------------------
 	// Various resource providers/delegates
+	// These must be set prior to calling `Initialize()` and MUST NOT be changed
+	// until player destruction.
 	//
 	virtual void SetStaticResourceProviderCallback(const TSharedPtr<IAdaptiveStreamingPlayerResourceProvider, ESPMode::ThreadSafe>& InStaticResourceProvider) = 0;
 	virtual void SetVideoDecoderResourceDelegate(const TSharedPtr<IVideoDecoderResourceDelegate, ESPMode::ThreadSafe>& ResourceDelegate) = 0;
+	virtual void SetPlayerDataCache(const TSharedPtr<IElectraPlayerDataCache, ESPMode::ThreadSafe>& InPlayerDataCache) = 0;
 
 
 	//-------------------------------------------------------------------------
diff --git a/Engine/Plugins/Media/ElectraPlayer/Source/ElectraPlayerRuntime/Private/Runtime/Player/AdaptiveStreamingPlayerABR.cpp b/Engine/Plugins/Media/ElectraPlayer/Source/ElectraPlayerRuntime/Private/Runtime/Player/AdaptiveStreamingPlayerABR.cpp
--- a/Engine/Plugins/Media/ElectraPlayer/Source/ElectraPlayerRuntime/Private/Runtime/Player/AdaptiveStreamingPlayerABR.cpp
+++ b/Engine/Plugins/Media/ElectraPlayer/Source/ElectraPlayerRuntime/Private/Runtime/Player/AdaptiveStreamingPlayerABR.cpp
@@ -88,6 +88,8 @@ namespace Electra
 		int32 GetBandwidthCeiling() override
 		{ return BandwidthCeiling; }
 
+		ESegmentAction SelectSuitableStreamByType(FTimeValue& OutDelay, TSharedPtrTS<const IStreamSegment> CurrentSegment, TSharedPtrTS<IManifest::IPlayPeriod> CurrentPlayPeriod, EStreamType StreamType);
+
 		FTimeValue ABRGetPlayPosition() const override
 		{ return PlayerLiveControl->ABRGetPlayPosition(); }
 		FTimeRange ABRGetTimeline() const override
@@ -598,6 +600,36 @@ namespace Electra
 			return ESegmentAction::FetchNext;
 		}
 
+		// Specific request?
+		if (CurrentSegment.IsValid())
+		{
+			return SelectSuitableStreamByType(OutDelay, CurrentSegment, CurrentPlayPeriod, StreamType);
+		}
+		else
+		{
+			// Initial request. The play period is the same for all streams at this point.
+			ESegmentAction ActionV = ESegmentAction::FetchNext;
+			ESegmentAction ActionA = ESegmentAction::FetchNext;
+			if (StreamInformationVideo.Num())
+			{
+				ActionV = SelectSuitableStreamByType(OutDelay, CurrentSegment, CurrentPlayPeriod, EStreamType::Video);
+			}
+			if (StreamInformationAudio.Num())
+			{
+				ActionA = SelectSuitableStreamByType(OutDelay, CurrentSegment, CurrentPlayPeriod, EStreamType::Audio);
+			}
+			if (ActionV == IAdaptiveStreamSelector::ESegmentAction::Fail || ActionA == IAdaptiveStreamSelector::ESegmentAction::Fail)
+			{
+				return IAdaptiveStreamSelector::ESegmentAction::Fail;
+			}
+			return StreamInformationVideo.Num() ? ActionV : ActionA;
+		}
+	}
+
+	IAdaptiveStreamSelector::ESegmentAction FAdaptiveStreamSelector::SelectSuitableStreamByType(FTimeValue& OutDelay, TSharedPtrTS<const IStreamSegment> CurrentSegment, TSharedPtrTS<IManifest::IPlayPeriod> CurrentPlayPeriod, EStreamType StreamType)
+	{
+		FTimeValue TimeNow = PlayerSessionServices->GetSynchronizedUTCTime()->GetTime();
+
 		// Get the initial list of representations that are currently possible to choose from.
 		TArray<TSharedPtrTS<FABRStreamInformation>> CandidateRepresentations;
 		ABRMethod->PrepareStreamCandidateList(CandidateRepresentations, StreamType, CurrentPlayPeriod, TimeNow);
diff --git a/Engine/Plugins/Media/ElectraPlayer/Source/ElectraPlayerRuntime/Private/Runtime/Player/AdaptiveStreamingPlayerInternal.h b/Engine/Plugins/Media/ElectraPlayer/Source/ElectraPlayerRuntime/Private/Runtime/Player/AdaptiveStreamingPlayerInternal.h
--- a/Engine/Plugins/Media/ElectraPlayer/Source/ElectraPlayerRuntime/Private/Runtime/Player/AdaptiveStreamingPlayerInternal.h
+++ b/Engine/Plugins/Media/ElectraPlayer/Source/ElectraPlayerRuntime/Private/Runtime/Player/AdaptiveStreamingPlayerInternal.h
@@ -965,6 +965,7 @@ public:
 
 	void SetStaticResourceProviderCallback(const TSharedPtr<IAdaptiveStreamingPlayerResourceProvider, ESPMode::ThreadSafe>& InStaticResourceProvider) override;
 	void SetVideoDecoderResourceDelegate(const TSharedPtr<IVideoDecoderResourceDelegate, ESPMode::ThreadSafe>& ResourceDelegate) override;
+	void SetPlayerDataCache(const TSharedPtr<IElectraPlayerDataCache, ESPMode::ThreadSafe>& InPlayerDataCache) override;
 
 	void AddMetricsReceiver(IAdaptiveStreamingPlayerMetrics* InMetricsReceiver) override;
 	void RemoveMetricsReceiver(IAdaptiveStreamingPlayerMetrics* InMetricsReceiver) override;
@@ -2149,6 +2150,7 @@ private:
 
 	TWeakPtr<IAdaptiveStreamingPlayerResourceProvider, ESPMode::ThreadSafe> StaticResourceProvider;
 	TWeakPtr<IVideoDecoderResourceDelegate, ESPMode::ThreadSafe>		VideoDecoderResourceDelegate;
+	TSharedPtr<IElectraPlayerDataCache, ESPMode::ThreadSafe>			ExternalCache;
 
 	TSharedPtrTS<FAdaptiveStreamingPlayerEventHandler>					EventDispatcher;
 	TSharedPtrTS<FAdaptiveStreamingPlayerWorkerThread>					SharedWorkerThread;
diff --git a/Engine/Plugins/Media/ElectraPlayer/Source/ElectraPlayerRuntime/Private/Runtime/Player/AdaptiveStreamingPlayerManifest.cpp b/Engine/Plugins/Media/ElectraPlayer/Source/ElectraPlayerRuntime/Private/Runtime/Player/AdaptiveStreamingPlayerManifest.cpp
--- a/Engine/Plugins/Media/ElectraPlayer/Source/ElectraPlayerRuntime/Private/Runtime/Player/AdaptiveStreamingPlayerManifest.cpp
+++ b/Engine/Plugins/Media/ElectraPlayer/Source/ElectraPlayerRuntime/Private/Runtime/Player/AdaptiveStreamingPlayerManifest.cpp
@@ -343,6 +343,12 @@ bool FAdaptiveStreamingPlayer::SelectManifest()
 			// Let the ABR know the format as well.
 			StreamSelector->SetFormatType(ManifestType);
 
+			// Live streams must not use any cache!
+			if (NewPresentation->GetPresentationType() == IManifest::EType::Live && HttpResponseCache.IsValid())
+			{
+				HttpResponseCache->Disable();
+			}
+
 			return true;
 		}
 		else
diff --git a/Engine/Plugins/Media/ElectraPlayer/Source/ElectraPlayerRuntime/Private/Runtime/Player/AdaptiveStreamingPlayerResourceRequest.h b/Engine/Plugins/Media/ElectraPlayer/Source/ElectraPlayerRuntime/Private/Runtime/Player/AdaptiveStreamingPlayerResourceRequest.h
--- a/Engine/Plugins/Media/ElectraPlayer/Source/ElectraPlayerRuntime/Private/Runtime/Player/AdaptiveStreamingPlayerResourceRequest.h
+++ b/Engine/Plugins/Media/ElectraPlayer/Source/ElectraPlayerRuntime/Private/Runtime/Player/AdaptiveStreamingPlayerResourceRequest.h
@@ -114,6 +114,9 @@ public:
 	virtual FHTTPResourceRequest& Object(TSharedPtrTS<IHTTPResourceRequestObject> InUserObject)
 	{ UserObject = InUserObject; return *this; }
 
+	virtual FHTTPResourceRequest& StreamTypeAndQuality(EStreamType InStreamType, int32 InQualityIndex, int32 InMaxQualityIndex)
+	{ Request->Parameters.StreamType = InStreamType; Request->Parameters.QualityIndex = InQualityIndex; Request->Parameters.MaxQualityIndex = InMaxQualityIndex; return *this; }
+
 	virtual FOnRequestCompletedCallback& Callback()
 	{ return CompletedCallback; }
 
@@ -153,6 +156,9 @@ public:
 	virtual const HTTP::FConnectionInfo* GetConnectionInfo() const
 	{ return Request.IsValid() ? &Request->ConnectionInfo : nullptr; }
 
+	virtual TSharedPtrTS<IElectraHttpManager::FRequest> GetRequest() const
+	{ return Request; }
+
 private:
 
 	class FStaticResourceRequest : public IAdaptiveStreamingPlayerResourceRequest
diff --git a/Engine/Plugins/Media/ElectraPlayer/Source/ElectraPlayerRuntime/Private/Runtime/Player/DASH/MPDElementsDASH.h b/Engine/Plugins/Media/ElectraPlayer/Source/ElectraPlayerRuntime/Private/Runtime/Player/DASH/MPDElementsDASH.h
--- a/Engine/Plugins/Media/ElectraPlayer/Source/ElectraPlayerRuntime/Private/Runtime/Player/DASH/MPDElementsDASH.h
+++ b/Engine/Plugins/Media/ElectraPlayer/Source/ElectraPlayerRuntime/Private/Runtime/Player/DASH/MPDElementsDASH.h
@@ -5,6 +5,7 @@
 #include "Containers/ChunkedArray.h"
 #include "OptionalValue.h"
 #include "Runtime/ErrorDetail.h"
+#include "PlayerTime.h"
 
 namespace Electra
 {
diff --git a/Engine/Plugins/Media/ElectraPlayer/Source/ElectraPlayerRuntime/Private/Runtime/Player/DASH/ManifestDASH.cpp b/Engine/Plugins/Media/ElectraPlayer/Source/ElectraPlayerRuntime/Private/Runtime/Player/DASH/ManifestDASH.cpp
--- a/Engine/Plugins/Media/ElectraPlayer/Source/ElectraPlayerRuntime/Private/Runtime/Player/DASH/ManifestDASH.cpp
+++ b/Engine/Plugins/Media/ElectraPlayer/Source/ElectraPlayerRuntime/Private/Runtime/Player/DASH/ManifestDASH.cpp
@@ -200,7 +200,7 @@ public:
 	FString GetSelectedAdaptationSetID(EStreamType StreamType) override;
 	ETrackChangeResult ChangeTrackStreamPreference(EStreamType ForStreamType, const FStreamSelectionAttributes& StreamAttributes) override;
 	TSharedPtrTS<ITimelineMediaAsset> GetMediaAsset() const override;
-	void SelectStream(const FString& AdaptationSetID, const FString& RepresentationID) override;
+	void SelectStream(const FString& AdaptationSetID, const FString& RepresentationID, int32 QualityIndex, int32 MaxQualityIndex) override;
 	void TriggerInitSegmentPreload(const TArray<FInitSegmentPreload>& InitSegmentsToPreload) override;
 	IManifest::FResult GetStartingSegment(TSharedPtrTS<IStreamSegment>& OutSegment, const FPlayerSequenceState& InSequenceState, const FPlayStartPosition& StartPosition, IManifest::ESearchType SearchType) override;
 	IManifest::FResult GetContinuationSegment(TSharedPtrTS<IStreamSegment>& OutSegment, EStreamType StreamType, const FPlayerSequenceState& SequenceState, const FPlayStartPosition& StartPosition, IManifest::ESearchType SearchType) override;
@@ -284,6 +284,20 @@ private:
 	FString ActiveAudioRepresentationID;
 	FString ActiveSubtitleRepresentationID;
 
+	struct FSelectedQualityIndex
+	{
+		int32 Index = 0;
+		int32 MaxIndex = 0;
+		void Reset()
+		{
+			Index = 0;
+			MaxIndex = 0;
+		}
+	};
+	FSelectedQualityIndex ActiveVideoQualityIndex;
+	FSelectedQualityIndex ActiveAudioQualityIndex;
+	FSelectedQualityIndex ActiveSubtitleQualityIndex;
+
 	TSharedPtrTS<FBufferSourceInfo> SourceBufferInfoVideo;
 	TSharedPtrTS<FBufferSourceInfo> SourceBufferInfoAudio;
 	TSharedPtrTS<FBufferSourceInfo> SourceBufferInfoSubtitles;
@@ -1137,20 +1151,26 @@ TSharedPtrTS<ITimelineMediaAsset> FDASHPlayPeriod::GetMediaAsset() const
 	return nullptr;
 }
 
-void FDASHPlayPeriod::SelectStream(const FString& AdaptationSetID, const FString& RepresentationID)
+void FDASHPlayPeriod::SelectStream(const FString& AdaptationSetID, const FString& RepresentationID, int32 QualityIndex, int32 MaxQualityIndex)
 {
 	// The ABR must not try to switch adaptation sets at the moment. As such the adaptation set passed in must be one of the already active ones.
 	if (AdaptationSetID == ActiveVideoAdaptationSetID)
 	{
 		ActiveVideoRepresentationID = RepresentationID;
+		ActiveVideoQualityIndex.Index = QualityIndex;
+		ActiveVideoQualityIndex.MaxIndex = MaxQualityIndex;
 	}
 	else if (AdaptationSetID == ActiveAudioAdaptationSetID)
 	{
 		ActiveAudioRepresentationID = RepresentationID;
+		ActiveAudioQualityIndex.Index = QualityIndex;
+		ActiveAudioQualityIndex.MaxIndex = MaxQualityIndex;
 	}
 	else if (AdaptationSetID == ActiveSubtitleAdaptationSetID)
 	{
 		ActiveSubtitleRepresentationID = RepresentationID;
+		ActiveSubtitleQualityIndex.Index = QualityIndex;
+		ActiveSubtitleQualityIndex.MaxIndex = MaxQualityIndex;
 	}
 	else
 	{
@@ -1239,19 +1259,20 @@ IManifest::FResult FDASHPlayPeriod::GetStartingSegment(TSharedPtrTS<IStreamSegme
 		EStreamType	StreamType;
 		FString		RepresentationID;
 		FString		AdaptationSetID;
+		FSelectedQualityIndex QualityIndex;
 	};
 	TArray<FSelectedStream> ActiveSelection;
 	if (!ActiveVideoRepresentationID.IsEmpty())
 	{
-		ActiveSelection.Emplace(FSelectedStream({EStreamType::Video, ActiveVideoRepresentationID, ActiveVideoAdaptationSetID}));
+		ActiveSelection.Emplace(FSelectedStream({EStreamType::Video, ActiveVideoRepresentationID, ActiveVideoAdaptationSetID, ActiveVideoQualityIndex}));
 	}
 	if (!ActiveAudioRepresentationID.IsEmpty())
 	{
-		ActiveSelection.Emplace(FSelectedStream({EStreamType::Audio, ActiveAudioRepresentationID, ActiveAudioAdaptationSetID}));
+		ActiveSelection.Emplace(FSelectedStream({EStreamType::Audio, ActiveAudioRepresentationID, ActiveAudioAdaptationSetID, ActiveAudioQualityIndex}));
 	}
 	if (!ActiveSubtitleRepresentationID.IsEmpty())
 	{
-		ActiveSelection.Emplace(FSelectedStream({EStreamType::Subtitle, ActiveSubtitleRepresentationID, ActiveSubtitleAdaptationSetID}));
+		ActiveSelection.Emplace(FSelectedStream({EStreamType::Subtitle, ActiveSubtitleRepresentationID, ActiveSubtitleAdaptationSetID, ActiveSubtitleQualityIndex}));
 	}
 
 	bool bDidAdjustStartTime = false;
@@ -1284,6 +1305,9 @@ IManifest::FResult FDASHPlayPeriod::GetStartingSegment(TSharedPtrTS<IStreamSegme
 			SearchOpt.bHasFollowingPeriod = Period->GetHasFollowingPeriod();
 			SearchOpt.SearchType = SearchType;
 			SearchOpt.bFrameAccurateSearch = bFrameAccurateSearch;
+			SearchOpt.QualityIndex = ActiveSelection[i].QualityIndex.Index;
+			SearchOpt.MaxQualityIndex = ActiveSelection[i].QualityIndex.MaxIndex;
+			SearchOpt.StreamType = ActiveSelection[i].StreamType;
 			FManifestDASHInternal::FRepresentation::ESearchResult SearchResult = Repr->FindSegment(PlayerSessionServices, SegmentInfo, RemoteElementLoadRequests, SearchOpt);
 			if (SearchResult == FManifestDASHInternal::FRepresentation::ESearchResult::NeedElement)
 			{
@@ -1355,6 +1379,8 @@ IManifest::FResult FDASHPlayPeriod::GetStartingSegment(TSharedPtrTS<IStreamSegme
 
 				TSharedPtrTS<FStreamSegmentRequestDASH> SegmentRequest = MakeSharedTS<FStreamSegmentRequestDASH>();
 				SegmentRequest->StreamType = ActiveSelection[i].StreamType;
+				SegmentRequest->QualityIndex = ActiveSelection[i].QualityIndex.Index;
+				SegmentRequest->MaxQualityIndex = ActiveSelection[i].QualityIndex.MaxIndex;
 				SegmentRequest->CodecInfo = Repr->GetCodecInformation();
 				SegmentRequest->Representation = Repr;
 				SegmentRequest->AdaptationSet = Adapt;
@@ -1458,19 +1484,23 @@ IManifest::FResult FDASHPlayPeriod::GetNextOrRetrySegment(TSharedPtrTS<IStreamSe
 
 	FString ActiveRepresentationIDByType;
 	FString ActiveAdaptationSetIDByType;
+	FSelectedQualityIndex ActiveQualityIndex;
 	switch(Current->GetType())
 	{
 		case EStreamType::Video:
 			ActiveAdaptationSetIDByType = ActiveVideoAdaptationSetID;
 			ActiveRepresentationIDByType = ActiveVideoRepresentationID;
+			ActiveQualityIndex = ActiveVideoQualityIndex;
 			break;
 		case EStreamType::Audio:
 			ActiveAdaptationSetIDByType = ActiveAudioAdaptationSetID;
 			ActiveRepresentationIDByType = ActiveAudioRepresentationID;
+			ActiveQualityIndex = ActiveAudioQualityIndex;
 			break;
 		case EStreamType::Subtitle:
 			ActiveAdaptationSetIDByType = ActiveSubtitleAdaptationSetID;
 			ActiveRepresentationIDByType = ActiveSubtitleRepresentationID;
+			ActiveQualityIndex = ActiveSubtitleQualityIndex;
 			break;
 		default:
 			break;
@@ -1546,6 +1576,9 @@ IManifest::FResult FDASHPlayPeriod::GetNextOrRetrySegment(TSharedPtrTS<IStreamSe
 	{
 		SearchOpt.PeriodDuration = Manifest->GetLastPeriodEndTime() - AST;
 	}
+	SearchOpt.QualityIndex = ActiveQualityIndex.Index;
+	SearchOpt.MaxQualityIndex = ActiveQualityIndex.MaxIndex;
+	SearchOpt.StreamType = Current->GetType();
 
 	FManifestDASHInternal::FRepresentation::ESearchResult SearchResult;
 	if (Repr.IsValid())
@@ -1604,6 +1637,8 @@ IManifest::FResult FDASHPlayPeriod::GetNextOrRetrySegment(TSharedPtrTS<IStreamSe
 		TSharedPtrTS<FStreamSegmentRequestDASH> SegmentRequest = MakeSharedTS<FStreamSegmentRequestDASH>();
 		SegmentRequest->TimestampSequenceIndex = Current->TimestampSequenceIndex;
 		SegmentRequest->StreamType = Current->GetType();
+		SegmentRequest->QualityIndex = ActiveQualityIndex.Index;
+		SegmentRequest->MaxQualityIndex = ActiveQualityIndex.MaxIndex;
 		SegmentRequest->CodecInfo = Repr->GetCodecInformation();
 		SegmentRequest->Representation = Repr;
 		SegmentRequest->AdaptationSet = Adapt;
@@ -2094,7 +2129,7 @@ TMediaOptionalValue<bool> GetSegmentAvailabilityTimeComplete(const TArray<TShare
 
 
 
-FManifestDASHInternal::FRepresentation::ESearchResult FManifestDASHInternal::FRepresentation::PrepareSegmentIndex(IPlayerSessionServices* InPlayerSessionServices, const TArray<TSharedPtrTS<FDashMPD_SegmentBaseType>>& SegmentBase, TArray<TWeakPtrTS<FMPDLoadRequestDASH>>& OutRemoteElementLoadRequests)
+FManifestDASHInternal::FRepresentation::ESearchResult FManifestDASHInternal::FRepresentation::PrepareSegmentIndex(IPlayerSessionServices* InPlayerSessionServices, const TArray<TSharedPtrTS<FDashMPD_SegmentBaseType>>& SegmentBase, TArray<TWeakPtrTS<FMPDLoadRequestDASH>>& OutRemoteElementLoadRequests, const FSegmentSearchOption& InSearchOptions)
 {
 	// If the segment index has been requested and is still pending, return right away.
 	if (PendingSegmentIndexLoadRequest.IsValid())
@@ -2195,6 +2230,9 @@ FManifestDASHInternal::FRepresentation::ESearchResult FManifestDASHInternal::FRe
 					}
 					LoadReq->PlayerSessionServices = InPlayerSessionServices;
 					LoadReq->XLinkElement = MPDRepresentation;
+					LoadReq->SegmentStreamType = InSearchOptions.StreamType;
+					LoadReq->SegmentQualityIndex = InSearchOptions.QualityIndex;
+					LoadReq->SegmentQualityIndexMax = InSearchOptions.MaxQualityIndex;
 					LoadReq->CompleteCallback.BindThreadSafeSP(AsShared(), &FManifestDASHInternal::FRepresentation::SegmentIndexDownloadComplete);
 					OutRemoteElementLoadRequests.Emplace(LoadReq);
 					PendingSegmentIndexLoadRequest = MoveTemp(LoadReq);
@@ -2218,6 +2256,9 @@ FManifestDASHInternal::FRepresentation::ESearchResult FManifestDASHInternal::FRe
 					}
 					LoadReq->PlayerSessionServices = InPlayerSessionServices;
 					LoadReq->XLinkElement = MPDRepresentation;
+					LoadReq->SegmentStreamType = InSearchOptions.StreamType;
+					LoadReq->SegmentQualityIndex = InSearchOptions.QualityIndex;
+					LoadReq->SegmentQualityIndexMax = InSearchOptions.MaxQualityIndex;
 					LoadReq->CompleteCallback.BindThreadSafeSP(AsShared(), &FManifestDASHInternal::FRepresentation::SegmentIndexDownloadComplete);
 					OutRemoteElementLoadRequests.Emplace(LoadReq);
 					PendingSegmentIndexLoadRequest = MoveTemp(LoadReq);
@@ -2680,7 +2721,7 @@ FManifestDASHInternal::FRepresentation::ESearchResult FManifestDASHInternal::FRe
 
 FManifestDASHInternal::FRepresentation::ESearchResult FManifestDASHInternal::FRepresentation::FindSegment_Base_MP4(IPlayerSessionServices* InPlayerSessionServices, FSegmentInformation& OutSegmentInfo, TArray<TWeakPtrTS<FMPDLoadRequestDASH>>& OutRemoteElementLoadRequests, const FSegmentSearchOption& InSearchOptions, const TSharedPtrTS<FDashMPD_RepresentationType>& MPDRepresentation, const TArray<TSharedPtrTS<FDashMPD_SegmentBaseType>>& SegmentBase)
 {
-	ESearchResult SegIndexResult = PrepareSegmentIndex(InPlayerSessionServices, SegmentBase, OutRemoteElementLoadRequests);
+	ESearchResult SegIndexResult = PrepareSegmentIndex(InPlayerSessionServices, SegmentBase, OutRemoteElementLoadRequests, InSearchOptions);
 	if (SegIndexResult != ESearchResult::Found)
 	{
 		return SegIndexResult;
@@ -2905,7 +2946,7 @@ FManifestDASHInternal::FRepresentation::ESearchResult FManifestDASHInternal::FRe
 
 FManifestDASHInternal::FRepresentation::ESearchResult FManifestDASHInternal::FRepresentation::FindSegment_Base_MKV(IPlayerSessionServices* InPlayerSessionServices, FSegmentInformation& OutSegmentInfo, TArray<TWeakPtrTS<FMPDLoadRequestDASH>>& OutRemoteElementLoadRequests, const FSegmentSearchOption& InSearchOptions, const TSharedPtrTS<FDashMPD_RepresentationType>& MPDRepresentation, const TArray<TSharedPtrTS<FDashMPD_SegmentBaseType>>& SegmentBase)
 {
-	ESearchResult SegIndexResult = PrepareSegmentIndex(InPlayerSessionServices, SegmentBase, OutRemoteElementLoadRequests);
+	ESearchResult SegIndexResult = PrepareSegmentIndex(InPlayerSessionServices, SegmentBase, OutRemoteElementLoadRequests, InSearchOptions);
 	if (SegIndexResult != ESearchResult::Found)
 	{
 		return SegIndexResult;
@@ -3738,6 +3779,9 @@ void FManifestDASHInternal::FRepresentation::SegmentIndexDownloadComplete(TShare
 				NextChainedReq->Headers = LoadRequest->Headers;
 				NextChainedReq->PlayerSessionServices = LoadRequest->PlayerSessionServices;
 				NextChainedReq->XLinkElement = LoadRequest->XLinkElement;
+				NextChainedReq->SegmentStreamType = LoadRequest->SegmentStreamType;
+				NextChainedReq->SegmentQualityIndex = LoadRequest->SegmentQualityIndex;
+				NextChainedReq->SegmentQualityIndexMax = LoadRequest->SegmentQualityIndexMax;
 				NextChainedReq->CompleteCallback.BindThreadSafeSP(AsShared(), &FManifestDASHInternal::FRepresentation::SegmentIndexDownloadComplete);
 				NextChainedReq->CompletedRequestChain.Emplace(MoveTemp(LoadRequest));
 				TSharedPtrTS<IPlaylistReader> ManifestReader = NextChainedReq->PlayerSessionServices->GetManifestReader();
diff --git a/Engine/Plugins/Media/ElectraPlayer/Source/ElectraPlayerRuntime/Private/Runtime/Player/DASH/PlaylistReaderDASH.cpp b/Engine/Plugins/Media/ElectraPlayer/Source/ElectraPlayerRuntime/Private/Runtime/Player/DASH/PlaylistReaderDASH.cpp
--- a/Engine/Plugins/Media/ElectraPlayer/Source/ElectraPlayerRuntime/Private/Runtime/Player/DASH/PlaylistReaderDASH.cpp
+++ b/Engine/Plugins/Media/ElectraPlayer/Source/ElectraPlayerRuntime/Private/Runtime/Player/DASH/PlaylistReaderDASH.cpp
@@ -923,7 +923,8 @@ void FPlaylistReaderDASH::SetupRequestTimeouts(FResourceLoadRequestPtr InRequest
 			InRequest->Request->ConnectionTimeout(FTimeValue(FTimeValue::MillisecondsToHNS(1000 * 5)));
 			InRequest->Request->NoDataTimeout(FTimeValue(FTimeValue::MillisecondsToHNS(1000 * 2)));
 			InRequest->Request->AcceptEncoding(TEXT("identity"));
-			break;
+			InRequest->Request->StreamTypeAndQuality(InRequest->SegmentStreamType, InRequest->SegmentQualityIndex, InRequest->SegmentQualityIndexMax);
+			InRequest->Request->GetRequest()->ResponseCache = PlayerSessionServices->GetHTTPResponseCache();
 		}
 		case FMPDLoadRequestDASH::ELoadType::XLink_Period:
 		case FMPDLoadRequestDASH::ELoadType::XLink_AdaptationSet:
diff --git a/Engine/Plugins/Media/ElectraPlayer/Source/ElectraPlayerRuntime/Private/Runtime/Player/DASH/PlaylistReaderDASH_Internal.h b/Engine/Plugins/Media/ElectraPlayer/Source/ElectraPlayerRuntime/Private/Runtime/Player/DASH/PlaylistReaderDASH_Internal.h
--- a/Engine/Plugins/Media/ElectraPlayer/Source/ElectraPlayerRuntime/Private/Runtime/Player/DASH/PlaylistReaderDASH_Internal.h
+++ b/Engine/Plugins/Media/ElectraPlayer/Source/ElectraPlayerRuntime/Private/Runtime/Player/DASH/PlaylistReaderDASH_Internal.h
@@ -6,7 +6,7 @@
 #include "PlayerCore.h"
 
 #include "MPDElementsDASH.h"
-
+#include "Player/Manifest.h"
 #include "Player/AdaptiveStreamingPlayerResourceRequest.h"
 #include "Utilities/URLParser.h"
 
@@ -80,6 +80,10 @@ struct FMPDLoadRequestDASH : public IHTTPResourceRequestObject
 	// This allows checking if - after a dynamic MPD update - the requesting MPD is still valid and in use.
 	TWeakPtrTS<FManifestDASHInternal> OwningManifest;
 
+	EStreamType SegmentStreamType = EStreamType::Unsupported;
+	int32 SegmentQualityIndex = 0;
+	int32 SegmentQualityIndexMax = 0;
+
 	IPlayerSessionServices* PlayerSessionServices = nullptr;
 	TSharedPtrTS<FHTTPResourceRequest> Request;
 	int32 Attempt = 0;
@@ -248,6 +252,9 @@ public:
 		bool bHasFollowingPeriod = false;			//!< true if we know for sure there is another period following.
 		bool bFrameAccurateSearch = false;			//!< true to prepare segments for frame-accurate decoding and rendering
 		bool bInitSegmentSetupOnly = false;			//!< true to get the initialization segment information only.
+		EStreamType StreamType = EStreamType::Unsupported;
+		int32 QualityIndex = 0;
+		int32 MaxQualityIndex = 0;
 	};
 
 	class FRepresentation : public IPlaybackAssetRepresentation, public TSharedFromThis<FRepresentation, ESPMode::ThreadSafe>
@@ -309,7 +316,7 @@ public:
 		}
 
 	private:
-		ESearchResult PrepareSegmentIndex(IPlayerSessionServices* PlayerSessionServices, const TArray<TSharedPtrTS<FDashMPD_SegmentBaseType>>& SegmentBase, TArray<TWeakPtrTS<FMPDLoadRequestDASH>>& OutRemoteElementLoadRequests);
+		ESearchResult PrepareSegmentIndex(IPlayerSessionServices* PlayerSessionServices, const TArray<TSharedPtrTS<FDashMPD_SegmentBaseType>>& SegmentBase, TArray<TWeakPtrTS<FMPDLoadRequestDASH>>& OutRemoteElementLoadRequests, const FSegmentSearchOption& InSearchOptions);
 
 		ESearchResult FindSegment_Base_MP4(IPlayerSessionServices* PlayerSessionServices, FSegmentInformation& OutSegmentInfo, TArray<TWeakPtrTS<FMPDLoadRequestDASH>>& OutRemoteElementLoadRequests, const FSegmentSearchOption& InSearchOptions, const TSharedPtrTS<FDashMPD_RepresentationType>& MPDRepresentation, const TArray<TSharedPtrTS<FDashMPD_SegmentBaseType>>& SegmentBase);
 		ESearchResult FindSegment_Base_MKV(IPlayerSessionServices* PlayerSessionServices, FSegmentInformation& OutSegmentInfo, TArray<TWeakPtrTS<FMPDLoadRequestDASH>>& OutRemoteElementLoadRequests, const FSegmentSearchOption& InSearchOptions, const TSharedPtrTS<FDashMPD_RepresentationType>& MPDRepresentation, const TArray<TSharedPtrTS<FDashMPD_SegmentBaseType>>& SegmentBase);
diff --git a/Engine/Plugins/Media/ElectraPlayer/Source/ElectraPlayerRuntime/Private/Runtime/Player/DASH/StreamReaderDASH.cpp b/Engine/Plugins/Media/ElectraPlayer/Source/ElectraPlayerRuntime/Private/Runtime/Player/DASH/StreamReaderDASH.cpp
--- a/Engine/Plugins/Media/ElectraPlayer/Source/ElectraPlayerRuntime/Private/Runtime/Player/DASH/StreamReaderDASH.cpp
+++ b/Engine/Plugins/Media/ElectraPlayer/Source/ElectraPlayerRuntime/Private/Runtime/Player/DASH/StreamReaderDASH.cpp
@@ -472,6 +472,9 @@ FErrorDetail FStreamReaderDASH::FStreamHandler::LoadInitSegment(TSharedPtrTS<FMP
 	// Create the download request.
 	OutLoadRequest = MakeSharedTS<FMPDLoadRequestDASH>();
 	OutLoadRequest->LoadType = FMPDLoadRequestDASH::ELoadType::Segment;
+	OutLoadRequest->SegmentStreamType = Request->StreamType;
+	OutLoadRequest->SegmentQualityIndex = Request->QualityIndex;
+	OutLoadRequest->SegmentQualityIndexMax = Request->MaxQualityIndex;
 	OutLoadRequest->URL = Request->Segment.InitializationURL.URL;
 	OutLoadRequest->Range = Request->Segment.InitializationURL.Range;
 	if (Request->Segment.InitializationURL.CustomHeader.Len())
@@ -1053,8 +1056,11 @@ void FStreamReaderDASH::FStreamHandler::HandleRequestMP4()
 			ProgressListener->ProgressDelegate   = IElectraHttpManager::FProgressListener::FProgressDelegate::CreateRaw(this, &FStreamHandler::HTTPProgressCallback);
 			ProgressListener->CompletionDelegate = IElectraHttpManager::FProgressListener::FCompletionDelegate::CreateRaw(this, &FStreamHandler::HTTPCompletionCallback);
 			TSharedPtrTS<IElectraHttpManager::FRequest> HTTP(new IElectraHttpManager::FRequest);
-			HTTP->Parameters.URL   = RequestURL;
-			HTTP->ReceiveBuffer    = ReadBuffer.ReceiveBuffer;
+			HTTP->Parameters.URL = RequestURL;
+			HTTP->Parameters.StreamType = Request->StreamType;
+			HTTP->Parameters.QualityIndex = Request->QualityIndex;
+			HTTP->Parameters.MaxQualityIndex = Request->MaxQualityIndex;
+			HTTP->ReceiveBuffer = ReadBuffer.ReceiveBuffer;
 			HTTP->ProgressListener = ProgressListener;
 			HTTP->ResponseCache = PlayerSessionService->GetHTTPResponseCache();
 			HTTP->ExternalDataReader = PlayerSessionService->GetExternalDataReader();
@@ -1883,6 +1889,9 @@ void FStreamReaderDASH::FStreamHandler::HandleRequestMKV()
 			ProgressListener->CompletionDelegate = IElectraHttpManager::FProgressListener::FCompletionDelegate::CreateRaw(this, &FStreamHandler::HTTPCompletionCallback);
 			TSharedPtrTS<IElectraHttpManager::FRequest> HTTP(new IElectraHttpManager::FRequest);
 			HTTP->Parameters.URL = RequestURL;
+			HTTP->Parameters.StreamType = Request->StreamType;
+			HTTP->Parameters.QualityIndex = Request->QualityIndex;
+			HTTP->Parameters.MaxQualityIndex = Request->MaxQualityIndex;
 			HTTP->ReceiveBuffer = ReadBuffer.ReceiveBuffer;
 			HTTP->ProgressListener = ProgressListener;
 			HTTP->ResponseCache = PlayerSessionService->GetHTTPResponseCache();
diff --git a/Engine/Plugins/Media/ElectraPlayer/Source/ElectraPlayerRuntime/Private/Runtime/Player/DASH/StreamReaderDASH.h b/Engine/Plugins/Media/ElectraPlayer/Source/ElectraPlayerRuntime/Private/Runtime/Player/DASH/StreamReaderDASH.h
--- a/Engine/Plugins/Media/ElectraPlayer/Source/ElectraPlayerRuntime/Private/Runtime/Player/DASH/StreamReaderDASH.h
+++ b/Engine/Plugins/Media/ElectraPlayer/Source/ElectraPlayerRuntime/Private/Runtime/Player/DASH/StreamReaderDASH.h
@@ -42,6 +42,8 @@ public:
 	bool GetStartupDelay(FTimeValue& OutStartTime, FTimeValue& OutTimeIntoSegment, FTimeValue& OutSegmentDuration) const override;
 
 	EStreamType												StreamType = EStreamType::Unsupported;				//!< Type of stream (video, audio, etc.)
+	int32													QualityIndex = 0;
+	int32													MaxQualityIndex = 0;
 	FStreamCodecInformation									CodecInfo;											//!< Partial codec info as can be collected from the MPD.
 	TSharedPtrTS<IPlaybackAssetRepresentation>				Representation;										//!< The representation this request belongs to.
 	TSharedPtrTS<IPlaybackAssetAdaptationSet>				AdaptationSet;										//!< The adaptation set the representation belongs to.
diff --git a/Engine/Plugins/Media/ElectraPlayer/Source/ElectraPlayerRuntime/Private/Runtime/Player/HLS/ManifestHLS.cpp b/Engine/Plugins/Media/ElectraPlayer/Source/ElectraPlayerRuntime/Private/Runtime/Player/HLS/ManifestHLS.cpp
--- a/Engine/Plugins/Media/ElectraPlayer/Source/ElectraPlayerRuntime/Private/Runtime/Player/HLS/ManifestHLS.cpp
+++ b/Engine/Plugins/Media/ElectraPlayer/Source/ElectraPlayerRuntime/Private/Runtime/Player/HLS/ManifestHLS.cpp
@@ -56,10 +56,23 @@ public:
 
 	TSharedPtrTS<ITimelineMediaAsset> GetMediaAsset() const override;
 
-	void SelectStream(const FString& AdaptationSetID, const FString& RepresentationID) override;
+	void SelectStream(const FString& AdaptationSetID, const FString& RepresentationID, int32 QualityIndex, int32 MaxQualityIndex) override;
 	void TriggerInitSegmentPreload(const TArray<FInitSegmentPreload>& InitSegmentsToPreload) override;
 
 private:
+	struct FSelectedQualityIndex
+	{
+		uint32 UniqueID = 0;
+		int32 Index = 0;
+		int32 MaxIndex = 0;
+		void Reset()
+		{
+			UniqueID = 0;
+			Index = 0;
+			MaxIndex = 0;
+		}
+	};
+
 	struct FSegSearchParam
 	{
 		FSegSearchParam()
@@ -88,7 +101,7 @@ private:
 
 	IManifest::FResult GetNextOrRetrySegment(TSharedPtrTS<IStreamSegment>& OutSegment, bool& OutWaitForRemoteElement, TSharedPtrTS<const IStreamSegment> CurrentSegment, const FPlayStartOptions& Options, bool bRetry);
 
-	IManifest::FResult FindSegment(TSharedPtrTS<FStreamSegmentRequestHLSfmp4>& OutRequest, TSharedPtrTS<FManifestHLSInternal::FPlaylistBase> InPlaylist, TSharedPtrTS<FManifestHLSInternal::FMediaStream> InStream, uint32 StreamUniqueID, EStreamType StreamType, const FSegSearchParam& SearchParam, IManifest::ESearchType SearchType);
+	IManifest::FResult FindSegment(TSharedPtrTS<FStreamSegmentRequestHLSfmp4>& OutRequest, TSharedPtrTS<FManifestHLSInternal::FPlaylistBase> InPlaylist, TSharedPtrTS<FManifestHLSInternal::FMediaStream> InStream, const FSelectedQualityIndex& ActiveQuality, EStreamType StreamType, const FSegSearchParam& SearchParam, IManifest::ESearchType SearchType);
 
 	void RefreshDenylistState();
 
@@ -97,8 +110,8 @@ private:
 	IPlaylistReaderHLS*							PlaylistReader;
 	EReadyState									CurrentReadyState;
 
-	uint32										ActiveVideoUniqueID;
-	uint32										ActiveAudioUniqueID;
+	FSelectedQualityIndex						ActiveVideoQuality;
+	FSelectedQualityIndex						ActiveAudioQuality;
 
 	TSharedPtrTS<FBufferSourceInfo>				CurrentSourceBufferInfoVideo;
 	TSharedPtrTS<FBufferSourceInfo>				CurrentSourceBufferInfoAudio;
@@ -318,12 +331,9 @@ FPlayPeriodHLS::FPlayPeriodHLS(IPlayerSessionServices* InSessionServices, IPlayl
 	, PlaylistReader(InPlaylistReader)
 {
 	check(PlaylistReader);
-
 	CurrentReadyState = IManifest::IPlayPeriod::EReadyState::NotLoaded;
-
-	// Set the active video and audio stream IDs to 0, which means none are selected.
-	ActiveVideoUniqueID = 0;
-	ActiveAudioUniqueID = 0;
+	ActiveVideoQuality.Reset();
+	ActiveAudioQuality.Reset();
 }
 
 FPlayPeriodHLS::~FPlayPeriodHLS()
@@ -423,14 +433,14 @@ void FPlayPeriodHLS::PrepareForPlay()
 
 	FManifestHLSInternal::ScopedLockPlaylists lock(InternalManifest);
 
-	uint32 OldVideoUniqueID = ActiveVideoUniqueID;
-	uint32 OldAudioUniqueID = ActiveAudioUniqueID;
+	uint32 OldVideoUniqueID = ActiveVideoQuality.UniqueID;
+	uint32 OldAudioUniqueID = ActiveAudioQuality.UniqueID;
 
 	for(int32 i=0; i<InternalManifest->VariantStreams.Num(); ++i)
 	{
 		if (InternalManifest->VariantStreams[i]->Internal.LoadState == FManifestHLSInternal::FPlaylistBase::FInternal::ELoadState::Loaded)
 		{
-			ActiveVideoUniqueID = InternalManifest->VariantStreams[i]->Internal.UniqueID;
+			ActiveVideoQuality.UniqueID = InternalManifest->VariantStreams[i]->Internal.UniqueID;
 			break;
 		}
 	}
@@ -439,37 +449,37 @@ void FPlayPeriodHLS::PrepareForPlay()
 	{
 		if (It.Value()->Internal.LoadState == FManifestHLSInternal::FPlaylistBase::FInternal::ELoadState::Loaded)
 		{
-			ActiveAudioUniqueID = It.Value()->Internal.UniqueID;
+			ActiveAudioQuality.UniqueID = It.Value()->Internal.UniqueID;
 			break;
 		}
 	}
 	// In case there is an audio rendition without a dedicated playlist we look at audio-only variant streams
-	if (ActiveAudioUniqueID == 0)
+	if (ActiveAudioQuality.UniqueID == 0)
 	{
 		for(int32 i=0; i<InternalManifest->AudioOnlyStreams.Num(); ++i)
 		{
 			if (InternalManifest->AudioOnlyStreams[i]->Internal.LoadState == FManifestHLSInternal::FPlaylistBase::FInternal::ELoadState::Loaded)
 			{
-				ActiveAudioUniqueID = InternalManifest->AudioOnlyStreams[i]->Internal.UniqueID;
+				ActiveAudioQuality.UniqueID = InternalManifest->AudioOnlyStreams[i]->Internal.UniqueID;
 				break;
 			}
 		}
 	}
 
 	// Tell the manifest which stream IDs are now actively used.
-	InternalManifest->SelectActiveStreamID(ActiveVideoUniqueID, OldVideoUniqueID);
-	InternalManifest->SelectActiveStreamID(ActiveAudioUniqueID, OldAudioUniqueID);
+	InternalManifest->SelectActiveStreamID(ActiveVideoQuality.UniqueID, OldVideoUniqueID);
+	InternalManifest->SelectActiveStreamID(ActiveAudioQuality.UniqueID, OldAudioUniqueID);
 
 	// Set up source buffer information for video and audio.
 	// These are currently dummies since we do not support track switching yet.
-	if (ActiveVideoUniqueID)
+	if (ActiveVideoQuality.UniqueID)
 	{
 		CurrentSourceBufferInfoVideo = MakeSharedTS<FBufferSourceInfo>();
 		CurrentSourceBufferInfoVideo->PeriodID = InternalManifest->CurrentMediaAsset->GetAssetIdentifier();
 		CurrentSourceBufferInfoVideo->PeriodAdaptationSetID = TEXT("video.0");
 		CurrentSourceBufferInfoVideo->HardIndex = 0;
 	}
-	if (ActiveAudioUniqueID)
+	if (ActiveAudioQuality.UniqueID)
 	{
 		CurrentSourceBufferInfoAudio = MakeSharedTS<FBufferSourceInfo>();
 		CurrentSourceBufferInfoAudio->PeriodID = InternalManifest->CurrentMediaAsset->GetAssetIdentifier();
@@ -575,14 +585,14 @@ IManifest::FResult FPlayPeriodHLS::GetMediaStreamForID(TSharedPtrTS<FManifestHLS
  * @param OutRequest
  * @param InPlaylist
  * @param InStream
- * @param StreamUniqueID
+ * @param ActiveQuality
  * @param StreamType
  * @param SearchParam
  * @param SearchType
  *
  * @return
  */
-IManifest::FResult FPlayPeriodHLS::FindSegment(TSharedPtrTS<FStreamSegmentRequestHLSfmp4>& OutRequest, TSharedPtrTS<FManifestHLSInternal::FPlaylistBase> InPlaylist, TSharedPtrTS<FManifestHLSInternal::FMediaStream> InStream, uint32 StreamUniqueID, EStreamType StreamType, const FPlayPeriodHLS::FSegSearchParam& SearchParam, IManifest::ESearchType SearchType)
+IManifest::FResult FPlayPeriodHLS::FindSegment(TSharedPtrTS<FStreamSegmentRequestHLSfmp4>& OutRequest, TSharedPtrTS<FManifestHLSInternal::FPlaylistBase> InPlaylist, TSharedPtrTS<FManifestHLSInternal::FMediaStream> InStream, const FSelectedQualityIndex& ActiveQuality, EStreamType StreamType, const FPlayPeriodHLS::FSegSearchParam& SearchParam, IManifest::ESearchType SearchType)
 {
 	SCOPE_CYCLE_COUNTER(STAT_ElectraPlayer_HLS_FindSegment);
 	CSV_SCOPED_TIMING_STAT(ElectraPlayer, HLS_FindSegment);
@@ -595,8 +605,10 @@ IManifest::FResult FPlayPeriodHLS::FindSegment(TSharedPtrTS<FStreamSegmentReques
 	FURL_RFC3986 UrlBuilder;
 	UrlBuilder.Parse(InPlaylist->Internal.PlaylistLoadRequest.URL);
 	TSharedPtrTS<FStreamSegmentRequestHLSfmp4>	Req(new FStreamSegmentRequestHLSfmp4);
-	Req->StreamType 	= StreamType;
-	Req->StreamUniqueID = StreamUniqueID;
+	Req->StreamType = StreamType;
+	Req->StreamUniqueID = ActiveQuality.UniqueID;
+	Req->QualityIndex = ActiveQuality.Index;
+	Req->MaxQualityIndex = ActiveQuality.MaxIndex;
 	Req->MediaAsset = InternalManifest->CurrentMediaAsset;
 	if (!Req->MediaAsset.IsValid())
 	{
@@ -611,20 +623,20 @@ IManifest::FResult FPlayPeriodHLS::FindSegment(TSharedPtrTS<FStreamSegmentReques
 	{
 		return IManifest::FResult(IManifest::FResult::EType::NotFound).SetErrorDetail(FErrorDetail().SetMessage(FString::Printf(TEXT("Internal error, no %s rendition group found on asset timeline!"), GetStreamTypeName(StreamType))));
 	}
-	Req->Representation = Req->AdaptationSet->GetRepresentationByUniqueIdentifier(LexToString(StreamUniqueID));
+	Req->Representation = Req->AdaptationSet->GetRepresentationByUniqueIdentifier(LexToString(ActiveQuality.UniqueID));
 	if (!Req->Representation.IsValid())
 	{
 		return IManifest::FResult(IManifest::FResult::EType::NotFound).SetErrorDetail(FErrorDetail().SetMessage(FString::Printf(TEXT("Internal error, %s rendition not found in group on asset timeline!"), GetStreamTypeName(StreamType))));
 	}
-	Req->InitSegmentCache   	   = InternalManifest->InitSegmentCache;
-	Req->LicenseKeyCache		   = InternalManifest->LicenseKeyCache;
-	Req->bHasEncryptedSegments     = InStream->bHasEncryptedSegments;
+	Req->InitSegmentCache = InternalManifest->InitSegmentCache;
+	Req->LicenseKeyCache = InternalManifest->LicenseKeyCache;
+	Req->bHasEncryptedSegments = InStream->bHasEncryptedSegments;
 	if (StreamType == EStreamType::Video)
 	{
 		Req->SourceBufferInfo = CurrentSourceBufferInfoVideo;
-		Req->Bitrate				   = InPlaylist->GetBitrate();
+		Req->Bitrate = InPlaylist->GetBitrate();
 		check(InternalManifest->BandwidthToQualityIndex.Find(Req->Bitrate) != nullptr);
-		Req->QualityLevel   		   = InternalManifest->BandwidthToQualityIndex[Req->Bitrate];
+		Req->QualityLevel = InternalManifest->BandwidthToQualityIndex[Req->Bitrate];
 	}
 	else if (StreamType == EStreamType::Audio)
 	{
@@ -861,8 +873,8 @@ IManifest::FResult FPlayPeriodHLS::GetStartingSegment(TSharedPtrTS<IStreamSegmen
 	TSharedPtrTS<FManifestHLSInternal::FMediaStream>	AudioStream;
 
 	// Get the streams that are selected, if there are selected ones.
-	IManifest::FResult vidResult = ActiveVideoUniqueID ? GetMediaStreamForID(VideoPlaylist, VideoStream, ActiveVideoUniqueID) : IManifest::FResult(IManifest::FResult::EType::Found);
-	IManifest::FResult audResult = ActiveAudioUniqueID ? GetMediaStreamForID(AudioPlaylist, AudioStream, ActiveAudioUniqueID) : IManifest::FResult(IManifest::FResult::EType::Found);
+	IManifest::FResult vidResult = ActiveVideoQuality.UniqueID ? GetMediaStreamForID(VideoPlaylist, VideoStream, ActiveVideoQuality.UniqueID) : IManifest::FResult(IManifest::FResult::EType::Found);
+	IManifest::FResult audResult = ActiveAudioQuality.UniqueID ? GetMediaStreamForID(AudioPlaylist, AudioStream, ActiveAudioQuality.UniqueID) : IManifest::FResult(IManifest::FResult::EType::Found);
 	if (vidResult.IsSuccess() && audResult.IsSuccess())
 	{
 		TSharedPtrTS<FStreamSegmentRequestHLSfmp4> 	VideoSegmentRequest;
@@ -887,9 +899,9 @@ IManifest::FResult FPlayPeriodHLS::GetStartingSegment(TSharedPtrTS<IStreamSegmen
 
 
 		// Do we have both video and audio?
-		if (ActiveVideoUniqueID && ActiveAudioUniqueID)
+		if (ActiveVideoQuality.UniqueID && ActiveAudioQuality.UniqueID)
 		{
-			vidResult = FindSegment(VideoSegmentRequest, VideoPlaylist, VideoStream, ActiveVideoUniqueID, EStreamType::Video, SearchParam, SearchType);
+			vidResult = FindSegment(VideoSegmentRequest, VideoPlaylist, VideoStream, ActiveVideoQuality, EStreamType::Video, SearchParam, SearchType);
 			// Found and PastEOS are valid results here. Everything else is not.
 			if (vidResult.GetType() != IManifest::FResult::EType::Found && vidResult.GetType() != IManifest::FResult::EType::PastEOS)
 			{
@@ -912,7 +924,7 @@ IManifest::FResult FPlayPeriodHLS::GetStartingSegment(TSharedPtrTS<IStreamSegmen
 				}
 			}
 			// Search for audio.
-			audResult = FindSegment(AudioSegmentRequest, AudioPlaylist, AudioStream, ActiveAudioUniqueID, EStreamType::Audio, SearchParam, SearchType);
+			audResult = FindSegment(AudioSegmentRequest, AudioPlaylist, AudioStream, ActiveAudioQuality, EStreamType::Audio, SearchParam, SearchType);
 			// Equally here, if successful or PastEOS is acceptable and everything else is not.
 			if (audResult.GetType() != IManifest::FResult::EType::Found && audResult.GetType() != IManifest::FResult::EType::PastEOS)
 			{
@@ -945,9 +957,9 @@ IManifest::FResult FPlayPeriodHLS::GetStartingSegment(TSharedPtrTS<IStreamSegmen
 			}
 		}
 		// Video only?
-		else if (ActiveVideoUniqueID)
+		else if (ActiveVideoQuality.UniqueID)
 		{
-			vidResult = FindSegment(VideoSegmentRequest, VideoPlaylist, VideoStream, ActiveVideoUniqueID, EStreamType::Video, SearchParam, SearchType);
+			vidResult = FindSegment(VideoSegmentRequest, VideoPlaylist, VideoStream, ActiveVideoQuality, EStreamType::Video, SearchParam, SearchType);
 			if (vidResult.IsSuccess())
 			{
 				VideoSegmentRequest->bIsInitialStartRequest = true;
@@ -962,7 +974,7 @@ IManifest::FResult FPlayPeriodHLS::GetStartingSegment(TSharedPtrTS<IStreamSegmen
 		// Audio only
 		else
 		{
-			audResult = FindSegment(AudioSegmentRequest, AudioPlaylist, AudioStream, ActiveAudioUniqueID, EStreamType::Audio, SearchParam, SearchType);
+			audResult = FindSegment(AudioSegmentRequest, AudioPlaylist, AudioStream, ActiveAudioQuality, EStreamType::Audio, SearchParam, SearchType);
 			if (audResult.IsSuccess())
 			{
 				AudioSegmentRequest->bIsInitialStartRequest = true;
@@ -988,7 +1000,7 @@ IManifest::FResult FPlayPeriodHLS::GetStartingSegment(TSharedPtrTS<IStreamSegmen
 			UrlBuilder.Parse(InternalManifest->MasterPlaylistVars.PlaylistLoadRequest.URL);
 			FPlaylistLoadRequestHLS Req;
 			Req.LoadType			   = FPlaylistLoadRequestHLS::ELoadType::First;
-			Req.InternalUniqueID	   = ActiveVideoUniqueID;
+			Req.InternalUniqueID	   = ActiveVideoQuality.UniqueID;
 			Req.RequestedAtTime 	   = SessionServices->GetSynchronizedUTCTime()->GetTime();
 			Req.URL 				   = FURL_RFC3986(UrlBuilder).ResolveWith(VideoPlaylist->GetURL()).Get();
 			Req.AdaptationSetUniqueID  = VideoPlaylist->Internal.AdaptationSetUniqueID;
@@ -1007,7 +1019,7 @@ IManifest::FResult FPlayPeriodHLS::GetStartingSegment(TSharedPtrTS<IStreamSegmen
 			UrlBuilder.Parse(InternalManifest->MasterPlaylistVars.PlaylistLoadRequest.URL);
 			FPlaylistLoadRequestHLS Req;
 			Req.LoadType			   = FPlaylistLoadRequestHLS::ELoadType::First;
-			Req.InternalUniqueID	   = ActiveAudioUniqueID;
+			Req.InternalUniqueID	   = ActiveAudioQuality.UniqueID;
 			Req.RequestedAtTime 	   = SessionServices->GetSynchronizedUTCTime()->GetTime();
 			Req.URL 				   = FURL_RFC3986(UrlBuilder).ResolveWith(AudioPlaylist->GetURL()).Get();
 			Req.AdaptationSetUniqueID  = AudioPlaylist->Internal.AdaptationSetUniqueID;
@@ -1082,17 +1094,17 @@ IManifest::FResult FPlayPeriodHLS::GetNextOrRetrySegment(TSharedPtrTS<IStreamSeg
 	}
 
 	FManifestHLSInternal::ScopedLockPlaylists lock(InternalManifest);
-	uint32 ForStreamID = 0;
+	FSelectedQualityIndex ForQuality;
 	switch(CurrentRequest->GetType())
 	{
 		case EStreamType::Video:
 		{
-			ForStreamID = ActiveVideoUniqueID;
+			ForQuality = ActiveVideoQuality;
 			break;
 		}
 		case EStreamType::Audio:
 		{
-			ForStreamID = ActiveAudioUniqueID;
+			ForQuality = ActiveAudioQuality;
 			break;
 		}
 		default:
@@ -1100,16 +1112,16 @@ IManifest::FResult FPlayPeriodHLS::GetNextOrRetrySegment(TSharedPtrTS<IStreamSeg
 			return IManifest::FResult(IManifest::FResult::EType::NotFound).SetErrorDetail(FErrorDetail().SetMessage(FString::Printf(TEXT("Cannot get next segment for unsupported stream type \"%s\"!"), GetStreamTypeName(CurrentRequest->GetType()))));
 		}
 	}
-	if (ForStreamID == 0)
+	if (ForQuality.UniqueID == 0)
 	{
 		return IManifest::FResult(IManifest::FResult::EType::NotFound).SetErrorDetail(FErrorDetail().SetMessage(FString::Printf(TEXT("Cannot get next segment stream type \"%s\" since no stream is actively selected!"), GetStreamTypeName(CurrentRequest->GetType()))));
 	}
 
 	RefreshDenylistState();
 
-	TSharedPtrTS<FManifestHLSInternal::FMediaStream>	Stream;
+	TSharedPtrTS<FManifestHLSInternal::FMediaStream> Stream;
 	TSharedPtrTS<FManifestHLSInternal::FPlaylistBase> Playlist;
-	IManifest::FResult Result = GetMediaStreamForID(Playlist, Stream, ForStreamID);
+	IManifest::FResult Result = GetMediaStreamForID(Playlist, Stream, ForQuality.UniqueID);
 	if (Result.IsSuccess())
 	{
 		TSharedPtrTS<FStreamSegmentRequestHLSfmp4> 	NextSegmentRequest;
@@ -1127,9 +1139,9 @@ IManifest::FResult FPlayPeriodHLS::GetNextOrRetrySegment(TSharedPtrTS<IStreamSeg
 		SearchParam.MediaSequence   	  = CurrentRequest->MediaSequence;
 		SearchParam.DiscontinuitySequence = CurrentRequest->DiscontinuitySequence;
 		SearchParam.LocalIndex  		  = CurrentRequest->LocalIndex;
-		SearchParam.StreamUniqueID  	  = CurrentRequest->StreamUniqueID == ForStreamID ? ForStreamID : 0;
+		SearchParam.StreamUniqueID  	  = CurrentRequest->StreamUniqueID == ForQuality.UniqueID ? ForQuality.UniqueID : 0;
 		SearchParam.TimestampSequenceIndex = CurrentRequest->TimestampSequenceIndex;
-		Result = FindSegment(NextSegmentRequest, Playlist, Stream, ForStreamID, CurrentRequest->GetType(), SearchParam, bRetry ? IManifest::ESearchType::Same : IManifest::ESearchType::StrictlyAfter);
+		Result = FindSegment(NextSegmentRequest, Playlist, Stream, ForQuality, CurrentRequest->GetType(), SearchParam, bRetry ? IManifest::ESearchType::Same : IManifest::ESearchType::StrictlyAfter);
 		if (Result.IsSuccess())
 		{
 			// Continuing with the next segment implicitly means there is no AU time offset.
@@ -1169,7 +1181,7 @@ IManifest::FResult FPlayPeriodHLS::GetNextOrRetrySegment(TSharedPtrTS<IStreamSeg
 			UrlBuilder.Parse(InternalManifest->MasterPlaylistVars.PlaylistLoadRequest.URL);
 			FPlaylistLoadRequestHLS Req;
 			Req.LoadType			   = FPlaylistLoadRequestHLS::ELoadType::First;
-			Req.InternalUniqueID	   = ForStreamID;
+			Req.InternalUniqueID	   = ForQuality.UniqueID;
 			Req.RequestedAtTime 	   = SessionServices->GetSynchronizedUTCTime()->GetTime();
 			Req.URL 				   = UrlBuilder.ResolveWith(Playlist->GetURL()).Get();
 			Req.AdaptationSetUniqueID  = Playlist->Internal.AdaptationSetUniqueID;
@@ -1415,7 +1427,7 @@ TSharedPtrTS<ITimelineMediaAsset> FPlayPeriodHLS::GetMediaAsset() const
 
 
 
-void FPlayPeriodHLS::SelectStream(const FString& AdaptationSetID, const FString& RepresentationID)
+void FPlayPeriodHLS::SelectStream(const FString& AdaptationSetID, const FString& RepresentationID, int32 QualityIndex, int32 MaxQualityIndex)
 {
 	RefreshDenylistState();
 
@@ -1427,29 +1439,33 @@ void FPlayPeriodHLS::SelectStream(const FString& AdaptationSetID, const FString&
 	if (AdaptationSetID.Equals(TEXT("$video$")))
 	{
 		// Different from what we have actively selected?
-		if (UniqueID != ActiveVideoUniqueID)
+		if (UniqueID != ActiveVideoQuality.UniqueID)
 		{
 			// FIXME: We could emit a QoS event here. Although selecting the stream does not mean it will actually get used.
 			//        There is still a chance that (if the playlist is not loaded yet) another stream will be selected right away.
 
 			// Tell the manifest that we are now using a different stream.
-			InternalManifest->SelectActiveStreamID(UniqueID, ActiveVideoUniqueID);
+			InternalManifest->SelectActiveStreamID(UniqueID, ActiveVideoQuality.UniqueID);
 
-			ActiveVideoUniqueID = UniqueID;
+			ActiveVideoQuality.UniqueID = UniqueID;
+			ActiveVideoQuality.Index = QualityIndex;
+			ActiveVideoQuality.MaxIndex = MaxQualityIndex;
 		}
 	}
 	else if (AdaptationSetID.Equals(TEXT("$audio$")))
 	{
 		// Different from what we have actively selected?
-		if (UniqueID != ActiveAudioUniqueID)
+		if (UniqueID != ActiveAudioQuality.UniqueID)
 		{
 			// FIXME: We could emit a QoS event here. Although selecting the stream does not mean it will actually get used.
 			//        There is still a chance that (if the playlist is not loaded yet) another stream will be selected right away.
 
 			// Tell the manifest that we are now using a different stream.
-			InternalManifest->SelectActiveStreamID(UniqueID, ActiveAudioUniqueID);
+			InternalManifest->SelectActiveStreamID(UniqueID, ActiveAudioQuality.UniqueID);
 
-			ActiveAudioUniqueID = UniqueID;
+			ActiveAudioQuality.UniqueID = UniqueID;
+			ActiveAudioQuality.Index = QualityIndex;
+			ActiveAudioQuality.MaxIndex = MaxQualityIndex;
 		}
 	}
 }
diff --git a/Engine/Plugins/Media/ElectraPlayer/Source/ElectraPlayerRuntime/Private/Runtime/Player/HLS/StreamReaderHLSfmp4.cpp b/Engine/Plugins/Media/ElectraPlayer/Source/ElectraPlayerRuntime/Private/Runtime/Player/HLS/StreamReaderHLSfmp4.cpp
--- a/Engine/Plugins/Media/ElectraPlayer/Source/ElectraPlayerRuntime/Private/Runtime/Player/HLS/StreamReaderHLSfmp4.cpp
+++ b/Engine/Plugins/Media/ElectraPlayer/Source/ElectraPlayerRuntime/Private/Runtime/Player/HLS/StreamReaderHLSfmp4.cpp
@@ -23,22 +23,24 @@ namespace Electra
 
 FStreamSegmentRequestHLSfmp4::FStreamSegmentRequestHLSfmp4()
 {
-	StreamType  			  = EStreamType::Video;
-	StreamUniqueID  		  = 0;
-	Bitrate 				  = 0;
-	QualityLevel			  = 0;
-	MediaSequence   		  = -1;
-	DiscontinuitySequence     = -1;
-	LocalIndex  			  = -1;
-	TimestampSequenceIndex	  = 0;
-	bIsPrefetch 			  = false;
-	bIsEOSSegment   		  = false;
-	bHasEncryptedSegments     = false;
-	NumOverallRetries   	  = 0;
-	bInsertFillerData   	  = false;
-	bIsInitialStartRequest    = false;
+	StreamType = EStreamType::Video;
+	QualityIndex = 0;
+	MaxQualityIndex = 0;
+	StreamUniqueID = 0;
+	Bitrate = 0;
+	QualityLevel = 0;
+	MediaSequence = -1;
+	DiscontinuitySequence = -1;
+	LocalIndex = -1;
+	TimestampSequenceIndex = 0;
+	bIsPrefetch = false;
+	bIsEOSSegment = false;
+	bHasEncryptedSegments = false;
+	NumOverallRetries = 0;
+	bInsertFillerData = false;
+	bIsInitialStartRequest = false;
 	CurrentPlaybackSequenceID = ~0U;
-	bFrameAccuracyRequired	  = false;
+	bFrameAccuracyRequired = false;
 }
 
 FStreamSegmentRequestHLSfmp4::~FStreamSegmentRequestHLSfmp4()
@@ -524,9 +526,9 @@ FStreamReaderHLSfmp4::FStreamHandler::EInitSegmentResult FStreamReaderHLSfmp4::F
 				}
 			}
 
-			TSharedPtrTS<IElectraHttpManager::FRequest> 			HTTP(new IElectraHttpManager::FRequest);
-			TSharedPtrTS<IElectraHttpManager::FProgressListener>	ProgressListener(new IElectraHttpManager::FProgressListener);
-			ProgressListener->ProgressDelegate   = IElectraHttpManager::FProgressListener::FProgressDelegate::CreateRaw(this, &FStreamHandler::HTTPProgressCallback);
+			TSharedPtrTS<IElectraHttpManager::FRequest> HTTP(new IElectraHttpManager::FRequest);
+			TSharedPtrTS<IElectraHttpManager::FProgressListener> ProgressListener(new IElectraHttpManager::FProgressListener);
+			ProgressListener->ProgressDelegate    = IElectraHttpManager::FProgressListener::FProgressDelegate::CreateRaw(this, &FStreamHandler::HTTPProgressCallback);
 			ProgressListener->CompletionDelegate = IElectraHttpManager::FProgressListener::FCompletionDelegate::CreateRaw(this, &FStreamHandler::HTTPCompletionCallback);
 			ReadBuffer.Reset();
 			ReadBuffer.ReceiveBuffer = MakeSharedTS<IElectraHttpManager::FReceiveBuffer>();
@@ -534,20 +536,24 @@ FStreamReaderHLSfmp4::FStreamHandler::EInitSegmentResult FStreamReaderHLSfmp4::F
 			FString RequestURL = Request->InitSegmentInfo->URI;
 
 			Metrics::FSegmentDownloadStats& ds = Request->DownloadStats;
-			ds.URL  	   = RequestURL;
+			ds.URL = RequestURL;
 			ds.SegmentType = Metrics::ESegmentType::Init;
 
-			HTTP->Parameters.URL   = RequestURL;
-			HTTP->ReceiveBuffer    = ReadBuffer.ReceiveBuffer;
+			HTTP->Parameters.URL = RequestURL;
+			HTTP->Parameters.StreamType = Request->StreamType;
+			HTTP->Parameters.QualityIndex = Request->QualityIndex;
+			HTTP->Parameters.MaxQualityIndex = Request->MaxQualityIndex;
+			HTTP->ReceiveBuffer = ReadBuffer.ReceiveBuffer;
 			HTTP->ProgressListener = ProgressListener;
 			if (Request->InitSegmentInfo->ByteRange.IsSet())
 			{
-				HTTP->Parameters.Range.Start		= Request->InitSegmentInfo->ByteRange.GetStart();
+				HTTP->Parameters.Range.Start = Request->InitSegmentInfo->ByteRange.GetStart();
 				HTTP->Parameters.Range.EndIncluding = Request->InitSegmentInfo->ByteRange.GetEnd();
 			}
 			// Set timeouts for init segment retrieval
 			HTTP->Parameters.ConnectTimeout = FTimeValue().SetFromMilliseconds(1000 * 4);
 			HTTP->Parameters.NoDataTimeout = FTimeValue().SetFromMilliseconds(1000 * 2);
+			HTTP->ResponseCache = PlayerSessionService->GetHTTPResponseCache();
 			HTTP->ExternalDataReader = PlayerSessionService->GetExternalDataReader();
 
 			ProgressReportCount = 0;
@@ -821,12 +827,15 @@ void FStreamReaderHLSfmp4::FStreamHandler::HandleRequest()
 				// Clear out the current connection info which may now be populated with the init segment fetch results.
 				CurrentRequest->ConnectionInfo = {};
 
-				TSharedPtrTS<IElectraHttpManager::FProgressListener>	ProgressListener(new IElectraHttpManager::FProgressListener);
-				ProgressListener->ProgressDelegate   = IElectraHttpManager::FProgressListener::FProgressDelegate::CreateRaw(this, &FStreamHandler::HTTPProgressCallback);
+				TSharedPtrTS<IElectraHttpManager::FProgressListener> ProgressListener(new IElectraHttpManager::FProgressListener);
+				ProgressListener->ProgressDelegate = IElectraHttpManager::FProgressListener::FProgressDelegate::CreateRaw(this, &FStreamHandler::HTTPProgressCallback);
 				ProgressListener->CompletionDelegate = IElectraHttpManager::FProgressListener::FCompletionDelegate::CreateRaw(this, &FStreamHandler::HTTPCompletionCallback);
 				TSharedPtrTS<IElectraHttpManager::FRequest> HTTP(new IElectraHttpManager::FRequest);
-				HTTP->Parameters.URL   = RequestURL;
-				HTTP->ReceiveBuffer    = ReadBuffer.ReceiveBuffer;
+				HTTP->Parameters.URL = RequestURL;
+				HTTP->Parameters.StreamType = Request->StreamType;
+				HTTP->Parameters.QualityIndex = Request->QualityIndex;
+				HTTP->Parameters.MaxQualityIndex = Request->MaxQualityIndex;
+				HTTP->ReceiveBuffer = ReadBuffer.ReceiveBuffer;
 				HTTP->ProgressListener = ProgressListener;
 				if (Request->Range.IsSet())
 				{
diff --git a/Engine/Plugins/Media/ElectraPlayer/Source/ElectraPlayerRuntime/Private/Runtime/Player/HLS/StreamReaderHLSfmp4.h b/Engine/Plugins/Media/ElectraPlayer/Source/ElectraPlayerRuntime/Private/Runtime/Player/HLS/StreamReaderHLSfmp4.h
--- a/Engine/Plugins/Media/ElectraPlayer/Source/ElectraPlayerRuntime/Private/Runtime/Player/HLS/StreamReaderHLSfmp4.h
+++ b/Engine/Plugins/Media/ElectraPlayer/Source/ElectraPlayerRuntime/Private/Runtime/Player/HLS/StreamReaderHLSfmp4.h
@@ -45,6 +45,8 @@ public:
 	ElectraHTTPStream::FHttpRange												Range;
 
 	EStreamType																	StreamType;							//!< Type of stream (video, audio, etc.)
+	int32																		QualityIndex;
+	int32																		MaxQualityIndex;
 	uint32																		StreamUniqueID;						//!< The unique stream ID identifying the stream for which this is a request.
 	int32																		Bitrate;
 	int32																		QualityLevel;
diff --git a/Engine/Plugins/Media/ElectraPlayer/Source/ElectraPlayerRuntime/Private/Runtime/Player/Manifest.h b/Engine/Plugins/Media/ElectraPlayer/Source/ElectraPlayerRuntime/Private/Runtime/Player/Manifest.h
--- a/Engine/Plugins/Media/ElectraPlayer/Source/ElectraPlayerRuntime/Private/Runtime/Player/Manifest.h
+++ b/Engine/Plugins/Media/ElectraPlayer/Source/ElectraPlayerRuntime/Private/Runtime/Player/Manifest.h
@@ -413,7 +413,7 @@ namespace Electra
 			 * @param AdaptationSetID
 			 * @param RepresentationID
 			 */
-			virtual void SelectStream(const FString& AdaptationSetID, const FString& RepresentationID) = 0;
+			virtual void SelectStream(const FString& AdaptationSetID, const FString& RepresentationID, int32 QualityIndex, int32 MaxQualityIndex) = 0;
 
 			struct FInitSegmentPreload
 			{
diff --git a/Engine/Plugins/Media/ElectraPlayer/Source/ElectraPlayerRuntime/Private/Runtime/Player/mkv/ManifestMKV.cpp b/Engine/Plugins/Media/ElectraPlayer/Source/ElectraPlayerRuntime/Private/Runtime/Player/mkv/ManifestMKV.cpp
--- a/Engine/Plugins/Media/ElectraPlayer/Source/ElectraPlayerRuntime/Private/Runtime/Player/mkv/ManifestMKV.cpp
+++ b/Engine/Plugins/Media/ElectraPlayer/Source/ElectraPlayerRuntime/Private/Runtime/Player/mkv/ManifestMKV.cpp
@@ -488,7 +488,7 @@ TSharedPtrTS<ITimelineMediaAsset> FManifestMKVInternal::FPlayPeriodMKV::GetMedia
  * @param AdaptationSetID
  * @param RepresentationID
  */
-void FManifestMKVInternal::FPlayPeriodMKV::SelectStream(const FString& AdaptationSetID, const FString& RepresentationID)
+void FManifestMKVInternal::FPlayPeriodMKV::SelectStream(const FString& AdaptationSetID, const FString& RepresentationID, int32 QualityIndex, int32 MaxQualityIndex)
 {
 	// Presently this method is only called by the ABR to switch between quality levels.
 	// Since a single mkv doesn't have different quality levels we ignore this for now.
diff --git a/Engine/Plugins/Media/ElectraPlayer/Source/ElectraPlayerRuntime/Private/Runtime/Player/mkv/ManifestMKV.h b/Engine/Plugins/Media/ElectraPlayer/Source/ElectraPlayerRuntime/Private/Runtime/Player/mkv/ManifestMKV.h
--- a/Engine/Plugins/Media/ElectraPlayer/Source/ElectraPlayerRuntime/Private/Runtime/Player/mkv/ManifestMKV.h
+++ b/Engine/Plugins/Media/ElectraPlayer/Source/ElectraPlayerRuntime/Private/Runtime/Player/mkv/ManifestMKV.h
@@ -274,7 +274,7 @@ public:
 		FString GetSelectedAdaptationSetID(EStreamType StreamType) override;
 		ETrackChangeResult ChangeTrackStreamPreference(EStreamType ForStreamType, const FStreamSelectionAttributes& StreamAttributes) override;
 		TSharedPtrTS<ITimelineMediaAsset> GetMediaAsset() const override;
-		void SelectStream(const FString& AdaptationSetID, const FString& RepresentationID) override;
+		void SelectStream(const FString& AdaptationSetID, const FString& RepresentationID, int32 QualityIndex, int32 MaxQualityIndex) override;
 		void TriggerInitSegmentPreload(const TArray<FInitSegmentPreload>& InitSegmentsToPreload) override;
 		FResult GetStartingSegment(TSharedPtrTS<IStreamSegment>& OutSegment, const FPlayerSequenceState& InSequenceState, const FPlayStartPosition& StartPosition, ESearchType SearchType) override;
 		FResult GetContinuationSegment(TSharedPtrTS<IStreamSegment>& OutSegment, EStreamType StreamType, const FPlayerSequenceState& LoopState, const FPlayStartPosition& StartPosition, ESearchType SearchType) override;
diff --git a/Engine/Plugins/Media/ElectraPlayer/Source/ElectraPlayerRuntime/Private/Runtime/Player/mp4/ManifestMP4.cpp b/Engine/Plugins/Media/ElectraPlayer/Source/ElectraPlayerRuntime/Private/Runtime/Player/mp4/ManifestMP4.cpp
--- a/Engine/Plugins/Media/ElectraPlayer/Source/ElectraPlayerRuntime/Private/Runtime/Player/mp4/ManifestMP4.cpp
+++ b/Engine/Plugins/Media/ElectraPlayer/Source/ElectraPlayerRuntime/Private/Runtime/Player/mp4/ManifestMP4.cpp
@@ -507,7 +507,7 @@ TSharedPtrTS<ITimelineMediaAsset> FManifestMP4Internal::FPlayPeriodMP4::GetMedia
  * @param AdaptationSetID
  * @param RepresentationID
  */
-void FManifestMP4Internal::FPlayPeriodMP4::SelectStream(const FString& AdaptationSetID, const FString& RepresentationID)
+void FManifestMP4Internal::FPlayPeriodMP4::SelectStream(const FString& AdaptationSetID, const FString& RepresentationID, int32 QualityIndex, int32 MaxQualityIndex)
 {
 	// Presently this method is only called by the ABR to switch between quality levels.
 	// Since a single mp4 doesn't have different quality levels (technically it could, but we are concerning ourselves only with different bitrates and that doesn't apply since we are streaming
diff --git a/Engine/Plugins/Media/ElectraPlayer/Source/ElectraPlayerRuntime/Private/Runtime/Player/mp4/ManifestMP4.h b/Engine/Plugins/Media/ElectraPlayer/Source/ElectraPlayerRuntime/Private/Runtime/Player/mp4/ManifestMP4.h
--- a/Engine/Plugins/Media/ElectraPlayer/Source/ElectraPlayerRuntime/Private/Runtime/Player/mp4/ManifestMP4.h
+++ b/Engine/Plugins/Media/ElectraPlayer/Source/ElectraPlayerRuntime/Private/Runtime/Player/mp4/ManifestMP4.h
@@ -280,7 +280,7 @@ public:
 		FString GetSelectedAdaptationSetID(EStreamType StreamType) override;
 		ETrackChangeResult ChangeTrackStreamPreference(EStreamType ForStreamType, const FStreamSelectionAttributes& StreamAttributes) override;
 		TSharedPtrTS<ITimelineMediaAsset> GetMediaAsset() const override;
-		void SelectStream(const FString& AdaptationSetID, const FString& RepresentationID) override;
+		void SelectStream(const FString& AdaptationSetID, const FString& RepresentationID, int32 QualityIndex, int32 MaxQualityIndex) override;
 		void TriggerInitSegmentPreload(const TArray<FInitSegmentPreload>& InitSegmentsToPreload) override;
 		FResult GetStartingSegment(TSharedPtrTS<IStreamSegment>& OutSegment, const FPlayerSequenceState& InSequenceState, const FPlayStartPosition& StartPosition, ESearchType SearchType) override;
 		FResult GetContinuationSegment(TSharedPtrTS<IStreamSegment>& OutSegment, EStreamType StreamType, const FPlayerSequenceState& LoopState, const FPlayStartPosition& StartPosition, ESearchType SearchType) override;
diff --git a/Engine/Plugins/Media/ElectraPlayer/Source/ElectraPlayerRuntime/Public/IElectraPlayerDataCache.h b/Engine/Plugins/Media/ElectraPlayer/Source/ElectraPlayerRuntime/Public/IElectraPlayerDataCache.h
new file mode 100644
--- /dev/null
+++ b/Engine/Plugins/Media/ElectraPlayer/Source/ElectraPlayerRuntime/Public/IElectraPlayerDataCache.h
@@ -0,0 +1,106 @@
+// Copyright Epic Games, Inc. All Rights Reserved.
+#pragma once
+
+#include "CoreMinimal.h"
+
+class IElectraPlayerDataCache : public TSharedFromThis<IElectraPlayerDataCache, ESPMode::ThreadSafe>
+{
+public:
+	virtual ~IElectraPlayerDataCache() = default;
+
+	struct FItemInfo
+	{
+		struct FRange
+		{
+			/** 
+			 * Offset to the first byte in the resource.
+			 */
+			int64 Start = -1;
+			
+			/**
+			 * Offset to the last byte, including.
+			 * 
+			 * A range of `Start` = 0, `EndIncluding` = 0 describes one byte at offset 0.
+			 */
+			int64 EndIncluding = -1;
+
+			/** 
+			 * The total size of the resource. This _may_ be set when adding to the cache
+			 * but is not required as there are cases where the total resource size is unknown.
+			 * This does not describe the number of bytes between `Start` and `EndIncluding`,
+			 * but the total size of the resource of which the range is a part of.
+			 */
+			int64 TotalSize = -1;
+		};
+
+		enum class EStreamType
+		{
+			Video,
+			Audio,
+			Other
+		};
+
+		/**
+		 * Unique resource identifier. Typically a URL.
+		 */
+		FString URI;
+
+		/**
+		 * Possible sub-range within the resource. Will be set when adding to the cache.
+		 * When asking the cache for a hit the `TotalSize` member will not be set (-1) but is
+		 * expected to be filled in when returning the cached element with the value that was
+		 * provided when adding the entry to the cache.
+		 */
+		FRange Range;
+
+		/**
+		 * Type of elementary stream this request is for.
+		 */
+		EStreamType StreamType = EStreamType::Other;
+		
+		/**
+		 * The quality index of this stream. Zero indicates the "worst" (or only) quality.
+		 * If this is the highest quality that is available this value will be the same
+		 * as `MaxQualityIndex`.
+		 */
+		int32 QualityIndex = 0;
+
+		/**
+		 * The maximum quality index of this stream. Zero indicates this is the only available
+		 * quality level.
+		 * This can be used to cache only the highest quality stream, or any other.
+		 */
+		int32 MaxQualityIndex = 0;
+	};
+
+	using FCacheDataPtr = TSharedPtr<TArray<uint8>, ESPMode::ThreadSafe>;
+
+	/**
+	 * Asks to add an item to the cache.
+	 * Whether or not the item gets added is under the purview of the cache implementation.
+	 * As such, this method has no return value to indicate success.
+	 * 
+	 * This method should not block to wait until the cache has been updated.
+	 * It should merely schedule an asynchronous update and defer writing the data to the storage
+	 * to some asynchronous task.
+	 */
+	virtual void AddElementToCache(const FItemInfo& ItemInfoToCache, FCacheDataPtr DataToCache) = 0;
+
+	enum class ECacheResult
+	{
+		// Element not available in the cache.
+		Miss,
+		// Element available in the cache.
+		Hit
+	};
+
+	DECLARE_DELEGATE_ThreeParams(FCachedDataReadCompleted, ECacheResult /*Result*/, FItemInfo /*ItemInfoFromCache*/, FCacheDataPtr /*DataFromCache*/);
+
+	/**
+	 * Asks to get an item from the cache.
+	 * The item to get is described by `InItemToGetFromCache`, with the exception of the `.Range.TotalSize` member
+	 * which is not known ahead of time.
+	 * If the cache fulfills the request the `OutCachedItemInfo` is expected to be set up, including the `.Range.TotalSize` member.
+	 */
+	virtual void GetElementFromCache(const FItemInfo& ItemToGetFromCache, FCachedDataReadCompleted CompletionDelegate) = 0;
+};
diff --git a/Engine/Plugins/Media/ElectraPlayer/Source/ElectraPlayerRuntime/Public/IElectraPlayerInterface.h b/Engine/Plugins/Media/ElectraPlayer/Source/ElectraPlayerRuntime/Public/IElectraPlayerInterface.h
--- a/Engine/Plugins/Media/ElectraPlayer/Source/ElectraPlayerRuntime/Public/IElectraPlayerInterface.h
+++ b/Engine/Plugins/Media/ElectraPlayer/Source/ElectraPlayerRuntime/Public/IElectraPlayerInterface.h
@@ -9,6 +9,7 @@
 
 #include "ElectraPlayerMisc.h"
 #include "ElectraPlayerResourceDelegate.h"
+#include "IElectraPlayerDataCache.h"
 #include "MediaStreamMetadata.h"
 
 class FVideoDecoderOutput;
@@ -182,6 +183,7 @@ public:
 		bool						bDoNotPreload = false;
 		
 		TSharedPtr<IElectraPlayerExternalDataReader, ESPMode::ThreadSafe> ExternalDataReader;
+		TSharedPtr<IElectraPlayerDataCache, ESPMode::ThreadSafe> ExternalDataCache;
 	};
 
 	enum class EOpenType
