ElectraPlayer: added pointer tests

#rnx

diff --git a/Engine/Plugins/Media/ElectraPlayer/Source/ElectraPlayerRuntime/Private/Runtime/Player/AdaptiveStreamingPlayer.cpp b/Engine/Plugins/Media/ElectraPlayer/Source/ElectraPlayerRuntime/Private/Runtime/Player/AdaptiveStreamingPlayer.cpp
--- a/Engine/Plugins/Media/ElectraPlayer/Source/ElectraPlayerRuntime/Private/Runtime/Player/AdaptiveStreamingPlayer.cpp
+++ b/Engine/Plugins/Media/ElectraPlayer/Source/ElectraPlayerRuntime/Private/Runtime/Player/AdaptiveStreamingPlayer.cpp
@@ -65,12 +65,16 @@ FAdaptiveStreamingPlayer::FAdaptiveStreamingPlayer(const IAdaptiveStreamingPlaye
 {
 	Electra::AddActivePlayerInstance();
 
+	if (InCreateParameters.VideoRenderer)
+	{
+		VideoRender.Renderer = CreateWrappedRenderer(InCreateParameters.VideoRenderer, EStreamType::Video);
+	}
+	if (InCreateParameters.AudioRenderer)
+	{
+		AudioRender.Renderer = CreateWrappedRenderer(InCreateParameters.AudioRenderer, EStreamType::Audio);
+	}
 	ExternalPlayerGUID		= InCreateParameters.ExternalPlayerGUID;
-
 	ManifestType			= EMediaFormatType::Unknown;
-
-	VideoRender.Renderer    = CreateWrappedRenderer(InCreateParameters.VideoRenderer, EStreamType::Video);
-	AudioRender.Renderer    = CreateWrappedRenderer(InCreateParameters.AudioRenderer, EStreamType::Audio);
 	CurrentState   		    = EPlayerState::eState_Idle;
 	PipelineState  		    = EPipelineState::ePipeline_Stopped;
 	DecoderState   		    = EDecoderState::eDecoder_Paused;
@@ -1040,11 +1044,17 @@ void FAdaptiveStreamingPlayer::StopRendering()
 int32 FAdaptiveStreamingPlayer::CreateRenderers()
 {
 	// Set the render clock with the renderes.
-	VideoRender.Renderer->SetRenderClock(RenderClock);
-	AudioRender.Renderer->SetRenderClock(RenderClock);
+	if (VideoRender.Renderer)
+	{
+		VideoRender.Renderer->SetRenderClock(RenderClock);
+	}
+	if (AudioRender.Renderer)
+	{
+		AudioRender.Renderer->SetRenderClock(RenderClock);
+	}
 
 	// Hold back all frames during preroll or emit the first frame for scrubbing?
-	if (PlayerOptions.HaveKey(OptionKeyDoNotHoldBackFirstVideoFrame))
+	if (VideoRender.Renderer && PlayerOptions.HaveKey(OptionKeyDoNotHoldBackFirstVideoFrame))
 	{
 		VideoRender.Renderer->DisableHoldbackOfFirstRenderableVideoFrame(PlayerOptions.GetValue(OptionKeyDoNotHoldBackFirstVideoFrame).SafeGetBool(false));
 	}
