ElectraPlayer: Inject DASH timing sources from URL fragment

#rnx

diff --git a/Engine/Plugins/Media/ElectraPlayer/Source/ElectraPlayerRuntime/Private/Runtime/Player/DASH/MPDElementsDASH.h b/Engine/Plugins/Media/ElectraPlayer/Source/ElectraPlayerRuntime/Private/Runtime/Player/DASH/MPDElementsDASH.h
--- a/Engine/Plugins/Media/ElectraPlayer/Source/ElectraPlayerRuntime/Private/Runtime/Player/DASH/MPDElementsDASH.h
+++ b/Engine/Plugins/Media/ElectraPlayer/Source/ElectraPlayerRuntime/Private/Runtime/Player/DASH/MPDElementsDASH.h
@@ -372,6 +372,10 @@ public:
 	void SetCustomElementAndAttributeJSON(const FString& InCustomElementAndAttributeJSON)
 	{ CustomElementAndAttributeJSON = InCustomElementAndAttributeJSON; }
 
+	void SetSchemeIdUri(const FString& InSchemeIdUri) { SchemeIdUri = InSchemeIdUri; }
+	void SetValue(const FString& InValue) { Value = InValue; }
+	void SetID(const FString& InID) { ID = InID; }
+
 private:
 	virtual bool ProcessElement(FManifestParserDASH* Builder, const TCHAR* ElementName, const TCHAR* ElementData, int32 XmlFileLineNumber) override;
 	virtual bool ProcessAttribute(FManifestParserDASH* Builder, const TCHAR* AttributeName, const TCHAR* AttributeValue) override;
@@ -1720,6 +1724,14 @@ public:
 		FScopeLock lock(&UpdateLock);
 		UTCTimings.Remove(TimingElement);
 	}
+	void AddUTCTimingElement(const FString& InData, const FString& InScheme, const FString& InID, const FString& InValue)
+	{
+		FScopeLock lock(&UpdateLock);
+		auto Element = UTCTimings.Emplace_GetRef(MakeSharedTS<FDashMPD_DescriptorType>(TEXT("UTCTiming"), *InData));
+		Element->SetSchemeIdUri(InScheme);
+		Element->SetID(InID);
+		Element->SetValue(InValue);
+	}
 
 	// Methods to manipulate the presentation type.
 	void LockAccess()											{ UpdateLock.Lock(); }
diff --git a/Engine/Plugins/Media/ElectraPlayer/Source/ElectraPlayerRuntime/Private/Runtime/Player/DASH/ManifestBuilderDASH.cpp b/Engine/Plugins/Media/ElectraPlayer/Source/ElectraPlayerRuntime/Private/Runtime/Player/DASH/ManifestBuilderDASH.cpp
--- a/Engine/Plugins/Media/ElectraPlayer/Source/ElectraPlayerRuntime/Private/Runtime/Player/DASH/ManifestBuilderDASH.cpp
+++ b/Engine/Plugins/Media/ElectraPlayer/Source/ElectraPlayerRuntime/Private/Runtime/Player/DASH/ManifestBuilderDASH.cpp
@@ -58,6 +58,8 @@ namespace
 {
 	const TCHAR* const Custom_EpicStaticStart = TEXT("EpicStaticStart");
 	const TCHAR* const Custom_EpicDynamicStart = TEXT("EpicDynamicStart");
+	const TCHAR* const Custom_EpicUTCUrl = TEXT("EpicUTCUrl");
+	const TCHAR* const Custom_EpicUTCNow = TEXT("EpicUTCNow");
 
 	const TCHAR* const XLinkActuateOnLoad = TEXT("onLoad");
 	const TCHAR* const XLinkActuateOnRequest = TEXT("onRequest");
@@ -927,6 +929,25 @@ FErrorDetail FManifestDASHInternal::PrepareRemoteElementLoadRequest(TArray<TWeak
 }
 
 
+void FManifestDASHInternal::InjectEpicTimingSources()
+{
+	const TArray<TSharedPtrTS<FDashMPD_DescriptorType>>& UTCTimings = MPDRoot->GetUTCTimings();
+
+	for(int32 i=0,iMax=URLFragmentComponents.Num(); i<iMax; ++i)
+	{
+		if (URLFragmentComponents[i].Name.Equals(Custom_EpicUTCUrl))
+		{
+			MPDRoot->AddUTCTimingElement(FString(), DASH::Schemes::TimingSources::Scheme_urn_mpeg_dash_utc_httpxsdate2014, FString(), URLFragmentComponents[i].Value);
+			bDidInjectUTCTimingElements = true;
+		}
+		else if (URLFragmentComponents[i].Name.Equals(Custom_EpicUTCNow))
+		{
+			MPDRoot->AddUTCTimingElement(FString(), DASH::Schemes::TimingSources::Scheme_urn_mpeg_dash_utc_direct2014, FString(), URLFragmentComponents[i].Value);
+			bDidInjectUTCTimingElements = true;
+		}
+	}
+}
+
 void FManifestDASHInternal::TransformIntoEpicEvent()
 {
 	// For this to work the presentation must be 'static'
@@ -2421,6 +2442,10 @@ bool FManifestDASHInternal::IsStaticType() const
 	return GetPresentationType() == FManifestDASHInternal::EPresentationType::Static;
 }
 
+bool FManifestDASHInternal::HasInjectedTimingSources() const
+{
+	return bDidInjectUTCTimingElements;
+}
 
 FTimeValue FManifestDASHInternal::GetMinimumUpdatePeriod() const
 {
diff --git a/Engine/Plugins/Media/ElectraPlayer/Source/ElectraPlayerRuntime/Private/Runtime/Player/DASH/ManifestDASH.cpp b/Engine/Plugins/Media/ElectraPlayer/Source/ElectraPlayerRuntime/Private/Runtime/Player/DASH/ManifestDASH.cpp
--- a/Engine/Plugins/Media/ElectraPlayer/Source/ElectraPlayerRuntime/Private/Runtime/Player/DASH/ManifestDASH.cpp
+++ b/Engine/Plugins/Media/ElectraPlayer/Source/ElectraPlayerRuntime/Private/Runtime/Player/DASH/ManifestDASH.cpp
@@ -521,6 +521,7 @@ IManifest::FResult FManifestDASH::FindPlayPeriod(TSharedPtrTS<IPlayPeriod>& OutP
 
 	// Quick out if the time falls outside the presentation.
 	FTimeValue TotalEndTime = Manifest->GetLastPeriodEndTime();
+	TotalEndTime -= Manifest->GetAnchorTime();
 	if (PlayRangeEnd.IsValid() && TotalEndTime.IsValid() && PlayRangeEnd < TotalEndTime)
 	{
 		TotalEndTime = PlayRangeEnd;
@@ -604,6 +605,18 @@ IManifest::FResult FManifestDASH::FindPlayPeriod(TSharedPtrTS<IPlayPeriod>& OutP
 		{
 			return IManifest::FResult(IManifest::FResult::EType::PastEOS);
 		}
+
+		// Check if we would start outside the permitted event range.
+		if (Manifest->IsDynamicEpicEvent())
+		{
+			FTimeValue Now = PlayerSessionServices->GetSynchronizedUTCTime()->GetTime();
+			FTimeValue End = Manifest->GetLastPeriodEndTime();
+			if (Now >= End)
+			{
+				return IManifest::FResult(IManifest::FResult::EType::PastEOS);
+			}
+		}
+
 		// Is the original period still there?
 		TSharedPtrTS<FDashMPD_PeriodType> MPDPeriod = SelectedPeriod->GetMPDPeriod();
 		if (MPDPeriod.IsValid())
diff --git a/Engine/Plugins/Media/ElectraPlayer/Source/ElectraPlayerRuntime/Private/Runtime/Player/DASH/PlaylistReaderDASH.cpp b/Engine/Plugins/Media/ElectraPlayer/Source/ElectraPlayerRuntime/Private/Runtime/Player/DASH/PlaylistReaderDASH.cpp
--- a/Engine/Plugins/Media/ElectraPlayer/Source/ElectraPlayerRuntime/Private/Runtime/Player/DASH/PlaylistReaderDASH.cpp
+++ b/Engine/Plugins/Media/ElectraPlayer/Source/ElectraPlayerRuntime/Private/Runtime/Player/DASH/PlaylistReaderDASH.cpp
@@ -665,7 +665,7 @@ void FPlaylistReaderDASH::TriggerTimeSynchronization()
 	if (Manifest.IsValid())
 	{
 		// Time sync is only necessary when there is either an MPD@availabilityStartTime or MPD@type is dynamic.
-		if (Manifest->UsesAST() || !Manifest->IsStaticType())
+		if (Manifest->UsesAST() || !Manifest->IsStaticType() || Manifest->HasInjectedTimingSources())
 		{
 			TSharedPtrTS<FDashMPD_MPDType> MPDRoot = Manifest->GetMPDRoot();
 			if (MPDRoot.IsValid())
@@ -1183,6 +1183,8 @@ void FPlaylistReaderDASH::ManifestDownloadCompleted(FResourceLoadRequestPtr Requ
 					TArray<FURL_RFC3986::FQueryParam> URLFragmentComponents;
 					FURL_RFC3986::GetQueryParams(URLFragmentComponents, Fragment, false);	// The fragment is already URL escaped, so no need to do it again.
 					NewManifest->SetURLFragmentComponents(MoveTemp(URLFragmentComponents));
+					// Inject fake <UTCTiming> elements that were passed in the fragments.
+					NewManifest->InjectEpicTimingSources();
 				}
 				Manifest = NewManifest;
 
@@ -1260,6 +1262,14 @@ void FPlaylistReaderDASH::FinishManifestBuildingAfterTimesync(bool bGotTheTime)
 				// No need to use the Date header any more.
 				InitialTimesync.HttpDateHeader.Reset();
 			}
+			// If parsing failed then maybe the response is just a number (possibly with frational digits) giving the
+			// current Unix epoch time.
+			else if (UnixEpoch::ParseFloatString(DirectTime, InitialTimesync.Utc_direct2014))
+			{
+				PlayerSessionServices->GetSynchronizedUTCTime()->SetTime(DirectTime + TimeDiffSinceStart);
+				// No need to use the Date header any more.
+				InitialTimesync.HttpDateHeader.Reset();
+			}
 		}
 		// As a last resort use the Date header from the HTTP response
 		if (InitialTimesync.HttpDateHeader.Len())
diff --git a/Engine/Plugins/Media/ElectraPlayer/Source/ElectraPlayerRuntime/Private/Runtime/Player/DASH/PlaylistReaderDASH_Internal.h b/Engine/Plugins/Media/ElectraPlayer/Source/ElectraPlayerRuntime/Private/Runtime/Player/DASH/PlaylistReaderDASH_Internal.h
--- a/Engine/Plugins/Media/ElectraPlayer/Source/ElectraPlayerRuntime/Private/Runtime/Player/DASH/PlaylistReaderDASH_Internal.h
+++ b/Engine/Plugins/Media/ElectraPlayer/Source/ElectraPlayerRuntime/Private/Runtime/Player/DASH/PlaylistReaderDASH_Internal.h
@@ -863,6 +863,7 @@ public:
 		return URLFragmentComponents;
 	}
 
+	void InjectEpicTimingSources();
 	void TransformIntoEpicEvent();
 
 	ELECTRA_IMPL_DEFAULT_ERROR_METHODS(DASHManifest);
@@ -903,6 +904,8 @@ public:
 	bool AreUpdatesExpected() const;
 	bool IsStaticType() const;
 	bool UsesAST() const;
+	bool HasInjectedTimingSources() const;
+
 	FTimeValue GetAvailabilityEndTime() const;
 	FTimeValue GetTimeshiftBufferDepth() const;
 
@@ -951,6 +954,8 @@ private:
 	mutable FTimeRange SeekableTimeRange;
 	FTimeValue DefaultStartTime;
 	mutable bool bWarnedAboutTooSmallSuggestedPresentationDelay = false;
+
+	bool bDidInjectUTCTimingElements = false;
 };
 
 
