ElectraPlayer: Sending a last audio buffer with the 'eos' flag set to indicate that the last output has been processed

#rnx

diff --git a/Engine/Plugins/Media/ElectraPlayer/Source/ElectraPlayerRuntime/Private/Runtime/Decoder/AudioDecoder.cpp b/Engine/Plugins/Media/ElectraPlayer/Source/ElectraPlayerRuntime/Private/Runtime/Decoder/AudioDecoder.cpp
--- a/Engine/Plugins/Media/ElectraPlayer/Source/ElectraPlayerRuntime/Private/Runtime/Decoder/AudioDecoder.cpp
+++ b/Engine/Plugins/Media/ElectraPlayer/Source/ElectraPlayerRuntime/Private/Runtime/Decoder/AudioDecoder.cpp
@@ -116,6 +116,12 @@ private:
 	};
 	FRIEND_ENUM_CLASS_FLAGS(EAUChangeFlags);
 
+	enum class ESendMode
+	{
+		SendSilence,
+		SendEOS
+	};
+
 	void StartThread();
 	void StopThread();
 	void WorkerThread();
@@ -143,7 +149,7 @@ private:
 
 	IElectraDecoder::EOutputStatus HandleOutput();
 	bool HandleDecoding();
-	bool HandleDummyDecoding();
+	bool SendSilenceOrEOS(ESendMode InMode);
 	void StartDraining(EDecodingState InNextStateAfterDraining);
 	bool CheckForFlush();
 	bool CheckBackgrounding();
@@ -866,7 +872,7 @@ bool FAudioDecoderImpl::HandleDecoding()
 
 		if ((bInDummyDecodeMode = CurrentAccessUnit->AccessUnit->bIsDummyData) == true)
 		{
-			bool bOk = HandleDummyDecoding();
+			bool bOk = SendSilenceOrEOS(ESendMode::SendSilence);
 			CurrentAccessUnit.Reset();
 			return bOk;
 		}
@@ -930,9 +936,9 @@ bool FAudioDecoderImpl::HandleDecoding()
 }
 
 
-bool FAudioDecoderImpl::HandleDummyDecoding()
+bool FAudioDecoderImpl::SendSilenceOrEOS(ESendMode InSendMode)
 {
-	check(CurrentAccessUnit.IsValid());
+	check(CurrentAccessUnit.IsValid() || InSendMode == ESendMode::SendEOS);
 	check(bIsDecoderClean);
 	
 	// Get output unless flushing or terminating
@@ -966,33 +972,48 @@ bool FAudioDecoderImpl::HandleDummyDecoding()
 			SCOPE_CYCLE_COUNTER(STAT_ElectraPlayer_AudioConvertOutput);
 			CSV_SCOPED_TIMING_STAT(ElectraPlayer, AudioConvertOutput);
 
-			// Clear to silence
-			int64 MaxBufferSize = CurrentOutputBuffer->GetBufferProperties().GetValue(TEXT("size")).GetInt64();
-			FMemory::Memzero(CurrentOutputBuffer->GetBufferProperties().GetValue(TEXT("address")).GetPointer(), MaxBufferSize);
+			if (InSendMode == ESendMode::SendSilence)
+			{
+				// Clear to silence
+				int64 MaxBufferSize = CurrentOutputBuffer->GetBufferProperties().GetValue(TEXT("size")).GetInt64();
+				FMemory::Memzero(CurrentOutputBuffer->GetBufferProperties().GetValue(TEXT("address")).GetPointer(), MaxBufferSize);
 
-			int32 NumChannels = CurrentOutputFormat.NumChannels ? CurrentOutputFormat.NumChannels : 2;
-			int32 SampleRate = CurrentOutputFormat.SampleRate ? CurrentOutputFormat.SampleRate : 48000;
-			int32 SamplesPerBlock = CurrentAccessUnit->AdjustedDuration.IsValid() ? CurrentAccessUnit->AdjustedDuration.GetAsHNS() * SampleRate / 10000000 : 0;
-			int32 EmptySize = NumChannels * sizeof(float) * SamplesPerBlock;
+				int32 NumChannels = CurrentOutputFormat.NumChannels ? CurrentOutputFormat.NumChannels : 2;
+				int32 SampleRate = CurrentOutputFormat.SampleRate ? CurrentOutputFormat.SampleRate : 48000;
+				int32 SamplesPerBlock = CurrentAccessUnit->AdjustedDuration.IsValid() ? CurrentAccessUnit->AdjustedDuration.GetAsHNS() * SampleRate / 10000000 : 0;
+				int32 EmptySize = NumChannels * sizeof(float) * SamplesPerBlock;
 
-			if (SamplesPerBlock)
-			{
-				FTimeValue Duration;
-				Duration.SetFromND(SamplesPerBlock, (uint32) SampleRate);
+				if (SamplesPerBlock)
+				{
+					FTimeValue Duration;
+					Duration.SetFromND(SamplesPerBlock, (uint32) SampleRate);
 
-				OutputBufferSampleProperties.Clear();
-				OutputBufferSampleProperties.Set(TEXT("num_channels"), FVariantValue((int64)NumChannels));
-				OutputBufferSampleProperties.Set(TEXT("sample_rate"), FVariantValue((int64)SampleRate));
-				OutputBufferSampleProperties.Set(TEXT("byte_size"), FVariantValue((int64)(EmptySize < MaxBufferSize ? EmptySize : MaxBufferSize)));
-				OutputBufferSampleProperties.Set(TEXT("duration"), FVariantValue(Duration));
-				OutputBufferSampleProperties.Set(TEXT("pts"), FVariantValue(CurrentAccessUnit->AdjustedPTS));
+					OutputBufferSampleProperties.Clear();
+					OutputBufferSampleProperties.Set(TEXT("num_channels"), FVariantValue((int64)NumChannels));
+					OutputBufferSampleProperties.Set(TEXT("sample_rate"), FVariantValue((int64)SampleRate));
+					OutputBufferSampleProperties.Set(TEXT("byte_size"), FVariantValue((int64)(EmptySize < MaxBufferSize ? EmptySize : MaxBufferSize)));
+					OutputBufferSampleProperties.Set(TEXT("duration"), FVariantValue(Duration));
+					OutputBufferSampleProperties.Set(TEXT("pts"), FVariantValue(CurrentAccessUnit->AdjustedPTS));
 
-				Renderer->ReturnBuffer(CurrentOutputBuffer, true, OutputBufferSampleProperties);
-				CurrentOutputBuffer = nullptr;
+					Renderer->ReturnBuffer(CurrentOutputBuffer, true, OutputBufferSampleProperties);
+					CurrentOutputBuffer = nullptr;
+				}
+				else
+				{
+					ReturnUnusedOutputBuffer();
+				}
 			}
 			else
 			{
-				ReturnUnusedOutputBuffer();
+				OutputBufferSampleProperties.Clear();
+				OutputBufferSampleProperties.Set(TEXT("num_channels"), FVariantValue((int64)0));
+				OutputBufferSampleProperties.Set(TEXT("byte_size"), FVariantValue((int64)0));
+				OutputBufferSampleProperties.Set(TEXT("sample_rate"), FVariantValue((int64) 0));
+				OutputBufferSampleProperties.Set(TEXT("duration"), FVariantValue(FTimeValue::GetZero()));
+				OutputBufferSampleProperties.Set(TEXT("pts"), FVariantValue(FTimeValue::GetInvalid()));
+				OutputBufferSampleProperties.Set(TEXT("eos"), FVariantValue(true));
+				Renderer->ReturnBuffer(CurrentOutputBuffer, false, OutputBufferSampleProperties);
+				CurrentOutputBuffer = nullptr;
 			}
 			return true;
 		}
@@ -1157,6 +1178,11 @@ void FAudioDecoderImpl::WorkerThread()
 				else if (OS == IElectraDecoder::EOutputStatus::EndOfData ||
 						 OS == IElectraDecoder::EOutputStatus::NeedInput)
 				{
+					if (OS == IElectraDecoder::EOutputStatus::EndOfData)
+					{
+						SendSilenceOrEOS(ESendMode::SendEOS);
+					}
+
 					// All output has been retrieved
 					InDecoderInput.Empty();
 					ChannelMapper.Reset();
@@ -1199,7 +1225,13 @@ void FAudioDecoderImpl::WorkerThread()
 				// Is the buffer at EOD?
 				if (NextAccessUnits.ReachedEOD())
 				{
-					if (!bIsDecoderClean)
+					if (bInDummyDecodeMode)
+					{
+						SendSilenceOrEOS(ESendMode::SendEOS);
+						CurrentDecodingState = EDecodingState::NormalDecoding;
+						bInDummyDecodeMode = false;
+					}
+					else if (!bIsDecoderClean)
 					{
 						StartDraining(EDecodingState::NormalDecoding);
 					}
