ElectraPlayer: added pointer checks for safety

#rnx

diff --git a/Engine/Plugins/Media/ElectraPlayer/Source/ElectraPlayerRuntime/Private/Runtime/Player/DASH/StreamReaderDASH.cpp b/Engine/Plugins/Media/ElectraPlayer/Source/ElectraPlayerRuntime/Private/Runtime/Player/DASH/StreamReaderDASH.cpp
--- a/Engine/Plugins/Media/ElectraPlayer/Source/ElectraPlayerRuntime/Private/Runtime/Player/DASH/StreamReaderDASH.cpp
+++ b/Engine/Plugins/Media/ElectraPlayer/Source/ElectraPlayerRuntime/Private/Runtime/Player/DASH/StreamReaderDASH.cpp
@@ -350,27 +350,29 @@ void FStreamReaderDASH::FStreamHandler::WorkerThread()
 	LLM_SCOPE(ELLMTag::ElectraPlayer);
 
 	StreamSelector = PlayerSessionService->GetStreamSelector();
-	check(StreamSelector.IsValid());
-	while(!bTerminate)
+	if (StreamSelector.IsValid())
 	{
-		WorkSignal.Obtain();
-		if (!bTerminate)
+		while(!bTerminate)
 		{
-			if (CurrentRequest.IsValid())
+			WorkSignal.Obtain();
+			if (!bTerminate)
 			{
-				IsIdleSignal.Reset();
-				if (!bRequestCanceled)
-				{
-					HandleRequest();
-				}
-				else
+				if (CurrentRequest.IsValid())
 				{
-					CurrentRequest.Reset();
+					IsIdleSignal.Reset();
+					if (!bRequestCanceled)
+					{
+						HandleRequest();
+					}
+					else
+					{
+						CurrentRequest.Reset();
+					}
+					IsIdleSignal.Signal();
 				}
-				IsIdleSignal.Signal();
+				bRequestCanceled = false;
+				bSilentCancellation = false;
 			}
-			bRequestCanceled = false;
-			bSilentCancellation = false;
 		}
 	}
 	StreamSelector.Reset();
@@ -381,23 +383,25 @@ void FStreamReaderDASH::FStreamHandler::RunInThreadPool()
 	LLM_SCOPE(ELLMTag::ElectraPlayer);
 
 	StreamSelector = PlayerSessionService->GetStreamSelector();
-	check(StreamSelector.IsValid());
-	WorkSignal.Obtain();
-	if (!bTerminate)
+	if (StreamSelector.IsValid())
 	{
-		if (CurrentRequest.IsValid())
+		WorkSignal.Obtain();
+		if (!bTerminate)
 		{
-			if (!bRequestCanceled)
-			{
-				HandleRequest();
-			}
-			else
+			if (CurrentRequest.IsValid())
 			{
-				CurrentRequest.Reset();
+				if (!bRequestCanceled)
+				{
+					HandleRequest();
+				}
+				else
+				{
+					CurrentRequest.Reset();
+				}
 			}
+			bRequestCanceled = false;
+			bSilentCancellation = false;
 		}
-		bRequestCanceled = false;
-		bSilentCancellation = false;
 	}
 	StreamSelector.Reset();
 	IsIdleSignal.Signal();
@@ -541,9 +545,15 @@ FErrorDetail FStreamReaderDASH::FStreamHandler::GetInitSegment(TSharedPtrTS<cons
 		// No init segment required. We're done.
 		return FErrorDetail();
 	}
+	// Get the entity cache. If it's not there we return.
+	TSharedPtrTS<IPlayerEntityCache> EntityCache = PlayerSessionService ? PlayerSessionService->GetEntityCache() : nullptr;
+	if (!EntityCache.IsValid())
+	{
+		return FErrorDetail();
+	}
 	// Check if we already have this cached.
 	IPlayerEntityCache::FCacheItem CachedItem;
-	if (PlayerSessionService->GetEntityCache()->GetCachedEntity(CachedItem, Request->Segment.InitializationURL.URL, Request->Segment.InitializationURL.Range))
+	if (EntityCache->GetCachedEntity(CachedItem, Request->Segment.InitializationURL.URL, Request->Segment.InitializationURL.Range))
 	{
 		// Already cached. Use it.
 		OutMP4InitSegment = CachedItem.Parsed14496_12Data;
@@ -563,7 +573,7 @@ FErrorDetail FStreamReaderDASH::FStreamHandler::GetInitSegment(TSharedPtrTS<cons
 	FMP4StaticDataReader StaticDataReader;
 	StaticDataReader.SetParseData(LoadReq->Request->GetResponseBuffer());
 	TSharedPtrTS<IParserISO14496_12> Init = IParserISO14496_12::CreateParser();
-	UEMediaError parseError = Init->ParseHeader(&StaticDataReader, this, PlayerSessionService, nullptr);
+	UEMediaError parseError = Init.IsValid() ? Init->ParseHeader(&StaticDataReader, this, PlayerSessionService, nullptr) : UEMEDIA_ERROR_OOM;
 	if (parseError == UEMEDIA_ERROR_OK || parseError == UEMEDIA_ERROR_END_OF_STREAM)
 	{
 		// Parse the tracks of the init segment. We do this mainly to get to the CSD we might need should we have to insert filler data later.
@@ -575,7 +585,7 @@ FErrorDetail FStreamReaderDASH::FStreamHandler::GetInitSegment(TSharedPtrTS<cons
 			CacheItem.URL = LoadReq->URL;
 			CacheItem.Range = LoadReq->Range;
 			CacheItem.Parsed14496_12Data = Init;
-			PlayerSessionService->GetEntityCache()->CacheEntity(CacheItem);
+			EntityCache->CacheEntity(CacheItem);
 			OutMP4InitSegment = Init;
 
 			StreamSelector->ReportDownloadEnd(ds);
@@ -605,9 +615,15 @@ FErrorDetail FStreamReaderDASH::FStreamHandler::GetInitSegment(TSharedPtrTS<cons
 		// No init segment required. We're done.
 		return FErrorDetail();
 	}
+	// Get the entity cache. If it's not there we return.
+	TSharedPtrTS<IPlayerEntityCache> EntityCache = PlayerSessionService ? PlayerSessionService->GetEntityCache() : nullptr;
+	if (!EntityCache.IsValid())
+	{
+		return FErrorDetail();
+	}
 	// Check if we already have this cached.
 	IPlayerEntityCache::FCacheItem CachedItem;
-	if (PlayerSessionService->GetEntityCache()->GetCachedEntity(CachedItem, Request->Segment.InitializationURL.URL, Request->Segment.InitializationURL.Range))
+	if (EntityCache->GetCachedEntity(CachedItem, Request->Segment.InitializationURL.URL, Request->Segment.InitializationURL.Range))
 	{
 		// Already cached. Use it.
 		OutMKVInitSegment = CachedItem.ParsedMatroskaData;
@@ -627,49 +643,64 @@ FErrorDetail FStreamReaderDASH::FStreamHandler::GetInitSegment(TSharedPtrTS<cons
 	FMKVStaticDataReader StaticDataReader;
 	StaticDataReader.SetParseData(LoadReq->Request->GetResponseBuffer());
 	TSharedPtrTS<IParserMKV> Init = IParserMKV::CreateParser(nullptr);
-	FErrorDetail parseError = Init->ParseHeader(&StaticDataReader, static_cast<Electra::IParserMKV::EParserFlags>(IParserMKV::EParserFlags::ParseFlag_OnlyTracks | IParserMKV::EParserFlags::ParseFlag_SuppressCueWarning));
-	if (parseError.IsOK())
+	if (Init.IsValid())
 	{
-		parseError = Init->PrepareTracks();
+		FErrorDetail parseError = Init->ParseHeader(&StaticDataReader, static_cast<Electra::IParserMKV::EParserFlags>(IParserMKV::EParserFlags::ParseFlag_OnlyTracks | IParserMKV::EParserFlags::ParseFlag_SuppressCueWarning));
 		if (parseError.IsOK())
 		{
-			// Add this to the entity cache in case it needs to be retrieved again.
-			IPlayerEntityCache::FCacheItem CacheItem;
-			CacheItem.URL = LoadReq->URL;
-			CacheItem.Range = LoadReq->Range;
-			CacheItem.ParsedMatroskaData = Init;
-			PlayerSessionService->GetEntityCache()->CacheEntity(CacheItem);
-			OutMKVInitSegment = Init;
-
-			StreamSelector->ReportDownloadEnd(ds);
-			return FErrorDetail();
+			parseError = Init->PrepareTracks();
+			if (parseError.IsOK())
+			{
+				// Add this to the entity cache in case it needs to be retrieved again.
+				IPlayerEntityCache::FCacheItem CacheItem;
+				CacheItem.URL = LoadReq->URL;
+				CacheItem.Range = LoadReq->Range;
+				CacheItem.ParsedMatroskaData = Init;
+				EntityCache->CacheEntity(CacheItem);
+				OutMKVInitSegment = Init;
+
+				StreamSelector->ReportDownloadEnd(ds);
+				return FErrorDetail();
+			}
+			else
+			{
+				ds.bParseFailure = true;
+				Request->ConnectionInfo = *LoadReq->GetConnectionInfo();
+				StreamSelector->ReportDownloadEnd(ds);
+				return CreateError(FString::Printf(TEXT("Track preparation of init segment \"%s\" failed. %d"), *LoadReq->URL, *parseError.GetMessage()), INTERNAL_ERROR_INIT_SEGMENT_PARSE_ERROR);
+			}
 		}
 		else
 		{
 			ds.bParseFailure = true;
-			Request->ConnectionInfo = *LoadReq->GetConnectionInfo();
 			StreamSelector->ReportDownloadEnd(ds);
-			return CreateError(FString::Printf(TEXT("Track preparation of init segment \"%s\" failed. %d"), *LoadReq->URL, *parseError.GetMessage()), INTERNAL_ERROR_INIT_SEGMENT_PARSE_ERROR);
+			return CreateError(FString::Printf(TEXT("Parse error of init segment \"%s\". %s"), *LoadReq->URL, *parseError.GetMessage()), INTERNAL_ERROR_INIT_SEGMENT_PARSE_ERROR);
 		}
 	}
 	else
 	{
 		ds.bParseFailure = true;
 		StreamSelector->ReportDownloadEnd(ds);
-		return CreateError(FString::Printf(TEXT("Parse error of init segment \"%s\". %s"), *LoadReq->URL, *parseError.GetMessage()), INTERNAL_ERROR_INIT_SEGMENT_PARSE_ERROR);
+		return CreateError(FString::Printf(TEXT("Parse error of init segment \"%s\". Could not create MKV parser."), *LoadReq->URL), INTERNAL_ERROR_INIT_SEGMENT_PARSE_ERROR);
 	}
 }
 
+
 FErrorDetail FStreamReaderDASH::FStreamHandler::RetrieveSideloadedFile(TSharedPtrTS<const TArray<uint8>>& OutData, const TSharedPtrTS<FStreamSegmentRequestDASH>& Request)
 {
 	if (Request->Segment.MediaURL.URL.IsEmpty())
 	{
 		return FErrorDetail();
 	}
-
+	// Get the entity cache. If it's not there we return.
+	TSharedPtrTS<IPlayerEntityCache> EntityCache = PlayerSessionService ? PlayerSessionService->GetEntityCache() : nullptr;
+	if (!EntityCache.IsValid())
+	{
+		return FErrorDetail();
+	}
 	// Check if we already have this cached.
 	IPlayerEntityCache::FCacheItem CachedItem;
-	if (PlayerSessionService->GetEntityCache()->GetCachedEntity(CachedItem, Request->Segment.MediaURL.URL, Request->Segment.MediaURL.Range))
+	if (EntityCache->GetCachedEntity(CachedItem, Request->Segment.MediaURL.URL, Request->Segment.MediaURL.Range))
 	{
 		// Already cached. Use it.
 		OutData = CachedItem.RawPayloadData;
@@ -758,7 +789,7 @@ FErrorDetail FStreamReaderDASH::FStreamHandler::RetrieveSideloadedFile(TSharedPt
 	CacheItem.URL = LoadReq->URL;
 	CacheItem.Range = LoadReq->Range;
 	CacheItem.RawPayloadData = OutData;
-	PlayerSessionService->GetEntityCache()->CacheEntity(CacheItem);
+	EntityCache->CacheEntity(CacheItem);
 
 	StreamSelector->ReportDownloadEnd(ds);
 	return FErrorDetail();
