ElectraPlayer: additional pointer check

#rnx

diff --git a/Engine/Plugins/Media/ElectraPlayer/Source/ElectraPlayerRuntime/Private/Runtime/Player/DASH/StreamReaderDASH.cpp b/Engine/Plugins/Media/ElectraPlayer/Source/ElectraPlayerRuntime/Private/Runtime/Player/DASH/StreamReaderDASH.cpp
--- a/Engine/Plugins/Media/ElectraPlayer/Source/ElectraPlayerRuntime/Private/Runtime/Player/DASH/StreamReaderDASH.cpp
+++ b/Engine/Plugins/Media/ElectraPlayer/Source/ElectraPlayerRuntime/Private/Runtime/Player/DASH/StreamReaderDASH.cpp
@@ -570,10 +570,14 @@ FErrorDetail FStreamReaderDASH::FStreamHandler::GetInitSegment(TSharedPtrTS<cons
 	SCOPE_CYCLE_COUNTER(STAT_ElectraPlayer_DASH_StreamReader);
 	CSV_SCOPED_TIMING_STAT(ElectraPlayer, DASH_StreamReader);
 
-	FMP4StaticDataReader StaticDataReader;
-	StaticDataReader.SetParseData(LoadReq->Request->GetResponseBuffer());
 	TSharedPtrTS<IParserISO14496_12> Init = IParserISO14496_12::CreateParser();
-	UEMediaError parseError = Init.IsValid() ? Init->ParseHeader(&StaticDataReader, this, PlayerSessionService, nullptr) : UEMEDIA_ERROR_OOM;
+	UEMediaError parseError = UEMEDIA_ERROR_FORMAT_ERROR;
+	if (LoadReq->Request.IsValid() && LoadReq->Request->GetResponseBuffer().IsValid())
+	{
+		FMP4StaticDataReader StaticDataReader;
+		StaticDataReader.SetParseData(LoadReq->Request->GetResponseBuffer());
+		parseError = Init->ParseHeader(&StaticDataReader, this, PlayerSessionService, nullptr);
+	}
 	if (parseError == UEMEDIA_ERROR_OK || parseError == UEMEDIA_ERROR_END_OF_STREAM)
 	{
 		// Parse the tracks of the init segment. We do this mainly to get to the CSD we might need should we have to insert filler data later.
@@ -640,11 +644,11 @@ FErrorDetail FStreamReaderDASH::FStreamHandler::GetInitSegment(TSharedPtrTS<cons
 	SCOPE_CYCLE_COUNTER(STAT_ElectraPlayer_DASH_StreamReader);
 	CSV_SCOPED_TIMING_STAT(ElectraPlayer, DASH_StreamReader);
 
-	FMKVStaticDataReader StaticDataReader;
-	StaticDataReader.SetParseData(LoadReq->Request->GetResponseBuffer());
 	TSharedPtrTS<IParserMKV> Init = IParserMKV::CreateParser(nullptr);
-	if (Init.IsValid())
+	if (LoadReq->Request.IsValid() && LoadReq->Request->GetResponseBuffer().IsValid())
 	{
+		FMKVStaticDataReader StaticDataReader;
+		StaticDataReader.SetParseData(LoadReq->Request->GetResponseBuffer());
 		FErrorDetail parseError = Init->ParseHeader(&StaticDataReader, static_cast<Electra::IParserMKV::EParserFlags>(IParserMKV::EParserFlags::ParseFlag_OnlyTracks | IParserMKV::EParserFlags::ParseFlag_SuppressCueWarning));
 		if (parseError.IsOK())
 		{
