ElectraDecoder: Configuring Opus decoder to use the same rate as was used in the encoding; changed header version check to match the definition in RFC 7845 exactly.

#rnx

diff --git a/Engine/Plugins/Media/ElectraPlayer/Source/ElectraPlayerRuntime/Private/Runtime/Demuxer/ParserMKV.cpp b/Engine/Plugins/Media/ElectraPlayer/Source/ElectraPlayerRuntime/Private/Runtime/Demuxer/ParserMKV.cpp
--- a/Engine/Plugins/Media/ElectraPlayer/Source/ElectraPlayerRuntime/Private/Runtime/Demuxer/ParserMKV.cpp
+++ b/Engine/Plugins/Media/ElectraPlayer/Source/ElectraPlayerRuntime/Private/Runtime/Demuxer/ParserMKV.cpp
@@ -3168,7 +3168,7 @@ namespace Electra
 				return false;
 			}
 			// Check minimum header size, version must be 1 and mapping family either 0, 1 or 255
-			if (OpusHead.Num() < 19 || OpusHead[8]!=1 || (OpusHead[18]!=0 && OpusHead[18]!=1 && OpusHead[18]!=255))
+			if (OpusHead.Num() < 19 || OpusHead[8]>15 || (OpusHead[18]!=0 && OpusHead[18]!=1 && OpusHead[18]!=255))
 			{
 				return false;
 			}
@@ -3202,6 +3202,7 @@ namespace Electra
 			OutCodecInformation.SetCodecSpecifierRFC6381(TEXT("Opus"));
 			OutCodecInformation.SetSamplingRate(Audio->GetOutputSampleRate());
 			OutCodecInformation.SetNumberOfChannels(Audio->GetNumberOfChannels());
+#if 0
 			// If there is no default duration set we try to calculate it from the sample rate.
 			const int32 NumDecodedSamplesPerBlock = 960;
 			if (InFromTrack->GetDefaultDurationNanos() == 0)
@@ -3217,6 +3218,22 @@ namespace Electra
 					}
 				}
 			}
+#else
+			// If there is no default duration set we assume the encoded frame size was 20ms
+			if (InFromTrack->GetDefaultDurationNanos() == 0)
+			{
+				FTimeFraction fr(20, 1000);
+				if (fr.IsValid())
+				{
+					int64 nanos = fr.GetAsTimebase(1000000000);
+					check(nanos >= 0);
+					if (nanos >= 0)
+					{
+						InFromTrack->SetDefaultDurationNanos((uint64)nanos);
+					}
+				}
+			}
+#endif
 			return true;
 		}
 		return false;
diff --git a/Engine/Plugins/Media/OpusDecoderElectra/Source/OpusDecoderElectra/Private/OpusDecoder/ElectraMediaOpusDecoder.cpp b/Engine/Plugins/Media/OpusDecoderElectra/Source/OpusDecoderElectra/Private/OpusDecoder/ElectraMediaOpusDecoder.cpp
--- a/Engine/Plugins/Media/OpusDecoderElectra/Source/OpusDecoderElectra/Private/OpusDecoder/ElectraMediaOpusDecoder.cpp
+++ b/Engine/Plugins/Media/OpusDecoderElectra/Source/OpusDecoderElectra/Private/OpusDecoder/ElectraMediaOpusDecoder.cpp
@@ -196,10 +196,9 @@ private:
 		return (static_cast<uint32>(A) << 24) | (static_cast<uint32>(B) << 16) | (static_cast<uint32>(C) << 8) | static_cast<uint32>(D);
 	}
 
-	static constexpr int32 OpusSamplingRate()
+	int32 OpusSamplingRate() const
 	{
-		// Opus internal rate is fixed at 48kHz, no matter what the source rate has been.
-		return 48000;
+		return DOpsConfig.SampleRate;
 	}
 
 	bool Parse_dOps(FDOpsConfig& OutConfig, const TArray<uint8>& IndOpsBox, bool bFailOnError);
