ElectraPlayer: fixed calling a wrapper method instead of the actual method to check for read abort. no functional difference since the wrapper calls the real method, but that could change in the future.

#rnx

diff --git a/Engine/Plugins/Media/ElectraPlayer/Source/ElectraPlayerRuntime/Private/Runtime/Player/DASH/StreamReaderDASH.cpp b/Engine/Plugins/Media/ElectraPlayer/Source/ElectraPlayerRuntime/Private/Runtime/Player/DASH/StreamReaderDASH.cpp
--- a/Engine/Plugins/Media/ElectraPlayer/Source/ElectraPlayerRuntime/Private/Runtime/Player/DASH/StreamReaderDASH.cpp
+++ b/Engine/Plugins/Media/ElectraPlayer/Source/ElectraPlayerRuntime/Private/Runtime/Player/DASH/StreamReaderDASH.cpp
@@ -2419,7 +2419,7 @@ FStreamReaderDASH::FStreamHandler::EEmitResult FStreamReaderDASH::FStreamHandler
 		ActiveTrackData.SortedAccessUnitFIFO.Empty();
 	}
 
-	while(ActiveTrackData.AccessUnitFIFO.Num() && !bTerminate && !MKVHasReadBeenAborted())
+	while(ActiveTrackData.AccessUnitFIFO.Num() && !bTerminate && !HasReadBeenAborted())
 	{
 		if (ActiveTrackData.bNeedToRecalculateDurations)
 		{
@@ -2514,7 +2514,7 @@ FStreamReaderDASH::FStreamHandler::EEmitResult FStreamReaderDASH::FStreamHandler
 			ActiveTrackData.AccessUnitFIFO[LastIndex].AU->bIsLastInPeriod = true;
 		}
 
-		while(!bTerminate && !MKVHasReadBeenAborted())
+		while(!bTerminate && !HasReadBeenAborted())
 		{
 			if (Parameters.EventListener->OnFragmentAccessUnitReceived(pNext))
 			{
