ElectraPlayer: allowing for 2 retries of audio-only streams playing on the only or lowest quality level

#rnx

diff --git a/Engine/Plugins/Media/ElectraPlayer/Source/ElectraPlayerRuntime/Private/Runtime/Player/ABRRules/ABRVoDPlus.cpp b/Engine/Plugins/Media/ElectraPlayer/Source/ElectraPlayerRuntime/Private/Runtime/Player/ABRRules/ABRVoDPlus.cpp
--- a/Engine/Plugins/Media/ElectraPlayer/Source/ElectraPlayerRuntime/Private/Runtime/Player/ABRRules/ABRVoDPlus.cpp
+++ b/Engine/Plugins/Media/ElectraPlayer/Source/ElectraPlayerRuntime/Private/Runtime/Player/ABRRules/ABRVoDPlus.cpp
@@ -726,9 +726,14 @@ IAdaptiveStreamSelector::ESegmentAction FABROnDemandPlus::EvaluateForError(TArra
 										CurrentStreamInfo->Health.BecomesAvailableAgainAtUTC = TimeNow + FTimeValue().SetFromMilliseconds(1000);
 									}
 								}
+								else if (IsAudioOnly() && Stats.RetryNumber < 2)
+								{
+									OutDelay.SetFromMilliseconds(400);
+									bRetryIfPossible = true;
+								}
 								else
 								{
-									// Need to insert filler data instead.
+									// Insert filler data instead.
 									bRetryIfPossible = false;
 								}
 							}
