ElectraPlayer: added media option query to enable local caching and scrubbing optimizations

#rnx

diff --git a/Engine/Plugins/Media/ElectraPlayer/Source/ElectraPlayerPlugin/Private/ElectraPlayerPlugin.cpp b/Engine/Plugins/Media/ElectraPlayer/Source/ElectraPlayerPlugin/Private/ElectraPlayerPlugin.cpp
--- a/Engine/Plugins/Media/ElectraPlayer/Source/ElectraPlayerPlugin/Private/ElectraPlayerPlugin.cpp
+++ b/Engine/Plugins/Media/ElectraPlayer/Source/ElectraPlayerPlugin/Private/ElectraPlayerPlugin.cpp
@@ -677,7 +677,24 @@ bool FElectraPlayerPlugin::Open(const FString& Url, const IMediaOptions* Options
 			UE_LOG(LogElectraPlayerPlugin, Log, TEXT("[%p] IMediaPlayer::Open: CDN HTTP status %d will deny a stream permanently"), this, HTTPStatus);
 		}
 	}
-
+	// Check if there are options to tweak for scrubbing
+	bool bUseScrubOptimizations = Options->GetMediaOption(TEXT("ElectraScrubOptimization"), (bool)false);
+	if (bUseScrubOptimizations)
+	{
+		int64 ScrubSeekBitrate = Options->GetMediaOption(TEXT("ElectraScrubBitrate"), (int64)-1);
+		int64 ScrubCacheSizeKiB = Options->GetMediaOption(TEXT("ElectraScrubCacheSizeKiB"), (int64)-1);
+		if (ScrubSeekBitrate >= 0)
+		{
+			PlayerOptions.SetOrUpdate("seekstart_bitrate", Electra::FVariantValue(ScrubSeekBitrate));
+		}
+		if (ScrubCacheSizeKiB > 0)
+		{
+			PlayerOptions.SetOrUpdate("httpcache_max_bytesize", Electra::FVariantValue(ScrubCacheSizeKiB << 10));
+			PlayerOptions.SetOrUpdate("httpcache_max_entries", Electra::FVariantValue((int64)10000));
+		}
+		PlayerOptions.SetOrUpdate("optimize_seek_for_scrubbing", Electra::FVariantValue(true));
+		PlayerOptions.SetOrUpdate("do_not_hold_back_first_frame", Electra::FVariantValue(true));
+	}
 
 	// Check for options that can be changed during playback and apply them at startup already.
 	// If a media source supports the MaxResolutionForMediaStreaming option then we can override the max resolution.
