ElectraPlayer: fixed issue with last access units not being sent into the player buffers when the asset is truncated by the end of a DASH period

#rnx

diff --git a/Engine/Plugins/Media/ElectraPlayer/Source/ElectraPlayerRuntime/Private/Runtime/Player/DASH/StreamReaderDASH.cpp b/Engine/Plugins/Media/ElectraPlayer/Source/ElectraPlayerRuntime/Private/Runtime/Player/DASH/StreamReaderDASH.cpp
--- a/Engine/Plugins/Media/ElectraPlayer/Source/ElectraPlayerRuntime/Private/Runtime/Player/DASH/StreamReaderDASH.cpp
+++ b/Engine/Plugins/Media/ElectraPlayer/Source/ElectraPlayerRuntime/Private/Runtime/Player/DASH/StreamReaderDASH.cpp
@@ -1693,6 +1693,8 @@ void FStreamReaderDASH::FStreamHandler::HandleRequestMP4()
 	ActiveTrackData.bGotAllSamples = true;
 	// Emit all remaining pending AUs
 	EmitSamples(EEmitType::AllRemaining, Request);
+	ActiveTrackData.AccessUnitFIFO.Empty();
+	ActiveTrackData.SortedAccessUnitFIFO.Empty();
 
 	// If we did not set a specific failure message yet set the one from the downloader.
 	if (ds.FailureReason.Len() == 0)
@@ -2304,6 +2306,8 @@ void FStreamReaderDASH::FStreamHandler::HandleRequestMKV()
 	ActiveTrackData.bGotAllSamples = true;
 	// Emit all remaining pending AUs
 	EmitSamples(EEmitType::AllRemaining, Request);
+	ActiveTrackData.AccessUnitFIFO.Empty();
+	ActiveTrackData.SortedAccessUnitFIFO.Empty();
 
 	// If we did not set a specific failure message yet set the one from the downloader.
 	if (ds.FailureReason.Len() == 0)
@@ -2413,12 +2417,6 @@ FStreamReaderDASH::FStreamHandler::EEmitResult FStreamReaderDASH::FStreamHandler
 {
 	EEmitResult Result = EEmitResult::SentNothing;
 
-	if (ActiveTrackData.bReadPastLastPTS)
-	{
-		ActiveTrackData.AccessUnitFIFO.Empty();
-		ActiveTrackData.SortedAccessUnitFIFO.Empty();
-	}
-
 	while(ActiveTrackData.AccessUnitFIFO.Num() && !bTerminate && !HasReadBeenAborted())
 	{
 		if (ActiveTrackData.bNeedToRecalculateDurations)
