ElectraPlayer: Fixed incorrect tagging of last period access unit

#rnx

diff --git a/Engine/Plugins/Media/ElectraPlayer/Source/ElectraPlayerRuntime/Private/Runtime/Player/AdaptiveStreamingPlayerABR.cpp b/Engine/Plugins/Media/ElectraPlayer/Source/ElectraPlayerRuntime/Private/Runtime/Player/AdaptiveStreamingPlayerABR.cpp
--- a/Engine/Plugins/Media/ElectraPlayer/Source/ElectraPlayerRuntime/Private/Runtime/Player/AdaptiveStreamingPlayerABR.cpp
+++ b/Engine/Plugins/Media/ElectraPlayer/Source/ElectraPlayerRuntime/Private/Runtime/Player/AdaptiveStreamingPlayerABR.cpp
@@ -607,16 +607,16 @@ namespace Electra
 		}
 		else
 		{
-			// Initial request. The play period is the same for all streams at this point.
+			// Initial request.
 			ESegmentAction ActionV = ESegmentAction::FetchNext;
 			ESegmentAction ActionA = ESegmentAction::FetchNext;
 			if (StreamInformationVideo.Num())
 			{
-				ActionV = SelectSuitableStreamByType(OutDelay, CurrentSegment, CurrentPlayPeriod, EStreamType::Video);
+				ActionV = SelectSuitableStreamByType(OutDelay, CurrentSegment, CurrentPlayPeriodVideo, EStreamType::Video);
 			}
 			if (StreamInformationAudio.Num())
 			{
-				ActionA = SelectSuitableStreamByType(OutDelay, CurrentSegment, CurrentPlayPeriod, EStreamType::Audio);
+				ActionA = SelectSuitableStreamByType(OutDelay, CurrentSegment, CurrentPlayPeriodAudio, EStreamType::Audio);
 			}
 			if (ActionV == IAdaptiveStreamSelector::ESegmentAction::Fail || ActionA == IAdaptiveStreamSelector::ESegmentAction::Fail)
 			{
diff --git a/Engine/Plugins/Media/ElectraPlayer/Source/ElectraPlayerRuntime/Private/Runtime/Player/DASH/StreamReaderDASH.cpp b/Engine/Plugins/Media/ElectraPlayer/Source/ElectraPlayerRuntime/Private/Runtime/Player/DASH/StreamReaderDASH.cpp
--- a/Engine/Plugins/Media/ElectraPlayer/Source/ElectraPlayerRuntime/Private/Runtime/Player/DASH/StreamReaderDASH.cpp
+++ b/Engine/Plugins/Media/ElectraPlayer/Source/ElectraPlayerRuntime/Private/Runtime/Player/DASH/StreamReaderDASH.cpp
@@ -2486,30 +2486,55 @@ FStreamReaderDASH::FStreamHandler::EEmitResult FStreamReaderDASH::FStreamHandler
 
 		FAccessUnit* pNext = ActiveTrackData.AccessUnitFIFO[0].AU;
 		// Check if this is the last access unit in the requested time range.
-		if (pNext->PTS >= pNext->LatestPTS)
+		if (!ActiveTrackData.bTaggedLastSample && pNext->PTS >= pNext->LatestPTS)
 		{
-			// Because of B frames the last frame that must be decoded could actually be
-			// a later frame in decode order.
-			// Suppose the sequence IPBB with timestamps 0,3,1,2 respectively. Even though the
-			// P frame with timestamp 3 is "the last" one, it will enter the decoder before the B frames.
-			// As such we need to tag the last B frame as "the last one" even though its timstamp
-			// is before the last time requested.
-			// Note: This is not necessary to do for audio frames, but the logic is the same
-			//       so we do not differentiate here.
-			FTimeValue HighestPTSBelow(-1.0);
-			int32 LastIndex = 0;
-			for(int32 i=1; i<ActiveTrackData.AccessUnitFIFO.Num(); ++i)
+			/*
+				Because of B frames the last frame that must be decoded could actually be
+				a later frame in decode order.
+				Suppose the sequence IPBB with timestamps 0,3,1,2 respectively. Even though the
+				P frame with timestamp 3 is "the last" one in presentation order, it will enter
+				the decoder before the B frames.
+				As such we need to tag the last B frame (2) as "the last one" even though its timestamp
+				is before the last time requested.
+				This would be easy if we had access to reliable DTS, but Matroska files only provide PTS.
+				Note: This may seem superfluous since we are tagging as "last" which happens to be the
+				      actual last element in the list, but there could really be even later frames in
+				      the list that we will then remove to avoid sending frames into the decoder that
+				      will be discarded after decoding, which is a waste of decode cycles.
+			*/
+
+			const FTimeValue NextPTS(pNext->PTS);
+			// Sort the remaining access units by ascending PTS
+			ActiveTrackData.AccessUnitFIFO.Sort([](const FActiveTrackData::FSample& a, const FActiveTrackData::FSample& b){return a.PTS < b.PTS;});
+			// Go backwards over the list and drop all access units that _follow_ the next one.
+			for(int32 i=ActiveTrackData.AccessUnitFIFO.Num()-1; i>0; --i)
 			{
-				if (ActiveTrackData.AccessUnitFIFO[i].PTS < pNext->LatestPTS)
+				if (ActiveTrackData.AccessUnitFIFO[i].PTS > NextPTS)
 				{
-					if (ActiveTrackData.AccessUnitFIFO[i].PTS > HighestPTSBelow)
+					for(int32 j=0; j<ActiveTrackData.SortedAccessUnitFIFO.Num(); ++j)
 					{
-						HighestPTSBelow = ActiveTrackData.AccessUnitFIFO[i].PTS;
-						LastIndex = i;
+						if (ActiveTrackData.SortedAccessUnitFIFO[j].PTS == ActiveTrackData.AccessUnitFIFO[i].PTS)
+						{
+							ActiveTrackData.SortedAccessUnitFIFO.RemoveAt(j);
+							break;
+						}
 					}
+					ActiveTrackData.AccessUnitFIFO.RemoveAt(i);
+				}
+				else
+				{
+					break;
 				}
 			}
-			ActiveTrackData.AccessUnitFIFO[LastIndex].AU->bIsLastInPeriod = true;
+			// Sort the list back to index order.
+			ActiveTrackData.AccessUnitFIFO.Sort([](const FActiveTrackData::FSample& a, const FActiveTrackData::FSample& b){return a.SequentialIndex < b.SequentialIndex;});
+			// Whichever element is the last in the list now is the one that needs to be tagged as such.
+			ActiveTrackData.AccessUnitFIFO.Last().AU->bIsLastInPeriod = true;
+			ActiveTrackData.bReadPastLastPTS = true;
+			ActiveTrackData.bTaggedLastSample = true;
+
+			check(pNext == ActiveTrackData.AccessUnitFIFO[0].AU);
+			pNext = ActiveTrackData.AccessUnitFIFO[0].AU;
 		}
 
 		while(!bTerminate && !HasReadBeenAborted())
@@ -2520,10 +2545,6 @@ FStreamReaderDASH::FStreamHandler::EEmitResult FStreamReaderDASH::FStreamHandler
 				ActiveTrackData.AccessUnitFIFO[0].AU = nullptr;
 				ActiveTrackData.AccessUnitFIFO.RemoveAt(0);
 				Result = Result == EEmitResult::SentNothing ? EEmitResult::Sent : Result;
-				if (pNext->bIsLastInPeriod)
-				{
-					ActiveTrackData.bReadPastLastPTS = true;
-				}
 				break;
 			}
 			// If emitting as much as we can we leave this loop now that the receiver is blocked.
diff --git a/Engine/Plugins/Media/ElectraPlayer/Source/ElectraPlayerRuntime/Private/Runtime/Player/DASH/StreamReaderDASH.h b/Engine/Plugins/Media/ElectraPlayer/Source/ElectraPlayerRuntime/Private/Runtime/Player/DASH/StreamReaderDASH.h
--- a/Engine/Plugins/Media/ElectraPlayer/Source/ElectraPlayerRuntime/Private/Runtime/Player/DASH/StreamReaderDASH.h
+++ b/Engine/Plugins/Media/ElectraPlayer/Source/ElectraPlayerRuntime/Private/Runtime/Player/DASH/StreamReaderDASH.h
@@ -139,6 +139,7 @@ private:
 				NumAddedTotal = 0;
 				bIsFirstInSequence = true;
 				bReadPastLastPTS = false;
+				bTaggedLastSample = false;
 				bGotAllSamples = false;
 				bReachedEndOfKnownDuration = false;
 
@@ -154,10 +155,12 @@ private:
 			{
 				FTimeValue PTS;
 				FAccessUnit* AU = nullptr;
-				FSample(FAccessUnit* InAU) : PTS(InAU->PTS), AU(InAU) { InAU->AddRef(); }
+				uint32 SequentialIndex = 0;
+				FSample(FAccessUnit* InAU, uint32 InSequentialIndex) : PTS(InAU->PTS), AU(InAU), SequentialIndex(InSequentialIndex) { InAU->AddRef(); }
 				FSample(const FSample& rhs)
 				{
 					PTS = rhs.PTS;
+					SequentialIndex = rhs.SequentialIndex;
 					if ((AU = rhs.AU) != nullptr)
 					{
 						AU->AddRef();
@@ -178,10 +181,10 @@ private:
 			{
 				if (InAU)
 				{
-					AccessUnitFIFO.Emplace(FActiveTrackData::FSample(InAU));
+					AccessUnitFIFO.Emplace(FActiveTrackData::FSample(InAU, NumAddedTotal));
 					if (bNeedToRecalculateDurations)
 					{
-						SortedAccessUnitFIFO.Emplace(FActiveTrackData::FSample(InAU));
+						SortedAccessUnitFIFO.Emplace(FActiveTrackData::FSample(InAU, NumAddedTotal));
 						SortedAccessUnitFIFO.Sort([](const FActiveTrackData::FSample& a, const FActiveTrackData::FSample& b){return a.PTS < b.PTS;});
 					}
 					// If a valid non-zero duration exists on the AU we take it as the average duration.
@@ -220,6 +223,7 @@ private:
 			int32 Bitrate = 0;
 			bool bIsFirstInSequence = true;
 			bool bReadPastLastPTS = false;
+			bool bTaggedLastSample = false;
 			bool bGotAllSamples = false;
 			bool bNeedToRecalculateDurations = false;
 			bool bReachedEndOfKnownDuration = false;
