Added optional soundwave cloud streaming via dedicated plugin (disabled by default)
Added streamed media segment download cache (disabled by default)
Added commandlet to tag select soundwave assets for cloud streaming

#rnx

diff --git a/Engine/Plugins/Media/ElectraPlayer/Source/ElectraPlayerPlugin/Private/ElectraPlayerPlugin.cpp b/Engine/Plugins/Media/ElectraPlayer/Source/ElectraPlayerPlugin/Private/ElectraPlayerPlugin.cpp
--- a/Engine/Plugins/Media/ElectraPlayer/Source/ElectraPlayerPlugin/Private/ElectraPlayerPlugin.cpp
+++ b/Engine/Plugins/Media/ElectraPlayer/Source/ElectraPlayerPlugin/Private/ElectraPlayerPlugin.cpp
@@ -716,6 +716,19 @@ bool FElectraPlayerPlugin::Open(const FString& Url, const IMediaOptions* Options
 	bMetadataChanged = false;
 	CurrentMetadata.Reset();
 
+	// Check if we can get a segment cache interface for this playback request...
+	TSharedPtr<FElectraPlayerDataCacheContainer, ESPMode::ThreadSafe> ElectraPlayerDataCacheContainer;
+	TSharedPtr<IMediaOptions::FDataContainer, ESPMode::ThreadSafe> ElectraPlayerDataCacheDefaultValue;
+	TSharedPtr<IMediaOptions::FDataContainer, ESPMode::ThreadSafe> DataContainer = Options->GetMediaOption(TEXT("ElectraPlayerDataCache"), ElectraPlayerDataCacheDefaultValue);
+	if (DataContainer.IsValid())
+	{
+		ElectraPlayerDataCacheContainer = StaticCastSharedPtr<FElectraPlayerDataCacheContainer, IMediaOptions::FDataContainer, ESPMode::ThreadSafe>(DataContainer);
+		if (ElectraPlayerDataCacheContainer.IsValid())
+		{
+			LocalPlaystartOptions.ExternalDataCache = ElectraPlayerDataCacheContainer->Data;
+		}
+	}
+
 	return Player->OpenInternal(Url, PlayerOptions, LocalPlaystartOptions, IElectraPlayerInterface::EOpenType::Media);
 }
 
diff --git a/Engine/Plugins/Media/ElectraPlayer/Source/ElectraPlayerPlugin/Private/ElectraPlayerPluginModule.cpp b/Engine/Plugins/Media/ElectraPlayer/Source/ElectraPlayerPlugin/Private/ElectraPlayerPluginModule.cpp
--- a/Engine/Plugins/Media/ElectraPlayer/Source/ElectraPlayerPlugin/Private/ElectraPlayerPluginModule.cpp
+++ b/Engine/Plugins/Media/ElectraPlayer/Source/ElectraPlayerPlugin/Private/ElectraPlayerPluginModule.cpp
@@ -10,6 +10,7 @@
 #include "IElectraPlayerRuntimeModule.h"
 #include "ElectraPlayerPlugin.h"
 #include "ParameterDictionary.h"
+#include "SimpleElectraAudioPlayer.h"
 
 #define LOCTEXT_NAMESPACE "ElectraPlayerPluginModule"
 
@@ -52,6 +53,7 @@ public:
 	virtual void SendAnalyticMetricsPerMinute(const TSharedPtr<IAnalyticsProviderET>& AnalyticsProvider) override
 	{
 		SendAnalyticMetricsPerMinuteDelegate.Broadcast(AnalyticsProvider);
+		ISimpleElectraAudioPlayer::SendAnalyticMetrics(AnalyticsProvider);
 	}
 
 	virtual void ReportVideoStreamingError(const FGuid& PlayerGuid, const FString& LastError) override
diff --git a/Engine/Plugins/Media/ElectraPlayer/Source/ElectraPlayerPlugin/Public/IElectraPlayerPluginModule.h b/Engine/Plugins/Media/ElectraPlayer/Source/ElectraPlayerPlugin/Public/IElectraPlayerPluginModule.h
--- a/Engine/Plugins/Media/ElectraPlayer/Source/ElectraPlayerPlugin/Public/IElectraPlayerPluginModule.h
+++ b/Engine/Plugins/Media/ElectraPlayer/Source/ElectraPlayerPlugin/Public/IElectraPlayerPluginModule.h
@@ -12,7 +12,7 @@
 class IAnalyticsProviderET;
 class IMediaEventSink;
 class IMediaPlayer;
-
+class IElectraPlayerDataCache;
 
 /**
  * This class is used to get safe access to an IMediaOptions interface.
@@ -107,6 +107,13 @@ public:
 	TArray<FTimespan> Data;
 };
 
+class FElectraPlayerDataCacheContainer : public IMediaOptions::FDataContainer
+{
+public:
+	FElectraPlayerDataCacheContainer(TSharedPtr<IElectraPlayerDataCache, ESPMode::ThreadSafe> InData) : Data(InData) {}
+	virtual ~FElectraPlayerDataCacheContainer() {}
+	TSharedPtr<IElectraPlayerDataCache, ESPMode::ThreadSafe> Data;
+};
 
 /**
  * Interface for the ElectraPlayerPlugin module.
diff --git a/Engine/Plugins/Media/ElectraPlayer/Source/ElectraPlayerRuntime/Private/Runtime/Player/ABRRules/ABRVoDPlus.cpp b/Engine/Plugins/Media/ElectraPlayer/Source/ElectraPlayerRuntime/Private/Runtime/Player/ABRRules/ABRVoDPlus.cpp
--- a/Engine/Plugins/Media/ElectraPlayer/Source/ElectraPlayerRuntime/Private/Runtime/Player/ABRRules/ABRVoDPlus.cpp
+++ b/Engine/Plugins/Media/ElectraPlayer/Source/ElectraPlayerRuntime/Private/Runtime/Player/ABRRules/ABRVoDPlus.cpp
@@ -473,13 +473,16 @@ void FABROnDemandPlus::ReportDownloadEnd(const Metrics::FSegmentDownloadStats& S
 			const int64 LastAddedBandwidth = WorkVars->AverageBandwidth.GetLastSample();
 			const int32 NextHighestQualityIndex = QualityIndex + 1 < StreamInformation.Num() ? QualityIndex + 1 : QualityIndex;
 			const int32 NextHighestBitrate = StreamInformation[NextHighestQualityIndex]->Bitrate;
+			const double InverseUsableBandwidthScaleFactor = IsAudioOnly() ? 2.0 / UsableBandwidthScaleFactor : 1.0 / UsableBandwidthScaleFactor;
+			const double PretendedLatency = IsAudioOnly() ? 0.06 : 0.12;
+
 			// Last seen bandwidth already higher than what we might want to switch up to next?
 			if (LastAddedBandwidth > NextHighestBitrate)
 			{
 				// We want to bring the average down a bit in preparation for the next uncached segment in case the network
 				// has degraded. This way we are not overshooting the target.
 				ThisBitrate = (LastAddedBandwidth + NextHighestBitrate) / 2;
-				WorkVars->AverageBandwidth.AddValue(ThisBitrate);
+				WorkVars->AverageBandwidth.AddValue(ThisBitrate * InverseUsableBandwidthScaleFactor);
 			}
 			else if (LastAddedBandwidth > ThisBitrate)
 			{
@@ -488,14 +491,15 @@ void FABROnDemandPlus::ReportDownloadEnd(const Metrics::FSegmentDownloadStats& S
 				// at lightning speed and have its duration worth of time to spend on attempting the next segment download.
 				// If that doesn't finish in time we can downswitch without hurting things too much.
 				ThisBitrate = NextHighestBitrate;
-				WorkVars->AverageBandwidth.InitializeTo(ThisBitrate);
+				WorkVars->AverageBandwidth.InitializeTo(ThisBitrate * InverseUsableBandwidthScaleFactor);
 			}
 			else
 			{
 				// Let's try to bring the average up a bit.
 				ThisBitrate = (int32)(ThisBitrate * 0.2 + NextHighestBitrate * 0.8);
-				WorkVars->AverageBandwidth.AddValue(ThisBitrate);
+				WorkVars->AverageBandwidth.AddValue(ThisBitrate * InverseUsableBandwidthScaleFactor);
 			}
+			WorkVars->AverageLatency.AddValue(PretendedLatency);
 		}
 
 		// Adjust any forced bitrate duration.
diff --git a/Engine/Plugins/Media/ElectraPlayer/Source/ElectraPlayerRuntime/Private/Runtime/Player/AdaptiveStreamingPlayer.cpp b/Engine/Plugins/Media/ElectraPlayer/Source/ElectraPlayerRuntime/Private/Runtime/Player/AdaptiveStreamingPlayer.cpp
--- a/Engine/Plugins/Media/ElectraPlayer/Source/ElectraPlayerRuntime/Private/Runtime/Player/AdaptiveStreamingPlayer.cpp
+++ b/Engine/Plugins/Media/ElectraPlayer/Source/ElectraPlayerRuntime/Private/Runtime/Player/AdaptiveStreamingPlayer.cpp
@@ -1376,6 +1376,13 @@ bool FAdaptiveStreamingPlayer::InternalHandleThreadMessages()
 			}
 			case FWorkerThreadMessages::FMessage::EType::LoadManifest:
 			{
+				// Check if a new or different cache has been set prior to loading the manifest.
+				if (ExternalCache.IsValid())
+				{
+					HttpResponseCache.Reset();
+					HttpResponseCache = IHTTPResponseCache::Create(this, PlayerOptions, MoveTemp(ExternalCache));
+				}
+
 				InternalLoadManifest(msg.Data.ManifestToLoad.URL, msg.Data.ManifestToLoad.MimeType);
 				break;
 			}
@@ -4093,7 +4100,7 @@ void FAdaptiveStreamingPlayer::CheckForStreamEnd()
 
 						// Check if the video stream is seamlessly decodable (next AU is a keyframe and it will also be rendered and not cut off)
 						// Note: at this point we assume audio is always decodable on every AU.
-						bool bSeamlessSwitchPossible = IsSeamlessBufferSwitchPossible(EStreamType::Video, NewOutputBuffers);
+						bool bSeamlessSwitchPossible = bHaveVid ? IsSeamlessBufferSwitchPossible(EStreamType::Video, NewOutputBuffers) : true;
 						if (!bSeamlessSwitchPossible)
 						{
 							StopRendering();
diff --git a/Engine/Plugins/Media/ElectraPlayer/Source/ElectraPlayerRuntime/Private/Runtime/Player/AdaptiveStreamingPlayerDecoders.cpp b/Engine/Plugins/Media/ElectraPlayer/Source/ElectraPlayerRuntime/Private/Runtime/Player/AdaptiveStreamingPlayerDecoders.cpp
--- a/Engine/Plugins/Media/ElectraPlayer/Source/ElectraPlayerRuntime/Private/Runtime/Player/AdaptiveStreamingPlayerDecoders.cpp
+++ b/Engine/Plugins/Media/ElectraPlayer/Source/ElectraPlayerRuntime/Private/Runtime/Player/AdaptiveStreamingPlayerDecoders.cpp
@@ -232,6 +232,10 @@ bool FAdaptiveStreamingPlayer::FindMatchingStreamInfo(FStreamCodecInformation& O
  */
 int32 FAdaptiveStreamingPlayer::CreateDecoder(EStreamType type)
 {
+	if (IsTrackDeselected(type))
+	{
+		return 0;
+	}
 	if (type == EStreamType::Video)
 	{
 		if (VideoDecoder.Decoder == nullptr)
@@ -330,7 +334,7 @@ int32 FAdaptiveStreamingPlayer::CreateDecoder(EStreamType type)
 				}
 				FAccessUnit::Release(AccessUnit);
 
-					// Create the audio decoder
+				// Create the audio decoder
 				AudioDecoder.Decoder = IAudioDecoder::Create();
 				if (AudioDecoder.Decoder)
 				{
diff --git a/Engine/Plugins/Media/ElectraPlayer/Source/ElectraPlayerRuntime/Private/Runtime/Player/AdaptiveStreamingPlayerMetrics.h b/Engine/Plugins/Media/ElectraPlayer/Source/ElectraPlayerRuntime/Private/Runtime/Player/AdaptiveStreamingPlayerMetrics.h
--- a/Engine/Plugins/Media/ElectraPlayer/Source/ElectraPlayerRuntime/Private/Runtime/Player/AdaptiveStreamingPlayerMetrics.h
+++ b/Engine/Plugins/Media/ElectraPlayer/Source/ElectraPlayerRuntime/Private/Runtime/Player/AdaptiveStreamingPlayerMetrics.h
@@ -125,6 +125,8 @@ namespace Metrics
 		double			PresentationTime = 0.0;				//!< Presentation time on media timeline
 		double			Duration = 0.0;						//!< Duration of segment as specified in manifest
 		int32			Bitrate = 0;						//!< Stream bitrate as specified in manifest
+		int32			QualityIndex = 0;					//!< Quality index of this segment
+		int32			HighestQualityIndex = 0;			//!< The highest quality index that could be had.
 		int32			RetryNumber = 0;
 		bool			bIsMissingSegment = false;			//!< true if the segment was not actually downloaded because it is missing on the timeline.
 
diff --git a/Engine/Plugins/Media/ElectraPlayer/Source/ElectraPlayerRuntime/Private/Runtime/Player/AdaptiveStreamingPlayerRenderers.cpp b/Engine/Plugins/Media/ElectraPlayer/Source/ElectraPlayerRuntime/Private/Runtime/Player/AdaptiveStreamingPlayerRenderers.cpp
--- a/Engine/Plugins/Media/ElectraPlayer/Source/ElectraPlayerRuntime/Private/Runtime/Player/AdaptiveStreamingPlayerRenderers.cpp
+++ b/Engine/Plugins/Media/ElectraPlayer/Source/ElectraPlayerRuntime/Private/Runtime/Player/AdaptiveStreamingPlayerRenderers.cpp
@@ -124,6 +124,7 @@ namespace Electra
 		UEMediaError ReturnBuffer(IBuffer* Buffer, bool bRender, const FParamDict& InSampleProperties) override;
 		UEMediaError ReleaseBufferPool() override;
 		bool CanReceiveOutputFrames(uint64 NumFrames) const override;
+		bool GetEnqueuedFrameInfo(int32& OutNumberOfEnqueuedFrames, FTimeValue& OutDurationOfEnqueuedFrames) const override;
 		void SetRenderClock(TSharedPtr<IMediaRenderClock, ESPMode::ThreadSafe> InRenderClock) override;
 		void SetParentRenderer(TWeakPtr<IMediaRenderer, ESPMode::ThreadSafe> ParentRenderer) override;
 		void SetNextApproximatePresentationTime(const FTimeValue& NextApproxPTS) override;
@@ -525,6 +526,11 @@ void FAdaptiveStreamingWrappedRenderer::SetRenderClock(TSharedPtr<IMediaRenderCl
 	WrappedRenderer->SetRenderClock(InRenderClock);
 }
 
+bool FAdaptiveStreamingWrappedRenderer::GetEnqueuedFrameInfo(int32& OutNumberOfEnqueuedFrames, FTimeValue& OutDurationOfEnqueuedFrames) const
+{
+	return WrappedRenderer->GetEnqueuedFrameInfo(OutNumberOfEnqueuedFrames, OutDurationOfEnqueuedFrames);
+}
+
 void FAdaptiveStreamingWrappedRenderer::SetParentRenderer(TWeakPtr<IMediaRenderer, ESPMode::ThreadSafe> ParentRenderer)
 {
 	check(!"this must not be called");
@@ -597,11 +603,20 @@ FTimeValue FAdaptiveStreamingWrappedRenderer::GetEnqueuedSampleDuration()
 int32 FAdaptiveStreamingWrappedRenderer::GetNumEnqueuedSamples(FTimeValue* OutOptionalDuration)
 {
 	FScopeLock lock(&Lock);
+
+	int32 NumAvail = 0;
+	FTimeValue DurAvail(FTimeValue::GetZero());
+	if (!WrappedRenderer->GetEnqueuedFrameInfo(NumAvail, DurAvail))
+	{
+		NumAvail = 0;
+		DurAvail = FTimeValue::GetZero();
+	}
+
 	if (OutOptionalDuration)
 	{
-		*OutOptionalDuration = EnqueuedDuration;
+		*OutOptionalDuration = EnqueuedDuration + DurAvail;
 	}
-	return NumEnqueuedSamples;
+	return NumEnqueuedSamples + NumAvail;
 }
 
 void FAdaptiveStreamingWrappedRenderer::DisableHoldbackOfFirstRenderableVideoFrame(bool bInDisableHoldback)
diff --git a/Engine/Plugins/Media/ElectraPlayer/Source/ElectraPlayerRuntime/Private/Runtime/Player/DASH/StreamReaderDASH.cpp b/Engine/Plugins/Media/ElectraPlayer/Source/ElectraPlayerRuntime/Private/Runtime/Player/DASH/StreamReaderDASH.cpp
--- a/Engine/Plugins/Media/ElectraPlayer/Source/ElectraPlayerRuntime/Private/Runtime/Player/DASH/StreamReaderDASH.cpp
+++ b/Engine/Plugins/Media/ElectraPlayer/Source/ElectraPlayerRuntime/Private/Runtime/Player/DASH/StreamReaderDASH.cpp
@@ -527,6 +527,8 @@ FErrorDetail FStreamReaderDASH::FStreamHandler::LoadInitSegment(TSharedPtrTS<FMP
 	ds.AdaptationSetID = Request->AdaptationSet.IsValid() ? Request->AdaptationSet->GetUniqueIdentifier() : "";
 	ds.RepresentationID = Request->Representation.IsValid() ? Request->Representation->GetUniqueIdentifier() : "";
 	ds.Bitrate = Request->GetBitrate();
+	ds.QualityIndex = Request->QualityIndex;
+	ds.HighestQualityIndex = Request->MaxQualityIndex;
 
 	if (!LoadResult->bSuccess)
 	{
@@ -958,6 +960,8 @@ void FStreamReaderDASH::FStreamHandler::HandleRequestMP4()
 	ds.SegmentType  	   = Metrics::ESegmentType::Init;
 	ds.PresentationTime    = Request->GetFirstPTS().GetAsSeconds();
 	ds.Bitrate  		   = Request->GetBitrate();
+	ds.QualityIndex        = Request->QualityIndex;
+	ds.HighestQualityIndex = Request->MaxQualityIndex;
 	ds.Duration 		   = SegmentDuration.GetAsSeconds();
 	ds.DurationDownloaded  = 0.0;
 	ds.DurationDelivered   = 0.0;
diff --git a/Engine/Plugins/Media/ElectraPlayer/Source/ElectraPlayerRuntime/Private/Runtime/Player/HLS/StreamReaderHLSfmp4.cpp b/Engine/Plugins/Media/ElectraPlayer/Source/ElectraPlayerRuntime/Private/Runtime/Player/HLS/StreamReaderHLSfmp4.cpp
--- a/Engine/Plugins/Media/ElectraPlayer/Source/ElectraPlayerRuntime/Private/Runtime/Player/HLS/StreamReaderHLSfmp4.cpp
+++ b/Engine/Plugins/Media/ElectraPlayer/Source/ElectraPlayerRuntime/Private/Runtime/Player/HLS/StreamReaderHLSfmp4.cpp
@@ -528,7 +528,7 @@ FStreamReaderHLSfmp4::FStreamHandler::EInitSegmentResult FStreamReaderHLSfmp4::F
 
 			TSharedPtrTS<IElectraHttpManager::FRequest> HTTP(new IElectraHttpManager::FRequest);
 			TSharedPtrTS<IElectraHttpManager::FProgressListener> ProgressListener(new IElectraHttpManager::FProgressListener);
-			ProgressListener->ProgressDelegate    = IElectraHttpManager::FProgressListener::FProgressDelegate::CreateRaw(this, &FStreamHandler::HTTPProgressCallback);
+			ProgressListener->ProgressDelegate = IElectraHttpManager::FProgressListener::FProgressDelegate::CreateRaw(this, &FStreamHandler::HTTPProgressCallback);
 			ProgressListener->CompletionDelegate = IElectraHttpManager::FProgressListener::FCompletionDelegate::CreateRaw(this, &FStreamHandler::HTTPCompletionCallback);
 			ReadBuffer.Reset();
 			ReadBuffer.ReceiveBuffer = MakeSharedTS<IElectraHttpManager::FReceiveBuffer>();
@@ -653,6 +653,8 @@ void FStreamReaderHLSfmp4::FStreamHandler::HandleRequest()
 	ds.SegmentType  	   = Metrics::ESegmentType::Media;
 	ds.PresentationTime    = Request->AbsoluteDateTime.GetAsSeconds();
 	ds.Bitrate  		   = Request->Bitrate;
+	ds.QualityIndex        = Request->QualityIndex;
+	ds.HighestQualityIndex = Request->MaxQualityIndex;
 	ds.Duration 		   = Request->SegmentDuration.GetAsSeconds();
 	ds.DurationDownloaded  = 0.0;
 	ds.DurationDelivered   = 0.0;
diff --git a/Engine/Plugins/Media/ElectraPlayer/Source/ElectraPlayerRuntime/Private/Runtime/Player/SimpleElectraAudioPlayer.cpp b/Engine/Plugins/Media/ElectraPlayer/Source/ElectraPlayerRuntime/Private/Runtime/Player/SimpleElectraAudioPlayer.cpp
new file mode 100644
--- /dev/null
+++ b/Engine/Plugins/Media/ElectraPlayer/Source/ElectraPlayerRuntime/Private/Runtime/Player/SimpleElectraAudioPlayer.cpp
@@ -0,0 +1,1945 @@
+// Copyright Epic Games, Inc. All Rights Reserved.
+
+#include "SimpleElectraAudioPlayer.h"
+#include "AdaptiveStreamingPlayer.h"
+#include "Renderer/RendererAudio.h"
+#include "CoreMinimal.h"
+#include "Misc/Guid.h"
+#include "Misc/Optional.h"
+#include "Async/Async.h"
+#include "Tickable.h"
+#include "CoreGlobals.h"
+#include "Stats/Stats.h"
+#include "Stats/Stats2.h"
+#include "MediaAudioDecoderOutput.h"
+#include "Utilities/URLParser.h"
+#include "HAL/IConsoleManager.h"
+#include "IAnalyticsProviderET.h"
+
+#if !UE_BUILD_SHIPPING && WITH_ENGINE
+#define ENABLE_DEBUG_STATS 1
+#else
+#define ENABLE_DEBUG_STATS 0
+#endif
+
+#if ENABLE_DEBUG_STATS
+#include "HAL/PlatformTime.h"
+#include "Engine/Engine.h"
+#include "Engine/Canvas.h"
+#include "Debug/DebugDrawService.h"
+#endif
+
+DEFINE_LOG_CATEGORY_STATIC(LogSimpleElectraPlayer, Log, All);
+
+
+namespace
+{
+// How many _different_ stream analytics to keep in memory.
+static int32 ElectraCSAudio_MaxStreamAnalytics = 100;
+
+// Whether or not analytics reporting is enabled.
+static int32 ElectraCSAudio_SendAnalytics = 1;
+
+
+// Remaining play time of a stream to be considered for reactivation.
+static float ElectraCSAudio_MinRemainPlaytimeForReactivation = 6.0f;
+// Maximum time a stream is allowed to have been in the stopped state
+// to be allowed to be reactivated. If stopped for longer than this
+// it won't be considered any more.
+static float ElectraCSAudio_MaxStoppedDurationForReactivation = 20.0f;
+
+// Maximum number of concurrent instances. 0=unlimited
+static int32 ElectraCSAudio_MaxInstances = 0;
+
+// Whether or not to resume instaces that were stopped due to max instance limit being reached.
+static int32 ElectraCSAudio_ResumeStopped = 1;
+
+static int32 ElectraCSAudio_BlobCacheSize = 1024 * 512;
+static float ElectraCSAudio_BlobCacheAge = 3600.0f;
+}
+
+FAutoConsoleVariableRef CVarElectraCSAudio_MaxInstances(
+	TEXT("ElectraCSAudio.MaxInstances"),
+	ElectraCSAudio_MaxInstances,
+	TEXT("Maximum concurrent cloud stream instances (0=no limit)"),
+	ECVF_Default);
+
+FAutoConsoleVariableRef CVarElectraCSAudio_ResumeStopped(
+	TEXT("ElectraCSAudio.ResumeStopped"),
+	ElectraCSAudio_ResumeStopped,
+	TEXT("Whether or not to resume previously stopped streams when an instance becomes available again (0=do not resume, 1=resume)"),
+	ECVF_Default);
+
+FAutoConsoleVariableRef CVarElectraCSAudio_MinRemainResumeTime(
+	TEXT("ElectraCSAudio.MinRemainResumeTime"),
+	ElectraCSAudio_MinRemainPlaytimeForReactivation,
+	TEXT("How many seconds of remaining playtime a stopped stream needs to have to be resumed."),
+	ECVF_Default);
+
+FAutoConsoleVariableRef CVarElectraCSAudio_MaxStoppedResumeDuration(
+	TEXT("ElectraCSAudio.MaxStoppedResumeDuration"),
+	ElectraCSAudio_MaxStoppedDurationForReactivation,
+	TEXT("How many seconds a stream may be stopped at most in order to be resumed."),
+	ECVF_Default);
+
+FAutoConsoleVariableRef CVarElectraCSAudio_MaxBlobCacheSize(
+	TEXT("ElectraCSAudio.MaxBlobCacheSize"),
+	ElectraCSAudio_BlobCacheSize,
+	TEXT("Size of the blob response cache in bytes."),
+	ECVF_Default);
+
+FAutoConsoleVariableRef CVarElectraCSAudio_MaxBlobCacheDuration(
+	TEXT("ElectraCSAudio.MaxBlobCacheDuration"),
+	ElectraCSAudio_BlobCacheAge,
+	TEXT("How many seconds a blob response remains valid in the cache."),
+	ECVF_Default);
+
+FAutoConsoleVariableRef CVarElectraCSAudio_MaxStreamAnalytics(
+	TEXT("ElectraCSAudio.MaxStreamAnalytics"),
+	ElectraCSAudio_MaxStreamAnalytics,
+	TEXT("Maximum number of stream URLs for which to track analytics at the same time."),
+	ECVF_Default);
+
+FAutoConsoleVariableRef CVarElectraCSAudio_SendAnalytics(
+	TEXT("ElectraCSAudio.SendAnalytics"),
+	ElectraCSAudio_SendAnalytics,
+	TEXT("Whether or not to send analytic metrics. (0=off (default), 1=on)"),
+	ECVF_Default);
+
+using namespace Electra;
+
+namespace
+{
+#if UE_BUILD_SHIPPING
+#define HIDE_URLS_FROM_LOG	1
+#else
+#define HIDE_URLS_FROM_LOG	0
+#endif
+
+static FString SanitizeMessage(FString InMessage)
+{
+#if !HIDE_URLS_FROM_LOG
+	return MoveTemp(InMessage);
+#else
+	int32 searchPos = 0;
+	while(1)
+	{
+		static FString SchemeStr(TEXT("://"));
+		static FString DotDotDotStr(TEXT("..."));
+		static FString TermChars(TEXT("'\",; "));
+		int32 schemePos = InMessage.Find(SchemeStr, ESearchCase::IgnoreCase, ESearchDir::FromStart, searchPos);
+		if (schemePos != INDEX_NONE)
+		{
+			schemePos += SchemeStr.Len();
+			// There may be a generic user message following a potential URL that we do not want to clobber.
+			// We search for any next character that tends to end a URL in a user message, like one of ['",; ]
+			int32 EndPos = InMessage.Len();
+			int32 Start = schemePos;
+			while(Start < EndPos)
+			{
+				int32 pos;
+				if (TermChars.FindChar(InMessage[Start], pos))
+				{
+					break;
+				}
+				++Start;
+			}
+			InMessage.RemoveAt(schemePos, Start-schemePos);
+			InMessage.InsertAt(schemePos, DotDotDotStr);
+			searchPos = schemePos + SchemeStr.Len();
+		}
+		else
+		{
+			break;
+		}
+	}
+	return InMessage;
+#endif
+}
+}
+
+class FSimpleElectraAudioPlayerRenderer;
+
+class FSimpleElectraAudioPlayer : public ISimpleElectraAudioPlayer, public IAdaptiveStreamingPlayerMetrics, public FTickableGameObject, public TSharedFromThis<FSimpleElectraAudioPlayer, ESPMode::ThreadSafe>
+{
+public:
+	static TSharedPtr<FSimpleElectraAudioPlayer, ESPMode::ThreadSafe> Create(const FCreateParams& InCreateParams)
+	{
+		return MakeShareable(new FSimpleElectraAudioPlayer(InCreateParams), TDeleter());
+	}
+	static void SendAnalyticMetrics(const TSharedPtr<IAnalyticsProviderET>& InAnalyticsProvider);
+
+	FSimpleElectraAudioPlayer(const FCreateParams& InCreateParams);
+	~FSimpleElectraAudioPlayer();
+	bool Open(const TMap<FString, FVariant>& InOptions, const FString& InManifestURL, const FTimespan& InStartPosition, const FTimespan& InEncodedDuration, bool bInAutoPlay, bool bInSetLooping, TSharedPtr<IElectraPlayerDataCache, ESPMode::ThreadSafe> InPlayerDataCache) override;
+	void SeekTo(const FTimespan& NewPosition) override
+	{
+		FlushAudio(true);
+		StartPosition = NewPosition;
+		IAdaptiveStreamingPlayer::FSeekParam sp;
+		sp.Time.SetFromTimespan(NewPosition);
+		if (TSharedPtr<IAdaptiveStreamingPlayer, ESPMode::ThreadSafe> LockedPlayer = Player)
+		{
+			LockedPlayer->SeekTo(sp);
+		}
+		else
+		{
+			bMaybeReopenAfterSeek = true;
+		}
+	}
+	void Pause() override
+	{
+		if (TSharedPtr<IAdaptiveStreamingPlayer, ESPMode::ThreadSafe> LockedPlayer = Player)
+		{
+			LockedPlayer->Pause();
+		}
+	}
+	void Resume() override
+	{ 
+		if (TSharedPtr<IAdaptiveStreamingPlayer, ESPMode::ThreadSafe> LockedPlayer = Player)
+		{
+			LockedPlayer->Resume();
+		}
+	}
+	void Stop() override
+	{ 
+		if (TSharedPtr<IAdaptiveStreamingPlayer, ESPMode::ThreadSafe> LockedPlayer = Player)
+		{
+			LockedPlayer->Stop(); 
+		}
+	}
+	bool HasErrored() const override
+	{ return bHasErrored; }
+	FString GetError() const override
+	{ return ErrorMessage; }
+	bool HaveMetadata() const override
+	{ 
+		if (TSharedPtr<IAdaptiveStreamingPlayer, ESPMode::ThreadSafe> LockedPlayer = Player)
+		{
+			return LockedPlayer->HaveMetadata();
+		}
+		return false;
+	}
+	int32 GetBinaryMetadata(TSharedPtr<TArray<uint8>, ESPMode::ThreadSafe>& OutMetadata) const override;
+	FTimespan GetDuration() const override
+	{ 
+		if (TSharedPtr<IAdaptiveStreamingPlayer, ESPMode::ThreadSafe> LockedPlayer = Player)
+		{
+			return LockedPlayer->GetDuration().GetAsTimespan();
+		}
+		return FTimespan::Zero();
+	}
+	FTimespan GetPlayPosition() const override
+	{ return NextPTS; }
+	bool HasEnded() const override
+	{ 
+		if (TSharedPtr<IAdaptiveStreamingPlayer, ESPMode::ThreadSafe> LockedPlayer = Player)
+		{
+			return LockedPlayer->HasEnded();
+		}
+		return false;
+	}
+	bool IsBuffering() const override
+	{ 
+		if (TSharedPtr<IAdaptiveStreamingPlayer, ESPMode::ThreadSafe> LockedPlayer = Player)
+		{
+			return LockedPlayer->IsBuffering(); 
+		}
+		return false;
+	}
+	bool IsSeeking() const override
+	{ 
+		if (TSharedPtr<IAdaptiveStreamingPlayer, ESPMode::ThreadSafe> LockedPlayer = Player)
+		{
+			return LockedPlayer->IsSeeking(); 
+		}
+		return false;
+	}
+	bool IsPlaying() const override
+	{ 
+		if (TSharedPtr<IAdaptiveStreamingPlayer, ESPMode::ThreadSafe> LockedPlayer = Player)
+		{
+			return LockedPlayer->IsPlaying(); 
+		}
+		return false;
+	}
+	bool IsPaused() const override
+	{ 
+		if (TSharedPtr<IAdaptiveStreamingPlayer, ESPMode::ThreadSafe> LockedPlayer = Player)
+		{
+			return LockedPlayer->IsPaused(); 
+		}
+		return false;
+	}
+	void SetBitrateCeiling(int32 HighestSelectableBitrate) override
+	{ 
+		if (TSharedPtr<IAdaptiveStreamingPlayer, ESPMode::ThreadSafe> LockedPlayer = Player)
+		{
+			LockedPlayer->SetBitrateCeiling(HighestSelectableBitrate); 
+		}
+	}
+
+	bool GetStreamFormat(FStreamFormat& OutFormat) const override;
+	int64 GetNumAvailableSamples() const override;
+	int64 GetNextSamples(float* OutBuffer, int32 InBufferSizeInSamples, int32 InNumSamplesToGet, const FDefaultSampleInfo& InDefaultSampleInfo) override;
+	int64 GetNextSamples(int16* OutBuffer, int32 InBufferSizeInSamples, int32 InNumSamplesToGet, const FDefaultSampleInfo& InDefaultSampleInfo) override;
+	bool IsAtEOS() const override;
+
+	void FlushAudio(bool bForceFlush);
+	bool CanAcceptAudioFrames(int32 InNumFrames);
+	bool GetEnqueuedFrameInfo(int32& OutNumberOfEnqueuedFrames, FTimespan& OutDurationOfEnqueuedFrames) const;
+	bool EnqueueAudioFrames(const void* InBufferAddress, int32 InNumSamples, int32 InNumChannels, int32 InSampleRate, const FTimespan& InPTS, const FTimespan& InDuration);
+	void SetReceivedLastBuffer();
+
+	void ReportOpenSource(const FString& InURL) override
+	{
+		if (!bIsResuming)
+		{
+			FString Msg = FString::Printf(TEXT("%s: Opening"), *AssetName);
+			UE_LOG(LogSimpleElectraPlayer, Verbose, TEXT("%s"), *Msg);
+		}
+	}
+	void ReportReceivedMasterPlaylist(const FString& InEffectiveURL) override
+	{ }
+	void ReportReceivedPlaylists() override
+	{
+		if (TSharedPtr<IAdaptiveStreamingPlayer, ESPMode::ThreadSafe> LockedPlayer = Player)
+		{
+			Duration = LockedPlayer->GetDuration().GetAsTimespan();
+			LockedPlayer->DeselectTrack(EStreamType::Video);
+			IAdaptiveStreamingPlayer::FSeekParam sp;
+			if (StartPosition > FTimespan::Zero())
+			{
+				sp.Time.SetFromTimespan(StartPosition);
+			}
+			LockedPlayer->SeekTo(sp);
+		}
+	}
+	void ReportTracksChanged() override
+	{ }
+	void ReportPlaylistDownload(const Metrics::FPlaylistDownloadStats& InPlaylistDownloadStats) override
+	{ }
+	void ReportCleanStart() override
+	{ }
+	void ReportBufferingStart(Metrics::EBufferingReason InBufferingReason) override
+	{ }
+	void ReportBufferingEnd(Metrics::EBufferingReason InBufferingReason) override
+	{ }
+	void ReportBandwidth(int64 InEffectiveBps, int64 InThroughputBps, double InLatencyInSeconds) override
+	{ }
+	void ReportBufferUtilization(const Metrics::FBufferStats& InBufferStats) override
+	{ }
+	void ReportSegmentDownload(const Metrics::FSegmentDownloadStats& InSegmentDownloadStats) override
+	{
+		if (!InSegmentDownloadStats.bWasSuccessful || InSegmentDownloadStats.RetryNumber)
+		{
+			FString Msg = FString::Printf(TEXT("%s: Segment download issue \"%s\", retry:%d, aborted:%d, silence:%d"), *AssetName,
+											*InSegmentDownloadStats.FailureReason, InSegmentDownloadStats.RetryNumber, InSegmentDownloadStats.bWasAborted, InSegmentDownloadStats.bInsertedFillerData);
+			UE_LOG(LogSimpleElectraPlayer, Verbose, TEXT("%s"), *Msg);
+		}
+
+		// Update analytics
+		Analytics.BestQualityLevel = InSegmentDownloadStats.HighestQualityIndex;
+		if (InSegmentDownloadStats.bWasSuccessful)
+		{
+			++Analytics.NumSegmentsLoadedAtQuality.FindOrAdd(InSegmentDownloadStats.QualityIndex);
+			if (InSegmentDownloadStats.bIsCachedResponse)
+			{
+				++Analytics.CacheHitsAtQuality.FindOrAdd(InSegmentDownloadStats.QualityIndex);
+			}
+		}
+		else if (!InSegmentDownloadStats.bWasAborted && !InSegmentDownloadStats.bIsMissingSegment && !InSegmentDownloadStats.bWaitingForRemoteRetryElement)
+		{
+			// Use error `1` for timeout and `10` for parse errors.
+			int32 ErrorCode = InSegmentDownloadStats.bDidTimeout ? 1 : InSegmentDownloadStats.bParseFailure ? 10 : InSegmentDownloadStats.HTTPStatusCode;
+			++Analytics.ErrorCodeAndCount.FindOrAdd(ErrorCode);
+		}
+	}
+	void ReportLicenseKey(const Metrics::FLicenseKeyStats& InLicenseKeyStats) override
+	{ }
+	void ReportVideoQualityChange(int32 InNewBitrate, int32 InPreviousBitrate, bool bInIsDrasticDownswitch) override
+	{ }
+	void ReportDataAvailabilityChange(const Metrics::FDataAvailabilityChange& InDataAvailability) override
+	{ }
+	void ReportDecodingFormatChange(const FStreamCodecInformation& InNewDecodingFormat) override
+	{ }
+	void ReportPrerollStart() override
+	{ }
+	void ReportPrerollEnd() override
+	{
+		if (bAutoPlay)
+		{
+			if (TSharedPtr<IAdaptiveStreamingPlayer, ESPMode::ThreadSafe> LockedPlayer = Player)
+			{
+				LockedPlayer->Resume();
+			}
+		}
+	}
+	void ReportPlaybackStart() override
+	{ }
+	void ReportPlaybackPaused() override
+	{ }
+	void ReportPlaybackResumed() override
+	{ }
+	void ReportPlaybackEnded() override
+	{ }
+	void ReportJumpInPlayPosition(const FTimeValue& InToNewTime, const FTimeValue& InFromTime, Metrics::ETimeJumpReason InTimejumpReason) override
+	{ }
+	void ReportPlaybackStopped() override
+	{ }
+	void ReportSeekCompleted() override
+	{ }
+	void ReportMediaMetadataChanged(TSharedPtrTS<UtilsMP4::FMetadataParser> InMetadata) override
+	{ }
+	void ReportError(const FString& InErrorReason) override
+	{
+		bHasErrored = true;
+		ErrorMessage = InErrorReason;
+		FString Msg = FString::Printf(TEXT("%s: %s"), *AssetName, *SanitizeMessage(InErrorReason));
+		UE_LOG(LogSimpleElectraPlayer, Error, TEXT("%s"), *Msg);
+	}
+	void ReportLogMessage(IInfoLog::ELevel InLogLevel, const FString& InLogMessage, int64 InPlayerWallclockMilliseconds) override
+	{
+		// Ignore common authoring mistakes.
+		if (InLogMessage.Contains(TEXT("differs from timescale")))
+		{
+			return;
+		}
+
+		FString Msg = FString::Printf(TEXT("%s: %s"), *AssetName, *SanitizeMessage(InLogMessage));
+		switch(InLogLevel)
+		{
+			case IInfoLog::ELevel::Verbose:
+			{
+				UE_LOG(LogSimpleElectraPlayer, Verbose, TEXT("%s"), *Msg);
+				break;
+			}
+			case IInfoLog::ELevel::Info:
+			{
+				UE_LOG(LogSimpleElectraPlayer, Log, TEXT("%s"), *Msg)
+				break;
+			}
+			case IInfoLog::ELevel::Warning:
+			{
+				UE_LOG(LogSimpleElectraPlayer, Warning, TEXT("%s"), *Msg)
+				break;
+			}
+			case IInfoLog::ELevel::Error:
+			{
+				UE_LOG(LogSimpleElectraPlayer, Error, TEXT("%s"), *Msg)
+				break;
+			}
+		}
+	}
+	void ReportDroppedVideoFrame() override
+	{ }
+	void ReportDroppedAudioFrame() override
+	{ }
+
+private:
+	void Tick(float DeltaTime) override;
+	ETickableTickType GetTickableTickType() const override
+	{ return ETickableTickType::Conditional; }
+	bool IsTickable() const override
+	{ return bIsReadyToTick; }
+	bool IsAllowedToTick() const override
+	{ return bIsReadyToTick; }
+	TStatId GetStatId() const override
+	{ RETURN_QUICK_DECLARE_CYCLE_STAT(FSimpleElectraAudioPlayer, STATGROUP_Tickables); }
+	bool IsTickableWhenPaused() const override
+	{ return true; }
+	bool IsTickableInEditor() const override
+	{ return true; }
+
+
+	class TDeleter
+	{
+	public:
+		void operator()(FSimpleElectraAudioPlayer* InInstanceToDelete)
+		{
+			InInstanceToDelete->bDestructionRequested = true;
+		}
+	};
+
+	class FResourceProvider : public IAdaptiveStreamingPlayerResourceProvider
+	{
+	public:
+		void SetPlaylistURL(const FString& InPlaylistURL)
+		{ PlaylistURL = InPlaylistURL; }
+		void SetOptions(const TMap<FString, FVariant>& InOptions)
+		{ Options = InOptions; }
+		void ProvideStaticPlaybackDataForURL(TSharedPtr<IAdaptiveStreamingPlayerResourceRequest, ESPMode::ThreadSafe> InOutRequest) override;
+	private:
+		FString PlaylistURL;
+		TMap<FString, FVariant> Options;
+	};
+
+	class FBlobRequest
+	{
+	public:
+		FBlobRequest() : Request(MakeShared<Electra::FHTTPResourceRequest, ESPMode::ThreadSafe>())
+		{ }
+		void OnBlobRequestComplete(TSharedPtrTS<Electra::FHTTPResourceRequest> InRequest)
+		{
+			ErrorCode = 0;
+			if (!InRequest->GetWasCanceled())
+			{
+				ErrorCode = InRequest->GetError();
+				BlobData = MakeShared<TArray<uint8>, ESPMode::ThreadSafe>();
+				TSharedPtrTS<IElectraHttpManager::FReceiveBuffer> ResponseBuffer = InRequest->GetResponseBuffer();
+				if (ResponseBuffer.IsValid())
+				{
+					BlobData->Append((const uint8*)ResponseBuffer->Buffer.GetLinearReadData(), ResponseBuffer->Buffer.Num());
+				}
+			}
+			bIsComplete = true;
+			Request.Reset();
+		}
+		void SetBlobData(TSharedPtr<TArray<uint8>, ESPMode::ThreadSafe> InBlobData)
+		{
+			ErrorCode = 0;
+			bIsComplete = true;
+			BlobData = MoveTemp(InBlobData);
+			Request.Reset();
+		}
+		TSharedPtr<Electra::FHTTPResourceRequest, ESPMode::ThreadSafe> Request;
+		TSharedPtr<TArray<uint8>, ESPMode::ThreadSafe> BlobData;
+		FString URL;
+		volatile int32 ErrorCode = -1;
+		volatile bool bIsComplete = false;
+		bool bReportedError = false;
+		bool bWasCached = false;
+	};
+
+	struct FBlobCacheEntry
+	{
+		TSharedPtr<TArray<uint8>, ESPMode::ThreadSafe> Data;
+		FString URL;
+		double TimeAdded = 0.0;
+	};
+
+	class FBlobCache
+	{
+	public:
+		TSharedPtr<TArray<uint8>, ESPMode::ThreadSafe> GetData(const FString& InForURL);
+		void AddData(const FString& InForURL, const TSharedPtr<TArray<uint8>, ESPMode::ThreadSafe>& InData);
+	private:
+		void CheckLimits(int32 InBytesNeeded);
+		FCriticalSection Lock;
+		TMap<FString, TSharedPtr<FBlobCacheEntry, ESPMode::ThreadSafe>> EntryMap;
+		TArray<TSharedPtr<FBlobCacheEntry, ESPMode::ThreadSafe>> EntryLRU;
+		int32 SizeInUse = 0;
+	};
+
+	struct FSampleBlock
+	{
+		FTimespan PTS;
+		void* Buffer = nullptr;
+		int32 NumSamples = 0;
+		int32 Offset = 0;
+		void Free()
+		{
+			if (Buffer)
+			{
+				FMemory::Free(Buffer);
+				Buffer = nullptr;
+			}
+		}
+	};
+
+	enum class EState
+	{
+		Uninitialized,
+		OpenBlob,
+		OpeningStream,
+		TryReopeningStream,
+		Active,
+		Stopped,
+		Errored
+	};
+
+	const int32 kNumDummyBlockSamples = 1024;
+
+	FCreateParams CreateParams;
+
+	volatile EState CurrentState = EState::Uninitialized;
+	volatile bool bIsReadyToTick = false;
+	volatile bool bAskedToBeReleased = false;
+	volatile bool bCloseInProgress = false;
+	volatile bool bDestructionRequested = false;
+	volatile bool bDestructionIssued = false;
+	volatile bool bDestructionComplete = false;
+
+	const double SuggestedBufferDuration = 0.5;
+	const int32 MinNumSamplesInBuffer = 16384;	// Minimum number of samples, regardless of sample rate.
+
+	TOptional<FString> BlobParameters;
+	TSharedPtr<FBlobRequest, ESPMode::ThreadSafe> PendingBlobRequest;
+
+	TMap<FString, FVariant> Options;
+	TSharedPtr<IAdaptiveStreamingPlayer, ESPMode::ThreadSafe> Player;
+	TSharedPtr<FSimpleElectraAudioPlayerRenderer, ESPMode::ThreadSafe> Renderer;
+	TSharedPtr<FResourceProvider, ESPMode::ThreadSafe> ResourceProvider;
+	TWeakPtr<IElectraPlayerDataCache, ESPMode::ThreadSafe> PlayerDataCache;
+
+	FString URL;
+	FString BaseURL;
+	FString AssetName;
+	FTimespan StartPosition;
+	bool bAutoPlay = true;
+	bool bSetLooping = false;
+
+	FTimespan EncodedDuration;
+	FTimespan Duration;
+	FString ErrorMessage;
+	bool bHasErrored = false;
+
+	mutable FCriticalSection Lock;
+	FStreamFormat StreamFormat;
+	TArray<FSampleBlock> SampleBlocks;
+	int64 NumSamplesAvailable = 0;
+	bool bReachedEOS = false;
+	FTimespan NextPTS;
+
+	double TimeCreated = -1.0;
+	double TimeUntilReady = -1.0;
+	double TimeStopped = -1.0;
+	bool bMaybeReopenAfterSeek = false;
+	bool bIsResuming = false;
+
+	void CreatePlayerInstance();
+	void ClosePlayerInstance(bool bDeleteThis);
+
+	struct FClosePlayerInstanceHelper
+	{
+		TSharedPtr<IAdaptiveStreamingPlayer, ESPMode::ThreadSafe> PlayerInstance;
+		FSimpleElectraAudioPlayer* This = nullptr;
+		bool bDeleteThis = false;
+	};
+	static void DoClosePlayerAsync(TSharedPtr<FClosePlayerInstanceHelper, ESPMode::ThreadSafe>&& InPlayerHelper);
+
+	bool HasBeenTerminated() const;
+	void UpdateAssetName();
+	void InternalApplyOptions();
+	void HandleOpenStream(bool bOnlyIfAvailable);
+	static void MaybeRestartAStoppedStream(FSimpleElectraAudioPlayer* This);
+
+	FBlobCache& GetBlobCache() const
+	{
+		static FBlobCache Cache;
+		return Cache;
+	}
+
+	static FCriticalSection InstanceLock;
+	static TArray<FSimpleElectraAudioPlayer*> AllInstances;
+	static TArray<FSimpleElectraAudioPlayer*> ActiveInstances;
+	static TArray<FSimpleElectraAudioPlayer*> StoppedInstances;
+
+
+	struct FAnalyticsEntry
+	{
+		int32 NumTimesStarted = 0;
+		int32 NumTimesSuspended = 0;
+		int32 NumTimesResumed = 0;
+		int32 BestQualityLevel = 0;
+		double StartLatency = 0.0;
+		TMap<int32, int32> ErrorCodeAndCount;
+		TMap<int32, int32> NumSegmentsLoadedAtQuality;
+		TMap<int32, int32> CacheHitsAtQuality;
+		FString Analytics_ErrorCodeAndCount;
+		FString Analytics_NumSegmentsLoadedAtQuality;
+		FString Analytics_CacheHitsAtQuality;
+
+		void PrepareAnalytics()
+		{
+			Analytics_ErrorCodeAndCount = MakeAnalyticsStringFromMap(ErrorCodeAndCount);
+			Analytics_NumSegmentsLoadedAtQuality = MakeAnalyticsStringFromMap(NumSegmentsLoadedAtQuality);
+			Analytics_CacheHitsAtQuality = MakeAnalyticsStringFromMap(CacheHitsAtQuality);
+		}
+	
+		FString MakeAnalyticsStringFromMap(TMap<int32, int32>& InMap) const
+		{
+			InMap.KeySort([](int32 A, int32 B){return A<B;});
+			FString s;
+			int32 i=0;
+			for(auto It=InMap.CreateConstIterator(); It; ++It, ++i)
+			{
+				if (i)
+				{
+					s.AppendChar(TCHAR(','));
+				}
+				s += FString::Printf(TEXT("%d:%d"), It->Key, It->Value);
+			}
+			return s;
+		}
+	};
+	mutable FAnalyticsEntry Analytics;
+	static FCriticalSection AnalyticsLock;
+	static TMap<FString, FAnalyticsEntry> AnalyticEntries;
+	static void MergeAnalytics(const FString& AssetId, const FAnalyticsEntry& InAnalytics);
+
+
+#if ENABLE_DEBUG_STATS
+public:
+	void DebugDrawInst(UCanvas* InCanvas);
+
+	/** Needed for debug draw management. */
+	static void EnableDebugDraw(bool bEnable);
+	static void DebugDraw(UCanvas* InCanvas, APlayerController* InPlayerController);
+	static void DebugDrawPrintLine(UCanvas* InCanvas, const FString& InString, const FLinearColor& InColor);
+	static FDelegateHandle DebugDrawDelegateHandle;
+	static UFont* DebugDrawFont;
+	static int32 DebugDrawFontRowHeight;
+	static FVector2D DebugDrawTextPos;
+	static const FLinearColor DebugColor_Orange;
+	static const FLinearColor DebugColor_Green;
+	static const FLinearColor DebugColor_Red;
+#endif
+};
+FCriticalSection FSimpleElectraAudioPlayer::InstanceLock;
+TArray<FSimpleElectraAudioPlayer*> FSimpleElectraAudioPlayer::AllInstances;
+TArray<FSimpleElectraAudioPlayer*> FSimpleElectraAudioPlayer::ActiveInstances;
+TArray<FSimpleElectraAudioPlayer*> FSimpleElectraAudioPlayer::StoppedInstances;
+
+FCriticalSection FSimpleElectraAudioPlayer::AnalyticsLock;
+TMap<FString, FSimpleElectraAudioPlayer::FAnalyticsEntry> FSimpleElectraAudioPlayer::AnalyticEntries;
+
+
+class FSimpleElectraAudioPlayerRenderer : public IMediaRenderer, public TSharedFromThis<FSimpleElectraAudioPlayerRenderer, ESPMode::ThreadSafe>
+{
+public:
+	FSimpleElectraAudioPlayerRenderer(TSharedPtr<FSimpleElectraAudioPlayer, ESPMode::ThreadSafe> InParentPlayer)
+	{
+		Player = InParentPlayer;
+	}
+	virtual ~FSimpleElectraAudioPlayerRenderer()
+	{ }
+
+	//-------------------------------------------------------------------------
+	// Methods for IMediaRenderer
+	//
+	const FParamDict& GetBufferPoolProperties() const override
+	{ return BufferPoolProperties; }
+	UEMediaError CreateBufferPool(const FParamDict& Parameters) override;
+	UEMediaError AcquireBuffer(IBuffer*& OutBuffer, int32 TimeoutInMicroseconds, const FParamDict& InParameters) override;
+	UEMediaError ReturnBuffer(IBuffer* Buffer, bool bRender, const FParamDict& InSampleProperties) override;
+	UEMediaError ReleaseBufferPool() override;
+	bool CanReceiveOutputFrames(uint64 NumFrames) const override;
+	bool GetEnqueuedFrameInfo(int32& OutNumberOfEnqueuedFrames, FTimeValue& OutDurationOfEnqueuedFrames) const override;
+	void SetRenderClock(TSharedPtr<IMediaRenderClock, ESPMode::ThreadSafe> RenderClock) override
+	{ }
+	void SetParentRenderer(TWeakPtr<IMediaRenderer, ESPMode::ThreadSafe> ParentRenderer) override;
+	void SetNextApproximatePresentationTime(const FTimeValue& NextApproxPTS) override
+	{ }
+	UEMediaError Flush(const FParamDict& InOptions) override;
+	void StartRendering(const FParamDict& InOptions) override
+	{ }
+	void StopRendering(const FParamDict& InOptions) override
+	{ }
+	void SampleReleasedToPool(IDecoderOutput* InDecoderOutput) override;
+private:
+	class FMediaBufferSharedPtrWrapper : public IMediaRenderer::IBuffer
+	{
+	public:
+		explicit FMediaBufferSharedPtrWrapper(IAudioDecoderOutputPtr InDecoderOutput) : DecoderOutput(InDecoderOutput)
+		{ }
+		virtual ~FMediaBufferSharedPtrWrapper() = default;
+		const FParamDict& GetBufferProperties() const override
+		{ return BufferProperties; }
+		FParamDict& GetMutableBufferProperties() override
+		{ return BufferProperties; }
+
+		FParamDict BufferProperties;
+		IAudioDecoderOutputPtr DecoderOutput;
+	};
+
+	TWeakPtr<FSimpleElectraAudioPlayer, ESPMode::ThreadSafe> Player;
+	TWeakPtr<IMediaRenderer, ESPMode::ThreadSafe> ParentRenderer;
+	TDecoderOutputObjectPool<IAudioDecoderOutput, FElectraPlayerPlatformAudioDecoderOutputFactory> DecoderOutputPool;
+	FParamDict BufferPoolProperties;
+	uint32 MaxBufferSize = 0;
+	int32 NumOutputAudioBuffersInUse = 0;
+	int32 NumBuffers = 0;
+	int32 NumBuffersAcquiredForDecoder = 0;
+};
+
+
+TSharedPtr<ISimpleElectraAudioPlayer, ESPMode::ThreadSafe> ISimpleElectraAudioPlayer::Create(const FCreateParams& InCreateParams)
+{
+	return FSimpleElectraAudioPlayer::Create(InCreateParams);
+}
+
+void ISimpleElectraAudioPlayer::SendAnalyticMetrics(const TSharedPtr<IAnalyticsProviderET>& InAnalyticsProvider)
+{
+	if (ElectraCSAudio_SendAnalytics)
+	{
+		FSimpleElectraAudioPlayer::SendAnalyticMetrics(InAnalyticsProvider);
+	}
+}
+
+
+
+FSimpleElectraAudioPlayer::FSimpleElectraAudioPlayer(const FCreateParams& InCreateParams)
+	: CreateParams(InCreateParams)
+{
+	if (!CreateParams.InstanceGUID.IsValid())
+	{
+		CreateParams.InstanceGUID = FGuid::NewGuid();
+	}
+
+	StreamFormat.NumChannels = 0;
+	StreamFormat.SampleRate = 0;
+	NextPTS = FTimespan::MinValue();
+
+	InstanceLock.Lock();
+	AllInstances.AddUnique(this);
+	InstanceLock.Unlock();
+
+	bIsReadyToTick = true;
+}
+
+FSimpleElectraAudioPlayer::~FSimpleElectraAudioPlayer()
+{
+	bIsReadyToTick = false;
+
+	InstanceLock.Lock();
+	check(AllInstances.Find(this) == INDEX_NONE);
+	check(ActiveInstances.Find(this) == INDEX_NONE);
+	check(StoppedInstances.Find(this) == INDEX_NONE);
+	InstanceLock.Unlock();
+	// The player instance must have been destroyed already!
+	check(!Player.IsValid());
+	ResourceProvider.Reset();
+	Renderer.Reset();
+	FlushAudio(true);
+}
+
+void FSimpleElectraAudioPlayer::Tick(float DeltaTime)
+{
+	if (bDestructionRequested)
+	{
+		if (!bDestructionIssued)
+		{
+			if (!bCloseInProgress)
+			{
+				FString Msg = FString::Printf(TEXT("%s: Closing"), *AssetName);
+				UE_LOG(LogSimpleElectraPlayer, Verbose, TEXT("%s"), *Msg);
+
+				bDestructionIssued = true;
+				ClosePlayerInstance(true);
+				InstanceLock.Lock();
+				StoppedInstances.Remove(this);
+				InstanceLock.Unlock();
+			}
+		}
+		if (bDestructionComplete)
+		{
+			bIsReadyToTick = false;
+			InstanceLock.Lock();
+			AllInstances.Remove(this);
+			ActiveInstances.Remove(this);
+			StoppedInstances.Remove(this);
+			InstanceLock.Unlock();
+
+			if (ElectraCSAudio_ResumeStopped > 0)
+			{
+				MaybeRestartAStoppedStream(this);
+			}
+			delete this;
+		}
+		return;
+	}
+
+	switch(CurrentState)
+	{
+		case EState::OpeningStream:
+		{
+			HandleOpenStream(false);
+			break;
+		}
+		case EState::TryReopeningStream:
+		{
+			CurrentState = EState::Stopped;
+			HandleOpenStream(true);
+			break;
+		}
+		case EState::Errored:
+		{
+			return;
+		}
+		case EState::Stopped:
+		{
+			if (bMaybeReopenAfterSeek)
+			{
+				bMaybeReopenAfterSeek = false;
+				StartPosition = NextPTS;
+				HandleOpenStream(true);
+			}
+			return;
+		}
+	}
+
+	if (bAskedToBeReleased)
+	{
+		FString Msg = FString::Printf(TEXT("%s: Stopping due to limit"), *AssetName);
+		UE_LOG(LogSimpleElectraPlayer, Verbose, TEXT("%s"), *Msg);
+
+		TimeStopped = FPlatformTime::Seconds();
+		++Analytics.NumTimesSuspended;
+		ClosePlayerInstance(false);
+	}
+}
+
+
+void FSimpleElectraAudioPlayer::HandleOpenStream(bool bOnlyIfAvailable)
+{
+	int32 MaxAllowed = CreateParams.MaxTotalPlayerInstances;
+	// CVar overrides unconditionally.
+	if (ElectraCSAudio_MaxInstances > 0)
+	{
+		MaxAllowed = ElectraCSAudio_MaxInstances;
+	}
+
+	bool bCreateNow = false;
+
+	// Is there a limit on the number of streams?
+	if (MaxAllowed > 0)
+	{
+		// Check if we need to make room for this new stream.
+		FScopeLock lock(&InstanceLock);
+		
+		int32 NumActive = ActiveInstances.Num();
+		int32 NumRequestedToBeReleased = 0;
+		for(auto& Inst : ActiveInstances)
+		{
+			NumRequestedToBeReleased += Inst->bAskedToBeReleased ? 1 : 0;
+		}
+		if (NumActive >= MaxAllowed)
+		{
+			if (!bOnlyIfAvailable)
+			{
+				int32 NumToRemove = NumActive - NumRequestedToBeReleased - MaxAllowed + 1;
+
+				for(int32 i=0; NumToRemove>0 && i<NumActive; ++i)
+				{
+					if (!ActiveInstances[i]->bAskedToBeReleased)
+					{
+						ActiveInstances[i]->bAskedToBeReleased = true;
+						--NumToRemove;
+					}
+				}
+			}
+		}
+		else
+		{
+			bCreateNow = true;
+		}
+	}
+	else
+	{
+		bCreateNow = true;
+	}
+
+	if (bCreateNow)
+	{
+		// No. Go!
+		FString Msg = FString::Printf(TEXT("%s: %s playback"), *AssetName, !bAskedToBeReleased ? TEXT("Preparing") : TEXT("Resuming"));
+		UE_LOG(LogSimpleElectraPlayer, Verbose, TEXT("%s"), *Msg);
+
+		// Track use count
+		if (bAskedToBeReleased)
+		{
+			++Analytics.NumTimesResumed;
+		}
+		else
+		{
+			++Analytics.NumTimesStarted;
+		}
+
+		CreatePlayerInstance();
+		InternalApplyOptions();
+		Player->SetPlayerDataCache(PlayerDataCache.Pin());
+		Player->LoadManifest(URL);
+		if (bSetLooping)
+		{
+			IAdaptiveStreamingPlayer::FLoopParam lp;
+			lp.bEnableLooping = true;
+			Player->SetLooping(lp);
+		}
+			
+		// Add to the list of active instances.
+		InstanceLock.Lock();
+		ActiveInstances.AddUnique(this);
+		StoppedInstances.Remove(this);
+		InstanceLock.Unlock();
+
+		// Set state to active.
+		bIsResuming = bAskedToBeReleased;
+		CurrentState = EState::Active;
+		bAskedToBeReleased = false;
+		TimeStopped = -1.0;
+	}
+}
+
+void FSimpleElectraAudioPlayer::MergeAnalytics(const FString& AssetId, const FAnalyticsEntry& InAnalytics)
+{
+	// If this entry was not active and had no failure downloading the blob we
+	// do not merge it.
+	if (!InAnalytics.NumTimesStarted && InAnalytics.ErrorCodeAndCount.IsEmpty())
+	{
+		return;
+	}
+
+	FScopeLock lock(&AnalyticsLock);
+
+	// To conserve memory we restrict ourselves to keeping only a small number of stream analytics in memory.
+	// This is to avoid growing indefinitely if the analytics are not periodically picked up by user code.
+	const int32 MaxAnalytics = FMath::Max((int32)0, ElectraCSAudio_MaxStreamAnalytics);
+	while(AnalyticEntries.Num() >= MaxAnalytics)
+	{
+		// Remove whichever entry comes first when enumerating.
+		auto it = AnalyticEntries.CreateConstIterator();
+		AnalyticEntries.Remove(it->Key);
+	}
+
+	FAnalyticsEntry* Entry = AnalyticEntries.Find(AssetId);
+	if (!Entry)
+	{
+		AnalyticEntries.Emplace(AssetId, InAnalytics).PrepareAnalytics();
+	}
+	else
+	{
+		Entry->NumTimesSuspended += InAnalytics.NumTimesSuspended;
+		Entry->NumTimesResumed += InAnalytics.NumTimesResumed;
+		Entry->BestQualityLevel = InAnalytics.BestQualityLevel;
+		Entry->StartLatency = (Entry->StartLatency * Entry->NumTimesStarted + InAnalytics.StartLatency) / (Entry->NumTimesStarted + 1);
+		++Entry->NumTimesStarted;
+		for(auto& it : InAnalytics.ErrorCodeAndCount)
+		{
+			Entry->ErrorCodeAndCount.FindOrAdd(it.Key) += it.Value;
+		}
+		for(auto& it : InAnalytics.NumSegmentsLoadedAtQuality)
+		{
+			Entry->NumSegmentsLoadedAtQuality.FindOrAdd(it.Key) += it.Value;
+		}
+		for(auto& it : InAnalytics.CacheHitsAtQuality)
+		{
+			Entry->CacheHitsAtQuality.FindOrAdd(it.Key) += it.Value;
+		}
+		Entry->PrepareAnalytics();
+	}
+}
+
+void FSimpleElectraAudioPlayer::SendAnalyticMetrics(const TSharedPtr<IAnalyticsProviderET>& InAnalyticsProvider)
+{
+	if (InAnalyticsProvider.IsValid())
+	{
+		AnalyticsLock.Lock();
+		TMap<FString, FAnalyticsEntry> CurrentEntries(MoveTemp(AnalyticEntries));
+		AnalyticsLock.Unlock();
+
+		for(auto& Elem : CurrentEntries)
+		{
+			TArray<FAnalyticsEventAttribute> ParamArray;
+			ParamArray.Add(FAnalyticsEventAttribute(TEXT("URL"), Elem.Key));
+			ParamArray.Add(FAnalyticsEventAttribute(TEXT("count"), Elem.Value.NumTimesStarted));
+			ParamArray.Add(FAnalyticsEventAttribute(TEXT("nsusp"), Elem.Value.NumTimesSuspended));
+			ParamArray.Add(FAnalyticsEventAttribute(TEXT("nresm"), Elem.Value.NumTimesResumed));
+			ParamArray.Add(FAnalyticsEventAttribute(TEXT("bestq"), Elem.Value.BestQualityLevel));
+			ParamArray.Add(FAnalyticsEventAttribute(TEXT("avgltcy"), Elem.Value.StartLatency));
+			ParamArray.Add(FAnalyticsEventAttribute(TEXT("errlst"), Elem.Value.Analytics_ErrorCodeAndCount));
+			ParamArray.Add(FAnalyticsEventAttribute(TEXT("loaded"), Elem.Value.Analytics_NumSegmentsLoadedAtQuality));
+			ParamArray.Add(FAnalyticsEventAttribute(TEXT("chits"), Elem.Value.Analytics_CacheHitsAtQuality));
+			InAnalyticsProvider->RecordEvent(TEXT("ElectraCSA.Report"), MoveTemp(ParamArray));
+		}
+	}
+}
+
+
+void FSimpleElectraAudioPlayer::MaybeRestartAStoppedStream(FSimpleElectraAudioPlayer* This)
+{
+	double TimeNow = FPlatformTime::Seconds();
+	FScopeLock lock(&InstanceLock);
+	for(auto& StoppedInst :	StoppedInstances)
+	{
+		// Don't restart ourselves.
+		if (StoppedInst == This)
+		{
+			continue;
+		}
+		// Don't revive the living dead!
+		if (StoppedInst->bDestructionRequested)
+		{
+			continue;
+		}
+		// Only those that are actually stopped and not in error or already restarting.
+		if (StoppedInst->CurrentState != EState::Stopped)
+		{
+			continue;
+		}
+
+		// Calculate the remaining play time of that stream.
+		FTimespan CurrentTime = StoppedInst->NextPTS < FTimespan::Zero() ? FTimespan::Zero() : StoppedInst->NextPTS;
+		FTimespan AssetDuration;
+		if (StoppedInst->EncodedDuration > FTimespan::Zero() && StoppedInst->EncodedDuration < FTimespan::MaxValue())
+		{
+			AssetDuration = StoppedInst->EncodedDuration;
+		}
+		else if (StoppedInst->Duration > FTimespan::Zero())
+		{
+			AssetDuration = StoppedInst->Duration;
+		}
+		else
+		{
+			continue;
+		}
+		double TimeRemaining = (AssetDuration - CurrentTime).GetTotalSeconds();
+		// If this stream is about to end don't bother restarting it.
+		if (TimeRemaining < ElectraCSAudio_MinRemainPlaytimeForReactivation)
+		{
+			continue;
+		}
+
+		// If the stream was stopped too long ago don't restart it either.
+		if (TimeNow - StoppedInst->TimeStopped > ElectraCSAudio_MaxStoppedDurationForReactivation)
+		{
+			continue;
+		}
+
+		// Ask to open again.
+		StoppedInst->StartPosition = CurrentTime;
+		StoppedInst->CurrentState = EState::TryReopeningStream;
+		break;
+	}
+}
+
+void FSimpleElectraAudioPlayer::CreatePlayerInstance()
+{
+	if (!Player.IsValid())
+	{
+		// Create a renderer for the new player.
+		Renderer = MakeShared<FSimpleElectraAudioPlayerRenderer, ESPMode::ThreadSafe>(SharedThis(this));
+
+		// Create a resource provider for the new player.
+		ResourceProvider = MakeShared<FResourceProvider, ESPMode::ThreadSafe>();
+
+		TSharedPtr<IAdaptiveStreamingPlayer, ESPMode::ThreadSafe> NewPlayer;
+		FParamDict NoOptions;
+		IAdaptiveStreamingPlayer::FCreateParam cp;
+		cp.ExternalPlayerGUID = CreateParams.InstanceGUID;
+		cp.AudioRenderer = Renderer;
+		NewPlayer = IAdaptiveStreamingPlayer::Create(cp);
+		NewPlayer->AddMetricsReceiver(this);
+		NewPlayer->SetStaticResourceProviderCallback(ResourceProvider);
+		NewPlayer->Initialize(NoOptions);
+		Player = NewPlayer;
+	}
+}
+
+void FSimpleElectraAudioPlayer::DoClosePlayerAsync(TSharedPtr<FClosePlayerInstanceHelper, ESPMode::ThreadSafe>&& InPlayerHelper)
+{
+	TFunction<void()> CloseTask = [InPlayerHelper]()
+	{
+		if (InPlayerHelper->PlayerInstance.IsValid())
+		{
+			InPlayerHelper->PlayerInstance->Stop();
+			InPlayerHelper->PlayerInstance.Reset();
+		}
+		InstanceLock.Lock();
+		ActiveInstances.Remove(InPlayerHelper->This);
+		if (!InPlayerHelper->bDeleteThis)
+		{
+			StoppedInstances.AddUnique(InPlayerHelper->This);
+		}
+		InstanceLock.Unlock();
+		if (InPlayerHelper->bDeleteThis)
+		{
+			MergeAnalytics(InPlayerHelper->This->BaseURL, InPlayerHelper->This->Analytics);
+			InPlayerHelper->This->bDestructionComplete = true;
+		}
+		InPlayerHelper->This->bCloseInProgress = false;
+	};
+
+	if (GIsRunning)
+	{
+		Async(EAsyncExecution::ThreadPool, MoveTemp(CloseTask));
+	}
+	else
+	{
+		CloseTask();
+	}
+}
+
+void FSimpleElectraAudioPlayer::ClosePlayerInstance(bool bDeleteThis)
+{
+	CurrentState = CurrentState == EState::Errored ? CurrentState : EState::Stopped;
+	bCloseInProgress = true;
+
+	TSharedPtr<FClosePlayerInstanceHelper, ESPMode::ThreadSafe> Helper = MakeShared<FClosePlayerInstanceHelper, ESPMode::ThreadSafe>();
+	Helper->PlayerInstance = MoveTemp(Player);
+	Helper->This = this;
+	Helper->bDeleteThis = bDeleteThis;
+	if (Helper->PlayerInstance.IsValid())
+	{
+		Helper->PlayerInstance->SetStaticResourceProviderCallback(nullptr);
+		Helper->PlayerInstance->RemoveMetricsReceiver(this);
+	}
+	DoClosePlayerAsync(MoveTemp(Helper));
+}
+
+bool FSimpleElectraAudioPlayer::Open(const TMap<FString, FVariant>& InOptions, const FString& InManifestURL, const FTimespan& InStartPosition, const FTimespan& InEncodedDuration, bool bInAutoPlay, bool bInSetLooping, TSharedPtr<IElectraPlayerDataCache, ESPMode::ThreadSafe> InPlayerDataCache)
+{
+	if (TimeCreated < 0.0)
+	{
+		TimeCreated = FPlatformTime::Seconds();
+	}
+
+	URL = InManifestURL;
+	Electra::FURL_RFC3986 UrlParser;
+	if (UrlParser.Parse(URL))
+	{
+		BaseURL = UrlParser.Get(false, false);
+	}
+
+	EncodedDuration = InEncodedDuration;
+	StartPosition = InStartPosition;
+	bAutoPlay = bInAutoPlay;
+	bSetLooping = bInSetLooping;
+	Options.Append(InOptions);
+	UpdateAssetName();
+
+	// Check if this is a blob request. These can be performed without having to restrict the number of active players.
+	// TBD: We may also elect to cache successful blob requests.
+	static FString BlobParams(TEXT("blobparams"));
+	const FVariant* BlobVar = nullptr;
+	if ((BlobVar = Options.Find(BlobParams)) != nullptr && BlobVar->GetType() == EVariantTypes::String)
+	{
+		BlobParameters = BlobVar->GetValue<FString>();
+		Options.Remove(BlobParams);
+	}
+	if (BlobParameters.IsSet())
+	{
+		PendingBlobRequest = MakeShared<FBlobRequest, ESPMode::ThreadSafe>();
+		PendingBlobRequest->URL = InManifestURL;
+		TSharedPtr<TArray<uint8>, ESPMode::ThreadSafe> BlobData = GetBlobCache().GetData(InManifestURL);
+		if (BlobData.IsValid())
+		{
+			PendingBlobRequest->bWasCached = true;
+			PendingBlobRequest->SetBlobData(BlobData);
+			CurrentState = EState::OpenBlob;
+		}
+		else
+		{
+			FString Msg = FString::Printf(TEXT("%s: Loading blob"), *AssetName);
+			UE_LOG(LogSimpleElectraPlayer, Verbose, TEXT("%s"), *Msg);
+			CreatePlayerInstance();
+			if (!PendingBlobRequest->Request->SetFromJSON(BlobParameters.GetValue()))
+			{
+				CurrentState = EState::Errored;
+				bHasErrored = true;
+				ErrorMessage = TEXT("Could not parse blob parameters");
+				PendingBlobRequest.Reset();
+				return false;
+			}
+			CurrentState = EState::OpenBlob;
+			PendingBlobRequest->Request->URL(InManifestURL).Callback().BindThreadSafeSP(PendingBlobRequest.ToSharedRef(), &FBlobRequest::OnBlobRequestComplete);
+			Player->LoadBlob(PendingBlobRequest->Request);
+		}
+		BlobParameters.Reset();
+		return true;
+	}
+
+	// Finish the set up and change the state to opening to let the Tick() method take over.
+	PendingBlobRequest.Reset();
+	PlayerDataCache = InPlayerDataCache;
+
+	CurrentState = EState::OpeningStream;
+	return true;
+}
+
+void FSimpleElectraAudioPlayer::InternalApplyOptions()
+{
+	if (ResourceProvider.IsValid())
+	{
+		ResourceProvider->SetOptions(Options);
+		ResourceProvider->SetPlaylistURL(URL);
+	}
+
+	FParamDict Opts, NoOpts;
+	const FVariant* Mimetype = Options.Find(TEXT("mimetype"));
+	if (Mimetype)
+	{
+		Opts.SetOrUpdate(TEXT("mime_type"), FVariantValue(Mimetype->GetValue<FString>()));
+	}
+
+	const FVariant* InitialBitrate = Options.Find(TEXT("initial_bitrate"));
+	if (InitialBitrate)
+	{
+		Opts.SetOrUpdate(TEXT("initial_bitrate"), FVariantValue((int64) InitialBitrate->GetValue<int32>()));
+	}
+
+	TSharedPtr<IAdaptiveStreamingPlayer, ESPMode::ThreadSafe> LockedPlayer(Player);
+	if (LockedPlayer.IsValid())
+	{
+		LockedPlayer->ModifyOptions(Opts, NoOpts);
+	}
+}
+
+void FSimpleElectraAudioPlayer::UpdateAssetName()
+{
+	FString AssetNameParam(TEXT("assetname"));
+	const FVariant* AssetNameVar = nullptr;
+	if ((AssetNameVar = Options.Find(AssetNameParam)) != nullptr && AssetNameVar->GetType() == EVariantTypes::String)
+	{
+		AssetName = AssetNameVar->GetValue<FString>();
+	}
+	else
+	{
+		AssetName = SanitizeMessage(BaseURL);
+	}
+}
+
+void FSimpleElectraAudioPlayer::FResourceProvider::ProvideStaticPlaybackDataForURL(TSharedPtr<IAdaptiveStreamingPlayerResourceRequest, ESPMode::ThreadSafe> InOutRequest)
+{
+	if (InOutRequest->GetResourceType() == Electra::IAdaptiveStreamingPlayerResourceRequest::EPlaybackResourceType::Playlist)
+	{
+		FString URI = InOutRequest->GetResourceURL();
+		if (URI == PlaylistURL)
+		{
+			const FVariant* PlaylistData = Options.Find(TEXT("playlistdata"));
+			if (PlaylistData && PlaylistData->GetType() == EVariantTypes::String)
+			{
+				FString Playlist = PlaylistData->GetValue<FString>();
+				TSharedPtr<TArray<uint8>, ESPMode::ThreadSafe> ResponseDataPtr = MakeShared<TArray<uint8>, ESPMode::ThreadSafe>((const uint8*)TCHAR_TO_UTF8(*Playlist), Playlist.Len());
+				InOutRequest->SetPlaybackData(ResponseDataPtr, 0);
+			}
+		}
+	}
+	else if (InOutRequest->GetResourceType() == Electra::IAdaptiveStreamingPlayerResourceRequest::EPlaybackResourceType::LicenseKey)
+	{
+		const FVariant* KeysData = Options.Find(TEXT("keys"));
+		if (KeysData && KeysData->GetType() == EVariantTypes::String)
+		{
+			FString Keys = KeysData->GetValue<FString>();
+			TSharedPtr<TArray<uint8>, ESPMode::ThreadSafe> ResponseDataPtr = MakeShared<TArray<uint8>, ESPMode::ThreadSafe>((const uint8*)TCHAR_TO_UTF8(*Keys), Keys.Len());
+			InOutRequest->SetPlaybackData(ResponseDataPtr, 0);
+		}
+	}
+	InOutRequest->SignalDataReady();
+}
+
+int32 FSimpleElectraAudioPlayer::GetBinaryMetadata(TSharedPtr<TArray<uint8>, ESPMode::ThreadSafe>& OutMetadata) const
+{
+	FScopeLock lock(&Lock);
+	if (PendingBlobRequest.IsValid())
+	{
+		if (PendingBlobRequest->bIsComplete)
+		{
+			if (PendingBlobRequest->ErrorCode != 0 && !PendingBlobRequest->bReportedError)
+			{
+				++Analytics.ErrorCodeAndCount.FindOrAdd((int32) PendingBlobRequest->ErrorCode);
+
+				PendingBlobRequest->bReportedError = true;
+				const_cast<FSimpleElectraAudioPlayer*>(this)->ReportError(FString::Printf(TEXT("Blob retrieval failed with error %d"), PendingBlobRequest->ErrorCode));
+			}
+			if (PendingBlobRequest->ErrorCode == 0 && !PendingBlobRequest->bWasCached)
+			{
+				GetBlobCache().AddData(PendingBlobRequest->URL, PendingBlobRequest->BlobData);
+			}
+			OutMetadata = PendingBlobRequest->BlobData;
+			return PendingBlobRequest->ErrorCode;
+		}
+	}
+	return -1;
+}
+
+bool FSimpleElectraAudioPlayer::GetStreamFormat(FStreamFormat& OutFormat) const
+{
+	FScopeLock lock(&Lock);
+	OutFormat = StreamFormat;
+	return StreamFormat.SampleRate && StreamFormat.NumChannels;
+}
+
+bool FSimpleElectraAudioPlayer::CanAcceptAudioFrames(int32 InNumFrames)
+{
+	FScopeLock lock(&Lock);
+
+	if (StreamFormat.SampleRate)
+	{
+		double DurationInBuffer = (double)NumSamplesAvailable / StreamFormat.SampleRate;
+		return DurationInBuffer < SuggestedBufferDuration || NumSamplesAvailable < MinNumSamplesInBuffer;
+	}
+	return NumSamplesAvailable < MinNumSamplesInBuffer;
+}
+
+bool FSimpleElectraAudioPlayer::GetEnqueuedFrameInfo(int32& OutNumberOfEnqueuedFrames, FTimespan& OutDurationOfEnqueuedFrames) const
+{
+	FScopeLock lock(&Lock);
+	OutNumberOfEnqueuedFrames = SampleBlocks.Num();
+	if (StreamFormat.SampleRate)
+	{
+		double DurationInBuffer = (double)NumSamplesAvailable / StreamFormat.SampleRate;
+		OutDurationOfEnqueuedFrames = FTimespan::FromSeconds(DurationInBuffer);
+	}
+	else
+	{
+		OutDurationOfEnqueuedFrames.Zero();
+	}
+	return true;
+}
+
+bool FSimpleElectraAudioPlayer::EnqueueAudioFrames(const void* InBufferAddress, int32 InNumSamples, int32 InNumChannels, int32 InSampleRate, const FTimespan& InPTS, const FTimespan& InDuration)
+{
+	if (InBufferAddress && InNumSamples && InNumChannels)
+	{
+		if (TimeUntilReady < 0.0)
+		{
+			TimeUntilReady = FPlatformTime::Seconds() - TimeCreated;
+			Analytics.StartLatency = TimeUntilReady;
+			FString Msg = FString::Printf(TEXT("%s: ready after %.2fs"), *AssetName, TimeUntilReady);
+			UE_LOG(LogSimpleElectraPlayer, Verbose, TEXT("%s"), *Msg);
+		}
+
+		FSampleBlock sb;
+		int32 nb = InNumSamples * InNumChannels * sizeof(float);
+		sb.Buffer = FMemory::Malloc(nb);
+		sb.NumSamples = InNumSamples;
+		sb.PTS = InPTS;
+		FMemory::Memcpy(sb.Buffer, InBufferAddress, nb);
+
+		FScopeLock lock(&Lock);
+
+		// Get first PTS.
+		if (NextPTS < FTimespan::Zero())
+		{
+			NextPTS = InPTS;
+		}
+
+		SampleBlocks.Emplace(MoveTemp(sb));
+		NumSamplesAvailable += InNumSamples;
+		if (StreamFormat.SampleRate == 0 || StreamFormat.NumChannels == 0)
+		{
+			StreamFormat.SampleRate = InSampleRate;
+			StreamFormat.NumChannels = InNumChannels;
+		}
+		bReachedEOS = false;
+		return true;
+	}
+	return false;
+}
+
+void FSimpleElectraAudioPlayer::SetReceivedLastBuffer()
+{
+	FScopeLock lock(&Lock);
+	bReachedEOS = true;
+}
+
+void FSimpleElectraAudioPlayer::FlushAudio(bool bForceFlush)
+{
+	bool bTerminated = HasBeenTerminated();
+	FScopeLock lock(&Lock);
+	// Only flush when not terminated. Otherwise we want to return all the remaining samples we have
+	if (!bTerminated || bForceFlush)
+	{
+		for(int32 i=0; i<SampleBlocks.Num(); ++i)
+		{
+			SampleBlocks[i].Free();
+		}
+		SampleBlocks.Empty();
+		NumSamplesAvailable = 0;
+		bReachedEOS = false;
+		NextPTS = FTimespan::MinValue();
+	}
+}
+
+int64 FSimpleElectraAudioPlayer::GetNumAvailableSamples() const
+{
+	bool bTerminated = HasBeenTerminated();
+	FScopeLock lock(&Lock);
+	if (NumSamplesAvailable)
+	{
+		return NumSamplesAvailable;
+	}
+	if (bTerminated)
+	{
+		return bReachedEOS ? 0 : kNumDummyBlockSamples;
+	}
+	return 0;
+}
+
+bool FSimpleElectraAudioPlayer::IsAtEOS() const
+{
+	FScopeLock lock(&Lock);
+	return bReachedEOS && SampleBlocks.Num() == 0;
+}
+
+bool FSimpleElectraAudioPlayer::HasBeenTerminated() const
+{
+	return CurrentState == EState::Errored || CurrentState == EState::Stopped;
+}
+
+int64 FSimpleElectraAudioPlayer::GetNextSamples(float* OutBuffer, int32 InBufferSizeInSamples, int32 InNumSamplesToGet, const FDefaultSampleInfo& InDefaultSampleInfo)
+{
+	check(!"this is not currently implemented");
+	return 0;
+}
+
+int64 FSimpleElectraAudioPlayer::GetNextSamples(int16* OutBuffer, int32 InBufferSizeInSamples, int32 InNumSamplesToGet, const FDefaultSampleInfo& InDefaultSampleInfo)
+{
+	if (!OutBuffer || InBufferSizeInSamples <= 0 || InNumSamplesToGet <= 0)
+	{
+		return 0;
+	}
+	FScopeLock lock(&Lock);
+	if (SampleBlocks.Num() == 0)
+	{
+		bool bTerminated = HasBeenTerminated();
+		if (!bTerminated)
+		{
+			return 0;
+		}
+		else
+		{
+			int32 nc = StreamFormat.NumChannels > 0 ? StreamFormat.NumChannels : InDefaultSampleInfo.NumChannels;
+			int32 sr = StreamFormat.SampleRate > 0 ? StreamFormat.SampleRate : InDefaultSampleInfo.SampleRate;
+
+			if (InDefaultSampleInfo.ExpectedCurrentSamplePos >= 0 && InDefaultSampleInfo.NumTotalSamples >= 0)
+			{
+				int64 RemainingToGo = InDefaultSampleInfo.NumTotalSamples - InDefaultSampleInfo.ExpectedCurrentSamplePos;
+				if (InNumSamplesToGet > RemainingToGo)
+				{
+					InNumSamplesToGet = RemainingToGo;
+					bReachedEOS = true;
+				}
+				NextPTS = FTimespan::FromSeconds((double)(InDefaultSampleInfo.ExpectedCurrentSamplePos + InNumSamplesToGet) / sr);
+			}
+			else
+			{
+				// Advance the PTS
+				if (NextPTS < FTimespan::Zero())
+				{
+					NextPTS = StartPosition;
+				}
+				NextPTS = NextPTS + FTimespan::FromSeconds((double)InNumSamplesToGet / sr);
+			}
+			if (NextPTS >= Duration)
+			{
+				bReachedEOS = true;
+			}
+
+			// Fill the output with silence.
+			FMemory::Memzero(OutBuffer, InNumSamplesToGet * sizeof(int16) * nc);
+			return InNumSamplesToGet; 
+		}
+	}
+	int32 NumSamplesToGo = InNumSamplesToGet < NumSamplesAvailable ? InNumSamplesToGet : NumSamplesAvailable;
+	int32 NumGot = 0;
+	FTimespan FrontPTS = SampleBlocks[0].PTS;
+	int32 FrontPTSSampleOffset = SampleBlocks[0].Offset;
+	while(NumSamplesToGo)
+	{
+		int32 NumSamplesLeftInBlock = SampleBlocks[0].NumSamples - SampleBlocks[0].Offset;
+		int32 NumBlockSamples = NumSamplesToGo < NumSamplesLeftInBlock ? NumSamplesToGo : NumSamplesLeftInBlock;
+		const float* Src = reinterpret_cast<const float*>(SampleBlocks[0].Buffer) + SampleBlocks[0].Offset * StreamFormat.NumChannels;
+		for(int32 i=0, iMax=NumBlockSamples*StreamFormat.NumChannels; i<iMax; ++i)
+		{
+			*OutBuffer++ = FMath::Clamp((int32)(*Src++ * 32768.0f), (int32)-32768, (int32)32767);
+		}
+		if ((SampleBlocks[0].Offset += NumBlockSamples) >= SampleBlocks[0].NumSamples)
+		{
+			SampleBlocks[0].Free();
+			SampleBlocks.RemoveAt(0);
+		}
+		NumSamplesAvailable -= NumBlockSamples;
+		NumSamplesToGo -= NumBlockSamples;
+		NumGot += NumBlockSamples;
+	}
+	NextPTS = FrontPTS + FTimespan::FromSeconds((double)(FrontPTSSampleOffset + NumGot) / StreamFormat.SampleRate);
+	return NumGot;
+}
+
+TSharedPtr<TArray<uint8>, ESPMode::ThreadSafe> FSimpleElectraAudioPlayer::FBlobCache::GetData(const FString& InForURL)
+{
+	double Now = FPlatformTime::Seconds();
+	FScopeLock lock(&Lock);
+
+	TSharedPtr<FBlobCacheEntry, ESPMode::ThreadSafe>* Entry = EntryMap.Find(InForURL);
+	if (Entry)
+	{
+		// Still good?
+		if (Now - (*Entry)->TimeAdded < ElectraCSAudio_BlobCacheAge)
+		{
+			return (*Entry)->Data;
+		}
+		else
+		{
+			EntryLRU.Remove(*Entry);
+			EntryMap.Remove(InForURL);
+			SizeInUse -= (*Entry)->Data->Num();
+		}
+	}
+	return nullptr;
+}
+
+void FSimpleElectraAudioPlayer::FBlobCache::AddData(const FString& InForURL, const TSharedPtr<TArray<uint8>, ESPMode::ThreadSafe>& InData)
+{
+	double Now = FPlatformTime::Seconds();
+	FScopeLock lock(&Lock);
+
+	TSharedPtr<FBlobCacheEntry, ESPMode::ThreadSafe>* Entry = EntryMap.Find(InForURL);
+	if (!Entry)
+	{
+		CheckLimits(InData->Num());
+		SizeInUse += InData->Num();
+		TSharedPtr<FBlobCacheEntry, ESPMode::ThreadSafe> NewEntry = MakeShared<FBlobCacheEntry, ESPMode::ThreadSafe>();
+		NewEntry->URL = InForURL;
+		NewEntry->TimeAdded = Now;
+		NewEntry->Data = InData;
+		EntryMap.Emplace(InForURL, NewEntry);
+		EntryLRU.Emplace(MoveTemp(NewEntry));
+	}
+	else
+	{
+		(*Entry)->TimeAdded = Now;
+		EntryLRU.Remove(*Entry);
+		EntryLRU.Emplace(*Entry);
+	}
+}
+
+void FSimpleElectraAudioPlayer::FBlobCache::CheckLimits(int32 InBytesNeeded)
+{
+	while(EntryLRU.Num() && SizeInUse + InBytesNeeded > ElectraCSAudio_BlobCacheSize)
+	{
+		TSharedPtr<FBlobCacheEntry, ESPMode::ThreadSafe> Oldest = EntryLRU[0];
+		EntryLRU.RemoveAt(0);
+		EntryMap.Remove(Oldest->URL);
+		SizeInUse -= Oldest->Data->Num();
+	}
+}
+
+
+
+
+
+UEMediaError FSimpleElectraAudioPlayerRenderer::CreateBufferPool(const FParamDict& Parameters)
+{
+	DecoderOutputPool.Reset();
+	NumOutputAudioBuffersInUse = 0;
+	NumBuffers = (int32)Parameters.GetValue("num_buffers").SafeGetInt64(-1);
+	if (NumBuffers <= 0)
+	{
+		return UEMEDIA_ERROR_BAD_ARGUMENTS;
+	}
+	BufferPoolProperties.Set("max_buffers", Electra::FVariantValue((int64)NumBuffers));
+	MaxBufferSize = (uint32)Parameters.GetValue("max_buffer_size").SafeGetInt64(0);
+	return MaxBufferSize ? UEMEDIA_ERROR_OK : UEMEDIA_ERROR_BAD_ARGUMENTS;
+}
+
+UEMediaError FSimpleElectraAudioPlayerRenderer::AcquireBuffer(IBuffer*& OutBuffer, int32 TimeoutInMicroseconds, const FParamDict& InParameters)
+{
+	if (TimeoutInMicroseconds != 0)
+	{
+		return UEMEDIA_ERROR_BAD_ARGUMENTS;
+	}
+	if (NumOutputAudioBuffersInUse >= NumBuffers)
+	{
+		return UEMEDIA_ERROR_INSUFFICIENT_DATA;
+	}
+	IAudioDecoderOutputPtr DecoderOutput = DecoderOutputPool.AcquireShared();
+	DecoderOutput->SetOwner(SharedThis(this));
+	FPlatformAtomics::InterlockedIncrement(&NumBuffersAcquiredForDecoder);
+	DecoderOutput->Reserve(MaxBufferSize);
+	FPlatformAtomics::InterlockedIncrement(&NumOutputAudioBuffersInUse);
+
+	FMediaBufferSharedPtrWrapper* MediaBufferSharedPtrWrapper = new FMediaBufferSharedPtrWrapper(DecoderOutput);
+	MediaBufferSharedPtrWrapper->BufferProperties.Set("size", FVariantValue((int64)DecoderOutput->GetReservedBufferBytes()));
+	MediaBufferSharedPtrWrapper->BufferProperties.Set("address", FVariantValue(const_cast<void*>(DecoderOutput->GetBuffer())));
+	OutBuffer = MediaBufferSharedPtrWrapper;
+	return UEMEDIA_ERROR_OK;
+}
+
+UEMediaError FSimpleElectraAudioPlayerRenderer::ReturnBuffer(IBuffer* Buffer, bool bRender, const FParamDict& InSampleProperties)
+{
+	if (Buffer == nullptr)
+	{
+		return UEMEDIA_ERROR_BAD_ARGUMENTS;
+	}
+	FMediaBufferSharedPtrWrapper* MediaBufferSharedPtrWrapper = static_cast<FMediaBufferSharedPtrWrapper*>(Buffer);
+	if (bRender)
+	{
+		int32 NumChannels = (int32)InSampleProperties.GetValue("num_channels").SafeGetInt64(-1);
+		int32 SampleRate = (int32)InSampleProperties.GetValue("sample_rate").SafeGetInt64(-1);
+		int32 UsedBufferBytes = (int32)InSampleProperties.GetValue("byte_size").SafeGetInt64(-1);
+		FTimeValue decPTS = InSampleProperties.GetValue("pts").SafeGetTimeValue();
+		FTimespan PTS = decPTS.GetAsTimespan();
+		int64 InSequenceIndex = decPTS.GetSequenceIndex();
+
+		FTimespan Duration = InSampleProperties.GetValue("duration").SafeGetTimeValue().GetAsTimespan();
+		const void* BufferAddr = MediaBufferSharedPtrWrapper->BufferProperties.GetValue("address").SafeGetPointer();
+		if (NumChannels <= 0 || SampleRate <= 0 || UsedBufferBytes <= 0)
+		{
+			return UEMEDIA_ERROR_BAD_ARGUMENTS;
+		}
+
+		IAudioDecoderOutputPtr DecoderOutput = MediaBufferSharedPtrWrapper->DecoderOutput;
+		DecoderOutput->GetMutablePropertyDictionary() = InSampleProperties;
+		DecoderOutput->Initialize(IAudioDecoderOutput::ESampleFormat::Float, NumChannels, SampleRate, Duration, FDecoderTimeStamp(PTS, InSequenceIndex), UsedBufferBytes);
+
+		TSharedPtr<FSimpleElectraAudioPlayer, ESPMode::ThreadSafe> PinnedPlayer = Player.Pin();
+		if (PinnedPlayer.IsValid())
+		{
+			PinnedPlayer->EnqueueAudioFrames(BufferAddr, UsedBufferBytes / (sizeof(float)*NumChannels), NumChannels, SampleRate, PTS, Duration);
+		}
+	}
+	else
+	{
+		if (InSampleProperties.GetValue("eos").SafeGetBool(false))
+		{
+			TSharedPtr<FSimpleElectraAudioPlayer, ESPMode::ThreadSafe> PinnedPlayer = Player.Pin();
+			if (PinnedPlayer.IsValid())
+			{
+				PinnedPlayer->SetReceivedLastBuffer();
+			}
+		}
+	}
+	FPlatformAtomics::InterlockedDecrement(&NumBuffersAcquiredForDecoder);
+	delete MediaBufferSharedPtrWrapper;
+	return UEMEDIA_ERROR_OK;
+}
+
+UEMediaError FSimpleElectraAudioPlayerRenderer::ReleaseBufferPool()
+{
+	DecoderOutputPool.Reset();
+	return UEMEDIA_ERROR_OK;
+}
+
+bool FSimpleElectraAudioPlayerRenderer::CanReceiveOutputFrames(uint64 NumFrames) const
+{
+	TSharedPtr<FSimpleElectraAudioPlayer, ESPMode::ThreadSafe> PinnedPlayer = Player.Pin();
+	if (!PinnedPlayer.IsValid())
+	{
+		return false;
+	}
+	return PinnedPlayer->CanAcceptAudioFrames(NumFrames);
+}
+
+bool FSimpleElectraAudioPlayerRenderer::GetEnqueuedFrameInfo(int32& OutNumberOfEnqueuedFrames, FTimeValue& OutDurationOfEnqueuedFrames) const
+{
+	TSharedPtr<FSimpleElectraAudioPlayer, ESPMode::ThreadSafe> PinnedPlayer = Player.Pin();
+	if (!PinnedPlayer.IsValid())
+	{
+		OutNumberOfEnqueuedFrames = 0;
+		OutDurationOfEnqueuedFrames.SetToZero();
+		return false;
+	}
+	FTimespan DurAvail(0);
+	bool bGot = PinnedPlayer->GetEnqueuedFrameInfo(OutNumberOfEnqueuedFrames, DurAvail);
+	if (bGot)
+	{
+		OutDurationOfEnqueuedFrames.SetFromTimespan(DurAvail);
+	}
+	return bGot;
+}
+
+void FSimpleElectraAudioPlayerRenderer::SetParentRenderer(TWeakPtr<IMediaRenderer, ESPMode::ThreadSafe> InParentRenderer)
+{
+	ParentRenderer = MoveTemp(InParentRenderer);
+}
+
+UEMediaError FSimpleElectraAudioPlayerRenderer::Flush(const FParamDict& InOptions)
+{
+	if (NumBuffersAcquiredForDecoder != 0)
+	{
+		return UEMEDIA_ERROR_INTERNAL;
+	}
+	if (TSharedPtr<FSimpleElectraAudioPlayer, ESPMode::ThreadSafe> PinnedPlayer = Player.Pin())
+	{
+		PinnedPlayer->FlushAudio(false);
+	}
+	return UEMEDIA_ERROR_OK;
+}
+
+void FSimpleElectraAudioPlayerRenderer::SampleReleasedToPool(IDecoderOutput* InDecoderOutput)
+{
+	FPlatformAtomics::InterlockedDecrement(&NumOutputAudioBuffersInUse);
+	check(NumOutputAudioBuffersInUse >= 0 && NumOutputAudioBuffersInUse <= NumBuffers)
+
+	TSharedPtr<IMediaRenderer, ESPMode::ThreadSafe> Parent = ParentRenderer.Pin();
+	if (Parent.IsValid())
+	{
+		Parent->SampleReleasedToPool(InDecoderOutput);
+	}
+}
+
+
+#if ENABLE_DEBUG_STATS
+
+static TAutoConsoleVariable<bool> CVarSimpleElectraAudioPlayer_ShowOverlay(
+	TEXT("ElectraCSAudio.ShowOverlay"),
+	false,
+	TEXT("If true, displays a debug overlay of active instances."),
+	FConsoleVariableDelegate::CreateLambda([](IConsoleVariable* CVar)
+		{ FSimpleElectraAudioPlayer::EnableDebugDraw(CVar->GetBool()); }),
+	ECVF_Default);
+
+
+FDelegateHandle FSimpleElectraAudioPlayer::DebugDrawDelegateHandle;
+FVector2D FSimpleElectraAudioPlayer::DebugDrawTextPos;
+UFont* FSimpleElectraAudioPlayer::DebugDrawFont = nullptr;
+int32 FSimpleElectraAudioPlayer::DebugDrawFontRowHeight = 0;
+const FLinearColor FSimpleElectraAudioPlayer::DebugColor_Orange { 1.0f, 0.5f, 0.0f };
+const FLinearColor FSimpleElectraAudioPlayer::DebugColor_Green { 0.0f, 1.0f, 0.0f };
+const FLinearColor FSimpleElectraAudioPlayer::DebugColor_Red { 1.0f, 0.0f, 0.0f };
+
+void FSimpleElectraAudioPlayer::EnableDebugDraw(bool bEnable)
+{
+	if (bEnable && !DebugDrawDelegateHandle.IsValid())
+	{
+		DebugDrawDelegateHandle = UDebugDrawService::Register(TEXT("Game"), FDebugDrawDelegate::CreateStatic(&FSimpleElectraAudioPlayer::DebugDraw));
+	}
+	else if (!bEnable && DebugDrawDelegateHandle.IsValid())
+	{
+		UDebugDrawService::Unregister(DebugDrawDelegateHandle);
+		DebugDrawDelegateHandle.Reset();
+	}
+}
+
+void FSimpleElectraAudioPlayer::DebugDraw(UCanvas* InCanvas, APlayerController* /*InPlayerController*/)
+{
+	if (CVarSimpleElectraAudioPlayer_ShowOverlay.GetValueOnAnyThread() && InCanvas)
+	{
+		if(!DebugDrawFont)
+		{
+			// Pick a larger font on console.
+			//DebugDrawFont = FPlatformProperties::SupportsWindowedMode() ? GEngine->GetSmallFont() : GEngine->GetMediumFont();
+			DebugDrawFont = GEngine->GetMediumFont();
+			DebugDrawFontRowHeight = FMath::TruncToInt(DebugDrawFont->GetMaxCharHeight() * 1.075f);
+		}
+
+		// Reset drawing pos...
+		DebugDrawTextPos.X = 20.0f;
+		DebugDrawTextPos.Y = 20.0f;
+
+
+		InstanceLock.Lock();
+		TArray<FSimpleElectraAudioPlayer*> active;
+		TArray<FSimpleElectraAudioPlayer*> ready;
+		TArray<FSimpleElectraAudioPlayer*> stopped;
+
+		for(int32 i=0; i<AllInstances.Num(); ++i)
+		{
+			if (ActiveInstances.Contains(AllInstances[i]))
+			{
+				active.Emplace(AllInstances[i]);
+			}
+			else if (StoppedInstances.Contains(AllInstances[i]))
+			{
+				stopped.Emplace(AllInstances[i]);
+			}
+			else
+			{
+				ready.Emplace(AllInstances[i]);
+			}
+		}
+
+		for(auto& act : active)
+		{
+			act->DebugDrawInst(InCanvas);
+		}
+		for(auto& rdy : ready)
+		{
+			rdy->DebugDrawInst(InCanvas);
+		}
+		for(auto& stp : stopped)
+		{
+			stp->DebugDrawInst(InCanvas);
+		}
+		InstanceLock.Unlock();
+	}
+}
+
+void FSimpleElectraAudioPlayer::DebugDrawPrintLine(UCanvas* InCanvas, const FString& InString, const FLinearColor& InColor)
+{
+	if (InCanvas)
+	{
+		FCanvasTextItem textItem(DebugDrawTextPos, FText::FromString(InString), DebugDrawFont, InColor);
+		const FLinearColor ShadowColor = FLinearColor::Black;
+		textItem.EnableShadow(ShadowColor);
+		InCanvas->DrawItem(textItem);
+		DebugDrawTextPos.Y += DebugDrawFontRowHeight;
+	}
+}
+
+void FSimpleElectraAudioPlayer::DebugDrawInst(UCanvas* InCanvas)
+{
+	FString Msg(AssetName);
+	Msg += TEXT(": ");
+	TSharedPtr<IAdaptiveStreamingPlayer, ESPMode::ThreadSafe> LockedPlayer(Player);
+
+	FLinearColor Color;
+	if (CurrentState == EState::Stopped) Color = DebugColor_Red;
+	else if (CurrentState == EState::Active) Color = DebugColor_Green;
+	else Color = DebugColor_Orange;
+
+	if (LockedPlayer.IsValid())
+	{
+		auto GetOptionVal = [](const FString& InOptName, const TMap<FString, FVariant>& InOptions) -> int32
+		{
+			const FVariant* Var = nullptr;
+			if ((Var = InOptions.Find(InOptName)) != nullptr && Var->GetType() == EVariantTypes::Int32)
+			{
+				return Var->GetValue<int32>();
+			}
+			return 0;
+		};
+
+		// Time until ready
+		double ttr = TimeUntilReady >= 0.0 ? TimeUntilReady : FPlatformTime::Seconds() - TimeCreated;
+		Msg += FString::Printf(TEXT("(%.2f) "), ttr);
+
+		int32 Retry = GetOptionVal(TEXT("retrynum"), Options);
+		int32 RetryMax = GetOptionVal(TEXT("retrymax"), Options);
+		if (Retry)
+		{
+			Msg += FString::Printf(TEXT(" retry %d of %d "), Retry, RetryMax);
+		}
+
+		Electra::FErrorDetail Err = LockedPlayer->GetError();
+		if (Err.IsOK())
+		{
+			if (PendingBlobRequest.IsValid())
+			{
+				Msg += TEXT("(Retrieving blob)");
+			}
+			else if (!LockedPlayer->HaveMetadata())
+			{
+				Msg += TEXT("(Opening)");
+			}
+			else
+			{
+				if (LockedPlayer->IsSeeking())
+				{
+					Msg += TEXT("(Seeking)");
+				}
+				else if (LockedPlayer->IsBuffering())
+				{
+					Msg += TEXT("(Buffering)");
+				}
+				double pts = NextPTS.GetTotalSeconds();
+				Msg += FString::Printf(TEXT(" %.2f/%.2f"), pts>=0.0?pts:0.0, Duration.GetTotalSeconds());
+			}
+		}
+		else
+		{
+			Msg += Err.GetMessage();
+		}
+	}
+	else
+	{
+		double pts = NextPTS.GetTotalSeconds();
+		Msg += FString::Printf(TEXT(" %.2f/%.2f"), pts>=0.0?pts:0.0, Duration.GetTotalSeconds());
+	}
+	DebugDrawPrintLine(InCanvas, Msg, Color);
+}
+
+#endif
diff --git a/Engine/Plugins/Media/ElectraPlayer/Source/ElectraPlayerRuntime/Private/Runtime/Renderer/RendererAudio.cpp b/Engine/Plugins/Media/ElectraPlayer/Source/ElectraPlayerRuntime/Private/Runtime/Renderer/RendererAudio.cpp
--- a/Engine/Plugins/Media/ElectraPlayer/Source/ElectraPlayerRuntime/Private/Runtime/Renderer/RendererAudio.cpp
+++ b/Engine/Plugins/Media/ElectraPlayer/Source/ElectraPlayerRuntime/Private/Runtime/Renderer/RendererAudio.cpp
@@ -214,6 +214,13 @@ bool FElectraRendererAudio::CanReceiveOutputFrames(uint64 NumFrames) const
 	return PinnedPlayer->CanPresentAudioFrames(NumFrames);
 }
 
+bool FElectraRendererAudio::GetEnqueuedFrameInfo(int32& OutNumberOfEnqueuedFrames, Electra::FTimeValue& OutDurationOfEnqueuedFrames) const
+{
+	OutNumberOfEnqueuedFrames = 0;
+	OutDurationOfEnqueuedFrames.SetToZero();
+	return false;
+}
+
 /**
  *
  */
diff --git a/Engine/Plugins/Media/ElectraPlayer/Source/ElectraPlayerRuntime/Private/Runtime/Renderer/RendererAudio.h b/Engine/Plugins/Media/ElectraPlayer/Source/ElectraPlayerRuntime/Private/Runtime/Renderer/RendererAudio.h
--- a/Engine/Plugins/Media/ElectraPlayer/Source/ElectraPlayerRuntime/Private/Runtime/Renderer/RendererAudio.h
+++ b/Engine/Plugins/Media/ElectraPlayer/Source/ElectraPlayerRuntime/Private/Runtime/Renderer/RendererAudio.h
@@ -73,6 +73,8 @@ public:
 
 	bool CanReceiveOutputFrames(uint64 NumFrames) const override;
 
+	bool GetEnqueuedFrameInfo(int32& OutNumberOfEnqueuedFrames, Electra::FTimeValue& OutDurationOfEnqueuedFrames) const override;
+
 	// Receives the render clock we need to update with the most recently rendered sample's timestamp.
 	void SetRenderClock(TSharedPtr<Electra::IMediaRenderClock, ESPMode::ThreadSafe> RenderClock) override;
 
diff --git a/Engine/Plugins/Media/ElectraPlayer/Source/ElectraPlayerRuntime/Private/Runtime/Renderer/RendererBase.h b/Engine/Plugins/Media/ElectraPlayer/Source/ElectraPlayerRuntime/Private/Runtime/Renderer/RendererBase.h
--- a/Engine/Plugins/Media/ElectraPlayer/Source/ElectraPlayerRuntime/Private/Runtime/Renderer/RendererBase.h
+++ b/Engine/Plugins/Media/ElectraPlayer/Source/ElectraPlayerRuntime/Private/Runtime/Renderer/RendererBase.h
@@ -151,6 +151,13 @@ namespace Electra
 		 */
 		virtual bool CanReceiveOutputFrames(uint64 NumFrames) const = 0;
 
+		/**
+		 * Returns how many frames of ready to use data the renderer has in its queue.
+		 * 
+		 * @return True if information is available, false if not.
+		 */
+		virtual bool GetEnqueuedFrameInfo(int32& OutNumberOfEnqueuedFrames, FTimeValue& OutDurationOfEnqueuedFrames) const = 0;
+
 		//=================================================================================================================
 		// Methods called from the player (and from within the player thread)
 		//
diff --git a/Engine/Plugins/Media/ElectraPlayer/Source/ElectraPlayerRuntime/Private/Runtime/Renderer/RendererVideo.cpp b/Engine/Plugins/Media/ElectraPlayer/Source/ElectraPlayerRuntime/Private/Runtime/Renderer/RendererVideo.cpp
--- a/Engine/Plugins/Media/ElectraPlayer/Source/ElectraPlayerRuntime/Private/Runtime/Renderer/RendererVideo.cpp
+++ b/Engine/Plugins/Media/ElectraPlayer/Source/ElectraPlayerRuntime/Private/Runtime/Renderer/RendererVideo.cpp
@@ -238,6 +238,13 @@ bool FElectraRendererVideo::CanReceiveOutputFrames(uint64 NumFrames) const
 	return PinnedPlayer->CanPresentVideoFrames(NumFrames);
 }
 
+bool FElectraRendererVideo::GetEnqueuedFrameInfo(int32& OutNumberOfEnqueuedFrames, Electra::FTimeValue& OutDurationOfEnqueuedFrames) const
+{
+	OutNumberOfEnqueuedFrames = 0;
+	OutDurationOfEnqueuedFrames.SetToZero();
+	return false;
+}
+
 /**
  *
  */
diff --git a/Engine/Plugins/Media/ElectraPlayer/Source/ElectraPlayerRuntime/Private/Runtime/Renderer/RendererVideo.h b/Engine/Plugins/Media/ElectraPlayer/Source/ElectraPlayerRuntime/Private/Runtime/Renderer/RendererVideo.h
--- a/Engine/Plugins/Media/ElectraPlayer/Source/ElectraPlayerRuntime/Private/Runtime/Renderer/RendererVideo.h
+++ b/Engine/Plugins/Media/ElectraPlayer/Source/ElectraPlayerRuntime/Private/Runtime/Renderer/RendererVideo.h
@@ -78,6 +78,8 @@ public:
 	UEMediaError ReleaseBufferPool() override;
 
 	bool CanReceiveOutputFrames(uint64 NumFrames) const override;
+	
+	bool GetEnqueuedFrameInfo(int32& OutNumberOfEnqueuedFrames, Electra::FTimeValue& OutDurationOfEnqueuedFrames) const override;
 
 	// Receives the render clock we need to update with the most recently rendered sample's timestamp.
 	void SetRenderClock(TSharedPtr<Electra::IMediaRenderClock, ESPMode::ThreadSafe> RenderClock) override;
diff --git a/Engine/Plugins/Media/ElectraPlayer/Source/ElectraPlayerRuntime/Public/SimpleElectraAudioPlayer.h b/Engine/Plugins/Media/ElectraPlayer/Source/ElectraPlayerRuntime/Public/SimpleElectraAudioPlayer.h
new file mode 100644
--- /dev/null
+++ b/Engine/Plugins/Media/ElectraPlayer/Source/ElectraPlayerRuntime/Public/SimpleElectraAudioPlayer.h
@@ -0,0 +1,79 @@
+// Copyright Epic Games, Inc. All Rights Reserved.
+#pragma once
+
+#include "CoreMinimal.h"
+#include "Containers/Map.h"
+#include "Misc/Variant.h"
+
+class IElectraPlayerDataCache;
+class IAnalyticsProviderET;
+
+class ELECTRAPLAYERRUNTIME_API ISimpleElectraAudioPlayer
+{
+public:
+	struct FCreateParams
+	{
+		/** The number of audio players that can be simultaneously active at any time. 0=unlimited */
+		int32 MaxTotalPlayerInstances = 0;
+
+		/** A GUID to identify the new player instance. */
+		FGuid InstanceGUID;
+	};
+	static TSharedPtr<ISimpleElectraAudioPlayer, ESPMode::ThreadSafe> Create(const FCreateParams& InCreateParams);
+
+	/**
+	 * Delivers the aggregated metrics to the analytics provider and clears the internal list.
+	 */
+	static void SendAnalyticMetrics(const TSharedPtr<IAnalyticsProviderET>& InAnalyticsProvider);
+
+	virtual ~ISimpleElectraAudioPlayer() = default;
+	virtual bool Open(const TMap<FString, FVariant>& InOptions, const FString& ManifestURL, const FTimespan& StartPosition, const FTimespan& EncodedDuration, bool bAutoPlay, bool bSetLooping, TSharedPtr<IElectraPlayerDataCache, ESPMode::ThreadSafe> InPlayerDataCache) = 0;
+	virtual void SeekTo(const FTimespan& NewPosition) = 0;
+	virtual void Pause() = 0;
+	virtual void Resume() = 0;
+	virtual void Stop() = 0;
+
+	virtual bool HasErrored() const = 0;
+	virtual FString GetError() const = 0;
+
+	struct FStreamFormat
+	{
+		int32 SampleRate = 0;
+		int32 NumChannels = 0;
+	};
+
+	struct FDefaultSampleInfo
+	{
+		int64 NumTotalSamples = -1;
+		int64 ExpectedCurrentSamplePos = -1;
+		int32 SampleRate = 0;
+		int32 NumChannels = 0;
+	};
+
+	virtual int64 GetNumAvailableSamples() const = 0;
+	virtual int64 GetNextSamples(float* OutBuffer, int32 InBufferSizeInSamples, int32 InNumSamplesToGet, const FDefaultSampleInfo& InDefaultSampleInfo) = 0;
+	virtual int64 GetNextSamples(int16* OutBuffer, int32 InBufferSizeInSamples, int32 InNumSamplesToGet, const FDefaultSampleInfo& InDefaultSampleInfo) = 0;
+	virtual bool IsAtEOS() const = 0;
+
+
+	//-------------------------------------------------------------------------
+	// State functions
+	//
+	virtual bool HaveMetadata() const = 0;
+	virtual int32 GetBinaryMetadata(TSharedPtr<TArray<uint8>, ESPMode::ThreadSafe>& OutMetadata) const = 0;
+	virtual bool GetStreamFormat(FStreamFormat& OutFormat) const = 0;
+	virtual FTimespan GetDuration() const = 0;
+	virtual FTimespan GetPlayPosition() const = 0;
+	virtual bool HasEnded() const = 0;
+	virtual bool IsBuffering() const = 0;
+	virtual bool IsSeeking() const = 0;
+	virtual bool IsPlaying() const = 0;
+	virtual bool IsPaused() const = 0;
+
+	//-------------------------------------------------------------------------
+	// Manual stream selection functions
+	//
+
+	//! Sets the highest bitrate when selecting a candidate stream.
+	virtual void SetBitrateCeiling(int32 HighestSelectableBitrate) = 0;
+};
diff --git a/Engine/Source/Editor/DetailCustomizations/Private/SoundWaveDetails.cpp b/Engine/Source/Editor/DetailCustomizations/Private/SoundWaveDetails.cpp
--- a/Engine/Source/Editor/DetailCustomizations/Private/SoundWaveDetails.cpp
+++ b/Engine/Source/Editor/DetailCustomizations/Private/SoundWaveDetails.cpp
@@ -18,6 +18,7 @@
 #include "ScopedTransaction.h"
 #include "SlotBase.h"
 #include "Sound/SoundWave.h"
+#include "ISoundWaveCloudStreaming.h"
 #include "Templates/Casts.h"
 #include "Types/SlateEnums.h"
 #include "UObject/NameTypes.h"
@@ -45,6 +46,7 @@ TSharedRef<IDetailCustomization> FSoundWaveDetails::MakeInstance()
 void FSoundWaveDetails::CustomizeDetails(IDetailLayoutBuilder& DetailBuilder)
 {
 	CustomizeCurveDetails(DetailBuilder);
+	CustomizeCloudStreamingPlatformDetails(DetailBuilder);
 
 	if(!FModuleManager::Get().IsModuleLoaded("WaveformTransformations"))
 	{
@@ -174,4 +176,47 @@ FReply FSoundWaveDetails::HandleUseInternalCurves(USoundWave* SoundWave, TShared
 	return FReply::Handled();
 }
 
+void FSoundWaveDetails::CustomizeCloudStreamingPlatformDetails(IDetailLayoutBuilder& DetailBuilder)
+{
+#if WITH_EDITOR
+	IModularFeatures::FScopedLockModularFeatureList ScopedLockModularFeatureList;
+	TArray<Audio::ISoundWaveCloudStreamingFeature*> Features = IModularFeatures::Get().GetModularFeatureImplementations<Audio::ISoundWaveCloudStreamingFeature>(Audio::ISoundWaveCloudStreamingFeature::GetModularFeatureName());
+	bool bEnabled = false;
+	for(int32 i=0; i<Features.Num(); ++i)
+	{
+		if (Features[i]->AddCustomizationCloudStreamingPlatformDetails(DetailBuilder))
+		{
+			bEnabled = true;
+			break;
+		}
+	}
+	if (!bEnabled)
+	{
+		DetailBuilder.HideCategory("Platform specific");
+
+		// If there are objects that have the cloud streaming flag set without there being a suitable feature plugin available now,
+		// we do not hide the flag so that it can be disabled.
+		TArray<TWeakObjectPtr<UObject>> Objects;
+		DetailBuilder.GetObjectsBeingCustomized(Objects);
+		Objects = Objects.FilterByPredicate([](const TWeakObjectPtr<UObject>& ObjectPtr) { return ObjectPtr.IsValid() && ObjectPtr->IsA<USoundWave>() && Cast<USoundWave>(ObjectPtr)->IsCloudStreamingEnabled(); });
+		TSharedPtr<IPropertyHandle> CloudStreamingProperty = DetailBuilder.GetProperty(USoundWave::GetCloudStreamingEnabledPropertyName());
+		if (CloudStreamingProperty.IsValid())
+		{
+			if (Objects.Num() == 0)
+			{
+				DetailBuilder.HideProperty(CloudStreamingProperty);
+			}
+			else
+			{
+				FText CurrentToolTipText = CloudStreamingProperty->GetToolTipText();
+				FText NewToolTipText = FText::FromString(FString::Printf(TEXT("%s\n\nNo suitable plugin has been found. You can disable this flag to get back to non-cloud streaming settings."), *CurrentToolTipText.ToString()));
+				CloudStreamingProperty->SetToolTipText(NewToolTipText);
+			}
+		}
+	}
+#endif
+}
+
+
+
 #undef LOCTEXT_NAMESPACE
diff --git a/Engine/Source/Editor/DetailCustomizations/Private/SoundWaveDetails.h b/Engine/Source/Editor/DetailCustomizations/Private/SoundWaveDetails.h
--- a/Engine/Source/Editor/DetailCustomizations/Private/SoundWaveDetails.h
+++ b/Engine/Source/Editor/DetailCustomizations/Private/SoundWaveDetails.h
@@ -23,6 +23,8 @@ private:
 
 	void CustomizeCurveDetails(IDetailLayoutBuilder& DetailBuilder);
 
+	void CustomizeCloudStreamingPlatformDetails(IDetailLayoutBuilder& DetailBuilder);
+
 	EVisibility GetMakeInternalCurvesVisibility(USoundWave* SoundWave, TSharedRef<IPropertyHandle> CurvePropertyHandle) const;
 
 	EVisibility GetUseInternalCurvesVisibility(USoundWave* SoundWave, TSharedRef<IPropertyHandle> CurvePropertyHandle) const;
diff --git a/Engine/Source/Runtime/AudioExtensions/Public/ISoundWaveCloudStreaming.h b/Engine/Source/Runtime/AudioExtensions/Public/ISoundWaveCloudStreaming.h
new file mode 100644
--- /dev/null
+++ b/Engine/Source/Runtime/AudioExtensions/Public/ISoundWaveCloudStreaming.h
@@ -0,0 +1,114 @@
+// Copyright Epic Games, Inc. All Rights Reserved.
+
+#pragma once
+
+#include "UObject/Object.h"
+#include "Features/IModularFeature.h"
+#include "Misc/Guid.h"
+
+#include "ISoundWaveCloudStreaming.generated.h"
+
+//
+// Forward declarations
+//
+class USoundWave;
+class IDetailLayoutBuilder;
+
+
+namespace Audio
+{
+
+	class ISoundWaveCloudStreamingFeature : public IModularFeature
+	{
+	public:
+		static FName GetModularFeatureName() { return TEXT("SoundWaveCloudStreaming"); }
+		virtual ~ISoundWaveCloudStreamingFeature() = default;
+
+		/** Returns the GUID of this plugin. */
+		virtual FGuid GetPluginGUID() const = 0;
+
+		/** Checks if the given sound wave can be turned into cloud streamable. */
+		virtual bool CanOverrideFormat(const USoundWave* InWaveToOverride) = 0;
+		
+		/** Gets the format name to use when overriding the given sound wave for cloud streaming. */
+		virtual FName GetOverrideFormatName(const USoundWave* InWaveToOverride) = 0;
+#if WITH_EDITOR
+		/** Gets a hash of the parameters for the DDC. */
+		virtual FString GetOverrideParameterDDCHash(const USoundWave* InWaveToOverride) = 0;
+
+		/** Add editor customization for an instance. */
+		virtual bool AddCustomizationCloudStreamingPlatformDetails(IDetailLayoutBuilder& InDetailLayoutBuilder) = 0;
+#endif // WITH_EDITOR
+	};
+
+}
+
+
+/** Platform specific enabling of Sound Wave cloud streaming. */
+UENUM()
+enum class ESoundWaveCloudStreamingPlatformProjectEnableType : uint8
+{
+	/** Enabled for this platform. */
+	Enabled,
+
+	/** Disabled for this platform. */
+	Disabled,
+};
+
+/** Platform specific settings for Sound Wave cloud streaming. */
+USTRUCT()
+struct FSoundWaveCloudStreamingPlatformProjectSettings
+{
+	GENERATED_USTRUCT_BODY()
+
+	FSoundWaveCloudStreamingPlatformProjectSettings()
+	{
+		EnablementSetting = ESoundWaveCloudStreamingPlatformProjectEnableType::Disabled;
+	}
+
+	bool IsDefault() const
+	{
+		return EnablementSetting == ESoundWaveCloudStreamingPlatformProjectEnableType::Disabled;
+	}
+
+	/** Overrides whether to use cloud streaming on this platform. */
+	UPROPERTY(EditAnywhere, Category = Platforms)
+	ESoundWaveCloudStreamingPlatformProjectEnableType EnablementSetting;
+};
+
+
+
+/** Platform specific enabling of Sound Wave cloud streaming. */
+UENUM()
+enum class ESoundWaveCloudStreamingPlatformEnableType : uint8
+{
+	/** Use Sound Wave setting. */
+	Inherited,
+
+	/** Disables Sound Wave cloud streaming for this platform. */
+	Disabled,
+
+	/** Used in Slate widget configuration to indicate multiple selected objects have different values. */
+	SWC_MultipleValues UMETA(Hidden, DisplayName="Multiple values")
+};
+
+/** Platform specific settings for Sound Wave cloud streaming. */
+USTRUCT()
+struct FSoundWaveCloudStreamingPlatformSettings
+{
+	GENERATED_USTRUCT_BODY()
+
+	FSoundWaveCloudStreamingPlatformSettings()
+	{
+		EnablementSetting = ESoundWaveCloudStreamingPlatformEnableType::Inherited;
+	}
+
+	bool IsDefault() const
+	{
+		return EnablementSetting == ESoundWaveCloudStreamingPlatformEnableType::Inherited;
+	}
+
+	/** Overrides whether to use cloud streaming on this platform. */
+	UPROPERTY(EditAnywhere, Category = Platforms)
+	ESoundWaveCloudStreamingPlatformEnableType EnablementSetting;
+};
diff --git a/Engine/Source/Runtime/Engine/Classes/Sound/SampleBufferIO.h b/Engine/Source/Runtime/Engine/Classes/Sound/SampleBufferIO.h
--- a/Engine/Source/Runtime/Engine/Classes/Sound/SampleBufferIO.h
+++ b/Engine/Source/Runtime/Engine/Classes/Sound/SampleBufferIO.h
@@ -173,12 +173,17 @@ namespace Audio
 		// This writes out the InSampleBuffer as a wav file at the path specified by FilePath and FileName.
 		// If FilePath is a relative path, it will be relative to the /Saved/BouncedWavFiles folder, otherwise specified absolute path will be used.
 		// FileName should not contain the extension. This can be used in non-editor builds.
-		ENGINE_API bool BeginWriteToWavFile(const TSampleBuffer<>& InSampleBuffer, const FString& FileName, FString& FilePath, TFunction<void()> OnSuccess = []() {});
+		ENGINE_API bool BeginWriteToWavFile(const TSampleBuffer<>& InSampleBuffer, const FString& FileName, const FString& FilePath, TFunction<void()> OnSuccess = []() {});
 
 		// This is a blocking call that will return the SoundWave generated from InSampleBuffer.
 		// Optionally, if you're using the editor, you can also write the resulting soundwave out to the content browser using the FileName and FilePath parameters.
 		ENGINE_API USoundWave* SynchronouslyWriteSoundWave(const TSampleBuffer<>& InSampleBuffer, const FString* FileName = nullptr, const FString* FilePath = nullptr);
 
+		// This blocking call writes out the InSampleBuffer as a wav file at the path specified by FilePath and FileName.
+		// If FilePath is a relative path, it will be relative to the /Saved/BouncedWavFiles folder, otherwise specified absolute path will be used.
+		// FileName should not contain the extension. This can be used in non-editor builds.
+		ENGINE_API bool SynchronouslyWriteToWavFile(const TSampleBuffer<>& InSampleBuffer, const FString& FileName, const FString& FilePath, FString* OutFilePathName = nullptr);
+
 		// Call this on the game thread to continue the write operation. Optionally provide a pointer
 		// to an ESoundWavePCMWriterState which will be written to with the current state of the write operation.
 		// Returns a float value from 0 to 1 indicating how complete the write operation is.
@@ -241,6 +246,9 @@ namespace Audio
 
 		// This checks to see if a directory exists and, if it does not, recursively adds the directory.
 		ENGINE_API bool CreateDirectoryIfNeeded(FString& DirectoryPath);
+
+		// Do some shared prep steps for WAV file export
+		bool PrepWavFileOutput(const TSampleBuffer<>& InSampleBuffer, const FString& FileName, const FString& FilePath);
 	};
 
 	/************************************************************************/
diff --git a/Engine/Source/Runtime/Engine/Classes/Sound/SoundWave.h b/Engine/Source/Runtime/Engine/Classes/Sound/SoundWave.h
--- a/Engine/Source/Runtime/Engine/Classes/Sound/SoundWave.h
+++ b/Engine/Source/Runtime/Engine/Classes/Sound/SoundWave.h
@@ -28,6 +28,7 @@
 #include "ContentStreaming.h"
 #include "IAudioProxyInitializer.h"
 #include "IWaveformTransformation.h"
+#include "ISoundWaveCloudStreaming.h"
 #include "Templates/DontCopy.h"
 #include "SoundWave.generated.h"
 
@@ -131,7 +132,7 @@ struct FStreamedAudioPlatformData
 	mutable TDontCopy<FRWLock> AsyncTaskLock;
 	/** Async cache task if one is outstanding. */
 	struct FStreamedAudioAsyncCacheDerivedDataTask* AsyncTask;
-#endif // WITH_EDITORONLY_DATA
+#endif // #if WITH_EDITORONLY_DATA
 
 	/** Default constructor. */
 	ENGINE_API FStreamedAudioPlatformData();
@@ -422,7 +423,7 @@ class USoundWave : public USoundBase, public IAudioProxyDataFactory, public IInt
 private:
 
 	/** Platform agnostic compression quality. 1..100 with 1 being best compression and 100 being best quality. ADPCM and PCM sound asset compression types ignore this parameter. */
-	UPROPERTY(EditAnywhere, Category = "Format|Quality", meta = (DisplayName = "Compression", ClampMin = "1", ClampMax = "100", EditCondition = "SoundAssetCompressionType != ESoundAssetCompressionType::PCM || SoundAssetCompressionType != ESoundAssetCompressionType::ADPCM"), AssetRegistrySearchable)
+	UPROPERTY(EditAnywhere, Category = "Format|Quality", meta = (DisplayName = "Compression", ClampMin = "1", ClampMax = "100", EditCondition = "SoundAssetCompressionType != ESoundAssetCompressionType::PCM && SoundAssetCompressionType != ESoundAssetCompressionType::ADPCM"), AssetRegistrySearchable)
 	int32 CompressionQuality;
 
 public:
@@ -829,6 +830,24 @@ protected:
 	UPROPERTY()
 	TObjectPtr<class UCurveTable> InternalCurves;
 
+#if WITH_EDITORONLY_DATA
+protected:
+	/** If enabled, this wave may be streamed from the cloud using the Opus format. Loading behavior must NOT be `Force Inline`. Requires a suitable support plugin to be installed. */
+	UPROPERTY(EditAnywhere, Category = "Format", Meta=(DisplayName="Enable cloud streaming", DisplayAfter="SoundAssetCompressionType", EditCondition = "LoadingBehavior != ESoundWaveLoadingBehavior::ForceInline"), AssetRegistrySearchable)
+	uint8 bEnableCloudStreaming : 1;
+	/** Platform specific. */
+	UPROPERTY(EditAnywhere, config, Category="Platform specific", Meta=(DisplayName="Platform specific settings", ToolTip="Optionally disables cloud streaming per platform"))
+	TMap<FGuid, FSoundWaveCloudStreamingPlatformSettings> PlatformSettings;
+public:
+	static FName GetCloudStreamingEnabledPropertyName() { return GET_MEMBER_NAME_CHECKED(USoundWave, bEnableCloudStreaming); }
+	ENGINE_API void SetCloudStreamingEnabled(bool bEnabled);
+	ENGINE_API bool IsCloudStreamingEnabled() const;
+	ENGINE_API void TriggerRecookForCloudStreaming();
+	static FName GetCloudStreamingPlatformSettingsPropertyName() { return GET_MEMBER_NAME_CHECKED(USoundWave, PlatformSettings); }
+	ENGINE_API TMap<FGuid, FSoundWaveCloudStreamingPlatformSettings>& GetCloudStreamingPlatformSettings() { return PlatformSettings; }
+	ENGINE_API const TMap<FGuid, FSoundWaveCloudStreamingPlatformSettings>& GetCloudStreamingPlatformSettings() const { return PlatformSettings; }
+#endif // WITH_EDITORONLY_DATA
+
 public:
 	/**
 	* helper function for getting the cached name of the current platform.
@@ -901,6 +920,7 @@ public:
 	ENGINE_API virtual void BeginDestroy() override;
 #if WITH_EDITOR
 	ENGINE_API virtual void PostEditChangeProperty(FPropertyChangedEvent& PropertyChangedEvent) override;
+	ENGINE_API virtual bool CanEditChange(const FProperty* InProperty) const override;
 
 	/** IInterface_AsyncCompilation begin*/
 	ENGINE_API virtual bool IsCompiling() const override;
diff --git a/Engine/Source/Runtime/Engine/Private/AudioDerivedData.cpp b/Engine/Source/Runtime/Engine/Private/AudioDerivedData.cpp
--- a/Engine/Source/Runtime/Engine/Private/AudioDerivedData.cpp
+++ b/Engine/Source/Runtime/Engine/Private/AudioDerivedData.cpp
@@ -23,6 +23,7 @@
 #include "Containers/UnrealString.h"
 #include "Sound/StreamedAudioChunkSeekTable.h"
 #include "AudioDecompress.h"
+#include "ISoundWaveCloudStreaming.h"
 
 DEFINE_LOG_CATEGORY_STATIC(LogAudioDerivedData, Log, All);
 
@@ -198,6 +199,22 @@ static FString GetSoundWaveHash(const USoundWave& InWave, const ITargetPlatform*
 	FPCU::AppendHash(SoundWaveHash, TEXT("SRQ"), InWave.SampleRateQuality);
 	FPCU::AppendHash(SoundWaveHash, TEXT("CK1"), InWave.GetSizeOfFirstAudioChunkInSeconds(InTargetPlatform));
 
+	// Add cloud streaming parameters, if available and enabled.
+	if (InWave.IsCloudStreamingEnabled())
+	{
+		IModularFeatures::FScopedLockModularFeatureList ScopedLockModularFeatureList;
+		TArray<Audio::ISoundWaveCloudStreamingFeature*> Features = IModularFeatures::Get().GetModularFeatureImplementations<Audio::ISoundWaveCloudStreamingFeature>(Audio::ISoundWaveCloudStreamingFeature::GetModularFeatureName());
+		for(int32 i=0; i<Features.Num(); ++i)
+		{
+			if (Features[i]->CanOverrideFormat(&InWave))
+			{
+				FString Hash = Features[i]->GetOverrideParameterDDCHash(&InWave);
+				FPCU::AppendHash(SoundWaveHash, TEXT("CSP"), Hash);
+				break;
+			}
+		}
+	}
+
 	return SoundWaveHash;
 }
 
diff --git a/Engine/Source/Runtime/Engine/Private/AudioFormatSettings.cpp b/Engine/Source/Runtime/Engine/Private/AudioFormatSettings.cpp
--- a/Engine/Source/Runtime/Engine/Private/AudioFormatSettings.cpp
+++ b/Engine/Source/Runtime/Engine/Private/AudioFormatSettings.cpp
@@ -5,6 +5,7 @@
 #include "Sound/SoundWave.h"
 #include "Misc/ConfigCacheIni.h"
 #include "Audio.h"
+#include "ISoundWaveCloudStreaming.h"
 
 namespace Audio
 {
@@ -19,6 +20,32 @@ namespace Audio
 		return !IsAudioLinkEnabled;
 	}
 
+	FName GetCloudStreamingFormatOverride(const FName& InCurrentFormat, const USoundWave* InWave)
+	{
+#if WITH_EDITORONLY_DATA
+		// Is a cloud streaming feature available?
+		if (InWave->IsCloudStreamingEnabled())
+		{
+			IModularFeatures::FScopedLockModularFeatureList ScopedLockModularFeatureList;
+			TArray<Audio::ISoundWaveCloudStreamingFeature*> Features = IModularFeatures::Get().GetModularFeatureImplementations<Audio::ISoundWaveCloudStreamingFeature>(Audio::ISoundWaveCloudStreamingFeature::GetModularFeatureName());
+			// If there is more than one cloud streaming feature it will be ambiguous which one to use.
+			check(Features.Num() <= 1);
+			for(int32 i=0; i<Features.Num(); ++i)
+			{
+				if (Features[i]->CanOverrideFormat(InWave))
+				{
+					FName NewFormatName = Features[i]->GetOverrideFormatName(InWave);
+					if (NewFormatName.GetStringLength())
+					{
+						return NewFormatName;
+					}
+				}
+			}
+		}
+#endif // WITH_EDITORONLY_DATA
+		return InCurrentFormat;
+	}
+
 	FAudioFormatSettings::FAudioFormatSettings(FConfigCacheIni* InConfigSystem, const FString& InConfigFilename, const FString& InPlatformIdentifierForLogging)
 	{
 		ReadConfiguration(InConfigSystem, InConfigFilename, InPlatformIdentifierForLogging);
@@ -33,18 +60,19 @@ namespace Audio
 			{
 				if (Wave->IsStreaming())
 				{
-					return PlatformStreamingFormat;
+					FormatName = PlatformStreamingFormat;
 				}
 				else
 				{
-					return PlatformFormat;
+					FormatName = PlatformFormat;
 				}
 			}
 			else
 			{
-				return FallbackFormat;
+				FormatName = FallbackFormat;
 			}
 		}
+		FormatName = GetCloudStreamingFormatOverride(FormatName, Wave);
 		return FormatName;
 	}
 
diff --git a/Engine/Source/Runtime/Engine/Private/SampleBufferIO.cpp b/Engine/Source/Runtime/Engine/Private/SampleBufferIO.cpp
--- a/Engine/Source/Runtime/Engine/Private/SampleBufferIO.cpp
+++ b/Engine/Source/Runtime/Engine/Private/SampleBufferIO.cpp
@@ -213,14 +213,8 @@ namespace Audio
 		return true;
 	}
 
-	bool FSoundWavePCMWriter::BeginWriteToWavFile(const TSampleBuffer<>& InSampleBuffer, const FString& FileName, FString& FilePath, TFunction<void()> OnSuccess)
+	bool FSoundWavePCMWriter::PrepWavFileOutput(const TSampleBuffer<>& InSampleBuffer, const FString& FileName, const FString& FilePath)
 	{
-		if (!IsDone())
-		{
-			UE_LOG(LogAudio, Error, TEXT("This instance of FSoundWavePCMWriter is already processing another write operation."));
-			return false;
-		}
-
 		const bool bIsRelativePath = FPaths::IsRelative(FilePath);
 		if (bIsRelativePath)
 		{
@@ -252,6 +246,22 @@ namespace Audio
 
 		CurrentBuffer = InSampleBuffer;
 
+		return true;
+	}
+
+	bool FSoundWavePCMWriter::BeginWriteToWavFile(const TSampleBuffer<>& InSampleBuffer, const FString& FileName, const FString& FilePath, TFunction<void()> OnSuccess)
+	{
+		if (!IsDone())
+		{
+			UE_LOG(LogAudio, Error, TEXT("This instance of FSoundWavePCMWriter is already processing another write operation."));
+			return false;
+		}
+
+		if (!PrepWavFileOutput(InSampleBuffer, FileName, FilePath))
+		{
+			return false;
+		}
+
 		// For convenience in our async task, we only take void(USoundWave*) type lambdas. So let's wrap our void() lambda here:
 		TFunction<void(const USoundWave*)> WrappedCallback = [OnSuccess](const USoundWave*)
 		{
@@ -264,6 +274,30 @@ namespace Audio
 		return true;
 	}
 
+	bool FSoundWavePCMWriter::SynchronouslyWriteToWavFile(const TSampleBuffer<>& InSampleBuffer, const FString& FileName, const FString& FilePath, FString* OutFilePathName)
+	{
+		if (!IsDone())
+		{
+			UE_LOG(LogAudio, Error, TEXT("This instance of FSoundWavePCMWriter is already processing another write operation."));
+			return false;
+		}
+
+		if (!PrepWavFileOutput(InSampleBuffer, FileName, FilePath))
+		{
+			return false;
+		}
+
+		CurrentOperation.Reset(new FAsyncSoundWavePCMWriterTask(this, ESoundWavePCMWriteTaskType::WriteWavFile, [](const USoundWave*){}));
+		CurrentOperation->StartSynchronousTask();
+
+		if (OutFilePathName != nullptr)
+		{
+			*OutFilePathName = AbsoluteFilePath;
+		}
+
+		return true;
+	}
+
 	USoundWave* FSoundWavePCMWriter::SynchronouslyWriteSoundWave(const TSampleBuffer<>& InSampleBuffer, const FString* FileName, const FString* FilePath)
 	{
 		if (!IsDone())
diff --git a/Engine/Source/Runtime/Engine/Private/SoundWave.cpp b/Engine/Source/Runtime/Engine/Private/SoundWave.cpp
--- a/Engine/Source/Runtime/Engine/Private/SoundWave.cpp
+++ b/Engine/Source/Runtime/Engine/Private/SoundWave.cpp
@@ -31,6 +31,7 @@
 #include "Algo/BinarySearch.h"
 #include "Templates/UnrealTemplate.h"
 #include "UObject/ObjectSaveContext.h"
+#include "ISoundWaveCloudStreaming.h"
 
 static int32 SoundWaveDefaultLoadingBehaviorCVar = static_cast<int32>(ESoundWaveLoadingBehavior::LoadOnDemand);
 FAutoConsoleVariableRef CVarSoundWaveDefaultLoadingBehavior(
@@ -662,6 +663,26 @@ const uint8* USoundWave::GetResourceData() const
 	return SoundWaveDataPtr->GetResourceData();
 }
 
+#if WITH_EDITORONLY_DATA
+void USoundWave::SetCloudStreamingEnabled(bool bEnable)
+{
+	bEnableCloudStreaming = bEnable;
+}
+
+bool USoundWave::IsCloudStreamingEnabled() const
+{
+	return !!bEnableCloudStreaming;
+}
+
+void USoundWave::TriggerRecookForCloudStreaming()
+{
+	if (IsCloudStreamingEnabled())
+	{
+		InvalidateCompressedData();	// this will reinitialize the compressed data GUID and hence make sure we do generate new DDC entries at all times
+		MarkPackageDirty();			// also mark the package as dirty to ensure we trigger a cook
+	}
+}
+#endif // #if WITH_EDITORONLY_DATA
 
 ITargetPlatform* USoundWave::GetRunningPlatform()
 {
@@ -862,6 +883,7 @@ USoundWave::USoundWave(const FObjectInitializer& ObjectInitializer)
 	EnvelopeFollowerFrameSize = 1024;
 	EnvelopeFollowerAttackTime = 10;
 	EnvelopeFollowerReleaseTime = 100;
+	bEnableCloudStreaming = false;
 #endif
 
 	bCachedSampleRateFromPlatformSettings = false;
@@ -2777,6 +2799,8 @@ void USoundWave::PostEditChangeProperty(FPropertyChangedEvent& PropertyChangedEv
 	static const FName InitialChunkSizeFName = GET_MEMBER_NAME_CHECKED(USoundWave, InitialChunkSize_DEPRECATED);
 	static const FName TransformationsFName = GET_MEMBER_NAME_CHECKED(USoundWave, Transformations);
 	static const FName InlinedAudioInSecondsFName = GET_MEMBER_NAME_CHECKED(USoundWave, SizeOfFirstAudioChunkInSeconds);
+	static const FName CloudStreamingFName = GET_MEMBER_NAME_CHECKED(USoundWave, bEnableCloudStreaming);
+	static const FName CloudStreamingPlatformSettingsFName = GET_MEMBER_NAME_CHECKED(USoundWave, PlatformSettings);
 
 	// force proxy state to be up to date
 	SoundWaveDataPtr->InitializeDataFromSoundWave(*this);
@@ -2823,7 +2847,9 @@ void USoundWave::PostEditChangeProperty(FPropertyChangedEvent& PropertyChangedEv
 				|| Name == LoadingBehaviorFName
 				|| Name == InlinedAudioInSecondsFName
 				|| Name == InitialChunkSizeFName
-				|| Name == TransformationsFName)
+				|| Name == TransformationsFName
+				|| Name == CloudStreamingPlatformSettingsFName
+				|| Name == CloudStreamingFName)
 			{
 				UpdateAsset();
 			}
@@ -2841,6 +2867,29 @@ void USoundWave::PostEditChangeProperty(FPropertyChangedEvent& PropertyChangedEv
 	}
 }
 
+bool USoundWave::CanEditChange(const FProperty* InProperty) const
+{
+	if (!Super::CanEditChange(InProperty))
+	{
+		return false;
+	}
+
+	const FName& Name = InProperty->GetFName();
+	if (Name == GetCloudStreamingEnabledPropertyName())
+	{
+		// If it is currently set then it is always editable so that it can be turned off if there is no suitable feature plugin available.
+		if (IsCloudStreamingEnabled())
+		{
+			return true;
+		}
+		IModularFeatures::FScopedLockModularFeatureList ScopedLockModularFeatureList;
+		TArray<Audio::ISoundWaveCloudStreamingFeature*> Features = IModularFeatures::Get().GetModularFeatureImplementations<Audio::ISoundWaveCloudStreamingFeature>(Audio::ISoundWaveCloudStreamingFeature::GetModularFeatureName());
+		return Features.Num() > 0;
+	}
+	return true;
+}
+
+
 #endif // WITH_EDITOR
 
 #if WITH_EDITORONLY_DATA
