Fixed CI issue 364855

#rnx

diff --git a/Engine/Plugins/Media/ElectraPlayer/Source/ElectraPlayerRuntime/Private/Runtime/Player/SimpleElectraAudioPlayer.cpp b/Engine/Plugins/Media/ElectraPlayer/Source/ElectraPlayerRuntime/Private/Runtime/Player/SimpleElectraAudioPlayer.cpp
--- a/Engine/Plugins/Media/ElectraPlayer/Source/ElectraPlayerRuntime/Private/Runtime/Player/SimpleElectraAudioPlayer.cpp
+++ b/Engine/Plugins/Media/ElectraPlayer/Source/ElectraPlayerRuntime/Private/Runtime/Player/SimpleElectraAudioPlayer.cpp
@@ -108,7 +108,7 @@ FAutoConsoleVariableRef CVarElectraCSAudio_SendAnalytics(
 
 using namespace Electra;
 
-namespace
+namespace SimpleElectraAudioPlayerUtil
 {
 #if UE_BUILD_SHIPPING
 #define HIDE_URLS_FROM_LOG	1
@@ -116,7 +116,7 @@ namespace
 #define HIDE_URLS_FROM_LOG	0
 #endif
 
-static FString SanitizeMessage(FString InMessage)
+static FString RedactMessage(FString InMessage)
 {
 #if !HIDE_URLS_FROM_LOG
 	return MoveTemp(InMessage);
@@ -395,7 +395,7 @@ public:
 	{
 		bHasErrored = true;
 		ErrorMessage = InErrorReason;
-		FString Msg = FString::Printf(TEXT("%s: %s"), *AssetName, *SanitizeMessage(InErrorReason));
+		FString Msg = FString::Printf(TEXT("%s: %s"), *AssetName, *SimpleElectraAudioPlayerUtil::RedactMessage(InErrorReason));
 		UE_LOG(LogSimpleElectraPlayer, Error, TEXT("%s"), *Msg);
 	}
 	void ReportLogMessage(IInfoLog::ELevel InLogLevel, const FString& InLogMessage, int64 InPlayerWallclockMilliseconds) override
@@ -406,7 +406,7 @@ public:
 			return;
 		}
 
-		FString Msg = FString::Printf(TEXT("%s: %s"), *AssetName, *SanitizeMessage(InLogMessage));
+		FString Msg = FString::Printf(TEXT("%s: %s"), *AssetName, *SimpleElectraAudioPlayerUtil::RedactMessage(InLogMessage));
 		switch(InLogLevel)
 		{
 			case IInfoLog::ELevel::Verbose:
@@ -1295,7 +1295,7 @@ void FSimpleElectraAudioPlayer::UpdateAssetName()
 	}
 	else
 	{
-		AssetName = SanitizeMessage(BaseURL);
+		AssetName = SimpleElectraAudioPlayerUtil::RedactMessage(BaseURL);
 	}
 }
 
