Added seamless looping of cloud streamed soundwaves

#rnx

diff --git a/Engine/Plugins/Media/ElectraPlayer/Source/ElectraPlayerRuntime/Private/Runtime/Player/AdaptivePlayerOptionKeynames.h b/Engine/Plugins/Media/ElectraPlayer/Source/ElectraPlayerRuntime/Private/Runtime/Player/AdaptivePlayerOptionKeynames.h
--- a/Engine/Plugins/Media/ElectraPlayer/Source/ElectraPlayerRuntime/Private/Runtime/Player/AdaptivePlayerOptionKeynames.h
+++ b/Engine/Plugins/Media/ElectraPlayer/Source/ElectraPlayerRuntime/Private/Runtime/Player/AdaptivePlayerOptionKeynames.h
@@ -35,6 +35,9 @@ namespace Electra
 	
 	//! (bool) true to emit the first decoded video frame while prerolling so it can be displayed while scrubbing.
 	const TCHAR* const OptionKeyDoNotHoldBackFirstVideoFrame = TEXT("do_not_hold_back_first_frame");
+
+	//! (bool) true to not truncate the media segment access units at the end of the presentation. Must only be used without a set playback range end!
+	const TCHAR* const OptionKeyDoNotTruncateAtPresentationEnd = TEXT("do_not_truncate_at_presentation_end");
 	
 	//! (bool) true to have every request to read data break out to an external data reader.
 	const TCHAR* const OptionKeyUseExternalDataReader = TEXT("use_external_data_reader");
diff --git a/Engine/Plugins/Media/ElectraPlayer/Source/ElectraPlayerRuntime/Private/Runtime/Player/DASH/StreamReaderDASH.cpp b/Engine/Plugins/Media/ElectraPlayer/Source/ElectraPlayerRuntime/Private/Runtime/Player/DASH/StreamReaderDASH.cpp
--- a/Engine/Plugins/Media/ElectraPlayer/Source/ElectraPlayerRuntime/Private/Runtime/Player/DASH/StreamReaderDASH.cpp
+++ b/Engine/Plugins/Media/ElectraPlayer/Source/ElectraPlayerRuntime/Private/Runtime/Player/DASH/StreamReaderDASH.cpp
@@ -11,6 +11,7 @@
 #include "HAL/LowLevelMemTracker.h"
 #include "ElectraPlayerPrivate.h"
 #include "Player/AdaptiveStreamingPlayerABR.h"
+#include "Player/AdaptivePlayerOptionKeynames.h"
 #include "Player/PlayerEntityCache.h"
 #include "Player/DASH/PlaylistReaderDASH.h"
 #include "Player/DASH/OptionKeynamesDASH.h"
@@ -1009,6 +1010,9 @@ void FStreamReaderDASH::FStreamHandler::HandleRequestMP4()
 	FTimeValue LastKnownAUDuration;
 	FTimeValue TimeOffset = Request->PeriodStart + Request->AST + Request->AdditionalAdjustmentTime;
 
+	const FParamDict& Options = PlayerSessionService->GetOptions();
+	bool bDoNotTruncateAtPresentationEnd = Options.GetValue(OptionKeyDoNotTruncateAtPresentationEnd).SafeGetBool(false);
+
 	// Get the init segment if there is one. Either gets it from the entity cache or requests it now and adds it to the cache.
 	InitSegmentError = GetInitSegment(MP4InitSegment, Request);
 	// Get the CSD from the init segment in case there is a failure with the segment later.
@@ -1111,7 +1115,6 @@ void FStreamReaderDASH::FStreamHandler::HandleRequestMP4()
 			}
 			HTTP->Parameters.bCollectTimingTraces = Request->Segment.bLowLatencyChunkedEncodingExpected;
 			// Set timeouts for media segment retrieval
-			const FParamDict& Options = PlayerSessionService->GetOptions();
 			HTTP->Parameters.ConnectTimeout = Options.GetValue(DASH::OptionKeyMediaSegmentConnectTimeout).SafeGetTimeValue(FTimeValue().SetFromMilliseconds(1000 * 4));
 			HTTP->Parameters.NoDataTimeout = Options.GetValue(DASH::OptionKeyMediaSegmentNoDataTimeout).SafeGetTimeValue(FTimeValue().SetFromMilliseconds(1000 * 4));
 
@@ -1197,7 +1200,7 @@ void FStreamReaderDASH::FStreamHandler::HandleRequestMP4()
 								TrackTimescale = TrackIterator->GetTimescale();
 
 								int64 MediaLocalFirstAUTime = Request->Segment.MediaLocalFirstAUTime;
-								int64 MediaLocalLastAUTime = Request->Segment.MediaLocalLastAUTime;
+								int64 MediaLocalLastAUTime = bDoNotTruncateAtPresentationEnd ? TNumericLimits<int64>::Max() : Request->Segment.MediaLocalLastAUTime;
 								int64 PTO = Request->Segment.PTO;
 
 								if (TrackTimescale != Request->Segment.Timescale)
@@ -1869,6 +1872,9 @@ void FStreamReaderDASH::FStreamHandler::HandleRequestMKV()
 	bFillRemainingDuration = false;
 	ABRAbortReason.Empty();
 
+	const FParamDict& Options = PlayerSessionService->GetOptions();
+	bool bDoNotTruncateAtPresentationEnd = Options.GetValue(OptionKeyDoNotTruncateAtPresentationEnd).SafeGetBool(false);
+
 	// Get the init segment if there is one. Either gets it from the entity cache or requests it now and adds it to the cache.
 	InitSegmentError = GetInitSegment(MKVParser, Request);
 	// Get the CSD from the init segment in case there is a failure with the segment later.
@@ -1945,7 +1951,6 @@ void FStreamReaderDASH::FStreamHandler::HandleRequestMKV()
 			}
 			HTTP->Parameters.bCollectTimingTraces = Request->Segment.bLowLatencyChunkedEncodingExpected;
 			// Set timeouts for media segment retrieval
-			const FParamDict& Options = PlayerSessionService->GetOptions();
 			HTTP->Parameters.ConnectTimeout = Options.GetValue(DASH::OptionKeyMediaSegmentConnectTimeout).SafeGetTimeValue(FTimeValue().SetFromMilliseconds(1000 * 4));
 			HTTP->Parameters.NoDataTimeout = Options.GetValue(DASH::OptionKeyMediaSegmentNoDataTimeout).SafeGetTimeValue(FTimeValue().SetFromMilliseconds(1000 * 4));
 
@@ -1985,7 +1990,7 @@ void FStreamReaderDASH::FStreamHandler::HandleRequestMKV()
 			FTimeValue EarliestPTS(Request->Segment.MediaLocalFirstAUTime, Request->Segment.Timescale, Request->TimestampSequenceIndex);
 			EarliestPTS += TimeOffset - PTO;
 			FTimeValue LastPTS;
-			if (Request->Segment.MediaLocalLastAUTime != TNumericLimits<int64>::Max())
+			if (!bDoNotTruncateAtPresentationEnd && Request->Segment.MediaLocalLastAUTime != TNumericLimits<int64>::Max())
 			{
 				LastPTS.SetFromND(Request->Segment.MediaLocalLastAUTime, Request->Segment.Timescale, Request->TimestampSequenceIndex);
 				LastPTS += TimeOffset - PTO;
diff --git a/Engine/Plugins/Media/ElectraPlayer/Source/ElectraPlayerRuntime/Private/Runtime/Player/SimpleElectraAudioPlayer.cpp b/Engine/Plugins/Media/ElectraPlayer/Source/ElectraPlayerRuntime/Private/Runtime/Player/SimpleElectraAudioPlayer.cpp
--- a/Engine/Plugins/Media/ElectraPlayer/Source/ElectraPlayerRuntime/Private/Runtime/Player/SimpleElectraAudioPlayer.cpp
+++ b/Engine/Plugins/Media/ElectraPlayer/Source/ElectraPlayerRuntime/Private/Runtime/Player/SimpleElectraAudioPlayer.cpp
@@ -172,6 +172,7 @@ public:
 	FSimpleElectraAudioPlayer(const FCreateParams& InCreateParams);
 	~FSimpleElectraAudioPlayer();
 	bool Open(const TMap<FString, FVariant>& InOptions, const FString& InManifestURL, const FTimespan& InStartPosition, const FTimespan& InEncodedDuration, bool bInAutoPlay, bool bInSetLooping, TSharedPtr<IElectraPlayerDataCache, ESPMode::ThreadSafe> InPlayerDataCache) override;
+	void PrepareToLoopToBeginning() override;
 	void SeekTo(const FTimespan& NewPosition) override
 	{
 		FlushAudio(true);
@@ -280,15 +281,12 @@ public:
 	}
 
 	bool GetStreamFormat(FStreamFormat& OutFormat) const override;
-	int64 GetNumAvailableSamples() const override;
-	int64 GetNextSamples(float* OutBuffer, int32 InBufferSizeInSamples, int32 InNumSamplesToGet, const FDefaultSampleInfo& InDefaultSampleInfo) override;
-	int64 GetNextSamples(int16* OutBuffer, int32 InBufferSizeInSamples, int32 InNumSamplesToGet, const FDefaultSampleInfo& InDefaultSampleInfo) override;
-	bool IsAtEOS() const override;
+	int64 GetNextSamples(FTimespan& OutPTS, bool& bOutIsFirstBlock, int16* OutBuffer, int32 InBufferSizeInFrames, int32 InNumSamplesToGet, const FDefaultSampleInfo& InDefaultSampleInfo) override;
 
 	void FlushAudio(bool bForceFlush);
 	bool CanAcceptAudioFrames(int32 InNumFrames);
 	bool GetEnqueuedFrameInfo(int32& OutNumberOfEnqueuedFrames, FTimespan& OutDurationOfEnqueuedFrames) const;
-	bool EnqueueAudioFrames(const void* InBufferAddress, int32 InNumSamples, int32 InNumChannels, int32 InSampleRate, const FTimespan& InPTS, const FTimespan& InDuration);
+	bool EnqueueAudioFrames(const void* InBufferAddress, int32 InNumSamples, int32 InNumChannels, int32 InSampleRate, const FTimespan& InPTS, int64 InSequenceIndex, const FTimespan& InDuration);
 	void SetReceivedLastBuffer();
 
 	void ReportOpenSource(const FString& InURL) override
@@ -596,10 +594,41 @@ private:
 
 	mutable FCriticalSection Lock;
 	FStreamFormat StreamFormat;
-	TArray<FSampleBlock> SampleBlocks;
-	int64 NumSamplesAvailable = 0;
-	bool bReachedEOS = false;
+
+	struct FBlockSequence
+	{
+		~FBlockSequence()
+		{
+			for(int32 i=0; i<SampleBlocks.Num(); ++i)
+			{
+				SampleBlocks[i].Free();
+			}
+			SampleBlocks.Empty();
+		}
+
+		TArray<FSampleBlock> SampleBlocks;
+		int64 SequenceIndex = 0;
+		int64 NumFramesAvailable = 0;
+		bool bReachedEOS = false;
+		bool bReadEnded = false;
+	};
+	TArray<TSharedPtr<FBlockSequence, ESPMode::ThreadSafe>> NextPendingSampleBlocks;
+	TSharedPtr<FBlockSequence, ESPMode::ThreadSafe> CurrentWriteSampleBlock;
+	TSharedPtr<FBlockSequence, ESPMode::ThreadSafe> CurrentReadSampleBlock;
+
+	void ActivateBlockReadSequence()
+	{
+		if (!CurrentReadSampleBlock.IsValid() && NextPendingSampleBlocks.Num())
+		{
+			CurrentReadSampleBlock = NextPendingSampleBlocks[0];
+			NextPendingSampleBlocks.RemoveAt(0);
+		}
+	}
+
 	FTimespan NextPTS;
+	int64 NumActiveFramesTotal = 0;
+	int32 NumEnqueuedBlocks = 0;
+	bool bIsFirstSampleBlock = false;
 
 	double TimeCreated = -1.0;
 	double TimeUntilReady = -1.0;
@@ -618,7 +647,7 @@ private:
 	};
 	static void DoClosePlayerAsync(TSharedPtr<FClosePlayerInstanceHelper, ESPMode::ThreadSafe>&& InPlayerHelper);
 
-	bool HasBeenTerminated() const;
+	bool HasBeenSuspended() const;
 	void UpdateAssetName();
 	void InternalApplyOptions();
 	void HandleOpenStream(bool bOnlyIfAvailable);
@@ -890,7 +919,6 @@ void FSimpleElectraAudioPlayer::Tick(float DeltaTime)
 	}
 }
 
-
 void FSimpleElectraAudioPlayer::HandleOpenStream(bool bOnlyIfAvailable)
 {
 	int32 MaxAllowed = CreateParams.MaxTotalPlayerInstances;
@@ -956,6 +984,16 @@ void FSimpleElectraAudioPlayer::HandleOpenStream(bool bOnlyIfAvailable)
 			++Analytics.NumTimesStarted;
 		}
 
+		// Discard any potential leftovers.
+		Lock.Lock();
+		NextPendingSampleBlocks.Empty();
+		CurrentWriteSampleBlock.Reset();
+		CurrentReadSampleBlock.Reset();
+		NumActiveFramesTotal = 0;
+		NumEnqueuedBlocks = 0;
+		bIsFirstSampleBlock = true;
+		Lock.Unlock();
+
 		CreatePlayerInstance();
 		InternalApplyOptions();
 		Player->SetPlayerDataCache(PlayerDataCache.Pin());
@@ -977,6 +1015,7 @@ void FSimpleElectraAudioPlayer::HandleOpenStream(bool bOnlyIfAvailable)
 		bIsResuming = bAskedToBeReleased;
 		CurrentState = EState::Active;
 		bAskedToBeReleased = false;
+		bMaybeReopenAfterSeek = false;
 		TimeStopped = -1.0;
 	}
 }
@@ -1055,7 +1094,6 @@ void FSimpleElectraAudioPlayer::SendAnalyticMetrics(const TSharedPtr<IAnalyticsP
 	}
 }
 
-
 void FSimpleElectraAudioPlayer::MaybeRestartAStoppedStream(FSimpleElectraAudioPlayer* This)
 {
 	double TimeNow = FPlatformTime::Seconds();
@@ -1278,6 +1316,12 @@ void FSimpleElectraAudioPlayer::InternalApplyOptions()
 		Opts.SetOrUpdate(TEXT("initial_bitrate"), FVariantValue((int64) InitialBitrate->GetValue<int32>()));
 	}
 
+	const FVariant* DoNotTruncate = Options.Find(TEXT("do_not_truncate_at_presentation_end"));
+	if (DoNotTruncate)
+	{
+		Opts.SetOrUpdate(TEXT("do_not_truncate_at_presentation_end"), FVariantValue(DoNotTruncate->GetValue<bool>()));
+	}
+
 	TSharedPtr<IAdaptiveStreamingPlayer, ESPMode::ThreadSafe> LockedPlayer(Player);
 	if (LockedPlayer.IsValid())
 	{
@@ -1363,22 +1407,21 @@ bool FSimpleElectraAudioPlayer::GetStreamFormat(FStreamFormat& OutFormat) const
 bool FSimpleElectraAudioPlayer::CanAcceptAudioFrames(int32 InNumFrames)
 {
 	FScopeLock lock(&Lock);
-
 	if (StreamFormat.SampleRate)
 	{
-		double DurationInBuffer = (double)NumSamplesAvailable / StreamFormat.SampleRate;
-		return DurationInBuffer < SuggestedBufferDuration || NumSamplesAvailable < MinNumSamplesInBuffer;
+		double DurationInBuffer = (double)NumActiveFramesTotal / StreamFormat.SampleRate;
+		return DurationInBuffer < SuggestedBufferDuration || NumActiveFramesTotal < MinNumSamplesInBuffer;
 	}
-	return NumSamplesAvailable < MinNumSamplesInBuffer;
+	return NumActiveFramesTotal < MinNumSamplesInBuffer;
 }
 
 bool FSimpleElectraAudioPlayer::GetEnqueuedFrameInfo(int32& OutNumberOfEnqueuedFrames, FTimespan& OutDurationOfEnqueuedFrames) const
 {
 	FScopeLock lock(&Lock);
-	OutNumberOfEnqueuedFrames = SampleBlocks.Num();
+	OutNumberOfEnqueuedFrames = NumEnqueuedBlocks;
 	if (StreamFormat.SampleRate)
 	{
-		double DurationInBuffer = (double)NumSamplesAvailable / StreamFormat.SampleRate;
+		double DurationInBuffer = (double)NumActiveFramesTotal / StreamFormat.SampleRate;
 		OutDurationOfEnqueuedFrames = FTimespan::FromSeconds(DurationInBuffer);
 	}
 	else
@@ -1388,7 +1431,28 @@ bool FSimpleElectraAudioPlayer::GetEnqueuedFrameInfo(int32& OutNumberOfEnqueuedF
 	return true;
 }
 
-bool FSimpleElectraAudioPlayer::EnqueueAudioFrames(const void* InBufferAddress, int32 InNumSamples, int32 InNumChannels, int32 InSampleRate, const FTimespan& InPTS, const FTimespan& InDuration)
+void FSimpleElectraAudioPlayer::PrepareToLoopToBeginning()
+{
+	FScopeLock lock(&Lock);
+	bool bIsSuspended = HasBeenSuspended();
+	if (bIsSuspended)
+	{
+		bMaybeReopenAfterSeek = true;
+	}
+
+	// If we have a current block to read from then set its EOS state to not reached.
+	// This allows to produce silence if the stream is suspended.
+	if (CurrentReadSampleBlock.IsValid())
+	{
+		CurrentReadSampleBlock->bReachedEOS = false;
+	}
+	if (!bIsSuspended || (bIsSuspended && NextPendingSampleBlocks.Num()))
+	{
+		CurrentReadSampleBlock.Reset();
+	}
+}
+
+bool FSimpleElectraAudioPlayer::EnqueueAudioFrames(const void* InBufferAddress, int32 InNumSamples, int32 InNumChannels, int32 InSampleRate, const FTimespan& InPTS, int64 InSequenceIndex, const FTimespan& InDuration)
 {
 	if (InBufferAddress && InNumSamples && InNumChannels)
 	{
@@ -1415,14 +1479,23 @@ bool FSimpleElectraAudioPlayer::EnqueueAudioFrames(const void* InBufferAddress,
 			NextPTS = InPTS;
 		}
 
-		SampleBlocks.Emplace(MoveTemp(sb));
-		NumSamplesAvailable += InNumSamples;
+		if (!CurrentWriteSampleBlock.IsValid() || CurrentWriteSampleBlock->SequenceIndex != InSequenceIndex)
+		{
+			TSharedPtr<FBlockSequence, ESPMode::ThreadSafe> NewSeq = MakeShared<FBlockSequence, ESPMode::ThreadSafe>();
+			NewSeq->SequenceIndex = InSequenceIndex;
+			NextPendingSampleBlocks.Emplace(NewSeq);
+			CurrentWriteSampleBlock = MoveTemp(NewSeq);
+		}
+		CurrentWriteSampleBlock->SampleBlocks.Emplace(MoveTemp(sb));
+		CurrentWriteSampleBlock->NumFramesAvailable += InNumSamples;
+		NumActiveFramesTotal += InNumSamples;
+		++NumEnqueuedBlocks;
+
 		if (StreamFormat.SampleRate == 0 || StreamFormat.NumChannels == 0)
 		{
 			StreamFormat.SampleRate = InSampleRate;
 			StreamFormat.NumChannels = InNumChannels;
 		}
-		bReachedEOS = false;
 		return true;
 	}
 	return false;
@@ -1431,129 +1504,126 @@ bool FSimpleElectraAudioPlayer::EnqueueAudioFrames(const void* InBufferAddress,
 void FSimpleElectraAudioPlayer::SetReceivedLastBuffer()
 {
 	FScopeLock lock(&Lock);
-	bReachedEOS = true;
+	check(CurrentWriteSampleBlock.IsValid());
+	if (CurrentWriteSampleBlock.IsValid())
+	{
+		CurrentWriteSampleBlock->bReadEnded = true;
+	}
 }
 
 void FSimpleElectraAudioPlayer::FlushAudio(bool bForceFlush)
 {
-	bool bTerminated = HasBeenTerminated();
+	bool bIsSuspended = HasBeenSuspended();
 	FScopeLock lock(&Lock);
 	// Only flush when not terminated. Otherwise we want to return all the remaining samples we have
-	if (!bTerminated || bForceFlush)
+	if (!bIsSuspended || bForceFlush)
 	{
-		for(int32 i=0; i<SampleBlocks.Num(); ++i)
-		{
-			SampleBlocks[i].Free();
-		}
-		SampleBlocks.Empty();
-		NumSamplesAvailable = 0;
-		bReachedEOS = false;
+		NextPendingSampleBlocks.Empty();
+		CurrentWriteSampleBlock.Reset();
+		CurrentReadSampleBlock.Reset();
+		NumActiveFramesTotal = 0;
+		NumEnqueuedBlocks = 0;
 		NextPTS = FTimespan::MinValue();
 	}
 }
 
-int64 FSimpleElectraAudioPlayer::GetNumAvailableSamples() const
-{
-	bool bTerminated = HasBeenTerminated();
-	FScopeLock lock(&Lock);
-	if (NumSamplesAvailable)
-	{
-		return NumSamplesAvailable;
-	}
-	if (bTerminated)
-	{
-		return bReachedEOS ? 0 : kNumDummyBlockSamples;
-	}
-	return 0;
-}
-
-bool FSimpleElectraAudioPlayer::IsAtEOS() const
-{
-	FScopeLock lock(&Lock);
-	return bReachedEOS && SampleBlocks.Num() == 0;
-}
-
-bool FSimpleElectraAudioPlayer::HasBeenTerminated() const
+bool FSimpleElectraAudioPlayer::HasBeenSuspended() const
 {
 	return CurrentState == EState::Errored || CurrentState == EState::Stopped;
 }
 
-int64 FSimpleElectraAudioPlayer::GetNextSamples(float* OutBuffer, int32 InBufferSizeInSamples, int32 InNumSamplesToGet, const FDefaultSampleInfo& InDefaultSampleInfo)
+int64 FSimpleElectraAudioPlayer::GetNextSamples(FTimespan& OutPTS, bool& bOutIsFirstBlock, int16* OutBuffer, int32 InBufferSizeInFrames, int32 InNumFramesToGet, const FDefaultSampleInfo& InDefaultSampleInfo)
 {
-	check(!"this is not currently implemented");
-	return 0;
-}
-
-int64 FSimpleElectraAudioPlayer::GetNextSamples(int16* OutBuffer, int32 InBufferSizeInSamples, int32 InNumSamplesToGet, const FDefaultSampleInfo& InDefaultSampleInfo)
-{
-	if (!OutBuffer || InBufferSizeInSamples <= 0 || InNumSamplesToGet <= 0)
+	bOutIsFirstBlock = false;
+	if (!OutBuffer || InBufferSizeInFrames <= 0 || InNumFramesToGet <= 0)
 	{
 		return 0;
 	}
+
 	FScopeLock lock(&Lock);
-	if (SampleBlocks.Num() == 0)
+	bool bIsSuspended = HasBeenSuspended();
+	ActivateBlockReadSequence();
+	if (!CurrentReadSampleBlock.IsValid())
 	{
-		bool bTerminated = HasBeenTerminated();
-		if (!bTerminated)
+		return GetSamples_NotReady;
+	}
+
+	if (CurrentReadSampleBlock->SampleBlocks.Num() == 0)
+	{
+		// If the reading of the stream has not ended we are buffering and waiting for new data.
+		if (!bIsSuspended && !CurrentReadSampleBlock->bReadEnded)
 		{
-			return 0;
+			return GetSamples_NotReady;
 		}
+
+		// If all requested samples have been delivered already we are done.
+		if (CurrentReadSampleBlock->bReachedEOS)
+		{
+			return GetSamples_AtEOS;
+		}
+		// Otherwise we need to continue "playing" by returning silence.
 		else
 		{
 			int32 nc = StreamFormat.NumChannels > 0 ? StreamFormat.NumChannels : InDefaultSampleInfo.NumChannels;
 			int32 sr = StreamFormat.SampleRate > 0 ? StreamFormat.SampleRate : InDefaultSampleInfo.SampleRate;
 
-			if (InDefaultSampleInfo.ExpectedCurrentSamplePos >= 0 && InDefaultSampleInfo.NumTotalSamples >= 0)
+			if (InDefaultSampleInfo.ExpectedCurrentFramePos >= 0 && InDefaultSampleInfo.NumTotalFrames >= 0)
 			{
-				int64 RemainingToGo = InDefaultSampleInfo.NumTotalSamples - InDefaultSampleInfo.ExpectedCurrentSamplePos;
-				if (InNumSamplesToGet > RemainingToGo)
+				int64 RemainingToGo = InDefaultSampleInfo.NumTotalFrames - InDefaultSampleInfo.ExpectedCurrentFramePos;
+				if (InNumFramesToGet > RemainingToGo)
 				{
-					InNumSamplesToGet = RemainingToGo;
-					bReachedEOS = true;
+					InNumFramesToGet = RemainingToGo;
+					CurrentReadSampleBlock->bReachedEOS = true;
 				}
-				NextPTS = FTimespan::FromSeconds((double)(InDefaultSampleInfo.ExpectedCurrentSamplePos + InNumSamplesToGet) / sr);
+				OutPTS = FTimespan::FromSeconds((double)InDefaultSampleInfo.ExpectedCurrentFramePos / sr);
+				NextPTS = FTimespan::FromSeconds((double)(InDefaultSampleInfo.ExpectedCurrentFramePos + InNumFramesToGet) / sr);
 			}
 			else
 			{
+				OutPTS = NextPTS;
 				// Advance the PTS
 				if (NextPTS < FTimespan::Zero())
 				{
 					NextPTS = StartPosition;
 				}
-				NextPTS = NextPTS + FTimespan::FromSeconds((double)InNumSamplesToGet / sr);
-			}
-			if (NextPTS >= Duration)
-			{
-				bReachedEOS = true;
+				NextPTS = NextPTS + FTimespan::FromSeconds((double)InNumFramesToGet / sr);
+				if (NextPTS >= Duration)
+				{
+					CurrentReadSampleBlock->bReachedEOS = true;
+				}
 			}
-
 			// Fill the output with silence.
-			FMemory::Memzero(OutBuffer, InNumSamplesToGet * sizeof(int16) * nc);
-			return InNumSamplesToGet; 
+			FMemory::Memzero(OutBuffer, InNumFramesToGet * sizeof(int16) * nc);
+			return InNumFramesToGet; 
 		}
 	}
-	int32 NumSamplesToGo = InNumSamplesToGet < NumSamplesAvailable ? InNumSamplesToGet : NumSamplesAvailable;
+	int32 NumFramesToGo = InNumFramesToGet < CurrentReadSampleBlock->NumFramesAvailable ? InNumFramesToGet : CurrentReadSampleBlock->NumFramesAvailable;
 	int32 NumGot = 0;
-	FTimespan FrontPTS = SampleBlocks[0].PTS;
-	int32 FrontPTSSampleOffset = SampleBlocks[0].Offset;
-	while(NumSamplesToGo)
+	FTimespan FrontPTS = CurrentReadSampleBlock->SampleBlocks[0].PTS;
+	int32 FrontPTSSampleOffset = CurrentReadSampleBlock->SampleBlocks[0].Offset;
+	OutPTS = FrontPTS + FTimespan::FromSeconds((double)FrontPTSSampleOffset / StreamFormat.SampleRate);
+	while(NumFramesToGo)
 	{
-		int32 NumSamplesLeftInBlock = SampleBlocks[0].NumSamples - SampleBlocks[0].Offset;
-		int32 NumBlockSamples = NumSamplesToGo < NumSamplesLeftInBlock ? NumSamplesToGo : NumSamplesLeftInBlock;
-		const float* Src = reinterpret_cast<const float*>(SampleBlocks[0].Buffer) + SampleBlocks[0].Offset * StreamFormat.NumChannels;
-		for(int32 i=0, iMax=NumBlockSamples*StreamFormat.NumChannels; i<iMax; ++i)
+		int32 NumFramesLeftInBlock = CurrentReadSampleBlock->SampleBlocks[0].NumSamples - CurrentReadSampleBlock->SampleBlocks[0].Offset;
+		int32 NumBlockFrames = NumFramesToGo < NumFramesLeftInBlock ? NumFramesToGo : NumFramesLeftInBlock;
+		const float* Src = reinterpret_cast<const float*>(CurrentReadSampleBlock->SampleBlocks[0].Buffer) + CurrentReadSampleBlock->SampleBlocks[0].Offset * StreamFormat.NumChannels;
+		for(int32 i=0, iMax=NumBlockFrames*StreamFormat.NumChannels; i<iMax; ++i)
 		{
 			*OutBuffer++ = FMath::Clamp((int32)(*Src++ * 32768.0f), (int32)-32768, (int32)32767);
 		}
-		if ((SampleBlocks[0].Offset += NumBlockSamples) >= SampleBlocks[0].NumSamples)
+		if ((CurrentReadSampleBlock->SampleBlocks[0].Offset += NumBlockFrames) >= CurrentReadSampleBlock->SampleBlocks[0].NumSamples)
 		{
-			SampleBlocks[0].Free();
-			SampleBlocks.RemoveAt(0);
+			CurrentReadSampleBlock->SampleBlocks[0].Free();
+			CurrentReadSampleBlock->SampleBlocks.RemoveAt(0);
+			--NumEnqueuedBlocks;
 		}
-		NumSamplesAvailable -= NumBlockSamples;
-		NumSamplesToGo -= NumBlockSamples;
-		NumGot += NumBlockSamples;
+		CurrentReadSampleBlock->NumFramesAvailable -= NumBlockFrames;
+		NumActiveFramesTotal -= NumBlockFrames;
+		NumFramesToGo -= NumBlockFrames;
+		NumGot += NumBlockFrames;
 	}
+	bOutIsFirstBlock = bIsFirstSampleBlock;
+	bIsFirstSampleBlock = false;
 	NextPTS = FrontPTS + FTimespan::FromSeconds((double)(FrontPTSSampleOffset + NumGot) / StreamFormat.SampleRate);
 	return NumGot;
 }
@@ -1688,7 +1758,7 @@ UEMediaError FSimpleElectraAudioPlayerRenderer::ReturnBuffer(IBuffer* Buffer, bo
 		TSharedPtr<FSimpleElectraAudioPlayer, ESPMode::ThreadSafe> PinnedPlayer = Player.Pin();
 		if (PinnedPlayer.IsValid())
 		{
-			PinnedPlayer->EnqueueAudioFrames(BufferAddr, UsedBufferBytes / (sizeof(float)*NumChannels), NumChannels, SampleRate, PTS, Duration);
+			PinnedPlayer->EnqueueAudioFrames(BufferAddr, UsedBufferBytes / (sizeof(float)*NumChannels), NumChannels, SampleRate, PTS, InSequenceIndex, Duration);
 		}
 	}
 	else
diff --git a/Engine/Plugins/Media/ElectraPlayer/Source/ElectraPlayerRuntime/Public/SimpleElectraAudioPlayer.h b/Engine/Plugins/Media/ElectraPlayer/Source/ElectraPlayerRuntime/Public/SimpleElectraAudioPlayer.h
--- a/Engine/Plugins/Media/ElectraPlayer/Source/ElectraPlayerRuntime/Public/SimpleElectraAudioPlayer.h
+++ b/Engine/Plugins/Media/ElectraPlayer/Source/ElectraPlayerRuntime/Public/SimpleElectraAudioPlayer.h
@@ -4,6 +4,7 @@
 #include "CoreMinimal.h"
 #include "Containers/Map.h"
 #include "Misc/Variant.h"
+#include "Misc/Timespan.h"
 
 class IElectraPlayerDataCache;
 class IAnalyticsProviderET;
@@ -29,6 +30,7 @@ public:
 	virtual ~ISimpleElectraAudioPlayer() = default;
 	virtual bool Open(const TMap<FString, FVariant>& InOptions, const FString& ManifestURL, const FTimespan& StartPosition, const FTimespan& EncodedDuration, bool bAutoPlay, bool bSetLooping, TSharedPtr<IElectraPlayerDataCache, ESPMode::ThreadSafe> InPlayerDataCache) = 0;
 	virtual void SeekTo(const FTimespan& NewPosition) = 0;
+	virtual void PrepareToLoopToBeginning() = 0;
 	virtual void Pause() = 0;
 	virtual void Resume() = 0;
 	virtual void Stop() = 0;
@@ -44,16 +46,20 @@ public:
 
 	struct FDefaultSampleInfo
 	{
-		int64 NumTotalSamples = -1;
-		int64 ExpectedCurrentSamplePos = -1;
+		int64 NumTotalFrames = -1;
+		int64 ExpectedCurrentFramePos = -1;
 		int32 SampleRate = 0;
 		int32 NumChannels = 0;
 	};
 
-	virtual int64 GetNumAvailableSamples() const = 0;
-	virtual int64 GetNextSamples(float* OutBuffer, int32 InBufferSizeInSamples, int32 InNumSamplesToGet, const FDefaultSampleInfo& InDefaultSampleInfo) = 0;
-	virtual int64 GetNextSamples(int16* OutBuffer, int32 InBufferSizeInSamples, int32 InNumSamplesToGet, const FDefaultSampleInfo& InDefaultSampleInfo) = 0;
-	virtual bool IsAtEOS() const = 0;
+	// Negative return values indicating reason for having no samples in GetNextSamples()
+	enum
+	{
+		GetSamples_NotReady = -1,		// Player not ready, no samples available.
+		GetSamples_AtEOS = -2,			// Reached end of stream
+	};
+
+	virtual int64 GetNextSamples(FTimespan& OutPTS, bool& bOutIsFirstBlock, int16* OutBuffer, int32 InBufferSizeInFrames, int32 InNumFramesToGet, const FDefaultSampleInfo& InDefaultSampleInfo) = 0;
 
 
 	//-------------------------------------------------------------------------
