ElectraPlayer: added dedicated termination thread to avoid blocking the thread pool

#rnx

diff --git a/Engine/Plugins/Media/ElectraPlayer/Source/ElectraPlayerRuntime/Private/ElectraPlayer.cpp b/Engine/Plugins/Media/ElectraPlayer/Source/ElectraPlayerRuntime/Private/ElectraPlayer.cpp
--- a/Engine/Plugins/Media/ElectraPlayer/Source/ElectraPlayerRuntime/Private/ElectraPlayer.cpp
+++ b/Engine/Plugins/Media/ElectraPlayer/Source/ElectraPlayerRuntime/Private/ElectraPlayer.cpp
@@ -14,8 +14,6 @@
 #include "VideoDecoderResourceDelegate.h"
 #include "Utilities/Utilities.h"
 
-#include "Async/Async.h"
-
 #include "CoreGlobals.h"
 #include "Misc/ConfigCacheIni.h"
 
@@ -516,22 +514,7 @@ void FElectraPlayer::FInternalPlayerImpl::DoCloseAsync(TSharedPtr<FInternalPlaye
 	// Fallback to simple, sequential execution if the engine is already shutting down...
 	if (GIsRunning)
 	{
-		Async(EAsyncExecution::ThreadPool, MoveTemp(CloseTask));
-
-		#if !(UE_BUILD_SHIPPING || UE_BUILD_TEST)
-			FInternalPlayerImpl* PlayerImpl = Player.Get();
-
-			TFunction<void()> CheckTimeoutTask = [bClosedSig, PlayerImpl]()
-			{
-				// Wait for 3 seconds and if the player did not close by then log an error!
-				FPlatformProcess::Sleep(3.0f);
-				if (*bClosedSig == false)
-				{
-					UE_LOG(LogElectraPlayer, Error, TEXT("[%p] DoCloseAsync() did not complete in time. Player may be dead and dangling!"), PlayerImpl);
-				}
-			};
-			Async(EAsyncExecution::Thread, MoveTemp(CheckTimeoutTask));
-		#endif
+		FMediaRunnable::EnqueueTerminationFunction(MoveTemp(CloseTask));
 	}
 	else
 	{
diff --git a/Engine/Plugins/Media/ElectraPlayer/Source/ElectraPlayerRuntime/Private/Runtime/Core/MediaThreads.cpp b/Engine/Plugins/Media/ElectraPlayer/Source/ElectraPlayerRuntime/Private/Runtime/Core/MediaThreads.cpp
--- a/Engine/Plugins/Media/ElectraPlayer/Source/ElectraPlayerRuntime/Private/Runtime/Core/MediaThreads.cpp
+++ b/Engine/Plugins/Media/ElectraPlayer/Source/ElectraPlayerRuntime/Private/Runtime/Core/MediaThreads.cpp
@@ -1,10 +1,60 @@
 // Copyright Epic Games, Inc. All Rights Reserved.
 
 #include "Core/MediaThreads.h"
+#include "Core/MediaMessageQueue.h"
 
 #include "HAL/RunnableThread.h"
 #include "Misc/ScopeLock.h"
 
+
+namespace MediaRunnablePrivate
+{
+
+static TMediaMessageQueueDynamicNoTimeout<TFunction<void()>> TerminationFunctions;
+static volatile bool bEndForShutdown = false;
+static FMediaThread* TerminationThread = nullptr;
+static void TerminationThreadFN()
+{
+	while(!bEndForShutdown)
+	{
+		TFunction<void()> TerminateFN = TerminationFunctions.ReceiveMessage();
+		if (TerminateFN)
+		{
+			TerminateFN();
+		}
+	}
+}
+
+}
+
+void FMediaRunnable::Startup()
+{
+	MediaRunnablePrivate::bEndForShutdown = false;
+	MediaRunnablePrivate::TerminationThread = new FMediaThread("Electra::Termination");
+	MediaRunnablePrivate::TerminationThread->ThreadStart(FMediaRunnable::FStartDelegate::CreateStatic(&MediaRunnablePrivate::TerminationThreadFN));
+}
+
+void FMediaRunnable::Shutdown()
+{
+	check(MediaRunnablePrivate::TerminationThread);
+	
+	TFunction<void()> FinishTask = []()
+	{
+		MediaRunnablePrivate::bEndForShutdown = true;
+	};
+	EnqueueTerminationFunction(MoveTemp(FinishTask));
+
+	MediaRunnablePrivate::TerminationThread->ThreadWaitDone();
+	delete MediaRunnablePrivate::TerminationThread;
+	MediaRunnablePrivate::TerminationThread = nullptr;
+}
+
+void FMediaRunnable::EnqueueTerminationFunction(TFunction<void()>&& InFunctionToExecuteOnTerminationThread)
+{
+	MediaRunnablePrivate::TerminationFunctions.SendMessage(MoveTemp(InFunctionToExecuteOnTerminationThread));
+}
+
+
 // ----------------------------------------------------------------------------
 /**
  * Creates a new thread.
diff --git a/Engine/Plugins/Media/ElectraPlayer/Source/ElectraPlayerRuntime/Private/Runtime/Core/MediaThreads.h b/Engine/Plugins/Media/ElectraPlayer/Source/ElectraPlayerRuntime/Private/Runtime/Core/MediaThreads.h
--- a/Engine/Plugins/Media/ElectraPlayer/Source/ElectraPlayerRuntime/Private/Runtime/Core/MediaThreads.h
+++ b/Engine/Plugins/Media/ElectraPlayer/Source/ElectraPlayerRuntime/Private/Runtime/Core/MediaThreads.h
@@ -16,6 +16,8 @@
 #include "Containers/UnrealString.h"
 #include "Containers/StringConv.h"
 #include "Templates/SharedPointer.h"
+#include "Templates/Function.h"
+
 
 /**
  *
@@ -41,6 +43,10 @@ public:
 		int32				CoreAffinity;
 	};
 
+	static void Startup();
+	static void Shutdown();
+	static void EnqueueTerminationFunction(TFunction<void()>&& InFunctionToExecuteOnTerminationThread);
+
 	static FMediaRunnable* Create(int32 CoreAffinityMask, EThreadPriority Priority, uint32 StackSize, const FString& InThreadName);
 	static void Destroy(FMediaRunnable* Thread);
 
diff --git a/Engine/Plugins/Media/ElectraPlayer/Source/ElectraPlayerRuntime/Private/Runtime/HTTP/HTTPManager.cpp b/Engine/Plugins/Media/ElectraPlayer/Source/ElectraPlayerRuntime/Private/Runtime/HTTP/HTTPManager.cpp
--- a/Engine/Plugins/Media/ElectraPlayer/Source/ElectraPlayerRuntime/Private/Runtime/HTTP/HTTPManager.cpp
+++ b/Engine/Plugins/Media/ElectraPlayer/Source/ElectraPlayerRuntime/Private/Runtime/HTTP/HTTPManager.cpp
@@ -68,7 +68,7 @@ namespace Electra
 				{
 					delete InInstanceToDelete;
 				};
-				Async(EAsyncExecution::ThreadPool, MoveTemp(DeleteTask));
+				FMediaRunnable::EnqueueTerminationFunction(MoveTemp(DeleteTask));
 			}
 		};
 
diff --git a/Engine/Plugins/Media/ElectraPlayer/Source/ElectraPlayerRuntime/Private/Runtime/Player/DASH/ManifestDASH.cpp b/Engine/Plugins/Media/ElectraPlayer/Source/ElectraPlayerRuntime/Private/Runtime/Player/DASH/ManifestDASH.cpp
--- a/Engine/Plugins/Media/ElectraPlayer/Source/ElectraPlayerRuntime/Private/Runtime/Player/DASH/ManifestDASH.cpp
+++ b/Engine/Plugins/Media/ElectraPlayer/Source/ElectraPlayerRuntime/Private/Runtime/Player/DASH/ManifestDASH.cpp
@@ -21,7 +21,6 @@
 #include "Player/PlayerEntityCache.h"
 #include "Player/AdaptivePlayerOptionKeynames.h"
 #include "Player/DRM/DRMManager.h"
-#include "Async/Async.h"
 
 
 DECLARE_CYCLE_STAT(TEXT("FRepresentation::FindSegment"), STAT_ElectraPlayer_DASH_FindSegment, STATGROUP_ElectraPlayer);
diff --git a/Engine/Plugins/Media/ElectraPlayer/Source/ElectraPlayerRuntime/Private/Runtime/Player/DASH/StreamReaderDASH.cpp b/Engine/Plugins/Media/ElectraPlayer/Source/ElectraPlayerRuntime/Private/Runtime/Player/DASH/StreamReaderDASH.cpp
--- a/Engine/Plugins/Media/ElectraPlayer/Source/ElectraPlayerRuntime/Private/Runtime/Player/DASH/StreamReaderDASH.cpp
+++ b/Engine/Plugins/Media/ElectraPlayer/Source/ElectraPlayerRuntime/Private/Runtime/Player/DASH/StreamReaderDASH.cpp
@@ -188,13 +188,14 @@ UEMediaError FStreamReaderDASH::Create(IPlayerSessionServices* InPlayerSessionSe
 		StreamHandlers[i].bSilentCancellation  = false;
 		StreamHandlers[i].bHasErrored   	   = false;
 		StreamHandlers[i].IsIdleSignal.Signal();
-
+#if 0
+	// Disabled because the thread pool is timing sensitive and does not really allow for jobs like this.
 		// Subtitles get fetched running on the thread pool.
 		if (i == 2)
 		{
 			StreamHandlers[i].bRunOnThreadPool = true;
 		}
-
+#endif
 		StreamHandlers[i].ThreadSetName(i==0 ? "ElectraPlayer::DASH Video" : 
 										i==1 ? "ElectraPlayer::DASH Audio" :
 											   "ElectraPlayer::DASH Subtitle");
diff --git a/Engine/Plugins/Media/ElectraPlayer/Source/ElectraPlayerRuntime/Private/Runtime/Player/SimpleElectraAudioPlayer.cpp b/Engine/Plugins/Media/ElectraPlayer/Source/ElectraPlayerRuntime/Private/Runtime/Player/SimpleElectraAudioPlayer.cpp
--- a/Engine/Plugins/Media/ElectraPlayer/Source/ElectraPlayerRuntime/Private/Runtime/Player/SimpleElectraAudioPlayer.cpp
+++ b/Engine/Plugins/Media/ElectraPlayer/Source/ElectraPlayerRuntime/Private/Runtime/Player/SimpleElectraAudioPlayer.cpp
@@ -6,7 +6,6 @@
 #include "CoreMinimal.h"
 #include "Misc/Guid.h"
 #include "Misc/Optional.h"
-#include "Async/Async.h"
 #include "Tickable.h"
 #include "CoreGlobals.h"
 #include "Stats/Stats.h"
@@ -1200,7 +1199,7 @@ void FSimpleElectraAudioPlayer::DoClosePlayerAsync(TSharedPtr<FClosePlayerInstan
 
 	if (GIsRunning)
 	{
-		Async(EAsyncExecution::ThreadPool, MoveTemp(CloseTask));
+		FMediaRunnable::EnqueueTerminationFunction(MoveTemp(CloseTask));
 	}
 	else
 	{
diff --git a/Engine/Plugins/Media/ElectraPlayer/Source/ElectraPlayerRuntime/Private/Runtime/PlayerRuntimeGlobal.cpp b/Engine/Plugins/Media/ElectraPlayer/Source/ElectraPlayerRuntime/Private/Runtime/PlayerRuntimeGlobal.cpp
--- a/Engine/Plugins/Media/ElectraPlayer/Source/ElectraPlayerRuntime/Private/Runtime/PlayerRuntimeGlobal.cpp
+++ b/Engine/Plugins/Media/ElectraPlayer/Source/ElectraPlayerRuntime/Private/Runtime/PlayerRuntimeGlobal.cpp
@@ -5,6 +5,7 @@
 #include "Misc/CoreDelegates.h"
 #include "Modules/ModuleManager.h"
 
+
 namespace Electra
 {
 	namespace Global
@@ -81,6 +82,8 @@ namespace Electra
 	 */
 	bool Startup(const Configuration& InConfiguration)
 	{
+		FMediaRunnable::Startup();
+
 		if (!ApplicationTerminatingDelegate.IsValid())
 		{
 			ApplicationTerminatingDelegate = FCoreDelegates::GetApplicationWillTerminateDelegate().AddStatic(&HandleApplicationWillTerminate);
@@ -128,6 +131,8 @@ namespace Electra
 		{
 			FCoreDelegates::GetApplicationWillTerminateDelegate().Remove(ApplicationTerminatingDelegate);
 		}
+
+		FMediaRunnable::Shutdown();
 	}
 
 
@@ -157,6 +162,12 @@ namespace Electra
 	}
 
 
+	void EnqueueTerminationFunction(TFunction<void()>&& InFunctionToExecuteOnTerminationThread)
+	{
+		FMediaRunnable::EnqueueTerminationFunction(MoveTemp(InFunctionToExecuteOnTerminationThread));
+	}
+
+
 	/**
 	 * Check if the provided analytics event is enabled
 	 *
diff --git a/Engine/Plugins/Media/ElectraPlayer/Source/ElectraPlayerRuntime/Private/Runtime/PlayerRuntimeGlobal.h b/Engine/Plugins/Media/ElectraPlayer/Source/ElectraPlayerRuntime/Private/Runtime/PlayerRuntimeGlobal.h
--- a/Engine/Plugins/Media/ElectraPlayer/Source/ElectraPlayerRuntime/Private/Runtime/PlayerRuntimeGlobal.h
+++ b/Engine/Plugins/Media/ElectraPlayer/Source/ElectraPlayerRuntime/Private/Runtime/PlayerRuntimeGlobal.h
@@ -45,5 +45,8 @@ namespace Electra
 	};
 	bool AddBGFGNotificationHandler(TSharedPtrTS<FFGBGNotificationHandlers> InHandlers);
 	void RemoveBGFGNotificationHandler(TSharedPtrTS<FFGBGNotificationHandlers> InHandlers);
+
+
+	void EnqueueTerminationFunction(TFunction<void()>&& InFunctionToExecuteOnTerminationThread);
 };
 
