Electra Cloud Streaming Audio: caching of parsed entry point responses; moved player creation onto async worker

#rnx

diff --git a/Engine/Plugins/Media/ElectraPlayer/Source/ElectraPlayerRuntime/Private/ElectraPlayer.cpp b/Engine/Plugins/Media/ElectraPlayer/Source/ElectraPlayerRuntime/Private/ElectraPlayer.cpp
--- a/Engine/Plugins/Media/ElectraPlayer/Source/ElectraPlayerRuntime/Private/ElectraPlayer.cpp
+++ b/Engine/Plugins/Media/ElectraPlayer/Source/ElectraPlayerRuntime/Private/ElectraPlayer.cpp
@@ -514,7 +514,7 @@ void FElectraPlayer::FInternalPlayerImpl::DoCloseAsync(TSharedPtr<FInternalPlaye
 	// Fallback to simple, sequential execution if the engine is already shutting down...
 	if (GIsRunning)
 	{
-		FMediaRunnable::EnqueueTerminationFunction(MoveTemp(CloseTask));
+		FMediaRunnable::EnqueueAsyncTask(MoveTemp(CloseTask));
 	}
 	else
 	{
diff --git a/Engine/Plugins/Media/ElectraPlayer/Source/ElectraPlayerRuntime/Private/Runtime/Core/MediaThreads.cpp b/Engine/Plugins/Media/ElectraPlayer/Source/ElectraPlayerRuntime/Private/Runtime/Core/MediaThreads.cpp
--- a/Engine/Plugins/Media/ElectraPlayer/Source/ElectraPlayerRuntime/Private/Runtime/Core/MediaThreads.cpp
+++ b/Engine/Plugins/Media/ElectraPlayer/Source/ElectraPlayerRuntime/Private/Runtime/Core/MediaThreads.cpp
@@ -10,17 +10,17 @@
 namespace MediaRunnablePrivate
 {
 
-static TMediaMessageQueueDynamicNoTimeout<TFunction<void()>> TerminationFunctions;
+static TMediaMessageQueueDynamicNoTimeout<TFunction<void()>> AsyncFunctionsToRun;
 static volatile bool bEndForShutdown = false;
-static FMediaThread* TerminationThread = nullptr;
-static void TerminationThreadFN()
+static FMediaThread* AsyncExecutionThread = nullptr;
+static void AsyncWorkerThreadFN()
 {
 	while(!bEndForShutdown)
 	{
-		TFunction<void()> TerminateFN = TerminationFunctions.ReceiveMessage();
-		if (TerminateFN)
+		TFunction<void()> AsyncRunFN = AsyncFunctionsToRun.ReceiveMessage();
+		if (AsyncRunFN)
 		{
-			TerminateFN();
+			AsyncRunFN();
 		}
 	}
 }
@@ -30,28 +30,28 @@ static void TerminationThreadFN()
 void FMediaRunnable::Startup()
 {
 	MediaRunnablePrivate::bEndForShutdown = false;
-	MediaRunnablePrivate::TerminationThread = new FMediaThread("Electra::Termination");
-	MediaRunnablePrivate::TerminationThread->ThreadStart(FMediaRunnable::FStartDelegate::CreateStatic(&MediaRunnablePrivate::TerminationThreadFN));
+	MediaRunnablePrivate::AsyncExecutionThread = new FMediaThread("Electra::ExecAsync");
+	MediaRunnablePrivate::AsyncExecutionThread->ThreadStart(FMediaRunnable::FStartDelegate::CreateStatic(&MediaRunnablePrivate::AsyncWorkerThreadFN));
 }
 
 void FMediaRunnable::Shutdown()
 {
-	check(MediaRunnablePrivate::TerminationThread);
+	check(MediaRunnablePrivate::AsyncExecutionThread);
 	
 	TFunction<void()> FinishTask = []()
 	{
 		MediaRunnablePrivate::bEndForShutdown = true;
 	};
-	EnqueueTerminationFunction(MoveTemp(FinishTask));
+	EnqueueAsyncTask(MoveTemp(FinishTask));
 
-	MediaRunnablePrivate::TerminationThread->ThreadWaitDone();
-	delete MediaRunnablePrivate::TerminationThread;
-	MediaRunnablePrivate::TerminationThread = nullptr;
+	MediaRunnablePrivate::AsyncExecutionThread->ThreadWaitDone();
+	delete MediaRunnablePrivate::AsyncExecutionThread;
+	MediaRunnablePrivate::AsyncExecutionThread = nullptr;
 }
 
-void FMediaRunnable::EnqueueTerminationFunction(TFunction<void()>&& InFunctionToExecuteOnTerminationThread)
+void FMediaRunnable::EnqueueAsyncTask(TFunction<void()>&& InFunctionToExecuteOnAsyncThread)
 {
-	MediaRunnablePrivate::TerminationFunctions.SendMessage(MoveTemp(InFunctionToExecuteOnTerminationThread));
+	MediaRunnablePrivate::AsyncFunctionsToRun.SendMessage(MoveTemp(InFunctionToExecuteOnAsyncThread));
 }
 
 
diff --git a/Engine/Plugins/Media/ElectraPlayer/Source/ElectraPlayerRuntime/Private/Runtime/Core/MediaThreads.h b/Engine/Plugins/Media/ElectraPlayer/Source/ElectraPlayerRuntime/Private/Runtime/Core/MediaThreads.h
--- a/Engine/Plugins/Media/ElectraPlayer/Source/ElectraPlayerRuntime/Private/Runtime/Core/MediaThreads.h
+++ b/Engine/Plugins/Media/ElectraPlayer/Source/ElectraPlayerRuntime/Private/Runtime/Core/MediaThreads.h
@@ -45,7 +45,7 @@ public:
 
 	static void Startup();
 	static void Shutdown();
-	static void EnqueueTerminationFunction(TFunction<void()>&& InFunctionToExecuteOnTerminationThread);
+	static void EnqueueAsyncTask(TFunction<void()>&& InFunctionToExecuteOnAsyncThread);
 
 	static FMediaRunnable* Create(int32 CoreAffinityMask, EThreadPriority Priority, uint32 StackSize, const FString& InThreadName);
 	static void Destroy(FMediaRunnable* Thread);
diff --git a/Engine/Plugins/Media/ElectraPlayer/Source/ElectraPlayerRuntime/Private/Runtime/HTTP/HTTPManager.cpp b/Engine/Plugins/Media/ElectraPlayer/Source/ElectraPlayerRuntime/Private/Runtime/HTTP/HTTPManager.cpp
--- a/Engine/Plugins/Media/ElectraPlayer/Source/ElectraPlayerRuntime/Private/Runtime/HTTP/HTTPManager.cpp
+++ b/Engine/Plugins/Media/ElectraPlayer/Source/ElectraPlayerRuntime/Private/Runtime/HTTP/HTTPManager.cpp
@@ -68,7 +68,7 @@ namespace Electra
 				{
 					delete InInstanceToDelete;
 				};
-				FMediaRunnable::EnqueueTerminationFunction(MoveTemp(DeleteTask));
+				FMediaRunnable::EnqueueAsyncTask(MoveTemp(DeleteTask));
 			}
 		};
 
@@ -899,9 +899,6 @@ namespace Electra
 
 	void FElectraHttpManager::HandlePeriodicCallbacks(const FTimeValue& Now)
 	{
-		SCOPE_CYCLE_COUNTER(STAT_ElectraPlayer_FElectraHttpManager_Worker);
-		CSV_SCOPED_TIMING_STAT(ElectraPlayer, ElectraHttpManager_Worker);
-
 		for(TMap<FHandle*, TSharedPtrTS<FRequest>>::TIterator It = ActiveRequests.CreateIterator(); It; ++It)
 		{
 			FHandle* Handle = It.Key();
diff --git a/Engine/Plugins/Media/ElectraPlayer/Source/ElectraPlayerRuntime/Private/Runtime/Player/SimpleElectraAudioPlayer.cpp b/Engine/Plugins/Media/ElectraPlayer/Source/ElectraPlayerRuntime/Private/Runtime/Player/SimpleElectraAudioPlayer.cpp
--- a/Engine/Plugins/Media/ElectraPlayer/Source/ElectraPlayerRuntime/Private/Runtime/Player/SimpleElectraAudioPlayer.cpp
+++ b/Engine/Plugins/Media/ElectraPlayer/Source/ElectraPlayerRuntime/Private/Runtime/Player/SimpleElectraAudioPlayer.cpp
@@ -168,6 +168,9 @@ public:
 	}
 	static void SendAnalyticMetrics(const TSharedPtr<IAnalyticsProviderET>& InAnalyticsProvider);
 
+	static TSharedPtr<ICacheElementBase, ESPMode::ThreadSafe> GetCustomCacheElement(const FString& InForURL);
+	static void SetCustomCacheElement(const FString& InForURL, const TSharedPtr<ICacheElementBase, ESPMode::ThreadSafe>& InElement);
+
 	FSimpleElectraAudioPlayer(const FCreateParams& InCreateParams);
 	~FSimpleElectraAudioPlayer();
 	bool Open(const TMap<FString, FVariant>& InOptions, const FString& InManifestURL, const FTimespan& InStartPosition, const FTimespan& InEncodedDuration, bool bInAutoPlay, bool bInSetLooping, TSharedPtr<IElectraPlayerDataCache, ESPMode::ThreadSafe> InPlayerDataCache) override;
@@ -511,6 +514,7 @@ private:
 	struct FBlobCacheEntry
 	{
 		TSharedPtr<TArray<uint8>, ESPMode::ThreadSafe> Data;
+		TSharedPtr<ISimpleElectraAudioPlayer::ICacheElementBase, ESPMode::ThreadSafe> CustomData;
 		FString URL;
 		double TimeAdded = 0.0;
 	};
@@ -520,6 +524,9 @@ private:
 	public:
 		TSharedPtr<TArray<uint8>, ESPMode::ThreadSafe> GetData(const FString& InForURL);
 		void AddData(const FString& InForURL, const TSharedPtr<TArray<uint8>, ESPMode::ThreadSafe>& InData);
+
+		TSharedPtr<ISimpleElectraAudioPlayer::ICacheElementBase, ESPMode::ThreadSafe> GetCustomData(const FString& InForURL);
+		void AddCustomData(const FString& InForURL, const TSharedPtr<ISimpleElectraAudioPlayer::ICacheElementBase, ESPMode::ThreadSafe>& InData);
 	private:
 		void CheckLimits(int32 InBytesNeeded);
 		FCriticalSection Lock;
@@ -550,6 +557,7 @@ private:
 		OpenBlob,
 		OpeningStream,
 		TryReopeningStream,
+		CreatingPlayer,
 		Active,
 		Stopped,
 		Errored
@@ -617,7 +625,7 @@ private:
 
 	void ActivateBlockReadSequence()
 	{
-		if (!CurrentReadSampleBlock.IsValid() && NextPendingSampleBlocks.Num())
+		if ((!CurrentReadSampleBlock.IsValid() || CurrentReadSampleBlock->SequenceIndex == -1) && NextPendingSampleBlocks.Num())
 		{
 			CurrentReadSampleBlock = NextPendingSampleBlocks[0];
 			NextPendingSampleBlocks.RemoveAt(0);
@@ -651,8 +659,10 @@ private:
 	void InternalApplyOptions();
 	void HandleOpenStream(bool bOnlyIfAvailable);
 	static void MaybeRestartAStoppedStream(FSimpleElectraAudioPlayer* This);
+	void CreatePlayerAsync();
+	void CreateIdleBufferIfNecessary();
 
-	FBlobCache& GetBlobCache() const
+	static FBlobCache& GetBlobCache()
 	{
 		static FBlobCache Cache;
 		return Cache;
@@ -805,7 +815,25 @@ void ISimpleElectraAudioPlayer::SendAnalyticMetrics(const TSharedPtr<IAnalyticsP
 	}
 }
 
+TSharedPtr<ISimpleElectraAudioPlayer::ICacheElementBase, ESPMode::ThreadSafe> ISimpleElectraAudioPlayer::GetCustomCacheElement(const FString& InForURL)
+{
+	return FSimpleElectraAudioPlayer::GetCustomCacheElement(InForURL);
+}
 
+void ISimpleElectraAudioPlayer::SetCustomCacheElement(const FString& InForURL, const TSharedPtr<ISimpleElectraAudioPlayer::ICacheElementBase, ESPMode::ThreadSafe>& InElement)
+{
+	FSimpleElectraAudioPlayer::SetCustomCacheElement(InForURL, InElement);
+}
+
+TSharedPtr<ISimpleElectraAudioPlayer::ICacheElementBase, ESPMode::ThreadSafe> FSimpleElectraAudioPlayer::GetCustomCacheElement(const FString& InForURL)
+{
+	return GetBlobCache().GetCustomData(InForURL);
+}
+
+void FSimpleElectraAudioPlayer::SetCustomCacheElement(const FString& InForURL, const TSharedPtr<ISimpleElectraAudioPlayer::ICacheElementBase, ESPMode::ThreadSafe>& InElement)
+{
+	GetBlobCache().AddCustomData(InForURL, InElement);
+}
 
 FSimpleElectraAudioPlayer::FSimpleElectraAudioPlayer(const FCreateParams& InCreateParams)
 	: CreateParams(InCreateParams)
@@ -891,6 +919,10 @@ void FSimpleElectraAudioPlayer::Tick(float DeltaTime)
 			HandleOpenStream(true);
 			break;
 		}
+		case EState::CreatingPlayer:
+		{
+			return;
+		}
 		case EState::Errored:
 		{
 			return;
@@ -907,7 +939,7 @@ void FSimpleElectraAudioPlayer::Tick(float DeltaTime)
 		}
 	}
 
-	if (bAskedToBeReleased)
+	if (bAskedToBeReleased && CurrentState != EState::CreatingPlayer)
 	{
 		FString Msg = FString::Printf(TEXT("%s: Stopping due to limit"), *AssetName);
 		UE_LOG(LogSimpleElectraPlayer, Verbose, TEXT("%s"), *Msg);
@@ -969,7 +1001,6 @@ void FSimpleElectraAudioPlayer::HandleOpenStream(bool bOnlyIfAvailable)
 
 	if (bCreateNow)
 	{
-		// No. Go!
 		FString Msg = FString::Printf(TEXT("%s: %s playback"), *AssetName, !bAskedToBeReleased ? TEXT("Preparing") : TEXT("Resuming"));
 		UE_LOG(LogSimpleElectraPlayer, Verbose, TEXT("%s"), *Msg);
 
@@ -993,17 +1024,6 @@ void FSimpleElectraAudioPlayer::HandleOpenStream(bool bOnlyIfAvailable)
 		bIsFirstSampleBlock = true;
 		Lock.Unlock();
 
-		CreatePlayerInstance();
-		InternalApplyOptions();
-		Player->SetPlayerDataCache(PlayerDataCache.Pin());
-		Player->LoadManifest(URL);
-		if (bSetLooping)
-		{
-			IAdaptiveStreamingPlayer::FLoopParam lp;
-			lp.bEnableLooping = true;
-			Player->SetLooping(lp);
-		}
-			
 		// Add to the list of active instances.
 		InstanceLock.Lock();
 		ActiveInstances.AddUnique(this);
@@ -1012,13 +1032,52 @@ void FSimpleElectraAudioPlayer::HandleOpenStream(bool bOnlyIfAvailable)
 
 		// Set state to active.
 		bIsResuming = bAskedToBeReleased;
-		CurrentState = EState::Active;
+		CurrentState = EState::CreatingPlayer;
 		bAskedToBeReleased = false;
 		bMaybeReopenAfterSeek = false;
 		TimeStopped = -1.0;
+
+		TWeakPtr<FSimpleElectraAudioPlayer, ESPMode::ThreadSafe> This(AsShared());
+		TFunction<void()> CreateTask = [This]()
+		{
+			TSharedPtr<FSimpleElectraAudioPlayer, ESPMode::ThreadSafe> That = This.Pin();
+			if (That.IsValid())
+			{
+				That->CreatePlayerAsync();
+			}
+		};
+		FMediaRunnable::EnqueueAsyncTask(MoveTemp(CreateTask));
+	}
+}
+
+void FSimpleElectraAudioPlayer::CreatePlayerAsync()
+{
+	CreatePlayerInstance();
+	InternalApplyOptions();
+	Player->SetPlayerDataCache(PlayerDataCache.Pin());
+	Player->LoadManifest(URL);
+	if (bSetLooping)
+	{
+		IAdaptiveStreamingPlayer::FLoopParam lp;
+		lp.bEnableLooping = true;
+		Player->SetLooping(lp);
+	}
+	CurrentState = EState::Active;
+}
+
+void FSimpleElectraAudioPlayer::CreateIdleBufferIfNecessary()
+{
+	Lock.Lock();
+	if (!CurrentReadSampleBlock.IsValid() && NextPendingSampleBlocks.IsEmpty())
+	{
+		TSharedPtr<FBlockSequence, ESPMode::ThreadSafe> NewSeq = MakeShared<FBlockSequence, ESPMode::ThreadSafe>();
+		NewSeq->SequenceIndex = -1;
+		NextPendingSampleBlocks.Emplace(NewSeq);
 	}
+	Lock.Unlock();
 }
 
+
 void FSimpleElectraAudioPlayer::MergeAnalytics(const FString& AssetId, const FAnalyticsEntry& InAnalytics)
 {
 	// If this entry was not active and had no failure downloading the blob we
@@ -1182,6 +1241,10 @@ void FSimpleElectraAudioPlayer::DoClosePlayerAsync(TSharedPtr<FClosePlayerInstan
 			InPlayerHelper->PlayerInstance->Stop();
 			InPlayerHelper->PlayerInstance.Reset();
 		}
+		if (!InPlayerHelper->bDeleteThis)
+		{
+			InPlayerHelper->This->CreateIdleBufferIfNecessary();
+		}
 		InstanceLock.Lock();
 		ActiveInstances.Remove(InPlayerHelper->This);
 		if (!InPlayerHelper->bDeleteThis)
@@ -1199,7 +1262,7 @@ void FSimpleElectraAudioPlayer::DoClosePlayerAsync(TSharedPtr<FClosePlayerInstan
 
 	if (GIsRunning)
 	{
-		FMediaRunnable::EnqueueTerminationFunction(MoveTemp(CloseTask));
+		FMediaRunnable::EnqueueAsyncTask(MoveTemp(CloseTask));
 	}
 	else
 	{
@@ -1675,6 +1738,24 @@ void FSimpleElectraAudioPlayer::FBlobCache::AddData(const FString& InForURL, con
 	}
 }
 
+TSharedPtr<ISimpleElectraAudioPlayer::ICacheElementBase, ESPMode::ThreadSafe> FSimpleElectraAudioPlayer::FBlobCache::GetCustomData(const FString& InForURL)
+{
+	FScopeLock lock(&Lock);
+	TSharedPtr<FBlobCacheEntry, ESPMode::ThreadSafe>* Entry = EntryMap.Find(InForURL);
+	return Entry ? (*Entry)->CustomData : nullptr;
+}
+
+void FSimpleElectraAudioPlayer::FBlobCache::AddCustomData(const FString& InForURL, const TSharedPtr<ISimpleElectraAudioPlayer::ICacheElementBase, ESPMode::ThreadSafe>& InData)
+{
+	FScopeLock lock(&Lock);
+	TSharedPtr<FBlobCacheEntry, ESPMode::ThreadSafe>* Entry = EntryMap.Find(InForURL);
+	if (Entry)
+	{
+		(*Entry)->CustomData = InData;
+	}
+}
+
+
 void FSimpleElectraAudioPlayer::FBlobCache::CheckLimits(int32 InBytesNeeded)
 {
 	while(EntryLRU.Num() && SizeInUse + InBytesNeeded > ElectraCSAudio_BlobCacheSize)
diff --git a/Engine/Plugins/Media/ElectraPlayer/Source/ElectraPlayerRuntime/Private/Runtime/PlayerRuntimeGlobal.cpp b/Engine/Plugins/Media/ElectraPlayer/Source/ElectraPlayerRuntime/Private/Runtime/PlayerRuntimeGlobal.cpp
--- a/Engine/Plugins/Media/ElectraPlayer/Source/ElectraPlayerRuntime/Private/Runtime/PlayerRuntimeGlobal.cpp
+++ b/Engine/Plugins/Media/ElectraPlayer/Source/ElectraPlayerRuntime/Private/Runtime/PlayerRuntimeGlobal.cpp
@@ -162,12 +162,6 @@ namespace Electra
 	}
 
 
-	void EnqueueTerminationFunction(TFunction<void()>&& InFunctionToExecuteOnTerminationThread)
-	{
-		FMediaRunnable::EnqueueTerminationFunction(MoveTemp(InFunctionToExecuteOnTerminationThread));
-	}
-
-
 	/**
 	 * Check if the provided analytics event is enabled
 	 *
diff --git a/Engine/Plugins/Media/ElectraPlayer/Source/ElectraPlayerRuntime/Private/Runtime/PlayerRuntimeGlobal.h b/Engine/Plugins/Media/ElectraPlayer/Source/ElectraPlayerRuntime/Private/Runtime/PlayerRuntimeGlobal.h
--- a/Engine/Plugins/Media/ElectraPlayer/Source/ElectraPlayerRuntime/Private/Runtime/PlayerRuntimeGlobal.h
+++ b/Engine/Plugins/Media/ElectraPlayer/Source/ElectraPlayerRuntime/Private/Runtime/PlayerRuntimeGlobal.h
@@ -45,8 +45,5 @@ namespace Electra
 	};
 	bool AddBGFGNotificationHandler(TSharedPtrTS<FFGBGNotificationHandlers> InHandlers);
 	void RemoveBGFGNotificationHandler(TSharedPtrTS<FFGBGNotificationHandlers> InHandlers);
-
-
-	void EnqueueTerminationFunction(TFunction<void()>&& InFunctionToExecuteOnTerminationThread);
 };
 
diff --git a/Engine/Plugins/Media/ElectraPlayer/Source/ElectraPlayerRuntime/Public/SimpleElectraAudioPlayer.h b/Engine/Plugins/Media/ElectraPlayer/Source/ElectraPlayerRuntime/Public/SimpleElectraAudioPlayer.h
--- a/Engine/Plugins/Media/ElectraPlayer/Source/ElectraPlayerRuntime/Public/SimpleElectraAudioPlayer.h
+++ b/Engine/Plugins/Media/ElectraPlayer/Source/ElectraPlayerRuntime/Public/SimpleElectraAudioPlayer.h
@@ -27,6 +27,16 @@ public:
 	 */
 	static void SendAnalyticMetrics(const TSharedPtr<IAnalyticsProviderET>& InAnalyticsProvider);
 
+
+	class ICacheElementBase : public TSharedFromThis<ICacheElementBase, ESPMode::ThreadSafe>
+	{
+	public:
+		virtual ~ICacheElementBase() {};
+	};
+	static TSharedPtr<ICacheElementBase, ESPMode::ThreadSafe> GetCustomCacheElement(const FString& InForURL);
+	static void SetCustomCacheElement(const FString& InForURL, const TSharedPtr<ICacheElementBase, ESPMode::ThreadSafe>& InElement);
+
+
 	virtual ~ISimpleElectraAudioPlayer() = default;
 	virtual bool Open(const TMap<FString, FVariant>& InOptions, const FString& ManifestURL, const FTimespan& StartPosition, const FTimespan& EncodedDuration, bool bAutoPlay, bool bSetLooping, TSharedPtr<IElectraPlayerDataCache, ESPMode::ThreadSafe> InPlayerDataCache) = 0;
 	virtual void SeekTo(const FTimespan& NewPosition) = 0;
