Electra Cloud Streaming Audio: limiting number of concurrent streams to 10 by default

#rnx

diff --git a/Engine/Plugins/Media/ElectraPlayer/Source/ElectraPlayerRuntime/Private/Runtime/Player/SimpleElectraAudioPlayer.cpp b/Engine/Plugins/Media/ElectraPlayer/Source/ElectraPlayerRuntime/Private/Runtime/Player/SimpleElectraAudioPlayer.cpp
--- a/Engine/Plugins/Media/ElectraPlayer/Source/ElectraPlayerRuntime/Private/Runtime/Player/SimpleElectraAudioPlayer.cpp
+++ b/Engine/Plugins/Media/ElectraPlayer/Source/ElectraPlayerRuntime/Private/Runtime/Player/SimpleElectraAudioPlayer.cpp
@@ -569,7 +569,7 @@ private:
 
 	volatile EState CurrentState = EState::Uninitialized;
 	volatile bool bIsReadyToTick = false;
-	volatile bool bAskedToBeReleased = false;
+	volatile bool bWasAskedToRelease = false;
 	volatile bool bCloseInProgress = false;
 	volatile bool bDestructionRequested = false;
 	volatile bool bDestructionIssued = false;
@@ -657,7 +657,7 @@ private:
 	bool HasBeenSuspended() const;
 	void UpdateAssetName();
 	void InternalApplyOptions();
-	void HandleOpenStream(bool bOnlyIfAvailable);
+	void HandleOpenStream(bool bTryToReopen);
 	static void MaybeRestartAStoppedStream(FSimpleElectraAudioPlayer* This);
 	void CreatePlayerAsync();
 	void CreateIdleBufferIfNecessary();
@@ -672,6 +672,7 @@ private:
 	static TArray<FSimpleElectraAudioPlayer*> AllInstances;
 	static TArray<FSimpleElectraAudioPlayer*> ActiveInstances;
 	static TArray<FSimpleElectraAudioPlayer*> StoppedInstances;
+	static std::atomic<int32> NumAwaitingStart;
 
 
 	struct FAnalyticsEntry
@@ -719,7 +720,7 @@ private:
 
 #if ENABLE_DEBUG_STATS
 public:
-	void DebugDrawInst(UCanvas* InCanvas);
+	void DebugDrawInst(UCanvas* InCanvas, int32& InOutNum);
 
 	/** Needed for debug draw management. */
 	static void EnableDebugDraw(bool bEnable);
@@ -738,6 +739,7 @@ FCriticalSection FSimpleElectraAudioPlayer::InstanceLock;
 TArray<FSimpleElectraAudioPlayer*> FSimpleElectraAudioPlayer::AllInstances;
 TArray<FSimpleElectraAudioPlayer*> FSimpleElectraAudioPlayer::ActiveInstances;
 TArray<FSimpleElectraAudioPlayer*> FSimpleElectraAudioPlayer::StoppedInstances;
+std::atomic<int32> FSimpleElectraAudioPlayer::NumAwaitingStart {0};
 
 FCriticalSection FSimpleElectraAudioPlayer::AnalyticsLock;
 TMap<FString, FSimpleElectraAudioPlayer::FAnalyticsEntry> FSimpleElectraAudioPlayer::AnalyticEntries;
@@ -939,7 +941,7 @@ void FSimpleElectraAudioPlayer::Tick(float DeltaTime)
 		}
 	}
 
-	if (bAskedToBeReleased && CurrentState != EState::CreatingPlayer)
+	if (bWasAskedToRelease && CurrentState != EState::Stopped && CurrentState != EState::CreatingPlayer)
 	{
 		FString Msg = FString::Printf(TEXT("%s: Stopping due to limit"), *AssetName);
 		UE_LOG(LogSimpleElectraPlayer, Verbose, TEXT("%s"), *Msg);
@@ -950,7 +952,7 @@ void FSimpleElectraAudioPlayer::Tick(float DeltaTime)
 	}
 }
 
-void FSimpleElectraAudioPlayer::HandleOpenStream(bool bOnlyIfAvailable)
+void FSimpleElectraAudioPlayer::HandleOpenStream(bool bTryToReopen)
 {
 	int32 MaxAllowed = CreateParams.MaxTotalPlayerInstances;
 	// CVar overrides unconditionally.
@@ -968,23 +970,22 @@ void FSimpleElectraAudioPlayer::HandleOpenStream(bool bOnlyIfAvailable)
 		FScopeLock lock(&InstanceLock);
 		
 		int32 NumActive = ActiveInstances.Num();
-		int32 NumRequestedToBeReleased = 0;
-		for(auto& Inst : ActiveInstances)
-		{
-			NumRequestedToBeReleased += Inst->bAskedToBeReleased ? 1 : 0;
-		}
 		if (NumActive >= MaxAllowed)
 		{
-			if (!bOnlyIfAvailable)
+			if (!bTryToReopen)
 			{
-				int32 NumToRemove = NumActive - NumRequestedToBeReleased - MaxAllowed + 1;
-
-				for(int32 i=0; NumToRemove>0 && i<NumActive; ++i)
+				int32 NumTaggedForRelease = 0;
+				for(auto& Inst : ActiveInstances)
+				{
+					NumTaggedForRelease += Inst->bWasAskedToRelease ? 1 : 0;
+				}
+				int32 NumNeedToRelease = NumAwaitingStart - NumTaggedForRelease;
+				for(int32 i=0; NumNeedToRelease>0 && i<NumActive; ++i)
 				{
-					if (!ActiveInstances[i]->bAskedToBeReleased)
+					if (!ActiveInstances[i]->bWasAskedToRelease)
 					{
-						ActiveInstances[i]->bAskedToBeReleased = true;
-						--NumToRemove;
+						ActiveInstances[i]->bWasAskedToRelease = true;
+						--NumNeedToRelease;
 					}
 				}
 			}
@@ -1001,11 +1002,16 @@ void FSimpleElectraAudioPlayer::HandleOpenStream(bool bOnlyIfAvailable)
 
 	if (bCreateNow)
 	{
-		FString Msg = FString::Printf(TEXT("%s: %s playback"), *AssetName, !bAskedToBeReleased ? TEXT("Preparing") : TEXT("Resuming"));
+		FString Msg = FString::Printf(TEXT("%s: %s playback"), *AssetName, !bTryToReopen ? TEXT("Preparing") : TEXT("Resuming"));
 		UE_LOG(LogSimpleElectraPlayer, Verbose, TEXT("%s"), *Msg);
 
+		if (!bTryToReopen)
+		{
+			--NumAwaitingStart;
+		}
+
 		// Track use count
-		if (bAskedToBeReleased)
+		if (bTryToReopen)
 		{
 			++Analytics.NumTimesResumed;
 		}
@@ -1031,9 +1037,8 @@ void FSimpleElectraAudioPlayer::HandleOpenStream(bool bOnlyIfAvailable)
 		InstanceLock.Unlock();
 
 		// Set state to active.
-		bIsResuming = bAskedToBeReleased;
+		bIsResuming = bTryToReopen;
 		CurrentState = EState::CreatingPlayer;
-		bAskedToBeReleased = false;
 		bMaybeReopenAfterSeek = false;
 		TimeStopped = -1.0;
 
@@ -1063,6 +1068,7 @@ void FSimpleElectraAudioPlayer::CreatePlayerAsync()
 		Player->SetLooping(lp);
 	}
 	CurrentState = EState::Active;
+	bWasAskedToRelease = false;
 }
 
 void FSimpleElectraAudioPlayer::CreateIdleBufferIfNecessary()
@@ -1169,7 +1175,7 @@ void FSimpleElectraAudioPlayer::MaybeRestartAStoppedStream(FSimpleElectraAudioPl
 			continue;
 		}
 		// Only those that are actually stopped and not in error or already restarting.
-		if (StoppedInst->CurrentState != EState::Stopped)
+		if (StoppedInst->CurrentState != EState::Stopped && StoppedInst->CurrentState != EState::TryReopeningStream)
 		{
 			continue;
 		}
@@ -1354,6 +1360,7 @@ bool FSimpleElectraAudioPlayer::Open(const TMap<FString, FVariant>& InOptions, c
 	PlayerDataCache = InPlayerDataCache;
 
 	CurrentState = EState::OpeningStream;
+	++NumAwaitingStart;
 	return true;
 }
 
@@ -1992,17 +1999,18 @@ void FSimpleElectraAudioPlayer::DebugDraw(UCanvas* InCanvas, APlayerController*
 			}
 		}
 
+		int32 Num = 0;
 		for(auto& act : active)
 		{
-			act->DebugDrawInst(InCanvas);
+			act->DebugDrawInst(InCanvas, Num);
 		}
 		for(auto& rdy : ready)
 		{
-			rdy->DebugDrawInst(InCanvas);
+			rdy->DebugDrawInst(InCanvas, Num);
 		}
 		for(auto& stp : stopped)
 		{
-			stp->DebugDrawInst(InCanvas);
+			stp->DebugDrawInst(InCanvas, Num);
 		}
 		InstanceLock.Unlock();
 	}
@@ -2020,9 +2028,12 @@ void FSimpleElectraAudioPlayer::DebugDrawPrintLine(UCanvas* InCanvas, const FStr
 	}
 }
 
-void FSimpleElectraAudioPlayer::DebugDrawInst(UCanvas* InCanvas)
+void FSimpleElectraAudioPlayer::DebugDrawInst(UCanvas* InCanvas, int32 &InOutNum)
 {
-	FString Msg(AssetName);
+	++InOutNum;
+	FString Msg;
+	//Msg = FString::Printf(TEXT("%2d %p "), InOutNum, this);
+	Msg += AssetName;
 	Msg += TEXT(": ");
 	TSharedPtr<IAdaptiveStreamingPlayer, ESPMode::ThreadSafe> LockedPlayer(Player);
 
