Electra Cloud Streaming Audio: updated the simple player's internal state machine

#rnx

diff --git a/Engine/Plugins/Media/ElectraPlayer/Source/ElectraPlayerRuntime/Private/Runtime/Player/SimpleElectraAudioPlayer.cpp b/Engine/Plugins/Media/ElectraPlayer/Source/ElectraPlayerRuntime/Private/Runtime/Player/SimpleElectraAudioPlayer.cpp
--- a/Engine/Plugins/Media/ElectraPlayer/Source/ElectraPlayerRuntime/Private/Runtime/Player/SimpleElectraAudioPlayer.cpp
+++ b/Engine/Plugins/Media/ElectraPlayer/Source/ElectraPlayerRuntime/Private/Runtime/Player/SimpleElectraAudioPlayer.cpp
@@ -48,7 +48,7 @@ static float ElectraCSAudio_MinRemainPlaytimeForReactivation = 6.0f;
 static float ElectraCSAudio_MaxStoppedDurationForReactivation = 20.0f;
 
 // Maximum number of concurrent instances. 0=unlimited
-static int32 ElectraCSAudio_MaxInstances = 0;
+static int32 ElectraCSAudio_MaxInstances = 10;
 
 // Whether or not to resume instaces that were stopped due to max instance limit being reached.
 static int32 ElectraCSAudio_ResumeStopped = 1;
@@ -159,13 +159,70 @@ static FString RedactMessage(FString InMessage)
 
 class FSimpleElectraAudioPlayerRenderer;
 
-class FSimpleElectraAudioPlayer : public ISimpleElectraAudioPlayer, public IAdaptiveStreamingPlayerMetrics, public FTickableGameObject, public TSharedFromThis<FSimpleElectraAudioPlayer, ESPMode::ThreadSafe>
+class FSimpleElectraAudioPlayer : public ISimpleElectraAudioPlayer, public IAdaptiveStreamingPlayerMetrics, public TSharedFromThis<FSimpleElectraAudioPlayer, ESPMode::ThreadSafe>
 {
+	class FTicker : public FTickableGameObject
+	{
+	public:
+		static void Acquire()
+		{
+			FScopeLock lock(&TickerCreateLock);
+			if (++NumClients == 1 && !Self)
+			{
+				Self = new FTicker;
+			}
+		}
+		static void Release()
+		{
+			FScopeLock lock(&TickerCreateLock);
+			if (--NumClients == 0)
+			{
+				FTicker* Prev = Self;
+				Self = nullptr;
+				lock.Unlock();
+				delete Prev;
+			}
+		}
+
+		virtual ~FTicker() {}
+		void Tick(float DeltaTime) override
+		{
+			FSimpleElectraAudioPlayer::TickAllInstances(DeltaTime);
+		}
+		ETickableTickType GetTickableTickType() const override
+		{ return ETickableTickType::Conditional; }
+		bool IsTickable() const override
+		{ return true; }
+		bool IsAllowedToTick() const override
+		{ return true; }
+		TStatId GetStatId() const override
+		{ RETURN_QUICK_DECLARE_CYCLE_STAT(FSimpleElectraAudioPlayer, STATGROUP_Tickables); }
+		bool IsTickableWhenPaused() const override
+		{ return true; }
+		bool IsTickableInEditor() const override
+		{ return true; }
+	private:
+		static FCriticalSection TickerCreateLock;
+		static volatile int32 NumClients;
+		static FTicker* Self;
+	};
+
+
 public:
-	static TSharedPtr<FSimpleElectraAudioPlayer, ESPMode::ThreadSafe> Create(const FCreateParams& InCreateParams)
+	static FSimpleElectraAudioPlayer* Create(const FCreateParams& InCreateParams)
 	{
-		return MakeShareable(new FSimpleElectraAudioPlayer(InCreateParams), TDeleter());
+		TSharedPtr<FSimpleElectraAudioPlayer, ESPMode::ThreadSafe> Inst = MakeShared<FSimpleElectraAudioPlayer, ESPMode::ThreadSafe>(InCreateParams);
+		Inst->Self = Inst;
+		return Inst.Get();
 	}
+	static void CloseAndDestroy(FSimpleElectraAudioPlayer* InInstance)
+	{
+		if (InInstance)
+		{
+			InInstance->bDestructionRequested = true;
+		}
+	}
+
 	static void SendAnalyticMetrics(const TSharedPtr<IAnalyticsProviderET>& InAnalyticsProvider);
 
 	static TSharedPtr<ICacheElementBase, ESPMode::ThreadSafe> GetCustomCacheElement(const FString& InForURL);
@@ -436,31 +493,9 @@ public:
 	void ReportDroppedAudioFrame() override
 	{ }
 
-private:
-	void Tick(float DeltaTime) override;
-	ETickableTickType GetTickableTickType() const override
-	{ return ETickableTickType::Conditional; }
-	bool IsTickable() const override
-	{ return bIsReadyToTick; }
-	bool IsAllowedToTick() const override
-	{ return bIsReadyToTick; }
-	TStatId GetStatId() const override
-	{ RETURN_QUICK_DECLARE_CYCLE_STAT(FSimpleElectraAudioPlayer, STATGROUP_Tickables); }
-	bool IsTickableWhenPaused() const override
-	{ return true; }
-	bool IsTickableInEditor() const override
-	{ return true; }
-
-
-	class TDeleter
-	{
-	public:
-		void operator()(FSimpleElectraAudioPlayer* InInstanceToDelete)
-		{
-			InInstanceToDelete->bDestructionRequested = true;
-		}
-	};
+	static void TickAllInstances(float InDeltaTime);
 
+private:
 	class FResourceProvider : public IAdaptiveStreamingPlayerResourceProvider
 	{
 	public:
@@ -560,20 +595,20 @@ private:
 		CreatingPlayer,
 		Active,
 		Stopped,
+		Stopping,
+		ClosingForDestruction,
+		Destructing,
 		Errored
 	};
 
 	const int32 kNumDummyBlockSamples = 1024;
 
 	FCreateParams CreateParams;
+	TSharedPtr<FSimpleElectraAudioPlayer, ESPMode::ThreadSafe> Self;
 
 	volatile EState CurrentState = EState::Uninitialized;
-	volatile bool bIsReadyToTick = false;
 	volatile bool bWasAskedToRelease = false;
-	volatile bool bCloseInProgress = false;
 	volatile bool bDestructionRequested = false;
-	volatile bool bDestructionIssued = false;
-	volatile bool bDestructionComplete = false;
 
 	const double SuggestedBufferDuration = 0.5;
 	const int32 MinNumSamplesInBuffer = 16384;	// Minimum number of samples, regardless of sample rate.
@@ -640,9 +675,13 @@ private:
 	double TimeCreated = -1.0;
 	double TimeUntilReady = -1.0;
 	double TimeStopped = -1.0;
+	int32 NewStreamOpenCount = 0;
 	bool bMaybeReopenAfterSeek = false;
 	bool bIsResuming = false;
 
+	void Tick();
+	bool HandleStopIfRequested();
+	bool HandleDestructionIfRequested();
 	void CreatePlayerInstance();
 	void ClosePlayerInstance(bool bDeleteThis);
 
@@ -672,7 +711,7 @@ private:
 	static TArray<FSimpleElectraAudioPlayer*> AllInstances;
 	static TArray<FSimpleElectraAudioPlayer*> ActiveInstances;
 	static TArray<FSimpleElectraAudioPlayer*> StoppedInstances;
-	static std::atomic<int32> NumAwaitingStart;
+	static std::atomic<int32> NumNewStreamStarts;
 
 
 	struct FAnalyticsEntry
@@ -739,10 +778,12 @@ FCriticalSection FSimpleElectraAudioPlayer::InstanceLock;
 TArray<FSimpleElectraAudioPlayer*> FSimpleElectraAudioPlayer::AllInstances;
 TArray<FSimpleElectraAudioPlayer*> FSimpleElectraAudioPlayer::ActiveInstances;
 TArray<FSimpleElectraAudioPlayer*> FSimpleElectraAudioPlayer::StoppedInstances;
-std::atomic<int32> FSimpleElectraAudioPlayer::NumAwaitingStart {0};
-
+std::atomic<int32> FSimpleElectraAudioPlayer::NumNewStreamStarts {0};
 FCriticalSection FSimpleElectraAudioPlayer::AnalyticsLock;
 TMap<FString, FSimpleElectraAudioPlayer::FAnalyticsEntry> FSimpleElectraAudioPlayer::AnalyticEntries;
+volatile int32 FSimpleElectraAudioPlayer::FTicker::NumClients = 0;
+FSimpleElectraAudioPlayer::FTicker* FSimpleElectraAudioPlayer::FTicker::Self = nullptr;
+FCriticalSection FSimpleElectraAudioPlayer::FTicker::TickerCreateLock;
 
 
 class FSimpleElectraAudioPlayerRenderer : public IMediaRenderer, public TSharedFromThis<FSimpleElectraAudioPlayerRenderer, ESPMode::ThreadSafe>
@@ -804,11 +845,16 @@ private:
 };
 
 
-TSharedPtr<ISimpleElectraAudioPlayer, ESPMode::ThreadSafe> ISimpleElectraAudioPlayer::Create(const FCreateParams& InCreateParams)
+ISimpleElectraAudioPlayer* ISimpleElectraAudioPlayer::Create(const FCreateParams& InCreateParams)
 {
 	return FSimpleElectraAudioPlayer::Create(InCreateParams);
 }
 
+void ISimpleElectraAudioPlayer::CloseAndDestroy(ISimpleElectraAudioPlayer* InInstance)
+{
+	FSimpleElectraAudioPlayer::CloseAndDestroy(static_cast<FSimpleElectraAudioPlayer*>(InInstance));
+}
+
 void ISimpleElectraAudioPlayer::SendAnalyticMetrics(const TSharedPtr<IAnalyticsProviderET>& InAnalyticsProvider)
 {
 	if (ElectraCSAudio_SendAnalytics)
@@ -853,95 +899,145 @@ FSimpleElectraAudioPlayer::FSimpleElectraAudioPlayer(const FCreateParams& InCrea
 	AllInstances.AddUnique(this);
 	InstanceLock.Unlock();
 
-	bIsReadyToTick = true;
+	FTicker::Acquire();
 }
 
 FSimpleElectraAudioPlayer::~FSimpleElectraAudioPlayer()
 {
-	bIsReadyToTick = false;
-
+#if !UE_BUILD_SHIPPING
 	InstanceLock.Lock();
 	check(AllInstances.Find(this) == INDEX_NONE);
 	check(ActiveInstances.Find(this) == INDEX_NONE);
 	check(StoppedInstances.Find(this) == INDEX_NONE);
 	InstanceLock.Unlock();
+#endif
 	// The player instance must have been destroyed already!
 	check(!Player.IsValid());
 	ResourceProvider.Reset();
 	Renderer.Reset();
 	FlushAudio(true);
+	FTicker::Release();
 }
 
-void FSimpleElectraAudioPlayer::Tick(float DeltaTime)
+void FSimpleElectraAudioPlayer::TickAllInstances(float InDeltaTime)
 {
-	if (bDestructionRequested)
+	InstanceLock.Lock();
+
+	// Check for limit and tag excess streams for stopping.
+	int32 MaxAllowed = ElectraCSAudio_MaxInstances;
+	// Is there a limit on the number of streams?
+	if (MaxAllowed > 0)
 	{
-		if (!bDestructionIssued)
+		int32 NumActive = ActiveInstances.Num();
+		if (NumActive >= MaxAllowed)
 		{
-			if (!bCloseInProgress)
+			int32 NumTaggedForRelease = 0;
+			for(auto& Inst : ActiveInstances)
 			{
-				FString Msg = FString::Printf(TEXT("%s: Closing"), *AssetName);
-				UE_LOG(LogSimpleElectraPlayer, Verbose, TEXT("%s"), *Msg);
-
-				bDestructionIssued = true;
-				ClosePlayerInstance(true);
-				InstanceLock.Lock();
-				StoppedInstances.Remove(this);
-				InstanceLock.Unlock();
+				NumTaggedForRelease += Inst->bWasAskedToRelease ? 1 : 0;
 			}
-		}
-		if (bDestructionComplete)
-		{
-			bIsReadyToTick = false;
-			InstanceLock.Lock();
-			AllInstances.Remove(this);
-			ActiveInstances.Remove(this);
-			StoppedInstances.Remove(this);
-			InstanceLock.Unlock();
-
-			if (ElectraCSAudio_ResumeStopped > 0)
+			int32 NumNeedToRelease = (NumActive - MaxAllowed) + NumNewStreamStarts - NumTaggedForRelease;
+			for(int32 i=0; NumNeedToRelease>0 && i<NumActive; ++i)
 			{
-				MaybeRestartAStoppedStream(this);
+				if (!ActiveInstances[i]->bWasAskedToRelease)
+				{
+					ActiveInstances[i]->bWasAskedToRelease = true;
+					--NumNeedToRelease;
+				}
 			}
-			delete this;
 		}
-		return;
 	}
 
+	TArray<FSimpleElectraAudioPlayer*> CurrentInstances = AllInstances;
+	InstanceLock.Unlock();
+	for(auto& Inst : CurrentInstances)
+	{
+		Inst->Tick();
+	}
+}
+
+void FSimpleElectraAudioPlayer::Tick()
+{
 	switch(CurrentState)
 	{
+		case EState::Uninitialized:
+		case EState::OpenBlob:
+		case EState::Errored:
+		{
+			HandleDestructionIfRequested();
+			return;
+		}
+		case EState::Stopping:
+		case EState::ClosingForDestruction:
+		{
+			return;
+		}
 		case EState::OpeningStream:
 		{
-			HandleOpenStream(false);
-			break;
+			if (!HandleDestructionIfRequested())
+			{
+				if (!HandleStopIfRequested())
+				{
+					HandleOpenStream(false);
+				}
+			}
+			return;
 		}
 		case EState::TryReopeningStream:
 		{
-			CurrentState = EState::Stopped;
-			HandleOpenStream(true);
-			break;
+			if (!HandleDestructionIfRequested())
+			{
+				CurrentState = EState::Stopped;
+				HandleOpenStream(true);
+			}
+			return;
 		}
 		case EState::CreatingPlayer:
 		{
 			return;
 		}
-		case EState::Errored:
+		case EState::Active:
 		{
+			if (!HandleDestructionIfRequested())
+			{
+				HandleStopIfRequested();
+			}
 			return;
 		}
 		case EState::Stopped:
 		{
-			if (bMaybeReopenAfterSeek)
+			if (!HandleDestructionIfRequested())
 			{
-				bMaybeReopenAfterSeek = false;
-				StartPosition = NextPTS;
-				HandleOpenStream(true);
+				if (bMaybeReopenAfterSeek)
+				{
+					bMaybeReopenAfterSeek = false;
+					StartPosition = NextPTS;
+					HandleOpenStream(true);
+				}
+			}
+			return;
+		}
+		case EState::Destructing:
+		{
+			InstanceLock.Lock();
+			AllInstances.Remove(this);
+			ActiveInstances.Remove(this);
+			StoppedInstances.Remove(this);
+			InstanceLock.Unlock();
+			if (ElectraCSAudio_ResumeStopped > 0)
+			{
+				MaybeRestartAStoppedStream(this);
 			}
+			check(Self.IsValid() && Self.IsUnique());
+			Self.Reset();
 			return;
 		}
 	}
+}
 
-	if (bWasAskedToRelease && CurrentState != EState::Stopped && CurrentState != EState::CreatingPlayer)
+bool FSimpleElectraAudioPlayer::HandleStopIfRequested()
+{
+	if (bWasAskedToRelease)
 	{
 		FString Msg = FString::Printf(TEXT("%s: Stopping due to limit"), *AssetName);
 		UE_LOG(LogSimpleElectraPlayer, Verbose, TEXT("%s"), *Msg);
@@ -949,48 +1045,37 @@ void FSimpleElectraAudioPlayer::Tick(float DeltaTime)
 		TimeStopped = FPlatformTime::Seconds();
 		++Analytics.NumTimesSuspended;
 		ClosePlayerInstance(false);
+		return true;
 	}
+	return false;
 }
 
-void FSimpleElectraAudioPlayer::HandleOpenStream(bool bTryToReopen)
+bool FSimpleElectraAudioPlayer::HandleDestructionIfRequested()
 {
-	int32 MaxAllowed = CreateParams.MaxTotalPlayerInstances;
-	// CVar overrides unconditionally.
-	if (ElectraCSAudio_MaxInstances > 0)
+	if (bDestructionRequested)
 	{
-		MaxAllowed = ElectraCSAudio_MaxInstances;
+		FString Msg = FString::Printf(TEXT("%s: Closing"), *AssetName);
+		UE_LOG(LogSimpleElectraPlayer, Verbose, TEXT("%s"), *Msg);
+		ClosePlayerInstance(true);
+		InstanceLock.Lock();
+		StoppedInstances.Remove(this);
+		InstanceLock.Unlock();
+		return true;
 	}
+	return false;
+}
 
-	bool bCreateNow = false;
 
+void FSimpleElectraAudioPlayer::HandleOpenStream(bool bTryToReopen)
+{
+	bool bCreateNow = false;
+	int32 MaxAllowed = ElectraCSAudio_MaxInstances;
 	// Is there a limit on the number of streams?
 	if (MaxAllowed > 0)
 	{
 		// Check if we need to make room for this new stream.
-		FScopeLock lock(&InstanceLock);
-		
 		int32 NumActive = ActiveInstances.Num();
-		if (NumActive >= MaxAllowed)
-		{
-			if (!bTryToReopen)
-			{
-				int32 NumTaggedForRelease = 0;
-				for(auto& Inst : ActiveInstances)
-				{
-					NumTaggedForRelease += Inst->bWasAskedToRelease ? 1 : 0;
-				}
-				int32 NumNeedToRelease = NumAwaitingStart - NumTaggedForRelease;
-				for(int32 i=0; NumNeedToRelease>0 && i<NumActive; ++i)
-				{
-					if (!ActiveInstances[i]->bWasAskedToRelease)
-					{
-						ActiveInstances[i]->bWasAskedToRelease = true;
-						--NumNeedToRelease;
-					}
-				}
-			}
-		}
-		else
+		if (NumActive < MaxAllowed)
 		{
 			bCreateNow = true;
 		}
@@ -1005,11 +1090,6 @@ void FSimpleElectraAudioPlayer::HandleOpenStream(bool bTryToReopen)
 		FString Msg = FString::Printf(TEXT("%s: %s playback"), *AssetName, !bTryToReopen ? TEXT("Preparing") : TEXT("Resuming"));
 		UE_LOG(LogSimpleElectraPlayer, Verbose, TEXT("%s"), *Msg);
 
-		if (!bTryToReopen)
-		{
-			--NumAwaitingStart;
-		}
-
 		// Track use count
 		if (bTryToReopen)
 		{
@@ -1020,6 +1100,9 @@ void FSimpleElectraAudioPlayer::HandleOpenStream(bool bTryToReopen)
 			++Analytics.NumTimesStarted;
 		}
 
+		NumNewStreamStarts = NumNewStreamStarts - NewStreamOpenCount;
+		NewStreamOpenCount = 0;
+
 		// Discard any potential leftovers.
 		Lock.Lock();
 		NextPendingSampleBlocks.Empty();
@@ -1046,6 +1129,7 @@ void FSimpleElectraAudioPlayer::HandleOpenStream(bool bTryToReopen)
 		TFunction<void()> CreateTask = [This]()
 		{
 			TSharedPtr<FSimpleElectraAudioPlayer, ESPMode::ThreadSafe> That = This.Pin();
+			check(That.IsValid());
 			if (That.IsValid())
 			{
 				That->CreatePlayerAsync();
@@ -1083,7 +1167,6 @@ void FSimpleElectraAudioPlayer::CreateIdleBufferIfNecessary()
 	Lock.Unlock();
 }
 
-
 void FSimpleElectraAudioPlayer::MergeAnalytics(const FString& AssetId, const FAnalyticsEntry& InAnalytics)
 {
 	// If this entry was not active and had no failure downloading the blob we
@@ -1261,9 +1344,8 @@ void FSimpleElectraAudioPlayer::DoClosePlayerAsync(TSharedPtr<FClosePlayerInstan
 		if (InPlayerHelper->bDeleteThis)
 		{
 			MergeAnalytics(InPlayerHelper->This->BaseURL, InPlayerHelper->This->Analytics);
-			InPlayerHelper->This->bDestructionComplete = true;
 		}
-		InPlayerHelper->This->bCloseInProgress = false;
+		InPlayerHelper->This->CurrentState = InPlayerHelper->bDeleteThis ? EState::Destructing : EState::Stopped;
 	};
 
 	if (GIsRunning)
@@ -1278,9 +1360,8 @@ void FSimpleElectraAudioPlayer::DoClosePlayerAsync(TSharedPtr<FClosePlayerInstan
 
 void FSimpleElectraAudioPlayer::ClosePlayerInstance(bool bDeleteThis)
 {
-	CurrentState = CurrentState == EState::Errored ? CurrentState : EState::Stopped;
-	bCloseInProgress = true;
-
+	CurrentState = bDeleteThis ? EState::ClosingForDestruction : EState::Stopping;
+	NumNewStreamStarts = NumNewStreamStarts - NewStreamOpenCount;
 	TSharedPtr<FClosePlayerInstanceHelper, ESPMode::ThreadSafe> Helper = MakeShared<FClosePlayerInstanceHelper, ESPMode::ThreadSafe>();
 	Helper->PlayerInstance = MoveTemp(Player);
 	Helper->This = this;
@@ -1360,7 +1441,8 @@ bool FSimpleElectraAudioPlayer::Open(const TMap<FString, FVariant>& InOptions, c
 	PlayerDataCache = InPlayerDataCache;
 
 	CurrentState = EState::OpeningStream;
-	++NumAwaitingStart;
+	NewStreamOpenCount = 1;
+	NumNewStreamStarts = NumNewStreamStarts + NewStreamOpenCount;
 	return true;
 }
 
diff --git a/Engine/Plugins/Media/ElectraPlayer/Source/ElectraPlayerRuntime/Public/SimpleElectraAudioPlayer.h b/Engine/Plugins/Media/ElectraPlayer/Source/ElectraPlayerRuntime/Public/SimpleElectraAudioPlayer.h
--- a/Engine/Plugins/Media/ElectraPlayer/Source/ElectraPlayerRuntime/Public/SimpleElectraAudioPlayer.h
+++ b/Engine/Plugins/Media/ElectraPlayer/Source/ElectraPlayerRuntime/Public/SimpleElectraAudioPlayer.h
@@ -11,16 +11,16 @@ class IAnalyticsProviderET;
 
 class ELECTRAPLAYERRUNTIME_API ISimpleElectraAudioPlayer
 {
+protected:
+	virtual ~ISimpleElectraAudioPlayer() = default;
 public:
 	struct FCreateParams
 	{
-		/** The number of audio players that can be simultaneously active at any time. 0=unlimited */
-		int32 MaxTotalPlayerInstances = 0;
-
 		/** A GUID to identify the new player instance. */
 		FGuid InstanceGUID;
 	};
-	static TSharedPtr<ISimpleElectraAudioPlayer, ESPMode::ThreadSafe> Create(const FCreateParams& InCreateParams);
+	static ISimpleElectraAudioPlayer* Create(const FCreateParams& InCreateParams);
+	static void CloseAndDestroy(ISimpleElectraAudioPlayer* InInstance);
 
 	/**
 	 * Delivers the aggregated metrics to the analytics provider and clears the internal list.
@@ -37,7 +37,6 @@ public:
 	static void SetCustomCacheElement(const FString& InForURL, const TSharedPtr<ICacheElementBase, ESPMode::ThreadSafe>& InElement);
 
 
-	virtual ~ISimpleElectraAudioPlayer() = default;
 	virtual bool Open(const TMap<FString, FVariant>& InOptions, const FString& ManifestURL, const FTimespan& StartPosition, const FTimespan& EncodedDuration, bool bAutoPlay, bool bSetLooping, TSharedPtr<IElectraPlayerDataCache, ESPMode::ThreadSafe> InPlayerDataCache) = 0;
 	virtual void SeekTo(const FTimespan& NewPosition) = 0;
 	virtual void PrepareToLoopToBeginning() = 0;
