ElectraPlayer: Added EOS buffer to pending buffer return list (if necessary) so it is not sent out while there are still buffers pending to be returned which would break the ordering.

#rnx

diff --git a/Engine/Plugins/Media/ElectraPlayer/Source/ElectraPlayerRuntime/Private/Runtime/Player/AdaptiveStreamingPlayerRenderers.cpp b/Engine/Plugins/Media/ElectraPlayer/Source/ElectraPlayerRuntime/Private/Runtime/Player/AdaptiveStreamingPlayerRenderers.cpp
--- a/Engine/Plugins/Media/ElectraPlayer/Source/ElectraPlayerRuntime/Private/Runtime/Player/AdaptiveStreamingPlayerRenderers.cpp
+++ b/Engine/Plugins/Media/ElectraPlayer/Source/ElectraPlayerRuntime/Private/Runtime/Player/AdaptiveStreamingPlayerRenderers.cpp
@@ -16,6 +16,7 @@
 #define VALIDITY_VALUE_KEY TEXT("$renderVV$")
 #define DURATION_VALUE_KEY TEXT("duration")
 #define TIMESTAMP_VALUE_KEY TEXT("pts")
+#define EOS_VALUE_KEY TEXT("eos")
 
 #define AUDIO_BUFFER_SIZE TEXT("max_buffer_size")
 #define AUDIO_BUFFER_NUM TEXT("num_buffers")
@@ -464,36 +465,38 @@ UEMediaError FAdaptiveStreamingWrappedRenderer::ReturnBufferCommon(IBuffer* Buff
 		InSampleProperties.SetOrUpdate(DURATION_VALUE_KEY, FVariantValue(Duration));
 	}
 
+	bool bIsUnusedReturnBuffer = bRender == false && InSampleProperties.GetValue(EOS_VALUE_KEY).SafeGetBool(false) == false;
+
 	FScopeLock lock(&Lock);
 	EnqueuedDuration += Duration;
 	++NumEnqueuedSamples;
 
-	if (bRender)
+	if (!bIsUnusedReturnBuffer)
 	{
-	    bool bHoldback = !bIsRunning;
-	    // If the video renderer shall not hold back the first frame (used for scrubbing video)
-	    // then we pass it out. The count is reset in Flush().
-	    if (Type == EStreamType::Video && bDoNotHoldBackFirstVideoFrame)
-	    {
-		    if (NumBuffersNotHeldBack == 0)
-		    {
-			    bHoldback = false;
-		    }
-		    if (!bHoldback)
-		    {
-			    ++NumBuffersNotHeldBack;
-		    }
-	    }
-	    if (bHoldback)
-	    {
-		    FPendingReturnBuffer pb;
-		    pb.Buffer = Buffer;
-		    pb.bRender = bRender;
-		    pb.Properties = InSampleProperties;
-		    PendingReturnBuffers.Enqueue(MoveTemp(pb));
-		    ++NumPendingReturnBuffers;
-		    return UEMEDIA_ERROR_OK;
-	    }
+		bool bHoldback = !bIsRunning;
+		// If the video renderer shall not hold back the first frame (used for scrubbing video)
+		// then we pass it out. The count is reset in Flush().
+		if (Type == EStreamType::Video && bDoNotHoldBackFirstVideoFrame)
+		{
+			if (NumBuffersNotHeldBack == 0)
+			{
+				bHoldback = false;
+			}
+			if (!bHoldback)
+			{
+				++NumBuffersNotHeldBack;
+			}
+		}
+		if (bHoldback)
+		{
+			FPendingReturnBuffer pb;
+			pb.Buffer = Buffer;
+			pb.bRender = bRender;
+			pb.Properties = InSampleProperties;
+			PendingReturnBuffers.Enqueue(MoveTemp(pb));
+			++NumPendingReturnBuffers;
+			return UEMEDIA_ERROR_OK;
+		}
 	}
 	lock.Unlock();
 	return WrappedRenderer->ReturnBuffer(Buffer, bRender, InSampleProperties);
diff --git a/Engine/Plugins/Media/ElectraPlayer/Source/ElectraPlayerRuntime/Private/Runtime/Player/SimpleElectraAudioPlayer.cpp b/Engine/Plugins/Media/ElectraPlayer/Source/ElectraPlayerRuntime/Private/Runtime/Player/SimpleElectraAudioPlayer.cpp
--- a/Engine/Plugins/Media/ElectraPlayer/Source/ElectraPlayerRuntime/Private/Runtime/Player/SimpleElectraAudioPlayer.cpp
+++ b/Engine/Plugins/Media/ElectraPlayer/Source/ElectraPlayerRuntime/Private/Runtime/Player/SimpleElectraAudioPlayer.cpp
@@ -564,7 +564,7 @@ private:
 		void AddCustomData(const FString& InForURL, const TSharedPtr<ISimpleElectraAudioPlayer::ICacheElementBase, ESPMode::ThreadSafe>& InData);
 	private:
 		void CheckLimits(int32 InBytesNeeded);
-		FCriticalSection Lock;
+		FCriticalSection AccessLock;
 		TMap<FString, TSharedPtr<FBlobCacheEntry, ESPMode::ThreadSafe>> EntryMap;
 		TArray<TSharedPtr<FBlobCacheEntry, ESPMode::ThreadSafe>> EntryLRU;
 		int32 SizeInUse = 0;
@@ -634,7 +634,7 @@ private:
 	FString ErrorMessage;
 	bool bHasErrored = false;
 
-	mutable FCriticalSection Lock;
+	mutable FCriticalSection InstanceLock;
 	FStreamFormat StreamFormat;
 
 	struct FBlockSequence
@@ -707,7 +707,7 @@ private:
 		return Cache;
 	}
 
-	static FCriticalSection InstanceLock;
+	static FCriticalSection InstanceListLock;
 	static TArray<FSimpleElectraAudioPlayer*> AllInstances;
 	static TArray<FSimpleElectraAudioPlayer*> ActiveInstances;
 	static TArray<FSimpleElectraAudioPlayer*> StoppedInstances;
@@ -774,7 +774,7 @@ public:
 	static const FLinearColor DebugColor_Red;
 #endif
 };
-FCriticalSection FSimpleElectraAudioPlayer::InstanceLock;
+FCriticalSection FSimpleElectraAudioPlayer::InstanceListLock;
 TArray<FSimpleElectraAudioPlayer*> FSimpleElectraAudioPlayer::AllInstances;
 TArray<FSimpleElectraAudioPlayer*> FSimpleElectraAudioPlayer::ActiveInstances;
 TArray<FSimpleElectraAudioPlayer*> FSimpleElectraAudioPlayer::StoppedInstances;
@@ -895,9 +895,9 @@ FSimpleElectraAudioPlayer::FSimpleElectraAudioPlayer(const FCreateParams& InCrea
 	StreamFormat.SampleRate = 0;
 	NextPTS = FTimespan::MinValue();
 
-	InstanceLock.Lock();
+	InstanceListLock.Lock();
 	AllInstances.AddUnique(this);
-	InstanceLock.Unlock();
+	InstanceListLock.Unlock();
 
 	FTicker::Acquire();
 }
@@ -905,11 +905,11 @@ FSimpleElectraAudioPlayer::FSimpleElectraAudioPlayer(const FCreateParams& InCrea
 FSimpleElectraAudioPlayer::~FSimpleElectraAudioPlayer()
 {
 #if !UE_BUILD_SHIPPING
-	InstanceLock.Lock();
+	InstanceListLock.Lock();
 	check(AllInstances.Find(this) == INDEX_NONE);
 	check(ActiveInstances.Find(this) == INDEX_NONE);
 	check(StoppedInstances.Find(this) == INDEX_NONE);
-	InstanceLock.Unlock();
+	InstanceListLock.Unlock();
 #endif
 	// The player instance must have been destroyed already!
 	check(!Player.IsValid());
@@ -921,7 +921,7 @@ FSimpleElectraAudioPlayer::~FSimpleElectraAudioPlayer()
 
 void FSimpleElectraAudioPlayer::TickAllInstances(float InDeltaTime)
 {
-	InstanceLock.Lock();
+	InstanceListLock.Lock();
 
 	// Check for limit and tag excess streams for stopping.
 	int32 MaxAllowed = ElectraCSAudio_MaxInstances;
@@ -949,7 +949,7 @@ void FSimpleElectraAudioPlayer::TickAllInstances(float InDeltaTime)
 	}
 
 	TArray<FSimpleElectraAudioPlayer*> CurrentInstances = AllInstances;
-	InstanceLock.Unlock();
+	InstanceListLock.Unlock();
 	for(auto& Inst : CurrentInstances)
 	{
 		Inst->Tick();
@@ -1019,11 +1019,11 @@ void FSimpleElectraAudioPlayer::Tick()
 		}
 		case EState::Destructing:
 		{
-			InstanceLock.Lock();
+			InstanceListLock.Lock();
 			AllInstances.Remove(this);
 			ActiveInstances.Remove(this);
 			StoppedInstances.Remove(this);
-			InstanceLock.Unlock();
+			InstanceListLock.Unlock();
 			if (ElectraCSAudio_ResumeStopped > 0)
 			{
 				MaybeRestartAStoppedStream(this);
@@ -1057,9 +1057,9 @@ bool FSimpleElectraAudioPlayer::HandleDestructionIfRequested()
 		FString Msg = FString::Printf(TEXT("%s: Closing"), *AssetName);
 		UE_LOG(LogSimpleElectraPlayer, Verbose, TEXT("%s"), *Msg);
 		ClosePlayerInstance(true);
-		InstanceLock.Lock();
+		InstanceListLock.Lock();
 		StoppedInstances.Remove(this);
-		InstanceLock.Unlock();
+		InstanceListLock.Unlock();
 		return true;
 	}
 	return false;
@@ -1104,20 +1104,20 @@ void FSimpleElectraAudioPlayer::HandleOpenStream(bool bTryToReopen)
 		NewStreamOpenCount = 0;
 
 		// Discard any potential leftovers.
-		Lock.Lock();
+		InstanceLock.Lock();
 		NextPendingSampleBlocks.Empty();
 		CurrentWriteSampleBlock.Reset();
 		CurrentReadSampleBlock.Reset();
 		NumActiveFramesTotal = 0;
 		NumEnqueuedBlocks = 0;
 		bIsFirstSampleBlock = true;
-		Lock.Unlock();
+		InstanceLock.Unlock();
 
 		// Add to the list of active instances.
-		InstanceLock.Lock();
+		InstanceListLock.Lock();
 		ActiveInstances.AddUnique(this);
 		StoppedInstances.Remove(this);
-		InstanceLock.Unlock();
+		InstanceListLock.Unlock();
 
 		// Set state to active.
 		bIsResuming = bTryToReopen;
@@ -1157,14 +1157,14 @@ void FSimpleElectraAudioPlayer::CreatePlayerAsync()
 
 void FSimpleElectraAudioPlayer::CreateIdleBufferIfNecessary()
 {
-	Lock.Lock();
+	InstanceLock.Lock();
 	if (!CurrentReadSampleBlock.IsValid() && NextPendingSampleBlocks.IsEmpty())
 	{
 		TSharedPtr<FBlockSequence, ESPMode::ThreadSafe> NewSeq = MakeShared<FBlockSequence, ESPMode::ThreadSafe>();
 		NewSeq->SequenceIndex = -1;
 		NextPendingSampleBlocks.Emplace(NewSeq);
 	}
-	Lock.Unlock();
+	InstanceLock.Unlock();
 }
 
 void FSimpleElectraAudioPlayer::MergeAnalytics(const FString& AssetId, const FAnalyticsEntry& InAnalytics)
@@ -1244,7 +1244,7 @@ void FSimpleElectraAudioPlayer::SendAnalyticMetrics(const TSharedPtr<IAnalyticsP
 void FSimpleElectraAudioPlayer::MaybeRestartAStoppedStream(FSimpleElectraAudioPlayer* This)
 {
 	double TimeNow = FPlatformTime::Seconds();
-	FScopeLock lock(&InstanceLock);
+	FScopeLock lock(&InstanceListLock);
 	for(auto& StoppedInst :	StoppedInstances)
 	{
 		// Don't restart ourselves.
@@ -1334,13 +1334,13 @@ void FSimpleElectraAudioPlayer::DoClosePlayerAsync(TSharedPtr<FClosePlayerInstan
 		{
 			InPlayerHelper->This->CreateIdleBufferIfNecessary();
 		}
-		InstanceLock.Lock();
+		InstanceListLock.Lock();
 		ActiveInstances.Remove(InPlayerHelper->This);
 		if (!InPlayerHelper->bDeleteThis)
 		{
 			StoppedInstances.AddUnique(InPlayerHelper->This);
 		}
-		InstanceLock.Unlock();
+		InstanceListLock.Unlock();
 		if (InPlayerHelper->bDeleteThis)
 		{
 			MergeAnalytics(InPlayerHelper->This->BaseURL, InPlayerHelper->This->Analytics);
@@ -1422,12 +1422,14 @@ bool FSimpleElectraAudioPlayer::Open(const TMap<FString, FVariant>& InOptions, c
 			CreatePlayerInstance();
 			if (!PendingBlobRequest->Request->SetFromJSON(BlobParameters.GetValue()))
 			{
+				FScopeLock lock(&InstanceLock);
 				CurrentState = EState::Errored;
 				bHasErrored = true;
 				ErrorMessage = TEXT("Could not parse blob parameters");
 				PendingBlobRequest.Reset();
 				return false;
 			}
+			FScopeLock lock(&InstanceLock);
 			CurrentState = EState::OpenBlob;
 			PendingBlobRequest->Request->URL(InManifestURL).Callback().BindThreadSafeSP(PendingBlobRequest.ToSharedRef(), &FBlobRequest::OnBlobRequestComplete);
 			Player->LoadBlob(PendingBlobRequest->Request);
@@ -1440,6 +1442,7 @@ bool FSimpleElectraAudioPlayer::Open(const TMap<FString, FVariant>& InOptions, c
 	PendingBlobRequest.Reset();
 	PlayerDataCache = InPlayerDataCache;
 
+	FScopeLock lock(&InstanceLock);
 	CurrentState = EState::OpeningStream;
 	NewStreamOpenCount = 1;
 	NumNewStreamStarts = NumNewStreamStarts + NewStreamOpenCount;
@@ -1525,7 +1528,7 @@ void FSimpleElectraAudioPlayer::FResourceProvider::ProvideStaticPlaybackDataForU
 
 int32 FSimpleElectraAudioPlayer::GetBinaryMetadata(TSharedPtr<TArray<uint8>, ESPMode::ThreadSafe>& OutMetadata) const
 {
-	FScopeLock lock(&Lock);
+	FScopeLock lock(&InstanceLock);
 	if (PendingBlobRequest.IsValid())
 	{
 		if (PendingBlobRequest->bIsComplete)
@@ -1550,14 +1553,14 @@ int32 FSimpleElectraAudioPlayer::GetBinaryMetadata(TSharedPtr<TArray<uint8>, ESP
 
 bool FSimpleElectraAudioPlayer::GetStreamFormat(FStreamFormat& OutFormat) const
 {
-	FScopeLock lock(&Lock);
+	FScopeLock lock(&InstanceLock);
 	OutFormat = StreamFormat;
 	return StreamFormat.SampleRate && StreamFormat.NumChannels;
 }
 
 bool FSimpleElectraAudioPlayer::CanAcceptAudioFrames(int32 InNumFrames)
 {
-	FScopeLock lock(&Lock);
+	FScopeLock lock(&InstanceLock);
 	if (StreamFormat.SampleRate)
 	{
 		double DurationInBuffer = (double)NumActiveFramesTotal / StreamFormat.SampleRate;
@@ -1568,7 +1571,7 @@ bool FSimpleElectraAudioPlayer::CanAcceptAudioFrames(int32 InNumFrames)
 
 bool FSimpleElectraAudioPlayer::GetEnqueuedFrameInfo(int32& OutNumberOfEnqueuedFrames, FTimespan& OutDurationOfEnqueuedFrames) const
 {
-	FScopeLock lock(&Lock);
+	FScopeLock lock(&InstanceLock);
 	OutNumberOfEnqueuedFrames = NumEnqueuedBlocks;
 	if (StreamFormat.SampleRate)
 	{
@@ -1584,7 +1587,7 @@ bool FSimpleElectraAudioPlayer::GetEnqueuedFrameInfo(int32& OutNumberOfEnqueuedF
 
 void FSimpleElectraAudioPlayer::PrepareToLoopToBeginning()
 {
-	FScopeLock lock(&Lock);
+	FScopeLock lock(&InstanceLock);
 	bool bIsSuspended = HasBeenSuspended();
 	if (bIsSuspended)
 	{
@@ -1622,7 +1625,7 @@ bool FSimpleElectraAudioPlayer::EnqueueAudioFrames(const void* InBufferAddress,
 		sb.PTS = InPTS;
 		FMemory::Memcpy(sb.Buffer, InBufferAddress, nb);
 
-		FScopeLock lock(&Lock);
+		FScopeLock lock(&InstanceLock);
 
 		// Get first PTS.
 		if (NextPTS < FTimespan::Zero())
@@ -1654,18 +1657,24 @@ bool FSimpleElectraAudioPlayer::EnqueueAudioFrames(const void* InBufferAddress,
 
 void FSimpleElectraAudioPlayer::SetReceivedLastBuffer()
 {
-	FScopeLock lock(&Lock);
-	check(CurrentWriteSampleBlock.IsValid());
+	FScopeLock lock(&InstanceLock);
 	if (CurrentWriteSampleBlock.IsValid())
 	{
 		CurrentWriteSampleBlock->bReadEnded = true;
 	}
+	else
+	{
+		TSharedPtr<FBlockSequence, ESPMode::ThreadSafe> NewSeq = MakeShared<FBlockSequence, ESPMode::ThreadSafe>();
+		NewSeq->bReadEnded = true;
+		NextPendingSampleBlocks.Emplace(NewSeq);
+		CurrentWriteSampleBlock = MoveTemp(NewSeq);
+	}
 }
 
 void FSimpleElectraAudioPlayer::FlushAudio(bool bForceFlush)
 {
 	bool bIsSuspended = HasBeenSuspended();
-	FScopeLock lock(&Lock);
+	FScopeLock lock(&InstanceLock);
 	// Only flush when not terminated. Otherwise we want to return all the remaining samples we have
 	if (!bIsSuspended || bForceFlush)
 	{
@@ -1691,7 +1700,7 @@ int64 FSimpleElectraAudioPlayer::GetNextSamples(FTimespan& OutPTS, bool& bOutIsF
 		return 0;
 	}
 
-	FScopeLock lock(&Lock);
+	FScopeLock lock(&InstanceLock);
 	bool bIsSuspended = HasBeenSuspended();
 	ActivateBlockReadSequence();
 	if (!CurrentReadSampleBlock.IsValid())
@@ -1743,6 +1752,7 @@ int64 FSimpleElectraAudioPlayer::GetNextSamples(FTimespan& OutPTS, bool& bOutIsF
 					CurrentReadSampleBlock->bReachedEOS = true;
 				}
 			}
+			lock.Unlock();
 			// Fill the output with silence.
 			FMemory::Memzero(OutBuffer, InNumFramesToGet * sizeof(int16) * nc);
 			return InNumFramesToGet; 
@@ -1782,7 +1792,7 @@ int64 FSimpleElectraAudioPlayer::GetNextSamples(FTimespan& OutPTS, bool& bOutIsF
 TSharedPtr<TArray<uint8>, ESPMode::ThreadSafe> FSimpleElectraAudioPlayer::FBlobCache::GetData(const FString& InForURL)
 {
 	double Now = FPlatformTime::Seconds();
-	FScopeLock lock(&Lock);
+	FScopeLock lock(&AccessLock);
 
 	TSharedPtr<FBlobCacheEntry, ESPMode::ThreadSafe>* Entry = EntryMap.Find(InForURL);
 	if (Entry)
@@ -1805,7 +1815,7 @@ TSharedPtr<TArray<uint8>, ESPMode::ThreadSafe> FSimpleElectraAudioPlayer::FBlobC
 void FSimpleElectraAudioPlayer::FBlobCache::AddData(const FString& InForURL, const TSharedPtr<TArray<uint8>, ESPMode::ThreadSafe>& InData)
 {
 	double Now = FPlatformTime::Seconds();
-	FScopeLock lock(&Lock);
+	FScopeLock lock(&AccessLock);
 
 	TSharedPtr<FBlobCacheEntry, ESPMode::ThreadSafe>* Entry = EntryMap.Find(InForURL);
 	if (!Entry)
@@ -1829,14 +1839,14 @@ void FSimpleElectraAudioPlayer::FBlobCache::AddData(const FString& InForURL, con
 
 TSharedPtr<ISimpleElectraAudioPlayer::ICacheElementBase, ESPMode::ThreadSafe> FSimpleElectraAudioPlayer::FBlobCache::GetCustomData(const FString& InForURL)
 {
-	FScopeLock lock(&Lock);
+	FScopeLock lock(&AccessLock);
 	TSharedPtr<FBlobCacheEntry, ESPMode::ThreadSafe>* Entry = EntryMap.Find(InForURL);
 	return Entry ? (*Entry)->CustomData : nullptr;
 }
 
 void FSimpleElectraAudioPlayer::FBlobCache::AddCustomData(const FString& InForURL, const TSharedPtr<ISimpleElectraAudioPlayer::ICacheElementBase, ESPMode::ThreadSafe>& InData)
 {
-	FScopeLock lock(&Lock);
+	FScopeLock lock(&AccessLock);
 	TSharedPtr<FBlobCacheEntry, ESPMode::ThreadSafe>* Entry = EntryMap.Find(InForURL);
 	if (Entry)
 	{
@@ -2060,7 +2070,7 @@ void FSimpleElectraAudioPlayer::DebugDraw(UCanvas* InCanvas, APlayerController*
 		DebugDrawTextPos.Y = 20.0f;
 
 
-		InstanceLock.Lock();
+		InstanceListLock.Lock();
 		TArray<FSimpleElectraAudioPlayer*> active;
 		TArray<FSimpleElectraAudioPlayer*> ready;
 		TArray<FSimpleElectraAudioPlayer*> stopped;
@@ -2094,7 +2104,7 @@ void FSimpleElectraAudioPlayer::DebugDraw(UCanvas* InCanvas, APlayerController*
 		{
 			stp->DebugDrawInst(InCanvas, Num);
 		}
-		InstanceLock.Unlock();
+		InstanceListLock.Unlock();
 	}
 }
 
