ElectraPlayer: fixed race condition creating the singleton task counter

#rnx

diff --git a/Engine/Plugins/Media/ElectraPlayer/Source/ElectraPlayerRuntime/Private/Runtime/PlayerRuntimeGlobal.cpp b/Engine/Plugins/Media/ElectraPlayer/Source/ElectraPlayerRuntime/Private/Runtime/PlayerRuntimeGlobal.cpp
--- a/Engine/Plugins/Media/ElectraPlayer/Source/ElectraPlayerRuntime/Private/Runtime/PlayerRuntimeGlobal.cpp
+++ b/Engine/Plugins/Media/ElectraPlayer/Source/ElectraPlayerRuntime/Private/Runtime/PlayerRuntimeGlobal.cpp
@@ -238,7 +238,7 @@ namespace Electra
 			if (!AllDoneSignal)
 			{
 				FMediaEvent* NewSignal = new FMediaEvent();
-				if (TMediaInterlockedExchangePointer(AllDoneSignal, NewSignal) != nullptr)
+				if (FMediaInterlockedCompareExchangePointer((void* volatile&)AllDoneSignal, (void*)NewSignal, (void*)nullptr) != nullptr)
 				{
 					delete NewSignal;
 				}
