Electra: Changed parameter dictionary to use FName as key to reduce temporary allocations as much as possible

#rnx

diff --git a/Engine/Plugins/Media/ElectraPlayer/Source/ElectraPlayerPlugin/Private/ElectraPlayerPlugin.cpp b/Engine/Plugins/Media/ElectraPlayer/Source/ElectraPlayerPlugin/Private/ElectraPlayerPlugin.cpp
--- a/Engine/Plugins/Media/ElectraPlayer/Source/ElectraPlayerPlugin/Private/ElectraPlayerPlugin.cpp
+++ b/Engine/Plugins/Media/ElectraPlayer/Source/ElectraPlayerPlugin/Private/ElectraPlayerPlugin.cpp
@@ -607,7 +607,7 @@ bool FElectraPlayerPlugin::Open(const FString& Url, const IMediaOptions* Options
 
 	// Set up options to initialize the internal player with.
 	Electra::FParamDict PlayerOptions;
-	const TCHAR * const OptionsByString[] = { TEXT("excluded_codecs_video"), TEXT("excluded_codecs_audio"), TEXT("excluded_codecs_subtitles"), TEXT("preferred_codecs_video"), TEXT("preferred_codecs_audio"),TEXT("preferred_codecs_subtitles") };
+	const FName OptionsByString[] = { TEXT("excluded_codecs_video"), TEXT("excluded_codecs_audio"), TEXT("excluded_codecs_subtitles"), TEXT("preferred_codecs_video"), TEXT("preferred_codecs_audio"),TEXT("preferred_codecs_subtitles") };
 	for(auto &StringOption : OptionsByString)
 	{
 		FString Value = Options->GetMediaOption(StringOption, FString());
@@ -621,49 +621,49 @@ bool FElectraPlayerPlugin::Open(const FString& Url, const IMediaOptions* Options
 	int64 InitialStreamBitrate = Options->GetMediaOption(TEXT("ElectraInitialBitrate"), (int64)-1);
 	if (InitialStreamBitrate > 0)
 	{
-		PlayerOptions.Set("initial_bitrate", Electra::FVariantValue(InitialStreamBitrate));
+		PlayerOptions.Set(TEXT("initial_bitrate"), Electra::FVariantValue(InitialStreamBitrate));
 		UE_LOG(LogElectraPlayerPlugin, Log, TEXT("[%p] IMediaPlayer::Open: Using initial bitrate of %d bits/second"), this, (int32)InitialStreamBitrate);
 	}
 	FString MediaMimeType = Options->GetMediaOption(TEXT("mimetype"), FString());
 	if (MediaMimeType.Len())
 	{
-		PlayerOptions.Set("mime_type", Electra::FVariantValue(MediaMimeType));
+		PlayerOptions.Set(TEXT("mime_type"), Electra::FVariantValue(MediaMimeType));
 		UE_LOG(LogElectraPlayerPlugin, Log, TEXT("[%p] IMediaPlayer::Open: Setting media mime type to \"%s\""), this, *MediaMimeType);
 	}
 	int64 MaxVerticalHeight = Options->GetMediaOption(TEXT("MaxElectraVerticalResolution"), (int64)-1);
 	if (MaxVerticalHeight > 0)
 	{
-		PlayerOptions.Set("max_resoY", Electra::FVariantValue(MaxVerticalHeight));
+		PlayerOptions.Set(TEXT("max_resoY"), Electra::FVariantValue(MaxVerticalHeight));
 		UE_LOG(LogElectraPlayerPlugin, Log, TEXT("[%p] IMediaPlayer::Open: Limiting vertical resolution to %d for all streams"), this, (int32)MaxVerticalHeight);
 	}
 	int64 MaxVerticalHeightAt60 = Options->GetMediaOption(TEXT("MaxElectraVerticalResolutionOf60fpsVideos"), (int64)-1);
 	if (MaxVerticalHeightAt60 > 0)
 	{
-		PlayerOptions.Set("max_resoY_above_30fps", Electra::FVariantValue(MaxVerticalHeightAt60));
+		PlayerOptions.Set(TEXT("max_resoY_above_30fps"), Electra::FVariantValue(MaxVerticalHeightAt60));
 		UE_LOG(LogElectraPlayerPlugin, Log, TEXT("[%p] IMediaPlayer::Open: Limiting vertical resolution to %d for streams >30fps"), this, (int32)MaxVerticalHeightAt60);
 	}
 	double LiveEdgeDistanceForNormalPresentation = Options->GetMediaOption(TEXT("ElectraLivePresentationOffset"), (double)-1.0);
 	if (LiveEdgeDistanceForNormalPresentation > 0.0)
 	{
-		PlayerOptions.Set("seekable_range_live_end_offset", Electra::FVariantValue(Electra::FTimeValue().SetFromSeconds(LiveEdgeDistanceForNormalPresentation)));
+		PlayerOptions.Set(TEXT("seekable_range_live_end_offset"), Electra::FVariantValue(Electra::FTimeValue().SetFromSeconds(LiveEdgeDistanceForNormalPresentation)));
 		UE_LOG(LogElectraPlayerPlugin, Log, TEXT("[%p] IMediaPlayer::Open: Setting distance to live edge for normal presentations to %.3f seconds"), this, LiveEdgeDistanceForNormalPresentation);
 	}
 	double LiveEdgeDistanceForAudioOnlyPresentation = Options->GetMediaOption(TEXT("ElectraLiveAudioPresentationOffset"), (double)-1.0);
 	if (LiveEdgeDistanceForAudioOnlyPresentation > 0.0)
 	{
-		PlayerOptions.Set("seekable_range_live_end_offset_audioonly", Electra::FVariantValue(Electra::FTimeValue().SetFromSeconds(LiveEdgeDistanceForAudioOnlyPresentation)));
+		PlayerOptions.Set(TEXT("seekable_range_live_end_offset_audioonly"), Electra::FVariantValue(Electra::FTimeValue().SetFromSeconds(LiveEdgeDistanceForAudioOnlyPresentation)));
 		UE_LOG(LogElectraPlayerPlugin, Log, TEXT("[%p] IMediaPlayer::Open: Setting distance to live edge for audio-only presentation to %.3f seconds"), this, LiveEdgeDistanceForAudioOnlyPresentation);
 	}
 	bool bUseConservativeLiveEdgeDistance = Options->GetMediaOption(TEXT("ElectraLiveUseConservativePresentationOffset"), (bool)false);
 	if (bUseConservativeLiveEdgeDistance)
 	{
-		PlayerOptions.Set("seekable_range_live_end_offset_conservative", Electra::FVariantValue(bUseConservativeLiveEdgeDistance));
+		PlayerOptions.Set(TEXT("seekable_range_live_end_offset_conservative"), Electra::FVariantValue(bUseConservativeLiveEdgeDistance));
 		UE_LOG(LogElectraPlayerPlugin, Log, TEXT("[%p] IMediaPlayer::Open: Using conservative live edge for distance calculation"), this);
 	}
 	bool bThrowErrorWhenRebuffering = Options->GetMediaOption(TEXT("ElectraThrowErrorWhenRebuffering"), (bool)false);
 	if (bThrowErrorWhenRebuffering)
 	{
-		PlayerOptions.Set("throw_error_when_rebuffering", Electra::FVariantValue(bThrowErrorWhenRebuffering));
+		PlayerOptions.Set(TEXT("throw_error_when_rebuffering"), Electra::FVariantValue(bThrowErrorWhenRebuffering));
 		UE_LOG(LogElectraPlayerPlugin, Log, TEXT("[%p] IMediaPlayer::Open: Throw playback error when rebuffering"), this);
 	}
 	FString CDNHTTPStatusDenyStream = Options->GetMediaOption(TEXT("ElectraGetDenyStreamCode"), FString());
@@ -673,7 +673,7 @@ bool FElectraPlayerPlugin::Open(const FString& Url, const IMediaOptions* Options
 		LexFromString(HTTPStatus, *CDNHTTPStatusDenyStream);
 		if (HTTPStatus > 0 && HTTPStatus < 1000)
 		{
-			PlayerOptions.Set("abr:cdn_deny_httpstatus", Electra::FVariantValue((int64)HTTPStatus));
+			PlayerOptions.Set(TEXT("abr:cdn_deny_httpstatus"), Electra::FVariantValue((int64)HTTPStatus));
 			UE_LOG(LogElectraPlayerPlugin, Log, TEXT("[%p] IMediaPlayer::Open: CDN HTTP status %d will deny a stream permanently"), this, HTTPStatus);
 		}
 	}
@@ -685,15 +685,15 @@ bool FElectraPlayerPlugin::Open(const FString& Url, const IMediaOptions* Options
 		int64 ScrubCacheSizeKiB = Options->GetMediaOption(TEXT("ElectraScrubCacheSizeKiB"), (int64)-1);
 		if (ScrubSeekBitrate >= 0)
 		{
-			PlayerOptions.SetOrUpdate("seekstart_bitrate", Electra::FVariantValue(ScrubSeekBitrate));
+			PlayerOptions.Set(TEXT("seekstart_bitrate"), Electra::FVariantValue(ScrubSeekBitrate));
 		}
 		if (ScrubCacheSizeKiB > 0)
 		{
-			PlayerOptions.SetOrUpdate("httpcache_max_bytesize", Electra::FVariantValue(ScrubCacheSizeKiB << 10));
-			PlayerOptions.SetOrUpdate("httpcache_max_entries", Electra::FVariantValue((int64)10000));
+			PlayerOptions.Set(TEXT("httpcache_max_bytesize"), Electra::FVariantValue(ScrubCacheSizeKiB << 10));
+			PlayerOptions.Set(TEXT("httpcache_max_entries"), Electra::FVariantValue((int64)10000));
 		}
-		PlayerOptions.SetOrUpdate("optimize_seek_for_scrubbing", Electra::FVariantValue(true));
-		PlayerOptions.SetOrUpdate("do_not_hold_back_first_frame", Electra::FVariantValue(true));
+		PlayerOptions.Set(TEXT("optimize_seek_for_scrubbing"), Electra::FVariantValue(true));
+		PlayerOptions.Set(TEXT("do_not_hold_back_first_frame"), Electra::FVariantValue(true));
 	}
 
 	// Check for options that can be changed during playback and apply them at startup already.
diff --git a/Engine/Plugins/Media/ElectraPlayer/Source/ElectraPlayerPlugin/Private/ElectraPlayerPluginModule.cpp b/Engine/Plugins/Media/ElectraPlayer/Source/ElectraPlayerPlugin/Private/ElectraPlayerPluginModule.cpp
--- a/Engine/Plugins/Media/ElectraPlayer/Source/ElectraPlayerPlugin/Private/ElectraPlayerPluginModule.cpp
+++ b/Engine/Plugins/Media/ElectraPlayer/Source/ElectraPlayerPlugin/Private/ElectraPlayerPluginModule.cpp
@@ -107,7 +107,7 @@ public:
 			}
 
 			Electra::FParamDict Params;
-			Params.Set("GetDeviceTypeCallback", Electra::FVariantValue((void*)&FElectraPlayerPluginModule::GetDynamicRHIInfo));
+			Params.Set(FName(TEXT("GetDeviceTypeCallback")), Electra::FVariantValue((void*)&FElectraPlayerPluginModule::GetDynamicRHIInfo));
 			if (!FElectraPlayerPlatform::StartupPlatformResources(Params))
 			{
 				UE_LOG(LogElectraPlayerPlugin, Log, TEXT("Platform resource setup failed! Electra Player plugin is not initialised."));
diff --git a/Engine/Plugins/Media/ElectraPlayer/Source/ElectraPlayerRuntime/Private/ElectraPlayer.cpp b/Engine/Plugins/Media/ElectraPlayer/Source/ElectraPlayerRuntime/Private/ElectraPlayer.cpp
--- a/Engine/Plugins/Media/ElectraPlayer/Source/ElectraPlayerRuntime/Private/ElectraPlayer.cpp
+++ b/Engine/Plugins/Media/ElectraPlayer/Source/ElectraPlayerRuntime/Private/ElectraPlayer.cpp
@@ -7,6 +7,7 @@
 
 #include "PlayerRuntimeGlobal.h"
 #include "Player/AdaptiveStreamingPlayer.h"
+#include "Player/AdaptivePlayerOptionKeynames.h"
 #include "Player/IExternalDataReader.h"
 #include "Renderer/RendererVideo.h"
 #include "Renderer/RendererAudio.h"
@@ -267,7 +268,7 @@ bool FElectraPlayer::OpenInternal(const FString& Url, const FParamDict& InPlayer
 		FParamDict PlayerOptions(InPlayerOptions);
 		if (PlaystartOptions.ExternalDataReader.IsValid())
 		{
-			PlayerOptions.Set(TEXT("use_external_data_reader"), FVariantValue(true));
+			PlayerOptions.Set(Electra::OptionKeyUseExternalDataReader, FVariantValue(true));
 			StaticResourceProvider->SetExternalDataReader(PlaystartOptions.ExternalDataReader);
 		}
 
@@ -306,7 +307,7 @@ bool FElectraPlayer::OpenInternal(const FString& Url, const FParamDict& InPlayer
 
 		if (InOpenType == IElectraPlayerInterface::EOpenType::Blob)
 		{
-			const TCHAR * const KeyBlob = TEXT("blobparams");
+			const FName KeyBlob(TEXT("blobparams"));
 			if (PlayerOptions.HaveKey(KeyBlob))
 			{
 				BlobParams = PlayerOptions.GetValue(KeyBlob).SafeGetFString();
diff --git a/Engine/Plugins/Media/ElectraPlayer/Source/ElectraPlayerRuntime/Private/Runtime/Decoder/Android/AndroidElectraDecoderResourceManager.cpp b/Engine/Plugins/Media/ElectraPlayer/Source/ElectraPlayerRuntime/Private/Runtime/Decoder/Android/AndroidElectraDecoderResourceManager.cpp
--- a/Engine/Plugins/Media/ElectraPlayer/Source/ElectraPlayerRuntime/Private/Runtime/Decoder/Android/AndroidElectraDecoderResourceManager.cpp
+++ b/Engine/Plugins/Media/ElectraPlayer/Source/ElectraPlayerRuntime/Private/Runtime/Decoder/Android/AndroidElectraDecoderResourceManager.cpp
@@ -192,26 +192,25 @@ bool FElectraDecoderResourceManagerAndroid::SetupRenderBufferFromDecoderOutput(I
 	
 	FInstanceVars* InInst = static_cast<FInstanceVars*>(InPlatformSpecificResource);
 	TSharedPtr<IVideoDecoderResourceDelegate, ESPMode::ThreadSafe> PinnedResourceDelegate = InInst->VideoDecoderResourceDelegate.Pin();
-	TSharedPtr<FElectraPlayerVideoDecoderOutputAndroid, ESPMode::ThreadSafe> DecoderOutput = InOutBufferToSetup->GetBufferProperties().GetValue("texture").GetSharedPointer<FElectraPlayerVideoDecoderOutputAndroid>();
+	TSharedPtr<FElectraPlayerVideoDecoderOutputAndroid, ESPMode::ThreadSafe> DecoderOutput = InOutBufferToSetup->GetBufferProperties().GetValue(RenderOptionKeys::Texture).GetSharedPointer<FElectraPlayerVideoDecoderOutputAndroid>();
 	if (DecoderOutput.IsValid() && PinnedResourceDelegate.IsValid())
 	{
 		FElectraVideoDecoderOutputCropValues Crop = InDecoderOutput->GetCropValues();
-		InOutBufferPropertes->Set(TEXT("width"), FVariantValue((int64) InDecoderOutput->GetWidth()));
-		InOutBufferPropertes->Set(TEXT("height"), FVariantValue((int64) InDecoderOutput->GetHeight()));
-		InOutBufferPropertes->Set(TEXT("crop_left"), FVariantValue((int64) Crop.Left));
-		InOutBufferPropertes->Set(TEXT("crop_right"), FVariantValue((int64) Crop.Right));
-		InOutBufferPropertes->Set(TEXT("crop_top"), FVariantValue((int64) Crop.Top));
-		InOutBufferPropertes->Set(TEXT("crop_bottom"), FVariantValue((int64) Crop.Bottom));
-		InOutBufferPropertes->Set(TEXT("aspect_ratio"), FVariantValue((double) InDecoderOutput->GetAspectRatioW() / (double) InDecoderOutput->GetAspectRatioH()));
-		InOutBufferPropertes->Set(TEXT("aspect_w"), FVariantValue((int64) InDecoderOutput->GetAspectRatioW()));
-		InOutBufferPropertes->Set(TEXT("aspect_h"), FVariantValue((int64) InDecoderOutput->GetAspectRatioH()));
-		InOutBufferPropertes->Set(TEXT("fps_num"), FVariantValue((int64) InDecoderOutput->GetFrameRateNumerator()));
-		InOutBufferPropertes->Set(TEXT("fps_denom"), FVariantValue((int64) InDecoderOutput->GetFrameRateDenominator()));
+		InOutBufferPropertes->Set(IDecoderOutputOptionNames::Width, FVariantValue((int64)InDecoderOutput->GetWidth()));
+		InOutBufferPropertes->Set(IDecoderOutputOptionNames::Height, FVariantValue((int64)InDecoderOutput->GetHeight()));
+		InOutBufferPropertes->Set(IDecoderOutputOptionNames::CropLeft, FVariantValue((int64)Crop.Left));
+		InOutBufferPropertes->Set(IDecoderOutputOptionNames::CropRight, FVariantValue((int64)Crop.Right));
+		InOutBufferPropertes->Set(IDecoderOutputOptionNames::CropTop, FVariantValue((int64)Crop.Top));
+		InOutBufferPropertes->Set(IDecoderOutputOptionNames::CropBottom, FVariantValue((int64)Crop.Bottom));
+		InOutBufferPropertes->Set(IDecoderOutputOptionNames::AspectRatio, FVariantValue((double)InDecoderOutput->GetAspectRatioW() / (double)InDecoderOutput->GetAspectRatioH()));
+		InOutBufferPropertes->Set(IDecoderOutputOptionNames::AspectW, FVariantValue((int64)InDecoderOutput->GetAspectRatioW()));
+		InOutBufferPropertes->Set(IDecoderOutputOptionNames::AspectH, FVariantValue((int64)InDecoderOutput->GetAspectRatioH()));
+		InOutBufferPropertes->Set(IDecoderOutputOptionNames::FPSNumerator, FVariantValue((int64)InDecoderOutput->GetFrameRateNumerator()));
+		InOutBufferPropertes->Set(IDecoderOutputOptionNames::FPSDenominator, FVariantValue((int64)InDecoderOutput->GetFrameRateDenominator()));
 		int32 NumBits = InDecoderOutput->GetNumberOfBits();
-		InOutBufferPropertes->Set(TEXT("pixelfmt"), FVariantValue((int64)((NumBits > 8)  ? EPixelFormat::PF_A2B10G10R10 : EPixelFormat::PF_B8G8R8A8)));
-		InOutBufferPropertes->Set(TEXT("bits_per"), FVariantValue((int64) NumBits));
-
-		InOutBufferPropertes->Set(TEXT("pitch"), FVariantValue((int64)InDecoderOutput->GetDecodedWidth() * ((NumBits > 8) ? 2 : 1)));
+		InOutBufferPropertes->Set(IDecoderOutputOptionNames::PixelFormat, FVariantValue((int64)((NumBits > 8)  ? EPixelFormat::PF_A2B10G10R10 : EPixelFormat::PF_B8G8R8A8)));
+		InOutBufferPropertes->Set(IDecoderOutputOptionNames::BitsPerComponent, FVariantValue((int64) NumBits));
+		InOutBufferPropertes->Set(IDecoderOutputOptionNames::Pitch, FVariantValue((int64)InDecoderOutput->GetDecodedWidth() * ((NumBits > 8) ? 2 : 1)));
 
 		FElectraDecoderVideoOutputCopyResources cr;
 		IElectraDecoderVideoOutput::EImageCopyResult CopyResult = InDecoderOutput->CopyPlatformImage(&cr);
diff --git a/Engine/Plugins/Media/ElectraPlayer/Source/ElectraPlayerRuntime/Private/Runtime/Decoder/Apple/AppleElectraDecoderResourceManager.cpp b/Engine/Plugins/Media/ElectraPlayer/Source/ElectraPlayerRuntime/Private/Runtime/Decoder/Apple/AppleElectraDecoderResourceManager.cpp
--- a/Engine/Plugins/Media/ElectraPlayer/Source/ElectraPlayerRuntime/Private/Runtime/Decoder/Apple/AppleElectraDecoderResourceManager.cpp
+++ b/Engine/Plugins/Media/ElectraPlayer/Source/ElectraPlayerRuntime/Private/Runtime/Decoder/Apple/AppleElectraDecoderResourceManager.cpp
@@ -91,29 +91,27 @@ bool FElectraDecoderResourceManagerApple::SetupRenderBufferFromDecoderOutput(IMe
 	check(InOutBufferPropertes.IsValid());
 	check(InDecoderOutput.IsValid());
 
-	TSharedPtr<FElectraPlayerVideoDecoderOutputApple, ESPMode::ThreadSafe> DecoderOutput = InOutBufferToSetup->GetBufferProperties().GetValue("texture").GetSharedPointer<FElectraPlayerVideoDecoderOutputApple>();
+	TSharedPtr<FElectraPlayerVideoDecoderOutputApple, ESPMode::ThreadSafe> DecoderOutput = InOutBufferToSetup->GetBufferProperties().GetValue(RenderOptionKeys::Texture).GetSharedPointer<FElectraPlayerVideoDecoderOutputApple>();
 	FInstanceVars* Vars = static_cast<FInstanceVars*>(InPlatformSpecificResource);
 
 	if (DecoderOutput.IsValid())
 	{
 		FElectraVideoDecoderOutputCropValues Crop = InDecoderOutput->GetCropValues();
-		InOutBufferPropertes->Set(TEXT("width"), FVariantValue((int64)InDecoderOutput->GetWidth()));
-		InOutBufferPropertes->Set(TEXT("height"), FVariantValue((int64)InDecoderOutput->GetHeight()));
-		InOutBufferPropertes->Set(TEXT("crop_left"), FVariantValue((int64)Crop.Left));
-		InOutBufferPropertes->Set(TEXT("crop_right"), FVariantValue((int64)Crop.Right));
-		InOutBufferPropertes->Set(TEXT("crop_top"), FVariantValue((int64)Crop.Top));
-		InOutBufferPropertes->Set(TEXT("crop_bottom"), FVariantValue((int64)Crop.Bottom));
-		InOutBufferPropertes->Set(TEXT("aspect_ratio"), FVariantValue((double)InDecoderOutput->GetAspectRatioW() / (double)InDecoderOutput->GetAspectRatioH()));
-		InOutBufferPropertes->Set(TEXT("aspect_w"), FVariantValue((int64)InDecoderOutput->GetAspectRatioW()));
-		InOutBufferPropertes->Set(TEXT("aspect_h"), FVariantValue((int64)InDecoderOutput->GetAspectRatioH()));
-		InOutBufferPropertes->Set(TEXT("fps_num"), FVariantValue((int64)InDecoderOutput->GetFrameRateNumerator()));
-		InOutBufferPropertes->Set(TEXT("fps_denom"), FVariantValue((int64)InDecoderOutput->GetFrameRateDenominator()));
+		InOutBufferPropertes->Set(IDecoderOutputOptionNames::Width, FVariantValue((int64)InDecoderOutput->GetWidth()));
+		InOutBufferPropertes->Set(IDecoderOutputOptionNames::Height, FVariantValue((int64)InDecoderOutput->GetHeight()));
+		InOutBufferPropertes->Set(IDecoderOutputOptionNames::CropLeft, FVariantValue((int64)Crop.Left));
+		InOutBufferPropertes->Set(IDecoderOutputOptionNames::CropRight, FVariantValue((int64)Crop.Right));
+		InOutBufferPropertes->Set(IDecoderOutputOptionNames::CropTop, FVariantValue((int64)Crop.Top));
+		InOutBufferPropertes->Set(IDecoderOutputOptionNames::CropBottom, FVariantValue((int64)Crop.Bottom));
+		InOutBufferPropertes->Set(IDecoderOutputOptionNames::AspectRatio, FVariantValue((double)InDecoderOutput->GetAspectRatioW() / (double)InDecoderOutput->GetAspectRatioH()));
+		InOutBufferPropertes->Set(IDecoderOutputOptionNames::AspectW, FVariantValue((int64)InDecoderOutput->GetAspectRatioW()));
+		InOutBufferPropertes->Set(IDecoderOutputOptionNames::AspectH, FVariantValue((int64)InDecoderOutput->GetAspectRatioH()));
+		InOutBufferPropertes->Set(IDecoderOutputOptionNames::FPSNumerator, FVariantValue((int64)InDecoderOutput->GetFrameRateNumerator()));
+		InOutBufferPropertes->Set(IDecoderOutputOptionNames::FPSDenominator, FVariantValue((int64)InDecoderOutput->GetFrameRateDenominator()));
+		InOutBufferPropertes->Set(IDecoderOutputOptionNames::BitsPerComponent, FVariantValue((int64)InDecoderOutput->GetNumberOfBits()));
 
 		int32 Width = InDecoderOutput->GetDecodedWidth();
 		int32 Height = InDecoderOutput->GetDecodedHeight();
-
-		InOutBufferPropertes->Set(TEXT("bits_per"), FVariantValue((int64)InDecoderOutput->GetNumberOfBits()));
-
 		//
 		// Image buffers?
 		//
@@ -165,9 +163,9 @@ bool FElectraDecoderResourceManagerApple::SetupRenderBufferFromDecoderOutput(IMe
 
 			int32 Pitch = ImageBuffers->GetBufferPitchByIndex(0);
 
-			InOutBufferPropertes->Set(TEXT("pixelfmt"), FVariantValue((int64)RHIPixFmt));
-			InOutBufferPropertes->Set(TEXT("pixelenc"), FVariantValue((int64)DecPixEnc));
-			InOutBufferPropertes->Set(TEXT("pitch"), FVariantValue((int64)Pitch));
+			InOutBufferPropertes->Set(IDecoderOutputOptionNames::PixelFormat, FVariantValue((int64)RHIPixFmt));
+			InOutBufferPropertes->Set(IDecoderOutputOptionNames::PixelEncoding, FVariantValue((int64)DecPixEnc));
+			InOutBufferPropertes->Set(IDecoderOutputOptionNames::Pitch, FVariantValue((int64)Pitch));
 
 // [...] ALPHA BUFFER -- HOW DO WE PASS IT ON!?
 // [...] ANY CS/HDR INFO FROM THE DECODER? -- PASSING IT ON WOULD BE EASY, BUT RIGHT NOW WE ALWAYS READ IT FROM THE CONTAINER (in code further up the chain)
diff --git a/Engine/Plugins/Media/ElectraPlayer/Source/ElectraPlayerRuntime/Private/Runtime/Decoder/AudioDecoder.cpp b/Engine/Plugins/Media/ElectraPlayer/Source/ElectraPlayerRuntime/Private/Runtime/Decoder/AudioDecoder.cpp
--- a/Engine/Plugins/Media/ElectraPlayer/Source/ElectraPlayerRuntime/Private/Runtime/Decoder/AudioDecoder.cpp
+++ b/Engine/Plugins/Media/ElectraPlayer/Source/ElectraPlayerRuntime/Private/Runtime/Decoder/AudioDecoder.cpp
@@ -199,8 +199,9 @@ private:
 	bool																	bMustBeSuspendedInBackground = false;
 
 	IMediaRenderer::IBuffer*												CurrentOutputBuffer = nullptr;
-	FParamDict																BufferAcquireOptions;
+	FParamDict																EmptyOptions;
 	FParamDict																OutputBufferSampleProperties;
+	FParamDict																EOSBufferProperties;
 	FAudioChannelMapper														ChannelMapper;
 };
 ENUM_CLASS_FLAGS(FAudioDecoderImpl::EAUChangeFlags);
@@ -267,6 +268,12 @@ TSharedPtr<IElectraCodecFactory, ESPMode::ThreadSafe> FAudioDecoderImpl::GetDeco
 FAudioDecoderImpl::FAudioDecoderImpl()
 	: FMediaThread("ElectraPlayer::Audio decoder")
 {
+	EOSBufferProperties.Set(RenderOptionKeys::NumChannels, FVariantValue((int64)0));
+	EOSBufferProperties.Set(RenderOptionKeys::UsedByteSize, FVariantValue((int64)0));
+	EOSBufferProperties.Set(RenderOptionKeys::SampleRate, FVariantValue((int64) 0));
+	EOSBufferProperties.Set(RenderOptionKeys::Duration, FVariantValue(FTimeValue::GetZero()));
+	EOSBufferProperties.Set(RenderOptionKeys::PTS, FVariantValue(FTimeValue::GetInvalid()));
+	EOSBufferProperties.Set(RenderOptionKeys::EOSFlag, FVariantValue(true));
 }
 
 FAudioDecoderImpl::~FAudioDecoderImpl()
@@ -365,13 +372,13 @@ void FAudioDecoderImpl::CreateDecoderOutputPool()
 	const int32 kMaxChannels = 16;
 	const int32 kMaxSamplesPerBlock = 2048;
 	uint32 frameSize = sizeof(float) * kMaxChannels * kMaxSamplesPerBlock;
-	poolOpts.Set(TEXT("max_buffer_size"), FVariantValue((int64) frameSize));
-	poolOpts.Set(TEXT("num_buffers"), FVariantValue((int64) 8));
-	poolOpts.Set(TEXT("samples_per_block"), FVariantValue((int64) kMaxSamplesPerBlock));
-	poolOpts.Set(TEXT("max_channels"), FVariantValue((int64) kMaxChannels));
+	poolOpts.Set(RenderOptionKeys::MaxBufferSize, FVariantValue((int64) frameSize));
+	poolOpts.Set(RenderOptionKeys::NumBuffers, FVariantValue((int64) 8));
+	poolOpts.Set(RenderOptionKeys::SamplesPerBlock, FVariantValue((int64) kMaxSamplesPerBlock));
+	poolOpts.Set(RenderOptionKeys::MaxChannels, FVariantValue((int64) kMaxChannels));
 	if (Renderer->CreateBufferPool(poolOpts) == UEMEDIA_ERROR_OK)
 	{
-		MaxDecodeBufferSize = (int32) Renderer->GetBufferPoolProperties().GetValue(TEXT("max_buffers")).GetInt64();
+		MaxDecodeBufferSize = (int32) Renderer->GetBufferPoolProperties().GetValue(RenderOptionKeys::MaxBuffers).GetInt64();
 	}
 	else
 	{
@@ -526,8 +533,7 @@ void FAudioDecoderImpl::ReturnUnusedOutputBuffer()
 {
 	if (CurrentOutputBuffer)
 	{
-		OutputBufferSampleProperties.Clear();
-		Renderer->ReturnBuffer(CurrentOutputBuffer, false, OutputBufferSampleProperties);
+		Renderer->ReturnBuffer(CurrentOutputBuffer, false, EmptyOptions);
 		CurrentOutputBuffer = nullptr;
 	}
 }
@@ -691,7 +697,7 @@ IElectraDecoder::EOutputStatus FAudioDecoderImpl::HandleOutput()
 			{
 				SCOPE_CYCLE_COUNTER(STAT_ElectraPlayer_AudioConvertOutput);
 				CSV_SCOPED_TIMING_STAT(ElectraPlayer, AudioConvertOutput);
-				UEMediaError bufResult = Renderer->AcquireBuffer(CurrentOutputBuffer, 0, BufferAcquireOptions);
+				UEMediaError bufResult = Renderer->AcquireBuffer(CurrentOutputBuffer, 0, EmptyOptions);
 				check(bufResult == UEMEDIA_ERROR_OK || bufResult == UEMEDIA_ERROR_INSUFFICIENT_DATA);
 				if (bufResult != UEMEDIA_ERROR_OK && bufResult != UEMEDIA_ERROR_INSUFFICIENT_DATA)
 				{
@@ -819,19 +825,18 @@ IElectraDecoder::EOutputStatus FAudioDecoderImpl::HandleOutput()
 				int32 NumDecodedSamplesPerChannel = NumDecodedBytes / NumBytesPerFrame;
 				if (NumDecodedSamplesPerChannel)
 				{
-					int32 CurrentRenderOutputBufferSize = (int32)CurrentOutputBuffer->GetBufferProperties().GetValue(TEXT("size")).GetInt64();
-					void* CurrentRenderOutputBufferAddress = CurrentOutputBuffer->GetBufferProperties().GetValue(TEXT("address")).GetPointer();
+					int32 CurrentRenderOutputBufferSize = (int32)CurrentOutputBuffer->GetBufferProperties().GetValue(RenderOptionKeys::AllocatedSize).GetInt64();
+					void* CurrentRenderOutputBufferAddress = CurrentOutputBuffer->GetBufferProperties().GetValue(RenderOptionKeys::AllocatedAddress).GetPointer();
 					ChannelMapper.MapChannels(CurrentRenderOutputBufferAddress, CurrentRenderOutputBufferSize, AdvancePointer(PCMBuffer, ByteOffsetToFirstSample), NumDecodedBytes, NumDecodedSamplesPerChannel);
 
 					FTimeValue Duration;
 					Duration.SetFromND(NumSamplesProduced, SamplingRate);
 
-					OutputBufferSampleProperties.Clear();
-					OutputBufferSampleProperties.Set(TEXT("num_channels"), FVariantValue((int64)ChannelMapper.GetNumTargetChannels()));
-					OutputBufferSampleProperties.Set(TEXT("byte_size"), FVariantValue((int64)((int64)NumDecodedSamplesPerChannel * ChannelMapper.GetNumTargetChannels() * NumBytesPerSample)));
-					OutputBufferSampleProperties.Set(TEXT("sample_rate"), FVariantValue((int64) SamplingRate));
-					OutputBufferSampleProperties.Set(TEXT("duration"), FVariantValue(Duration));
-					OutputBufferSampleProperties.Set(TEXT("pts"), FVariantValue(MatchingInput->AdjustedPTS));
+					OutputBufferSampleProperties.Set(RenderOptionKeys::NumChannels, FVariantValue((int64)ChannelMapper.GetNumTargetChannels()));
+					OutputBufferSampleProperties.Set(RenderOptionKeys::UsedByteSize, FVariantValue((int64)((int64)NumDecodedSamplesPerChannel * ChannelMapper.GetNumTargetChannels() * NumBytesPerSample)));
+					OutputBufferSampleProperties.Set(RenderOptionKeys::SampleRate, FVariantValue((int64) SamplingRate));
+					OutputBufferSampleProperties.Set(RenderOptionKeys::Duration, FVariantValue(Duration));
+					OutputBufferSampleProperties.Set(RenderOptionKeys::PTS, FVariantValue(MatchingInput->AdjustedPTS));
 
 					Renderer->ReturnBuffer(CurrentOutputBuffer, true, OutputBufferSampleProperties);
 					CurrentOutputBuffer = nullptr;
@@ -949,7 +954,7 @@ bool FAudioDecoderImpl::SendSilenceOrEOS(ESendMode InSendMode)
 		{
 			SCOPE_CYCLE_COUNTER(STAT_ElectraPlayer_AudioConvertOutput);
 			CSV_SCOPED_TIMING_STAT(ElectraPlayer, AudioConvertOutput);
-			UEMediaError bufResult = Renderer->AcquireBuffer(CurrentOutputBuffer, 0, BufferAcquireOptions);
+			UEMediaError bufResult = Renderer->AcquireBuffer(CurrentOutputBuffer, 0, EmptyOptions);
 			check(bufResult == UEMEDIA_ERROR_OK || bufResult == UEMEDIA_ERROR_INSUFFICIENT_DATA);
 			if (bufResult != UEMEDIA_ERROR_OK && bufResult != UEMEDIA_ERROR_INSUFFICIENT_DATA)
 			{
@@ -975,8 +980,8 @@ bool FAudioDecoderImpl::SendSilenceOrEOS(ESendMode InSendMode)
 			if (InSendMode == ESendMode::SendSilence)
 			{
 				// Clear to silence
-				int64 MaxBufferSize = CurrentOutputBuffer->GetBufferProperties().GetValue(TEXT("size")).GetInt64();
-				FMemory::Memzero(CurrentOutputBuffer->GetBufferProperties().GetValue(TEXT("address")).GetPointer(), MaxBufferSize);
+				int64 MaxBufferSize = CurrentOutputBuffer->GetBufferProperties().GetValue(RenderOptionKeys::AllocatedSize).GetInt64();
+				FMemory::Memzero(CurrentOutputBuffer->GetBufferProperties().GetValue(RenderOptionKeys::AllocatedAddress).GetPointer(), MaxBufferSize);
 
 				int32 NumChannels = CurrentOutputFormat.NumChannels ? CurrentOutputFormat.NumChannels : 2;
 				int32 SampleRate = CurrentOutputFormat.SampleRate ? CurrentOutputFormat.SampleRate : 48000;
@@ -988,12 +993,11 @@ bool FAudioDecoderImpl::SendSilenceOrEOS(ESendMode InSendMode)
 					FTimeValue Duration;
 					Duration.SetFromND(SamplesPerBlock, (uint32) SampleRate);
 
-					OutputBufferSampleProperties.Clear();
-					OutputBufferSampleProperties.Set(TEXT("num_channels"), FVariantValue((int64)NumChannels));
-					OutputBufferSampleProperties.Set(TEXT("sample_rate"), FVariantValue((int64)SampleRate));
-					OutputBufferSampleProperties.Set(TEXT("byte_size"), FVariantValue((int64)(EmptySize < MaxBufferSize ? EmptySize : MaxBufferSize)));
-					OutputBufferSampleProperties.Set(TEXT("duration"), FVariantValue(Duration));
-					OutputBufferSampleProperties.Set(TEXT("pts"), FVariantValue(CurrentAccessUnit->AdjustedPTS));
+					OutputBufferSampleProperties.Set(RenderOptionKeys::NumChannels, FVariantValue((int64)NumChannels));
+					OutputBufferSampleProperties.Set(RenderOptionKeys::UsedByteSize, FVariantValue((int64)(EmptySize < MaxBufferSize ? EmptySize : MaxBufferSize)));
+					OutputBufferSampleProperties.Set(RenderOptionKeys::SampleRate, FVariantValue((int64)SampleRate));
+					OutputBufferSampleProperties.Set(RenderOptionKeys::Duration, FVariantValue(Duration));
+					OutputBufferSampleProperties.Set(RenderOptionKeys::PTS, FVariantValue(CurrentAccessUnit->AdjustedPTS));
 
 					Renderer->ReturnBuffer(CurrentOutputBuffer, true, OutputBufferSampleProperties);
 					CurrentOutputBuffer = nullptr;
@@ -1005,14 +1009,7 @@ bool FAudioDecoderImpl::SendSilenceOrEOS(ESendMode InSendMode)
 			}
 			else
 			{
-				OutputBufferSampleProperties.Clear();
-				OutputBufferSampleProperties.Set(TEXT("num_channels"), FVariantValue((int64)0));
-				OutputBufferSampleProperties.Set(TEXT("byte_size"), FVariantValue((int64)0));
-				OutputBufferSampleProperties.Set(TEXT("sample_rate"), FVariantValue((int64) 0));
-				OutputBufferSampleProperties.Set(TEXT("duration"), FVariantValue(FTimeValue::GetZero()));
-				OutputBufferSampleProperties.Set(TEXT("pts"), FVariantValue(FTimeValue::GetInvalid()));
-				OutputBufferSampleProperties.Set(TEXT("eos"), FVariantValue(true));
-				Renderer->ReturnBuffer(CurrentOutputBuffer, false, OutputBufferSampleProperties);
+				Renderer->ReturnBuffer(CurrentOutputBuffer, false, EOSBufferProperties);
 				CurrentOutputBuffer = nullptr;
 			}
 			return true;
diff --git a/Engine/Plugins/Media/ElectraPlayer/Source/ElectraPlayerRuntime/Private/Runtime/Decoder/Linux/LinuxElectraDecoderResourceManager.cpp b/Engine/Plugins/Media/ElectraPlayer/Source/ElectraPlayerRuntime/Private/Runtime/Decoder/Linux/LinuxElectraDecoderResourceManager.cpp
--- a/Engine/Plugins/Media/ElectraPlayer/Source/ElectraPlayerRuntime/Private/Runtime/Decoder/Linux/LinuxElectraDecoderResourceManager.cpp
+++ b/Engine/Plugins/Media/ElectraPlayer/Source/ElectraPlayerRuntime/Private/Runtime/Decoder/Linux/LinuxElectraDecoderResourceManager.cpp
@@ -254,7 +254,7 @@ bool FElectraDecoderResourceManagerLinux::SetupRenderBufferFromDecoderOutput(IMe
 	TMap<FString, FVariant> ExtraValues;
 	InDecoderOutput->GetExtraValues(ExtraValues);
 
-	TSharedPtr<FElectraPlayerVideoDecoderOutputLinux, ESPMode::ThreadSafe> DecoderOutput = InOutBufferToSetup->GetBufferProperties().GetValue("texture").GetSharedPointer<FElectraPlayerVideoDecoderOutputLinux>();
+	TSharedPtr<FElectraPlayerVideoDecoderOutputLinux, ESPMode::ThreadSafe> DecoderOutput = InOutBufferToSetup->GetBufferProperties().GetValue(RenderOptionKeys::Texture).GetSharedPointer<FElectraPlayerVideoDecoderOutputLinux>();
 	if (DecoderOutput.IsValid())
 	{
 		ILibavDecoderDecodedImage* DecodedImage = reinterpret_cast<ILibavDecoderDecodedImage*>(InDecoderOutput->GetPlatformOutputHandle(EElectraDecoderPlatformOutputHandleType::LibavDecoderDecodedImage));
@@ -263,29 +263,28 @@ bool FElectraDecoderResourceManagerLinux::SetupRenderBufferFromDecoderOutput(IMe
 		if (DecodedImage || ImageBuffers)
 		{
 			FElectraVideoDecoderOutputCropValues Crop = InDecoderOutput->GetCropValues();
-			InOutBufferPropertes->Set(TEXT("width"), FVariantValue((int64) InDecoderOutput->GetWidth()));
-			InOutBufferPropertes->Set(TEXT("height"), FVariantValue((int64) InDecoderOutput->GetHeight()));
-			InOutBufferPropertes->Set(TEXT("crop_left"), FVariantValue((int64) Crop.Left));
-			InOutBufferPropertes->Set(TEXT("crop_right"), FVariantValue((int64) Crop.Right));
-			InOutBufferPropertes->Set(TEXT("crop_top"), FVariantValue((int64) Crop.Top));
-			InOutBufferPropertes->Set(TEXT("crop_bottom"), FVariantValue((int64) Crop.Bottom));
-			InOutBufferPropertes->Set(TEXT("aspect_ratio"), FVariantValue((double) InDecoderOutput->GetAspectRatioW() / (double) InDecoderOutput->GetAspectRatioH()));
-			InOutBufferPropertes->Set(TEXT("aspect_w"), FVariantValue((int64) InDecoderOutput->GetAspectRatioW()));
-			InOutBufferPropertes->Set(TEXT("aspect_h"), FVariantValue((int64) InDecoderOutput->GetAspectRatioH()));
-			InOutBufferPropertes->Set(TEXT("fps_num"), FVariantValue((int64) InDecoderOutput->GetFrameRateNumerator()));
-			InOutBufferPropertes->Set(TEXT("fps_denom"), FVariantValue((int64) InDecoderOutput->GetFrameRateDenominator()));
+			InOutBufferPropertes->Set(IDecoderOutputOptionNames::Width, FVariantValue((int64)InDecoderOutput->GetWidth()));
+			InOutBufferPropertes->Set(IDecoderOutputOptionNames::Height, FVariantValue((int64)InDecoderOutput->GetHeight()));
+			InOutBufferPropertes->Set(IDecoderOutputOptionNames::CropLeft, FVariantValue((int64)Crop.Left));
+			InOutBufferPropertes->Set(IDecoderOutputOptionNames::CropRight, FVariantValue((int64)Crop.Right));
+			InOutBufferPropertes->Set(IDecoderOutputOptionNames::CropTop, FVariantValue((int64)Crop.Top));
+			InOutBufferPropertes->Set(IDecoderOutputOptionNames::CropBottom, FVariantValue((int64)Crop.Bottom));
+			InOutBufferPropertes->Set(IDecoderOutputOptionNames::AspectRatio, FVariantValue((double)InDecoderOutput->GetAspectRatioW() / (double)InDecoderOutput->GetAspectRatioH()));
+			InOutBufferPropertes->Set(IDecoderOutputOptionNames::AspectW, FVariantValue((int64)InDecoderOutput->GetAspectRatioW()));
+			InOutBufferPropertes->Set(IDecoderOutputOptionNames::AspectH, FVariantValue((int64)InDecoderOutput->GetAspectRatioH()));
+			InOutBufferPropertes->Set(IDecoderOutputOptionNames::FPSNumerator, FVariantValue((int64)InDecoderOutput->GetFrameRateNumerator()));
+			InOutBufferPropertes->Set(IDecoderOutputOptionNames::FPSDenominator, FVariantValue((int64)InDecoderOutput->GetFrameRateDenominator()));
 
 			int32 num_bits = InDecoderOutput->GetNumberOfBits();
 		// We are currently converting output to 8 bit NV12 format!
 			num_bits = 8;
-
-			InOutBufferPropertes->Set(TEXT("bits_per"), FVariantValue((int64) num_bits));
-			InOutBufferPropertes->Set(TEXT("pixelfmt"), FVariantValue((int64)((num_bits <= 8) ? EPixelFormat::PF_NV12 : EPixelFormat::PF_P010)));
+			InOutBufferPropertes->Set(IDecoderOutputOptionNames::BitsPerComponent, FVariantValue((int64) num_bits));
+			InOutBufferPropertes->Set(IDecoderOutputOptionNames::PixelFormat, FVariantValue((int64)((num_bits <= 8) ? EPixelFormat::PF_NV12 : EPixelFormat::PF_P010)));
 
 			if (DecodedImage)
 			{
 #if WITH_LIBAV
-				InOutBufferPropertes->Set(TEXT("pitch"), FVariantValue((int64)InDecoderOutput->GetDecodedWidth() * ((num_bits <= 8) ? 1 : 2)));
+				InOutBufferPropertes->Set(IDecoderOutputOptionNames::Pitch, FVariantValue((int64)InDecoderOutput->GetDecodedWidth() * ((num_bits <= 8) ? 1 : 2)));
 
 				const ILibavDecoderVideoCommon::FOutputInfo& DecodedImageInfo = DecodedImage->GetOutputInfo();
 				FIntPoint BufferDim(DecodedImageInfo.Planes[0].Width, DecodedImageInfo.Planes[0].Height);
@@ -304,7 +303,7 @@ bool FElectraDecoderResourceManagerLinux::SetupRenderBufferFromDecoderOutput(IMe
 				if ((ImageBuffers->GetBufferFormatByIndex(0) == EElectraDecoderPlatformPixelFormat::NV12 || ImageBuffers->GetBufferFormatByIndex(0) == EElectraDecoderPlatformPixelFormat::P010)&&
 					ImageBuffers->GetBufferEncodingByIndex(0) == EElectraDecoderPlatformPixelEncoding::Native)
 				{
-					InOutBufferPropertes->Set(TEXT("pitch"), FVariantValue((int64)ImageBuffers->GetBufferPitchByIndex(0)));
+					InOutBufferPropertes->Set(IDecoderOutputOptionNames::Pitch, FVariantValue((int64)ImageBuffers->GetBufferPitchByIndex(0)));
 
 					FIntPoint BufferDim((int32) InDecoderOutput->GetWidth(), (int32) InDecoderOutput->GetHeight());
 					if (DecoderOutput->InitializeForBuffer(BufferDim, num_bits <= 8 ? EPixelFormat::PF_NV12 : EPixelFormat::PF_P010, num_bits, InOutBufferPropertes, false))
diff --git a/Engine/Plugins/Media/ElectraPlayer/Source/ElectraPlayerRuntime/Private/Runtime/Decoder/SubtitleDecoder.cpp b/Engine/Plugins/Media/ElectraPlayer/Source/ElectraPlayerRuntime/Private/Runtime/Decoder/SubtitleDecoder.cpp
--- a/Engine/Plugins/Media/ElectraPlayer/Source/ElectraPlayerRuntime/Private/Runtime/Decoder/SubtitleDecoder.cpp
+++ b/Engine/Plugins/Media/ElectraPlayer/Source/ElectraPlayerRuntime/Private/Runtime/Decoder/SubtitleDecoder.cpp
@@ -8,6 +8,19 @@
 namespace Electra
 {
 
+namespace SubtitleOptionKeys
+{
+	static const FName SendEmptySubtitleDuringGaps(TEXT("sendEmptySubtitleDuringGaps"));
+	static const FName SideloadedID(TEXT("SideloadedID"));
+	static const FName PTO(TEXT("PTO"));
+	static const FName Width(TEXT("width"));
+	static const FName Height(TEXT("height"));
+	static const FName OffsetX(TEXT("offset_x"));
+	static const FName OffsetY(TEXT("offset_y"));
+	static const FName Timescale(TEXT("timescale"));
+}
+
+
 class FSubtitleDecoderPlugins
 {
 public:
@@ -227,11 +240,11 @@ void FSubtitleDecoder::Open(const FParamDict& Options)
 	{
 		PluginDecoderOptions = Options;
 		FParamDict Addtl(Options);
-		Addtl.Set(TEXT("width"), FVariantValue((int64) CodecInfo.GetResolution().Width));
-		Addtl.Set(TEXT("height"), FVariantValue((int64) CodecInfo.GetResolution().Height));
-		Addtl.Set(TEXT("offset_x"), FVariantValue((int64) CodecInfo.GetTranslation().GetX()));
-		Addtl.Set(TEXT("offset_y"), FVariantValue((int64) CodecInfo.GetTranslation().GetY()));
-		Addtl.Set(TEXT("timescale"), FVariantValue((int64) CodecInfo.GetFrameRate().GetDenominator()));
+		Addtl.Set(SubtitleOptionKeys::Width, FVariantValue((int64) CodecInfo.GetResolution().Width));
+		Addtl.Set(SubtitleOptionKeys::Height, FVariantValue((int64) CodecInfo.GetResolution().Height));
+		Addtl.Set(SubtitleOptionKeys::OffsetX, FVariantValue((int64) CodecInfo.GetTranslation().GetX()));
+		Addtl.Set(SubtitleOptionKeys::OffsetY, FVariantValue((int64) CodecInfo.GetTranslation().GetY()));
+		Addtl.Set(SubtitleOptionKeys::Timescale, FVariantValue((int64) CodecInfo.GetFrameRate().GetDenominator()));
 		if (PluginDecoder->InitializeStreamWithCSD(RawCSD, Addtl))
 		{
 			DeliveryOffset = PluginDecoder->GetStreamedDeliveryTimeOffset();
@@ -352,9 +365,9 @@ void FSubtitleDecoder::AUdataPushAU(FAccessUnit* AccessUnit)
 			// whether or not it has already parsed this data before (when seeking for instance).
 			if (AccessUnit->bIsSideloaded && AccessUnit->BufferSourceInfo.IsValid())
 			{
-				Addtl.Set(FString(TEXT("SideloadedID")), FVariantValue(AccessUnit->BufferSourceInfo->PeriodAdaptationSetID));
+				Addtl.Set(SubtitleOptionKeys::SideloadedID, FVariantValue(AccessUnit->BufferSourceInfo->PeriodAdaptationSetID));
 			}
-			Addtl.Set(FString(TEXT("PTO")), FVariantValue(AccessUnit->PTO));
+			Addtl.Set(SubtitleOptionKeys::PTO, FVariantValue(AccessUnit->PTO));
 			PluginDecoder->AddStreamedSubtitleData(TArray<uint8>((const uint8*)AccessUnit->AUData, (int32)AccessUnit->AUSize), AccessUnit->PTS, AccessUnit->Duration, Addtl);
 		}
 	}
diff --git a/Engine/Plugins/Media/ElectraPlayer/Source/ElectraPlayerRuntime/Private/Runtime/Decoder/VideoDecoder.cpp b/Engine/Plugins/Media/ElectraPlayer/Source/ElectraPlayerRuntime/Private/Runtime/Decoder/VideoDecoder.cpp
--- a/Engine/Plugins/Media/ElectraPlayer/Source/ElectraPlayerRuntime/Private/Runtime/Decoder/VideoDecoder.cpp
+++ b/Engine/Plugins/Media/ElectraPlayer/Source/ElectraPlayerRuntime/Private/Runtime/Decoder/VideoDecoder.cpp
@@ -212,8 +212,8 @@ private:
 	TSharedPtr<IElectraDecoderVideoOutput, ESPMode::ThreadSafe>				CurrentDecoderOutput;
 
 	IMediaRenderer::IBuffer*												CurrentOutputBuffer = nullptr;
-	FParamDict																BufferAcquireOptions;
-	FParamDict																OutputBufferSampleProperties;
+	FParamDict																EmptyOptions;
+	FParamDict																DummyBufferSampleProperties;
 };
 ENUM_CLASS_FLAGS(FVideoDecoderImpl::EAUChangeFlags);
 
@@ -390,10 +390,10 @@ void FVideoDecoderImpl::CreateDecoderOutputPool()
 
 // TODO/FIXME: get the default value of 8 from some config option?
 	int64 NumOutputFrames = ElectraDecodersUtil::GetVariantValueSafeI64(DecoderConfigOptions, IElectraDecoderFeature::MinimumNumberOfOutputFrames, 8);
-	poolOpts.Set(TEXT("num_buffers"), FVariantValue(NumOutputFrames));
+	poolOpts.Set(RenderOptionKeys::NumBuffers, FVariantValue(NumOutputFrames));
 	if (Renderer->CreateBufferPool(poolOpts) == UEMEDIA_ERROR_OK)
 	{
-		MaxDecodeBufferSize = (int32) Renderer->GetBufferPoolProperties().GetValue(TEXT("max_buffers")).GetInt64();
+		MaxDecodeBufferSize = (int32) Renderer->GetBufferPoolProperties().GetValue(RenderOptionKeys::MaxBuffers).GetInt64();
 	}
 	else
 	{
@@ -569,8 +569,7 @@ void FVideoDecoderImpl::ReturnUnusedOutputBuffer()
 {
 	if (CurrentOutputBuffer)
 	{
-		OutputBufferSampleProperties.Clear();
-		Renderer->ReturnBuffer(CurrentOutputBuffer, false, OutputBufferSampleProperties);
+		Renderer->ReturnBuffer(CurrentOutputBuffer, false, EmptyOptions);
 		CurrentOutputBuffer = nullptr;
 	}
 }
@@ -777,7 +776,7 @@ IElectraDecoder::EOutputStatus FVideoDecoderImpl::HandleOutput()
 				{
 					SCOPE_CYCLE_COUNTER(STAT_ElectraPlayer_VideoConvertOutput);
 					CSV_SCOPED_TIMING_STAT(ElectraPlayer, VideoConvertOutput);
-					UEMediaError bufResult = Renderer->AcquireBuffer(CurrentOutputBuffer, 0, BufferAcquireOptions);
+					UEMediaError bufResult = Renderer->AcquireBuffer(CurrentOutputBuffer, 0, EmptyOptions);
 					check(bufResult == UEMEDIA_ERROR_OK || bufResult == UEMEDIA_ERROR_INSUFFICIENT_DATA);
 					if (bufResult != UEMEDIA_ERROR_OK && bufResult != UEMEDIA_ERROR_INSUFFICIENT_DATA)
 					{
@@ -852,8 +851,8 @@ IElectraDecoder::EOutputStatus FVideoDecoderImpl::HandleOutput()
 
 				// Create the platform specific decoder output.
 				TSharedPtr<FParamDict, ESPMode::ThreadSafe> BufferProperties(new FParamDict);
-				BufferProperties->Set(TEXT("pts"), FVariantValue(MatchingInput->AdjustedPTS));
-				BufferProperties->Set(TEXT("duration"), FVariantValue(MatchingInput->AdjustedDuration));
+				BufferProperties->Set(RenderOptionKeys::PTS, FVariantValue(MatchingInput->AdjustedPTS));
+				BufferProperties->Set(RenderOptionKeys::Duration, FVariantValue(MatchingInput->AdjustedDuration));
 
 				// Set properties from the bitstream messages.
 				BitstreamProcessor->SetPropertiesOnOutput(CurrentDecoderOutput, BufferProperties.Get(), MatchingInput->BitstreamInfo);
@@ -1104,7 +1103,7 @@ bool FVideoDecoderImpl::HandleDummyDecoding()
 		{
 			SCOPE_CYCLE_COUNTER(STAT_ElectraPlayer_VideoConvertOutput);
 			CSV_SCOPED_TIMING_STAT(ElectraPlayer, VideoConvertOutput);
-			UEMediaError bufResult = Renderer->AcquireBuffer(CurrentOutputBuffer, 0, BufferAcquireOptions);
+			UEMediaError bufResult = Renderer->AcquireBuffer(CurrentOutputBuffer, 0, EmptyOptions);
 			check(bufResult == UEMEDIA_ERROR_OK || bufResult == UEMEDIA_ERROR_INSUFFICIENT_DATA);
 			if (bufResult != UEMEDIA_ERROR_OK && bufResult != UEMEDIA_ERROR_INSUFFICIENT_DATA)
 			{
@@ -1121,11 +1120,10 @@ bool FVideoDecoderImpl::HandleDummyDecoding()
 
 		NotifyReadyBufferListener(true);
 
-		OutputBufferSampleProperties.Clear();
-		OutputBufferSampleProperties.Set(TEXT("duration"), FVariantValue(CurrentAccessUnit->AdjustedDuration));
-		OutputBufferSampleProperties.Set(TEXT("pts"), FVariantValue(CurrentAccessUnit->AdjustedPTS));
-		OutputBufferSampleProperties.Set("is_dummy", FVariantValue(true));
-		Renderer->ReturnBuffer(CurrentOutputBuffer, true, OutputBufferSampleProperties);
+		DummyBufferSampleProperties.Set(RenderOptionKeys::Duration, FVariantValue(CurrentAccessUnit->AdjustedDuration));
+		DummyBufferSampleProperties.Set(RenderOptionKeys::PTS, FVariantValue(CurrentAccessUnit->AdjustedPTS));
+		DummyBufferSampleProperties.Set(RenderOptionKeys::DummyBufferFlag, FVariantValue(true));
+		Renderer->ReturnBuffer(CurrentOutputBuffer, true, DummyBufferSampleProperties);
 		CurrentOutputBuffer = nullptr;
 		return true;
 	}
diff --git a/Engine/Plugins/Media/ElectraPlayer/Source/ElectraPlayerRuntime/Private/Runtime/Decoder/VideoDecoderHelpers.cpp b/Engine/Plugins/Media/ElectraPlayer/Source/ElectraPlayerRuntime/Private/Runtime/Decoder/VideoDecoderHelpers.cpp
--- a/Engine/Plugins/Media/ElectraPlayer/Source/ElectraPlayerRuntime/Private/Runtime/Decoder/VideoDecoderHelpers.cpp
+++ b/Engine/Plugins/Media/ElectraPlayer/Source/ElectraPlayerRuntime/Private/Runtime/Decoder/VideoDecoderHelpers.cpp
@@ -59,7 +59,7 @@ void FColorimetryHelper::Update(const TArray<uint8>& InFromCOLRBox)
 
 void FColorimetryHelper::UpdateParamDict(FParamDict& InOutDictionary)
 {
-	InOutDictionary.Set(TEXT("colorimetry"), FVariantValue(CurrentColorimetry));
+	InOutDictionary.Set(IDecoderOutputOptionNames::Colorimetry, FVariantValue(CurrentColorimetry));
 }
 
 bool FColorimetryHelper::GetCurrentValues(uint8& colour_primaries, uint8& transfer_characteristics, uint8& matrix_coeffs) const
@@ -306,7 +306,7 @@ void FHDRHelper::Update(int32 BitDepth, const FColorimetryHelper& InColorimetry,
 
 void FHDRHelper::UpdateParamDict(FParamDict& InOutDictionary)
 {
-	InOutDictionary.Set(TEXT("hdr_info"), FVariantValue(CurrentHDRInfo));
+	InOutDictionary.Set(IDecoderOutputOptionNames::HDRInfo, FVariantValue(CurrentHDRInfo));
 }
 
 
diff --git a/Engine/Plugins/Media/ElectraPlayer/Source/ElectraPlayerRuntime/Private/Runtime/Decoder/Windows/WindowsElectraDecoderResourceManager.cpp b/Engine/Plugins/Media/ElectraPlayer/Source/ElectraPlayerRuntime/Private/Runtime/Decoder/Windows/WindowsElectraDecoderResourceManager.cpp
--- a/Engine/Plugins/Media/ElectraPlayer/Source/ElectraPlayerRuntime/Private/Runtime/Decoder/Windows/WindowsElectraDecoderResourceManager.cpp
+++ b/Engine/Plugins/Media/ElectraPlayer/Source/ElectraPlayerRuntime/Private/Runtime/Decoder/Windows/WindowsElectraDecoderResourceManager.cpp
@@ -116,21 +116,21 @@ bool FElectraDecoderResourceManagerWindows::GetD3DDevice(void **OutD3DDevice, in
 
 bool FElectraDecoderResourceManagerWindows::SetupRenderBufferFromDecoderOutputFromMFSample(IMediaRenderer::IBuffer* InOutBufferToSetup, TSharedPtr<FParamDict, ESPMode::ThreadSafe> InOutBufferPropertes, TSharedPtr<IElectraDecoderVideoOutput, ESPMode::ThreadSafe> InDecoderOutput, FElectraDecoderResourceManagerWindows::IDecoderPlatformResource* InPlatformSpecificResource)
 {
-	TSharedPtr<FElectraPlayerVideoDecoderOutputPC, ESPMode::ThreadSafe> DecoderOutput = InOutBufferToSetup->GetBufferProperties().GetValue("texture").GetSharedPointer<FElectraPlayerVideoDecoderOutputPC>();
+	TSharedPtr<FElectraPlayerVideoDecoderOutputPC, ESPMode::ThreadSafe> DecoderOutput = InOutBufferToSetup->GetBufferProperties().GetValue(RenderOptionKeys::Texture).GetSharedPointer<FElectraPlayerVideoDecoderOutputPC>();
 	if (DecoderOutput.IsValid())
 	{
 		FElectraVideoDecoderOutputCropValues Crop = InDecoderOutput->GetCropValues();
-		InOutBufferPropertes->Set(TEXT("width"), FVariantValue((int64)InDecoderOutput->GetWidth()));
-		InOutBufferPropertes->Set(TEXT("height"), FVariantValue((int64)InDecoderOutput->GetHeight()));
-		InOutBufferPropertes->Set(TEXT("crop_left"), FVariantValue((int64)Crop.Left));
-		InOutBufferPropertes->Set(TEXT("crop_right"), FVariantValue((int64)Crop.Right));
-		InOutBufferPropertes->Set(TEXT("crop_top"), FVariantValue((int64)Crop.Top));
-		InOutBufferPropertes->Set(TEXT("crop_bottom"), FVariantValue((int64)Crop.Bottom));
-		InOutBufferPropertes->Set(TEXT("aspect_ratio"), FVariantValue((double)InDecoderOutput->GetAspectRatioW() / (double)InDecoderOutput->GetAspectRatioH()));
-		InOutBufferPropertes->Set(TEXT("aspect_w"), FVariantValue((int64)InDecoderOutput->GetAspectRatioW()));
-		InOutBufferPropertes->Set(TEXT("aspect_h"), FVariantValue((int64)InDecoderOutput->GetAspectRatioH()));
-		InOutBufferPropertes->Set(TEXT("fps_num"), FVariantValue((int64)InDecoderOutput->GetFrameRateNumerator()));
-		InOutBufferPropertes->Set(TEXT("fps_denom"), FVariantValue((int64)InDecoderOutput->GetFrameRateDenominator()));
+		InOutBufferPropertes->Set(IDecoderOutputOptionNames::Width, FVariantValue((int64)InDecoderOutput->GetWidth()));
+		InOutBufferPropertes->Set(IDecoderOutputOptionNames::Height, FVariantValue((int64)InDecoderOutput->GetHeight()));
+		InOutBufferPropertes->Set(IDecoderOutputOptionNames::CropLeft, FVariantValue((int64)Crop.Left));
+		InOutBufferPropertes->Set(IDecoderOutputOptionNames::CropRight, FVariantValue((int64)Crop.Right));
+		InOutBufferPropertes->Set(IDecoderOutputOptionNames::CropTop, FVariantValue((int64)Crop.Top));
+		InOutBufferPropertes->Set(IDecoderOutputOptionNames::CropBottom, FVariantValue((int64)Crop.Bottom));
+		InOutBufferPropertes->Set(IDecoderOutputOptionNames::AspectRatio, FVariantValue((double)InDecoderOutput->GetAspectRatioW() / (double)InDecoderOutput->GetAspectRatioH()));
+		InOutBufferPropertes->Set(IDecoderOutputOptionNames::AspectW, FVariantValue((int64)InDecoderOutput->GetAspectRatioW()));
+		InOutBufferPropertes->Set(IDecoderOutputOptionNames::AspectH, FVariantValue((int64)InDecoderOutput->GetAspectRatioH()));
+		InOutBufferPropertes->Set(IDecoderOutputOptionNames::FPSNumerator, FVariantValue((int64)InDecoderOutput->GetFrameRateNumerator()));
+		InOutBufferPropertes->Set(IDecoderOutputOptionNames::FPSDenominator, FVariantValue((int64)InDecoderOutput->GetFrameRateDenominator()));
 
 		// What type of decoder output do we have here?
 		TMap<FString, FVariant> ExtraValues;
@@ -188,8 +188,8 @@ bool FElectraDecoderResourceManagerWindows::SetupRenderBufferFromDecoderOutputFr
 				return false;
 			}
 
-			InOutBufferPropertes->Set(TEXT("pixelfmt"), FVariantValue((int64)(TextureDesc.Format == DXGI_FORMAT_NV12 ? EPixelFormat::PF_NV12 : EPixelFormat::PF_P010)));
-			InOutBufferPropertes->Set(TEXT("bits_per"), FVariantValue((int64)(TextureDesc.Format == DXGI_FORMAT_NV12 ? 8 : 10)));
+			InOutBufferPropertes->Set(IDecoderOutputOptionNames::PixelFormat, FVariantValue((int64)(TextureDesc.Format == DXGI_FORMAT_NV12 ? EPixelFormat::PF_NV12 : EPixelFormat::PF_P010)));
+			InOutBufferPropertes->Set(IDecoderOutputOptionNames::BitsPerComponent, FVariantValue((int64)(TextureDesc.Format == DXGI_FORMAT_NV12 ? 8 : 10)));
 
 			if (DXVersion == 0 || DXVersion >= 12000)
 			{
@@ -213,7 +213,7 @@ bool FElectraDecoderResourceManagerWindows::SetupRenderBufferFromDecoderOutputFr
 					return false;
 				}
 
-				InOutBufferPropertes->Set(TEXT("pitch"), FVariantValue((int64)Pitch));
+				InOutBufferPropertes->Set(IDecoderOutputOptionNames::Pitch, FVariantValue((int64)Pitch));
 
 				// Get decoded with (e.g. featuring any height adjustments for CPU buffer usage of NV12 etc.)
 				int32 Width = InDecoderOutput->GetDecodedWidth();
@@ -295,8 +295,8 @@ bool FElectraDecoderResourceManagerWindows::SetupRenderBufferFromDecoderOutputFr
 
 			EPixelFormat PixFmt = (SamplePixFmt == EElectraDecoderPlatformPixelFormat::NV12) ? EPixelFormat::PF_NV12 : EPixelFormat::PF_P010;
 
-			InOutBufferPropertes->Set(TEXT("pixelfmt"), FVariantValue((int64)PixFmt));
-			InOutBufferPropertes->Set(TEXT("bits_per"), FVariantValue((int64)((PixFmt == EPixelFormat::PF_NV12) ? 8 : 10)));
+			InOutBufferPropertes->Set(IDecoderOutputOptionNames::PixelFormat, FVariantValue((int64)PixFmt));
+			InOutBufferPropertes->Set(IDecoderOutputOptionNames::BitsPerComponent, FVariantValue((int64)((PixFmt == EPixelFormat::PF_NV12) ? 8 : 10)));
 
 			TRefCountPtr<IMF2DBuffer> Buffer2D;
 			if ((Result = Buffer->QueryInterface(__uuidof(IMF2DBuffer), (void**)Buffer2D.GetInitReference())) == S_OK)
@@ -309,7 +309,7 @@ bool FElectraDecoderResourceManagerWindows::SetupRenderBufferFromDecoderOutputFr
 					return false;
 				}
 
-				InOutBufferPropertes->Set(TEXT("pitch"), FVariantValue((int64)Pitch));
+				InOutBufferPropertes->Set(IDecoderOutputOptionNames::Pitch, FVariantValue((int64)Pitch));
 
 				DecoderOutput->InitializeWithBuffer(Data, Pitch * Height,
 					Pitch,						// Buffer stride
@@ -339,7 +339,7 @@ bool FElectraDecoderResourceManagerWindows::SetupRenderBufferFromDecoderOutputFr
 
 				int32 Pitch = Width * ((PixFmt == EPixelFormat::PF_NV12) ? 1 : 2);
 
-				InOutBufferPropertes->Set(TEXT("pitch"), FVariantValue((int64)Pitch));
+				InOutBufferPropertes->Set(IDecoderOutputOptionNames::Pitch, FVariantValue((int64)Pitch));
 
 				DecoderOutput->InitializeWithBuffer(Data, BufferSize,
 					Pitch,						// Buffer stride
@@ -365,9 +365,8 @@ bool FElectraDecoderResourceManagerWindows::SetupRenderBufferFromDecoderOutput(I
 	check(InOutBufferPropertes.IsValid());
 	check(InDecoderOutput.IsValid());
 
-	TSharedPtr<FElectraPlayerVideoDecoderOutputPC, ESPMode::ThreadSafe> DecoderOutput = InOutBufferToSetup->GetBufferProperties().GetValue("texture").GetSharedPointer<FElectraPlayerVideoDecoderOutputPC>();
+	TSharedPtr<FElectraPlayerVideoDecoderOutputPC, ESPMode::ThreadSafe> DecoderOutput = InOutBufferToSetup->GetBufferProperties().GetValue(RenderOptionKeys::Texture).GetSharedPointer<FElectraPlayerVideoDecoderOutputPC>();
 	FInstanceVars* Vars = static_cast<FInstanceVars*>(InPlatformSpecificResource);
-
 	if (DecoderOutput.IsValid())
 	{
 		//
@@ -377,22 +376,22 @@ bool FElectraDecoderResourceManagerWindows::SetupRenderBufferFromDecoderOutput(I
 		if (ImageBuffers != nullptr)
 		{
 			FElectraVideoDecoderOutputCropValues Crop = InDecoderOutput->GetCropValues();
-			InOutBufferPropertes->Set(TEXT("width"), FVariantValue((int64)InDecoderOutput->GetWidth()));
-			InOutBufferPropertes->Set(TEXT("height"), FVariantValue((int64)InDecoderOutput->GetHeight()));
-			InOutBufferPropertes->Set(TEXT("crop_left"), FVariantValue((int64)Crop.Left));
-			InOutBufferPropertes->Set(TEXT("crop_right"), FVariantValue((int64)Crop.Right));
-			InOutBufferPropertes->Set(TEXT("crop_top"), FVariantValue((int64)Crop.Top));
-			InOutBufferPropertes->Set(TEXT("crop_bottom"), FVariantValue((int64)Crop.Bottom));
-			InOutBufferPropertes->Set(TEXT("aspect_ratio"), FVariantValue((double)InDecoderOutput->GetAspectRatioW() / (double)InDecoderOutput->GetAspectRatioH()));
-			InOutBufferPropertes->Set(TEXT("aspect_w"), FVariantValue((int64)InDecoderOutput->GetAspectRatioW()));
-			InOutBufferPropertes->Set(TEXT("aspect_h"), FVariantValue((int64)InDecoderOutput->GetAspectRatioH()));
-			InOutBufferPropertes->Set(TEXT("fps_num"), FVariantValue((int64)InDecoderOutput->GetFrameRateNumerator()));
-			InOutBufferPropertes->Set(TEXT("fps_denom"), FVariantValue((int64)InDecoderOutput->GetFrameRateDenominator()));
+			InOutBufferPropertes->Set(IDecoderOutputOptionNames::Width, FVariantValue((int64)InDecoderOutput->GetWidth()));
+			InOutBufferPropertes->Set(IDecoderOutputOptionNames::Height, FVariantValue((int64)InDecoderOutput->GetHeight()));
+			InOutBufferPropertes->Set(IDecoderOutputOptionNames::CropLeft, FVariantValue((int64)Crop.Left));
+			InOutBufferPropertes->Set(IDecoderOutputOptionNames::CropRight, FVariantValue((int64)Crop.Right));
+			InOutBufferPropertes->Set(IDecoderOutputOptionNames::CropTop, FVariantValue((int64)Crop.Top));
+			InOutBufferPropertes->Set(IDecoderOutputOptionNames::CropBottom, FVariantValue((int64)Crop.Bottom));
+			InOutBufferPropertes->Set(IDecoderOutputOptionNames::AspectRatio, FVariantValue((double)InDecoderOutput->GetAspectRatioW() / (double)InDecoderOutput->GetAspectRatioH()));
+			InOutBufferPropertes->Set(IDecoderOutputOptionNames::AspectW, FVariantValue((int64)InDecoderOutput->GetAspectRatioW()));
+			InOutBufferPropertes->Set(IDecoderOutputOptionNames::AspectH, FVariantValue((int64)InDecoderOutput->GetAspectRatioH()));
+			InOutBufferPropertes->Set(IDecoderOutputOptionNames::FPSNumerator, FVariantValue((int64)InDecoderOutput->GetFrameRateNumerator()));
+			InOutBufferPropertes->Set(IDecoderOutputOptionNames::FPSDenominator, FVariantValue((int64)InDecoderOutput->GetFrameRateDenominator()));
 
 			int32 Width = InDecoderOutput->GetDecodedWidth();
 			int32 Height = InDecoderOutput->GetDecodedHeight();
 
-			InOutBufferPropertes->Set(TEXT("bits_per"), FVariantValue((int64)InDecoderOutput->GetNumberOfBits()));
+			InOutBufferPropertes->Set(IDecoderOutputOptionNames::BitsPerComponent, FVariantValue((int64)InDecoderOutput->GetNumberOfBits()));
 
 			int32 NumImageBuffers = ImageBuffers->GetNumberOfBuffers();
 			check(NumImageBuffers == 1 || NumImageBuffers == 2);
@@ -402,7 +401,7 @@ bool FElectraDecoderResourceManagerWindows::SetupRenderBufferFromDecoderOutput(I
 			EElectraDecoderPlatformPixelEncoding PixEnc = ImageBuffers->GetBufferEncodingByIndex(0);
 
 			EPixelFormat RHIPixFmt;
-			switch (PixFmt)
+			switch(PixFmt)
 			{
 				case EElectraDecoderPlatformPixelFormat::R8G8B8A8:		RHIPixFmt = EPixelFormat::PF_R8G8B8A8; break;
 				case EElectraDecoderPlatformPixelFormat::A8R8G8B8:		RHIPixFmt = EPixelFormat::PF_A8R8G8B8; break;
@@ -438,9 +437,9 @@ bool FElectraDecoderResourceManagerWindows::SetupRenderBufferFromDecoderOutput(I
 			
 			int32 Pitch = ImageBuffers->GetBufferPitchByIndex(0);
 
-			InOutBufferPropertes->Set(TEXT("pixelfmt"), FVariantValue((int64)RHIPixFmt));
-			InOutBufferPropertes->Set(TEXT("pixelenc"), FVariantValue((int64)DecPixEnc));
-			InOutBufferPropertes->Set(TEXT("pitch"), FVariantValue((int64)Pitch));
+			InOutBufferPropertes->Set(IDecoderOutputOptionNames::PixelFormat, FVariantValue((int64)RHIPixFmt));
+			InOutBufferPropertes->Set(IDecoderOutputOptionNames::PixelEncoding, FVariantValue((int64)DecPixEnc));
+			InOutBufferPropertes->Set(IDecoderOutputOptionNames::Pitch, FVariantValue((int64)Pitch));
 
 // [...] ALPHA BUFFER -- HOW DO WE PASS IT ON!?
 // [...] ANY CS/HDR INFO FROM THE DECODER? -- PASSING IT ON WOULD BE EASY, BUT RIGHT NOW WE ALWAYS READ IT FROM THE CONTAINER (in code further up the chain)
diff --git a/Engine/Plugins/Media/ElectraPlayer/Source/ElectraPlayerRuntime/Private/Runtime/Demuxer/ParserISO14496-12.cpp b/Engine/Plugins/Media/ElectraPlayer/Source/ElectraPlayerRuntime/Private/Runtime/Demuxer/ParserISO14496-12.cpp
--- a/Engine/Plugins/Media/ElectraPlayer/Source/ElectraPlayerRuntime/Private/Runtime/Demuxer/ParserISO14496-12.cpp
+++ b/Engine/Plugins/Media/ElectraPlayer/Source/ElectraPlayerRuntime/Private/Runtime/Demuxer/ParserISO14496-12.cpp
@@ -5967,7 +5967,7 @@ private:
 		};
 
 
-		void AddExtraBoxToCodecInfo(FStreamCodecInformation& InOutCodecInformation, const FMP4Box* InBox, IParserISO14496_12::FBoxType InBoxType, const FString& InAs)
+		void AddExtraBoxToCodecInfo(FStreamCodecInformation& InOutCodecInformation, const FMP4Box* InBox, IParserISO14496_12::FBoxType InBoxType, const FName& InAs)
 		{
 			const FMP4BoxIgnored* Box = static_cast<const FMP4BoxIgnored*>(InBox->FindBox(InBoxType, 3));
 			if (Box && Box->GetFillerData())
@@ -5976,7 +5976,7 @@ private:
 				InOutCodecInformation.GetExtras().Set(InAs, FVariantValue(Data));
 			}
 		}
-		void AddExtraBoxToCodecInfo(FStreamCodecInformation& InOutCodecInformation, const FMP4Box* InBox, const FString& InAs)
+		void AddExtraBoxToCodecInfo(FStreamCodecInformation& InOutCodecInformation, const FMP4Box* InBox, const FName& InAs)
 		{
 			const FMP4BoxIgnored* Box = static_cast<const FMP4BoxIgnored*>(InBox);
 			if (Box && Box->GetFillerData())
@@ -6115,7 +6115,7 @@ private:
 			{
 				if (CodecInformation.GetCodec4CC() == Utils::Make4CC('O','p','u','s'))
 				{
-					FVariantValue dOps = CodecInformation.GetExtras().GetValue(TEXT("dOps_box"));
+					FVariantValue dOps = CodecInformation.GetExtras().GetValue(StreamCodecInformationOptions::DOpsBox);
 					if (dOps.IsType(FVariantValue::EDataType::TypeU8Array))
 					{
 						return dOps.GetArray();
@@ -7312,7 +7312,7 @@ private:
 						{
 							Track->CodecInformation.SetNumberOfChannels(2);
 						}
-						Track->CodecInformation.GetExtras().Set(TEXT("samples_per_block"), FVariantValue(ConfigRecord.SBRSignal > 0 ? (int64)2048 : (int64)1024));
+						Track->CodecInformation.GetExtras().Set(StreamCodecInformationOptions::SamplesPerBlock, FVariantValue(ConfigRecord.SBRSignal > 0 ? (int64)2048 : (int64)1024));
 					}
 
 					// Typically an mp4a track will not have a 'btrt' box because the bitrate is stored in the DecoderConfigDescriptor.
@@ -7356,9 +7356,9 @@ private:
 			{
 				if (AudioSampleEntry->HaveFormatSpecificFlags())
 				{
-					Track->CodecInformation.GetExtras().Set(TEXT("FormatSpecificFlags"), FVariantValue((int64) AudioSampleEntry->GetFormatSpecificFlags()));
+					Track->CodecInformation.GetExtras().Set(StreamCodecInformationOptions::FormatSpecificFlags, FVariantValue((int64) AudioSampleEntry->GetFormatSpecificFlags()));
 				}
-				AddExtraBoxToCodecInfo(Track->CodecInformation, SampleBox->GetParentBox(), FMP4Box::kBox_chan, TEXT("chan_box"));
+				AddExtraBoxToCodecInfo(Track->CodecInformation, SampleBox->GetParentBox(), FMP4Box::kBox_chan, StreamCodecInformationOptions::ChanBox);
 			}
 
 			for(int32 j=0, jMax=AudioSampleEntry->GetNumberOfChildren(); j<jMax; ++j)
@@ -7537,7 +7537,7 @@ private:
 								// Note: adding the box via the filler data only works as long as these boxes are not explicitly handled in this module!
 								FString BoxName = Utils::Printable4CC(SampleEntry->GetType());
 								BoxName.Append(TEXT("_box"));
-								AddExtraBoxToCodecInfo(CodecInformation, SampleEntry, BoxName);
+								AddExtraBoxToCodecInfo(CodecInformation, SampleEntry, FName(*BoxName));
 								break;
 							}
 						}
@@ -7570,12 +7570,12 @@ private:
 					CodecInformation.SetNumberOfChannels(AudioSampleEntry->GetChannelCount());
 					CodecInformation.SetCodecSpecifierRFC6381(Utils::Printable4CC(CodecInformation.GetCodec4CC()));
 
-					CodecInformation.GetExtras().Set(TEXT("SampleSize"), FVariantValue((int64)AudioSampleEntry->GetSampleSize()));
+					CodecInformation.GetExtras().Set(StreamCodecInformationOptions::SampleSize, FVariantValue((int64)AudioSampleEntry->GetSampleSize()));
 					if (AudioSampleEntry->HaveFormatSpecificFlags())
 					{
-						CodecInformation.GetExtras().Set(TEXT("FormatSpecificFlags"), FVariantValue((int64) AudioSampleEntry->GetFormatSpecificFlags()));
-						CodecInformation.GetExtras().Set(TEXT("ConstBytesPerAudioPacket"), FVariantValue((int64)AudioSampleEntry->GetConstBytesPerAudioPacket()));
-						CodecInformation.GetExtras().Set(TEXT("ConstLPCMFramesPerAudioPacket"), FVariantValue((int64)AudioSampleEntry->GetConstLPCMFramesPerAudioPacket()));
+						CodecInformation.GetExtras().Set(StreamCodecInformationOptions::FormatSpecificFlags, FVariantValue((int64) AudioSampleEntry->GetFormatSpecificFlags()));
+						CodecInformation.GetExtras().Set(StreamCodecInformationOptions::ConstBytesPerAudioPacket, FVariantValue((int64)AudioSampleEntry->GetConstBytesPerAudioPacket()));
+						CodecInformation.GetExtras().Set(StreamCodecInformationOptions::ConstLPCMFramesPerAudioPacket, FVariantValue((int64)AudioSampleEntry->GetConstLPCMFramesPerAudioPacket()));
 					}
 
 					// Look at the additional boxes
@@ -7597,7 +7597,7 @@ private:
 								// Note: adding the box via the filler data only works as long as these boxes are not explicitly handled in this module!
 								FString BoxName = Utils::Printable4CC(SampleEntry->GetType());
 								BoxName.Append(TEXT("_box"));
-								AddExtraBoxToCodecInfo(CodecInformation, SampleEntry, BoxName);
+								AddExtraBoxToCodecInfo(CodecInformation, SampleEntry, FName(*BoxName));
 								break;
 							}
 						}
diff --git a/Engine/Plugins/Media/ElectraPlayer/Source/ElectraPlayerRuntime/Private/Runtime/Demuxer/ParserMKV.cpp b/Engine/Plugins/Media/ElectraPlayer/Source/ElectraPlayerRuntime/Private/Runtime/Demuxer/ParserMKV.cpp
--- a/Engine/Plugins/Media/ElectraPlayer/Source/ElectraPlayerRuntime/Private/Runtime/Demuxer/ParserMKV.cpp
+++ b/Engine/Plugins/Media/ElectraPlayer/Source/ElectraPlayerRuntime/Private/Runtime/Demuxer/ParserMKV.cpp
@@ -2733,7 +2733,7 @@ namespace Electra
 
 			TMKVElementPtr<FMKVVideo> Video = InFromTrack->GetVideo();
 
-			OutCodecInformation.GetExtras().Set(TEXT("dcr"), FVariantValue(InFromTrack->GetCodecPrivate()));
+			OutCodecInformation.GetExtras().Set(StreamCodecInformationOptions::DecoderConfigurationRecord, FVariantValue(InFromTrack->GetCodecPrivate()));
 			OutCodecInformation.SetStreamType(EStreamType::Video);
 			OutCodecInformation.SetMimeType(TEXT("video/mp4"));
 			OutCodecInformation.SetCodec(FStreamCodecInformation::ECodec::H264);
@@ -2812,7 +2812,7 @@ namespace Electra
 
 			TMKVElementPtr<FMKVVideo> Video = InFromTrack->GetVideo();
 
-			OutCodecInformation.GetExtras().Set(TEXT("dcr"), FVariantValue(InFromTrack->GetCodecPrivate()));
+			OutCodecInformation.GetExtras().Set(StreamCodecInformationOptions::DecoderConfigurationRecord, FVariantValue(InFromTrack->GetCodecPrivate()));
 			OutCodecInformation.SetStreamType(EStreamType::Video);
 			OutCodecInformation.SetMimeType(TEXT("video/mp4"));
 			OutCodecInformation.SetCodec(FStreamCodecInformation::ECodec::H265);
@@ -2938,7 +2938,7 @@ namespace Electra
 			vpcC[7] = (uint8)OutCodecInformation.GetCodecVideoColorInfo().ColourPrimaries.Get(2);
 			vpcC[8] = (uint8)OutCodecInformation.GetCodecVideoColorInfo().TransferCharacteristics.Get(2);
 			vpcC[9] = (uint8)OutCodecInformation.GetCodecVideoColorInfo().MatrixCoefficients.Get(2);
-			OutCodecInformation.GetExtras().Set(TEXT("vpcC_box"), FVariantValue(vpcC));
+			OutCodecInformation.GetExtras().Set(StreamCodecInformationOptions::VPccBox, FVariantValue(vpcC));
 
 			int32 CropL = Video->GetCropLeft();
 			int32 CropR = Video->GetCropRight();
@@ -3082,7 +3082,7 @@ namespace Electra
 			vpcC[7] = (uint8)OutCodecInformation.GetCodecVideoColorInfo().ColourPrimaries.Get(2);
 			vpcC[8] = (uint8)OutCodecInformation.GetCodecVideoColorInfo().TransferCharacteristics.Get(2);
 			vpcC[9] = (uint8)OutCodecInformation.GetCodecVideoColorInfo().MatrixCoefficients.Get(2);
-			OutCodecInformation.GetExtras().Set(TEXT("vpcC_box"), FVariantValue(vpcC));
+			OutCodecInformation.GetExtras().Set(StreamCodecInformationOptions::VPccBox, FVariantValue(vpcC));
 
 			int32 CropL = Video->GetCropLeft();
 			int32 CropR = Video->GetCropRight();
@@ -3120,7 +3120,7 @@ namespace Electra
 				return false;
 			}
 
-			OutCodecInformation.GetExtras().Set(TEXT("dcr"), FVariantValue(InFromTrack->GetCodecPrivate()));
+			OutCodecInformation.GetExtras().Set(StreamCodecInformationOptions::DecoderConfigurationRecord, FVariantValue(InFromTrack->GetCodecPrivate()));
 			OutCodecInformation.SetStreamType(EStreamType::Audio);
 			OutCodecInformation.SetMimeType(TEXT("audio/mp4"));
 			OutCodecInformation.SetCodec(FStreamCodecInformation::ECodec::AAC);
@@ -3137,7 +3137,7 @@ namespace Electra
 				OutCodecInformation.SetNumberOfChannels(2);
 			}
 			const int32 NumDecodedSamplesPerBlock = AudioSpecificConfig.SBRSignal > 0 ? 2048 : 1024;
-			OutCodecInformation.GetExtras().Set(TEXT("samples_per_block"), FVariantValue((int64)NumDecodedSamplesPerBlock));
+			OutCodecInformation.GetExtras().Set(StreamCodecInformationOptions::SamplesPerBlock, FVariantValue((int64)NumDecodedSamplesPerBlock));
 			// If there is no default duration set we try to calculate it from the sample rate.
 			if (InFromTrack->GetDefaultDurationNanos() == 0)
 			{
@@ -3192,7 +3192,7 @@ namespace Electra
 				return false;
 			}
 
-			OutCodecInformation.GetExtras().Set(TEXT("dOps_box"), FVariantValue(OpusHead));
+			OutCodecInformation.GetExtras().Set(StreamCodecInformationOptions::DOpsBox, FVariantValue(OpusHead));
 			OutCodecInformation.SetStreamType(EStreamType::Audio);
 			OutCodecInformation.SetMimeType(TEXT("audio/mp4"));
 			OutCodecInformation.SetCodec(FStreamCodecInformation::ECodec::Audio4CC);
diff --git a/Engine/Plugins/Media/ElectraPlayer/Source/ElectraPlayerRuntime/Private/Runtime/HTTP/HTTPResponseCache.cpp b/Engine/Plugins/Media/ElectraPlayer/Source/ElectraPlayerRuntime/Private/Runtime/HTTP/HTTPResponseCache.cpp
--- a/Engine/Plugins/Media/ElectraPlayer/Source/ElectraPlayerRuntime/Private/Runtime/HTTP/HTTPResponseCache.cpp
+++ b/Engine/Plugins/Media/ElectraPlayer/Source/ElectraPlayerRuntime/Private/Runtime/HTTP/HTTPResponseCache.cpp
@@ -7,8 +7,8 @@
 
 namespace Electra
 {
-	const TCHAR* const OptionKeyResponseCacheMaxEntries = TEXT("httpcache_max_entries");
-	const TCHAR* const OptionKeyResponseCacheMaxByteSize = TEXT("httpcache_max_bytesize");
+	static const FName OptionKeyResponseCacheMaxEntries(TEXT("httpcache_max_entries"));
+	static const FName OptionKeyResponseCacheMaxByteSize(TEXT("httpcache_max_bytesize"));
 
 
 class FHTTPResponseCache : public IHTTPResponseCache
diff --git a/Engine/Plugins/Media/ElectraPlayer/Source/ElectraPlayerRuntime/Private/Runtime/Player/ABRRules/ABROptionKeynames.h b/Engine/Plugins/Media/ElectraPlayer/Source/ElectraPlayerRuntime/Private/Runtime/Player/ABRRules/ABROptionKeynames.h
--- a/Engine/Plugins/Media/ElectraPlayer/Source/ElectraPlayerRuntime/Private/Runtime/Player/ABRRules/ABROptionKeynames.h
+++ b/Engine/Plugins/Media/ElectraPlayer/Source/ElectraPlayerRuntime/Private/Runtime/Player/ABRRules/ABROptionKeynames.h
@@ -9,10 +9,10 @@ namespace Electra
 namespace ABR
 {
 	//! (bool) true to just finish the currently loading segment when rebuffering. false to start over with.
-	const TCHAR* const OptionKeyABR_RebufferingContinuesLoading = TEXT("abr:rebuffering_continues_loading");
+	const FName OptionKeyABR_RebufferingContinuesLoading(TEXT("abr:rebuffering_continues_loading"));
 
 	//! (int) HTTP status code if returned by CDN for any segment request will take that stream permanently offline.
-	const TCHAR* const OptionKeyABR_CDNSegmentDenyHTTPStatus = TEXT("abr:cdn_deny_httpstatus");
+	const FName OptionKeyABR_CDNSegmentDenyHTTPStatus(TEXT("abr:cdn_deny_httpstatus"));
 }
 } // namespace Electra
 
diff --git a/Engine/Plugins/Media/ElectraPlayer/Source/ElectraPlayerRuntime/Private/Runtime/Player/AdaptivePlayerOptionKeynames.h b/Engine/Plugins/Media/ElectraPlayer/Source/ElectraPlayerRuntime/Private/Runtime/Player/AdaptivePlayerOptionKeynames.h
--- a/Engine/Plugins/Media/ElectraPlayer/Source/ElectraPlayerRuntime/Private/Runtime/Player/AdaptivePlayerOptionKeynames.h
+++ b/Engine/Plugins/Media/ElectraPlayer/Source/ElectraPlayerRuntime/Private/Runtime/Player/AdaptivePlayerOptionKeynames.h
@@ -7,43 +7,44 @@
 namespace Electra
 {
 	//! (FString) mime type of URL to load
-	const TCHAR* const OptionKeyMimeType = TEXT("mime_type");										
+	const FName OptionKeyMimeType(TEXT("mime_type"));										
 
 	//! (int64) value indicating the bitrate to start playback with (initial start).
-	const TCHAR* const OptionKeyInitialBitrate = TEXT("initial_bitrate");
+	const FName OptionKeyInitialBitrate(TEXT("initial_bitrate"));
 
 	//! (int64) value indicating the bitrate to start buffering with when seeking (not at playback start)
-	const TCHAR* const OptionKeySeekStartBitrate = TEXT("seekstart_bitrate");
+	const FName OptionKeySeekStartBitrate(TEXT("seekstart_bitrate"));
 
 	//! (int64) value indicating the bitrate to start rebuffering with
-	const TCHAR* const OptionKeyRebufferStartBitrate = TEXT("rebufferstart_bitrate");
+	const FName OptionKeyRebufferStartBitrate(TEXT("rebufferstart_bitrate"));
 
 	//! (FTimeValue) value specifying how many seconds away from the Live media timeline the seekable range should end.
-	const TCHAR* const OptionKeyLiveSeekableEndOffset = TEXT("seekable_range_live_end_offset");
+	const FName OptionKeyLiveSeekableEndOffset(TEXT("seekable_range_live_end_offset"));
 
 	//! (bool) true to just finish the currently loading segment when rebuffering. false to start over with.
-	const TCHAR* const OptionRebufferingContinuesLoading = TEXT("rebuffering_continues_loading");
+	const FName OptionRebufferingContinuesLoading(TEXT("rebuffering_continues_loading"));
 
 	//! (bool) true to throw a playback error when rebuffering occurs, false to continue normally.
-	const TCHAR* const OptionThrowErrorWhenRebuffering = TEXT("throw_error_when_rebuffering");
+	const FName OptionThrowErrorWhenRebuffering(TEXT("throw_error_when_rebuffering"));
 
 	//! (bool) true to perform frame accurate seeks (slow as decoding and discarding data from a preceeding keyframe is required)
-	const TCHAR* const OptionKeyFrameAccurateSeek = TEXT("frame_accurate_seeking");
+	const FName OptionKeyFrameAccurateSeek(TEXT("frame_accurate_seeking"));
 
 	//! (bool) true to optimize seeking for faster frame scrubbing, false to optimize for playback.
-	const TCHAR* const OptionKeyFrameOptimizeSeekForScrubbing = TEXT("optimize_seek_for_scrubbing");
+	const FName OptionKeyFrameOptimizeSeekForScrubbing(TEXT("optimize_seek_for_scrubbing"));
 	
 	//! (bool) true to emit the first decoded video frame while prerolling so it can be displayed while scrubbing.
-	const TCHAR* const OptionKeyDoNotHoldBackFirstVideoFrame = TEXT("do_not_hold_back_first_frame");
+	const FName OptionKeyDoNotHoldBackFirstVideoFrame(TEXT("do_not_hold_back_first_frame"));
 
 	//! (bool) true to not truncate the media segment access units at the end of the presentation. Must only be used without a set playback range end!
-	const TCHAR* const OptionKeyDoNotTruncateAtPresentationEnd = TEXT("do_not_truncate_at_presentation_end");
+	const FName OptionKeyDoNotTruncateAtPresentationEnd(TEXT("do_not_truncate_at_presentation_end"));
 	
 	//! (bool) true to have every request to read data break out to an external data reader.
-	const TCHAR* const OptionKeyUseExternalDataReader = TEXT("use_external_data_reader");
+	const FName OptionKeyUseExternalDataReader(TEXT("use_external_data_reader"));
 
+	const FName OptionKeyHoldCurrentFrame(TEXT("hold_current_frame"));
 
-	const TCHAR* const OptionKeyCurrentAvgStartingVideoBitrate = TEXT("current:avg_video_bitrate");
+	const FName OptionKeyCurrentAvgStartingVideoBitrate(TEXT("current:avg_video_bitrate"));
 
 } // namespace Electra
 
diff --git a/Engine/Plugins/Media/ElectraPlayer/Source/ElectraPlayerRuntime/Private/Runtime/Player/AdaptiveStreamingPlayer.cpp b/Engine/Plugins/Media/ElectraPlayer/Source/ElectraPlayerRuntime/Private/Runtime/Player/AdaptiveStreamingPlayer.cpp
--- a/Engine/Plugins/Media/ElectraPlayer/Source/ElectraPlayerRuntime/Private/Runtime/Player/AdaptiveStreamingPlayer.cpp
+++ b/Engine/Plugins/Media/ElectraPlayer/Source/ElectraPlayerRuntime/Private/Runtime/Player/AdaptiveStreamingPlayer.cpp
@@ -282,7 +282,7 @@ void FAdaptiveStreamingPlayer::Initialize(const FParamDict& Options)
 void FAdaptiveStreamingPlayer::ModifyOptions(const FParamDict& InOptionsToSetOrChange, const FParamDict& InOptionsToClear)
 {
 	// Remove options
-	TArray<FString> Keys;
+	TArray<FName> Keys;
 	InOptionsToClear.GetKeysStartingWith(FString(), Keys);
 	for(auto &Key : Keys)
 	{
@@ -293,7 +293,7 @@ void FAdaptiveStreamingPlayer::ModifyOptions(const FParamDict& InOptionsToSetOrC
 	for(auto &Key : Keys)
 	{
 		FVariantValue Value = InOptionsToSetOrChange.GetValue(Key);
-		PlayerOptions.SetOrUpdate(Key, Value);
+		PlayerOptions.Set(Key, Value);
 	}
 }
 
@@ -2632,7 +2632,7 @@ void FAdaptiveStreamingPlayer::InternalHandlePendingStartRequest(const FTimeValu
 						}
 
 						// Set the current average video bitrate in the player options for the period to retrieve if necessary.
-						PlayerOptions.SetOrUpdate(OptionKeyCurrentAvgStartingVideoBitrate, FVariantValue(StreamSelector->GetAverageBandwidth()));
+						PlayerOptions.Set(OptionKeyCurrentAvgStartingVideoBitrate, FVariantValue(StreamSelector->GetAverageBandwidth()));
 
 						InitialPlayPeriod->PrepareForPlay();
 						break;
@@ -3193,7 +3193,7 @@ void FAdaptiveStreamingPlayer::HandlePendingMediaSegmentRequests()
 				else if (PeriodState == IManifest::IPlayPeriod::EReadyState::Loaded)
 				{
 					// Set the current average video bitrate in the player options for the period to retrieve if necessary.
-					PlayerOptions.SetOrUpdate(OptionKeyCurrentAvgStartingVideoBitrate, FVariantValue(StreamSelector->GetAverageBandwidth()));
+					PlayerOptions.Set(OptionKeyCurrentAvgStartingVideoBitrate, FVariantValue(StreamSelector->GetAverageBandwidth()));
 					FinishedReq.Period->PrepareForPlay();
 					NextPendingSegmentRequests.Enqueue(MoveTemp(FinishedReq));
 					continue;
@@ -4448,7 +4448,7 @@ void FAdaptiveStreamingPlayer::InternalInitialize()
 	CreateRenderers();
 
 	// Check for codecs that are not to be used as per the user's choice, even if the device supports them.
-	auto GetExcludedCodecPrefixes = [](TArray<FString>& OutList, const FParamDict& InOptions, const FString& InKey) -> void
+	auto GetExcludedCodecPrefixes = [](TArray<FString>& OutList, const FParamDict& InOptions, const FName& InKey) -> void
 	{
 		if (InOptions.HaveKey(InKey))
 		{
@@ -4467,7 +4467,7 @@ void FAdaptiveStreamingPlayer::InternalInitialize()
 	GetExcludedCodecPrefixes(ExcludedAudioDecoderPrefixes, PlayerOptions, TEXT("excluded_codecs_audio"));
 	GetExcludedCodecPrefixes(ExcludedSubtitleDecoderPrefixes, PlayerOptions, TEXT("excluded_codecs_subtitles"));
 
-	auto GetCodecSelectionPriorities = [this](FCodecSelectionPriorities& OutPriorities, const FParamDict& InOptions, const FString& InKey, const TCHAR* const InType) -> void
+	auto GetCodecSelectionPriorities = [this](FCodecSelectionPriorities& OutPriorities, const FParamDict& InOptions, const FName& InKey, const TCHAR* const InType) -> void
 	{
 		if (InOptions.HaveKey(InKey))
 		{
diff --git a/Engine/Plugins/Media/ElectraPlayer/Source/ElectraPlayerRuntime/Private/Runtime/Player/AdaptiveStreamingPlayerDecoders.cpp b/Engine/Plugins/Media/ElectraPlayer/Source/ElectraPlayerRuntime/Private/Runtime/Player/AdaptiveStreamingPlayerDecoders.cpp
--- a/Engine/Plugins/Media/ElectraPlayer/Source/ElectraPlayerRuntime/Private/Runtime/Player/AdaptiveStreamingPlayerDecoders.cpp
+++ b/Engine/Plugins/Media/ElectraPlayer/Source/ElectraPlayerRuntime/Private/Runtime/Player/AdaptiveStreamingPlayerDecoders.cpp
@@ -11,6 +11,12 @@
 
 namespace Electra
 {
+namespace DecoderOptionKeys
+{
+static const FName SendEmptySubtitleDuringGaps(TEXT("sendEmptySubtitleDuringGaps"));
+static const FName MaxResolutionY(TEXT("max_resoY"));
+static const FName MaxResolutionYAbove30FPS(TEXT("max_resoY_above_30fps"));
+}
 
 // AU memory
 void* FAdaptiveStreamingPlayer::AUAllocate(IAccessUnitMemoryProvider::EDataType type, SIZE_T NumBytes, SIZE_T Alignment)
@@ -279,9 +285,9 @@ int32 FAdaptiveStreamingPlayer::CreateDecoder(EStreamType type)
 
 					// Add in any player options that are for decoder use
 					FParamDict AdditionalOptions;
-					TArray<FString> DecoderOptionKeys;
-					PlayerOptions.GetKeysStartingWith("videoDecoder", DecoderOptionKeys);
-					for (const FString & Key : DecoderOptionKeys)
+					TArray<FName> DecoderOptionKeys;
+					PlayerOptions.GetKeysStartingWith(TEXT("videoDecoder"), DecoderOptionKeys);
+					for (const FName& Key : DecoderOptionKeys)
 					{
 						AdditionalOptions.Set(Key, PlayerOptions.GetValue(Key));
 					}
@@ -385,7 +391,7 @@ int32 FAdaptiveStreamingPlayer::CreateDecoder(EStreamType type)
 				if (SubtitleDecoder.Decoder)
 				{
 					FParamDict Options;
-					Options.Set(TEXT("sendEmptySubtitleDuringGaps"), Electra::FVariantValue(true));
+					Options.Set(DecoderOptionKeys::SendEmptySubtitleDuringGaps, Electra::FVariantValue(true));
 
 					SubtitleDecoder.Decoder->SetPlayerSessionServices(this);
 					SubtitleDecoder.Decoder->GetDecodedSubtitleReceiveDelegate().BindRaw(this, &FAdaptiveStreamingPlayer::OnDecodedSubtitleReceived);
@@ -803,15 +809,15 @@ void FAdaptiveStreamingPlayer::Deprecate_InternalInitializeDecoderLimits()
 	// compared against these limits and those that exceed the limit will not be considered for playback.
 
 	// Maximum allowed vertical resolution specified?
-	if (PlayerOptions.HaveKey(TEXT("max_resoY")))
+	if (PlayerOptions.HaveKey(DecoderOptionKeys::MaxResolutionY))
 	{
-		PlayerConfig.H264LimitUpto30fps.MaxResolution.Height = (int32)PlayerOptions.GetValue(TEXT("max_resoY")).GetInt64();
-		PlayerConfig.H264LimitAbove30fps.MaxResolution.Height = (int32)PlayerOptions.GetValue(TEXT("max_resoY")).GetInt64();
+		PlayerConfig.H264LimitUpto30fps.MaxResolution.Height =
+		PlayerConfig.H264LimitAbove30fps.MaxResolution.Height = (int32)PlayerOptions.GetValue(DecoderOptionKeys::MaxResolutionY).GetInt64();
 	}
 	// A limit in vertical resolution for streams with more than 30fps?
-	if (PlayerOptions.HaveKey(TEXT("max_resoY_above_30fps")))
+	if (PlayerOptions.HaveKey(DecoderOptionKeys::MaxResolutionYAbove30FPS))
 	{
-		PlayerConfig.H264LimitAbove30fps.MaxResolution.Height = (int32)PlayerOptions.GetValue(TEXT("max_resoY_above_30fps")).GetInt64();
+		PlayerConfig.H264LimitAbove30fps.MaxResolution.Height = (int32)PlayerOptions.GetValue(DecoderOptionKeys::MaxResolutionYAbove30FPS).GetInt64();
 	}
 }
 
diff --git a/Engine/Plugins/Media/ElectraPlayer/Source/ElectraPlayerRuntime/Private/Runtime/Player/AdaptiveStreamingPlayerInternal.h b/Engine/Plugins/Media/ElectraPlayer/Source/ElectraPlayerRuntime/Private/Runtime/Player/AdaptiveStreamingPlayerInternal.h
--- a/Engine/Plugins/Media/ElectraPlayer/Source/ElectraPlayerRuntime/Private/Runtime/Player/AdaptiveStreamingPlayerInternal.h
+++ b/Engine/Plugins/Media/ElectraPlayer/Source/ElectraPlayerRuntime/Private/Runtime/Player/AdaptiveStreamingPlayerInternal.h
@@ -23,6 +23,7 @@
 #include "Player/PlayerStreamFilter.h"
 #include "Player/PlayerEntityCache.h"
 #include "Player/AdaptiveStreamingPlayerABR.h"
+#include "Player/AdaptivePlayerOptionKeynames.h"
 
 #include "Utilities/UtilsMP4.h"
 
@@ -1185,7 +1186,7 @@ private:
 			if (Renderer.IsValid())
 			{
 				FParamDict options;
-				options.Set("hold_current_frame", FVariantValue(bHoldCurrentFrame));
+				options.Set(Electra::OptionKeyHoldCurrentFrame, FVariantValue(bHoldCurrentFrame));
 				Renderer->Flush(options);
 			}
 		}
diff --git a/Engine/Plugins/Media/ElectraPlayer/Source/ElectraPlayerRuntime/Private/Runtime/Player/AdaptiveStreamingPlayerRenderers.cpp b/Engine/Plugins/Media/ElectraPlayer/Source/ElectraPlayerRuntime/Private/Runtime/Player/AdaptiveStreamingPlayerRenderers.cpp
--- a/Engine/Plugins/Media/ElectraPlayer/Source/ElectraPlayerRuntime/Private/Runtime/Player/AdaptiveStreamingPlayerRenderers.cpp
+++ b/Engine/Plugins/Media/ElectraPlayer/Source/ElectraPlayerRuntime/Private/Runtime/Player/AdaptiveStreamingPlayerRenderers.cpp
@@ -13,21 +13,6 @@
 
 #include <atomic>
 
-#define VALIDITY_VALUE_KEY TEXT("$renderVV$")
-#define DURATION_VALUE_KEY TEXT("duration")
-#define TIMESTAMP_VALUE_KEY TEXT("pts")
-#define EOS_VALUE_KEY TEXT("eos")
-
-#define AUDIO_BUFFER_SIZE TEXT("max_buffer_size")
-#define AUDIO_BUFFER_NUM TEXT("num_buffers")
-#define AUDIO_BUFFER_MAX_CHANNELS TEXT("max_channels")
-#define AUDIO_BUFFER_SAMPLES_PER_BLOCK TEXT("samples_per_block")
-#define AUDIO_BUFFER_ALLOCATED_ADDRESS TEXT("address")
-#define AUDIO_BUFFER_ALLOCATED_SIZE TEXT("size")
-#define AUDIO_BUFFER_NUM_CHANNELS TEXT("num_channels")
-#define AUDIO_BUFFER_BYTE_SIZE TEXT("byte_size")
-#define AUDIO_BUFFER_SAMPLE_RATE TEXT("sample_rate")
-
 /***************************************************************************************************************************************************/
 
 DECLARE_STATS_GROUP(TEXT("Electra Audio"), STATGROUP_ElectraAudioProcessing, STATCAT_Advanced);
@@ -271,11 +256,11 @@ void FAdaptiveStreamingWrappedRenderer::SampleReleasedToPool(IDecoderOutput* InD
 	check(InDecoderOutput);
 	if (InDecoderOutput && RenderClock.IsValid())
 	{
-		int64 ValidityValue = InDecoderOutput->GetMutablePropertyDictionary().GetValue(VALIDITY_VALUE_KEY).SafeGetInt64(0);
+		int64 ValidityValue = InDecoderOutput->GetMutablePropertyDictionary().GetValue(RenderOptionKeys::ValidityValue).SafeGetInt64(0);
 		if (ValidityValue == CurrentValidityValue)
 		{
-			FTimeValue RenderTime = InDecoderOutput->GetMutablePropertyDictionary().GetValue(TIMESTAMP_VALUE_KEY).SafeGetTimeValue(FTimeValue::GetInvalid());
-			FTimeValue Duration = InDecoderOutput->GetMutablePropertyDictionary().GetValue(DURATION_VALUE_KEY).SafeGetTimeValue(FTimeValue::GetZero());
+			FTimeValue RenderTime = InDecoderOutput->GetMutablePropertyDictionary().GetValue(RenderOptionKeys::PTS).SafeGetTimeValue(FTimeValue::GetInvalid());
+			FTimeValue Duration = InDecoderOutput->GetMutablePropertyDictionary().GetValue(RenderOptionKeys::Duration).SafeGetTimeValue(FTimeValue::GetZero());
 			switch(Type)
 			{
 				case EStreamType::Video:
@@ -330,16 +315,16 @@ UEMediaError FAdaptiveStreamingWrappedRenderer::CreateBufferPool(const FParamDic
 		FMemory::Free(AudioVars.AudioTempSourceBuffer);
 		AudioVars.AudioTempSourceBuffer = nullptr;
 		AudioVars.NumAudioBuffersInUse = 0;
-		AudioVars.OriginalAudioBufferNum = Parameters.GetValue(AUDIO_BUFFER_NUM).SafeGetInt64(0);
-		AudioVars.OriginalAudioBufferSize = Parameters.GetValue(AUDIO_BUFFER_SIZE).SafeGetInt64(0);
+		AudioVars.OriginalAudioBufferNum = Parameters.GetValue(RenderOptionKeys::NumBuffers).SafeGetInt64(0);
+		AudioVars.OriginalAudioBufferSize = Parameters.GetValue(RenderOptionKeys::MaxBufferSize).SafeGetInt64(0);
 		if (AudioVars.OriginalAudioBufferSize)
 		{
 			// We need an occasional extra buffer when the input sample sequence counter changes. Double the number of buffers to accommodate.
-			Parameters.SetOrUpdate(AUDIO_BUFFER_NUM, FVariantValue(AudioVars.OriginalAudioBufferNum * 2));
+			Parameters.Set(RenderOptionKeys::NumBuffers, FVariantValue(AudioVars.OriginalAudioBufferNum * 2));
 			AudioVars.AudioTempSourceBuffer = (float*)FMemory::Malloc(AudioVars.OriginalAudioBufferSize);
 
-			int32 SamplesPerBlock = (int32) Parameters.GetValue(AUDIO_BUFFER_SAMPLES_PER_BLOCK).SafeGetInt64(2048);
-			int32 MaxChannels = (int32) Parameters.GetValue(AUDIO_BUFFER_MAX_CHANNELS).SafeGetInt64(8);
+			int32 SamplesPerBlock = (int32) Parameters.GetValue(RenderOptionKeys::SamplesPerBlock).SafeGetInt64(2048);
+			int32 MaxChannels = (int32) Parameters.GetValue(RenderOptionKeys::MaxChannels).SafeGetInt64(8);
 
 			// Get maximum number of samples we may produce when slowing down the most.
 			int32 NumTempoSamples = AudioVars.TempoChanger->GetNominalOutputSampleNum(MaxSampleRate, MinPlaybackSpeed);
@@ -347,7 +332,7 @@ UEMediaError FAdaptiveStreamingWrappedRenderer::CreateBufferPool(const FParamDic
 
 			AudioVars.MaxOutputSampleBlockSize = Utils::Max(Utils::Max(NumTempoSamples, NumResampleSamples), SamplesPerBlock);
 			AudioVars.AudioBufferSize = AudioVars.MaxOutputSampleBlockSize * MaxChannels * sizeof(float);
-			Parameters.SetOrUpdate(AUDIO_BUFFER_SIZE, FVariantValue(AudioVars.AudioBufferSize));
+			Parameters.Set(RenderOptionKeys::MaxBufferSize, FVariantValue(AudioVars.AudioBufferSize));
 			AudioVars.TempoChanger->SetMaxOutputSamples(AudioVars.MaxOutputSampleBlockSize);
 		}
 	}
@@ -372,7 +357,7 @@ UEMediaError FAdaptiveStreamingWrappedRenderer::AcquireBuffer(IBuffer*& OutBuffe
 		if (Error == UEMEDIA_ERROR_OK && OutBuffer)
 		{
 			++NumBuffersInCirculation;
-			OutBuffer->GetMutableBufferProperties().SetOrUpdate(AUDIO_BUFFER_ALLOCATED_SIZE, FVariantValue(AudioVars.OriginalAudioBufferSize));
+			OutBuffer->GetMutableBufferProperties().Set(RenderOptionKeys::AllocatedSize, FVariantValue(AudioVars.OriginalAudioBufferSize));
 		}
 		return Error;
 	}
@@ -393,7 +378,7 @@ UEMediaError FAdaptiveStreamingWrappedRenderer::ReturnBuffer(IBuffer* Buffer, bo
 
 	FParamDict SampleProperties(InSampleProperties);
 	Lock.Lock();
-	SampleProperties.SetOrUpdate(VALIDITY_VALUE_KEY, FVariantValue(CurrentValidityValue));
+	SampleProperties.Set(RenderOptionKeys::ValidityValue, FVariantValue(CurrentValidityValue));
 	Lock.Unlock();
 	if (Type == EStreamType::Video)
 	{
@@ -458,14 +443,14 @@ UEMediaError FAdaptiveStreamingWrappedRenderer::ReturnAudioBuffer(IBuffer* Buffe
 
 UEMediaError FAdaptiveStreamingWrappedRenderer::ReturnBufferCommon(IBuffer* Buffer, bool bRender, FParamDict& InSampleProperties)
 {
-	FTimeValue Duration = InSampleProperties.GetValue(DURATION_VALUE_KEY).SafeGetTimeValue(FTimeValue::GetZero());
+	FTimeValue Duration = InSampleProperties.GetValue(RenderOptionKeys::Duration).SafeGetTimeValue(FTimeValue::GetZero());
 	if (!Duration.IsValid())
 	{
 		Duration.SetToZero();
-		InSampleProperties.SetOrUpdate(DURATION_VALUE_KEY, FVariantValue(Duration));
+		InSampleProperties.Set(RenderOptionKeys::Duration, FVariantValue(Duration));
 	}
 
-	bool bIsUnusedReturnBuffer = bRender == false && InSampleProperties.GetValue(EOS_VALUE_KEY).SafeGetBool(false) == false;
+	bool bIsUnusedReturnBuffer = bRender == false && InSampleProperties.GetValue(RenderOptionKeys::EOSFlag).SafeGetBool(false) == false;
 
 	FScopeLock lock(&Lock);
 	EnqueuedDuration += Duration;
@@ -664,13 +649,13 @@ bool FAdaptiveStreamingWrappedRenderer::ProcessAudio(bool& bOutNeed2ndBuffer, IB
 	bool bGetResiduals = bOutNeed2ndBuffer;
 	bOutNeed2ndBuffer = false;
 
-	int32 SizeInBytes = (int32)InSampleProperties.GetValue(AUDIO_BUFFER_BYTE_SIZE).SafeGetInt64();
-	FTimeValue Timestamp = InSampleProperties.GetValue(TIMESTAMP_VALUE_KEY).SafeGetTimeValue(FTimeValue::GetInvalid());
-	FTimeValue Duration = InSampleProperties.GetValue(DURATION_VALUE_KEY).SafeGetTimeValue(FTimeValue::GetZero());
-	int32 NumChannels = (int32)InSampleProperties.GetValue(AUDIO_BUFFER_NUM_CHANNELS).SafeGetInt64();
+	int32 SizeInBytes = (int32)InSampleProperties.GetValue(RenderOptionKeys::UsedByteSize).SafeGetInt64();
+	FTimeValue Timestamp = InSampleProperties.GetValue(RenderOptionKeys::PTS).SafeGetTimeValue(FTimeValue::GetInvalid());
+	FTimeValue Duration = InSampleProperties.GetValue(RenderOptionKeys::Duration).SafeGetTimeValue(FTimeValue::GetZero());
+	int32 NumChannels = (int32)InSampleProperties.GetValue(RenderOptionKeys::NumChannels).SafeGetInt64();
 	int32 NumSamples = SizeInBytes / NumChannels / sizeof(float);
-	int32 SampleRate = (int32)InSampleProperties.GetValue(AUDIO_BUFFER_SAMPLE_RATE).SafeGetInt64();
-	float* BufferAddress = (float*)Buffer->GetBufferProperties().GetValue(AUDIO_BUFFER_ALLOCATED_ADDRESS).GetPointer();
+	int32 SampleRate = (int32)InSampleProperties.GetValue(RenderOptionKeys::SampleRate).SafeGetInt64();
+	float* BufferAddress = (float*)Buffer->GetBufferProperties().GetValue(RenderOptionKeys::AllocatedAddress).GetPointer();
 
 	if (AudioVars.CurrentConfig.DiffersFrom(SampleRate, NumChannels))
 	{
@@ -789,9 +774,9 @@ bool FAdaptiveStreamingWrappedRenderer::ProcessAudio(bool& bOutNeed2ndBuffer, IB
 	if (bUpdateProperties)
 	{
 		check(NumSamples <= AudioVars.AudioBufferSize / NumChannels / sizeof(float));
-		InSampleProperties.SetOrUpdate(AUDIO_BUFFER_BYTE_SIZE, FVariantValue((int64)(NumSamples * sizeof(float) * NumChannels)));
-		InSampleProperties.SetOrUpdate(TIMESTAMP_VALUE_KEY, FVariantValue(Timestamp));
-		InSampleProperties.SetOrUpdate(DURATION_VALUE_KEY, FVariantValue(FTimeValue(NumSamples, SampleRate, 0)));
+		InSampleProperties.Set(RenderOptionKeys::UsedByteSize, FVariantValue((int64)(NumSamples * sizeof(float) * NumChannels)));
+		InSampleProperties.Set(RenderOptionKeys::PTS, FVariantValue(Timestamp));
+		InSampleProperties.Set(RenderOptionKeys::Duration, FVariantValue(FTimeValue(NumSamples, SampleRate, 0)));
 	}
 
 	// Need to interpolate this block's start samples from the last block's last values?
diff --git a/Engine/Plugins/Media/ElectraPlayer/Source/ElectraPlayerRuntime/Private/Runtime/Player/DASH/ManifestDASH.cpp b/Engine/Plugins/Media/ElectraPlayer/Source/ElectraPlayerRuntime/Private/Runtime/Player/DASH/ManifestDASH.cpp
--- a/Engine/Plugins/Media/ElectraPlayer/Source/ElectraPlayerRuntime/Private/Runtime/Player/DASH/ManifestDASH.cpp
+++ b/Engine/Plugins/Media/ElectraPlayer/Source/ElectraPlayerRuntime/Private/Runtime/Player/DASH/ManifestDASH.cpp
@@ -37,7 +37,7 @@ namespace Electra
 namespace DashUtils
 {
 	#define GETPLAYEROPTION(Type, Getter)																						\
-		bool GetPlayerOption(IPlayerSessionServices* InPlayerSessionServices, Type& OutValue, const TCHAR* Key, Type Default)	\
+		bool GetPlayerOption(IPlayerSessionServices* InPlayerSessionServices, Type& OutValue, const FName& Key, Type Default)	\
 		{																														\
 			if (InPlayerSessionServices->GetOptions().HaveKey(Key))																\
 			{																													\
@@ -1040,7 +1040,7 @@ void FDASHPlayPeriod::PrepareForPlay()
 		{
 			if (llDesc->Latency.ReferenceID >= 0)
 			{
-				PlayerSessionServices->GetOptions().SetOrUpdate(DASH::OptionKey_LatencyReferenceId, FVariantValue(llDesc->Latency.ReferenceID));
+				PlayerSessionServices->GetOptions().Set(DASH::OptionKey_LatencyReferenceId, FVariantValue(llDesc->Latency.ReferenceID));
 			}
 			else
 			{
diff --git a/Engine/Plugins/Media/ElectraPlayer/Source/ElectraPlayerRuntime/Private/Runtime/Player/DASH/OptionKeynamesDASH.h b/Engine/Plugins/Media/ElectraPlayer/Source/ElectraPlayerRuntime/Private/Runtime/Player/DASH/OptionKeynamesDASH.h
--- a/Engine/Plugins/Media/ElectraPlayer/Source/ElectraPlayerRuntime/Private/Runtime/Player/DASH/OptionKeynamesDASH.h
+++ b/Engine/Plugins/Media/ElectraPlayer/Source/ElectraPlayerRuntime/Private/Runtime/Player/DASH/OptionKeynamesDASH.h
@@ -9,23 +9,23 @@ namespace Electra
 namespace DASH
 {
 
-const TCHAR* const OptionKeyMPDLoadConnectTimeout = TEXT("mpd_connection_timeout");					//!< (FTimeValue) value specifying connection timeout fetching the MPD
-const TCHAR* const OptionKeyMPDLoadNoDataTimeout = TEXT("mpd_nodata_timeout");						//!< (FTimeValue) value specifying no-data timeout fetching the MPD
-const TCHAR* const OptionKeyMPDReloadConnectTimeout = TEXT("mpd_update_connection_timeout");		//!< (FTimeValue) value specifying connection timeout fetching the MPD repeatedly
-const TCHAR* const OptionKeyMPDReloadNoDataTimeout = TEXT("mpd_update_nodata_timeout");				//!< (FTimeValue) value specifying no-data timeout fetching the MPD repeatedly
+const FName OptionKeyMPDLoadConnectTimeout(TEXT("mpd_connection_timeout"));					//!< (FTimeValue) value specifying connection timeout fetching the MPD
+const FName OptionKeyMPDLoadNoDataTimeout(TEXT("mpd_nodata_timeout"));						//!< (FTimeValue) value specifying no-data timeout fetching the MPD
+const FName OptionKeyMPDReloadConnectTimeout(TEXT("mpd_update_connection_timeout"));		//!< (FTimeValue) value specifying connection timeout fetching the MPD repeatedly
+const FName OptionKeyMPDReloadNoDataTimeout(TEXT("mpd_update_nodata_timeout"));				//!< (FTimeValue) value specifying no-data timeout fetching the MPD repeatedly
 
-const TCHAR* const OptionKeyInitSegmentConnectTimeout = TEXT("initsegment_connection_timeout");		//!< (FTimeValue) value specifying connection timeout fetching an init segment
-const TCHAR* const OptionKeyInitSegmentNoDataTimeout = TEXT("initsegment_nodata_timeout");			//!< (FTimeValue) value specifying no-data timeout fetching an init segment
-const TCHAR* const OptionKeyMediaSegmentConnectTimeout = TEXT("mediasegment_connection_timeout");	//!< (FTimeValue) value specifying connection timeout fetching a media segment
-const TCHAR* const OptionKeyMediaSegmentNoDataTimeout = TEXT("mediasegment_nodata_timeout");		//!< (FTimeValue) value specifying no-data timeout fetching a media segment
+const FName OptionKeyInitSegmentConnectTimeout(TEXT("initsegment_connection_timeout"));		//!< (FTimeValue) value specifying connection timeout fetching an init segment
+const FName OptionKeyInitSegmentNoDataTimeout(TEXT("initsegment_nodata_timeout"));			//!< (FTimeValue) value specifying no-data timeout fetching an init segment
+const FName OptionKeyMediaSegmentConnectTimeout(TEXT("mediasegment_connection_timeout"));	//!< (FTimeValue) value specifying connection timeout fetching a media segment
+const FName OptionKeyMediaSegmentNoDataTimeout(TEXT("mediasegment_nodata_timeout"));		//!< (FTimeValue) value specifying no-data timeout fetching a media segment
 
-const TCHAR* const OptionKey_MinTimeBetweenMPDUpdates = TEXT("dash:min_mpd_update_interval");		//!< (FTimeValue) value limiting the time between MPD updates unless these are required updates.
+const FName OptionKey_MinTimeBetweenMPDUpdates(TEXT("dash:min_mpd_update_interval"));		//!< (FTimeValue) value limiting the time between MPD updates unless these are required updates.
 
-const TCHAR* const OptionKey_CurrentCDN = TEXT("dash:current_cdn");			//!< A string giving the currently selected CDN to be used with <BaseURL@serviceLocation> elements.
+const FName OptionKey_CurrentCDN(TEXT("dash:current_cdn"));									//!< A string giving the currently selected CDN to be used with <BaseURL@serviceLocation> elements.
 
-const TCHAR* const OptionKey_LatencyReferenceId = TEXT("dash:latency_refId");	//!< The Id of the <ProducerReferenceTime> element used by <ServiceDescription><Latency>
+const FName OptionKey_LatencyReferenceId(TEXT("dash:latency_refId"));						//!< The Id of the <ProducerReferenceTime> element used by <ServiceDescription><Latency>
 
-const TCHAR* const HTTPHeaderOptionName = TEXT("MPEG-DASH-Param");			//!< HTTP request header name carrying optional DASH specific parameters.
+const FString HTTPHeaderOptionName(TEXT("MPEG-DASH-Param"));									//!< HTTP request header name carrying optional DASH specific parameters.
 
 }
 
diff --git a/Engine/Plugins/Media/ElectraPlayer/Source/ElectraPlayerRuntime/Private/Runtime/Player/DASH/StreamReaderDASH.cpp b/Engine/Plugins/Media/ElectraPlayer/Source/ElectraPlayerRuntime/Private/Runtime/Player/DASH/StreamReaderDASH.cpp
--- a/Engine/Plugins/Media/ElectraPlayer/Source/ElectraPlayerRuntime/Private/Runtime/Player/DASH/StreamReaderDASH.cpp
+++ b/Engine/Plugins/Media/ElectraPlayer/Source/ElectraPlayerRuntime/Private/Runtime/Player/DASH/StreamReaderDASH.cpp
@@ -1612,7 +1612,7 @@ void FStreamReaderDASH::FStreamHandler::HandleRequestMP4()
 						{
 							case FStreamCodecInformation::ECodec::AAC:
 							{
-								n = ActiveTrackData.CSD->ParsedInfo.GetExtras().GetValue("samples_per_block").SafeGetInt64(1024);
+								n = ActiveTrackData.CSD->ParsedInfo.GetExtras().GetValue(StreamCodecInformationOptions::SamplesPerBlock).SafeGetInt64(1024);
 								break;
 							}
 						}
@@ -1888,7 +1888,7 @@ void FStreamReaderDASH::FStreamHandler::HandleRequestMKV()
 			{
 				ActiveTrackData.CSD->ParsedInfo = Track->GetCodecInformation();
 				ActiveTrackData.CSD->CodecSpecificData = ActiveTrackData.CSD->ParsedInfo.GetCodecSpecificData();
-				FVariantValue dcr = ActiveTrackData.CSD->ParsedInfo.GetExtras().GetValue(TEXT("dcr"));
+				FVariantValue dcr = ActiveTrackData.CSD->ParsedInfo.GetExtras().GetValue(StreamCodecInformationOptions::DecoderConfigurationRecord);
 				if (dcr.IsValid() && dcr.IsType(FVariantValue::EDataType::TypeU8Array))
 				{
 					ActiveTrackData.CSD->RawCSD = dcr.GetArray();
diff --git a/Engine/Plugins/Media/ElectraPlayer/Source/ElectraPlayerRuntime/Private/Runtime/Player/HLS/ManifestBuilderHLS.cpp b/Engine/Plugins/Media/ElectraPlayer/Source/ElectraPlayerRuntime/Private/Runtime/Player/HLS/ManifestBuilderHLS.cpp
--- a/Engine/Plugins/Media/ElectraPlayer/Source/ElectraPlayerRuntime/Private/Runtime/Player/HLS/ManifestBuilderHLS.cpp
+++ b/Engine/Plugins/Media/ElectraPlayer/Source/ElectraPlayerRuntime/Private/Runtime/Player/HLS/ManifestBuilderHLS.cpp
@@ -473,20 +473,6 @@ FErrorDetail FManifestBuilderHLS::SetupVariants(FManifestHLSInternal* Manifest,
 				}
 				return false;
 			};
-
-			// Check indicator for b-frame presence.
-			FString CustAttrBFrame;
-			if (GetCustomAttribute(CustAttrBFrame, TEXT("BFRAMES")))
-			{
-				// The value is an int with 0=no B frames, 1=B frames present.
-				int32 bCustAttrBFrame = 0;
-				LexFromString(bCustAttrBFrame, *CustAttrBFrame);
-				if (bCustAttrBFrame != 0)
-				{
-					// We set custom options only when the attribute is explicitly set to non-0.
-					CompanyCustomExtraOptions.Set(TEXT("b_frames"), FVariantValue((int64) 1));
-				}
-			}
 		}
 
 		// The only required attribute is bandwidth, so let's make sure it's there.
diff --git a/Engine/Plugins/Media/ElectraPlayer/Source/ElectraPlayerRuntime/Private/Runtime/Player/HLS/PlaylistReaderHLS.cpp b/Engine/Plugins/Media/ElectraPlayer/Source/ElectraPlayerRuntime/Private/Runtime/Player/HLS/PlaylistReaderHLS.cpp
--- a/Engine/Plugins/Media/ElectraPlayer/Source/ElectraPlayerRuntime/Private/Runtime/Player/HLS/PlaylistReaderHLS.cpp
+++ b/Engine/Plugins/Media/ElectraPlayer/Source/ElectraPlayerRuntime/Private/Runtime/Player/HLS/PlaylistReaderHLS.cpp
@@ -32,16 +32,16 @@ DECLARE_CYCLE_STAT(TEXT("FPlaylistReaderHLS_WorkerThread"), STAT_ElectraPlayer_H
 namespace Electra
 {
 
-const FString IPlaylistReaderHLS::OptionKeyLiveSeekableStartOffset(TEXT("seekable_range_live_start_offset"));							//!< (FTimeValue) value specifying how many seconds away from the Live media timeline the seekable range should start.
-const FString IPlaylistReaderHLS::OptionKeyLiveSeekableEndOffsetAudioOnly(TEXT("seekable_range_live_end_offset_audioonly"));			//!< (FTimeValue) value specifying how many seconds away from the Live media timeline the seekable range should end for audio-only playlists.
-const FString IPlaylistReaderHLS::OptionKeyLiveSeekableEndOffsetBeConservative(TEXT("seekable_range_live_end_offset_conservative"));	//!< (bool) true to use a larger Live edge distance, false to go with the smaller absolute difference
-
-const FString IPlaylistReaderHLS::OptionKeyMasterPlaylistLoadConnectTimeout (TEXT("master_playlist_connection_timeout") );		//!< (FTimeValue) value specifying connection timeout fetching the master playlist
-const FString IPlaylistReaderHLS::OptionKeyMasterPlaylistLoadNoDataTimeout  (TEXT("master_playlist_nodata_timeout")     );		//!< (FTimeValue) value specifying no-data timeout fetching the master playlist
-const FString IPlaylistReaderHLS::OptionKeyVariantPlaylistLoadConnectTimeout(TEXT("variant_playlist_connection_timeout"));		//!< (FTimeValue) value specifying connection timeout fetching a variant playlist the first time
-const FString IPlaylistReaderHLS::OptionKeyVariantPlaylistLoadNoDataTimeout (TEXT("variant_playlist_nodata_timeout")    );		//!< (FTimeValue) value specifying no-data timeout fetching a variant playlist the first time
-const FString IPlaylistReaderHLS::OptionKeyUpdatePlaylistLoadConnectTimeout (TEXT("update_playlist_connection_timeout") );		//!< (FTimeValue) value specifying connection timeout fetching a variant playlist repeatedly
-const FString IPlaylistReaderHLS::OptionKeyUpdatePlaylistLoadNoDataTimeout  (TEXT("update_playlist_nodata_timeout" )    );		//!< (FTimeValue) value specifying no-data timeout fetching a variant playlist repeatedly
+const FName IPlaylistReaderHLS::OptionKeyLiveSeekableStartOffset(TEXT("seekable_range_live_start_offset"));							//!< (FTimeValue) value specifying how many seconds away from the Live media timeline the seekable range should start.
+const FName IPlaylistReaderHLS::OptionKeyLiveSeekableEndOffsetAudioOnly(TEXT("seekable_range_live_end_offset_audioonly"));			//!< (FTimeValue) value specifying how many seconds away from the Live media timeline the seekable range should end for audio-only playlists.
+const FName IPlaylistReaderHLS::OptionKeyLiveSeekableEndOffsetBeConservative(TEXT("seekable_range_live_end_offset_conservative"));	//!< (bool) true to use a larger Live edge distance, false to go with the smaller absolute difference
+
+const FName IPlaylistReaderHLS::OptionKeyMasterPlaylistLoadConnectTimeout (TEXT("master_playlist_connection_timeout") );		//!< (FTimeValue) value specifying connection timeout fetching the master playlist
+const FName IPlaylistReaderHLS::OptionKeyMasterPlaylistLoadNoDataTimeout  (TEXT("master_playlist_nodata_timeout")     );		//!< (FTimeValue) value specifying no-data timeout fetching the master playlist
+const FName IPlaylistReaderHLS::OptionKeyVariantPlaylistLoadConnectTimeout(TEXT("variant_playlist_connection_timeout"));		//!< (FTimeValue) value specifying connection timeout fetching a variant playlist the first time
+const FName IPlaylistReaderHLS::OptionKeyVariantPlaylistLoadNoDataTimeout (TEXT("variant_playlist_nodata_timeout")    );		//!< (FTimeValue) value specifying no-data timeout fetching a variant playlist the first time
+const FName IPlaylistReaderHLS::OptionKeyUpdatePlaylistLoadConnectTimeout (TEXT("update_playlist_connection_timeout") );		//!< (FTimeValue) value specifying connection timeout fetching a variant playlist repeatedly
+const FName IPlaylistReaderHLS::OptionKeyUpdatePlaylistLoadNoDataTimeout  (TEXT("update_playlist_nodata_timeout" )    );		//!< (FTimeValue) value specifying no-data timeout fetching a variant playlist repeatedly
 
 
 
diff --git a/Engine/Plugins/Media/ElectraPlayer/Source/ElectraPlayerRuntime/Private/Runtime/Player/HLS/PlaylistReaderHLS.h b/Engine/Plugins/Media/ElectraPlayer/Source/ElectraPlayerRuntime/Private/Runtime/Player/HLS/PlaylistReaderHLS.h
--- a/Engine/Plugins/Media/ElectraPlayer/Source/ElectraPlayerRuntime/Private/Runtime/Player/HLS/PlaylistReaderHLS.h
+++ b/Engine/Plugins/Media/ElectraPlayer/Source/ElectraPlayerRuntime/Private/Runtime/Player/HLS/PlaylistReaderHLS.h
@@ -21,15 +21,15 @@ public:
 
 	virtual ~IPlaylistReaderHLS() = default;
 
-	static const FString OptionKeyLiveSeekableStartOffset;				//!< (FTimeValue) value specifying how many seconds away from the Live media timeline the seekable range should start.
-	static const FString OptionKeyLiveSeekableEndOffsetAudioOnly;		//!< (FTimeValue) value specifying how many seconds away from the Live media timeline the seekable range should end for audio-only playlists.
-	static const FString OptionKeyLiveSeekableEndOffsetBeConservative;	//!< (bool) true to use a larger Live edge distance, false to go with the smaller absolute difference
-	static const FString OptionKeyMasterPlaylistLoadConnectTimeout;		//!< (FTimeValue) value specifying connection timeout fetching the master playlist
-	static const FString OptionKeyMasterPlaylistLoadNoDataTimeout;		//!< (FTimeValue) value specifying no-data timeout fetching the master playlist
-	static const FString OptionKeyVariantPlaylistLoadConnectTimeout;	//!< (FTimeValue) value specifying connection timeout fetching a variant playlist the first time
-	static const FString OptionKeyVariantPlaylistLoadNoDataTimeout;		//!< (FTimeValue) value specifying no-data timeout fetching a variant playlist the first time
-	static const FString OptionKeyUpdatePlaylistLoadConnectTimeout;		//!< (FTimeValue) value specifying connection timeout fetching a variant playlist repeatedly
-	static const FString OptionKeyUpdatePlaylistLoadNoDataTimeout;		//!< (FTimeValue) value specifying no-data timeout fetching a variant playlist repeatedly
+	static const FName OptionKeyLiveSeekableStartOffset;				//!< (FTimeValue) value specifying how many seconds away from the Live media timeline the seekable range should start.
+	static const FName OptionKeyLiveSeekableEndOffsetAudioOnly;		//!< (FTimeValue) value specifying how many seconds away from the Live media timeline the seekable range should end for audio-only playlists.
+	static const FName OptionKeyLiveSeekableEndOffsetBeConservative;	//!< (bool) true to use a larger Live edge distance, false to go with the smaller absolute difference
+	static const FName OptionKeyMasterPlaylistLoadConnectTimeout;		//!< (FTimeValue) value specifying connection timeout fetching the master playlist
+	static const FName OptionKeyMasterPlaylistLoadNoDataTimeout;		//!< (FTimeValue) value specifying no-data timeout fetching the master playlist
+	static const FName OptionKeyVariantPlaylistLoadConnectTimeout;	//!< (FTimeValue) value specifying connection timeout fetching a variant playlist the first time
+	static const FName OptionKeyVariantPlaylistLoadNoDataTimeout;		//!< (FTimeValue) value specifying no-data timeout fetching a variant playlist the first time
+	static const FName OptionKeyUpdatePlaylistLoadConnectTimeout;		//!< (FTimeValue) value specifying connection timeout fetching a variant playlist repeatedly
+	static const FName OptionKeyUpdatePlaylistLoadNoDataTimeout;		//!< (FTimeValue) value specifying no-data timeout fetching a variant playlist repeatedly
 
 	/**
 	 * Loads and parses the playlist.
diff --git a/Engine/Plugins/Media/ElectraPlayer/Source/ElectraPlayerRuntime/Private/Runtime/Player/HLS/StreamReaderHLSfmp4.cpp b/Engine/Plugins/Media/ElectraPlayer/Source/ElectraPlayerRuntime/Private/Runtime/Player/HLS/StreamReaderHLSfmp4.cpp
--- a/Engine/Plugins/Media/ElectraPlayer/Source/ElectraPlayerRuntime/Private/Runtime/Player/HLS/StreamReaderHLSfmp4.cpp
+++ b/Engine/Plugins/Media/ElectraPlayer/Source/ElectraPlayerRuntime/Private/Runtime/Player/HLS/StreamReaderHLSfmp4.cpp
@@ -1183,7 +1183,7 @@ void FStreamReaderHLSfmp4::FStreamHandler::HandleRequest()
 							{
 								case FStreamCodecInformation::ECodec::AAC:
 								{
-									n = CSD->ParsedInfo.GetExtras().GetValue("samples_per_block").SafeGetInt64(1024);
+									n = CSD->ParsedInfo.GetExtras().GetValue(StreamCodecInformationOptions::SamplesPerBlock).SafeGetInt64(1024);
 									break;
 								}
 							}
diff --git a/Engine/Plugins/Media/ElectraPlayer/Source/ElectraPlayerRuntime/Private/Runtime/Player/SimpleElectraAudioPlayer.cpp b/Engine/Plugins/Media/ElectraPlayer/Source/ElectraPlayerRuntime/Private/Runtime/Player/SimpleElectraAudioPlayer.cpp
--- a/Engine/Plugins/Media/ElectraPlayer/Source/ElectraPlayerRuntime/Private/Runtime/Player/SimpleElectraAudioPlayer.cpp
+++ b/Engine/Plugins/Media/ElectraPlayer/Source/ElectraPlayerRuntime/Private/Runtime/Player/SimpleElectraAudioPlayer.cpp
@@ -2,6 +2,7 @@
 
 #include "SimpleElectraAudioPlayer.h"
 #include "AdaptiveStreamingPlayer.h"
+#include "Player/AdaptivePlayerOptionKeynames.h"
 #include "Renderer/RendererAudio.h"
 #include "CoreMinimal.h"
 #include "Misc/Guid.h"
@@ -1461,19 +1462,19 @@ void FSimpleElectraAudioPlayer::InternalApplyOptions()
 	const FVariant* Mimetype = Options.Find(TEXT("mimetype"));
 	if (Mimetype)
 	{
-		Opts.SetOrUpdate(TEXT("mime_type"), FVariantValue(Mimetype->GetValue<FString>()));
+		Opts.Set(Electra::OptionKeyMimeType, FVariantValue(Mimetype->GetValue<FString>()));
 	}
 
 	const FVariant* InitialBitrate = Options.Find(TEXT("initial_bitrate"));
 	if (InitialBitrate)
 	{
-		Opts.SetOrUpdate(TEXT("initial_bitrate"), FVariantValue((int64) InitialBitrate->GetValue<int32>()));
+		Opts.Set(Electra::OptionKeyInitialBitrate, FVariantValue((int64) InitialBitrate->GetValue<int32>()));
 	}
 
 	const FVariant* DoNotTruncate = Options.Find(TEXT("do_not_truncate_at_presentation_end"));
 	if (DoNotTruncate)
 	{
-		Opts.SetOrUpdate(TEXT("do_not_truncate_at_presentation_end"), FVariantValue(DoNotTruncate->GetValue<bool>()));
+		Opts.Set(Electra::OptionKeyDoNotTruncateAtPresentationEnd, FVariantValue(DoNotTruncate->GetValue<bool>()));
 	}
 
 	TSharedPtr<IAdaptiveStreamingPlayer, ESPMode::ThreadSafe> LockedPlayer(Player);
@@ -1874,13 +1875,13 @@ UEMediaError FSimpleElectraAudioPlayerRenderer::CreateBufferPool(const FParamDic
 {
 	DecoderOutputPool.Reset();
 	NumOutputAudioBuffersInUse = 0;
-	NumBuffers = (int32)Parameters.GetValue("num_buffers").SafeGetInt64(-1);
+	NumBuffers = (int32)Parameters.GetValue(RenderOptionKeys::NumBuffers).SafeGetInt64(-1);
 	if (NumBuffers <= 0)
 	{
 		return UEMEDIA_ERROR_BAD_ARGUMENTS;
 	}
-	BufferPoolProperties.Set("max_buffers", Electra::FVariantValue((int64)NumBuffers));
-	MaxBufferSize = (uint32)Parameters.GetValue("max_buffer_size").SafeGetInt64(0);
+	BufferPoolProperties.Set(RenderOptionKeys::MaxBuffers, Electra::FVariantValue((int64)NumBuffers));
+	MaxBufferSize = (uint32)Parameters.GetValue(RenderOptionKeys::MaxBufferSize).SafeGetInt64(0);
 	return MaxBufferSize ? UEMEDIA_ERROR_OK : UEMEDIA_ERROR_BAD_ARGUMENTS;
 }
 
@@ -1899,10 +1900,9 @@ UEMediaError FSimpleElectraAudioPlayerRenderer::AcquireBuffer(IBuffer*& OutBuffe
 	FPlatformAtomics::InterlockedIncrement(&NumBuffersAcquiredForDecoder);
 	DecoderOutput->Reserve(MaxBufferSize);
 	FPlatformAtomics::InterlockedIncrement(&NumOutputAudioBuffersInUse);
-
 	FMediaBufferSharedPtrWrapper* MediaBufferSharedPtrWrapper = new FMediaBufferSharedPtrWrapper(DecoderOutput);
-	MediaBufferSharedPtrWrapper->BufferProperties.Set("size", FVariantValue((int64)DecoderOutput->GetReservedBufferBytes()));
-	MediaBufferSharedPtrWrapper->BufferProperties.Set("address", FVariantValue(const_cast<void*>(DecoderOutput->GetBuffer())));
+	MediaBufferSharedPtrWrapper->BufferProperties.Set(RenderOptionKeys::AllocatedSize, FVariantValue((int64)DecoderOutput->GetReservedBufferBytes()));
+	MediaBufferSharedPtrWrapper->BufferProperties.Set(RenderOptionKeys::AllocatedAddress, FVariantValue(const_cast<void*>(DecoderOutput->GetBuffer())));
 	OutBuffer = MediaBufferSharedPtrWrapper;
 	return UEMEDIA_ERROR_OK;
 }
@@ -1916,15 +1916,15 @@ UEMediaError FSimpleElectraAudioPlayerRenderer::ReturnBuffer(IBuffer* Buffer, bo
 	FMediaBufferSharedPtrWrapper* MediaBufferSharedPtrWrapper = static_cast<FMediaBufferSharedPtrWrapper*>(Buffer);
 	if (bRender)
 	{
-		int32 NumChannels = (int32)InSampleProperties.GetValue("num_channels").SafeGetInt64(-1);
-		int32 SampleRate = (int32)InSampleProperties.GetValue("sample_rate").SafeGetInt64(-1);
-		int32 UsedBufferBytes = (int32)InSampleProperties.GetValue("byte_size").SafeGetInt64(-1);
-		FTimeValue decPTS = InSampleProperties.GetValue("pts").SafeGetTimeValue();
+		int32 NumChannels = (int32)InSampleProperties.GetValue(RenderOptionKeys::NumChannels).SafeGetInt64(-1);
+		int32 SampleRate = (int32)InSampleProperties.GetValue(RenderOptionKeys::SampleRate).SafeGetInt64(-1);
+		int32 UsedBufferBytes = (int32)InSampleProperties.GetValue(RenderOptionKeys::UsedByteSize).SafeGetInt64(-1);
+		FTimeValue decPTS = InSampleProperties.GetValue(RenderOptionKeys::PTS).SafeGetTimeValue();
 		FTimespan PTS = decPTS.GetAsTimespan();
 		int64 InSequenceIndex = decPTS.GetSequenceIndex();
 
-		FTimespan Duration = InSampleProperties.GetValue("duration").SafeGetTimeValue().GetAsTimespan();
-		const void* BufferAddr = MediaBufferSharedPtrWrapper->BufferProperties.GetValue("address").SafeGetPointer();
+		FTimespan Duration = InSampleProperties.GetValue(RenderOptionKeys::Duration).SafeGetTimeValue().GetAsTimespan();
+		const void* BufferAddr = MediaBufferSharedPtrWrapper->BufferProperties.GetValue(RenderOptionKeys::AllocatedAddress).SafeGetPointer();
 		if (NumChannels <= 0 || SampleRate <= 0 || UsedBufferBytes <= 0)
 		{
 			return UEMEDIA_ERROR_BAD_ARGUMENTS;
@@ -1942,7 +1942,7 @@ UEMediaError FSimpleElectraAudioPlayerRenderer::ReturnBuffer(IBuffer* Buffer, bo
 	}
 	else
 	{
-		if (InSampleProperties.GetValue("eos").SafeGetBool(false))
+		if (InSampleProperties.GetValue(RenderOptionKeys::EOSFlag).SafeGetBool(false))
 		{
 			TSharedPtr<FSimpleElectraAudioPlayer, ESPMode::ThreadSafe> PinnedPlayer = Player.Pin();
 			if (PinnedPlayer.IsValid())
diff --git a/Engine/Plugins/Media/ElectraPlayer/Source/ElectraPlayerRuntime/Private/Runtime/Player/mkv/OptionKeynamesMKV.h b/Engine/Plugins/Media/ElectraPlayer/Source/ElectraPlayerRuntime/Private/Runtime/Player/mkv/OptionKeynamesMKV.h
--- a/Engine/Plugins/Media/ElectraPlayer/Source/ElectraPlayerRuntime/Private/Runtime/Player/mkv/OptionKeynamesMKV.h
+++ b/Engine/Plugins/Media/ElectraPlayer/Source/ElectraPlayerRuntime/Private/Runtime/Player/mkv/OptionKeynamesMKV.h
@@ -9,8 +9,8 @@ namespace Electra
 namespace MKV
 {
 
-const TCHAR* const OptionKeyMKVLoadConnectTimeout = TEXT("mkv_connection_timeout");				//!< (FTimeValue) value specifying connection timeout fetching the MKV
-const TCHAR* const OptionKeyMKVLoadNoDataTimeout = TEXT("mkv_nodata_timeout");					//!< (FTimeValue) value specifying no-data timeout fetching the MKV
+const FName OptionKeyMKVLoadConnectTimeout(TEXT("mkv_connection_timeout"));				//!< (FTimeValue) value specifying connection timeout fetching the MKV
+const FName OptionKeyMKVLoadNoDataTimeout(TEXT("mkv_nodata_timeout"));					//!< (FTimeValue) value specifying no-data timeout fetching the MKV
 
 }
 
diff --git a/Engine/Plugins/Media/ElectraPlayer/Source/ElectraPlayerRuntime/Private/Runtime/Player/mkv/StreamReaderMKV.cpp b/Engine/Plugins/Media/ElectraPlayer/Source/ElectraPlayerRuntime/Private/Runtime/Player/mkv/StreamReaderMKV.cpp
--- a/Engine/Plugins/Media/ElectraPlayer/Source/ElectraPlayerRuntime/Private/Runtime/Player/mkv/StreamReaderMKV.cpp
+++ b/Engine/Plugins/Media/ElectraPlayer/Source/ElectraPlayerRuntime/Private/Runtime/Player/mkv/StreamReaderMKV.cpp
@@ -395,7 +395,7 @@ void FStreamReaderMKV::HandleRequest()
 				TSharedPtrTS<FAccessUnit::CodecData> CSD(new FAccessUnit::CodecData);
 				CSD->ParsedInfo = SelectedTrackMetadata->CodecInfo;
 				CSD->CodecSpecificData = SelectedTrackMetadata->CodecInfo.GetCodecSpecificData();
-				FVariantValue dcr = SelectedTrackMetadata->CodecInfo.GetExtras().GetValue(TEXT("dcr"));
+				FVariantValue dcr = SelectedTrackMetadata->CodecInfo.GetExtras().GetValue(StreamCodecInformationOptions::DecoderConfigurationRecord);
 				if (dcr.IsValid() && dcr.IsType(FVariantValue::EDataType::TypeU8Array))
 				{
 					CSD->RawCSD = dcr.GetArray();
diff --git a/Engine/Plugins/Media/ElectraPlayer/Source/ElectraPlayerRuntime/Private/Runtime/Player/mp4/OptionKeynamesMP4.h b/Engine/Plugins/Media/ElectraPlayer/Source/ElectraPlayerRuntime/Private/Runtime/Player/mp4/OptionKeynamesMP4.h
--- a/Engine/Plugins/Media/ElectraPlayer/Source/ElectraPlayerRuntime/Private/Runtime/Player/mp4/OptionKeynamesMP4.h
+++ b/Engine/Plugins/Media/ElectraPlayer/Source/ElectraPlayerRuntime/Private/Runtime/Player/mp4/OptionKeynamesMP4.h
@@ -9,8 +9,8 @@ namespace Electra
 namespace MP4
 {
 
-const TCHAR* const OptionKeyMP4LoadConnectTimeout = TEXT("mp4_connection_timeout");				//!< (FTimeValue) value specifying connection timeout fetching the MP4
-const TCHAR* const OptionKeyMP4LoadNoDataTimeout = TEXT("mp4_nodata_timeout");					//!< (FTimeValue) value specifying no-data timeout fetching the MP4
+const FName OptionKeyMP4LoadConnectTimeout(TEXT("mp4_connection_timeout"));				//!< (FTimeValue) value specifying connection timeout fetching the MP4
+const FName OptionKeyMP4LoadNoDataTimeout(TEXT("mp4_nodata_timeout"));					//!< (FTimeValue) value specifying no-data timeout fetching the MP4
 
 }
 
diff --git a/Engine/Plugins/Media/ElectraPlayer/Source/ElectraPlayerRuntime/Private/Runtime/Renderer/RendererAudio.cpp b/Engine/Plugins/Media/ElectraPlayer/Source/ElectraPlayerRuntime/Private/Runtime/Renderer/RendererAudio.cpp
--- a/Engine/Plugins/Media/ElectraPlayer/Source/ElectraPlayerRuntime/Private/Runtime/Renderer/RendererAudio.cpp
+++ b/Engine/Plugins/Media/ElectraPlayer/Source/ElectraPlayerRuntime/Private/Runtime/Renderer/RendererAudio.cpp
@@ -72,7 +72,7 @@ UEMediaError FElectraRendererAudio::CreateBufferPool(const Electra::FParamDict&
 	DecoderOutputPool.Reset();
 	NumOutputAudioBuffersInUse = 0;
 
-	const FVariantValue& variantNumBuffers = Parameters.GetValue("num_buffers");
+	const FVariantValue& variantNumBuffers = Parameters.GetValue(RenderOptionKeys::NumBuffers);
 	if (!variantNumBuffers.IsType(FVariantValue::EDataType::TypeInt64))
 	{
 		return UEMEDIA_ERROR_BAD_ARGUMENTS;
@@ -80,9 +80,9 @@ UEMediaError FElectraRendererAudio::CreateBufferPool(const Electra::FParamDict&
 	NumBuffers = (int32)variantNumBuffers.GetInt64();
 
 	// Update dic for later query via GetBufferPoolProperties
-	BufferPoolProperties.Set("max_buffers", Electra::FVariantValue((int64)NumBuffers));
+	BufferPoolProperties.Set(RenderOptionKeys::MaxBuffers, Electra::FVariantValue((int64)NumBuffers));
 
-	const FVariantValue& variantMaxBufferSize = Parameters.GetValue("max_buffer_size");
+	const FVariantValue& variantMaxBufferSize = Parameters.GetValue(RenderOptionKeys::MaxBufferSize);
 	if (!variantMaxBufferSize.IsType(FVariantValue::EDataType::TypeInt64))
 	{
 		return UEMEDIA_ERROR_BAD_ARGUMENTS;
@@ -117,8 +117,8 @@ UEMediaError FElectraRendererAudio::AcquireBuffer(IBuffer*& OutBuffer, int32 Tim
 
 	// UE handles all with shared pointers...
 	FMediaBufferSharedPtrWrapper* MediaBufferSharedPtrWrapper = new FMediaBufferSharedPtrWrapper(DecoderOutput);
-	MediaBufferSharedPtrWrapper->BufferProperties.Set("size", FVariantValue((int64)DecoderOutput->GetReservedBufferBytes()));
-	MediaBufferSharedPtrWrapper->BufferProperties.Set("address", FVariantValue(const_cast<void*>(DecoderOutput->GetBuffer())));
+	MediaBufferSharedPtrWrapper->BufferProperties.Set(RenderOptionKeys::AllocatedSize, FVariantValue((int64)DecoderOutput->GetReservedBufferBytes()));
+	MediaBufferSharedPtrWrapper->BufferProperties.Set(RenderOptionKeys::AllocatedAddress, FVariantValue(const_cast<void*>(DecoderOutput->GetBuffer())));
 	OutBuffer = MediaBufferSharedPtrWrapper;
 
 	return UEMEDIA_ERROR_OK;
@@ -137,31 +137,31 @@ UEMediaError FElectraRendererAudio::ReturnBuffer(IBuffer* Buffer, bool bRender,
 
 	if (bRender)
 	{
-		const FVariantValue& variantNumChannels = InSampleProperties.GetValue("num_channels");
+		const FVariantValue& variantNumChannels = InSampleProperties.GetValue(RenderOptionKeys::NumChannels);
 		if (!variantNumChannels.IsType(FVariantValue::EDataType::TypeInt64))
 		{
 			return UEMEDIA_ERROR_BAD_ARGUMENTS;
 		}
 
-		const FVariantValue& variantSampleRate = InSampleProperties.GetValue("sample_rate");
+		const FVariantValue& variantSampleRate = InSampleProperties.GetValue(RenderOptionKeys::SampleRate);
 		if (!variantSampleRate.IsType(FVariantValue::EDataType::TypeInt64))
 		{
 			return UEMEDIA_ERROR_BAD_ARGUMENTS;
 		}
 
-		const FVariantValue& variantBufferUsedBytes = InSampleProperties.GetValue("byte_size");
+		const FVariantValue& variantBufferUsedBytes = InSampleProperties.GetValue(RenderOptionKeys::UsedByteSize);
 		if (!variantBufferUsedBytes.IsType(FVariantValue::EDataType::TypeInt64))
 		{
 			return UEMEDIA_ERROR_BAD_ARGUMENTS;
 		}
 
-		const FVariantValue& variantPts = InSampleProperties.GetValue("pts");
+		const FVariantValue& variantPts = InSampleProperties.GetValue(RenderOptionKeys::PTS);
 		if (!variantPts.IsType(FVariantValue::EDataType::TypeTimeValue))
 		{
 			return UEMEDIA_ERROR_BAD_ARGUMENTS;
 		}
 
-		const FVariantValue& variantDuration = InSampleProperties.GetValue("duration");
+		const FVariantValue& variantDuration = InSampleProperties.GetValue(RenderOptionKeys::Duration);
 		if (!variantDuration.IsType(FVariantValue::EDataType::TypeTimeValue))
 		{
 			return UEMEDIA_ERROR_BAD_ARGUMENTS;
diff --git a/Engine/Plugins/Media/ElectraPlayer/Source/ElectraPlayerRuntime/Private/Runtime/Renderer/RendererBase.h b/Engine/Plugins/Media/ElectraPlayer/Source/ElectraPlayerRuntime/Private/Runtime/Renderer/RendererBase.h
--- a/Engine/Plugins/Media/ElectraPlayer/Source/ElectraPlayerRuntime/Private/Runtime/Renderer/RendererBase.h
+++ b/Engine/Plugins/Media/ElectraPlayer/Source/ElectraPlayerRuntime/Private/Runtime/Renderer/RendererBase.h
@@ -225,6 +225,30 @@ namespace Electra
 		virtual void TickOutputBufferPool() {};
 	};
 
+
+	namespace RenderOptionKeys
+	{
+		static const FName ValidityValue(TEXT("$renderVV$"));
+		static const FName Duration(TEXT("duration"));
+		static const FName PTS(TEXT("pts"));
+		static const FName EOSFlag(TEXT("eos"));
+		static const FName DummyBufferFlag(TEXT("is_dummy"));
+
+		static const FName Texture(TEXT("texture"));
+		static const FName AllocatedAddress(TEXT("address"));
+		static const FName AllocatedSize(TEXT("size"));
+
+		static const FName MaxBuffers(TEXT("max_buffers"));
+		static const FName MaxBufferSize(TEXT("max_buffer_size"));
+		static const FName NumBuffers(TEXT("num_buffers"));
+		static const FName UsedByteSize(TEXT("byte_size"));
+
+		static const FName MaxChannels(TEXT("max_channels"));
+		static const FName NumChannels(TEXT("num_channels"));
+		static const FName SamplesPerBlock(TEXT("samples_per_block"));
+		static const FName SampleRate(TEXT("sample_rate"));
+	}
+
 } // namespace Electra
 
 
diff --git a/Engine/Plugins/Media/ElectraPlayer/Source/ElectraPlayerRuntime/Private/Runtime/Renderer/RendererVideo.cpp b/Engine/Plugins/Media/ElectraPlayer/Source/ElectraPlayerRuntime/Private/Runtime/Renderer/RendererVideo.cpp
--- a/Engine/Plugins/Media/ElectraPlayer/Source/ElectraPlayerRuntime/Private/Runtime/Renderer/RendererVideo.cpp
+++ b/Engine/Plugins/Media/ElectraPlayer/Source/ElectraPlayerRuntime/Private/Runtime/Renderer/RendererVideo.cpp
@@ -107,13 +107,13 @@ void FElectraRendererVideo::AcquireFromPool(FVideoDecoderOutputPtr& DelayedImage
  */
 UEMediaError FElectraRendererVideo::CreateBufferPool(const Electra::FParamDict& Parameters)
 {
-	const FVariantValue& variantNumBuffers = Parameters.GetValue("num_buffers");
+	const FVariantValue& variantNumBuffers = Parameters.GetValue(RenderOptionKeys::NumBuffers);
 	if (!variantNumBuffers.IsType(FVariantValue::EDataType::TypeInt64))
 		return UEMEDIA_ERROR_BAD_ARGUMENTS;
 	int32 RequestedNumBuffers = (int32)variantNumBuffers.GetInt64();
 
 	// Update dict for later query via GetBufferPoolProperties
-	BufferPoolProperties.Set("max_buffers", Electra::FVariantValue((int64)RequestedNumBuffers));
+	BufferPoolProperties.Set(RenderOptionKeys::MaxBuffers, Electra::FVariantValue((int64)RequestedNumBuffers));
 
 	// Currently, we only handle enlargement of the buffer pool. If the size should shrink,
 	if (RequestedNumBuffers != NumBuffers)
@@ -175,7 +175,7 @@ UEMediaError FElectraRendererVideo::AcquireBuffer(IBuffer*& OutBuffer, int32 Tim
 
 	FMediaBufferSharedPtrWrapper* MediaBufferSharedPtrWrapper = new FMediaBufferSharedPtrWrapper(DelayedImage);
 	check(MediaBufferSharedPtrWrapper);
-	MediaBufferSharedPtrWrapper->BufferProperties.Set("texture", FVariantValue(DelayedImage));
+	MediaBufferSharedPtrWrapper->BufferProperties.Set(RenderOptionKeys::Texture, FVariantValue(DelayedImage));
 	OutBuffer = MediaBufferSharedPtrWrapper;
 
 	FPlatformAtomics::InterlockedIncrement(&NumBuffersAcquiredForDecoder);
@@ -199,7 +199,7 @@ UEMediaError FElectraRendererVideo::ReturnBuffer(IBuffer* Buffer, bool bRender,
 	if (bRender)
 	{
 		//OPT/CHANGE: Note that "MediaBufferSharedPtrWrapper->DecoderOutput->GetDict()" is the very same as InSampleProperties!
-		bool bIsDummyBuffer = InSampleProperties.GetValue("is_dummy").SafeGetBool(false);
+		bool bIsDummyBuffer = InSampleProperties.GetValue(RenderOptionKeys::DummyBufferFlag).SafeGetBool(false);
 
 		// Put frame into output queue...
 		if (TSharedPtr<FElectraPlayer, ESPMode::ThreadSafe> PinnedPlayer = Player.Pin())
diff --git a/Engine/Plugins/Media/ElectraPlayer/Source/ElectraPlayerRuntime/Private/Runtime/StreamTypes.cpp b/Engine/Plugins/Media/ElectraPlayer/Source/ElectraPlayerRuntime/Private/Runtime/StreamTypes.cpp
--- a/Engine/Plugins/Media/ElectraPlayer/Source/ElectraPlayerRuntime/Private/Runtime/StreamTypes.cpp
+++ b/Engine/Plugins/Media/ElectraPlayer/Source/ElectraPlayerRuntime/Private/Runtime/StreamTypes.cpp
@@ -258,7 +258,7 @@ namespace Electra
 		else if (CodecOTI.StartsWith(TEXT("vp08")))
 		{
 			ElectraDecodersUtil::FMimeTypeVideoCodecInfo ci;
-			if (ElectraDecodersUtil::ParseCodecVP8(ci, CodecOTI, Extras.GetValue(TEXT("vpcC_box")).SafeGetArray()))
+			if (ElectraDecodersUtil::ParseCodecVP8(ci, CodecOTI, Extras.GetValue(StreamCodecInformationOptions::VPccBox).SafeGetArray()))
 			{
 				StreamType = EStreamType::Video;
 				CodecSpecifier = CodecOTI;
@@ -283,7 +283,7 @@ namespace Electra
 		else if (CodecOTI.StartsWith(TEXT("vp09")))
 		{
 			ElectraDecodersUtil::FMimeTypeVideoCodecInfo ci;
-			if (ElectraDecodersUtil::ParseCodecVP9(ci, CodecOTI, Extras.GetValue(TEXT("vpcC_box")).SafeGetArray()))
+			if (ElectraDecodersUtil::ParseCodecVP9(ci, CodecOTI, Extras.GetValue(StreamCodecInformationOptions::VPccBox).SafeGetArray()))
 			{
 				StreamType = EStreamType::Video;
 				CodecSpecifier = CodecOTI;
diff --git a/Engine/Plugins/Media/ElectraPlayer/Source/ElectraPlayerRuntime/Private/Runtime/StreamTypes.h b/Engine/Plugins/Media/ElectraPlayer/Source/ElectraPlayerRuntime/Private/Runtime/StreamTypes.h
--- a/Engine/Plugins/Media/ElectraPlayer/Source/ElectraPlayerRuntime/Private/Runtime/StreamTypes.h
+++ b/Engine/Plugins/Media/ElectraPlayer/Source/ElectraPlayerRuntime/Private/Runtime/StreamTypes.h
@@ -658,6 +658,19 @@ namespace Electra
 	};
 
 
+	namespace StreamCodecInformationOptions
+	{
+		static const FName DecoderConfigurationRecord(TEXT("dcr"));
+		static const FName VPccBox(TEXT("vpcC_box"));
+		static const FName DOpsBox(TEXT("dOps_box"));
+		static const FName ChanBox(TEXT("chan_box"));
+		static const FName SamplesPerBlock(TEXT("samples_per_block"));
+		static const FName FormatSpecificFlags(TEXT("FormatSpecificFlags"));
+		static const FName SampleSize(TEXT("SampleSize"));
+		static const FName ConstBytesPerAudioPacket(TEXT("ConstBytesPerAudioPacket"));
+		static const FName ConstLPCMFramesPerAudioPacket(TEXT("ConstLPCMFramesPerAudioPacket"));
+	}
+
 
 
 	/**
diff --git a/Engine/Plugins/Media/ElectraPlayer/Source/ElectraPlayerRuntime/Private/Windows/ElectraPlayer_Windows.cpp b/Engine/Plugins/Media/ElectraPlayer/Source/ElectraPlayerRuntime/Private/Windows/ElectraPlayer_Windows.cpp
--- a/Engine/Plugins/Media/ElectraPlayer/Source/ElectraPlayerRuntime/Private/Windows/ElectraPlayer_Windows.cpp
+++ b/Engine/Plugins/Media/ElectraPlayer/Source/ElectraPlayerRuntime/Private/Windows/ElectraPlayer_Windows.cpp
@@ -60,7 +60,7 @@ TSharedPtr<IVideoDecoderResourceDelegate, ESPMode::ThreadSafe> PlatformCreateVid
 bool FElectraPlayerPlatform::StartupPlatformResources(const Electra::FParamDict& Params)
 {
 	FElectraDecoderResourceManagerWindows::FCallbacks Callbacks;
-	auto GetDeviceTypeCallback = reinterpret_cast<void(*)(void**, int64*)>(Params.GetValue(TEXT("GetDeviceTypeCallback")).SafeGetPointer());
+	auto GetDeviceTypeCallback = reinterpret_cast<void(*)(void**, int64*)>(Params.GetValue(FName(TEXT("GetDeviceTypeCallback"))).SafeGetPointer());
 
 	Callbacks.GetD3DDevice = [GetDeviceTypeCallback](void **OutD3DDevice, int32* OutD3DVersionTimes1000, void*) -> bool
 	{
diff --git a/Engine/Plugins/Media/ElectraSubtitles/Source/ElectraSubtitles/Private/ttml/ElectraSubtitleDecoder_TTML.cpp b/Engine/Plugins/Media/ElectraSubtitles/Source/ElectraSubtitles/Private/ttml/ElectraSubtitleDecoder_TTML.cpp
--- a/Engine/Plugins/Media/ElectraSubtitles/Source/ElectraSubtitles/Private/ttml/ElectraSubtitleDecoder_TTML.cpp
+++ b/Engine/Plugins/Media/ElectraSubtitles/Source/ElectraSubtitles/Private/ttml/ElectraSubtitleDecoder_TTML.cpp
@@ -105,9 +105,14 @@ private:
 
 namespace ElectraSubtitleDecoderTTMLOptions
 {
-	static const TCHAR* const SendEmptySubtitleDuringGaps = TEXT("sendEmptySubtitleDuringGaps");
-	static const TCHAR* const SideloadedID = TEXT("SideloadedID");
-	static const TCHAR* const PTO = TEXT("PTO");
+	static const FName SendEmptySubtitleDuringGaps(TEXT("sendEmptySubtitleDuringGaps"));
+	static const FName SideloadedID(TEXT("SideloadedID"));
+	static const FName PTO(TEXT("PTO"));
+	static const FName Width(TEXT("width"));
+	static const FName Height(TEXT("height"));
+	static const FName OffsetX(TEXT("offset_x"));
+	static const FName OffsetY(TEXT("offset_y"));
+	static const FName Timescale(TEXT("timescale"));
 }
 
 
@@ -134,12 +139,11 @@ FElectraSubtitleDecoderTTML::~FElectraSubtitleDecoderTTML()
 bool FElectraSubtitleDecoderTTML::InitializeStreamWithCSD(const TArray<uint8>& InCSD, const Electra::FParamDict& InAdditionalInfo)
 {
 	// Get dimension, placement offset and other data from the sideband dictionary.
-	Width = (int32) InAdditionalInfo.GetValue(TEXT("width")).SafeGetInt64(0);
-	Height = (int32) InAdditionalInfo.GetValue(TEXT("height")).SafeGetInt64(0);
-	TranslationX = (int32) InAdditionalInfo.GetValue(TEXT("offset_x")).SafeGetInt64(0);
-	TranslationY = (int32) InAdditionalInfo.GetValue(TEXT("offset_y")).SafeGetInt64(0);
-	Timescale = (uint32) InAdditionalInfo.GetValue(TEXT("timescale")).SafeGetInt64(0);
-
+	Width = (int32) InAdditionalInfo.GetValue(ElectraSubtitleDecoderTTMLOptions::Width).SafeGetInt64(0);
+	Height = (int32) InAdditionalInfo.GetValue(ElectraSubtitleDecoderTTMLOptions::Height).SafeGetInt64(0);
+	TranslationX = (int32) InAdditionalInfo.GetValue(ElectraSubtitleDecoderTTMLOptions::OffsetX).SafeGetInt64(0);
+	TranslationY = (int32) InAdditionalInfo.GetValue(ElectraSubtitleDecoderTTMLOptions::OffsetY).SafeGetInt64(0);
+	Timescale = (uint32) InAdditionalInfo.GetValue(ElectraSubtitleDecoderTTMLOptions::Timescale).SafeGetInt64(0);
 	bSendEmptySubtitleDuringGaps = InAdditionalInfo.GetValue(ElectraSubtitleDecoderTTMLOptions::SendEmptySubtitleDuringGaps).SafeGetBool();
 
 	return true;
diff --git a/Engine/Plugins/Media/ElectraSubtitles/Source/ElectraSubtitles/Private/ttml/TTMLSubtitleList.cpp b/Engine/Plugins/Media/ElectraSubtitles/Source/ElectraSubtitles/Private/ttml/TTMLSubtitleList.cpp
--- a/Engine/Plugins/Media/ElectraSubtitles/Source/ElectraSubtitles/Private/ttml/TTMLSubtitleList.cpp
+++ b/Engine/Plugins/Media/ElectraSubtitles/Source/ElectraSubtitles/Private/ttml/TTMLSubtitleList.cpp
@@ -9,6 +9,10 @@
 namespace ElectraTTMLParser
 {
 
+namespace OptionKeys
+{
+	static const FName SendEmptySubtitleDuringGaps(TEXT("sendEmptySubtitleDuringGaps"));
+}
 
 
 class FTTMLSubtitleList : public TSharedFromThis<FTTMLSubtitleList, ESPMode::ThreadSafe>, public ITTMLSubtitleList
@@ -83,7 +87,7 @@ bool FTTMLSubtitleList::CreateFrom(TSharedPtr<ITTMLXMLElement, ESPMode::ThreadSa
 		return false;
 	}
 
-	bSendEmptySubtitleDuringGaps = InOptions.GetValue(TEXT("sendEmptySubtitleDuringGaps")).SafeGetBool();
+	bSendEmptySubtitleDuringGaps = InOptions.GetValue(OptionKeys::SendEmptySubtitleDuringGaps).SafeGetBool();
 
 	// Set the root temporal extent of the document.
 	TTMLTT->SetBegin(InDocumentStartTime);
diff --git a/Engine/Plugins/Media/ElectraSubtitles/Source/ElectraSubtitles/Private/tx3g/ElectraSubtitleDecoder_TX3G.cpp b/Engine/Plugins/Media/ElectraSubtitles/Source/ElectraSubtitles/Private/tx3g/ElectraSubtitleDecoder_TX3G.cpp
--- a/Engine/Plugins/Media/ElectraSubtitles/Source/ElectraSubtitles/Private/tx3g/ElectraSubtitleDecoder_TX3G.cpp
+++ b/Engine/Plugins/Media/ElectraSubtitles/Source/ElectraSubtitles/Private/tx3g/ElectraSubtitleDecoder_TX3G.cpp
@@ -8,6 +8,15 @@
 namespace ElectraSubtitleDecoderTX3GUtils
 {
 
+namespace OptionKeys
+{
+	static const FName Width(TEXT("width"));
+	static const FName Height(TEXT("height"));
+	static const FName OffsetX(TEXT("offset_x"));
+	static const FName OffsetY(TEXT("offset_y"));
+	static const FName Timescale(TEXT("timescale"));
+}
+
 #if !PLATFORM_LITTLE_ENDIAN
 	static inline uint8  GetFromBigEndian(uint8 value)		{ return value; }
 	static inline int8   GetFromBigEndian(int8 value)		{ return value; }
@@ -242,11 +251,11 @@ bool FElectraSubtitleDecoderTX3G::InitializeStreamWithCSD(const TArray<uint8>& I
 
 
 	// Get dimension, placement offset and other data from the sideband dictionary.
-	Width = (int32) InAdditionalInfo.GetValue(TEXT("width")).SafeGetInt64(0);
-	Height = (int32) InAdditionalInfo.GetValue(TEXT("height")).SafeGetInt64(0);
-	TranslationX = (int32) InAdditionalInfo.GetValue(TEXT("offset_x")).SafeGetInt64(0);
-	TranslationY = (int32) InAdditionalInfo.GetValue(TEXT("offset_y")).SafeGetInt64(0);
-	Timescale = (uint32) InAdditionalInfo.GetValue(TEXT("timescale")).SafeGetInt64(0);
+	Width = (int32) InAdditionalInfo.GetValue(ElectraSubtitleDecoderTX3GUtils::OptionKeys::Width).SafeGetInt64(0);
+	Height = (int32) InAdditionalInfo.GetValue(ElectraSubtitleDecoderTX3GUtils::OptionKeys::Height).SafeGetInt64(0);
+	TranslationX = (int32) InAdditionalInfo.GetValue(ElectraSubtitleDecoderTX3GUtils::OptionKeys::OffsetX).SafeGetInt64(0);
+	TranslationY = (int32) InAdditionalInfo.GetValue(ElectraSubtitleDecoderTX3GUtils::OptionKeys::OffsetY).SafeGetInt64(0);
+	Timescale = (uint32) InAdditionalInfo.GetValue(ElectraSubtitleDecoderTX3GUtils::OptionKeys::Timescale).SafeGetInt64(0);
 
 	ElectraSubtitleDecoderTX3GUtils::FDataReaderMP4 r(InCSD);
 
diff --git a/Engine/Plugins/Media/ElectraSubtitles/Source/ElectraSubtitles/Private/wvtt/ElectraSubtitleDecoder_WVTT.cpp b/Engine/Plugins/Media/ElectraSubtitles/Source/ElectraSubtitles/Private/wvtt/ElectraSubtitleDecoder_WVTT.cpp
--- a/Engine/Plugins/Media/ElectraSubtitles/Source/ElectraSubtitles/Private/wvtt/ElectraSubtitleDecoder_WVTT.cpp
+++ b/Engine/Plugins/Media/ElectraSubtitles/Source/ElectraSubtitles/Private/wvtt/ElectraSubtitleDecoder_WVTT.cpp
@@ -8,6 +8,15 @@
 namespace ElectraSubtitleDecoderWVTTUtils
 {
 
+namespace OptionKeys
+{
+	static const FName Width(TEXT("width"));
+	static const FName Height(TEXT("height"));
+	static const FName OffsetX(TEXT("offset_x"));
+	static const FName OffsetY(TEXT("offset_y"));
+	static const FName Timescale(TEXT("timescale"));
+}
+
 #if !PLATFORM_LITTLE_ENDIAN
 	static inline uint8 GetFromBigEndian(uint8 value)		{ return value; }
 	static inline int8 GetFromBigEndian(int8 value)			{ return value; }
@@ -242,11 +251,11 @@ bool FElectraSubtitleDecoderWVTT::InitializeStreamWithCSD(const TArray<uint8>& I
 		}														\
 
 	// Get dimension, placement offset and other data from the sideband dictionary.
-	Width = (int32) InAdditionalInfo.GetValue(TEXT("width")).SafeGetInt64(0);
-	Height = (int32) InAdditionalInfo.GetValue(TEXT("height")).SafeGetInt64(0);
-	TranslationX = (int32) InAdditionalInfo.GetValue(TEXT("offset_x")).SafeGetInt64(0);
-	TranslationY = (int32) InAdditionalInfo.GetValue(TEXT("offset_y")).SafeGetInt64(0);
-	Timescale = (uint32) InAdditionalInfo.GetValue(TEXT("timescale")).SafeGetInt64(0);
+	Width = (int32) InAdditionalInfo.GetValue(ElectraSubtitleDecoderWVTTUtils::OptionKeys::Width).SafeGetInt64(0);
+	Height = (int32) InAdditionalInfo.GetValue(ElectraSubtitleDecoderWVTTUtils::OptionKeys::Height).SafeGetInt64(0);
+	TranslationX = (int32) InAdditionalInfo.GetValue(ElectraSubtitleDecoderWVTTUtils::OptionKeys::OffsetX).SafeGetInt64(0);
+	TranslationY = (int32) InAdditionalInfo.GetValue(ElectraSubtitleDecoderWVTTUtils::OptionKeys::OffsetY).SafeGetInt64(0);
+	Timescale = (uint32) InAdditionalInfo.GetValue(ElectraSubtitleDecoderWVTTUtils::OptionKeys::Timescale).SafeGetInt64(0);
 
 	ElectraSubtitleDecoderWVTTUtils::FDataReaderMP4 r(InCSD);
 
diff --git a/Engine/Plugins/Media/ElectraUtil/Source/ElectraBase/Private/ParameterDictionary.cpp b/Engine/Plugins/Media/ElectraUtil/Source/ElectraBase/Private/ParameterDictionary.cpp
--- a/Engine/Plugins/Media/ElectraUtil/Source/ElectraBase/Private/ParameterDictionary.cpp
+++ b/Engine/Plugins/Media/ElectraUtil/Source/ElectraBase/Private/ParameterDictionary.cpp
@@ -440,27 +440,19 @@ void FParamDict::InternalCopy(const FParamDict& Other)
 	Dictionary = Other.Dictionary;
 }
 
-
-
 void FParamDict::Clear()
 {
 	FScopeLock lock(&Lock);
 	Dictionary.Empty();
 }
 
-void FParamDict::Set(const FString& Key, const FVariantValue& Value)
-{
-	FScopeLock lock(&Lock);
-	Dictionary.FindOrAdd(Key, Value);
-}
-
-bool FParamDict::HaveKey(const FString& Key) const
+bool FParamDict::HaveKey(const FName& Key) const
 {
 	FScopeLock lock(&Lock);
 	return Dictionary.Find(Key) != nullptr;
 }
 
-FVariantValue FParamDict::GetValue(const FString& Key) const
+FVariantValue FParamDict::GetValue(const FName& Key) const
 {
 	static FVariantValue Empty;
 	FScopeLock lock(&Lock);
@@ -468,26 +460,33 @@ FVariantValue FParamDict::GetValue(const FString& Key) const
 	return VariantValue ? *VariantValue : Empty;
 }
 
-void FParamDict::GetKeys(TArray<FString>& Keys) const
+void FParamDict::Remove(const FName& Key)
 {
-	GetKeysStartingWith(FString(), Keys);
+	FScopeLock lock(&Lock);
+	Dictionary.Remove(Key); 
+}
+
+void FParamDict::Set(const FName& Key, const FVariantValue& Value)
+{ 
+	FScopeLock lock(&Lock);
+	Dictionary.Emplace(Key, Value); 
 }
 
-void FParamDict::GetKeysStartingWith(const FString& StartsWith, TArray<FString>& Keys) const
+void FParamDict::GetKeysStartingWith(const FString& StartsWith, TArray<FName>& OutKeys) const
 {
-	Keys.Empty();
+	OutKeys.Empty();
 	FScopeLock lock(&Lock);
 	if (StartsWith.IsEmpty())
 	{
-		Dictionary.GenerateKeyArray(Keys);
+		Dictionary.GenerateKeyArray(OutKeys);
 	}
 	else
 	{
-		for(const TPair<FString, FVariantValue>& Pair : Dictionary)
+		for(const TPair<FName, FVariantValue>& Pair : Dictionary)
 		{
-			if (Pair.Key.StartsWith(StartsWith, ESearchCase::CaseSensitive))
+			if (Pair.Key.ToString().StartsWith(StartsWith, ESearchCase::CaseSensitive))
 			{
-				Keys.Emplace(Pair.Key);
+				OutKeys.Emplace(Pair.Key);
 			}
 		}
 	}
@@ -495,16 +494,20 @@ void FParamDict::GetKeysStartingWith(const FString& StartsWith, TArray<FString>&
 
 void FParamDict::ConvertKeysStartingWithTo(TMap<FString, FVariant>& OutVariantMap, const FString& InKeyStartsWith, const FString& InAddPrefixToKey) const
 {
+	OutVariantMap.Reserve(Dictionary.Num());
+	FString NewKey;
+	NewKey.Reserve(64);
 	FScopeLock lock(&Lock);
-	for(const TPair<FString, FVariantValue>& Pair : Dictionary)
+	for(const TPair<FName, FVariantValue>& Pair : Dictionary)
 	{
-		if (!InKeyStartsWith.IsEmpty() && !Pair.Key.StartsWith(InKeyStartsWith, ESearchCase::CaseSensitive))
+		FString s(Pair.Key.ToString());
+		if (!InKeyStartsWith.IsEmpty() && !s.StartsWith(InKeyStartsWith, ESearchCase::CaseSensitive))
 		{
 			continue;
 		}
 
-		FString NewKey = InAddPrefixToKey;
-		NewKey.Append(Pair.Key);
+		NewKey = InAddPrefixToKey;
+		NewKey.Append(s);
 		switch(Pair.Value.GetDataType())
 		{
 			case FVariantValue::EDataType::TypeFString:
diff --git a/Engine/Plugins/Media/ElectraUtil/Source/ElectraBase/Public/MediaDecoderOutput.h b/Engine/Plugins/Media/ElectraUtil/Source/ElectraBase/Public/MediaDecoderOutput.h
--- a/Engine/Plugins/Media/ElectraUtil/Source/ElectraBase/Public/MediaDecoderOutput.h
+++ b/Engine/Plugins/Media/ElectraUtil/Source/ElectraBase/Public/MediaDecoderOutput.h
@@ -300,3 +300,28 @@ public:
 private:
 	Electra::FParamDict PropertyDictionary;
 };
+
+
+namespace IDecoderOutputOptionNames
+{
+static const FName PTS(TEXT("pts"));
+static const FName Duration(TEXT("duration"));
+static const FName Width(TEXT("width"));
+static const FName Height(TEXT("height"));
+static const FName Pitch(TEXT("pitch"));
+static const FName AspectRatio(TEXT("aspect_ratio"));
+static const FName CropLeft(TEXT("crop_left"));
+static const FName CropTop(TEXT("crop_top"));
+static const FName CropRight(TEXT("crop_right"));
+static const FName CropBottom(TEXT("crop_bottom"));
+static const FName PixelFormat(TEXT("pixelfmt"));
+static const FName PixelEncoding(TEXT("pixelenc"));
+static const FName Orientation(TEXT("orientation"));
+static const FName BitsPerComponent(TEXT("bits_per"));
+static const FName HDRInfo(TEXT("hdr_info"));
+static const FName Colorimetry(TEXT("colorimetry"));
+static const FName AspectW(TEXT("aspect_w"));
+static const FName AspectH(TEXT("aspect_h"));
+static const FName FPSNumerator(TEXT("fps_num"));
+static const FName FPSDenominator(TEXT("fps_denom"));
+}
diff --git a/Engine/Plugins/Media/ElectraUtil/Source/ElectraBase/Public/MediaVideoDecoderOutput.h b/Engine/Plugins/Media/ElectraUtil/Source/ElectraBase/Public/MediaVideoDecoderOutput.h
--- a/Engine/Plugins/Media/ElectraUtil/Source/ElectraBase/Public/MediaVideoDecoderOutput.h
+++ b/Engine/Plugins/Media/ElectraUtil/Source/ElectraBase/Public/MediaVideoDecoderOutput.h
@@ -137,7 +137,7 @@ public:
 		{
 			return FDecoderTimeStamp(FTimespan::Zero(), 0);
 		}
-		Electra::FTimeValue pts(ParamDict->GetValue("pts").GetTimeValue());
+		Electra::FTimeValue pts(ParamDict->GetValue(IDecoderOutputOptionNames::PTS).GetTimeValue());
 		return FDecoderTimeStamp(pts.GetAsTimespan(), pts.GetSequenceIndex());
 	}
 
@@ -147,7 +147,7 @@ public:
 		{
 			return FTimespan(0);
 		}
-		return FTimespan(ParamDict->GetValue("duration").GetTimeValue().GetAsHNS());
+		return FTimespan(ParamDict->GetValue(IDecoderOutputOptionNames::Duration).GetTimeValue().GetAsHNS());
 	}
 
 	virtual FIntPoint GetOutputDim() const
@@ -158,8 +158,8 @@ public:
 		}
 		if (!(Cached.Flags & FCached::Valid_OutputDim))
 		{
-			Cached.OutputDim.X = (int32)ParamDict->GetValue("width").SafeGetInt64(0);
-			Cached.OutputDim.Y = (int32)ParamDict->GetValue("height").SafeGetInt64(0);
+			Cached.OutputDim.X = (int32)ParamDict->GetValue(IDecoderOutputOptionNames::Width).SafeGetInt64(0);
+			Cached.OutputDim.Y = (int32)ParamDict->GetValue(IDecoderOutputOptionNames::Height).SafeGetInt64(0);
 			FPlatformMisc::MemoryBarrier();
 			Cached.Flags |= FCached::Valid_OutputDim;
 		}
@@ -173,7 +173,7 @@ public:
 			return 1.0;
 		}
 		const FIntPoint Dim = GetOutputDim();
-		return ((double)Dim.X / (double)Dim.Y) * ParamDict->GetValue("aspect_ratio").SafeGetDouble(1.0);
+		return ((double)Dim.X / (double)Dim.Y) * ParamDict->GetValue(IDecoderOutputOptionNames::AspectRatio).SafeGetDouble(1.0);
 	}
 
 	virtual FVideoDecoderCropInfo GetCropInfo() const
@@ -184,10 +184,10 @@ public:
 		}
 		if (!(Cached.Flags & FCached::Valid_CropInfo))
 		{ 
-			Cached.CropInfo.CropLeft = (int32)ParamDict->GetValue("crop_left").SafeGetInt64(0);
-			Cached.CropInfo.CropTop = (int32)ParamDict->GetValue("crop_top").SafeGetInt64(0);
-			Cached.CropInfo.CropRight = (int32)ParamDict->GetValue("crop_right").SafeGetInt64(0);
-			Cached.CropInfo.CropBottom = (int32)ParamDict->GetValue("crop_bottom").SafeGetInt64(0);
+			Cached.CropInfo.CropLeft = (int32)ParamDict->GetValue(IDecoderOutputOptionNames::CropLeft).SafeGetInt64(0);
+			Cached.CropInfo.CropTop = (int32)ParamDict->GetValue(IDecoderOutputOptionNames::CropTop).SafeGetInt64(0);
+			Cached.CropInfo.CropRight = (int32)ParamDict->GetValue(IDecoderOutputOptionNames::CropRight).SafeGetInt64(0);
+			Cached.CropInfo.CropBottom = (int32)ParamDict->GetValue(IDecoderOutputOptionNames::CropBottom).SafeGetInt64(0);
 			FPlatformMisc::MemoryBarrier();
 			Cached.Flags |= FCached::Valid_CropInfo;
 		}
@@ -215,7 +215,7 @@ public:
 		{
 			return EPixelFormat::PF_Unknown;
 		}
-		return (EPixelFormat)ParamDict->GetValue("pixelfmt").SafeGetInt64((int64)EPixelFormat::PF_Unknown);
+		return (EPixelFormat)ParamDict->GetValue(IDecoderOutputOptionNames::PixelFormat).SafeGetInt64((int64)EPixelFormat::PF_Unknown);
 	}
 
 	virtual EVideoDecoderPixelEncoding GetFormatEncoding() const
@@ -224,7 +224,7 @@ public:
 		{
 			return EVideoDecoderPixelEncoding::Native;
 		}
-		return (EVideoDecoderPixelEncoding)ParamDict->GetValue("pixelenc").SafeGetInt64((int64)EVideoDecoderPixelEncoding::Native);
+		return (EVideoDecoderPixelEncoding)ParamDict->GetValue(IDecoderOutputOptionNames::PixelEncoding).SafeGetInt64((int64)EVideoDecoderPixelEncoding::Native);
 	}
 
 	virtual EVideoOrientation GetOrientation() const
@@ -235,7 +235,7 @@ public:
 		}
 		if (!(Cached.Flags & FCached::Valid_Orientation))
 		{
-			Cached.Orientation = (EVideoOrientation)ParamDict->GetValue("orientation").SafeGetInt64((int64)EVideoOrientation::Original);
+			Cached.Orientation = (EVideoOrientation)ParamDict->GetValue(IDecoderOutputOptionNames::Orientation).SafeGetInt64((int64)EVideoOrientation::Original);
 			FPlatformMisc::MemoryBarrier();
 			Cached.Flags |= FCached::Valid_Orientation;
 		}
@@ -247,25 +247,25 @@ public:
 		const int32 NumDefaultBits = 8;
 		if (ParamDict)
 		{
-			return (int32)ParamDict->GetValue("bits_per").SafeGetInt64((int64)NumDefaultBits);
+			return (int32)ParamDict->GetValue(IDecoderOutputOptionNames::BitsPerComponent).SafeGetInt64((int64)NumDefaultBits);
 		}
 		return NumDefaultBits;
 	}
 
 	virtual TSharedPtr<const IVideoDecoderHDRInformation, ESPMode::ThreadSafe> GetHDRInformation() const
 	{
-		if (ParamDict && ParamDict->HaveKey("hdr_info"))
+		if (ParamDict && ParamDict->HaveKey(IDecoderOutputOptionNames::HDRInfo))
 		{
-			return ParamDict->GetValue("hdr_info").GetSharedPointer<const IVideoDecoderHDRInformation>();
+			return ParamDict->GetValue(IDecoderOutputOptionNames::HDRInfo).GetSharedPointer<const IVideoDecoderHDRInformation>();
 		}
 		return nullptr;
 	}
 
 	virtual TSharedPtr<const IVideoDecoderColorimetry, ESPMode::ThreadSafe> GetColorimetry() const
 	{
-		if (ParamDict && ParamDict->HaveKey("colorimetry"))
+		if (ParamDict && ParamDict->HaveKey(IDecoderOutputOptionNames::Colorimetry))
 		{
-			return ParamDict->GetValue("colorimetry").GetSharedPointer<const IVideoDecoderColorimetry>();
+			return ParamDict->GetValue(IDecoderOutputOptionNames::Colorimetry).GetSharedPointer<const IVideoDecoderColorimetry>();
 		}
 		return nullptr;
 	}
diff --git a/Engine/Plugins/Media/ElectraUtil/Source/ElectraBase/Public/ParameterDictionary.h b/Engine/Plugins/Media/ElectraUtil/Source/ElectraBase/Public/ParameterDictionary.h
--- a/Engine/Plugins/Media/ElectraUtil/Source/ElectraBase/Public/ParameterDictionary.h
+++ b/Engine/Plugins/Media/ElectraUtil/Source/ElectraBase/Public/ParameterDictionary.h
@@ -19,7 +19,7 @@ class ELECTRABASE_API FVariantValue
 	{
 	public:
 		virtual ~FSharedPtrHolderBase() {}
-		virtual void SetValueOn(FSharedPtrHolderBase *Dst) const = 0;
+		virtual void SetValueOn(FSharedPtrHolderBase* Dst) const = 0;
 	};
 
 	template<typename T> class TSharedPtrHolder : public FSharedPtrHolderBase
@@ -33,7 +33,7 @@ class ELECTRABASE_API FVariantValue
 
 		SharedPtrType	Pointer;
 
-		virtual void SetValueOn(FSharedPtrHolderBase *Dst) const override
+		virtual void SetValueOn(FSharedPtrHolderBase* Dst) const override
 		{
 			new(Dst) TSharedPtrHolder<T>(reinterpret_cast<const SharedPtrType&>(Pointer));
 		}
@@ -43,7 +43,7 @@ public:
 	FVariantValue();
 	~FVariantValue();
 	FVariantValue(const FVariantValue& rhs);
-	FVariantValue& operator=(const FVariantValue&rhs);
+	FVariantValue& operator=(const FVariantValue& rhs);
 
 	explicit FVariantValue(const FString& StringValue);
 	explicit FVariantValue(double DoubleValue);
@@ -162,29 +162,19 @@ public:
 	FParamDict& operator=(const FParamDict& Other);
 	~FParamDict();
 	void Clear();
-	void Set(const FString& Key, const FVariantValue& Value);
-	bool HaveKey(const FString& Key) const;
-	FVariantValue GetValue(const FString& Key) const;
-	void Remove(const FString& Key)
-	{
-		FScopeLock lock(&Lock);
-		Dictionary.Remove(Key); 
-	}
-	void SetOrUpdate(const FString& Key, const FVariantValue& Value)
-	{ 
-		FScopeLock lock(&Lock);
-		Dictionary.Add(Key, Value); 
-	}
+	void Set(const FName& Key, const FVariantValue& Value);
+	bool HaveKey(const FName& Key) const;
+	FVariantValue GetValue(const FName& Key) const;
+	void Remove(const FName& Key);
 
 	void ConvertTo(TMap<FString, FVariant>& OutVariantMap, const FString& InAddPrefixToKey) const;
 	void ConvertKeysStartingWithTo(TMap<FString, FVariant>& OutVariantMap, const FString& InKeyStartsWith, const FString& InAddPrefixToKey) const;
 
-	void GetKeys(TArray<FString>& Keys) const;
-	void GetKeysStartingWith(const FString& StartsWith, TArray<FString>& Keys) const;
+	void GetKeysStartingWith(const FString& InStartsWith, TArray<FName>& OutKeys) const;
 private:
 	void InternalCopy(const FParamDict& Other);
 	mutable FCriticalSection Lock;
-	TMap<FString, FVariantValue> Dictionary;
+	TMap<FName, FVariantValue> Dictionary;
 };
 
 
diff --git a/Engine/Plugins/Media/ElectraUtil/Source/ElectraHTTPStream/Private/Curl/CurlElectraHTTPStream.cpp b/Engine/Plugins/Media/ElectraUtil/Source/ElectraHTTPStream/Private/Curl/CurlElectraHTTPStream.cpp
--- a/Engine/Plugins/Media/ElectraUtil/Source/ElectraHTTPStream/Private/Curl/CurlElectraHTTPStream.cpp
+++ b/Engine/Plugins/Media/ElectraUtil/Source/ElectraHTTPStream/Private/Curl/CurlElectraHTTPStream.cpp
@@ -877,6 +877,9 @@ void FElectraHTTPStreamRequestLibCurl::Close()
 /***************************************************************************************************************************************************/
 /***************************************************************************************************************************************************/
 
+FName FElectraHTTPStreamLibCurl::OptionName_Proxy(TEXT("proxy"));
+
+
 FElectraHTTPStreamLibCurl::FElectraHTTPStreamLibCurl()
 {
 }
@@ -890,9 +893,9 @@ bool FElectraHTTPStreamLibCurl::Initialize(const Electra::FParamDict& InOptions)
 {
 	LLM_SCOPE(ELLMTag::MediaStreaming);
 
-	if (InOptions.HaveKey(TEXT("proxy")))
+	if (InOptions.HaveKey(OptionName_Proxy))
 	{
-		ProxyAddressAndPort = InOptions.GetValue(TEXT("proxy")).SafeGetFString();
+		ProxyAddressAndPort = InOptions.GetValue(OptionName_Proxy).SafeGetFString();
 	}
 
 	CURLSHcode scode = CURLSHE_OK; (void)scode;
diff --git a/Engine/Plugins/Media/ElectraUtil/Source/ElectraHTTPStream/Private/Curl/CurlElectraHTTPStreamImpl.h b/Engine/Plugins/Media/ElectraUtil/Source/ElectraHTTPStream/Private/Curl/CurlElectraHTTPStreamImpl.h
--- a/Engine/Plugins/Media/ElectraUtil/Source/ElectraHTTPStream/Private/Curl/CurlElectraHTTPStreamImpl.h
+++ b/Engine/Plugins/Media/ElectraUtil/Source/ElectraHTTPStream/Private/Curl/CurlElectraHTTPStreamImpl.h
@@ -102,6 +102,8 @@ private:
 	TArray<TSharedPtr<FElectraHTTPStreamRequestLibCurl, ESPMode::ThreadSafe>> NewRequests;
 	TArray<TSharedPtr<FElectraHTTPStreamRequestLibCurl, ESPMode::ThreadSafe>> ActiveRequests;
 	TArray<TSharedPtr<FElectraHTTPStreamRequestLibCurl, ESPMode::ThreadSafe>> CompletedRequests;
+
+	static FName OptionName_Proxy;
 };
 
 #endif
diff --git a/Engine/Plugins/Media/ElectraUtil/Source/ElectraHTTPStream/Private/WinHttp/WinHttpElectraHTTPStream.cpp b/Engine/Plugins/Media/ElectraUtil/Source/ElectraHTTPStream/Private/WinHttp/WinHttpElectraHTTPStream.cpp
--- a/Engine/Plugins/Media/ElectraUtil/Source/ElectraHTTPStream/Private/WinHttp/WinHttpElectraHTTPStream.cpp
+++ b/Engine/Plugins/Media/ElectraUtil/Source/ElectraHTTPStream/Private/WinHttp/WinHttpElectraHTTPStream.cpp
@@ -207,10 +207,12 @@ private:
 	static std::atomic_ulong NextRequestIndex;
 	static FCriticalSection AllRequestHandleLock;
 	static TMap<uint64, FRequestPointers> AllRequestHandles;
+	static FName OptionName_Proxy;
 };
 std::atomic_ulong											FElectraHTTPStreamWinHttp::NextRequestIndex = {0};
 FCriticalSection											FElectraHTTPStreamWinHttp::AllRequestHandleLock;
 TMap<uint64, FElectraHTTPStreamWinHttp::FRequestPointers>	FElectraHTTPStreamWinHttp::AllRequestHandles;
+FName														FElectraHTTPStreamWinHttp::OptionName_Proxy(TEXT("proxy"));
 
 /***************************************************************************************************************************************************/
 /***************************************************************************************************************************************************/
@@ -1120,9 +1122,9 @@ bool FElectraHTTPStreamWinHttp::Initialize(const Electra::FParamDict& InOptions)
 #endif
 
 	HINTERNET sh;
-	if (InOptions.HaveKey(TEXT("proxy")))
+	if (InOptions.HaveKey(OptionName_Proxy))
 	{
-		FString ProxyNameAndPort = InOptions.GetValue(TEXT("proxy")).SafeGetFString();
+		FString ProxyNameAndPort = InOptions.GetValue(OptionName_Proxy).SafeGetFString();
 		if (ProxyNameAndPort.Len())
 		{
 			AccessType = WINHTTP_ACCESS_TYPE_NAMED_PROXY;
