Electra: reusing sample blocks in simple audio player

#rnx

diff --git a/Engine/Plugins/Media/ElectraPlayer/Source/ElectraPlayerRuntime/Private/Runtime/Player/SimpleElectraAudioPlayer.cpp b/Engine/Plugins/Media/ElectraPlayer/Source/ElectraPlayerRuntime/Private/Runtime/Player/SimpleElectraAudioPlayer.cpp
--- a/Engine/Plugins/Media/ElectraPlayer/Source/ElectraPlayerRuntime/Private/Runtime/Player/SimpleElectraAudioPlayer.cpp
+++ b/Engine/Plugins/Media/ElectraPlayer/Source/ElectraPlayerRuntime/Private/Runtime/Player/SimpleElectraAudioPlayer.cpp
@@ -353,8 +353,7 @@ public:
 	{
 		if (!bIsResuming)
 		{
-			FString Msg = FString::Printf(TEXT("%s: Opening"), *AssetName);
-			UE_LOG(LogSimpleElectraPlayer, Verbose, TEXT("%s"), *Msg);
+			UE_LOG(LogSimpleElectraPlayer, Verbose, TEXT("%s: Opening"), *AssetName);
 		}
 	}
 	void ReportReceivedMasterPlaylist(const FString& InEffectiveURL) override
@@ -391,9 +390,8 @@ public:
 	{
 		if (!InSegmentDownloadStats.bWasSuccessful || InSegmentDownloadStats.RetryNumber)
 		{
-			FString Msg = FString::Printf(TEXT("%s: Segment download issue \"%s\", retry:%d, aborted:%d, silence:%d"), *AssetName,
-											*InSegmentDownloadStats.FailureReason, InSegmentDownloadStats.RetryNumber, InSegmentDownloadStats.bWasAborted, InSegmentDownloadStats.bInsertedFillerData);
-			UE_LOG(LogSimpleElectraPlayer, Verbose, TEXT("%s"), *Msg);
+			UE_LOG(LogSimpleElectraPlayer, Verbose, TEXT("%s: Segment download issue \"%s\", retry:%d, aborted:%d, silence:%d"),
+				*AssetName, *InSegmentDownloadStats.FailureReason, InSegmentDownloadStats.RetryNumber, InSegmentDownloadStats.bWasAborted, InSegmentDownloadStats.bInsertedFillerData);
 		}
 
 		// Update analytics
@@ -453,8 +451,7 @@ public:
 	{
 		bHasErrored = true;
 		ErrorMessage = InErrorReason;
-		FString Msg = FString::Printf(TEXT("%s: %s"), *AssetName, *SimpleElectraAudioPlayerUtil::RedactMessage(InErrorReason));
-		UE_LOG(LogSimpleElectraPlayer, Error, TEXT("%s"), *Msg);
+		UE_LOG(LogSimpleElectraPlayer, Error, TEXT("%s: %s"), *AssetName, *SimpleElectraAudioPlayerUtil::RedactMessage(InErrorReason));
 	}
 	void ReportLogMessage(IInfoLog::ELevel InLogLevel, const FString& InLogMessage, int64 InPlayerWallclockMilliseconds) override
 	{
@@ -464,27 +461,26 @@ public:
 			return;
 		}
 
-		FString Msg = FString::Printf(TEXT("%s: %s"), *AssetName, *SimpleElectraAudioPlayerUtil::RedactMessage(InLogMessage));
 		switch(InLogLevel)
 		{
 			case IInfoLog::ELevel::Verbose:
 			{
-				UE_LOG(LogSimpleElectraPlayer, Verbose, TEXT("%s"), *Msg);
+				UE_LOG(LogSimpleElectraPlayer, Verbose, TEXT("%s: %s"), *AssetName, *SimpleElectraAudioPlayerUtil::RedactMessage(InLogMessage));
 				break;
 			}
 			case IInfoLog::ELevel::Info:
 			{
-				UE_LOG(LogSimpleElectraPlayer, Log, TEXT("%s"), *Msg)
+				UE_LOG(LogSimpleElectraPlayer, Log, TEXT("%s: %s"), *AssetName, *SimpleElectraAudioPlayerUtil::RedactMessage(InLogMessage));
 				break;
 			}
 			case IInfoLog::ELevel::Warning:
 			{
-				UE_LOG(LogSimpleElectraPlayer, Warning, TEXT("%s"), *Msg)
+				UE_LOG(LogSimpleElectraPlayer, Warning, TEXT("%s: %s"), *AssetName, *SimpleElectraAudioPlayerUtil::RedactMessage(InLogMessage));
 				break;
 			}
 			case IInfoLog::ELevel::Error:
 			{
-				UE_LOG(LogSimpleElectraPlayer, Error, TEXT("%s"), *Msg)
+				UE_LOG(LogSimpleElectraPlayer, Error, TEXT("%s: %s"), *AssetName, *SimpleElectraAudioPlayerUtil::RedactMessage(InLogMessage));
 				break;
 			}
 		}
@@ -577,12 +573,12 @@ private:
 		void* Buffer = nullptr;
 		int32 NumSamples = 0;
 		int32 Offset = 0;
-		void Free()
+		int32 MaxSamples = 0;
+		~FSampleBlock()
 		{
 			if (Buffer)
 			{
 				FMemory::Free(Buffer);
-				Buffer = nullptr;
 			}
 		}
 	};
@@ -623,6 +619,33 @@ private:
 	TSharedPtr<FResourceProvider, ESPMode::ThreadSafe> ResourceProvider;
 	TWeakPtr<IElectraPlayerDataCache, ESPMode::ThreadSafe> PlayerDataCache;
 
+	TArray<FSampleBlock*> SampleBlockPool;
+	FCriticalSection SampleBlockPoolLock;
+	FSampleBlock* AllocSampleBlockFromPool()
+	{
+		FScopeLock lock(&SampleBlockPoolLock);
+		if (SampleBlockPool.Num())
+		{
+			return SampleBlockPool.Pop();
+		}
+		return new FSampleBlock;
+	}
+	void ReturnSampleBlockToPool(FSampleBlock* InBufferToReturn)
+	{
+		if (InBufferToReturn)
+		{
+			FScopeLock lock(&SampleBlockPoolLock);
+			SampleBlockPool.Emplace(InBufferToReturn);
+		}
+	}
+	void DestroySampleBlockPool()
+	{
+		while(!SampleBlockPool.IsEmpty())
+		{
+			delete SampleBlockPool.Pop();
+		}
+	}
+
 	FString URL;
 	FString BaseURL;
 	FString AssetName;
@@ -640,16 +663,32 @@ private:
 
 	struct FBlockSequence
 	{
-		~FBlockSequence()
+		class TDeleter
 		{
-			for(int32 i=0; i<SampleBlocks.Num(); ++i)
+		public:
+			TDeleter(TFunction<void(FSampleBlock*)>&& InReturnBufferFN)
+				: ReturnBufferFN(MoveTemp(InReturnBufferFN))
+			{
+			}
+			void operator()(FBlockSequence* InInstanceToDelete)
 			{
-				SampleBlocks[i].Free();
+				for(int32 i=0; i<InInstanceToDelete->SampleBlocks.Num(); ++i)
+				{
+					ReturnBufferFN(InInstanceToDelete->SampleBlocks[i]);
+				}
+				InInstanceToDelete->SampleBlocks.Empty();
+				delete InInstanceToDelete;
 			}
-			SampleBlocks.Empty();
+		private:
+			TFunction<void(FSampleBlock*)> ReturnBufferFN;
+		};
+
+		~FBlockSequence()
+		{
+			check(SampleBlocks.IsEmpty());
 		}
 
-		TArray<FSampleBlock> SampleBlocks;
+		TArray<FSampleBlock*> SampleBlocks;
 		int64 SequenceIndex = 0;
 		int64 NumFramesAvailable = 0;
 		bool bReachedEOS = false;
@@ -918,6 +957,7 @@ FSimpleElectraAudioPlayer::~FSimpleElectraAudioPlayer()
 	Renderer.Reset();
 	FlushAudio(true);
 	FTicker::Release();
+	DestroySampleBlockPool();
 }
 
 void FSimpleElectraAudioPlayer::TickAllInstances(float InDeltaTime)
@@ -1040,8 +1080,7 @@ bool FSimpleElectraAudioPlayer::HandleStopIfRequested()
 {
 	if (bWasAskedToRelease)
 	{
-		FString Msg = FString::Printf(TEXT("%s: Stopping due to limit"), *AssetName);
-		UE_LOG(LogSimpleElectraPlayer, Verbose, TEXT("%s"), *Msg);
+		UE_LOG(LogSimpleElectraPlayer, Verbose, TEXT("%s: Stopping due to limit"), *AssetName);
 
 		TimeStopped = FPlatformTime::Seconds();
 		++Analytics.NumTimesSuspended;
@@ -1055,8 +1094,7 @@ bool FSimpleElectraAudioPlayer::HandleDestructionIfRequested()
 {
 	if (bDestructionRequested)
 	{
-		FString Msg = FString::Printf(TEXT("%s: Closing"), *AssetName);
-		UE_LOG(LogSimpleElectraPlayer, Verbose, TEXT("%s"), *Msg);
+		UE_LOG(LogSimpleElectraPlayer, Verbose, TEXT("%s: Closing"), *AssetName);
 		ClosePlayerInstance(true);
 		InstanceListLock.Lock();
 		StoppedInstances.Remove(this);
@@ -1088,8 +1126,7 @@ void FSimpleElectraAudioPlayer::HandleOpenStream(bool bTryToReopen)
 
 	if (bCreateNow)
 	{
-		FString Msg = FString::Printf(TEXT("%s: %s playback"), *AssetName, !bTryToReopen ? TEXT("Preparing") : TEXT("Resuming"));
-		UE_LOG(LogSimpleElectraPlayer, Verbose, TEXT("%s"), *Msg);
+		UE_LOG(LogSimpleElectraPlayer, Verbose, TEXT("%s: %s playback"), *AssetName, !bTryToReopen ? TEXT("Preparing") : TEXT("Resuming"));
 
 		// Track use count
 		if (bTryToReopen)
@@ -1161,7 +1198,7 @@ void FSimpleElectraAudioPlayer::CreateIdleBufferIfNecessary()
 	InstanceLock.Lock();
 	if (!CurrentReadSampleBlock.IsValid() && NextPendingSampleBlocks.IsEmpty())
 	{
-		TSharedPtr<FBlockSequence, ESPMode::ThreadSafe> NewSeq = MakeShared<FBlockSequence, ESPMode::ThreadSafe>();
+		TSharedPtr<FBlockSequence, ESPMode::ThreadSafe> NewSeq = MakeShareable(new FBlockSequence(), FBlockSequence::TDeleter([this](FSampleBlock* InBlk){ReturnSampleBlockToPool(InBlk);}));
 		NewSeq->SequenceIndex = -1;
 		NextPendingSampleBlocks.Emplace(NewSeq);
 	}
@@ -1418,8 +1455,7 @@ bool FSimpleElectraAudioPlayer::Open(const TMap<FString, FVariant>& InOptions, c
 		}
 		else
 		{
-			FString Msg = FString::Printf(TEXT("%s: Loading blob"), *AssetName);
-			UE_LOG(LogSimpleElectraPlayer, Verbose, TEXT("%s"), *Msg);
+			UE_LOG(LogSimpleElectraPlayer, Verbose, TEXT("%s: Loading blob"), *AssetName);
 			CreatePlayerInstance();
 			if (!PendingBlobRequest->Request->SetFromJSON(BlobParameters.GetValue()))
 			{
@@ -1615,16 +1651,20 @@ bool FSimpleElectraAudioPlayer::EnqueueAudioFrames(const void* InBufferAddress,
 		{
 			TimeUntilReady = FPlatformTime::Seconds() - TimeCreated;
 			Analytics.StartLatency = TimeUntilReady;
-			FString Msg = FString::Printf(TEXT("%s: ready after %.2fs"), *AssetName, TimeUntilReady);
-			UE_LOG(LogSimpleElectraPlayer, Verbose, TEXT("%s"), *Msg);
+			UE_LOG(LogSimpleElectraPlayer, Verbose, TEXT("%s: ready after %.2fs"), *AssetName, TimeUntilReady);
 		}
 
-		FSampleBlock sb;
+		FSampleBlock* sb = AllocSampleBlockFromPool();
 		int32 nb = InNumSamples * InNumChannels * sizeof(float);
-		sb.Buffer = FMemory::Malloc(nb);
-		sb.NumSamples = InNumSamples;
-		sb.PTS = InPTS;
-		FMemory::Memcpy(sb.Buffer, InBufferAddress, nb);
+		if (sb->MaxSamples < InNumSamples)
+		{
+			sb->Buffer = FMemory::Realloc(sb->Buffer, nb);
+			sb->MaxSamples = InNumSamples;
+		}
+		sb->NumSamples = InNumSamples;
+		sb->Offset = 0;
+		sb->PTS = InPTS;
+		FMemory::Memcpy(sb->Buffer, InBufferAddress, nb);
 
 		FScopeLock lock(&InstanceLock);
 
@@ -1636,12 +1676,12 @@ bool FSimpleElectraAudioPlayer::EnqueueAudioFrames(const void* InBufferAddress,
 
 		if (!CurrentWriteSampleBlock.IsValid() || CurrentWriteSampleBlock->SequenceIndex != InSequenceIndex)
 		{
-			TSharedPtr<FBlockSequence, ESPMode::ThreadSafe> NewSeq = MakeShared<FBlockSequence, ESPMode::ThreadSafe>();
+			TSharedPtr<FBlockSequence, ESPMode::ThreadSafe> NewSeq = MakeShareable(new FBlockSequence(), FBlockSequence::TDeleter([this](FSampleBlock* InBlk){ReturnSampleBlockToPool(InBlk);}));
 			NewSeq->SequenceIndex = InSequenceIndex;
 			NextPendingSampleBlocks.Emplace(NewSeq);
 			CurrentWriteSampleBlock = MoveTemp(NewSeq);
 		}
-		CurrentWriteSampleBlock->SampleBlocks.Emplace(MoveTemp(sb));
+		CurrentWriteSampleBlock->SampleBlocks.Emplace(sb);
 		CurrentWriteSampleBlock->NumFramesAvailable += InNumSamples;
 		NumActiveFramesTotal += InNumSamples;
 		++NumEnqueuedBlocks;
@@ -1761,21 +1801,21 @@ int64 FSimpleElectraAudioPlayer::GetNextSamples(FTimespan& OutPTS, bool& bOutIsF
 	}
 	int32 NumFramesToGo = InNumFramesToGet < CurrentReadSampleBlock->NumFramesAvailable ? InNumFramesToGet : CurrentReadSampleBlock->NumFramesAvailable;
 	int32 NumGot = 0;
-	FTimespan FrontPTS = CurrentReadSampleBlock->SampleBlocks[0].PTS;
-	int32 FrontPTSSampleOffset = CurrentReadSampleBlock->SampleBlocks[0].Offset;
+	FTimespan FrontPTS = CurrentReadSampleBlock->SampleBlocks[0]->PTS;
+	int32 FrontPTSSampleOffset = CurrentReadSampleBlock->SampleBlocks[0]->Offset;
 	OutPTS = FrontPTS + FTimespan::FromSeconds((double)FrontPTSSampleOffset / StreamFormat.SampleRate);
 	while(NumFramesToGo)
 	{
-		int32 NumFramesLeftInBlock = CurrentReadSampleBlock->SampleBlocks[0].NumSamples - CurrentReadSampleBlock->SampleBlocks[0].Offset;
+		int32 NumFramesLeftInBlock = CurrentReadSampleBlock->SampleBlocks[0]->NumSamples - CurrentReadSampleBlock->SampleBlocks[0]->Offset;
 		int32 NumBlockFrames = NumFramesToGo < NumFramesLeftInBlock ? NumFramesToGo : NumFramesLeftInBlock;
-		const float* Src = reinterpret_cast<const float*>(CurrentReadSampleBlock->SampleBlocks[0].Buffer) + CurrentReadSampleBlock->SampleBlocks[0].Offset * StreamFormat.NumChannels;
+		const float* Src = reinterpret_cast<const float*>(CurrentReadSampleBlock->SampleBlocks[0]->Buffer) + CurrentReadSampleBlock->SampleBlocks[0]->Offset * StreamFormat.NumChannels;
 		for(int32 i=0, iMax=NumBlockFrames*StreamFormat.NumChannels; i<iMax; ++i)
 		{
 			*OutBuffer++ = FMath::Clamp((int32)(*Src++ * 32768.0f), (int32)-32768, (int32)32767);
 		}
-		if ((CurrentReadSampleBlock->SampleBlocks[0].Offset += NumBlockFrames) >= CurrentReadSampleBlock->SampleBlocks[0].NumSamples)
+		if ((CurrentReadSampleBlock->SampleBlocks[0]->Offset += NumBlockFrames) >= CurrentReadSampleBlock->SampleBlocks[0]->NumSamples)
 		{
-			CurrentReadSampleBlock->SampleBlocks[0].Free();
+			ReturnSampleBlockToPool(CurrentReadSampleBlock->SampleBlocks[0]);
 			CurrentReadSampleBlock->SampleBlocks.RemoveAt(0);
 			--NumEnqueuedBlocks;
 		}
