ElectraPlayer: using TVariant for different message type union; removed parameter dictionary lock for internal use where it is not necessary and added safe query methods instead; GetLanguage() converts to ISO639-1 just once instead of on every call; removed unnecessary copies of connection infos; misc optimizations

#rnx

diff --git a/Engine/Plugins/Media/ElectraPlayer/Source/ElectraPlayerRuntime/Private/Runtime/Core/MediaThreads.h b/Engine/Plugins/Media/ElectraPlayer/Source/ElectraPlayerRuntime/Private/Runtime/Core/MediaThreads.h
--- a/Engine/Plugins/Media/ElectraPlayer/Source/ElectraPlayerRuntime/Private/Runtime/Core/MediaThreads.h
+++ b/Engine/Plugins/Media/ElectraPlayer/Source/ElectraPlayerRuntime/Private/Runtime/Core/MediaThreads.h
@@ -77,17 +77,17 @@ public:
 
 	static void SleepSeconds(uint32 Seconds)
 	{
-		FPlatformProcess::SleepNoStats((float)Seconds);
+		FPlatformProcess::Sleep((float)Seconds);
 	}
 
 	static void SleepMilliseconds(uint32 Milliseconds)
 	{
-		FPlatformProcess::SleepNoStats(Milliseconds / 1000.0f);
+		FPlatformProcess::Sleep(Milliseconds / 1000.0f);
 	}
 
 	static void SleepMicroseconds(uint32 Microseconds)
 	{
-		FPlatformProcess::SleepNoStats(Microseconds / 1000000.0f);
+		FPlatformProcess::Sleep(Microseconds / 1000000.0f);
 	}
 
 private:
diff --git a/Engine/Plugins/Media/ElectraPlayer/Source/ElectraPlayerRuntime/Private/Runtime/Decoder/AudioDecoder.cpp b/Engine/Plugins/Media/ElectraPlayer/Source/ElectraPlayerRuntime/Private/Runtime/Decoder/AudioDecoder.cpp
--- a/Engine/Plugins/Media/ElectraPlayer/Source/ElectraPlayerRuntime/Private/Runtime/Decoder/AudioDecoder.cpp
+++ b/Engine/Plugins/Media/ElectraPlayer/Source/ElectraPlayerRuntime/Private/Runtime/Decoder/AudioDecoder.cpp
@@ -168,6 +168,7 @@ private:
 
 	TSharedPtrTS<FAccessUnit::CodecData>									CurrentCodecData;
 	FCurrentOutputFormat													CurrentOutputFormat;
+	TMap<FString, FVariant>													CSDOptions;
 	bool																	bIsFirstAccessUnit = true;
 	bool																	bInDummyDecodeMode = false;
 	bool																	bDrainForCodecChange = false;
@@ -648,6 +649,9 @@ FAudioDecoderImpl::EAUChangeFlags FAudioDecoderImpl::GetAndPrepareInputAU()
 				if (CurrentCodecData != CurrentAccessUnit->AccessUnit->AUCodecData && CurrentAccessUnit->AccessUnit->AUCodecData.IsValid())
 				{
 					CurrentCodecData = CurrentAccessUnit->AccessUnit->AUCodecData;
+					CSDOptions.Empty();
+					CSDOptions.Emplace(TEXT("csd"), FVariant(CurrentCodecData->CodecSpecificData));
+					CurrentCodecData->ParsedInfo.GetExtras().ConvertTo(CSDOptions, TEXT("$"));
 				}
 
 				// The very first access unit can't have differences to the one before so we clear the flags.
@@ -667,13 +671,13 @@ IElectraDecoder::ECSDCompatibility FAudioDecoderImpl::IsCompatibleWith()
 	IElectraDecoder::ECSDCompatibility Compatibility = IElectraDecoder::ECSDCompatibility::Compatible;
 	if (DecoderInstance.IsValid() && CurrentAccessUnit.IsValid())
 	{
-		TMap<FString, FVariant> CSDOptions;
+		TMap<FString, FVariant> TestCSDOptions;
 		if (CurrentAccessUnit->AccessUnit->AUCodecData.IsValid())
 		{
-			CSDOptions.Emplace(TEXT("csd"), FVariant(CurrentAccessUnit->AccessUnit->AUCodecData->CodecSpecificData));
-			CSDOptions.Emplace(TEXT("dcr"), FVariant(CurrentAccessUnit->AccessUnit->AUCodecData->RawCSD));
-			CurrentAccessUnit->AccessUnit->AUCodecData->ParsedInfo.GetExtras().ConvertTo(CSDOptions, TEXT("$"));
-			Compatibility = DecoderInstance->IsCompatibleWith(CSDOptions);
+			TestCSDOptions.Emplace(TEXT("csd"), FVariant(CurrentAccessUnit->AccessUnit->AUCodecData->CodecSpecificData));
+			TestCSDOptions.Emplace(TEXT("dcr"), FVariant(CurrentAccessUnit->AccessUnit->AUCodecData->RawCSD));
+			CurrentAccessUnit->AccessUnit->AUCodecData->ParsedInfo.GetExtras().ConvertTo(TestCSDOptions, TEXT("$"));
+			Compatibility = DecoderInstance->IsCompatibleWith(TestCSDOptions);
 		}
 	}
 	return Compatibility;
@@ -903,12 +907,6 @@ bool FAudioDecoderImpl::HandleDecoding()
 			DecAU.Duration = CurrentAccessUnit->AccessUnit->Duration.GetAsTimespan();
 			DecAU.UserValue = CurrentAccessUnit->PTS;
 			DecAU.Flags = CurrentAccessUnit->AccessUnit->bIsSyncSample ? EElectraDecoderFlags::IsSyncSample : EElectraDecoderFlags::None;
-			TMap<FString, FVariant> CSDOptions;
-			if (CurrentCodecData.IsValid())
-			{
-				CSDOptions.Emplace(TEXT("csd"), FVariant(CurrentCodecData->CodecSpecificData));
-				CurrentCodecData->ParsedInfo.GetExtras().ConvertTo(CSDOptions, TEXT("$"));
-			}
 
 			IElectraDecoder::EDecoderError DecErr = DecoderInstance->DecodeAccessUnit(DecAU, CSDOptions);
 			if (DecErr == IElectraDecoder::EDecoderError::None)
diff --git a/Engine/Plugins/Media/ElectraPlayer/Source/ElectraPlayerRuntime/Private/Runtime/Decoder/VideoDecoder.cpp b/Engine/Plugins/Media/ElectraPlayer/Source/ElectraPlayerRuntime/Private/Runtime/Decoder/VideoDecoder.cpp
--- a/Engine/Plugins/Media/ElectraPlayer/Source/ElectraPlayerRuntime/Private/Runtime/Decoder/VideoDecoder.cpp
+++ b/Engine/Plugins/Media/ElectraPlayer/Source/ElectraPlayerRuntime/Private/Runtime/Decoder/VideoDecoder.cpp
@@ -54,7 +54,7 @@ public:
 	static TSharedPtr<IElectraCodecFactory, ESPMode::ThreadSafe> GetDecoderFactory(FString& OutFormat, TMap<FString, FVariant>& OutAddtlCfg, const FStreamCodecInformation& InCodecInfo, TSharedPtrTS<FAccessUnit::CodecData> InCodecData);
 
 	void SetPlayerSessionServices(IPlayerSessionServices* SessionServices) override;
-	void Open(TSharedPtrTS<FAccessUnit::CodecData> InCodecData, const FParamDict& InAdditionalOptions, const FStreamCodecInformation* InMaxStreamConfiguration) override;
+	void Open(TSharedPtrTS<FAccessUnit::CodecData> InCodecData, FParamDict&& InAdditionalOptions, const FStreamCodecInformation* InMaxStreamConfiguration) override;
 	void Close() override;
 	void DrainForCodecChange() override;
 	void SetVideoResourceDelegate(TWeakPtr<IVideoDecoderResourceDelegate, ESPMode::ThreadSafe> InVideoResourceDelegate) override;
@@ -294,10 +294,10 @@ void FVideoDecoderImpl::SetPlayerSessionServices(IPlayerSessionServices* InSessi
 	SessionServices = InSessionServices;
 }
 
-void FVideoDecoderImpl::Open(TSharedPtrTS<FAccessUnit::CodecData> InCodecData, const FParamDict& InAdditionalOptions, const FStreamCodecInformation* InMaxStreamConfiguration)
+void FVideoDecoderImpl::Open(TSharedPtrTS<FAccessUnit::CodecData> InCodecData, FParamDict&& InAdditionalOptions, const FStreamCodecInformation* InMaxStreamConfiguration)
 {
 	InitialCodecSpecificData = InCodecData;
-	InitialAdditionalOptions = InAdditionalOptions;
+	InitialAdditionalOptions = MoveTemp(InAdditionalOptions);
 	if (InMaxStreamConfiguration)
 	{
 		InitialMaxStreamProperties = *InMaxStreamConfiguration;
diff --git a/Engine/Plugins/Media/ElectraPlayer/Source/ElectraPlayerRuntime/Private/Runtime/Decoder/VideoDecoder.h b/Engine/Plugins/Media/ElectraPlayer/Source/ElectraPlayerRuntime/Private/Runtime/Decoder/VideoDecoder.h
--- a/Engine/Plugins/Media/ElectraPlayer/Source/ElectraPlayerRuntime/Private/Runtime/Decoder/VideoDecoder.h
+++ b/Engine/Plugins/Media/ElectraPlayer/Source/ElectraPlayerRuntime/Private/Runtime/Decoder/VideoDecoder.h
@@ -26,7 +26,7 @@ namespace Electra
 
 		virtual void SetPlayerSessionServices(IPlayerSessionServices* SessionServices) = 0;
 
-		virtual void Open(TSharedPtrTS<FAccessUnit::CodecData> InCodecData, const FParamDict& InAdditionalOptions, const FStreamCodecInformation* InMaxStreamConfiguration) = 0;
+		virtual void Open(TSharedPtrTS<FAccessUnit::CodecData> InCodecData, FParamDict&& InAdditionalOptions, const FStreamCodecInformation* InMaxStreamConfiguration) = 0;
 		virtual void Close() = 0;
 		virtual void DrainForCodecChange() = 0;
 
diff --git a/Engine/Plugins/Media/ElectraPlayer/Source/ElectraPlayerRuntime/Private/Runtime/Demuxer/ParserISO14496-12.cpp b/Engine/Plugins/Media/ElectraPlayer/Source/ElectraPlayerRuntime/Private/Runtime/Demuxer/ParserISO14496-12.cpp
--- a/Engine/Plugins/Media/ElectraPlayer/Source/ElectraPlayerRuntime/Private/Runtime/Demuxer/ParserISO14496-12.cpp
+++ b/Engine/Plugins/Media/ElectraPlayer/Source/ElectraPlayerRuntime/Private/Runtime/Demuxer/ParserISO14496-12.cpp
@@ -1471,9 +1471,9 @@ private:
 			return Timescale;
 		}
 
-		FString GetLanguage() const
+		const FString& GetLanguage() const
 		{
-			return FString(FMEDIA_STATIC_ARRAY_COUNT(Language), Language);
+			return Lang639_1;
 		}
 
 	private:
@@ -1509,6 +1509,8 @@ private:
 			Language[0] = (char)(0x60 + ((Value16 & 0x7c00) >> 10));
 			Language[1] = (char)(0x60 + ((Value16 & 0x03e0) >> 5));
 			Language[2] = (char)(0x60 + (Value16 & 0x001f));
+			// Try to map the ISO-639-2T language code to the shorter ISO-639-1 code if possible.
+			Lang639_1 = ISO639::MapTo639_1(FString(FMEDIA_STATIC_ARRAY_COUNT(Language), Language));
 			RETURN_IF_ERROR(ParseInfo->Reader()->Read(Value16));				// pre_defined (in QuickTime this held 'Quality')
 			return Error;
 		}
@@ -1517,6 +1519,7 @@ private:
 		uint64		Duration;
 		uint32		Timescale;
 		char		Language[3];
+		FString		Lang639_1;
 	};
 
 
@@ -5851,7 +5854,7 @@ private:
 			const TArray<uint8> GetCodecSpecificDataRAW() const override;
 			const FStreamCodecInformation GetCodecInformation() const override;
 			const FBitrateInfo GetBitrateInfo() const override;
-			const FString GetLanguage() const override;
+			const FString& GetLanguage() const override;
 			bool GetEncryptionInfo(FEncryptionInfo& OutEncryptionInfo) const override;
 			void GetPSSHBoxes(TArray<TArray<uint8>>& OutBoxes, bool bFromMOOV, bool bFromMOOF) const override;
 			void GetPRFTBoxes(TArray<ITrack::FProducerReferenceTime>& OutBoxes) const override;
@@ -6172,11 +6175,10 @@ private:
 		return BitrateInfo;
 	}
 
-	const FString FParserISO14496_12::FTrack::GetLanguage() const
+	const FString& FParserISO14496_12::FTrack::GetLanguage() const
 	{
-		FString Language = MDHDBox ? MDHDBox->GetLanguage() : FString(TEXT("und"));
-		// Try to map the ISO-639-2T language code to the shorter ISO-639-1 code if possible.
-		return ISO639::MapTo639_1(Language);
+		static FString Undefined(TEXT("und"));
+		return MDHDBox ? MDHDBox->GetLanguage() : Undefined;
 	}
 
 	bool FParserISO14496_12::FTrack::GetEncryptionInfo(FEncryptionInfo& OutEncryptionInfo) const
diff --git a/Engine/Plugins/Media/ElectraPlayer/Source/ElectraPlayerRuntime/Private/Runtime/Demuxer/ParserISO14496-12.h b/Engine/Plugins/Media/ElectraPlayer/Source/ElectraPlayerRuntime/Private/Runtime/Demuxer/ParserISO14496-12.h
--- a/Engine/Plugins/Media/ElectraPlayer/Source/ElectraPlayerRuntime/Private/Runtime/Demuxer/ParserISO14496-12.h
+++ b/Engine/Plugins/Media/ElectraPlayer/Source/ElectraPlayerRuntime/Private/Runtime/Demuxer/ParserISO14496-12.h
@@ -265,7 +265,7 @@ namespace Electra
 			virtual const TArray<uint8> GetCodecSpecificDataRAW() const = 0;
 			virtual const FStreamCodecInformation GetCodecInformation() const = 0;
 			virtual const FBitrateInfo GetBitrateInfo() const = 0;
-			virtual const FString GetLanguage() const = 0;
+			virtual const FString& GetLanguage() const = 0;
 			virtual bool GetEncryptionInfo(FEncryptionInfo& OutEncryptionInfo) const = 0;
 			virtual void GetPSSHBoxes(TArray<TArray<uint8>>& OutBoxes, bool bFromMOOV, bool bFromMOOF) const = 0;
 			virtual void GetPRFTBoxes(TArray<FProducerReferenceTime>& OutBoxes) const = 0;
diff --git a/Engine/Plugins/Media/ElectraPlayer/Source/ElectraPlayerRuntime/Private/Runtime/HTTP/HTTPManager.cpp b/Engine/Plugins/Media/ElectraPlayer/Source/ElectraPlayerRuntime/Private/Runtime/HTTP/HTTPManager.cpp
--- a/Engine/Plugins/Media/ElectraPlayer/Source/ElectraPlayerRuntime/Private/Runtime/HTTP/HTTPManager.cpp
+++ b/Engine/Plugins/Media/ElectraPlayer/Source/ElectraPlayerRuntime/Private/Runtime/HTTP/HTTPManager.cpp
@@ -1026,7 +1026,7 @@ namespace Electra
 
 		FParamDict HttpOptions;
 		#ifdef ELECTRA_HTTPMANAGER_DEBUG_PROXY_ADDRESS
-			HttpOptions.Set(FString(TEXT("proxy")), FVariantValue(FString(ELECTRA_HTTPMANAGER_DEBUG_PROXY_ADDRESS)));
+			HttpOptions.Set(FName(TEXT("proxy")), FVariantValue(FString(ELECTRA_HTTPMANAGER_DEBUG_PROXY_ADDRESS)));
 		#endif
 		HTTPStreamHandler = IElectraHTTPStream::Create(HttpOptions);
 		if (HTTPStreamHandler.IsValid())
diff --git a/Engine/Plugins/Media/ElectraPlayer/Source/ElectraPlayerRuntime/Private/Runtime/HTTP/HTTPResponseCache.cpp b/Engine/Plugins/Media/ElectraPlayer/Source/ElectraPlayerRuntime/Private/Runtime/HTTP/HTTPResponseCache.cpp
--- a/Engine/Plugins/Media/ElectraPlayer/Source/ElectraPlayerRuntime/Private/Runtime/HTTP/HTTPResponseCache.cpp
+++ b/Engine/Plugins/Media/ElectraPlayer/Source/ElectraPlayerRuntime/Private/Runtime/HTTP/HTTPResponseCache.cpp
@@ -14,7 +14,7 @@ namespace Electra
 class FHTTPResponseCache : public IHTTPResponseCache
 {
 public:
-	FHTTPResponseCache(IPlayerSessionServices* SessionServices, const FParamDict& Options, TSharedPtr<IElectraPlayerDataCache, ESPMode::ThreadSafe> ExternalCache);
+	FHTTPResponseCache(IPlayerSessionServices* SessionServices, TSharedPtr<IElectraPlayerDataCache, ESPMode::ThreadSafe> ExternalCache);
 	virtual ~FHTTPResponseCache();
 
 	void Disable() override;
@@ -359,30 +359,30 @@ private:
 /***************************************************************************************************************************************************/
 /***************************************************************************************************************************************************/
 
-TSharedPtrTS<IHTTPResponseCache> IHTTPResponseCache::Create(IPlayerSessionServices* SessionServices, const FParamDict& Options, TSharedPtr<IElectraPlayerDataCache, ESPMode::ThreadSafe> ExternalCache)
+TSharedPtrTS<IHTTPResponseCache> IHTTPResponseCache::Create(IPlayerSessionServices* SessionServices, TSharedPtr<IElectraPlayerDataCache, ESPMode::ThreadSafe> ExternalCache)
 {
-	return MakeSharedTS<FHTTPResponseCache>(SessionServices, Options, ExternalCache);
+	return MakeSharedTS<FHTTPResponseCache>(SessionServices, ExternalCache);
 }
 
 /***************************************************************************************************************************************************/
 /***************************************************************************************************************************************************/
 /***************************************************************************************************************************************************/
 
-FHTTPResponseCache::FHTTPResponseCache(IPlayerSessionServices* InSessionServices, const FParamDict& InOptions, TSharedPtr<IElectraPlayerDataCache, ESPMode::ThreadSafe> InExternalCache)
+FHTTPResponseCache::FHTTPResponseCache(IPlayerSessionServices* InSessionServices, TSharedPtr<IElectraPlayerDataCache, ESPMode::ThreadSafe> InExternalCache)
 	: ExternalCache(InExternalCache), SessionServices(InSessionServices)
 {
 	// Only configure our cache when there is no external one!
 	if (!ExternalCache.IsValid())
 	{
 		int64 v;
-		v = InOptions.GetValue(OptionKeyResponseCacheMaxByteSize).SafeGetInt64(0);
+		v = InSessionServices->GetOptionValue(OptionKeyResponseCacheMaxByteSize).SafeGetInt64(0);
 		if (v >= 0)
 		{
 			MaxElementSize = (int64)v;
 		}
 		// Max number of elements is probably not used a lot, if at all.
 		// If not specified we allow for some reasonable number.
-		v = InOptions.GetValue(OptionKeyResponseCacheMaxEntries).SafeGetInt64(8192);
+		v = InSessionServices->GetOptionValue(OptionKeyResponseCacheMaxEntries).SafeGetInt64(8192);
 		if (v >= 0)
 		{
 			MaxNumElements = (int32) v;
diff --git a/Engine/Plugins/Media/ElectraPlayer/Source/ElectraPlayerRuntime/Private/Runtime/HTTP/HTTPResponseCache.h b/Engine/Plugins/Media/ElectraPlayer/Source/ElectraPlayerRuntime/Private/Runtime/HTTP/HTTPResponseCache.h
--- a/Engine/Plugins/Media/ElectraPlayer/Source/ElectraPlayerRuntime/Private/Runtime/HTTP/HTTPResponseCache.h
+++ b/Engine/Plugins/Media/ElectraPlayer/Source/ElectraPlayerRuntime/Private/Runtime/HTTP/HTTPResponseCache.h
@@ -33,7 +33,7 @@ public:
 		FQualityInfo Quality;
 	};
 
-	static TSharedPtrTS<IHTTPResponseCache> Create(IPlayerSessionServices* SessionServices, const FParamDict& Options, TSharedPtr<IElectraPlayerDataCache, ESPMode::ThreadSafe> ExternalCache);
+	static TSharedPtrTS<IHTTPResponseCache> Create(IPlayerSessionServices* SessionServices, TSharedPtr<IElectraPlayerDataCache, ESPMode::ThreadSafe> ExternalCache);
 
 	virtual ~IHTTPResponseCache() = default;
 
diff --git a/Engine/Plugins/Media/ElectraPlayer/Source/ElectraPlayerRuntime/Private/Runtime/Player/ABRRules/ABRFixedStream.cpp b/Engine/Plugins/Media/ElectraPlayer/Source/ElectraPlayerRuntime/Private/Runtime/Player/ABRRules/ABRFixedStream.cpp
--- a/Engine/Plugins/Media/ElectraPlayer/Source/ElectraPlayerRuntime/Private/Runtime/Player/ABRRules/ABRFixedStream.cpp
+++ b/Engine/Plugins/Media/ElectraPlayer/Source/ElectraPlayerRuntime/Private/Runtime/Player/ABRRules/ABRFixedStream.cpp
@@ -27,7 +27,7 @@ public:
 
 	FTimeValue GetMinBufferTimeForPlayback(IAdaptiveStreamSelector::EMinBufferType InBufferingType, FTimeValue InDefaultMBT) override
 	{ return FTimeValue(); }
-	IAdaptiveStreamSelector::FRebufferAction GetRebufferAction(const FParamDict& CurrentPlayerOptions) override;
+	IAdaptiveStreamSelector::FRebufferAction GetRebufferAction() override;
 	IAdaptiveStreamSelector::EHandlingAction PeriodicHandle() override { return IAdaptiveStreamSelector::EHandlingAction::None; }
 	void DebugPrint(void* pThat, void (*pPrintFN)(void* pThat, const char *pFmt, ...)) override { }
 
@@ -130,14 +130,14 @@ double FABRFixedStream::GetAverageLatency()
 }
 
 
-IAdaptiveStreamSelector::FRebufferAction FABRFixedStream::GetRebufferAction(const FParamDict& CurrentPlayerOptions)
+IAdaptiveStreamSelector::FRebufferAction FABRFixedStream::GetRebufferAction()
 {
 	IAdaptiveStreamSelector::FRebufferAction Action;
-	if (CurrentPlayerOptions.GetValue(OptionThrowErrorWhenRebuffering).SafeGetBool(false))
+	if (Info->GetOptionValue(OptionThrowErrorWhenRebuffering).SafeGetBool(false))
 	{
 		Action.Action = IAdaptiveStreamSelector::FRebufferAction::EAction::ThrowError;
 	}
-	else if (CurrentPlayerOptions.GetValue(OptionRebufferingContinuesLoading).SafeGetBool(false))
+	else if (Info->GetOptionValue(OptionRebufferingContinuesLoading).SafeGetBool(false))
 	{
 		Action.Action = IAdaptiveStreamSelector::FRebufferAction::EAction::ContinueLoading;
 	}
diff --git a/Engine/Plugins/Media/ElectraPlayer/Source/ElectraPlayerRuntime/Private/Runtime/Player/ABRRules/ABRInternal.h b/Engine/Plugins/Media/ElectraPlayer/Source/ElectraPlayerRuntime/Private/Runtime/Player/ABRRules/ABRInternal.h
--- a/Engine/Plugins/Media/ElectraPlayer/Source/ElectraPlayerRuntime/Private/Runtime/Player/ABRRules/ABRInternal.h
+++ b/Engine/Plugins/Media/ElectraPlayer/Source/ElectraPlayerRuntime/Private/Runtime/Player/ABRRules/ABRInternal.h
@@ -45,7 +45,8 @@ namespace Electra
 	{
 	public:
 		virtual ~IABRInfoInterface() = default;
-		virtual FParamDict& GetPlayerOptions() = 0;
+		virtual bool HaveOptionValue(const FName& InOption) = 0;
+		virtual const FVariantValue GetOptionValue(const FName& InOption) = 0;
 
 		virtual TSharedPtrTS<FABRStreamInformation> GetStreamInformation(const Metrics::FSegmentDownloadStats& FromDownloadStats) = 0;
 		virtual const TArray<TSharedPtrTS<FABRStreamInformation>>& GetStreamInformations(EStreamType InForStreamType) = 0;
@@ -71,7 +72,7 @@ namespace Electra
 		virtual void ReportPlaybackEnded() = 0;
 
 		virtual FTimeValue GetMinBufferTimeForPlayback(IAdaptiveStreamSelector::EMinBufferType InBufferingType, FTimeValue InDefaultMBT) = 0;
-		virtual IAdaptiveStreamSelector::FRebufferAction GetRebufferAction(const FParamDict& CurrentPlayerOptions) = 0;
+		virtual IAdaptiveStreamSelector::FRebufferAction GetRebufferAction() = 0;
 		virtual IAdaptiveStreamSelector::EHandlingAction PeriodicHandle() = 0;
 		virtual void DebugPrint(void* pThat, void (*pPrintFN)(void* pThat, const char *pFmt, ...)) = 0;
 
diff --git a/Engine/Plugins/Media/ElectraPlayer/Source/ElectraPlayerRuntime/Private/Runtime/Player/ABRRules/ABRLiveStream.cpp b/Engine/Plugins/Media/ElectraPlayer/Source/ElectraPlayerRuntime/Private/Runtime/Player/ABRRules/ABRLiveStream.cpp
--- a/Engine/Plugins/Media/ElectraPlayer/Source/ElectraPlayerRuntime/Private/Runtime/Player/ABRRules/ABRLiveStream.cpp
+++ b/Engine/Plugins/Media/ElectraPlayer/Source/ElectraPlayerRuntime/Private/Runtime/Player/ABRRules/ABRLiveStream.cpp
@@ -40,7 +40,7 @@ public:
 	void ReportPlaybackEnded() override;
 
 	FTimeValue GetMinBufferTimeForPlayback(IAdaptiveStreamSelector::EMinBufferType InBufferingType, FTimeValue InDefaultMBT) override;
-	IAdaptiveStreamSelector::FRebufferAction GetRebufferAction(const FParamDict& CurrentPlayerOptions) override;
+	IAdaptiveStreamSelector::FRebufferAction GetRebufferAction() override;
 	IAdaptiveStreamSelector::EHandlingAction PeriodicHandle() override;
 	void DebugPrint(void* pThat, void (*pPrintFN)(void* pThat, const char *pFmt, ...)) override;
 
@@ -396,10 +396,10 @@ FABRLiveStream::FABRLiveStream(IABRInfoInterface* InInfo, EMediaFormatType InFor
 	, FormatType(InFormatType)
 	, PresentationType(InPresentationType)
 {
-	bRebufferingJustResumesLoading = Info->GetPlayerOptions().GetValue(OptionRebufferingContinuesLoading).SafeGetBool(false);
-	if (Info->GetPlayerOptions().HaveKey(ABR::OptionKeyABR_CDNSegmentDenyHTTPStatus))
+	bRebufferingJustResumesLoading = Info->GetOptionValue(OptionRebufferingContinuesLoading).SafeGetBool(false);
+	if (Info->HaveOptionValue(ABR::OptionKeyABR_CDNSegmentDenyHTTPStatus))
 	{
-		HTTPStatusCodeToDenyStream.Set((int32) Info->GetPlayerOptions().GetValue(ABR::OptionKeyABR_CDNSegmentDenyHTTPStatus).SafeGetInt64(-1));
+		HTTPStatusCodeToDenyStream.Set((int32) Info->GetOptionValue(ABR::OptionKeyABR_CDNSegmentDenyHTTPStatus).SafeGetInt64(-1));
 	}
 	PrepareLatencyConfiguration();
 
@@ -1600,7 +1600,7 @@ FTimeValue FABRLiveStream::GetMinBufferTimeForPlayback(IAdaptiveStreamSelector::
 	return FTimeValue();
 }
 
-IAdaptiveStreamSelector::FRebufferAction FABRLiveStream::GetRebufferAction(const FParamDict& CurrentPlayerOptions)
+IAdaptiveStreamSelector::FRebufferAction FABRLiveStream::GetRebufferAction()
 {
 	IAdaptiveStreamSelector::FRebufferAction Action;
 
diff --git a/Engine/Plugins/Media/ElectraPlayer/Source/ElectraPlayerRuntime/Private/Runtime/Player/ABRRules/ABRVoDPlus.cpp b/Engine/Plugins/Media/ElectraPlayer/Source/ElectraPlayerRuntime/Private/Runtime/Player/ABRRules/ABRVoDPlus.cpp
--- a/Engine/Plugins/Media/ElectraPlayer/Source/ElectraPlayerRuntime/Private/Runtime/Player/ABRRules/ABRVoDPlus.cpp
+++ b/Engine/Plugins/Media/ElectraPlayer/Source/ElectraPlayerRuntime/Private/Runtime/Player/ABRRules/ABRVoDPlus.cpp
@@ -42,7 +42,7 @@ public:
 	void ReportPlaybackEnded() override;
 
 	FTimeValue GetMinBufferTimeForPlayback(IAdaptiveStreamSelector::EMinBufferType InBufferingType, FTimeValue InDefaultMBT) override;
-	IAdaptiveStreamSelector::FRebufferAction GetRebufferAction(const FParamDict& CurrentPlayerOptions) override;
+	IAdaptiveStreamSelector::FRebufferAction GetRebufferAction() override;
 	IAdaptiveStreamSelector::EHandlingAction PeriodicHandle() override;
 	void DebugPrint(void* pThat, void (*pPrintFN)(void* pThat, const char *pFmt, ...)) override;
 
@@ -272,10 +272,10 @@ FABROnDemandPlus::FABROnDemandPlus(IABRInfoInterface* InInfo, EMediaFormatType I
 	, FormatType(InFormatType)
 	, PresentationType(InPresentationType)
 {
-	bRebufferingJustResumesLoading = Info->GetPlayerOptions().GetValue(OptionRebufferingContinuesLoading).SafeGetBool(false);
-	if (Info->GetPlayerOptions().HaveKey(ABR::OptionKeyABR_CDNSegmentDenyHTTPStatus))
+	bRebufferingJustResumesLoading = Info->GetOptionValue(OptionRebufferingContinuesLoading).SafeGetBool(false);
+	if (Info->HaveOptionValue(ABR::OptionKeyABR_CDNSegmentDenyHTTPStatus))
 	{
-		HTTPStatusCodeToDenyStream.Set((int32) Info->GetPlayerOptions().GetValue(ABR::OptionKeyABR_CDNSegmentDenyHTTPStatus).SafeGetInt64(-1));
+		HTTPStatusCodeToDenyStream.Set((int32) Info->GetOptionValue(ABR::OptionKeyABR_CDNSegmentDenyHTTPStatus).SafeGetInt64(-1));
 	}
 
 	FTimeRange RenderRange = Info->ABRGetSupportedRenderRateScale();
@@ -978,7 +978,7 @@ FTimeValue FABROnDemandPlus::GetMinBufferTimeForPlayback(IAdaptiveStreamSelector
 	return FTimeValue();
 }
 
-IAdaptiveStreamSelector::FRebufferAction FABROnDemandPlus::GetRebufferAction(const FParamDict& CurrentPlayerOptions)
+IAdaptiveStreamSelector::FRebufferAction FABROnDemandPlus::GetRebufferAction()
 {
 	IAdaptiveStreamSelector::FRebufferAction Action;
 	Action.Action = bRebufferingJustResumesLoading ? IAdaptiveStreamSelector::FRebufferAction::EAction::ContinueLoading : IAdaptiveStreamSelector::FRebufferAction::EAction::Restart;
diff --git a/Engine/Plugins/Media/ElectraPlayer/Source/ElectraPlayerRuntime/Private/Runtime/Player/AdaptivePlayerOptionKeynames.h b/Engine/Plugins/Media/ElectraPlayer/Source/ElectraPlayerRuntime/Private/Runtime/Player/AdaptivePlayerOptionKeynames.h
--- a/Engine/Plugins/Media/ElectraPlayer/Source/ElectraPlayerRuntime/Private/Runtime/Player/AdaptivePlayerOptionKeynames.h
+++ b/Engine/Plugins/Media/ElectraPlayer/Source/ElectraPlayerRuntime/Private/Runtime/Player/AdaptivePlayerOptionKeynames.h
@@ -46,6 +46,14 @@ namespace Electra
 
 	const FName OptionKeyCurrentAvgStartingVideoBitrate(TEXT("current:avg_video_bitrate"));
 
+	const FName OptionKeyExcludedCodecsVideo(TEXT("excluded_codecs_video"));
+	const FName OptionKeyExcludedCodecsAudio(TEXT("excluded_codecs_audio"));
+	const FName OptionKeyExcludedCodecsSubtitles(TEXT("excluded_codecs_subtitles"));
+	
+	const FName OptionKeyPreferredCodecsVideo(TEXT("preferred_codecs_video"));
+	const FName OptionKeyPreferredCodecsAudio(TEXT("preferred_codecs_audio"));
+	const FName OptionKeyPreferredCodecsSubtitles(TEXT("preferred_codecs_subtitles"));
+
 } // namespace Electra
 
 
diff --git a/Engine/Plugins/Media/ElectraPlayer/Source/ElectraPlayerRuntime/Private/Runtime/Player/AdaptiveStreamingPlayer.cpp b/Engine/Plugins/Media/ElectraPlayer/Source/ElectraPlayerRuntime/Private/Runtime/Player/AdaptiveStreamingPlayer.cpp
--- a/Engine/Plugins/Media/ElectraPlayer/Source/ElectraPlayerRuntime/Private/Runtime/Player/AdaptiveStreamingPlayer.cpp
+++ b/Engine/Plugins/Media/ElectraPlayer/Source/ElectraPlayerRuntime/Private/Runtime/Player/AdaptiveStreamingPlayer.cpp
@@ -281,20 +281,7 @@ void FAdaptiveStreamingPlayer::Initialize(const FParamDict& Options)
 
 void FAdaptiveStreamingPlayer::ModifyOptions(const FParamDict& InOptionsToSetOrChange, const FParamDict& InOptionsToClear)
 {
-	// Remove options
-	TArray<FName> Keys;
-	InOptionsToClear.GetKeysStartingWith(FString(), Keys);
-	for(auto &Key : Keys)
-	{
-		PlayerOptions.Remove(Key);
-	}
-
-	InOptionsToSetOrChange.GetKeysStartingWith(FString(), Keys);
-	for(auto &Key : Keys)
-	{
-		FVariantValue Value = InOptionsToSetOrChange.GetValue(Key);
-		PlayerOptions.Set(Key, Value);
-	}
+	WorkerThread.SendOptionChangeMessage(InOptionsToSetOrChange, InOptionsToClear);
 }
 
 
@@ -552,7 +539,7 @@ void FAdaptiveStreamingPlayer::StopWorkerThread()
  */
 void FAdaptiveStreamingPlayer::EnableFrameAccurateSeeking(bool bEnabled)
 {
-	GetOptions().Set(OptionKeyFrameAccurateSeek, FVariantValue(bEnabled));
+	PlayerOptions.Set(OptionKeyFrameAccurateSeek, FVariantValue(bEnabled));
 }
 
 
@@ -1151,11 +1138,18 @@ IAdaptiveStreamingPlayerAEMSHandler* FAdaptiveStreamingPlayer::GetAEMSEventHandl
 {
 	return AEMSEventHandler;
 }
-
-FParamDict& FAdaptiveStreamingPlayer::GetOptions()
+FParamDictTS& FAdaptiveStreamingPlayer::GetMutableOptions()
 {
 	return PlayerOptions;
 }
+bool FAdaptiveStreamingPlayer::HaveOptionValue(const FName& InOption)
+{
+	return PlayerOptions.HaveKey(InOption);
+}
+const FVariantValue FAdaptiveStreamingPlayer::GetOptionValue(const FName& InOption)
+{
+	return PlayerOptions.GetValue(InOption);
+}
 
 TSharedPtrTS<FDRMManager> FAdaptiveStreamingPlayer::GetDRMManager()
 {
@@ -1358,9 +1352,13 @@ bool FAdaptiveStreamingPlayer::InternalHandleThreadMessages()
 		// While closed ignore all messages.
 		if (bIsClosing)
 		{
-			if (msg.Data.MediaEvent.Event)
+			if (msg.Data.IsType<FWorkerThreadMessages::FMessage::FEvent>())
 			{
-				msg.Data.MediaEvent.Event->Signal();
+				FWorkerThreadMessages::FMessage::FEvent& ev = msg.Data.Get<FWorkerThreadMessages::FMessage::FEvent>();
+				if (ev.Event)
+				{
+					ev.Event->Signal();
+				}
 			}
 			continue;
 		}
@@ -1374,23 +1372,42 @@ bool FAdaptiveStreamingPlayer::InternalHandleThreadMessages()
 				FAdaptiveStreamingPlayer::InternalInitialize();
 				break;
 			}
+			case FWorkerThreadMessages::FMessage::EType::ChangeOptions:
+			{
+				FWorkerThreadMessages::FMessage::FOptionChange& ev = msg.Data.Get<FWorkerThreadMessages::FMessage::FOptionChange>();
+				TArray<FName> Keys;
+				ev.OptionsToClear.GetKeys(Keys);
+				for(auto &Key : Keys)
+				{
+					PlayerOptions.Remove(Key);
+				}
+				ev.OptionsToSetOrChange.GetKeys(Keys);
+				for(auto &Key : Keys)
+				{
+					FVariantValue Value = ev.OptionsToSetOrChange.GetValue(Key);
+					PlayerOptions.Set(Key, Value);
+				}
+				break;
+			}
 			case FWorkerThreadMessages::FMessage::EType::LoadManifest:
 			{
 				// Check if a new or different cache has been set prior to loading the manifest.
 				if (ExternalCache.IsValid())
 				{
 					HttpResponseCache.Reset();
-					HttpResponseCache = IHTTPResponseCache::Create(this, PlayerOptions, MoveTemp(ExternalCache));
+					HttpResponseCache = IHTTPResponseCache::Create(this, MoveTemp(ExternalCache));
 				}
 
-				InternalLoadManifest(msg.Data.ManifestToLoad.URL, msg.Data.ManifestToLoad.MimeType);
+				FWorkerThreadMessages::FMessage::FLoadManifest& ev = msg.Data.Get<FWorkerThreadMessages::FMessage::FLoadManifest>();
+				InternalLoadManifest(ev.URL, ev.MimeType);
 				break;
 			}
 			case FWorkerThreadMessages::FMessage::EType::LoadBlob:
 			{
-				if (msg.Data.BlobToLoad.BlobLoadRequest.IsValid())
+				FWorkerThreadMessages::FMessage::FLoadBlob& ev = msg.Data.Get<FWorkerThreadMessages::FMessage::FLoadBlob>();
+				if (ev.BlobLoadRequest.IsValid())
 				{
-					msg.Data.BlobToLoad.BlobLoadRequest->StartGet(this);
+					ev.BlobLoadRequest->StartGet(this);
 				}
 				break;
 			}
@@ -1408,11 +1425,13 @@ bool FAdaptiveStreamingPlayer::InternalHandleThreadMessages()
 			}
 			case FWorkerThreadMessages::FMessage::EType::Loop:
 			{
-				InternalSetLoop(msg.Data.Looping.Loop);
-				if (msg.Data.Looping.Signal)
+				FWorkerThreadMessages::FMessage::FLoop& ev = msg.Data.Get<FWorkerThreadMessages::FMessage::FLoop>();
+				InternalSetLoop(ev.Loop);
+				/*
+				if (ev.Signal)
 				{
-					msg.Data.Looping.Signal->Signal();
-				}
+					ev.Signal->Signal();
+				*/
 				break;
 			}
 			case FWorkerThreadMessages::FMessage::EType::Close:
@@ -1424,57 +1443,61 @@ bool FAdaptiveStreamingPlayer::InternalHandleThreadMessages()
 					InternalClose();
 					DispatchEvent(FMetricEvent::ReportPlaybackStopped());
 				}
-				if (msg.Data.MediaEvent.Event)
+				FWorkerThreadMessages::FMessage::FEvent& ev = msg.Data.Get<FWorkerThreadMessages::FMessage::FEvent>();
+				if (ev.Event)
 				{
-					msg.Data.MediaEvent.Event->Signal();
+					ev.Event->Signal();
 				}
 				break;
 			}
 
 			case FWorkerThreadMessages::FMessage::EType::ChangeBitrate:
 			{
-				BitrateCeiling = msg.Data.Bitrate.Value;
+				FWorkerThreadMessages::FMessage::FBitrate& ev = msg.Data.Get<FWorkerThreadMessages::FMessage::FBitrate>();
+				BitrateCeiling = ev.Value;
 				StreamSelector->SetBandwidthCeiling(BitrateCeiling);
 				break;
 			}
 
 			case FWorkerThreadMessages::FMessage::EType::LimitResolution:
 			{
-				VideoResolutionLimitWidth  = msg.Data.Resolution.Width;
-				VideoResolutionLimitHeight = msg.Data.Resolution.Height;
+				FWorkerThreadMessages::FMessage::FResolution& ev = msg.Data.Get<FWorkerThreadMessages::FMessage::FResolution>();
+				VideoResolutionLimitWidth  = ev.Width;
+				VideoResolutionLimitHeight = ev.Height;
 				UpdateStreamResolutionLimit();
 				break;
 			}
 
 			case FWorkerThreadMessages::FMessage::EType::InitialStreamAttributes:
 			{
+				FWorkerThreadMessages::FMessage::FInitialStreamSelect& ev = msg.Data.Get<FWorkerThreadMessages::FMessage::FInitialStreamSelect>();
 				// Map the language in case it is not yet.
-				if (msg.Data.InitialStreamAttribute.InitialSelection.Language_ISO639.IsSet())
+				if (ev.InitialSelection.Language_ISO639.IsSet())
 				{
-					msg.Data.InitialStreamAttribute.InitialSelection.Language_ISO639 = ISO639::MapTo639_1(msg.Data.InitialStreamAttribute.InitialSelection.Language_ISO639.GetValue());
+					ev.InitialSelection.Language_ISO639 = ISO639::MapTo639_1(ev.InitialSelection.Language_ISO639.GetValue());
 				}
-				bool bInitiallyDisabled = msg.Data.InitialStreamAttribute.InitialSelection.OverrideIndex.IsSet() && msg.Data.InitialStreamAttribute.InitialSelection.OverrideIndex.GetValue() < 0;
+				bool bInitiallyDisabled = ev.InitialSelection.OverrideIndex.IsSet() && ev.InitialSelection.OverrideIndex.GetValue() < 0;
 				if (bInitiallyDisabled)
 				{
-					msg.Data.InitialStreamAttribute.InitialSelection.ClearOverrideIndex();
+					ev.InitialSelection.ClearOverrideIndex();
 				}
-				switch(msg.Data.InitialStreamAttribute.StreamType)
+				switch(ev.StreamType)
 				{
 					case EStreamType::Video:
-						StreamSelectionAttributesVid = msg.Data.InitialStreamAttribute.InitialSelection;
+						StreamSelectionAttributesVid = ev.InitialSelection;
 						break;
 					case EStreamType::Audio:
-						StreamSelectionAttributesAud = msg.Data.InitialStreamAttribute.InitialSelection;
+						StreamSelectionAttributesAud = ev.InitialSelection;
 						break;
 					case EStreamType::Subtitle:
-						StreamSelectionAttributesTxt = msg.Data.InitialStreamAttribute.InitialSelection;
+						StreamSelectionAttributesTxt = ev.InitialSelection;
 						break;
 					default:
 						break;
 				}
 				if (bInitiallyDisabled)
 				{
-					InternalDeselectStream(msg.Data.InitialStreamAttribute.StreamType);
+					InternalDeselectStream(ev.StreamType);
 				}
 				break;
 			}
@@ -1483,23 +1506,24 @@ bool FAdaptiveStreamingPlayer::InternalHandleThreadMessages()
 			{
 				if (!bIsClosing)
 				{
+					FWorkerThreadMessages::FMessage::FMetadataTrackSelection& ev = msg.Data.Get<FWorkerThreadMessages::FMessage::FMetadataTrackSelection>();
 					// Map the language in case it is not yet.
-					if (msg.Data.TrackSelection.TrackAttributes.Language_ISO639.IsSet())
+					if (ev.TrackAttributes.Language_ISO639.IsSet())
 					{
-						msg.Data.TrackSelection.TrackAttributes.Language_ISO639 = ISO639::MapTo639_1(msg.Data.TrackSelection.TrackAttributes.Language_ISO639.GetValue());
+						ev.TrackAttributes.Language_ISO639 = ISO639::MapTo639_1(ev.TrackAttributes.Language_ISO639.GetValue());
 					}
-					switch(msg.Data.TrackSelection.StreamType)
+					switch(ev.StreamType)
 					{
 						case EStreamType::Video:
-							PendingTrackSelectionVid = MakeSharedTS<FStreamSelectionAttributes>(msg.Data.TrackSelection.TrackAttributes);
+							PendingTrackSelectionVid = MakeSharedTS<FStreamSelectionAttributes>(ev.TrackAttributes);
 							bIsVideoDeselected = false;
 							break;
 						case EStreamType::Audio:
-							PendingTrackSelectionAud = MakeSharedTS<FStreamSelectionAttributes>(msg.Data.TrackSelection.TrackAttributes);
+							PendingTrackSelectionAud = MakeSharedTS<FStreamSelectionAttributes>(ev.TrackAttributes);
 							bIsAudioDeselected = false;
 							break;
 						case EStreamType::Subtitle:
-							PendingTrackSelectionTxt = MakeSharedTS<FStreamSelectionAttributes>(msg.Data.TrackSelection.TrackAttributes);
+							PendingTrackSelectionTxt = MakeSharedTS<FStreamSelectionAttributes>(ev.TrackAttributes);
 							bIsTextDeselected = false;
 							break;
 						default:
@@ -1517,21 +1541,23 @@ bool FAdaptiveStreamingPlayer::InternalHandleThreadMessages()
 
 			case FWorkerThreadMessages::FMessage::EType::DeselectTrack:
 			{
-				InternalDeselectStream(msg.Data.TrackSelection.StreamType);
+				FWorkerThreadMessages::FMessage::FMetadataTrackSelection& ev = msg.Data.Get<FWorkerThreadMessages::FMessage::FMetadataTrackSelection>();
+				InternalDeselectStream(ev.StreamType);
 				break;
 			}
 
 			case FWorkerThreadMessages::FMessage::EType::FragmentOpen:
 			{
-				TSharedPtrTS<IStreamSegment> pRequest = msg.Data.StreamReader.Request;
-				EStreamType reqType = msg.Data.StreamReader.Request->GetType();
+				FWorkerThreadMessages::FMessage::FStreamReader& ev = msg.Data.Get<FWorkerThreadMessages::FMessage::FStreamReader>();
+				TSharedPtrTS<IStreamSegment> pRequest = ev.Request;
+				EStreamType reqType = ev.Request->GetType();
 				// Check that the request is for this current playback sequence and not an outdated one.
 				if (pRequest.IsValid() && pRequest->GetPlaybackSequenceID() == CurrentPlaybackSequenceID[StreamTypeToArrayIndex(reqType)])
 				{
-					DispatchBufferUtilizationEvent(msg.Data.StreamReader.Request->GetType());
+					DispatchBufferUtilizationEvent(ev.Request->GetType());
 
 					// Video bitrate change?
-					if (msg.Data.StreamReader.Request->GetType() == EStreamType::Video)
+					if (ev.Request->GetType() == EStreamType::Video)
 					{
 						int32 SegmentBitrate = pRequest->GetBitrate();
 						int32 SegmentQualityLevel = pRequest->GetQualityIndex();
@@ -1549,8 +1575,9 @@ bool FAdaptiveStreamingPlayer::InternalHandleThreadMessages()
 
 			case FWorkerThreadMessages::FMessage::EType::FragmentClose:
 			{
-				TSharedPtrTS<IStreamSegment> pRequest = msg.Data.StreamReader.Request;
-				EStreamType reqType = msg.Data.StreamReader.Request->GetType();
+				FWorkerThreadMessages::FMessage::FStreamReader& ev = msg.Data.Get<FWorkerThreadMessages::FMessage::FStreamReader>();
+				TSharedPtrTS<IStreamSegment> pRequest = ev.Request;
+				EStreamType reqType = ev.Request->GetType();
 				// Check that the request is for this current playback sequence and not an outdated one.
 				if (pRequest.IsValid() && pRequest->GetPlaybackSequenceID() == CurrentPlaybackSequenceID[StreamTypeToArrayIndex(reqType)])
 				{
@@ -1582,15 +1609,17 @@ bool FAdaptiveStreamingPlayer::InternalHandleThreadMessages()
 
 			case FWorkerThreadMessages::FMessage::EType::PlayerSession:
 			{
-				HandleSessionMessage(msg.Data.Session.PlayerMessage);
+				FWorkerThreadMessages::FMessage::FSession& ev = msg.Data.Get<FWorkerThreadMessages::FMessage::FSession>();
+				HandleSessionMessage(ev.PlayerMessage);
 				break;
 			}
 
 			case FWorkerThreadMessages::FMessage::EType::EndPlaybackAt:
 			{
+				FWorkerThreadMessages::FMessage::FEndPlaybackAt& ev = msg.Data.Get<FWorkerThreadMessages::FMessage::FEndPlaybackAt>();
 				// With a defined end time we just play out until then.
 				PlaybackState.SetShouldPlayOnLiveEdge(false);
-				PlaybackState.SetPlaybackEndAtTime(msg.Data.EndPlaybackAt.EndAtTime);
+				PlaybackState.SetPlaybackEndAtTime(ev.EndAtTime);
 				break;
 			}
 
@@ -3855,7 +3884,7 @@ void FAdaptiveStreamingPlayer::SetPlaystartOptions(FPlayStartOptions& OutOptions
 	OutOptions.PlaybackRange.Start = ActiveRange.Start.IsValid() ? ActiveRange.Start : FTimeValue::GetZero();
 	OutOptions.PlaybackRange.End = ActiveRange.End.IsValid() ? ActiveRange.End : FTimeValue::GetPositiveInfinity();
 
-	OutOptions.bFrameAccuracy = GetOptions().GetValue(OptionKeyFrameAccurateSeek).SafeGetBool(false);
+	OutOptions.bFrameAccuracy = PlayerOptions.GetValue(OptionKeyFrameAccurateSeek).SafeGetBool(false);
 }
 
 void FAdaptiveStreamingPlayer::ClampStartRequestTime(FTimeValue& InOutTimeToClamp)
@@ -4429,10 +4458,10 @@ void FAdaptiveStreamingPlayer::InternalInitialize()
 	DrmManager = FDRMManager::Create(this);
 
 	// Create an entity cache.
-	EntityCache = IPlayerEntityCache::Create(this, PlayerOptions);
+	EntityCache = IPlayerEntityCache::Create(this);
 
 	// Create an HTTP response cache. Hand over any externally set cache. We do not need it in here any further.
-	HttpResponseCache = IHTTPResponseCache::Create(this, PlayerOptions, MoveTemp(ExternalCache));
+	HttpResponseCache = IHTTPResponseCache::Create(this, MoveTemp(ExternalCache));
 
 	// If all read requests for this player are to be routed to an external reader we need to create a wrapper for it.
 	if (PlayerOptions.GetValue(OptionKeyUseExternalDataReader).SafeGetBool(false))
@@ -4448,7 +4477,7 @@ void FAdaptiveStreamingPlayer::InternalInitialize()
 	CreateRenderers();
 
 	// Check for codecs that are not to be used as per the user's choice, even if the device supports them.
-	auto GetExcludedCodecPrefixes = [](TArray<FString>& OutList, const FParamDict& InOptions, const FName& InKey) -> void
+	auto GetExcludedCodecPrefixes = [](TArray<FString>& OutList, const FParamDictTS& InOptions, const FName& InKey) -> void
 	{
 		if (InOptions.HaveKey(InKey))
 		{
@@ -4463,11 +4492,11 @@ void FAdaptiveStreamingPlayer::InternalInitialize()
 			}
 		}
 	};
-	GetExcludedCodecPrefixes(ExcludedVideoDecoderPrefixes, PlayerOptions, TEXT("excluded_codecs_video"));
-	GetExcludedCodecPrefixes(ExcludedAudioDecoderPrefixes, PlayerOptions, TEXT("excluded_codecs_audio"));
-	GetExcludedCodecPrefixes(ExcludedSubtitleDecoderPrefixes, PlayerOptions, TEXT("excluded_codecs_subtitles"));
+	GetExcludedCodecPrefixes(ExcludedVideoDecoderPrefixes, PlayerOptions, OptionKeyExcludedCodecsVideo);
+	GetExcludedCodecPrefixes(ExcludedAudioDecoderPrefixes, PlayerOptions, OptionKeyExcludedCodecsAudio);
+	GetExcludedCodecPrefixes(ExcludedSubtitleDecoderPrefixes, PlayerOptions, OptionKeyExcludedCodecsSubtitles);
 
-	auto GetCodecSelectionPriorities = [this](FCodecSelectionPriorities& OutPriorities, const FParamDict& InOptions, const FName& InKey, const TCHAR* const InType) -> void
+	auto GetCodecSelectionPriorities = [this](FCodecSelectionPriorities& OutPriorities, const FParamDictTS& InOptions, const FName& InKey, const TCHAR* const InType) -> void
 	{
 		if (InOptions.HaveKey(InKey))
 		{
@@ -4477,10 +4506,9 @@ void FAdaptiveStreamingPlayer::InternalInitialize()
 			}
 		}
 	};
-	GetCodecSelectionPriorities(CodecPrioritiesVideo, PlayerOptions, TEXT("preferred_codecs_video"), TEXT("video"));
-	GetCodecSelectionPriorities(CodecPrioritiesAudio, PlayerOptions, TEXT("preferred_codecs_audio"), TEXT("audio"));
-	GetCodecSelectionPriorities(CodecPrioritiesSubtitles, PlayerOptions, TEXT("preferred_codecs_subtitles"), TEXT("subtitle"));
-
+	GetCodecSelectionPriorities(CodecPrioritiesVideo, PlayerOptions, OptionKeyPreferredCodecsVideo, TEXT("video"));
+	GetCodecSelectionPriorities(CodecPrioritiesAudio, PlayerOptions, OptionKeyPreferredCodecsAudio, TEXT("audio"));
+	GetCodecSelectionPriorities(CodecPrioritiesSubtitles, PlayerOptions, OptionKeyPreferredCodecsSubtitles, TEXT("subtitle"));
 
 	// Unless already specified in the options to either value, enable frame accurate seeking now.
 	if (!PlayerOptions.HaveKey(OptionKeyFrameAccurateSeek))
@@ -4611,7 +4639,7 @@ void FAdaptiveStreamingPlayer::InternalRebuffer()
 	check(StreamSelector.IsValid());
 	if (StreamSelector.IsValid())
 	{
-		Action = StreamSelector->GetRebufferAction(GetOptions());
+		Action = StreamSelector->GetRebufferAction();
 	}
 
 	// Check if we are configured to throw a playback error instead of doing a rebuffer.
diff --git a/Engine/Plugins/Media/ElectraPlayer/Source/ElectraPlayerRuntime/Private/Runtime/Player/AdaptiveStreamingPlayerABR.cpp b/Engine/Plugins/Media/ElectraPlayer/Source/ElectraPlayerRuntime/Private/Runtime/Player/AdaptiveStreamingPlayerABR.cpp
--- a/Engine/Plugins/Media/ElectraPlayer/Source/ElectraPlayerRuntime/Private/Runtime/Player/AdaptiveStreamingPlayerABR.cpp
+++ b/Engine/Plugins/Media/ElectraPlayer/Source/ElectraPlayerRuntime/Private/Runtime/Player/AdaptiveStreamingPlayerABR.cpp
@@ -31,7 +31,7 @@ namespace Electra
 		void SetBandwidth(int64 bitsPerSecond) override;
 		void SetForcedNextBandwidth(int64 bitsPerSecond, double minBufferTimeBeforePlayback) override;
 		FTimeValue GetMinBufferTimeForPlayback(EMinBufferType InBufferingType, FTimeValue InDefaultMBT) override;
-		FRebufferAction GetRebufferAction(const FParamDict& CurrentPlayerOptions) override;
+		FRebufferAction GetRebufferAction() override;
 		EHandlingAction PeriodicHandle() override;
 		void MarkStreamAsUnavailable(const FDenylistedStream& DenylistedStream) override;
 		void MarkStreamAsAvailable(const FDenylistedStream& NoLongerDenylistedStream) override;
@@ -77,8 +77,10 @@ namespace Electra
 		void ReportDroppedAudioFrame() override {}
 
 	private:
-		FParamDict& GetPlayerOptions() override
-		{ return PlayerSessionServices->GetOptions(); }
+		bool HaveOptionValue(const FName& InOption) override
+		{ return PlayerSessionServices->HaveOptionValue(InOption); }
+		const FVariantValue GetOptionValue(const FName& InOption) override
+		{ return PlayerSessionServices->GetOptionValue(InOption); }
 		void LogMessage(IInfoLog::ELevel Level, const FString& Message) override
 		{ PlayerSessionServices->PostLog(Facility::EFacility::ABR, Level, Message); }
 		TSharedPtrTS<FABRStreamInformation> GetStreamInformation(const Metrics::FSegmentDownloadStats& FromDownloadStats) override;
@@ -449,9 +451,9 @@ namespace Electra
 		return ABRMethod.IsValid() ? ABRMethod->PeriodicHandle() : IAdaptiveStreamSelector::EHandlingAction::None;
 	}
 
-	IAdaptiveStreamSelector::FRebufferAction FAdaptiveStreamSelector::GetRebufferAction(const FParamDict& CurrentPlayerOptions)
+	IAdaptiveStreamSelector::FRebufferAction FAdaptiveStreamSelector::GetRebufferAction()
 	{
-		return ABRMethod.IsValid() ? ABRMethod->GetRebufferAction(CurrentPlayerOptions) : IAdaptiveStreamSelector::FRebufferAction();
+		return ABRMethod.IsValid() ? ABRMethod->GetRebufferAction() : IAdaptiveStreamSelector::FRebufferAction();
 	}
 
 	FTimeValue FAdaptiveStreamSelector::GetMinBufferTimeForPlayback(IAdaptiveStreamSelector::EMinBufferType InBufferingType, FTimeValue InDefaultMBT)
diff --git a/Engine/Plugins/Media/ElectraPlayer/Source/ElectraPlayerRuntime/Private/Runtime/Player/AdaptiveStreamingPlayerABR.h b/Engine/Plugins/Media/ElectraPlayer/Source/ElectraPlayerRuntime/Private/Runtime/Player/AdaptiveStreamingPlayerABR.h
--- a/Engine/Plugins/Media/ElectraPlayer/Source/ElectraPlayerRuntime/Private/Runtime/Player/AdaptiveStreamingPlayerABR.h
+++ b/Engine/Plugins/Media/ElectraPlayer/Source/ElectraPlayerRuntime/Private/Runtime/Player/AdaptiveStreamingPlayerABR.h
@@ -120,7 +120,7 @@ public:
 
 	// Call when rebuffering occurs to get the recommended action. This implies that rebuffering has actually occurred.
 	// The ABR may take internal action to switch quality when this is called!
-	virtual FRebufferAction GetRebufferAction(const FParamDict& CurrentPlayerOptions) = 0;
+	virtual FRebufferAction GetRebufferAction() = 0;
 
 	struct FDenylistedStream
 	{
diff --git a/Engine/Plugins/Media/ElectraPlayer/Source/ElectraPlayerRuntime/Private/Runtime/Player/AdaptiveStreamingPlayerDecoders.cpp b/Engine/Plugins/Media/ElectraPlayer/Source/ElectraPlayerRuntime/Private/Runtime/Player/AdaptiveStreamingPlayerDecoders.cpp
--- a/Engine/Plugins/Media/ElectraPlayer/Source/ElectraPlayerRuntime/Private/Runtime/Player/AdaptiveStreamingPlayerDecoders.cpp
+++ b/Engine/Plugins/Media/ElectraPlayer/Source/ElectraPlayerRuntime/Private/Runtime/Player/AdaptiveStreamingPlayerDecoders.cpp
@@ -282,16 +282,6 @@ int32 FAdaptiveStreamingPlayer::CreateDecoder(EStreamType type)
 				{
 					VideoDecoder.Parent = this;
 					VideoDecoder.Decoder->SetPlayerSessionServices(this);
-
-					// Add in any player options that are for decoder use
-					FParamDict AdditionalOptions;
-					TArray<FName> DecoderOptionKeys;
-					PlayerOptions.GetKeysStartingWith(TEXT("videoDecoder"), DecoderOptionKeys);
-					for (const FName& Key : DecoderOptionKeys)
-					{
-						AdditionalOptions.Set(Key, PlayerOptions.GetValue(Key));
-					}
-
 					// Attach video decoder buffer monitor.
 					VideoDecoder.Decoder->SetAUInputBufferListener(&VideoDecoder);
 					VideoDecoder.Decoder->SetReadyBufferListener(&VideoDecoder);
@@ -300,7 +290,7 @@ int32 FAdaptiveStreamingPlayer::CreateDecoder(EStreamType type)
 					// Hand it (may be nullptr) a delegate for platform for resource queries
 					VideoDecoder.Decoder->SetVideoResourceDelegate(VideoDecoderResourceDelegate);
 					// Open the decoder after having set all listeners.
-					VideoDecoder.Decoder->Open(VideoDecoder.LastSentAUCodecData, AdditionalOptions, &HighestStream);
+					VideoDecoder.Decoder->Open(VideoDecoder.LastSentAUCodecData, PlayerOptions.GetDictionary(), &HighestStream);
 				}
 
 				VideoDecoder.CheckIfNewDecoderMustBeSuspendedImmediately();
@@ -651,7 +641,7 @@ void FAdaptiveStreamingPlayer::FeedDecoder(EStreamType Type, IAccessUnitBufferIn
 					{
 						RebufferDetectedAtPlayPos = LastKnownPTS;
 					}
-					WorkerThread.Enqueue(FWorkerThreadMessages::FMessage::EType::BufferUnderrun);
+					WorkerThread.EnqueueBufferUnderrun();
 				}
 			}
 		}
diff --git a/Engine/Plugins/Media/ElectraPlayer/Source/ElectraPlayerRuntime/Private/Runtime/Player/AdaptiveStreamingPlayerInternal.h b/Engine/Plugins/Media/ElectraPlayer/Source/ElectraPlayerRuntime/Private/Runtime/Player/AdaptiveStreamingPlayerInternal.h
--- a/Engine/Plugins/Media/ElectraPlayer/Source/ElectraPlayerRuntime/Private/Runtime/Player/AdaptiveStreamingPlayerInternal.h
+++ b/Engine/Plugins/Media/ElectraPlayer/Source/ElectraPlayerRuntime/Private/Runtime/Player/AdaptiveStreamingPlayerInternal.h
@@ -31,6 +31,8 @@
 
 #include "InfoLog.h"
 
+#include "Misc/TVariant.h"
+
 #define INTERR_ALL_STREAMS_HAVE_FAILED			1
 #define INTERR_UNSUPPORTED_FORMAT				2
 #define INTERR_COULD_NOT_LOCATE_START_SEGMENT	3
@@ -1063,7 +1065,10 @@ private:
 	TSharedPtrTS<IPlayerEntityCache> GetEntityCache() override;
 	TSharedPtrTS<IHTTPResponseCache> GetHTTPResponseCache() override;
 	IAdaptiveStreamingPlayerAEMSHandler* GetAEMSEventHandler() override;
-	FParamDict& GetOptions() override;
+	FParamDictTS& GetMutableOptions() override;
+	bool HaveOptionValue(const FName& InOption) override;
+	const FVariantValue GetOptionValue(const FName& InOption) override;
+
 	TSharedPtrTS<FDRMManager> GetDRMManager() override;
 	void SetPlaybackEnd(const FTimeValue& InEndAtTime, IPlayerSessionServices::EPlayEndReason InEndingReason, TSharedPtrTS<IPlayerSessionServices::IPlayEndReason> InCustomManifestObject) override;
 
@@ -1422,6 +1427,7 @@ private:
 			{
 				// Commands
 				Initialize,
+				ChangeOptions,
 				LoadManifest,
 				LoadBlob,
 				Pause,
@@ -1444,110 +1450,91 @@ private:
 				BufferUnderrun,
 			};
 
-			struct FData
+			struct FLoadManifest
 			{
-				struct FLoadManifest
-				{
-					FString											URL;
-					FString											MimeType;
-				};
-				struct FLoadBlob
-				{
-					TSharedPtrTS<FHTTPResourceRequest>				BlobLoadRequest;
-				};
-				struct FStreamReader
-				{
-					FStreamReader()
-						: AUType(EStreamType::Video), AUSize(0)
-					{
-					}
-					TSharedPtrTS<IStreamSegment>					Request;
-					EStreamType										AUType;
-					SIZE_T											AUSize;
-				};
-				struct FEvent
-				{
-					FEvent()
-						: Event(nullptr)
-					{
-					}
-					FMediaEvent*									Event;
-				};
-				struct FLoop
-				{
-					FLoop()
-						: Signal(nullptr)
-					{
-					}
-					FLoopParam										Loop;
-					FMediaEvent*									Signal;
-				};
-				struct FBitrate
-				{
-					int32											Value;
-				};
-				struct FSession
-				{
-					TSharedPtrTS<IPlayerMessage>					PlayerMessage;
-				};
-				struct FResolution
-				{
-					FResolution() : Width(0), Height(0) { }
-					int32											Width;
-					int32											Height;
-				};
-				struct FInitialStreamSelect
-				{
-					EStreamType										StreamType;
-					FStreamSelectionAttributes						InitialSelection;
-				};
-				struct FMetadataTrackSelection
-				{
-					EStreamType										StreamType;
-					FTrackMetadata									TrackMetadata;
-					FStreamSelectionAttributes						TrackAttributes;
-				};
-				struct FEndPlaybackAt
-				{
-					FTimeValue											EndAtTime;
-					IPlayerSessionServices::EPlayEndReason				EndingReason;
-					TSharedPtrTS<IPlayerSessionServices::IPlayEndReason> CustomManifestObject;
-				};
-
-				FLoadManifest				ManifestToLoad;
-				FLoadBlob					BlobToLoad;
-				FStreamReader				StreamReader;
-				FEvent						MediaEvent;
-				FLoop						Looping;
-				FBitrate					Bitrate;
-				FSession					Session;
-				FResolution					Resolution;
-				FInitialStreamSelect		InitialStreamAttribute;
-				FMetadataTrackSelection		TrackSelection;
-				FEndPlaybackAt				EndPlaybackAt;
+				FString											URL;
+				FString											MimeType;
+			};
+			struct FLoadBlob
+			{
+				TSharedPtrTS<FHTTPResourceRequest>				BlobLoadRequest;
+			};
+			struct FOptionChange
+			{
+				FParamDict										OptionsToSetOrChange;
+				FParamDict										OptionsToClear;
+			};
+			struct FStreamReader
+			{
+				TSharedPtrTS<IStreamSegment>					Request;
+				EStreamType										AUType = EStreamType::Unsupported;
+				SIZE_T											AUSize = 0;
+			};
+			struct FEvent
+			{
+				FMediaEvent*									Event = nullptr;
+			};
+			struct FLoop
+			{
+				FLoopParam										Loop;
+				//FMediaEvent*									Signal = nullptr;
+			};
+			struct FBitrate
+			{
+				int32											Value = 0;
+			};
+			struct FSession
+			{
+				TSharedPtrTS<IPlayerMessage>					PlayerMessage;
+			};
+			struct FResolution
+			{
+				int32											Width = 0;
+				int32											Height = 0;
+			};
+			struct FInitialStreamSelect
+			{
+				EStreamType										StreamType = EStreamType::Unsupported;;
+				FStreamSelectionAttributes						InitialSelection;
+			};
+			struct FMetadataTrackSelection
+			{
+				EStreamType										StreamType = EStreamType::Unsupported;;
+				FTrackMetadata									TrackMetadata;
+				FStreamSelectionAttributes						TrackAttributes;
 			};
-			EType					Type;
-			FData   				Data;
+			struct FEndPlaybackAt
+			{
+				FTimeValue											EndAtTime;
+				IPlayerSessionServices::EPlayEndReason				EndingReason = IPlayerSessionServices::EPlayEndReason::EndAll;
+				TSharedPtrTS<IPlayerSessionServices::IPlayEndReason> CustomManifestObject;
+			};
+
+			EType Type;
+			TVariant<FLoadManifest, FLoadBlob, FOptionChange, FStreamReader, FEvent, FLoop, FBitrate, FSession, FResolution, FInitialStreamSelect, FMetadataTrackSelection, FEndPlaybackAt> Data;
 		};
 
-		void Enqueue(FMessage::EType type)
+		void Enqueue(FMessage::EType type, const TSharedPtrTS<IStreamSegment>& InRequest)
 		{
+			FMessage::FStreamReader sr { InRequest };
 			FMessage Msg;
 			Msg.Type = type;
+			Msg.Data.Emplace<FMessage::FStreamReader>(MoveTemp(sr));
 			TriggerSharedWorkerThread(MoveTemp(Msg));
 		}
-		void Enqueue(FMessage::EType type, TSharedPtrTS<IStreamSegment> pRequest)
+		void Enqueue(FMessage::EType type, FMediaEvent* InEventSignal)
 		{
+			check(InEventSignal);
+			FMessage::FEvent ev { InEventSignal };
 			FMessage Msg;
 			Msg.Type = type;
-			Msg.Data.StreamReader.Request = pRequest;
+			Msg.Data.Emplace<FMessage::FEvent>(MoveTemp(ev));
 			TriggerSharedWorkerThread(MoveTemp(Msg));
 		}
-		void Enqueue(FMessage::EType type, FMediaEvent* pEventSignal)
+		void EnqueueBufferUnderrun()
 		{
 			FMessage Msg;
-			Msg.Type = type;
-			Msg.Data.MediaEvent.Event = pEventSignal;
+			Msg.Type = FMessage::EType::BufferUnderrun;
 			TriggerSharedWorkerThread(MoveTemp(Msg));
 		}
 		void SendInitializeMessage()
@@ -1556,19 +1543,28 @@ private:
 			Msg.Type = FMessage::EType::Initialize;
 			TriggerSharedWorkerThread(MoveTemp(Msg));
 		}
-		void SendLoadManifestMessage(const FString& URL, const FString& MimeType)
+		void SendOptionChangeMessage(const FParamDict& InOptionsToSetOrChange, const FParamDict& InOptionsToClear)
+		{
+			FMessage::FOptionChange oc { InOptionsToSetOrChange, InOptionsToClear };
+			FMessage Msg;
+			Msg.Type = FMessage::EType::ChangeOptions;
+			Msg.Data.Emplace<FMessage::FOptionChange>(MoveTemp(oc));
+			TriggerSharedWorkerThread(MoveTemp(Msg));
+		}
+		void SendLoadManifestMessage(const FString& InURL, const FString& InMimeType)
 		{
+			FMessage::FLoadManifest lm { InURL, InMimeType };
 			FMessage Msg;
 			Msg.Type = FMessage::EType::LoadManifest;
-			Msg.Data.ManifestToLoad.URL = URL;
-			Msg.Data.ManifestToLoad.MimeType = MimeType;
+			Msg.Data.Emplace<FMessage::FLoadManifest>(MoveTemp(lm));
 			TriggerSharedWorkerThread(MoveTemp(Msg));
 		}
-		void SendLoadBlobMessage(TSharedPtr<FHTTPResourceRequest, ESPMode::ThreadSafe> InBlobLoadRequest)
+		void SendLoadBlobMessage(const TSharedPtr<FHTTPResourceRequest, ESPMode::ThreadSafe>& InBlobLoadRequest)
 		{
+			FMessage::FLoadBlob lb { InBlobLoadRequest };
 			FMessage Msg;
 			Msg.Type = FMessage::EType::LoadBlob;
-			Msg.Data.BlobToLoad.BlobLoadRequest = MoveTemp(InBlobLoadRequest);
+			Msg.Data.Emplace<FMessage::FLoadBlob>(MoveTemp(lb));
 			TriggerSharedWorkerThread(MoveTemp(Msg));
 		}
 		void SendPauseMessage()
@@ -1585,79 +1581,84 @@ private:
 		}
 		void SendLoopMessage(const FLoopParam& InLoopParams)
 		{
+			FMessage::FLoop loop { InLoopParams };
 			FMessage Msg;
 			Msg.Type = FMessage::EType::Loop;
-			Msg.Data.Looping.Loop = InLoopParams;
+			Msg.Data.Emplace<FMessage::FLoop>(MoveTemp(loop));
 			TriggerSharedWorkerThread(MoveTemp(Msg));
 		}
-		void SendCloseMessage(FMediaEvent* pEventSignal)
+		void SendCloseMessage(FMediaEvent* InEventSignal)
 		{
+			check(InEventSignal);
+			FMessage::FEvent ev { InEventSignal };
 			FMessage Msg;
 			Msg.Type = FMessage::EType::Close;
-			Msg.Data.MediaEvent.Event = pEventSignal;
+			Msg.Data.Emplace<FMessage::FEvent>(MoveTemp(ev));
 			TriggerSharedWorkerThread(MoveTemp(Msg));
 		}
-		void SendBitrateMessage(EStreamType type, int32 value, int32 which)
+		void SendBitrateMessage(EStreamType /*InType*/, int32 InValue, int32 /*InWhich*/)
 		{
+			FMessage::FBitrate br { InValue };
 			FMessage Msg;
 			Msg.Type = FMessage::EType::ChangeBitrate;
-			Msg.Data.Bitrate.Value = value;
+			Msg.Data.Emplace<FMessage::FBitrate>(MoveTemp(br));
 			TriggerSharedWorkerThread(MoveTemp(Msg));
 		}
-		void SendPlayerSessionMessage(const TSharedPtrTS<IPlayerMessage>& message)
+		void SendPlayerSessionMessage(const TSharedPtrTS<IPlayerMessage>& InMessage)
 		{
+			FMessage::FSession sess { InMessage };
 			FMessage Msg;
 			Msg.Type = FMessage::EType::PlayerSession;
-			Msg.Data.Session.PlayerMessage = message;
+			Msg.Data.Emplace<FMessage::FSession>(MoveTemp(sess));
 			TriggerSharedWorkerThread(MoveTemp(Msg));
 		}
-		void SendResolutionMessage(int32 Width, int32 Height)
+		void SendResolutionMessage(int32 InWidth, int32 InHeight)
 		{
+			FMessage::FResolution reso { InWidth, InHeight };
 			FMessage Msg;
 			Msg.Type = FMessage::EType::LimitResolution;
-			Msg.Data.Resolution.Width = Width;
-			Msg.Data.Resolution.Height = Height;
+			Msg.Data.Emplace<FMessage::FResolution>(MoveTemp(reso));
 			TriggerSharedWorkerThread(MoveTemp(Msg));
 		}
-		void SendInitialStreamAttributeMessage(EStreamType StreamType, const FStreamSelectionAttributes& InitialSelection)
+		void SendInitialStreamAttributeMessage(EStreamType InStreamType, const FStreamSelectionAttributes& InInitialSelection)
 		{
+			FMessage::FInitialStreamSelect iss { InStreamType, InInitialSelection };
 			FMessage Msg;
 			Msg.Type = FMessage::EType::InitialStreamAttributes;
-			Msg.Data.InitialStreamAttribute.StreamType = StreamType;
-			Msg.Data.InitialStreamAttribute.InitialSelection = InitialSelection;
+			Msg.Data.Emplace<FMessage::FInitialStreamSelect>(MoveTemp(iss));
 			TriggerSharedWorkerThread(MoveTemp(Msg));
 		}
-		void SendTrackSelectByMetadataMessage(EStreamType StreamType, const FTrackMetadata& TrackMetadata)
+		void SendTrackSelectByMetadataMessage(EStreamType InStreamType, const FTrackMetadata& InTrackMetadata)
 		{
+			FMessage::FMetadataTrackSelection mts { InStreamType, InTrackMetadata };
 			FMessage Msg;
 			Msg.Type = FMessage::EType::SelectTrackByMetadata;
-			Msg.Data.TrackSelection.StreamType = StreamType;
-			Msg.Data.TrackSelection.TrackMetadata = TrackMetadata;
+			Msg.Data.Emplace<FMessage::FMetadataTrackSelection>(MoveTemp(mts));
 			TriggerSharedWorkerThread(MoveTemp(Msg));
 		}
-		void SendTrackSelectByAttributeMessage(EStreamType StreamType, const FStreamSelectionAttributes& TrackAttributes)
+		void SendTrackSelectByAttributeMessage(EStreamType InStreamType, const FStreamSelectionAttributes& InTrackAttributes)
 		{
+			FMessage::FMetadataTrackSelection mts { InStreamType };
+			mts.TrackAttributes = InTrackAttributes;
 			FMessage Msg;
 			Msg.Type = FMessage::EType::SelectTrackByAttributes;
-			Msg.Data.TrackSelection.StreamType = StreamType;
-			Msg.Data.TrackSelection.TrackAttributes = TrackAttributes;
+			Msg.Data.Emplace<FMessage::FMetadataTrackSelection>(MoveTemp(mts));
 			TriggerSharedWorkerThread(MoveTemp(Msg));
 		}
-		void SendTrackDeselectMessage(EStreamType StreamType)
+		void SendTrackDeselectMessage(EStreamType InStreamType)
 		{
+			FMessage::FMetadataTrackSelection mts { InStreamType };
 			FMessage Msg;
 			Msg.Type = FMessage::EType::DeselectTrack;
-			Msg.Data.TrackSelection.StreamType = StreamType;
+			Msg.Data.Emplace<FMessage::FMetadataTrackSelection>(MoveTemp(mts));
 			TriggerSharedWorkerThread(MoveTemp(Msg));
 		}
-
-		void SendPlaybackEndMessage(const FTimeValue& InEndAtTime, IPlayerSessionServices::EPlayEndReason InEndingReason, TSharedPtrTS<IPlayerSessionServices::IPlayEndReason> InCustomManifestObject)
+		void SendPlaybackEndMessage(const FTimeValue& InEndAtTime, IPlayerSessionServices::EPlayEndReason InEndingReason, const TSharedPtrTS<IPlayerSessionServices::IPlayEndReason>& InCustomManifestObject)
 		{
+			FMessage::FEndPlaybackAt epa { InEndAtTime, InEndingReason, InCustomManifestObject };
 			FMessage Msg;
 			Msg.Type = FMessage::EType::EndPlaybackAt;
-			Msg.Data.EndPlaybackAt.EndAtTime = InEndAtTime;
-			Msg.Data.EndPlaybackAt.EndingReason = InEndingReason;
-			Msg.Data.EndPlaybackAt.CustomManifestObject = InCustomManifestObject;
+			Msg.Data.Emplace<FMessage::FEndPlaybackAt>(MoveTemp(epa));
 			TriggerSharedWorkerThread(MoveTemp(Msg));
 		}
 
@@ -2158,7 +2159,7 @@ private:
 	FWorkerThreadMessages												WorkerThread;
 	bool																bUseSharedWorkerThread;
 
-	FParamDict															PlayerOptions;
+	FParamDictTS														PlayerOptions;
 	FPlaybackState														PlaybackState;
 	ISynchronizedUTCTime*												SynchronizedUTCTime;
 	IAdaptiveStreamingPlayerAEMSHandler*								AEMSEventHandler;
diff --git a/Engine/Plugins/Media/ElectraPlayer/Source/ElectraPlayerRuntime/Private/Runtime/Player/AdaptiveStreamingPlayerManifest.cpp b/Engine/Plugins/Media/ElectraPlayer/Source/ElectraPlayerRuntime/Private/Runtime/Player/AdaptiveStreamingPlayerManifest.cpp
--- a/Engine/Plugins/Media/ElectraPlayer/Source/ElectraPlayerRuntime/Private/Runtime/Player/AdaptiveStreamingPlayerManifest.cpp
+++ b/Engine/Plugins/Media/ElectraPlayer/Source/ElectraPlayerRuntime/Private/Runtime/Player/AdaptiveStreamingPlayerManifest.cpp
@@ -186,9 +186,9 @@ void FAdaptiveStreamingPlayer::InternalLoadManifest(const FString& InURL, const
 	if (CurrentState == EPlayerState::eState_Idle)
 	{
 		FString mimeType = MimeType;
-		if (GetOptions().HaveKey(OptionKeyMimeType))
+		if (PlayerOptions.HaveKey(OptionKeyMimeType))
 		{
-			mimeType = GetOptions().GetValue(OptionKeyMimeType).GetFString();
+			mimeType = PlayerOptions.GetValue(OptionKeyMimeType).GetFString();
 		}
 		else if (mimeType.IsEmpty())
 		{
diff --git a/Engine/Plugins/Media/ElectraPlayer/Source/ElectraPlayerRuntime/Private/Runtime/Player/AdaptiveStreamingPlayerRenderers.cpp b/Engine/Plugins/Media/ElectraPlayer/Source/ElectraPlayerRuntime/Private/Runtime/Player/AdaptiveStreamingPlayerRenderers.cpp
--- a/Engine/Plugins/Media/ElectraPlayer/Source/ElectraPlayerRuntime/Private/Runtime/Player/AdaptiveStreamingPlayerRenderers.cpp
+++ b/Engine/Plugins/Media/ElectraPlayer/Source/ElectraPlayerRuntime/Private/Runtime/Player/AdaptiveStreamingPlayerRenderers.cpp
@@ -107,7 +107,7 @@ namespace Electra
 		const FParamDict& GetBufferPoolProperties() const override;
 		UEMediaError CreateBufferPool(const FParamDict& Parameters) override;
 		UEMediaError AcquireBuffer(IBuffer*& OutBuffer, int32 TimeoutInMicroseconds, const FParamDict& InParameters) override;
-		UEMediaError ReturnBuffer(IBuffer* Buffer, bool bRender, const FParamDict& InSampleProperties) override;
+		UEMediaError ReturnBuffer(IBuffer* Buffer, bool bRender, FParamDict& InOutSampleProperties) override;
 		UEMediaError ReleaseBufferPool() override;
 		bool CanReceiveOutputFrames(uint64 NumFrames) const override;
 		bool GetEnqueuedFrameInfo(int32& OutNumberOfEnqueuedFrames, FTimeValue& OutDurationOfEnqueuedFrames) const override;
@@ -372,21 +372,20 @@ UEMediaError FAdaptiveStreamingWrappedRenderer::AcquireBuffer(IBuffer*& OutBuffe
 	}
 }
 
-UEMediaError FAdaptiveStreamingWrappedRenderer::ReturnBuffer(IBuffer* Buffer, bool bRender, const FParamDict& InSampleProperties)
+UEMediaError FAdaptiveStreamingWrappedRenderer::ReturnBuffer(IBuffer* Buffer, bool bRender, FParamDict& InOutSampleProperties)
 {
 	LLM_SCOPE(ELLMTag::ElectraPlayer);
 
-	FParamDict SampleProperties(InSampleProperties);
 	Lock.Lock();
-	SampleProperties.Set(RenderOptionKeys::ValidityValue, FVariantValue(CurrentValidityValue));
+	InOutSampleProperties.Set(RenderOptionKeys::ValidityValue, FVariantValue(CurrentValidityValue));
 	Lock.Unlock();
 	if (Type == EStreamType::Video)
 	{
-		return ReturnVideoBuffer(Buffer, bRender, SampleProperties);
+		return ReturnVideoBuffer(Buffer, bRender, InOutSampleProperties);
 	}
 	else if (Type == EStreamType::Audio)
 	{
-		return ReturnAudioBuffer(Buffer, bRender, SampleProperties);
+		return ReturnAudioBuffer(Buffer, bRender, InOutSampleProperties);
 	}
 	else
 	{
diff --git a/Engine/Plugins/Media/ElectraPlayer/Source/ElectraPlayerRuntime/Private/Runtime/Player/DASH/ManifestBuilderDASH.cpp b/Engine/Plugins/Media/ElectraPlayer/Source/ElectraPlayerRuntime/Private/Runtime/Player/DASH/ManifestBuilderDASH.cpp
--- a/Engine/Plugins/Media/ElectraPlayer/Source/ElectraPlayerRuntime/Private/Runtime/Player/DASH/ManifestBuilderDASH.cpp
+++ b/Engine/Plugins/Media/ElectraPlayer/Source/ElectraPlayerRuntime/Private/Runtime/Player/DASH/ManifestBuilderDASH.cpp
@@ -2378,7 +2378,7 @@ FTimeValue FManifestDASHInternal::CalculateDistanceToLiveEdge() const
 	}
 
 	// Check if there is a user provided value. If there is it takes precedence over everything else.
-	FTimeValue Distance = PlayerSessionServices->GetOptions().GetValue(OptionKeyLiveSeekableEndOffset).SafeGetTimeValue(FTimeValue());
+	FTimeValue Distance = PlayerSessionServices->GetOptionValue(OptionKeyLiveSeekableEndOffset).SafeGetTimeValue(FTimeValue());
 
 	// If there is a low latency descriptor we use the target latency from it.
 	if (LowLatencyDescriptor.IsValid())
diff --git a/Engine/Plugins/Media/ElectraPlayer/Source/ElectraPlayerRuntime/Private/Runtime/Player/DASH/ManifestDASH.cpp b/Engine/Plugins/Media/ElectraPlayer/Source/ElectraPlayerRuntime/Private/Runtime/Player/DASH/ManifestDASH.cpp
--- a/Engine/Plugins/Media/ElectraPlayer/Source/ElectraPlayerRuntime/Private/Runtime/Player/DASH/ManifestDASH.cpp
+++ b/Engine/Plugins/Media/ElectraPlayer/Source/ElectraPlayerRuntime/Private/Runtime/Player/DASH/ManifestDASH.cpp
@@ -39,9 +39,9 @@ namespace DashUtils
 	#define GETPLAYEROPTION(Type, Getter)																						\
 		bool GetPlayerOption(IPlayerSessionServices* InPlayerSessionServices, Type& OutValue, const FName& Key, Type Default)	\
 		{																														\
-			if (InPlayerSessionServices->GetOptions().HaveKey(Key))																\
+			if (InPlayerSessionServices->HaveOptionValue(Key))																	\
 			{																													\
-				OutValue = InPlayerSessionServices->GetOptions().GetValue(Key).Getter(Default);									\
+				OutValue = InPlayerSessionServices->GetOptionValue(Key).Getter(Default);										\
 				return true;																									\
 			}																													\
 			OutValue = Default;																									\
@@ -998,7 +998,7 @@ void FDASHPlayPeriod::PrepareForPlay()
 		if (VideoAS.IsValid())
 		{
 			// Get the current average video bitrate with some sensible default if it is not set.
-			int64 StartingBitrate = PlayerSessionServices->GetOptions().GetValue(OptionKeyCurrentAvgStartingVideoBitrate).SafeGetInt64(2*1000*1000);
+			int64 StartingBitrate = PlayerSessionServices->GetOptionValue(OptionKeyCurrentAvgStartingVideoBitrate).SafeGetInt64(2*1000*1000);
 
 			TSharedPtrTS<IPlaybackAssetRepresentation> VideoRepr = GetRepresentationFromAdaptationByMaxBandwidth(VideoAS, (int32) StartingBitrate);
 			if (VideoRepr.IsValid())
@@ -1040,11 +1040,11 @@ void FDASHPlayPeriod::PrepareForPlay()
 		{
 			if (llDesc->Latency.ReferenceID >= 0)
 			{
-				PlayerSessionServices->GetOptions().Set(DASH::OptionKey_LatencyReferenceId, FVariantValue(llDesc->Latency.ReferenceID));
+				PlayerSessionServices->GetMutableOptions().Set(DASH::OptionKey_LatencyReferenceId, FVariantValue(llDesc->Latency.ReferenceID));
 			}
 			else
 			{
-				PlayerSessionServices->GetOptions().Remove(DASH::OptionKey_LatencyReferenceId);
+				PlayerSessionServices->GetMutableOptions().Remove(DASH::OptionKey_LatencyReferenceId);
 			}
 		}
 
diff --git a/Engine/Plugins/Media/ElectraPlayer/Source/ElectraPlayerRuntime/Private/Runtime/Player/DASH/PlaylistReaderDASH.cpp b/Engine/Plugins/Media/ElectraPlayer/Source/ElectraPlayerRuntime/Private/Runtime/Player/DASH/PlaylistReaderDASH.cpp
--- a/Engine/Plugins/Media/ElectraPlayer/Source/ElectraPlayerRuntime/Private/Runtime/Player/DASH/PlaylistReaderDASH.cpp
+++ b/Engine/Plugins/Media/ElectraPlayer/Source/ElectraPlayerRuntime/Private/Runtime/Player/DASH/PlaylistReaderDASH.cpp
@@ -806,7 +806,7 @@ void FPlaylistReaderDASH::InternalSetup()
 	PlayerSessionServices->GetAEMSEventHandler()->AddAEMSReceiver(AsShared(), DASH::Schemes::ManifestEvents::Scheme_urn_mpeg_dash_event_ttfn_2016, TEXT(""), IAdaptiveStreamingPlayerAEMSReceiver::EDispatchMode::OnStart, false);
 
 	// Get the minimum MPD update time limit.
-	MinTimeBetweenUpdates = PlayerSessionServices->GetOptions().GetValue(DASH::OptionKey_MinTimeBetweenMPDUpdates).SafeGetTimeValue(FTimeValue().SetFromMilliseconds(1000));
+	MinTimeBetweenUpdates = PlayerSessionServices->GetOptionValue(DASH::OptionKey_MinTimeBetweenMPDUpdates).SafeGetTimeValue(FTimeValue().SetFromMilliseconds(1000));
 
 	// Setup the playlist load request for the master playlist.
 	TSharedPtrTS<FMPDLoadRequestDASH> PlaylistLoadRequest = MakeSharedTS<FMPDLoadRequestDASH>();
@@ -903,19 +903,18 @@ void FPlaylistReaderDASH::ExecutePendingRequests(const FTimeValue& TimeNow)
 
 void FPlaylistReaderDASH::SetupRequestTimeouts(FResourceLoadRequestPtr InRequest)
 {
-	const FParamDict& Options = PlayerSessionServices->GetOptions();
 	switch(InRequest->LoadType)
 	{
 		case FMPDLoadRequestDASH::ELoadType::MPD:
 		{
-			InRequest->Request->ConnectionTimeout(Options.GetValue(DASH::OptionKeyMPDLoadConnectTimeout).SafeGetTimeValue(FTimeValue().SetFromMilliseconds(1000 * 8)));
-			InRequest->Request->NoDataTimeout(Options.GetValue(DASH::OptionKeyMPDLoadNoDataTimeout).SafeGetTimeValue(FTimeValue().SetFromMilliseconds(1000 * 5)));
+			InRequest->Request->ConnectionTimeout(PlayerSessionServices->GetOptionValue(DASH::OptionKeyMPDLoadConnectTimeout).SafeGetTimeValue(FTimeValue().SetFromMilliseconds(1000 * 8)));
+			InRequest->Request->NoDataTimeout(PlayerSessionServices->GetOptionValue(DASH::OptionKeyMPDLoadNoDataTimeout).SafeGetTimeValue(FTimeValue().SetFromMilliseconds(1000 * 5)));
 			break;
 		}
 		case FMPDLoadRequestDASH::ELoadType::MPDUpdate:
 		{
-			InRequest->Request->ConnectionTimeout(Options.GetValue(DASH::OptionKeyMPDReloadConnectTimeout).SafeGetTimeValue(FTimeValue().SetFromMilliseconds(1000 * 2)));
-			InRequest->Request->NoDataTimeout(Options.GetValue(DASH::OptionKeyMPDReloadNoDataTimeout).SafeGetTimeValue(FTimeValue().SetFromMilliseconds(1000 * 2)));
+			InRequest->Request->ConnectionTimeout(PlayerSessionServices->GetOptionValue(DASH::OptionKeyMPDReloadConnectTimeout).SafeGetTimeValue(FTimeValue().SetFromMilliseconds(1000 * 2)));
+			InRequest->Request->NoDataTimeout(PlayerSessionServices->GetOptionValue(DASH::OptionKeyMPDReloadNoDataTimeout).SafeGetTimeValue(FTimeValue().SetFromMilliseconds(1000 * 2)));
 			break;
 		}
 		case FMPDLoadRequestDASH::ELoadType::Segment:
diff --git a/Engine/Plugins/Media/ElectraPlayer/Source/ElectraPlayerRuntime/Private/Runtime/Player/DASH/StreamReaderDASH.cpp b/Engine/Plugins/Media/ElectraPlayer/Source/ElectraPlayerRuntime/Private/Runtime/Player/DASH/StreamReaderDASH.cpp
--- a/Engine/Plugins/Media/ElectraPlayer/Source/ElectraPlayerRuntime/Private/Runtime/Player/DASH/StreamReaderDASH.cpp
+++ b/Engine/Plugins/Media/ElectraPlayer/Source/ElectraPlayerRuntime/Private/Runtime/Player/DASH/StreamReaderDASH.cpp
@@ -439,9 +439,9 @@ void FStreamReaderDASH::FStreamHandler::HTTPUpdateStats(const FTimeValue& Curren
 	TSharedPtrTS<FStreamSegmentRequestDASH> SegmentRequest = CurrentRequest;
 	if (SegmentRequest.IsValid())
 	{
+		// Only update elements that are needed by the ABR here.
 		FMediaCriticalSection::ScopedLock lock(MetricUpdateLock);
-		SegmentRequest->ConnectionInfo = Request->ConnectionInfo;
-		// Update the current download stats which we report periodically to the ABR.
+		SegmentRequest->ConnectionInfo.RequestStartTime = Request->ConnectionInfo.RequestStartTime;
 		Metrics::FSegmentDownloadStats& ds = SegmentRequest->DownloadStats;
 		if (Request->ConnectionInfo.EffectiveURL.Len())
 		{
@@ -1011,8 +1011,7 @@ void FStreamReaderDASH::FStreamHandler::HandleRequestMP4()
 	FTimeValue LastKnownAUDuration;
 	FTimeValue TimeOffset = Request->PeriodStart + Request->AST + Request->AdditionalAdjustmentTime;
 
-	const FParamDict& Options = PlayerSessionService->GetOptions();
-	bool bDoNotTruncateAtPresentationEnd = Options.GetValue(OptionKeyDoNotTruncateAtPresentationEnd).SafeGetBool(false);
+	bool bDoNotTruncateAtPresentationEnd = PlayerSessionService->GetOptionValue(OptionKeyDoNotTruncateAtPresentationEnd).SafeGetBool(false);
 
 	// Get the init segment if there is one. Either gets it from the entity cache or requests it now and adds it to the cache.
 	InitSegmentError = GetInitSegment(MP4InitSegment, Request);
@@ -1116,8 +1115,8 @@ void FStreamReaderDASH::FStreamHandler::HandleRequestMP4()
 			}
 			HTTP->Parameters.bCollectTimingTraces = Request->Segment.bLowLatencyChunkedEncodingExpected;
 			// Set timeouts for media segment retrieval
-			HTTP->Parameters.ConnectTimeout = Options.GetValue(DASH::OptionKeyMediaSegmentConnectTimeout).SafeGetTimeValue(FTimeValue().SetFromMilliseconds(1000 * 4));
-			HTTP->Parameters.NoDataTimeout = Options.GetValue(DASH::OptionKeyMediaSegmentNoDataTimeout).SafeGetTimeValue(FTimeValue().SetFromMilliseconds(1000 * 4));
+			HTTP->Parameters.ConnectTimeout = PlayerSessionService->GetOptionValue(DASH::OptionKeyMediaSegmentConnectTimeout).SafeGetTimeValue(FTimeValue().SetFromMilliseconds(1000 * 4));
+			HTTP->Parameters.NoDataTimeout = PlayerSessionService->GetOptionValue(DASH::OptionKeyMediaSegmentNoDataTimeout).SafeGetTimeValue(FTimeValue().SetFromMilliseconds(1000 * 4));
 
 			ProgressReportCount = 0;
 			DownloadCompleteSignal.Reset();
@@ -1873,8 +1872,7 @@ void FStreamReaderDASH::FStreamHandler::HandleRequestMKV()
 	bFillRemainingDuration = false;
 	ABRAbortReason.Empty();
 
-	const FParamDict& Options = PlayerSessionService->GetOptions();
-	bool bDoNotTruncateAtPresentationEnd = Options.GetValue(OptionKeyDoNotTruncateAtPresentationEnd).SafeGetBool(false);
+	bool bDoNotTruncateAtPresentationEnd = PlayerSessionService->GetOptionValue(OptionKeyDoNotTruncateAtPresentationEnd).SafeGetBool(false);
 
 	// Get the init segment if there is one. Either gets it from the entity cache or requests it now and adds it to the cache.
 	InitSegmentError = GetInitSegment(MKVParser, Request);
@@ -1952,8 +1950,8 @@ void FStreamReaderDASH::FStreamHandler::HandleRequestMKV()
 			}
 			HTTP->Parameters.bCollectTimingTraces = Request->Segment.bLowLatencyChunkedEncodingExpected;
 			// Set timeouts for media segment retrieval
-			HTTP->Parameters.ConnectTimeout = Options.GetValue(DASH::OptionKeyMediaSegmentConnectTimeout).SafeGetTimeValue(FTimeValue().SetFromMilliseconds(1000 * 4));
-			HTTP->Parameters.NoDataTimeout = Options.GetValue(DASH::OptionKeyMediaSegmentNoDataTimeout).SafeGetTimeValue(FTimeValue().SetFromMilliseconds(1000 * 4));
+			HTTP->Parameters.ConnectTimeout = PlayerSessionService->GetOptionValue(DASH::OptionKeyMediaSegmentConnectTimeout).SafeGetTimeValue(FTimeValue().SetFromMilliseconds(1000 * 4));
+			HTTP->Parameters.NoDataTimeout = PlayerSessionService->GetOptionValue(DASH::OptionKeyMediaSegmentNoDataTimeout).SafeGetTimeValue(FTimeValue().SetFromMilliseconds(1000 * 4));
 
 			ProgressReportCount = 0;
 			DownloadCompleteSignal.Reset();
@@ -2606,14 +2604,13 @@ int64 FStreamReaderDASH::FStreamHandler::ReadData(void* IntoBuffer, int64 NumByt
 			if (CurrentRequest.IsValid())
 			{
 				MetricUpdateLock.Lock();
-				Metrics::FSegmentDownloadStats currentDownloadStats = CurrentRequest->DownloadStats;
+				Metrics::FSegmentDownloadStats& currentDownloadStats = CurrentRequest->DownloadStats;
 				currentDownloadStats.DurationDelivered = ActiveTrackData.DurationSuccessfullyDelivered.GetAsSeconds();
 				currentDownloadStats.DurationDownloaded = ActiveTrackData.DurationSuccessfullyRead.GetAsSeconds();
 				currentDownloadStats.TimeToDownload = (MEDIAutcTime::Current() - CurrentRequest->ConnectionInfo.RequestStartTime).GetAsSeconds();
+				FABRDownloadProgressDecision StreamSelectorDecision = StreamSelector->ReportDownloadProgress(currentDownloadStats);
 				MetricUpdateLock.Unlock();
 
-				Metrics::FSegmentDownloadStats& ds = CurrentRequest->DownloadStats;
-				FABRDownloadProgressDecision StreamSelectorDecision = StreamSelector->ReportDownloadProgress(currentDownloadStats);
 				if ((StreamSelectorDecision.Flags & FABRDownloadProgressDecision::EDecisionFlags::eABR_EmitPartialData) != 0)
 				{
 					SCOPE_CYCLE_COUNTER(STAT_ElectraPlayer_DASH_StreamReader);
diff --git a/Engine/Plugins/Media/ElectraPlayer/Source/ElectraPlayerRuntime/Private/Runtime/Player/HLS/ManifestBuilderHLS.cpp b/Engine/Plugins/Media/ElectraPlayer/Source/ElectraPlayerRuntime/Private/Runtime/Player/HLS/ManifestBuilderHLS.cpp
--- a/Engine/Plugins/Media/ElectraPlayer/Source/ElectraPlayerRuntime/Private/Runtime/Player/HLS/ManifestBuilderHLS.cpp
+++ b/Engine/Plugins/Media/ElectraPlayer/Source/ElectraPlayerRuntime/Private/Runtime/Player/HLS/ManifestBuilderHLS.cpp
@@ -1020,7 +1020,7 @@ FErrorDetail FManifestBuilderHLS::GetInitialPlaylistLoadRequests(TArray<FPlaylis
 	}
 
 	TSharedPtrTS<FManifestHLSInternal::FVariantStream> StartingVariantStream;
-	int32 InitialBitrate = (int32) PlayerSessionServices->GetOptions().GetValue(OptionKeyInitialBitrate).SafeGetInt64(0);
+	int32 InitialBitrate = (int32) PlayerSessionServices->GetOptionValue(OptionKeyInitialBitrate).SafeGetInt64(0);
 	bool bIsAudioOnlyVariant = false;
 
 	if (Manifest->VariantStreams.Num())
@@ -1754,7 +1754,6 @@ void FManifestBuilderHLS::UpdateManifestMetadataFromStream(FManifestHLSInternal*
 	bool bIsAudioOnly = Manifest->AudioOnlyStreams.Num() && Manifest->VariantStreams.Num() == 0;
 	const TArray<FManifestHLSInternal::FMediaStream::FMediaSegment>& SegmentList = Stream->SegmentList;
 	check(SegmentList.Num());
-	const FParamDict& Options = PlayerSessionServices->GetOptions();
 	if (SegmentList.Num())
 	{
 		EarliestTimeline = SegmentList[0].AbsoluteDateTime;
@@ -1769,7 +1768,7 @@ void FManifestBuilderHLS::UpdateManifestMetadataFromStream(FManifestHLSInternal*
 		}
 		else
 		{
-			FTimeValue FirstOffset = Options.GetValue(IPlaylistReaderHLS::OptionKeyLiveSeekableStartOffset).SafeGetTimeValue(FTimeValue::GetZero());
+			FTimeValue FirstOffset = PlayerSessionServices->GetOptionValue(IPlaylistReaderHLS::OptionKeyLiveSeekableStartOffset).SafeGetTimeValue(FTimeValue::GetZero());
 			for(int32 i=0, iMax=SegmentList.Num(); i<iMax; ++i)
 			{
 				if (SegmentList[i].AbsoluteDateTime >= FirstOffset)
@@ -1798,15 +1797,15 @@ void FManifestBuilderHLS::UpdateManifestMetadataFromStream(FManifestHLSInternal*
 			if (bIsAudioOnly)
 			{
 				// First try the audio-only config.
-				if (Options.HaveKey(IPlaylistReaderHLS::OptionKeyLiveSeekableEndOffsetAudioOnly))
+				if (PlayerSessionServices->HaveOptionValue(IPlaylistReaderHLS::OptionKeyLiveSeekableEndOffsetAudioOnly))
 				{
-					LastOffset = Options.GetValue(IPlaylistReaderHLS::OptionKeyLiveSeekableEndOffsetAudioOnly).SafeGetTimeValue(ThreeTargetDurations);
+					LastOffset = PlayerSessionServices->GetOptionValue(IPlaylistReaderHLS::OptionKeyLiveSeekableEndOffsetAudioOnly).SafeGetTimeValue(ThreeTargetDurations);
 				}
 			}
 			// Try the video config value regardless of presentation type.
-			if (!LastOffset.IsValid() && Options.HaveKey(OptionKeyLiveSeekableEndOffset))
+			if (!LastOffset.IsValid() && PlayerSessionServices->HaveOptionValue(OptionKeyLiveSeekableEndOffset))
 			{
-				LastOffset = Options.GetValue(OptionKeyLiveSeekableEndOffset).SafeGetTimeValue(ThreeTargetDurations);
+				LastOffset = PlayerSessionServices->GetOptionValue(OptionKeyLiveSeekableEndOffset).SafeGetTimeValue(ThreeTargetDurations);
 			}
 			// If still not valid use the default from the HLS specification.
 			if (!LastOffset.IsValid())
@@ -1817,7 +1816,7 @@ void FManifestBuilderHLS::UpdateManifestMetadataFromStream(FManifestHLSInternal*
 			// Now back off from the end until we have reached at least the required offset.
 			FTimeValue BackedOff(FTimeValue::GetZero());
 			FTimeValue PrevBackedOff;
-			bool bUseConservativeDistance = Options.GetValue(IPlaylistReaderHLS::OptionKeyLiveSeekableEndOffsetBeConservative).SafeGetBool(false);
+			bool bUseConservativeDistance = PlayerSessionServices->GetOptionValue(IPlaylistReaderHLS::OptionKeyLiveSeekableEndOffsetBeConservative).SafeGetBool(false);
 			for(int32 i=SegmentList.Num()-1; i>=0; --i)
 			{
 				BackedOff += SegmentList[i].Duration;
diff --git a/Engine/Plugins/Media/ElectraPlayer/Source/ElectraPlayerRuntime/Private/Runtime/Player/HLS/PlaylistReaderHLS.cpp b/Engine/Plugins/Media/ElectraPlayer/Source/ElectraPlayerRuntime/Private/Runtime/Player/HLS/PlaylistReaderHLS.cpp
--- a/Engine/Plugins/Media/ElectraPlayer/Source/ElectraPlayerRuntime/Private/Runtime/Player/HLS/PlaylistReaderHLS.cpp
+++ b/Engine/Plugins/Media/ElectraPlayer/Source/ElectraPlayerRuntime/Private/Runtime/Player/HLS/PlaylistReaderHLS.cpp
@@ -659,23 +659,22 @@ void FPlaylistReaderHLS::FPlaylistRequest::Execute(TSharedPtrTS<IElectraHttpMana
 	HTTPRequest->ConnectionInfo.RetryInfo = RetryInfo;
 
 	// Set connection timeouts for master and variant playlist retrieval.
-	const FParamDict& Options =	InPlayerSessionServices->GetOptions();
 	if (bIsMasterPlaylist)
 	{
-		HTTPRequest->Parameters.ConnectTimeout = Options.GetValue(IPlaylistReaderHLS::OptionKeyMasterPlaylistLoadConnectTimeout).SafeGetTimeValue(FTimeValue().SetFromMilliseconds(1000 * 8));
-		HTTPRequest->Parameters.NoDataTimeout  = Options.GetValue(IPlaylistReaderHLS::OptionKeyMasterPlaylistLoadNoDataTimeout).SafeGetTimeValue(FTimeValue().SetFromMilliseconds(1000 * 5));
+		HTTPRequest->Parameters.ConnectTimeout = InPlayerSessionServices->GetOptionValue(IPlaylistReaderHLS::OptionKeyMasterPlaylistLoadConnectTimeout).SafeGetTimeValue(FTimeValue().SetFromMilliseconds(1000 * 8));
+		HTTPRequest->Parameters.NoDataTimeout  = InPlayerSessionServices->GetOptionValue(IPlaylistReaderHLS::OptionKeyMasterPlaylistLoadNoDataTimeout).SafeGetTimeValue(FTimeValue().SetFromMilliseconds(1000 * 5));
 	}
 	else
 	{
 		if (PlaylistLoadRequest.LoadType == FPlaylistLoadRequestHLS::ELoadType::Update)
 		{
-			HTTPRequest->Parameters.ConnectTimeout = Options.GetValue(IPlaylistReaderHLS::OptionKeyUpdatePlaylistLoadConnectTimeout).SafeGetTimeValue(FTimeValue().SetFromMilliseconds(1000 * 2));
-			HTTPRequest->Parameters.NoDataTimeout  = Options.GetValue(IPlaylistReaderHLS::OptionKeyUpdatePlaylistLoadNoDataTimeout).SafeGetTimeValue(FTimeValue().SetFromMilliseconds(1000 * 2));
+			HTTPRequest->Parameters.ConnectTimeout = InPlayerSessionServices->GetOptionValue(IPlaylistReaderHLS::OptionKeyUpdatePlaylistLoadConnectTimeout).SafeGetTimeValue(FTimeValue().SetFromMilliseconds(1000 * 2));
+			HTTPRequest->Parameters.NoDataTimeout  = InPlayerSessionServices->GetOptionValue(IPlaylistReaderHLS::OptionKeyUpdatePlaylistLoadNoDataTimeout).SafeGetTimeValue(FTimeValue().SetFromMilliseconds(1000 * 2));
 		}
 		else
 		{
-			HTTPRequest->Parameters.ConnectTimeout = Options.GetValue(IPlaylistReaderHLS::OptionKeyVariantPlaylistLoadConnectTimeout).SafeGetTimeValue(FTimeValue().SetFromMilliseconds(1000 * 5));
-			HTTPRequest->Parameters.NoDataTimeout  = Options.GetValue(IPlaylistReaderHLS::OptionKeyVariantPlaylistLoadNoDataTimeout).SafeGetTimeValue(FTimeValue().SetFromMilliseconds(1000 * 2));
+			HTTPRequest->Parameters.ConnectTimeout = InPlayerSessionServices->GetOptionValue(IPlaylistReaderHLS::OptionKeyVariantPlaylistLoadConnectTimeout).SafeGetTimeValue(FTimeValue().SetFromMilliseconds(1000 * 5));
+			HTTPRequest->Parameters.NoDataTimeout  = InPlayerSessionServices->GetOptionValue(IPlaylistReaderHLS::OptionKeyVariantPlaylistLoadNoDataTimeout).SafeGetTimeValue(FTimeValue().SetFromMilliseconds(1000 * 2));
 		}
 	}
 
diff --git a/Engine/Plugins/Media/ElectraPlayer/Source/ElectraPlayerRuntime/Private/Runtime/Player/PlayerEntityCache.cpp b/Engine/Plugins/Media/ElectraPlayer/Source/ElectraPlayerRuntime/Private/Runtime/Player/PlayerEntityCache.cpp
--- a/Engine/Plugins/Media/ElectraPlayer/Source/ElectraPlayerRuntime/Private/Runtime/Player/PlayerEntityCache.cpp
+++ b/Engine/Plugins/Media/ElectraPlayer/Source/ElectraPlayerRuntime/Private/Runtime/Player/PlayerEntityCache.cpp
@@ -10,7 +10,7 @@ namespace Electra
 class FPlayerEntityCache : public IPlayerEntityCache
 {
 public:
-	FPlayerEntityCache(IPlayerSessionServices* SessionServices, const FParamDict& Options);
+	FPlayerEntityCache(IPlayerSessionServices* SessionServices);
 	virtual ~FPlayerEntityCache();
 
 	virtual void HandleEntityExpiration() override;
@@ -28,7 +28,6 @@ private:
 
 	static const uint32 kMaxCacheEntries = 256;
 
-	FParamDict Options;
 	IPlayerSessionServices* SessionServices;
 
 	mutable FCriticalSection Lock;
@@ -40,18 +39,17 @@ private:
 /***************************************************************************************************************************************************/
 /***************************************************************************************************************************************************/
 
-TSharedPtrTS<IPlayerEntityCache> IPlayerEntityCache::Create(IPlayerSessionServices* SessionServices, const FParamDict& Options)
+TSharedPtrTS<IPlayerEntityCache> IPlayerEntityCache::Create(IPlayerSessionServices* SessionServices)
 {
-	return MakeSharedTS<FPlayerEntityCache>(SessionServices, Options);
+	return MakeSharedTS<FPlayerEntityCache>(SessionServices);
 }
 
 /***************************************************************************************************************************************************/
 /***************************************************************************************************************************************************/
 /***************************************************************************************************************************************************/
 
-FPlayerEntityCache::FPlayerEntityCache(IPlayerSessionServices* InSessionServices, const FParamDict& InOptions)
-	: Options(InOptions)
-	, SessionServices(InSessionServices)
+FPlayerEntityCache::FPlayerEntityCache(IPlayerSessionServices* InSessionServices)
+	: SessionServices(InSessionServices)
 	, Cache(kMaxCacheEntries)
 {
 }
diff --git a/Engine/Plugins/Media/ElectraPlayer/Source/ElectraPlayerRuntime/Private/Runtime/Player/PlayerEntityCache.h b/Engine/Plugins/Media/ElectraPlayer/Source/ElectraPlayerRuntime/Private/Runtime/Player/PlayerEntityCache.h
--- a/Engine/Plugins/Media/ElectraPlayer/Source/ElectraPlayerRuntime/Private/Runtime/Player/PlayerEntityCache.h
+++ b/Engine/Plugins/Media/ElectraPlayer/Source/ElectraPlayerRuntime/Private/Runtime/Player/PlayerEntityCache.h
@@ -37,7 +37,7 @@ public:
 		TSharedPtrTS<const TArray<uint8>> RawPayloadData;
 	};
 
-	static TSharedPtrTS<IPlayerEntityCache> Create(IPlayerSessionServices* SessionServices, const FParamDict& Options);
+	static TSharedPtrTS<IPlayerEntityCache> Create(IPlayerSessionServices* SessionServices);
 
 	virtual ~IPlayerEntityCache() = default;
 
diff --git a/Engine/Plugins/Media/ElectraPlayer/Source/ElectraPlayerRuntime/Private/Runtime/Player/PlayerSessionServices.h b/Engine/Plugins/Media/ElectraPlayer/Source/ElectraPlayerRuntime/Private/Runtime/Player/PlayerSessionServices.h
--- a/Engine/Plugins/Media/ElectraPlayer/Source/ElectraPlayerRuntime/Private/Runtime/Player/PlayerSessionServices.h
+++ b/Engine/Plugins/Media/ElectraPlayer/Source/ElectraPlayerRuntime/Private/Runtime/Player/PlayerSessionServices.h
@@ -138,7 +138,17 @@ namespace Electra
 		 * Returns the mutable player option dictionary. Other than initial options this serves as an interface to
 		 * pass values between internal sub systems.
 		 */
-		 virtual FParamDict& GetOptions() = 0;
+		virtual FParamDictTS& GetMutableOptions() = 0;
+
+		/**
+		 * Checks if a certain player option has been set.
+		 */
+		virtual bool HaveOptionValue(const FName& InOption) = 0;
+
+		/**
+		 * Returns a player option value.
+		 */
+		virtual const FVariantValue GetOptionValue(const FName& InOption) = 0;
 
 		 /**
 		  * Returns the DRM manager, if any.
diff --git a/Engine/Plugins/Media/ElectraPlayer/Source/ElectraPlayerRuntime/Private/Runtime/Player/SimpleElectraAudioPlayer.cpp b/Engine/Plugins/Media/ElectraPlayer/Source/ElectraPlayerRuntime/Private/Runtime/Player/SimpleElectraAudioPlayer.cpp
--- a/Engine/Plugins/Media/ElectraPlayer/Source/ElectraPlayerRuntime/Private/Runtime/Player/SimpleElectraAudioPlayer.cpp
+++ b/Engine/Plugins/Media/ElectraPlayer/Source/ElectraPlayerRuntime/Private/Runtime/Player/SimpleElectraAudioPlayer.cpp
@@ -843,7 +843,7 @@ public:
 	{ return BufferPoolProperties; }
 	UEMediaError CreateBufferPool(const FParamDict& Parameters) override;
 	UEMediaError AcquireBuffer(IBuffer*& OutBuffer, int32 TimeoutInMicroseconds, const FParamDict& InParameters) override;
-	UEMediaError ReturnBuffer(IBuffer* Buffer, bool bRender, const FParamDict& InSampleProperties) override;
+	UEMediaError ReturnBuffer(IBuffer* Buffer, bool bRender, FParamDict& InOutSampleProperties) override;
 	UEMediaError ReleaseBufferPool() override;
 	bool CanReceiveOutputFrames(uint64 NumFrames) const override;
 	bool GetEnqueuedFrameInfo(int32& OutNumberOfEnqueuedFrames, FTimeValue& OutDurationOfEnqueuedFrames) const override;
@@ -1947,7 +1947,7 @@ UEMediaError FSimpleElectraAudioPlayerRenderer::AcquireBuffer(IBuffer*& OutBuffe
 	return UEMEDIA_ERROR_OK;
 }
 
-UEMediaError FSimpleElectraAudioPlayerRenderer::ReturnBuffer(IBuffer* Buffer, bool bRender, const FParamDict& InSampleProperties)
+UEMediaError FSimpleElectraAudioPlayerRenderer::ReturnBuffer(IBuffer* Buffer, bool bRender, FParamDict& InOutSampleProperties)
 {
 	if (Buffer == nullptr)
 	{
@@ -1956,14 +1956,14 @@ UEMediaError FSimpleElectraAudioPlayerRenderer::ReturnBuffer(IBuffer* Buffer, bo
 	FMediaBufferSharedPtrWrapper* MediaBufferSharedPtrWrapper = static_cast<FMediaBufferSharedPtrWrapper*>(Buffer);
 	if (bRender)
 	{
-		int32 NumChannels = (int32)InSampleProperties.GetValue(RenderOptionKeys::NumChannels).SafeGetInt64(-1);
-		int32 SampleRate = (int32)InSampleProperties.GetValue(RenderOptionKeys::SampleRate).SafeGetInt64(-1);
-		int32 UsedBufferBytes = (int32)InSampleProperties.GetValue(RenderOptionKeys::UsedByteSize).SafeGetInt64(-1);
-		FTimeValue decPTS = InSampleProperties.GetValue(RenderOptionKeys::PTS).SafeGetTimeValue();
+		int32 NumChannels = (int32)InOutSampleProperties.GetValue(RenderOptionKeys::NumChannels).SafeGetInt64(-1);
+		int32 SampleRate = (int32)InOutSampleProperties.GetValue(RenderOptionKeys::SampleRate).SafeGetInt64(-1);
+		int32 UsedBufferBytes = (int32)InOutSampleProperties.GetValue(RenderOptionKeys::UsedByteSize).SafeGetInt64(-1);
+		FTimeValue decPTS = InOutSampleProperties.GetValue(RenderOptionKeys::PTS).SafeGetTimeValue();
 		FTimespan PTS = decPTS.GetAsTimespan();
 		int64 InSequenceIndex = decPTS.GetSequenceIndex();
 
-		FTimespan Duration = InSampleProperties.GetValue(RenderOptionKeys::Duration).SafeGetTimeValue().GetAsTimespan();
+		FTimespan Duration = InOutSampleProperties.GetValue(RenderOptionKeys::Duration).SafeGetTimeValue().GetAsTimespan();
 		const void* BufferAddr = MediaBufferSharedPtrWrapper->BufferProperties.GetValue(RenderOptionKeys::AllocatedAddress).SafeGetPointer();
 		if (NumChannels <= 0 || SampleRate <= 0 || UsedBufferBytes <= 0)
 		{
@@ -1971,7 +1971,7 @@ UEMediaError FSimpleElectraAudioPlayerRenderer::ReturnBuffer(IBuffer* Buffer, bo
 		}
 
 		IAudioDecoderOutputPtr DecoderOutput = MediaBufferSharedPtrWrapper->DecoderOutput;
-		DecoderOutput->GetMutablePropertyDictionary() = InSampleProperties;
+		DecoderOutput->GetMutablePropertyDictionary() = InOutSampleProperties;
 		DecoderOutput->Initialize(IAudioDecoderOutput::ESampleFormat::Float, NumChannels, SampleRate, Duration, FDecoderTimeStamp(PTS, InSequenceIndex), UsedBufferBytes);
 
 		TSharedPtr<FSimpleElectraAudioPlayer, ESPMode::ThreadSafe> PinnedPlayer = Player.Pin();
@@ -1982,7 +1982,7 @@ UEMediaError FSimpleElectraAudioPlayerRenderer::ReturnBuffer(IBuffer* Buffer, bo
 	}
 	else
 	{
-		if (InSampleProperties.GetValue(RenderOptionKeys::EOSFlag).SafeGetBool(false))
+		if (InOutSampleProperties.GetValue(RenderOptionKeys::EOSFlag).SafeGetBool(false))
 		{
 			TSharedPtr<FSimpleElectraAudioPlayer, ESPMode::ThreadSafe> PinnedPlayer = Player.Pin();
 			if (PinnedPlayer.IsValid())
diff --git a/Engine/Plugins/Media/ElectraPlayer/Source/ElectraPlayerRuntime/Private/Runtime/Player/mkv/PlaylistReaderMKV.cpp b/Engine/Plugins/Media/ElectraPlayer/Source/ElectraPlayerRuntime/Private/Runtime/Player/mkv/PlaylistReaderMKV.cpp
--- a/Engine/Plugins/Media/ElectraPlayer/Source/ElectraPlayerRuntime/Private/Runtime/Player/mkv/PlaylistReaderMKV.cpp
+++ b/Engine/Plugins/Media/ElectraPlayer/Source/ElectraPlayerRuntime/Private/Runtime/Player/mkv/PlaylistReaderMKV.cpp
@@ -303,8 +303,6 @@ void FPlaylistReaderMKV::ReadChunk(uint8* DestinationBuffer, int64 InFromOffset,
 	ReceiveBuffer = MakeSharedTS<IElectraHttpManager::FReceiveBuffer>();
 	ReceiveBuffer->Buffer.SetExternalBuffer(DestinationBuffer, ChunkSize);
 
-	const FParamDict& Options = PlayerSessionServices->GetOptions();
-
 	Request = MakeSharedTS<IElectraHttpManager::FRequest>();
 	Request->Parameters.URL = PlaylistURL;
 	Request->Parameters.Range.SetStart(InFromOffset);
@@ -314,8 +312,8 @@ void FPlaylistReaderMKV::ReadChunk(uint8* DestinationBuffer, int64 InFromOffset,
 		LastByte = FileSize - 1;
 	}
 	Request->Parameters.Range.SetEndIncluding(LastByte);
-	Request->Parameters.ConnectTimeout = Options.GetValue(MKV::OptionKeyMKVLoadConnectTimeout).SafeGetTimeValue(FTimeValue().SetFromMilliseconds(1000 * 8));
-	Request->Parameters.NoDataTimeout = Options.GetValue(MKV::OptionKeyMKVLoadNoDataTimeout).SafeGetTimeValue(FTimeValue().SetFromMilliseconds(1000 * 6));
+	Request->Parameters.ConnectTimeout = PlayerSessionServices->GetOptionValue(MKV::OptionKeyMKVLoadConnectTimeout).SafeGetTimeValue(FTimeValue().SetFromMilliseconds(1000 * 8));
+	Request->Parameters.NoDataTimeout = PlayerSessionServices->GetOptionValue(MKV::OptionKeyMKVLoadNoDataTimeout).SafeGetTimeValue(FTimeValue().SetFromMilliseconds(1000 * 6));
 	Request->ReceiveBuffer = ReceiveBuffer;
 	Request->ProgressListener = ProgressListener;
 	Request->ResponseCache = PlayerSessionServices->GetHTTPResponseCache();
diff --git a/Engine/Plugins/Media/ElectraPlayer/Source/ElectraPlayerRuntime/Private/Runtime/Player/mkv/StreamReaderMKV.cpp b/Engine/Plugins/Media/ElectraPlayer/Source/ElectraPlayerRuntime/Private/Runtime/Player/mkv/StreamReaderMKV.cpp
--- a/Engine/Plugins/Media/ElectraPlayer/Source/ElectraPlayerRuntime/Private/Runtime/Player/mkv/StreamReaderMKV.cpp
+++ b/Engine/Plugins/Media/ElectraPlayer/Source/ElectraPlayerRuntime/Private/Runtime/Player/mkv/StreamReaderMKV.cpp
@@ -357,8 +357,6 @@ void FStreamReaderMKV::HandleRequest()
 	ReadBuffer.SetCurrentOffset(0);
 	ReadBuffer.SetEndOffset(Request->FileEndOffset);
 
-	const FParamDict& Options = PlayerSessionServices->GetOptions();
-
 	HTTPRequest = MakeSharedTS<IElectraHttpManager::FRequest>();
 	HTTPRequest->Parameters.URL = TimelineAsset->GetMediaURL();
 	HTTPRequest->Parameters.Range.Start = Request->FileStartOffset;
@@ -366,8 +364,8 @@ void FStreamReaderMKV::HandleRequest()
 	// No compression as this would not yield much with already compressed video/audio data.
 	HTTPRequest->Parameters.AcceptEncoding.Set(TEXT("identity"));
 	// Timeouts
-	HTTPRequest->Parameters.ConnectTimeout = Options.GetValue(MKV::OptionKeyMKVLoadConnectTimeout).SafeGetTimeValue(FTimeValue().SetFromMilliseconds(1000 * 8));
-	HTTPRequest->Parameters.NoDataTimeout = Options.GetValue(MKV::OptionKeyMKVLoadNoDataTimeout).SafeGetTimeValue(FTimeValue().SetFromMilliseconds(1000 * 6));
+	HTTPRequest->Parameters.ConnectTimeout = PlayerSessionServices->GetOptionValue(MKV::OptionKeyMKVLoadConnectTimeout).SafeGetTimeValue(FTimeValue().SetFromMilliseconds(1000 * 8));
+	HTTPRequest->Parameters.NoDataTimeout = PlayerSessionServices->GetOptionValue(MKV::OptionKeyMKVLoadNoDataTimeout).SafeGetTimeValue(FTimeValue().SetFromMilliseconds(1000 * 6));
 
 	// Explicit range?
 	int64 NumRequestedBytes = HTTPRequest->Parameters.Range.GetNumberOfBytes();
diff --git a/Engine/Plugins/Media/ElectraPlayer/Source/ElectraPlayerRuntime/Private/Runtime/Player/mp4/ManifestMP4.cpp b/Engine/Plugins/Media/ElectraPlayer/Source/ElectraPlayerRuntime/Private/Runtime/Player/mp4/ManifestMP4.cpp
--- a/Engine/Plugins/Media/ElectraPlayer/Source/ElectraPlayerRuntime/Private/Runtime/Player/mp4/ManifestMP4.cpp
+++ b/Engine/Plugins/Media/ElectraPlayer/Source/ElectraPlayerRuntime/Private/Runtime/Player/mp4/ManifestMP4.cpp
@@ -61,9 +61,8 @@ FManifestMP4Internal::~FManifestMP4Internal()
  *
  * @return
  */
-FErrorDetail FManifestMP4Internal::Build(TSharedPtrTS<IParserISO14496_12> MP4Parser, const FString& URL, const HTTP::FConnectionInfo& InConnectionInfo)
+FErrorDetail FManifestMP4Internal::Build(TSharedPtrTS<IParserISO14496_12> MP4Parser, const FString& URL, const HTTP::FConnectionInfo& /*InConnectionInfo*/)
 {
-	ConnectionInfo = InConnectionInfo;
 	MediaAsset = MakeSharedTS<FTimelineAssetMP4>();
 	FErrorDetail Result = MediaAsset->Build(PlayerSessionServices, MP4Parser, URL);
 	FTimeRange PlaybackRange = GetPlaybackRange();
diff --git a/Engine/Plugins/Media/ElectraPlayer/Source/ElectraPlayerRuntime/Private/Runtime/Player/mp4/ManifestMP4.h b/Engine/Plugins/Media/ElectraPlayer/Source/ElectraPlayerRuntime/Private/Runtime/Player/mp4/ManifestMP4.h
--- a/Engine/Plugins/Media/ElectraPlayer/Source/ElectraPlayerRuntime/Private/Runtime/Player/mp4/ManifestMP4.h
+++ b/Engine/Plugins/Media/ElectraPlayer/Source/ElectraPlayerRuntime/Private/Runtime/Player/mp4/ManifestMP4.h
@@ -328,7 +328,6 @@ public:
 
 	IPlayerSessionServices* 			PlayerSessionServices;
 	TSharedPtrTS<FTimelineAssetMP4>		MediaAsset;
-	HTTP::FConnectionInfo				ConnectionInfo;
 	// The MPD URL fragment components
 	TArray<FURL_RFC3986::FQueryParam>	URLFragmentComponents;
 	FTimeValue							DefaultStartTime;
diff --git a/Engine/Plugins/Media/ElectraPlayer/Source/ElectraPlayerRuntime/Private/Runtime/Player/mp4/PlaylistReaderMP4.cpp b/Engine/Plugins/Media/ElectraPlayer/Source/ElectraPlayerRuntime/Private/Runtime/Player/mp4/PlaylistReaderMP4.cpp
--- a/Engine/Plugins/Media/ElectraPlayer/Source/ElectraPlayerRuntime/Private/Runtime/Player/mp4/PlaylistReaderMP4.cpp
+++ b/Engine/Plugins/Media/ElectraPlayer/Source/ElectraPlayerRuntime/Private/Runtime/Player/mp4/PlaylistReaderMP4.cpp
@@ -334,8 +334,6 @@ void FPlaylistReaderMP4::ReadNextChunk(int64 InFromOffset, int64 ChunkSize)
 	ReceiveBuffer = MakeSharedTS<IElectraHttpManager::FReceiveBuffer>();
 	ReceiveBuffer->Buffer.Reserve(ChunkSize);
 
-	const FParamDict& Options = PlayerSessionServices->GetOptions();
-
 	Request = MakeSharedTS<IElectraHttpManager::FRequest>();
 	Request->Parameters.URL = MasterPlaylistURL;
 	Request->Parameters.Range.SetStart(InFromOffset);
@@ -345,8 +343,8 @@ void FPlaylistReaderMP4::ReadNextChunk(int64 InFromOffset, int64 ChunkSize)
 		LastByte = FileSize - 1;
 	}
 	Request->Parameters.Range.SetEndIncluding(LastByte);
-	Request->Parameters.ConnectTimeout = Options.GetValue(MP4::OptionKeyMP4LoadConnectTimeout).SafeGetTimeValue(FTimeValue().SetFromMilliseconds(1000 * 8));
-	Request->Parameters.NoDataTimeout = Options.GetValue(MP4::OptionKeyMP4LoadNoDataTimeout).SafeGetTimeValue(FTimeValue().SetFromMilliseconds(1000 * 6));
+	Request->Parameters.ConnectTimeout = PlayerSessionServices->GetOptionValue(MP4::OptionKeyMP4LoadConnectTimeout).SafeGetTimeValue(FTimeValue().SetFromMilliseconds(1000 * 8));
+	Request->Parameters.NoDataTimeout = PlayerSessionServices->GetOptionValue(MP4::OptionKeyMP4LoadNoDataTimeout).SafeGetTimeValue(FTimeValue().SetFromMilliseconds(1000 * 6));
 	Request->ReceiveBuffer = ReceiveBuffer;
 	Request->ProgressListener = ProgressListener;
 	Request->ResponseCache = PlayerSessionServices->GetHTTPResponseCache();
diff --git a/Engine/Plugins/Media/ElectraPlayer/Source/ElectraPlayerRuntime/Private/Runtime/Player/mp4/StreamReaderMP4.cpp b/Engine/Plugins/Media/ElectraPlayer/Source/ElectraPlayerRuntime/Private/Runtime/Player/mp4/StreamReaderMP4.cpp
--- a/Engine/Plugins/Media/ElectraPlayer/Source/ElectraPlayerRuntime/Private/Runtime/Player/mp4/StreamReaderMP4.cpp
+++ b/Engine/Plugins/Media/ElectraPlayer/Source/ElectraPlayerRuntime/Private/Runtime/Player/mp4/StreamReaderMP4.cpp
@@ -240,14 +240,12 @@ bool FStreamReaderMP4::HasErrored() const
 
 int32 FStreamReaderMP4::HTTPProgressCallback(const IElectraHttpManager::FRequest* InRequest)
 {
-	HTTPUpdateStats(MEDIAutcTime::Current(), InRequest);
 	// Aborted?
 	return HasBeenAborted() ? 1 : 0;
 }
 
 void FStreamReaderMP4::HTTPCompletionCallback(const IElectraHttpManager::FRequest* InRequest)
 {
-	HTTPUpdateStats(FTimeValue::GetInvalid(), InRequest);
 	bHasErrored = InRequest->ConnectionInfo.StatusInfo.ErrorDetail.IsError();
 	if (bHasErrored)
 	{
@@ -255,27 +253,6 @@ void FStreamReaderMP4::HTTPCompletionCallback(const IElectraHttpManager::FReques
 	}
 }
 
-void FStreamReaderMP4::HTTPUpdateStats(const FTimeValue& CurrentTime, const IElectraHttpManager::FRequest* Request)
-{
-	TSharedPtrTS<FStreamSegmentRequestMP4> SegmentRequest = CurrentRequest;
-	if (SegmentRequest.IsValid())
-	{
-		FMediaCriticalSection::ScopedLock lock(MetricUpdateLock);
-		SegmentRequest->ConnectionInfo = Request->ConnectionInfo;
-		// Update the current download stats which we report periodically to the ABR.
-		Metrics::FSegmentDownloadStats& ds = SegmentRequest->DownloadStats;
-		if (Request->ConnectionInfo.EffectiveURL.Len())
-		{
-			ds.URL 			  = Request->ConnectionInfo.EffectiveURL;
-		}
-		ds.HTTPStatusCode     = Request->ConnectionInfo.StatusInfo.HTTPStatus;
-		ds.TimeToFirstByte    = Request->ConnectionInfo.TimeUntilFirstByte;
-		ds.TimeToDownload     = ((CurrentTime.IsValid() ? CurrentTime : Request->ConnectionInfo.RequestEndTime) - Request->ConnectionInfo.RequestStartTime).GetAsSeconds();
-		ds.ByteSize 		  = Request->ConnectionInfo.ContentLength;
-		ds.NumBytesDownloaded = Request->ConnectionInfo.BytesReadSoFar;
-	}
-}
-
 void FStreamReaderMP4::HandleRequest()
 {
 	TSharedPtrTS<FStreamSegmentRequestMP4>	Request = CurrentRequest;
@@ -385,8 +362,6 @@ void FStreamReaderMP4::HandleRequest()
 	ReadBuffer.ReceiveBuffer = MakeSharedTS<IElectraHttpManager::FReceiveBuffer>();
 	ReadBuffer.SetCurrentPos(Request->FileStartOffset);
 
-	const FParamDict& Options = PlayerSessionServices->GetOptions();
-
 	TSharedPtrTS<IElectraHttpManager::FRequest> HTTP(new IElectraHttpManager::FRequest);
 	HTTP->Parameters.URL				= TimelineAsset->GetMediaURL();
 	HTTP->Parameters.Range.Start		= Request->FileStartOffset;
@@ -394,8 +369,8 @@ void FStreamReaderMP4::HandleRequest()
 	// No compression as this would not yield much with already compressed video/audio data.
 	HTTP->Parameters.AcceptEncoding.Set(TEXT("identity"));
 	// Timeouts
-	HTTP->Parameters.ConnectTimeout = Options.GetValue(MP4::OptionKeyMP4LoadConnectTimeout).SafeGetTimeValue(FTimeValue().SetFromMilliseconds(1000 * 8));
-	HTTP->Parameters.NoDataTimeout = Options.GetValue(MP4::OptionKeyMP4LoadNoDataTimeout).SafeGetTimeValue(FTimeValue().SetFromMilliseconds(1000 * 6));
+	HTTP->Parameters.ConnectTimeout = PlayerSessionServices->GetOptionValue(MP4::OptionKeyMP4LoadConnectTimeout).SafeGetTimeValue(FTimeValue().SetFromMilliseconds(1000 * 8));
+	HTTP->Parameters.NoDataTimeout = PlayerSessionServices->GetOptionValue(MP4::OptionKeyMP4LoadNoDataTimeout).SafeGetTimeValue(FTimeValue().SetFromMilliseconds(1000 * 6));
 
 	// Explicit range?
 	int64 NumRequestedBytes = HTTP->Parameters.Range.GetNumberOfBytes();
@@ -704,24 +679,22 @@ void FStreamReaderMP4::HandleRequest()
 	DurationSuccessfullyRead	  = PrimaryTrack.DurationSuccessfullyRead;
 	DurationSuccessfullyDelivered = PrimaryTrack.DurationSuccessfullyDelivered;
 
-	// Set up remaining download stat fields.
-// Note: currently commented out because of UE-88612.
-//       This must be reinstated once we set failure reasons in the loop above so they won't get replaced by this!
-//		if (ds.FailureReason.length() == 0)
+	// Set up download stat fields.
+	if (Request->ConnectionInfo.EffectiveURL.Len())
 	{
-		ds.FailureReason = Request->ConnectionInfo.StatusInfo.ErrorDetail.GetMessage();
+		ds.URL = Request->ConnectionInfo.EffectiveURL;
 	}
+	ds.HTTPStatusCode = Request->ConnectionInfo.StatusInfo.HTTPStatus;
+	ds.TimeToFirstByte = Request->ConnectionInfo.TimeUntilFirstByte;
+	ds.TimeToDownload = (Request->ConnectionInfo.RequestEndTime - Request->ConnectionInfo.RequestStartTime).GetAsSeconds();
+	ds.ByteSize = Request->ConnectionInfo.ContentLength;
+	ds.NumBytesDownloaded = Request->ConnectionInfo.BytesReadSoFar;
+
+	ds.FailureReason = Request->ConnectionInfo.StatusInfo.ErrorDetail.GetMessage();
 	if (ParsingErrorMessage.Len())
 	{
 		ds.FailureReason = ParsingErrorMessage;
 	}
-//		if (bAbortedByABR)
-//		{
-//			// If aborted set the reason as the download failure.
-//			ds.FailureReason = ds.ABRState.ProgressDecision.Reason;
-//		}
-//		ds.bWasAborted  	  = bAbortedByABR;
-//		ds.bWasSuccessful     = !bHasErrored && !bAbortedByABR;
 	ds.bWasSuccessful     = !bHasErrored;
 	ds.URL  			  = Request->ConnectionInfo.EffectiveURL;
 	ds.HTTPStatusCode     = Request->ConnectionInfo.StatusInfo.HTTPStatus;
diff --git a/Engine/Plugins/Media/ElectraPlayer/Source/ElectraPlayerRuntime/Private/Runtime/Player/mp4/StreamReaderMP4.h b/Engine/Plugins/Media/ElectraPlayer/Source/ElectraPlayerRuntime/Private/Runtime/Player/mp4/StreamReaderMP4.h
--- a/Engine/Plugins/Media/ElectraPlayer/Source/ElectraPlayerRuntime/Private/Runtime/Player/mp4/StreamReaderMP4.h
+++ b/Engine/Plugins/Media/ElectraPlayer/Source/ElectraPlayerRuntime/Private/Runtime/Player/mp4/StreamReaderMP4.h
@@ -108,7 +108,6 @@ private:
 	void LogMessage(IInfoLog::ELevel Level, const FString& Message);
 	int32 HTTPProgressCallback(const IElectraHttpManager::FRequest* InRequest);
 	void HTTPCompletionCallback(const IElectraHttpManager::FRequest* InRequest);
-	void HTTPUpdateStats(const FTimeValue& CurrentTime, const IElectraHttpManager::FRequest* Request);
 	void HandleRequest();
 
 	bool HasBeenAborted() const;
@@ -187,7 +186,6 @@ private:
 	TSharedPtrTS<FStreamSegmentRequestMP4>		CurrentRequest;
 	FMediaEvent									WorkSignal;
 	FReadBuffer									ReadBuffer;
-	FMediaCriticalSection						MetricUpdateLock;
 
 	TMap<uint32, FSelectedTrackData>			ActiveTrackMap;
 
diff --git a/Engine/Plugins/Media/ElectraPlayer/Source/ElectraPlayerRuntime/Private/Runtime/Renderer/RendererAudio.cpp b/Engine/Plugins/Media/ElectraPlayer/Source/ElectraPlayerRuntime/Private/Runtime/Renderer/RendererAudio.cpp
--- a/Engine/Plugins/Media/ElectraPlayer/Source/ElectraPlayerRuntime/Private/Runtime/Renderer/RendererAudio.cpp
+++ b/Engine/Plugins/Media/ElectraPlayer/Source/ElectraPlayerRuntime/Private/Runtime/Renderer/RendererAudio.cpp
@@ -126,7 +126,7 @@ UEMediaError FElectraRendererAudio::AcquireBuffer(IBuffer*& OutBuffer, int32 Tim
 
 
 
-UEMediaError FElectraRendererAudio::ReturnBuffer(IBuffer* Buffer, bool bRender, const FParamDict& InSampleProperties)
+UEMediaError FElectraRendererAudio::ReturnBuffer(IBuffer* Buffer, bool bRender, FParamDict& InOutSampleProperties)
 {
 	if (Buffer == nullptr)
 	{
@@ -137,31 +137,31 @@ UEMediaError FElectraRendererAudio::ReturnBuffer(IBuffer* Buffer, bool bRender,
 
 	if (bRender)
 	{
-		const FVariantValue& variantNumChannels = InSampleProperties.GetValue(RenderOptionKeys::NumChannels);
+		const FVariantValue& variantNumChannels = InOutSampleProperties.GetValue(RenderOptionKeys::NumChannels);
 		if (!variantNumChannels.IsType(FVariantValue::EDataType::TypeInt64))
 		{
 			return UEMEDIA_ERROR_BAD_ARGUMENTS;
 		}
 
-		const FVariantValue& variantSampleRate = InSampleProperties.GetValue(RenderOptionKeys::SampleRate);
+		const FVariantValue& variantSampleRate = InOutSampleProperties.GetValue(RenderOptionKeys::SampleRate);
 		if (!variantSampleRate.IsType(FVariantValue::EDataType::TypeInt64))
 		{
 			return UEMEDIA_ERROR_BAD_ARGUMENTS;
 		}
 
-		const FVariantValue& variantBufferUsedBytes = InSampleProperties.GetValue(RenderOptionKeys::UsedByteSize);
+		const FVariantValue& variantBufferUsedBytes = InOutSampleProperties.GetValue(RenderOptionKeys::UsedByteSize);
 		if (!variantBufferUsedBytes.IsType(FVariantValue::EDataType::TypeInt64))
 		{
 			return UEMEDIA_ERROR_BAD_ARGUMENTS;
 		}
 
-		const FVariantValue& variantPts = InSampleProperties.GetValue(RenderOptionKeys::PTS);
+		const FVariantValue& variantPts = InOutSampleProperties.GetValue(RenderOptionKeys::PTS);
 		if (!variantPts.IsType(FVariantValue::EDataType::TypeTimeValue))
 		{
 			return UEMEDIA_ERROR_BAD_ARGUMENTS;
 		}
 
-		const FVariantValue& variantDuration = InSampleProperties.GetValue(RenderOptionKeys::Duration);
+		const FVariantValue& variantDuration = InOutSampleProperties.GetValue(RenderOptionKeys::Duration);
 		if (!variantDuration.IsType(FVariantValue::EDataType::TypeTimeValue))
 		{
 			return UEMEDIA_ERROR_BAD_ARGUMENTS;
@@ -179,7 +179,7 @@ UEMediaError FElectraRendererAudio::ReturnBuffer(IBuffer* Buffer, bool bRender,
 
 		//UE_LOG(LogElectraPlayer, VeryVerbose, TEXT("-- FElectraRendererAudio::ReturnBuffer: Audio sample for time %s"), *InPts.ToString(TEXT("%h:%m:%s.%f")));
 
-		DecoderOutput->GetMutablePropertyDictionary() = InSampleProperties;
+		DecoderOutput->GetMutablePropertyDictionary() = InOutSampleProperties;
 		DecoderOutput->Initialize(IAudioDecoderOutput::ESampleFormat::Float, InNumChannels, InSampleRate, InDuration, FDecoderTimeStamp(InPts, InSequenceIndex), InUsedBufferBytes);
 
 		// Push buffer to output queue...
diff --git a/Engine/Plugins/Media/ElectraPlayer/Source/ElectraPlayerRuntime/Private/Runtime/Renderer/RendererAudio.h b/Engine/Plugins/Media/ElectraPlayer/Source/ElectraPlayerRuntime/Private/Runtime/Renderer/RendererAudio.h
--- a/Engine/Plugins/Media/ElectraPlayer/Source/ElectraPlayerRuntime/Private/Runtime/Renderer/RendererAudio.h
+++ b/Engine/Plugins/Media/ElectraPlayer/Source/ElectraPlayerRuntime/Private/Runtime/Renderer/RendererAudio.h
@@ -66,7 +66,7 @@ public:
 	UEMediaError AcquireBuffer(IBuffer*& OutBuffer, int32 TimeoutInMicroseconds, const Electra::FParamDict& InParameters) override;
 
 	// Releases the buffer for rendering and subsequent return to the buffer pool
-	UEMediaError ReturnBuffer(IBuffer* Buffer, bool bRender, const Electra::FParamDict& InSampleProperties) override;
+	UEMediaError ReturnBuffer(IBuffer* Buffer, bool bRender, Electra::FParamDict& InOutSampleProperties) override;
 
 	// Informs that the decoder is done with this pool. NO FREE!!!
 	UEMediaError ReleaseBufferPool() override;
diff --git a/Engine/Plugins/Media/ElectraPlayer/Source/ElectraPlayerRuntime/Private/Runtime/Renderer/RendererBase.h b/Engine/Plugins/Media/ElectraPlayer/Source/ElectraPlayerRuntime/Private/Runtime/Renderer/RendererBase.h
--- a/Engine/Plugins/Media/ElectraPlayer/Source/ElectraPlayerRuntime/Private/Runtime/Renderer/RendererBase.h
+++ b/Engine/Plugins/Media/ElectraPlayer/Source/ElectraPlayerRuntime/Private/Runtime/Renderer/RendererBase.h
@@ -131,7 +131,7 @@ namespace Electra
 		 *
 		 * @return Should be UEMEDIA_ERROR_OK only. Anything else will cause a playback error (like UEMEDIA_ERROR_BAD_ARGUMENTS for bad properties)
 		 */
-		virtual UEMediaError ReturnBuffer(IBuffer* Buffer, bool bRender, const FParamDict& InSampleProperties) = 0;
+		virtual UEMediaError ReturnBuffer(IBuffer* Buffer, bool bRender, FParamDict& InOutSampleProperties) = 0;
 
 		/**
 		 * Informs that the decoder is done with this pool. This only indicates no calls to
diff --git a/Engine/Plugins/Media/ElectraPlayer/Source/ElectraPlayerRuntime/Private/Runtime/Renderer/RendererVideo.cpp b/Engine/Plugins/Media/ElectraPlayer/Source/ElectraPlayerRuntime/Private/Runtime/Renderer/RendererVideo.cpp
--- a/Engine/Plugins/Media/ElectraPlayer/Source/ElectraPlayerRuntime/Private/Runtime/Renderer/RendererVideo.cpp
+++ b/Engine/Plugins/Media/ElectraPlayer/Source/ElectraPlayerRuntime/Private/Runtime/Renderer/RendererVideo.cpp
@@ -186,7 +186,7 @@ UEMediaError FElectraRendererVideo::AcquireBuffer(IBuffer*& OutBuffer, int32 Tim
 /**
  * Releases the buffer for rendering and subsequent return to the buffer pool
  */
-UEMediaError FElectraRendererVideo::ReturnBuffer(IBuffer* Buffer, bool bRender, const FParamDict& InSampleProperties)
+UEMediaError FElectraRendererVideo::ReturnBuffer(IBuffer* Buffer, bool bRender, FParamDict& InOutSampleProperties)
 {
 	if (Buffer == nullptr)
 	{
@@ -194,12 +194,12 @@ UEMediaError FElectraRendererVideo::ReturnBuffer(IBuffer* Buffer, bool bRender,
 	}
 
 	FMediaBufferSharedPtrWrapper* MediaBufferSharedPtrWrapper = static_cast<FMediaBufferSharedPtrWrapper*>(Buffer);
-	MediaBufferSharedPtrWrapper->DecoderOutput->GetMutablePropertyDictionary() = InSampleProperties;
+	MediaBufferSharedPtrWrapper->DecoderOutput->GetMutablePropertyDictionary() = InOutSampleProperties;
 
 	if (bRender)
 	{
-		//OPT/CHANGE: Note that "MediaBufferSharedPtrWrapper->DecoderOutput->GetDict()" is the very same as InSampleProperties!
-		bool bIsDummyBuffer = InSampleProperties.GetValue(RenderOptionKeys::DummyBufferFlag).SafeGetBool(false);
+		//OPT/CHANGE: Note that "MediaBufferSharedPtrWrapper->DecoderOutput->GetDict()" is the very same as InOutSampleProperties!
+		bool bIsDummyBuffer = InOutSampleProperties.GetValue(RenderOptionKeys::DummyBufferFlag).SafeGetBool(false);
 
 		// Put frame into output queue...
 		if (TSharedPtr<FElectraPlayer, ESPMode::ThreadSafe> PinnedPlayer = Player.Pin())
diff --git a/Engine/Plugins/Media/ElectraPlayer/Source/ElectraPlayerRuntime/Private/Runtime/Renderer/RendererVideo.h b/Engine/Plugins/Media/ElectraPlayer/Source/ElectraPlayerRuntime/Private/Runtime/Renderer/RendererVideo.h
--- a/Engine/Plugins/Media/ElectraPlayer/Source/ElectraPlayerRuntime/Private/Runtime/Renderer/RendererVideo.h
+++ b/Engine/Plugins/Media/ElectraPlayer/Source/ElectraPlayerRuntime/Private/Runtime/Renderer/RendererVideo.h
@@ -72,7 +72,7 @@ public:
 	UEMediaError AcquireBuffer(IBuffer*& OutBuffer, int32 TimeoutInMicroseconds, const Electra::FParamDict& InParameters) override;
 
 	// Releases the buffer for rendering and subsequent return to the buffer pool
-	UEMediaError ReturnBuffer(IBuffer* Buffer, bool bRender, const Electra::FParamDict& InSampleProperties) override;
+	UEMediaError ReturnBuffer(IBuffer* Buffer, bool bRender, Electra::FParamDict& InOutSampleProperties) override;
 
 	// Informs that the decoder is done with this pool. NO FREE!!!
 	UEMediaError ReleaseBufferPool() override;
diff --git a/Engine/Plugins/Media/ElectraUtil/Source/ElectraBase/Private/ParameterDictionary.cpp b/Engine/Plugins/Media/ElectraUtil/Source/ElectraBase/Private/ParameterDictionary.cpp
--- a/Engine/Plugins/Media/ElectraUtil/Source/ElectraBase/Private/ParameterDictionary.cpp
+++ b/Engine/Plugins/Media/ElectraUtil/Source/ElectraBase/Private/ParameterDictionary.cpp
@@ -411,10 +411,6 @@ const TArray<uint8>& FVariantValue::SafeGetArray() const
 
 
 
-FParamDict::FParamDict()
-{
-}
-
 FParamDict::FParamDict(const FParamDict& Other)
 {
 	InternalCopy(Other);
@@ -424,72 +420,47 @@ FParamDict& FParamDict::operator=(const FParamDict& Other)
 {
 	if (&Other != this)
 	{
-		FScopeLock lock(&Lock);
 		InternalCopy(Other);
 	}
 	return *this;
 }
 
-FParamDict::~FParamDict()
-{
-}
-
 void FParamDict::InternalCopy(const FParamDict& Other)
 {
-	FScopeLock lock(&Other.Lock);
 	Dictionary = Other.Dictionary;
 }
 
 void FParamDict::Clear()
 {
-	FScopeLock lock(&Lock);
 	Dictionary.Empty();
 }
 
 bool FParamDict::HaveKey(const FName& Key) const
 {
-	FScopeLock lock(&Lock);
 	return Dictionary.Find(Key) != nullptr;
 }
 
 FVariantValue FParamDict::GetValue(const FName& Key) const
 {
 	static FVariantValue Empty;
-	FScopeLock lock(&Lock);
 	const FVariantValue* VariantValue = Dictionary.Find(Key);
 	return VariantValue ? *VariantValue : Empty;
 }
 
 void FParamDict::Remove(const FName& Key)
 {
-	FScopeLock lock(&Lock);
 	Dictionary.Remove(Key); 
 }
 
 void FParamDict::Set(const FName& Key, const FVariantValue& Value)
 { 
-	FScopeLock lock(&Lock);
 	Dictionary.Emplace(Key, Value); 
 }
 
-void FParamDict::GetKeysStartingWith(const FString& StartsWith, TArray<FName>& OutKeys) const
+void FParamDict::GetKeys(TArray<FName>& OutKeys) const
 {
 	OutKeys.Empty();
-	FScopeLock lock(&Lock);
-	if (StartsWith.IsEmpty())
-	{
-		Dictionary.GenerateKeyArray(OutKeys);
-	}
-	else
-	{
-		for(const TPair<FName, FVariantValue>& Pair : Dictionary)
-		{
-			if (Pair.Key.ToString().StartsWith(StartsWith, ESearchCase::CaseSensitive))
-			{
-				OutKeys.Emplace(Pair.Key);
-			}
-		}
-	}
+	Dictionary.GenerateKeyArray(OutKeys);
 }
 
 void FParamDict::ConvertKeysStartingWithTo(TMap<FString, FVariant>& OutVariantMap, const FString& InKeyStartsWith, const FString& InAddPrefixToKey) const
@@ -497,7 +468,6 @@ void FParamDict::ConvertKeysStartingWithTo(TMap<FString, FVariant>& OutVariantMa
 	OutVariantMap.Reserve(Dictionary.Num());
 	FString NewKey;
 	NewKey.Reserve(64);
-	FScopeLock lock(&Lock);
 	for(const TPair<FName, FVariantValue>& Pair : Dictionary)
 	{
 		FString s(Pair.Key.ToString());
diff --git a/Engine/Plugins/Media/ElectraUtil/Source/ElectraBase/Public/ParameterDictionary.h b/Engine/Plugins/Media/ElectraUtil/Source/ElectraBase/Public/ParameterDictionary.h
--- a/Engine/Plugins/Media/ElectraUtil/Source/ElectraBase/Public/ParameterDictionary.h
+++ b/Engine/Plugins/Media/ElectraUtil/Source/ElectraBase/Public/ParameterDictionary.h
@@ -157,26 +157,101 @@ private:
 class ELECTRABASE_API FParamDict
 {
 public:
-	FParamDict();
+	FParamDict() {}
 	FParamDict(const FParamDict& Other);
 	FParamDict& operator=(const FParamDict& Other);
-	~FParamDict();
+	~FParamDict() = default;
 	void Clear();
 	void Set(const FName& Key, const FVariantValue& Value);
+	void GetKeys(TArray<FName>& OutKeys) const;
 	bool HaveKey(const FName& Key) const;
 	FVariantValue GetValue(const FName& Key) const;
 	void Remove(const FName& Key);
 
 	void ConvertTo(TMap<FString, FVariant>& OutVariantMap, const FString& InAddPrefixToKey) const;
 	void ConvertKeysStartingWithTo(TMap<FString, FVariant>& OutVariantMap, const FString& InKeyStartsWith, const FString& InAddPrefixToKey) const;
-
-	void GetKeysStartingWith(const FString& InStartsWith, TArray<FName>& OutKeys) const;
 private:
 	void InternalCopy(const FParamDict& Other);
-	mutable FCriticalSection Lock;
 	TMap<FName, FVariantValue> Dictionary;
 };
 
+class ELECTRABASE_API FParamDictTS
+{
+public:
+	FParamDictTS() {}
+	FParamDictTS(const FParamDictTS& Other)
+	{
+		Dictionary = Other.Dictionary;
+	}
+	FParamDictTS& operator=(const FParamDictTS& Other)
+	{
+		if (&Other != this)
+		{
+			FScopeLock lock(&Lock);
+			Dictionary = Other.Dictionary;
+		}
+		return *this;
+	}
+	FParamDictTS& operator=(const FParamDict& Other)
+	{
+		FScopeLock lock(&Lock);
+		Dictionary = Other;
+		return *this;
+	}
+
+	~FParamDictTS() = default;
+
+	FParamDict GetDictionary() const
+	{
+		FScopeLock lock(&Lock);
+		return Dictionary;
+	}
+
+	void Clear()
+	{
+		FScopeLock lock(&Lock);
+		Dictionary.Clear();
+	}
+	void Set(const FName& Key, const FVariantValue& Value)
+	{
+		FScopeLock lock(&Lock);
+		Dictionary.Set(Key, Value);
+	}
+	void GetKeys(TArray<FName>& OutKeys) const
+	{
+		FScopeLock lock(&Lock);
+		Dictionary.GetKeys(OutKeys);
+	}
+	bool HaveKey(const FName& Key) const
+	{
+		FScopeLock lock(&Lock);
+		return Dictionary.HaveKey(Key);
+	}
+	FVariantValue GetValue(const FName& Key) const
+	{
+		FScopeLock lock(&Lock);
+		return Dictionary.GetValue(Key);
+	}
+	void Remove(const FName& Key)
+	{
+		FScopeLock lock(&Lock);
+		Dictionary.Remove(Key);
+	}
+	void ConvertTo(TMap<FString, FVariant>& OutVariantMap, const FString& InAddPrefixToKey) const
+	{
+		FScopeLock lock(&Lock);
+		Dictionary.ConvertTo(OutVariantMap, InAddPrefixToKey);
+	}
+	void ConvertKeysStartingWithTo(TMap<FString, FVariant>& OutVariantMap, const FString& InKeyStartsWith, const FString& InAddPrefixToKey) const
+	{
+		FScopeLock lock(&Lock);
+		Dictionary.ConvertKeysStartingWithTo(OutVariantMap, InKeyStartsWith, InAddPrefixToKey);
+	}
+private:
+	mutable FCriticalSection Lock;
+	FParamDict Dictionary;
+};
+
 
 } // namespace Electra
 
