ElectraPlayer: Using maximum configured decoded texture size to avoid reallocations on resolution changes

#rnx

diff --git a/Engine/Plugins/Media/ElectraPlayer/Source/ElectraPlayerRuntime/Private/Runtime/Decoder/VideoDecoder.cpp b/Engine/Plugins/Media/ElectraPlayer/Source/ElectraPlayerRuntime/Private/Runtime/Decoder/VideoDecoder.cpp
--- a/Engine/Plugins/Media/ElectraPlayer/Source/ElectraPlayerRuntime/Private/Runtime/Decoder/VideoDecoder.cpp
+++ b/Engine/Plugins/Media/ElectraPlayer/Source/ElectraPlayerRuntime/Private/Runtime/Decoder/VideoDecoder.cpp
@@ -188,7 +188,7 @@ private:
 
 	TWeakPtr<IVideoDecoderResourceDelegate, ESPMode::ThreadSafe>			VideoResourceDelegate;
 	TSharedPtr<IMediaRenderer, ESPMode::ThreadSafe>							Renderer;
-	int32																	MaxDecodeBufferSize = 0;
+	int32																	MaxOutputBuffers = 0;
 
 	FCriticalSection														ListenerMutex;
 	IAccessUnitBufferListener*												InputBufferListener = nullptr;
@@ -393,7 +393,8 @@ void FVideoDecoderImpl::CreateDecoderOutputPool()
 	poolOpts.Set(RenderOptionKeys::NumBuffers, FVariantValue(NumOutputFrames));
 	if (Renderer->CreateBufferPool(poolOpts) == UEMEDIA_ERROR_OK)
 	{
-		MaxDecodeBufferSize = (int32) Renderer->GetBufferPoolProperties().GetValue(RenderOptionKeys::MaxBuffers).GetInt64();
+		MaxOutputBuffers = (int32) Renderer->GetBufferPoolProperties().GetValue(RenderOptionKeys::MaxBuffers).GetInt64();
+		DecoderFactoryAddtlCfg.Add(TEXT("max_output_buffers"), FVariant((uint32)MaxOutputBuffers));
 	}
 	else
 	{
@@ -411,7 +412,7 @@ void FVideoDecoderImpl::NotifyReadyBufferListener(bool bHaveOutput)
 	if (ReadyBufferListener)
 	{
 		IDecoderOutputBufferListener::FDecodeReadyStats stats;
-		stats.MaxDecodedElementsReady = MaxDecodeBufferSize;
+		stats.MaxDecodedElementsReady = MaxOutputBuffers;
 		stats.NumElementsInDecoder = CurrentOutputBuffer ? 1 : 0;
 		stats.bOutputStalled = !bHaveOutput;
 		stats.bEODreached = NextAccessUnits.ReachedEOD() && stats.NumDecodedElementsReady == 0 && stats.NumElementsInDecoder == 0;
