ElectraPlayer: Fixed where and when media events are sent to avoid sending a `TracksChanged` when the available tracks are not yet known (because track availability may change with multiperiod DASH streams and thus a time value is needed to determine which tracks are available there).

#rnx

diff --git a/Engine/Plugins/Media/ElectraPlayer/Source/ElectraPlayerRuntime/Private/ElectraPlayer.cpp b/Engine/Plugins/Media/ElectraPlayer/Source/ElectraPlayerRuntime/Private/ElectraPlayer.cpp
--- a/Engine/Plugins/Media/ElectraPlayer/Source/ElectraPlayerRuntime/Private/ElectraPlayer.cpp
+++ b/Engine/Plugins/Media/ElectraPlayer/Source/ElectraPlayerRuntime/Private/ElectraPlayer.cpp
@@ -2210,6 +2210,14 @@ void FElectraPlayer::HandlePlayerEventDataAvailabilityChange(const Electra::Metr
 void FElectraPlayer::HandlePlayerEventBufferingStart(Electra::Metrics::EBufferingReason BufferingReason)
 {
 	PlayerState.Status = PlayerState.Status | EPlayerStatus::Buffering;
+
+	// Send TracksChanged on the initial buffering event. Prior to that we do not know where in the stream
+	// playback will begin and what tracks are available there.
+	if (BufferingReason == Electra::Metrics::EBufferingReason::Initial)
+	{
+		DeferredEvents.Enqueue(IElectraPlayerAdapterDelegate::EPlayerEvent::TracksChanged);
+	}
+
 	DeferredEvents.Enqueue(IElectraPlayerAdapterDelegate::EPlayerEvent::MediaBuffering);
 
 	// Update statistics
@@ -3085,8 +3093,9 @@ void FElectraPlayer::MediaStateOnPreparingFinished()
 	CSV_EVENT(ElectraPlayer, TEXT("MediaStateOnPreparingFinished"));
 
 	PlayerState.State = EPlayerState::Stopped;
+	// Only report MediaOpened here and *not* TracksChanged as well.
+	// We do not know where playback will start at and what tracks are available at that point.
 	DeferredEvents.Enqueue(IElectraPlayerAdapterDelegate::EPlayerEvent::MediaOpened);
-	DeferredEvents.Enqueue(IElectraPlayerAdapterDelegate::EPlayerEvent::TracksChanged);
 }
 
 bool FElectraPlayer::MediaStateOnPlay()
