ElectraDecoders: Dropping frames that are to be skipped and not displayed while seeking earlier to speed up the seek process.

#rnx

diff --git a/Engine/Plugins/Media/ElectraCodecs/Source/ElectraDecoders/Private/Android/h264/VideoDecoderH264_Android.cpp b/Engine/Plugins/Media/ElectraCodecs/Source/ElectraDecoders/Private/Android/h264/VideoDecoderH264_Android.cpp
--- a/Engine/Plugins/Media/ElectraCodecs/Source/ElectraDecoders/Private/Android/h264/VideoDecoderH264_Android.cpp
+++ b/Engine/Plugins/Media/ElectraCodecs/Source/ElectraDecoders/Private/Android/h264/VideoDecoderH264_Android.cpp
@@ -259,6 +259,7 @@ void IElectraVideoDecoderH264_Android::GetConfigurationOptions(TMap<FString, FVa
 	OutOptions.Emplace(IElectraDecoderFeature::IsAdaptive, FVariant(false));
 	OutOptions.Emplace(IElectraDecoderFeature::NeedReplayDataOnDecoderLoss, FVariant(true));
 	OutOptions.Emplace(IElectraDecoderFeature::MustBeSuspendedInBackground, FVariant(true));
+	OutOptions.Emplace(IElectraDecoderFeature::SupportsDroppingOutput, FVariant(true));
 }
 
 TSharedPtr<IElectraDecoder, ESPMode::ThreadSafe> IElectraVideoDecoderH264_Android::Create(const TMap<FString, FVariant>& InOptions, TSharedPtr<IElectraDecoderResourceDelegate, ESPMode::ThreadSafe> InResourceDelegate)
@@ -433,6 +434,15 @@ IElectraDecoder::EDecoderError FElectraVideoDecoderH264_Android::DecodeAccessUni
 		return IElectraDecoder::EDecoderError::EndOfData;
 	}
 
+	// CSD only buffer is not handled at the moment.
+	check((InInputAccessUnit.Flags & EElectraDecoderFlags::InitCSDOnly) == EElectraDecoderFlags::None);
+
+	// If this is discardable and won't be output we do not need to handle it at all.
+	if ((InInputAccessUnit.Flags & (EElectraDecoderFlags::DoNotOutput | EElectraDecoderFlags::IsDiscardable)) == (EElectraDecoderFlags::DoNotOutput | EElectraDecoderFlags::IsDiscardable))
+	{
+		return IElectraDecoder::EDecoderError::None;
+	}
+
 	// Still creating the decoder?
 	if (DecodeState == EDecodeState::CreatingDecoder)
 	{
@@ -446,9 +456,6 @@ IElectraDecoder::EDecoderError FElectraVideoDecoderH264_Android::DecodeAccessUni
 		}
 	}
 
-	// CSD only buffer is not handled at the moment.
-	check((InInputAccessUnit.Flags & EElectraDecoderFlags::InitCSDOnly) == EElectraDecoderFlags::None);
-
 	// Need a valid CSD to create a decoder.
 	if (!ConfigRecord.IsValid())
 	{
@@ -1035,8 +1042,8 @@ FElectraVideoDecoderH264_Android::EConvertResult FElectraVideoDecoderH264_Androi
 		return EConvertResult::Failure;
 	}
 
-	// If this is a replay frame we do not need to set it up and merely release the buffer immediately.
-	if ((In->AccessUnit.Flags & EElectraDecoderFlags::IsReplaySample) != EElectraDecoderFlags::None)
+	// If this is a replay frame or a frame that will not be output, we do not need to set it up and merely release the buffer immediately.
+	if ((In->AccessUnit.Flags & (EElectraDecoderFlags::IsReplaySample | EElectraDecoderFlags::DoNotOutput)) != EElectraDecoderFlags::None)
 	{
 		DETAILLOG(LogElectraDecoders, Log, TEXT("VideoDecoderH264::ConvertDecoderOutput() - Discard replay output %lld, 0x%x"), (long long int)In->AccessUnit.PTS.GetTicks(), In->AccessUnit.Flags);
 		if (DecoderInstance.IsValid())
diff --git a/Engine/Plugins/Media/ElectraCodecs/Source/ElectraDecoders/Private/Android/h265/VideoDecoderH265_Android.cpp b/Engine/Plugins/Media/ElectraCodecs/Source/ElectraDecoders/Private/Android/h265/VideoDecoderH265_Android.cpp
--- a/Engine/Plugins/Media/ElectraCodecs/Source/ElectraDecoders/Private/Android/h265/VideoDecoderH265_Android.cpp
+++ b/Engine/Plugins/Media/ElectraCodecs/Source/ElectraDecoders/Private/Android/h265/VideoDecoderH265_Android.cpp
@@ -286,6 +286,7 @@ void IElectraVideoDecoderH265_Android::GetConfigurationOptions(TMap<FString, FVa
 	OutOptions.Emplace(IElectraDecoderFeature::IsAdaptive, FVariant(false));
 	OutOptions.Emplace(IElectraDecoderFeature::NeedReplayDataOnDecoderLoss, FVariant(true));
 	OutOptions.Emplace(IElectraDecoderFeature::MustBeSuspendedInBackground, FVariant(true));
+	OutOptions.Emplace(IElectraDecoderFeature::SupportsDroppingOutput, FVariant(true));
 }
 
 TSharedPtr<IElectraDecoder, ESPMode::ThreadSafe> IElectraVideoDecoderH265_Android::Create(const TMap<FString, FVariant>& InOptions, TSharedPtr<IElectraDecoderResourceDelegate, ESPMode::ThreadSafe> InResourceDelegate)
@@ -460,6 +461,15 @@ IElectraDecoder::EDecoderError FElectraVideoDecoderH265_Android::DecodeAccessUni
 		return IElectraDecoder::EDecoderError::EndOfData;
 	}
 
+	// CSD only buffer is not handled at the moment.
+	check((InInputAccessUnit.Flags & EElectraDecoderFlags::InitCSDOnly) == EElectraDecoderFlags::None);
+
+	// If this is discardable and won't be output we do not need to handle it at all.
+	if ((InInputAccessUnit.Flags & (EElectraDecoderFlags::DoNotOutput | EElectraDecoderFlags::IsDiscardable)) == (EElectraDecoderFlags::DoNotOutput | EElectraDecoderFlags::IsDiscardable))
+	{
+		return IElectraDecoder::EDecoderError::None;
+	}
+
 	// Still creating the decoder?
 	if (DecodeState == EDecodeState::CreatingDecoder)
 	{
@@ -473,9 +483,6 @@ IElectraDecoder::EDecoderError FElectraVideoDecoderH265_Android::DecodeAccessUni
 		}
 	}
 
-	// CSD only buffer is not handled at the moment.
-	check((InInputAccessUnit.Flags & EElectraDecoderFlags::InitCSDOnly) == EElectraDecoderFlags::None);
-
 	// Need a valid CSD to create a decoder.
 	if (!ConfigRecord.IsValid())
 	{
@@ -1063,8 +1070,8 @@ FElectraVideoDecoderH265_Android::EConvertResult FElectraVideoDecoderH265_Androi
 		return EConvertResult::Failure;
 	}
 
-	// If this is a replay frame we do not need to set it up and merely release the buffer immediately.
-	if ((In->AccessUnit.Flags & EElectraDecoderFlags::IsReplaySample) != EElectraDecoderFlags::None)
+	// If this is a replay frame or a frame that will not be output, we do not need to set it up and merely release the buffer immediately.
+	if ((In->AccessUnit.Flags & (EElectraDecoderFlags::IsReplaySample | EElectraDecoderFlags::DoNotOutput)) != EElectraDecoderFlags::None)
 	{
 		DETAILLOG(LogElectraDecoders, Log, TEXT("VideoDecoderH265::ConvertDecoderOutput() - Discard replay output %lld, 0x%x"), (long long int)In->AccessUnit.PTS.GetTicks(), In->AccessUnit.Flags);
 		if (DecoderInstance.IsValid())
diff --git a/Engine/Plugins/Media/ElectraCodecs/Source/ElectraDecoders/Private/Apple/h264/VideoDecoderH264_Apple.cpp b/Engine/Plugins/Media/ElectraCodecs/Source/ElectraDecoders/Private/Apple/h264/VideoDecoderH264_Apple.cpp
--- a/Engine/Plugins/Media/ElectraCodecs/Source/ElectraDecoders/Private/Apple/h264/VideoDecoderH264_Apple.cpp
+++ b/Engine/Plugins/Media/ElectraCodecs/Source/ElectraDecoders/Private/Apple/h264/VideoDecoderH264_Apple.cpp
@@ -228,6 +228,7 @@ private:
 		FInputAccessUnit AccessUnit;
 		TMap<FString, FVariant> AdditionalOptions;
 		TSharedPtr<ElectraDecodersUtil::MPEG::FISO14496_10_seq_parameter_set_data, ESPMode::ThreadSafe> SPS;
+		bool bDropOutput = false;
 	};
 
 	struct FDecodedImage
@@ -378,6 +379,7 @@ void IElectraVideoDecoderH264_Apple::GetConfigurationOptions(TMap<FString, FVari
 
 	OutOptions.Emplace(IElectraDecoderFeature::NeedReplayDataOnDecoderLoss, FVariant(true));
 	OutOptions.Emplace(IElectraDecoderFeature::MustBeSuspendedInBackground, FVariant(true));
+	OutOptions.Emplace(IElectraDecoderFeature::SupportsDroppingOutput, FVariant(true));
 }
 
 TSharedPtr<IElectraDecoder, ESPMode::ThreadSafe> IElectraVideoDecoderH264_Apple::Create(const TMap<FString, FVariant>& InOptions, TSharedPtr<IElectraDecoderResourceDelegate, ESPMode::ThreadSafe> InResourceDelegate)
@@ -546,6 +548,15 @@ IElectraDecoder::EDecoderError FElectraVideoDecoderH264_Apple::DecodeAccessUnit(
 		return IElectraDecoder::EDecoderError::EndOfData;
 	}
 
+	// CSD only buffer is not handled at the moment.
+	check((InInputAccessUnit.Flags & EElectraDecoderFlags::InitCSDOnly) == EElectraDecoderFlags::None);
+
+	// If this is discardable and won't be output we do not need to handle it at all.
+	if ((InInputAccessUnit.Flags & (EElectraDecoderFlags::DoNotOutput | EElectraDecoderFlags::IsDiscardable)) == (EElectraDecoderFlags::DoNotOutput | EElectraDecoderFlags::IsDiscardable))
+	{
+		return IElectraDecoder::EDecoderError::None;
+	}
+
 	// If there is pending output it is very likely that decoding this access unit would also generate output.
 	// Since that would result in loss of the pending output we return now.
 	if (CurrentOutput.IsValid())
@@ -553,9 +564,6 @@ IElectraDecoder::EDecoderError FElectraVideoDecoderH264_Apple::DecodeAccessUnit(
 		return IElectraDecoder::EDecoderError::NoBuffer;
 	}
 
-	// CSD only buffer is not handled at the moment.
-	check((InInputAccessUnit.Flags & EElectraDecoderFlags::InitCSDOnly) == EElectraDecoderFlags::None);
-
 	// If a new decoder is needed, destroy the current one.
 	if (bNewDecoderRequired)
 	{
@@ -617,6 +625,7 @@ IElectraDecoder::EDecoderError FElectraVideoDecoderH264_Apple::DecodeAccessUnit(
 		TSharedPtr<FDecoderInput, ESPMode::ThreadSafe> In(new FDecoderInput);
 		In->AdditionalOptions = InAdditionalOptions;
 		In->AccessUnit = InInputAccessUnit;
+		In->bDropOutput = (InInputAccessUnit.Flags & EElectraDecoderFlags::DoNotOutput) == EElectraDecoderFlags::DoNotOutput;
 		In->SPS = CurrentSPS;
 		// Zero the input pointer and size in the copy. That data is not owned by us and it's best not to have any
 		// values here that would lead us to think that we do.
@@ -632,7 +641,7 @@ IElectraDecoder::EDecoderError FElectraVideoDecoderH264_Apple::DecodeAccessUnit(
 
 		// Decode
 		VTDecodeFrameFlags DecodeFlags = kVTDecodeFrame_EnableAsynchronousDecompression | kVTDecodeFrame_EnableTemporalProcessing;
-		if ((InInputAccessUnit.Flags & (/*EElectraDecoderFlags::DoNotOutput |*/ EElectraDecoderFlags::IsReplaySample | EElectraDecoderFlags::IsLastReplaySample)) != EElectraDecoderFlags::None)
+		if ((InInputAccessUnit.Flags & (EElectraDecoderFlags::DoNotOutput | EElectraDecoderFlags::IsReplaySample | EElectraDecoderFlags::IsLastReplaySample)) != EElectraDecoderFlags::None)
 		{
 			DecodeFlags |= kVTDecodeFrame_DoNotOutputFrame;
 		}
@@ -972,7 +981,7 @@ void FElectraVideoDecoderH264_Apple::DecodeCallback(void* pSrcRef, OSStatus Resu
 
 	if (Result == 0)
 	{
-		if (imageBuffer != nullptr && (infoFlags & kVTDecodeInfo_FrameDropped) == 0 && MatchingInput.IsValid())
+		if (imageBuffer != nullptr && (infoFlags & kVTDecodeInfo_FrameDropped) == 0 && !MatchingInput->bDropOutput)
 		{
 			FDecodedImage NextImage;
 			NextImage.SourceInfo = MatchingInput;
@@ -984,6 +993,10 @@ void FElectraVideoDecoderH264_Apple::DecodeCallback(void* pSrcRef, OSStatus Resu
 			ReadyImageMutex.Unlock();
 		}
 	}
+	else if (Result == kVTVideoDecoderReferenceMissingErr)
+	{
+		// Ignore this.
+	}
 	else
 	{
 		FString Msg = FString::Printf(TEXT("Failed to decode video: DecodeCallback() returned Result 0x%x; flags 0x%x"), Result, infoFlags);
diff --git a/Engine/Plugins/Media/ElectraCodecs/Source/ElectraDecoders/Private/Apple/h265/VideoDecoderH265_Apple.cpp b/Engine/Plugins/Media/ElectraCodecs/Source/ElectraDecoders/Private/Apple/h265/VideoDecoderH265_Apple.cpp
--- a/Engine/Plugins/Media/ElectraCodecs/Source/ElectraDecoders/Private/Apple/h265/VideoDecoderH265_Apple.cpp
+++ b/Engine/Plugins/Media/ElectraCodecs/Source/ElectraDecoders/Private/Apple/h265/VideoDecoderH265_Apple.cpp
@@ -228,6 +228,7 @@ private:
 		FInputAccessUnit AccessUnit;
 		TMap<FString, FVariant> AdditionalOptions;
 		TSharedPtr<ElectraDecodersUtil::MPEG::FISO23008_2_seq_parameter_set_data, ESPMode::ThreadSafe> SPS;
+		bool bDropOutput = false;
 	};
 
 	struct FDecodedImage
@@ -378,6 +379,7 @@ void IElectraVideoDecoderH265_Apple::GetConfigurationOptions(TMap<FString, FVari
 
 	OutOptions.Emplace(IElectraDecoderFeature::NeedReplayDataOnDecoderLoss, FVariant(true));
 	OutOptions.Emplace(IElectraDecoderFeature::MustBeSuspendedInBackground, FVariant(true));
+	OutOptions.Emplace(IElectraDecoderFeature::SupportsDroppingOutput, FVariant(true));
 }
 
 TSharedPtr<IElectraDecoder, ESPMode::ThreadSafe> IElectraVideoDecoderH265_Apple::Create(const TMap<FString, FVariant>& InOptions, TSharedPtr<IElectraDecoderResourceDelegate, ESPMode::ThreadSafe> InResourceDelegate)
@@ -550,6 +552,15 @@ IElectraDecoder::EDecoderError FElectraVideoDecoderH265_Apple::DecodeAccessUnit(
 		return IElectraDecoder::EDecoderError::EndOfData;
 	}
 
+	// CSD only buffer is not handled at the moment.
+	check((InInputAccessUnit.Flags & EElectraDecoderFlags::InitCSDOnly) == EElectraDecoderFlags::None);
+
+	// If this is discardable and won't be output we do not need to handle it at all.
+	if ((InInputAccessUnit.Flags & (EElectraDecoderFlags::DoNotOutput | EElectraDecoderFlags::IsDiscardable)) == (EElectraDecoderFlags::DoNotOutput | EElectraDecoderFlags::IsDiscardable))
+	{
+		return IElectraDecoder::EDecoderError::None;
+	}
+
 	// If there is pending output it is very likely that decoding this access unit would also generate output.
 	// Since that would result in loss of the pending output we return now.
 	if (CurrentOutput.IsValid())
@@ -557,9 +568,6 @@ IElectraDecoder::EDecoderError FElectraVideoDecoderH265_Apple::DecodeAccessUnit(
 		return IElectraDecoder::EDecoderError::NoBuffer;
 	}
 
-	// CSD only buffer is not handled at the moment.
-	check((InInputAccessUnit.Flags & EElectraDecoderFlags::InitCSDOnly) == EElectraDecoderFlags::None);
-
 	// If a new decoder is needed, destroy the current one.
 	if (bNewDecoderRequired)
 	{
@@ -621,6 +629,7 @@ IElectraDecoder::EDecoderError FElectraVideoDecoderH265_Apple::DecodeAccessUnit(
 		TSharedPtr<FDecoderInput, ESPMode::ThreadSafe> In(new FDecoderInput);
 		In->AdditionalOptions = InAdditionalOptions;
 		In->AccessUnit = InInputAccessUnit;
+		In->bDropOutput = (InInputAccessUnit.Flags & EElectraDecoderFlags::DoNotOutput) == EElectraDecoderFlags::DoNotOutput;
 		In->SPS = CurrentSPS;
 		// Zero the input pointer and size in the copy. That data is not owned by us and it's best not to have any
 		// values here that would lead us to think that we do.
@@ -636,7 +645,7 @@ IElectraDecoder::EDecoderError FElectraVideoDecoderH265_Apple::DecodeAccessUnit(
 
 		// Decode
 		VTDecodeFrameFlags DecodeFlags = kVTDecodeFrame_EnableAsynchronousDecompression | kVTDecodeFrame_EnableTemporalProcessing;
-		if ((InInputAccessUnit.Flags & (/*EElectraDecoderFlags::DoNotOutput |*/ EElectraDecoderFlags::IsReplaySample | EElectraDecoderFlags::IsLastReplaySample)) != EElectraDecoderFlags::None)
+		if ((InInputAccessUnit.Flags & (EElectraDecoderFlags::DoNotOutput | EElectraDecoderFlags::IsReplaySample | EElectraDecoderFlags::IsLastReplaySample)) != EElectraDecoderFlags::None)
 		{
 			DecodeFlags |= kVTDecodeFrame_DoNotOutputFrame;
 		}
@@ -976,7 +985,7 @@ void FElectraVideoDecoderH265_Apple::DecodeCallback(void* pSrcRef, OSStatus Resu
 
 	if (Result == 0)
 	{
-		if (imageBuffer != nullptr && (infoFlags & kVTDecodeInfo_FrameDropped) == 0 && MatchingInput.IsValid())
+		if (imageBuffer != nullptr && (infoFlags & kVTDecodeInfo_FrameDropped) == 0 && !MatchingInput->bDropOutput)
 		{
 			FDecodedImage NextImage;
 			NextImage.SourceInfo = MatchingInput;
@@ -988,6 +997,10 @@ void FElectraVideoDecoderH265_Apple::DecodeCallback(void* pSrcRef, OSStatus Resu
 			ReadyImageMutex.Unlock();
 		}
 	}
+	else if (Result == kVTVideoDecoderReferenceMissingErr)
+	{
+		// Ignore this.
+	}
 	else
 	{
 		FString Msg = FString::Printf(TEXT("Failed to decode video: DecodeCallback() returned Result 0x%x; flags 0x%x"), Result, infoFlags);
diff --git a/Engine/Plugins/Media/ElectraCodecs/Source/ElectraDecoders/Private/DX/VideoDecoderH264_DX.cpp b/Engine/Plugins/Media/ElectraCodecs/Source/ElectraDecoders/Private/DX/VideoDecoderH264_DX.cpp
--- a/Engine/Plugins/Media/ElectraCodecs/Source/ElectraDecoders/Private/DX/VideoDecoderH264_DX.cpp
+++ b/Engine/Plugins/Media/ElectraCodecs/Source/ElectraDecoders/Private/DX/VideoDecoderH264_DX.cpp
@@ -158,6 +158,7 @@ private:
 		FInputAccessUnit AccessUnit;
 		TMap<FString, FVariant> AdditionalOptions;
 		void* InputDataCopy = nullptr;
+		bool bDropOutput = false;
 	};
 
 	struct FDecoderOutputBuffer
@@ -364,6 +365,15 @@ IElectraDecoder::EDecoderError FElectraVideoDecoderH264_DX::DecodeAccessUnit(con
 		return IElectraDecoder::EDecoderError::EndOfData;
 	}
 
+	// CSD only buffer is not handled at the moment.
+	check((InInputAccessUnit.Flags & EElectraDecoderFlags::InitCSDOnly) == EElectraDecoderFlags::None);
+
+	// If this is discardable and won't be output we do not need to handle it at all.
+	if ((InInputAccessUnit.Flags & (EElectraDecoderFlags::DoNotOutput | EElectraDecoderFlags::IsDiscardable)) == (EElectraDecoderFlags::DoNotOutput | EElectraDecoderFlags::IsDiscardable))
+	{
+		return IElectraDecoder::EDecoderError::None;
+	}
+
 	// If there is pending output it is very likely that decoding this access unit would also generate output.
 	// Since that would result in loss of the pending output we return now.
 	if (CurrentOutput.IsValid())
@@ -371,9 +381,6 @@ IElectraDecoder::EDecoderError FElectraVideoDecoderH264_DX::DecodeAccessUnit(con
 		return IElectraDecoder::EDecoderError::NoBuffer;
 	}
 
-	// CSD only buffer is not handled at the moment.
-	check((InInputAccessUnit.Flags & EElectraDecoderFlags::InitCSDOnly) == EElectraDecoderFlags::None);
-
 	// Create decoder transform if necessary.
 	if (!DecoderTransform.IsValid() && !InternalDecoderCreate(InAdditionalOptions))
 	{
@@ -416,6 +423,7 @@ IElectraDecoder::EDecoderError FElectraVideoDecoderH264_DX::DecodeAccessUnit(con
 
 					FDecoderInput In;
 					In.AdditionalOptions = InAdditionalOptions;
+					In.bDropOutput = (InInputAccessUnit.Flags & EElectraDecoderFlags::DoNotOutput) == EElectraDecoderFlags::DoNotOutput;
 					In.AccessUnit = InInputAccessUnit;
 					// If we need to hold on to the input data we need to make a local copy.
 					// For safety reasons we zero out the pointer we were given in the input data copy to now accidentally
@@ -851,6 +859,11 @@ bool FElectraVideoDecoderH264_DX::ConvertDecoderOutput()
 	}
 	*/
 
+	if (MatchingInput.bDropOutput)
+	{
+		return true;
+	}
+
 	TSharedPtr<FElectraVideoDecoderOutputH264_DX, ESPMode::ThreadSafe> NewOutput = MakeShared<FElectraVideoDecoderOutputH264_DX>();
 	NewOutput->PTS = MatchingInput.AccessUnit.PTS;
 	NewOutput->UserValue = MatchingInput.AccessUnit.UserValue;
diff --git a/Engine/Plugins/Media/ElectraCodecs/Source/ElectraDecoders/Private/DX/VideoDecoderH265_DX.cpp b/Engine/Plugins/Media/ElectraCodecs/Source/ElectraDecoders/Private/DX/VideoDecoderH265_DX.cpp
--- a/Engine/Plugins/Media/ElectraCodecs/Source/ElectraDecoders/Private/DX/VideoDecoderH265_DX.cpp
+++ b/Engine/Plugins/Media/ElectraCodecs/Source/ElectraDecoders/Private/DX/VideoDecoderH265_DX.cpp
@@ -158,6 +158,7 @@ private:
 		FInputAccessUnit AccessUnit;
 		TMap<FString, FVariant> AdditionalOptions;
 		void* InputDataCopy = nullptr;
+		bool bDropOutput = false;
 	};
 
 	struct FDecoderOutputBuffer
@@ -365,6 +366,15 @@ IElectraDecoder::EDecoderError FElectraVideoDecoderH265_DX::DecodeAccessUnit(con
 		return IElectraDecoder::EDecoderError::EndOfData;
 	}
 
+	// CSD only buffer is not handled at the moment.
+	check((InInputAccessUnit.Flags & EElectraDecoderFlags::InitCSDOnly) == EElectraDecoderFlags::None);
+
+	// If this is discardable and won't be output we do not need to handle it at all.
+	if ((InInputAccessUnit.Flags & (EElectraDecoderFlags::DoNotOutput | EElectraDecoderFlags::IsDiscardable)) == (EElectraDecoderFlags::DoNotOutput | EElectraDecoderFlags::IsDiscardable))
+	{
+		return IElectraDecoder::EDecoderError::None;
+	}
+
 	// If there is pending output it is very likely that decoding this access unit would also generate output.
 	// Since that would result in loss of the pending output we return now.
 	if (CurrentOutput.IsValid())
@@ -372,9 +382,6 @@ IElectraDecoder::EDecoderError FElectraVideoDecoderH265_DX::DecodeAccessUnit(con
 		return IElectraDecoder::EDecoderError::NoBuffer;
 	}
 
-	// CSD only buffer is not handled at the moment.
-	check((InInputAccessUnit.Flags & EElectraDecoderFlags::InitCSDOnly) == EElectraDecoderFlags::None);
-
 	// Create decoder transform if necessary.
 	if (!DecoderTransform.IsValid() && !InternalDecoderCreate(InAdditionalOptions))
 	{
@@ -417,6 +424,7 @@ IElectraDecoder::EDecoderError FElectraVideoDecoderH265_DX::DecodeAccessUnit(con
 
 					FDecoderInput In;
 					In.AdditionalOptions = InAdditionalOptions;
+					In.bDropOutput = (InInputAccessUnit.Flags & EElectraDecoderFlags::DoNotOutput) == EElectraDecoderFlags::DoNotOutput;
 					In.AccessUnit = InInputAccessUnit;
 					// If we need to hold on to the input data we need to make a local copy.
 					// For safety reasons we zero out the pointer we were given in the input data copy to now accidentally
@@ -849,6 +857,7 @@ bool FElectraVideoDecoderH265_DX::ConvertDecoderOutput()
 			break;
 		}
 	}
+
 	// If no exact match is found then the decoder transform may have interpolated the PTS given the individual
 	// frame durations or perhaps due to SEI messages.
 	if (!bFoundMatch)
@@ -865,6 +874,11 @@ bool FElectraVideoDecoderH265_DX::ConvertDecoderOutput()
 	}
 	*/
 
+	if (MatchingInput.bDropOutput)
+	{
+		return true;
+	}
+
 	TSharedPtr<FElectraVideoDecoderOutputH265_DX, ESPMode::ThreadSafe> NewOutput = MakeShared<FElectraVideoDecoderOutputH265_DX>();
 	NewOutput->PTS = MatchingInput.AccessUnit.PTS;
 	NewOutput->UserValue = MatchingInput.AccessUnit.UserValue;
diff --git a/Engine/Plugins/Media/ElectraCodecs/Source/ElectraDecoders/Private/Linux/h264/VideoDecoderH264_Linux.cpp b/Engine/Plugins/Media/ElectraCodecs/Source/ElectraDecoders/Private/Linux/h264/VideoDecoderH264_Linux.cpp
--- a/Engine/Plugins/Media/ElectraCodecs/Source/ElectraDecoders/Private/Linux/h264/VideoDecoderH264_Linux.cpp
+++ b/Engine/Plugins/Media/ElectraCodecs/Source/ElectraDecoders/Private/Linux/h264/VideoDecoderH264_Linux.cpp
@@ -135,6 +135,7 @@ private:
 		FInputAccessUnit AccessUnit;
 		TMap<FString, FVariant> AdditionalOptions;
 		TSharedPtr<ElectraDecodersUtil::MPEG::FISO14496_10_seq_parameter_set_data, ESPMode::ThreadSafe> SPS;
+		bool bDropOutput = false;
 	};
 
 	enum class EDecodeState
@@ -203,6 +204,7 @@ void IElectraVideoDecoderH264_Linux::GetConfigurationOptions(TMap<FString, FVari
 	OutOptions.Emplace(IElectraDecoderFeature::MinimumNumberOfOutputFrames, FVariant((int32)8));
 	// The decoder is adaptive. There is no need to call IsCompatibleWith() on stream switches.
 	OutOptions.Emplace(IElectraDecoderFeature::IsAdaptive, FVariant(true));
+	OutOptions.Emplace(IElectraDecoderFeature::SupportsDroppingOutput, FVariant(true));
 }
 
 TSharedPtr<IElectraDecoder, ESPMode::ThreadSafe> IElectraVideoDecoderH264_Linux::Create(const TMap<FString, FVariant>& InOptions, TSharedPtr<IElectraDecoderResourceDelegate, ESPMode::ThreadSafe> InResourceDelegate)
@@ -326,6 +328,15 @@ IElectraDecoder::EDecoderError FElectraVideoDecoderH264_Linux::DecodeAccessUnit(
 		return IElectraDecoder::EDecoderError::EndOfData;
 	}
 
+	// CSD only buffer is not handled at the moment.
+	check((InInputAccessUnit.Flags & EElectraDecoderFlags::InitCSDOnly) == EElectraDecoderFlags::None);
+
+	// If this is discardable and won't be output we do not need to handle it at all.
+	if ((InInputAccessUnit.Flags & (EElectraDecoderFlags::DoNotOutput | EElectraDecoderFlags::IsDiscardable)) == (EElectraDecoderFlags::DoNotOutput | EElectraDecoderFlags::IsDiscardable))
+	{
+		return IElectraDecoder::EDecoderError::None;
+	}
+
 	// If there is pending output it is very likely that decoding this access unit would also generate output.
 	// Since that would result in loss of the pending output we return now.
 	if (CurrentOutput.IsValid())
@@ -333,9 +344,6 @@ IElectraDecoder::EDecoderError FElectraVideoDecoderH264_Linux::DecodeAccessUnit(
 		return IElectraDecoder::EDecoderError::NoBuffer;
 	}
 
-	// CSD only buffer is not handled at the moment.
-	check((InInputAccessUnit.Flags & EElectraDecoderFlags::InitCSDOnly) == EElectraDecoderFlags::None);
-
 	// If a new decoder is needed, destroy the current one.
 	if (bNewDecoderRequired)
 	{
@@ -389,6 +397,7 @@ IElectraDecoder::EDecoderError FElectraVideoDecoderH264_Linux::DecodeAccessUnit(
 			TSharedPtr<FDecoderInput, ESPMode::NotThreadSafe> In(new FDecoderInput);
 			In->AdditionalOptions = InAdditionalOptions;
 			In->AccessUnit = InInputAccessUnit;
+			In->bDropOutput = (InInputAccessUnit.Flags & EElectraDecoderFlags::DoNotOutput) == EElectraDecoderFlags::DoNotOutput;
 			In->SPS = CurrentSPS;
 			// Zero the input pointer and size in the copy. That data is not owned by us and it's best not to have any
 			// values here that would lead us to think that we do.
@@ -605,7 +614,11 @@ FElectraVideoDecoderH264_Linux::EConvertResult FElectraVideoDecoderH264_Linux::C
 	NewOutput->ExtraValues.Emplace(TEXT("decoder"), FVariant(TEXT("libavcodec")));
 	NewOutput->ExtraValues.Emplace(TEXT("codec"), FVariant(TEXT("avc")));
 
-	CurrentOutput = MoveTemp(NewOutput);
-
+	// Take this output on if it is not to be dropped.
+	// We need to have pulled it out of libav before deciding this though.
+	if (!In->bDropOutput)
+	{
+		CurrentOutput = MoveTemp(NewOutput);
+	}
 	return EConvertResult::Success;
 }
diff --git a/Engine/Plugins/Media/ElectraCodecs/Source/ElectraDecoders/Private/Linux/h265/VideoDecoderH265_Linux.cpp b/Engine/Plugins/Media/ElectraCodecs/Source/ElectraDecoders/Private/Linux/h265/VideoDecoderH265_Linux.cpp
--- a/Engine/Plugins/Media/ElectraCodecs/Source/ElectraDecoders/Private/Linux/h265/VideoDecoderH265_Linux.cpp
+++ b/Engine/Plugins/Media/ElectraCodecs/Source/ElectraDecoders/Private/Linux/h265/VideoDecoderH265_Linux.cpp
@@ -135,6 +135,7 @@ private:
 		FInputAccessUnit AccessUnit;
 		TMap<FString, FVariant> AdditionalOptions;
 		TSharedPtr<ElectraDecodersUtil::MPEG::FISO23008_2_seq_parameter_set_data, ESPMode::ThreadSafe> SPS;
+		bool bDropOutput = false;
 	};
 
 	enum class EDecodeState
@@ -205,6 +206,7 @@ void IElectraVideoDecoderH265_Linux::GetConfigurationOptions(TMap<FString, FVari
 	OutOptions.Emplace(IElectraDecoderFeature::MinimumNumberOfOutputFrames, FVariant((int32)8));
 	// The decoder is adaptive. There is no need to call IsCompatibleWith() on stream switches.
 	OutOptions.Emplace(IElectraDecoderFeature::IsAdaptive, FVariant(true));
+	OutOptions.Emplace(IElectraDecoderFeature::SupportsDroppingOutput, FVariant(true));
 }
 
 TSharedPtr<IElectraDecoder, ESPMode::ThreadSafe> IElectraVideoDecoderH265_Linux::Create(const TMap<FString, FVariant>& InOptions, TSharedPtr<IElectraDecoderResourceDelegate, ESPMode::ThreadSafe> InResourceDelegate)
@@ -349,6 +351,15 @@ IElectraDecoder::EDecoderError FElectraVideoDecoderH265_Linux::DecodeAccessUnit(
 		return IElectraDecoder::EDecoderError::EndOfData;
 	}
 
+	// CSD only buffer is not handled at the moment.
+	check((InInputAccessUnit.Flags & EElectraDecoderFlags::InitCSDOnly) == EElectraDecoderFlags::None);
+
+	// If this is discardable and won't be output we do not need to handle it at all.
+	if ((InInputAccessUnit.Flags & (EElectraDecoderFlags::DoNotOutput | EElectraDecoderFlags::IsDiscardable)) == (EElectraDecoderFlags::DoNotOutput | EElectraDecoderFlags::IsDiscardable))
+	{
+		return IElectraDecoder::EDecoderError::None;
+	}
+
 	// If there is pending output it is very likely that decoding this access unit would also generate output.
 	// Since that would result in loss of the pending output we return now.
 	if (CurrentOutput.IsValid())
@@ -356,9 +367,6 @@ IElectraDecoder::EDecoderError FElectraVideoDecoderH265_Linux::DecodeAccessUnit(
 		return IElectraDecoder::EDecoderError::NoBuffer;
 	}
 
-	// CSD only buffer is not handled at the moment.
-	check((InInputAccessUnit.Flags & EElectraDecoderFlags::InitCSDOnly) == EElectraDecoderFlags::None);
-
 	// If a new decoder is needed, destroy the current one.
 	if (bNewDecoderRequired)
 	{
@@ -412,6 +420,7 @@ IElectraDecoder::EDecoderError FElectraVideoDecoderH265_Linux::DecodeAccessUnit(
 			TSharedPtr<FDecoderInput, ESPMode::NotThreadSafe> In(new FDecoderInput);
 			In->AdditionalOptions = InAdditionalOptions;
 			In->AccessUnit = InInputAccessUnit;
+			In->bDropOutput = (InInputAccessUnit.Flags & EElectraDecoderFlags::DoNotOutput) == EElectraDecoderFlags::DoNotOutput;
 			In->SPS = CurrentSPS;
 			// Zero the input pointer and size in the copy. That data is not owned by us and it's best not to have any
 			// values here that would lead us to think that we do.
@@ -628,7 +637,11 @@ FElectraVideoDecoderH265_Linux::EConvertResult FElectraVideoDecoderH265_Linux::C
 	NewOutput->ExtraValues.Emplace(TEXT("decoder"), FVariant(TEXT("libavcodec")));
 	NewOutput->ExtraValues.Emplace(TEXT("codec"), FVariant(TEXT("hevc")));
 
-	CurrentOutput = MoveTemp(NewOutput);
-
+	// Take this output on if it is not to be dropped.
+	// We need to have pulled it out of libav before deciding this though.
+	if (!In->bDropOutput)
+	{
+		CurrentOutput = MoveTemp(NewOutput);
+	}
 	return EConvertResult::Success;
 }
diff --git a/Engine/Plugins/Media/ElectraCodecs/Source/ElectraDecoders/Private/Linux/vpx/VideoDecoderVPx_Linux.cpp b/Engine/Plugins/Media/ElectraCodecs/Source/ElectraDecoders/Private/Linux/vpx/VideoDecoderVPx_Linux.cpp
--- a/Engine/Plugins/Media/ElectraCodecs/Source/ElectraDecoders/Private/Linux/vpx/VideoDecoderVPx_Linux.cpp
+++ b/Engine/Plugins/Media/ElectraCodecs/Source/ElectraDecoders/Private/Linux/vpx/VideoDecoderVPx_Linux.cpp
@@ -142,6 +142,7 @@ private:
 		FInputAccessUnit AccessUnit;
 		TMap<FString, FVariant> AdditionalOptions;
 		int32 NumBits = 8;
+		bool bDropOutput = false;
 	};
 
 	enum class EDecodeState
@@ -182,6 +183,7 @@ void IElectraVideoDecoderVPx_Linux::GetConfigurationOptions(TMap<FString, FVaria
 	OutOptions.Emplace(IElectraDecoderFeature::MinimumNumberOfOutputFrames, FVariant((int32)8));
 	// The decoder is adaptive. There is no need to call IsCompatibleWith() on stream switches.
 	OutOptions.Emplace(IElectraDecoderFeature::IsAdaptive, FVariant(true));
+	OutOptions.Emplace(IElectraDecoderFeature::SupportsDroppingOutput, FVariant(true));
 }
 
 TSharedPtr<IElectraDecoder, ESPMode::ThreadSafe> IElectraVideoDecoderVPx_Linux::Create(const TMap<FString, FVariant>& InOptions, TSharedPtr<IElectraDecoderResourceDelegate, ESPMode::ThreadSafe> InResourceDelegate)
@@ -273,6 +275,15 @@ IElectraDecoder::EDecoderError FElectraVideoDecoderVPx_Linux::DecodeAccessUnit(c
 		return IElectraDecoder::EDecoderError::EndOfData;
 	}
 
+	// CSD only buffer is not handled at the moment.
+	check((InInputAccessUnit.Flags & EElectraDecoderFlags::InitCSDOnly) == EElectraDecoderFlags::None);
+
+	// If this is discardable and won't be output we do not need to handle it at all.
+	if ((InInputAccessUnit.Flags & (EElectraDecoderFlags::DoNotOutput | EElectraDecoderFlags::IsDiscardable)) == (EElectraDecoderFlags::DoNotOutput | EElectraDecoderFlags::IsDiscardable))
+	{
+		return IElectraDecoder::EDecoderError::None;
+	}
+
 	// If there is pending output it is very likely that decoding this access unit would also generate output.
 	// Since that would result in loss of the pending output we return now.
 	if (CurrentOutput.IsValid())
@@ -280,9 +291,6 @@ IElectraDecoder::EDecoderError FElectraVideoDecoderVPx_Linux::DecodeAccessUnit(c
 		return IElectraDecoder::EDecoderError::NoBuffer;
 	}
 
-	// CSD only buffer is not handled at the moment.
-	check((InInputAccessUnit.Flags & EElectraDecoderFlags::InitCSDOnly) == EElectraDecoderFlags::None);
-
 	// If a new decoder is needed, destroy the current one.
 	if (bNewDecoderRequired)
 	{
@@ -335,6 +343,7 @@ IElectraDecoder::EDecoderError FElectraVideoDecoderVPx_Linux::DecodeAccessUnit(c
 			TSharedPtr<FDecoderInput, ESPMode::NotThreadSafe> In(new FDecoderInput);
 			In->AdditionalOptions = InAdditionalOptions;
 			In->AccessUnit = InInputAccessUnit;
+			In->bDropOutput = (InInputAccessUnit.Flags & EElectraDecoderFlags::DoNotOutput) == EElectraDecoderFlags::DoNotOutput;
 			In->NumBits = NumBits;
 			// Zero the input pointer and size in the copy. That data is not owned by us and it's best not to have any
 			// values here that would lead us to think that we do.
@@ -545,6 +554,11 @@ FElectraVideoDecoderVPx_Linux::EConvertResult FElectraVideoDecoderVPx_Linux::Con
 	NewOutput->ExtraValues.Emplace(TEXT("decoder"), FVariant(TEXT("libavcodec")));
 	NewOutput->ExtraValues.Emplace(TEXT("codec"), FVariant(TEXT("vpx")));
 
-	CurrentOutput = MoveTemp(NewOutput);
+	// Take this output on if it is not to be dropped.
+	// We need to have pulled it out of libav before deciding this though.
+	if (!In->bDropOutput)
+	{
+		CurrentOutput = MoveTemp(NewOutput);
+	}
 	return EConvertResult::Success;
 }
diff --git a/Engine/Plugins/Media/ElectraCodecs/Source/ElectraDecoders/Private/Windows/h264/H264_VideoDecoder_Windows.cpp b/Engine/Plugins/Media/ElectraCodecs/Source/ElectraDecoders/Private/Windows/h264/H264_VideoDecoder_Windows.cpp
--- a/Engine/Plugins/Media/ElectraCodecs/Source/ElectraDecoders/Private/Windows/h264/H264_VideoDecoder_Windows.cpp
+++ b/Engine/Plugins/Media/ElectraCodecs/Source/ElectraDecoders/Private/Windows/h264/H264_VideoDecoder_Windows.cpp
@@ -206,6 +206,7 @@ void IElectraVideoDecoderH264_DX::PlatformGetConfigurationOptions(TMap<FString,
 {
 	OutOptions.Emplace(IElectraDecoderFeature::MinimumNumberOfOutputFrames, FVariant((int32)8));
 	OutOptions.Emplace(IElectraDecoderFeature::IsAdaptive, FVariant(false));
+	OutOptions.Emplace(IElectraDecoderFeature::SupportsDroppingOutput, FVariant(true));
 }
 
 
diff --git a/Engine/Plugins/Media/ElectraCodecs/Source/ElectraDecoders/Private/Windows/h265/H265_VideoDecoder_Windows.cpp b/Engine/Plugins/Media/ElectraCodecs/Source/ElectraDecoders/Private/Windows/h265/H265_VideoDecoder_Windows.cpp
--- a/Engine/Plugins/Media/ElectraCodecs/Source/ElectraDecoders/Private/Windows/h265/H265_VideoDecoder_Windows.cpp
+++ b/Engine/Plugins/Media/ElectraCodecs/Source/ElectraDecoders/Private/Windows/h265/H265_VideoDecoder_Windows.cpp
@@ -192,6 +192,7 @@ void IElectraVideoDecoderH265_DX::PlatformGetConfigurationOptions(TMap<FString,
 {
 	OutOptions.Emplace(IElectraDecoderFeature::MinimumNumberOfOutputFrames, FVariant((int32)8));
 	OutOptions.Emplace(IElectraDecoderFeature::IsAdaptive, FVariant(false));
+	OutOptions.Emplace(IElectraDecoderFeature::SupportsDroppingOutput, FVariant(true));
 }
 
 
diff --git a/Engine/Plugins/Media/ElectraCodecs/Source/ElectraDecoders/Public/IElectraDecoder.h b/Engine/Plugins/Media/ElectraCodecs/Source/ElectraDecoders/Public/IElectraDecoder.h
--- a/Engine/Plugins/Media/ElectraCodecs/Source/ElectraDecoders/Public/IElectraDecoder.h
+++ b/Engine/Plugins/Media/ElectraCodecs/Source/ElectraDecoders/Public/IElectraDecoder.h
@@ -38,11 +38,11 @@ enum class EElectraDecoderFlags : uint32
 {
 	None = 0,
 	IsSyncSample = 0x01,
-	IsReplaySample = 0x02,
-	IsLastReplaySample = 0x04,
-	InitCSDOnly = 0x08,
-	//DoNotOutput = 0x10,
-	//TimeUpdateOnly = 0x20
+	IsDiscardable = 0x02,
+	IsReplaySample = 0x04,
+	IsLastReplaySample = 0x08,
+	InitCSDOnly = 0x10,
+	DoNotOutput = 0x20
 };
 ENUM_CLASS_FLAGS(EElectraDecoderFlags);
 
diff --git a/Engine/Plugins/Media/ElectraCodecs/Source/ElectraDecoders/Public/IElectraDecoderFeaturesAndOptions.h b/Engine/Plugins/Media/ElectraCodecs/Source/ElectraDecoders/Public/IElectraDecoderFeaturesAndOptions.h
--- a/Engine/Plugins/Media/ElectraCodecs/Source/ElectraDecoders/Public/IElectraDecoderFeaturesAndOptions.h
+++ b/Engine/Plugins/Media/ElectraCodecs/Source/ElectraDecoders/Public/IElectraDecoderFeaturesAndOptions.h
@@ -35,6 +35,16 @@ namespace IElectraDecoderFeature
 	 */
 	const TCHAR* const IsAdaptive = TEXT("is_adaptive");
 
+	/*
+	 * FVariant(bool) : true - adaptive decoding supported, false - not supported
+	 *
+	 * Indicates if the decoder supports dropping output frames that we do not want to display.
+	 * Such frames still need to be decoded to update the internal decoder state but the output
+	 * is not needed and can already be discarded on the decoder level instead of being converted
+	 * and passed up the pipeline where it will then be discarded.
+	 */
+	const TCHAR* const SupportsDroppingOutput = TEXT("supports_drop_output");
+
 	/**
 	 * FVariant(int32)
 	 *    Value to add to the length of the decoding unit.
diff --git a/Engine/Plugins/Media/ElectraPlayer/Source/ElectraPlayerRuntime/Private/Runtime/Decoder/BitstreamProcessors/VideoDecoderInputBitstreamProcessorVP8.cpp b/Engine/Plugins/Media/ElectraPlayer/Source/ElectraPlayerRuntime/Private/Runtime/Decoder/BitstreamProcessors/VideoDecoderInputBitstreamProcessorVP8.cpp
--- a/Engine/Plugins/Media/ElectraPlayer/Source/ElectraPlayerRuntime/Private/Runtime/Decoder/BitstreamProcessors/VideoDecoderInputBitstreamProcessorVP8.cpp
+++ b/Engine/Plugins/Media/ElectraPlayer/Source/ElectraPlayerRuntime/Private/Runtime/Decoder/BitstreamProcessors/VideoDecoderInputBitstreamProcessorVP8.cpp
@@ -85,7 +85,16 @@ IVideoDecoderInputBitstreamProcessor::EProcessResult FVideoDecoderInputBitstream
 		}
 	}
 
-	OutBSI.bIsSyncFrame = InOutAccessUnit->bIsSyncSample;
+	// Parse the bitstream header. The sync sample indicator is not always correctly set by the packager.
+	ElectraDecodersUtil::VPxVideo::FVP8UncompressedHeader Header;
+	if (ElectraDecodersUtil::VPxVideo::ParseVP8UncompressedHeader(Header, InOutAccessUnit->AUData, InOutAccessUnit->AUSize))
+	{
+		OutBSI.bIsSyncFrame = Header.IsKeyframe();
+	}
+	else
+	{
+		OutBSI.bIsSyncFrame = InOutAccessUnit->bIsSyncSample;
+	}
 	OutBSI.bIsDiscardable = false;
 
 	return Result;
diff --git a/Engine/Plugins/Media/ElectraPlayer/Source/ElectraPlayerRuntime/Private/Runtime/Decoder/BitstreamProcessors/VideoDecoderInputBitstreamProcessorVP9.cpp b/Engine/Plugins/Media/ElectraPlayer/Source/ElectraPlayerRuntime/Private/Runtime/Decoder/BitstreamProcessors/VideoDecoderInputBitstreamProcessorVP9.cpp
--- a/Engine/Plugins/Media/ElectraPlayer/Source/ElectraPlayerRuntime/Private/Runtime/Decoder/BitstreamProcessors/VideoDecoderInputBitstreamProcessorVP9.cpp
+++ b/Engine/Plugins/Media/ElectraPlayer/Source/ElectraPlayerRuntime/Private/Runtime/Decoder/BitstreamProcessors/VideoDecoderInputBitstreamProcessorVP9.cpp
@@ -85,7 +85,16 @@ IVideoDecoderInputBitstreamProcessor::EProcessResult FVideoDecoderInputBitstream
 		}
 	}
 
-	OutBSI.bIsSyncFrame = InOutAccessUnit->bIsSyncSample;
+	// Parse the bitstream header. The sync sample indicator is not always correctly set by the packager.
+	ElectraDecodersUtil::VPxVideo::FVP9UncompressedHeader Header;
+	if (ElectraDecodersUtil::VPxVideo::ParseVP9UncompressedHeader(Header, InOutAccessUnit->AUData, InOutAccessUnit->AUSize))
+	{
+		OutBSI.bIsSyncFrame = Header.IsKeyframe();
+	}
+	else
+	{
+		OutBSI.bIsSyncFrame = InOutAccessUnit->bIsSyncSample;
+	}
 	OutBSI.bIsDiscardable = false;
 
 	return Result;
diff --git a/Engine/Plugins/Media/ElectraPlayer/Source/ElectraPlayerRuntime/Private/Runtime/Decoder/VideoDecoder.cpp b/Engine/Plugins/Media/ElectraPlayer/Source/ElectraPlayerRuntime/Private/Runtime/Decoder/VideoDecoder.cpp
--- a/Engine/Plugins/Media/ElectraPlayer/Source/ElectraPlayerRuntime/Private/Runtime/Decoder/VideoDecoder.cpp
+++ b/Engine/Plugins/Media/ElectraPlayer/Source/ElectraPlayerRuntime/Private/Runtime/Decoder/VideoDecoder.cpp
@@ -169,10 +169,13 @@ private:
 	EDecodingState															NextDecodingStateAfterDrain = EDecodingState::NormalDecoding;
 	bool																	bIsDecoderClean = true;
 	bool																	bDrainAfterDecode = false;
+	int32																	MinLoopSleepTimeMsec = 0;
 
 	bool																	bIsFirstAccessUnit = true;
 	bool																	bInDummyDecodeMode = false;
 	bool																	bDrainForCodecChange = false;
+	bool																	bWaitForSyncSample = true;
+	bool																	bWarnedMissingSyncSample = false;
 
 	bool																	bError = false;
 
@@ -206,6 +209,7 @@ private:
 	TMap<FString, FVariant>													DecoderConfigOptions;
 	TSharedPtr<IElectraDecoder, ESPMode::ThreadSafe>						DecoderInstance;
 	bool																	bIsAdaptiveDecoder = false;
+	bool																	bSupportsDroppingOutput = false;
 	bool																	bNeedsReplayData = true;
 	bool																	bMustBeSuspendedInBackground = false;
 
@@ -522,6 +526,7 @@ bool FVideoDecoderImpl::InternalDecoderCreate()
 	TMap<FString, FVariant> Features;
 	DecoderInstance->GetFeatures(Features);
 	bIsAdaptiveDecoder = ElectraDecodersUtil::GetVariantValueSafeI64(Features, IElectraDecoderFeature::IsAdaptive, 0) != 0;
+	bSupportsDroppingOutput = ElectraDecodersUtil::GetVariantValueSafeI64(Features, IElectraDecoderFeature::SupportsDroppingOutput, 0) != 0;
 	bNeedsReplayData = ElectraDecodersUtil::GetVariantValueSafeI64(Features, IElectraDecoderFeature::NeedReplayDataOnDecoderLoss, 0) != 0;
 	// If replay data is not needed we can let go of anything we may have collected (which should be only the first access unit).
 	if (!bNeedsReplayData)
@@ -563,6 +568,7 @@ void FVideoDecoderImpl::InternalDecoderDestroy()
 		PlatformResource = nullptr;
 	}
 	bIsAdaptiveDecoder = false;
+	bSupportsDroppingOutput = false;
 	bNeedsReplayData = true;
 }
 
@@ -925,6 +931,14 @@ FVideoDecoderImpl::ENextDecodingState FVideoDecoderImpl::HandleDecoding()
 			DecAU.Duration = CurrentAccessUnit->AccessUnit->Duration.GetAsTimespan();
 			DecAU.UserValue = CurrentAccessUnit->PTS;
 			DecAU.Flags = CurrentAccessUnit->BitstreamInfo.bIsSyncFrame ? EElectraDecoderFlags::IsSyncSample : EElectraDecoderFlags::None;
+			if (CurrentAccessUnit->BitstreamInfo.bIsDiscardable)
+			{
+				DecAU.Flags |= EElectraDecoderFlags::IsDiscardable;
+			}
+			if (bSupportsDroppingOutput && !CurrentAccessUnit->AdjustedPTS.IsValid())
+			{
+				DecAU.Flags |= EElectraDecoderFlags::DoNotOutput;
+			}
 			TMap<FString, FVariant> CSDOptions;
 			if (CurrentAccessUnit->BitstreamInfo.bIsSyncFrame && CurrentAccessUnit->AccessUnit->AUCodecData.IsValid())
 			{
@@ -932,14 +946,38 @@ FVideoDecoderImpl::ENextDecodingState FVideoDecoderImpl::HandleDecoding()
 				CSDOptions.Emplace(TEXT("dcr"), FVariant(CurrentAccessUnit->AccessUnit->AUCodecData->RawCSD));
 			}
 
+			// Need to wait for a sync sample?
+			if (bWaitForSyncSample && !CurrentAccessUnit->BitstreamInfo.bIsSyncFrame)
+			{
+				if (!bWarnedMissingSyncSample)
+				{
+					bWarnedMissingSyncSample = true;
+					UE_LOG(LogElectraPlayer, Warning, TEXT("Expected a video sync sample at PTS %lld, but did not get one. The stream may be packaged incorrectly. Dropping frames until one arrives, which may take a while. Please wait!"), (long long int)DecAU.PTS.GetTicks());
+				}
+				bDrainAfterDecode = CurrentAccessUnit->AccessUnit->bIsLastInPeriod;
+				CurrentAccessUnit.Reset();
+				// Report this up as "stalled" so that we get out of prerolling.
+				// This case here happens when seeking due to bad sync frame information in the container format
+				// and the next sync frame may be too far away to satisfy the prerolling finished condition.
+				NotifyReadyBufferListener(false);
+				return ENextDecodingState::NormalDecoding;
+			}
+
 			IElectraDecoder::EDecoderError DecErr = DecoderInstance->DecodeAccessUnit(DecAU, CSDOptions);
 			if (DecErr == IElectraDecoder::EDecoderError::None)
 			{
-				InDecoderInput.Emplace(CurrentAccessUnit);
-				InDecoderInput.Sort([](const TSharedPtr<FDecoderInput, ESPMode::ThreadSafe>& a, const TSharedPtr<FDecoderInput, ESPMode::ThreadSafe>& b)
+				if ((DecAU.Flags & EElectraDecoderFlags::DoNotOutput) == EElectraDecoderFlags::None)
+				{
+					InDecoderInput.Emplace(CurrentAccessUnit);
+					InDecoderInput.Sort([](const TSharedPtr<FDecoderInput, ESPMode::ThreadSafe>& a, const TSharedPtr<FDecoderInput, ESPMode::ThreadSafe>& b)
+					{
+						return a->PTS < b->PTS;
+					});
+				}
+				else
 				{
-					return a->PTS < b->PTS;
-				});
+					MinLoopSleepTimeMsec = 0;
+				}
 
 				// If this was the last access unit in a period we need to drain the decoder _after_ having sent it
 				// for decoding. We need to get its decoded output.
@@ -947,6 +985,8 @@ FVideoDecoderImpl::ENextDecodingState FVideoDecoderImpl::HandleDecoding()
 				CurrentAccessUnit.Reset();
 				// Since we decoded something the decoder is no longer clean.
 				bIsDecoderClean = false;
+				// Likewise we are no longer waiting for a sync sample.
+				bWaitForSyncSample = false;
 			}
 			else if (DecErr == IElectraDecoder::EDecoderError::NoBuffer)
 			{
@@ -1177,6 +1217,8 @@ bool FVideoDecoderImpl::CheckForFlush()
 		CurrentAccessUnit.Reset();
 		bIsDecoderClean = true;
 		bInDummyDecodeMode = false;
+		bWaitForSyncSample = true;
+		bWarnedMissingSyncSample = false;
 		CurrentDecodingState = EDecodingState::NormalDecoding;
 		BitstreamProcessor->Clear();
 		FlushDecoderSignal.Reset();
@@ -1223,10 +1265,13 @@ void FVideoDecoderImpl::WorkerThread()
 	bIsFirstAccessUnit = true;
 	bInDummyDecodeMode = false;
 	bIsAdaptiveDecoder = false;
+	bSupportsDroppingOutput = false;
 	// Start out assuming replay data will be needed. We only know this for sure once we have created a decoder instance.
 	bNeedsReplayData = true;
 	bDrainAfterDecode = false;
 	bIsDecoderClean = true;
+	bWaitForSyncSample = true;
+	bWarnedMissingSyncSample = false;
 	CurrentDecodingState = EDecodingState::NormalDecoding;
 
 	check(InitialCodecSpecificData.IsValid());
@@ -1254,6 +1299,7 @@ void FVideoDecoderImpl::WorkerThread()
 	CreateDecoderOutputPool();
 
 	int64 TimeLast = MEDIAutcTime::CurrentMSec();
+	const int32 kDefaultMinLoopSleepTimeMS = 5;
 	while(!TerminateThreadSignal.IsSignaled())
 	{
 		if (CheckBackgrounding())
@@ -1271,12 +1317,17 @@ void FVideoDecoderImpl::WorkerThread()
 		// To prevent this from becoming a tight loop we make sure to sleep at least some time  here to throttle down.
 		int64 TimeNow = MEDIAutcTime::CurrentMSec();
 		int64 elapsedMS = TimeNow - TimeLast;
-		const int32 kTotalSleepTimeMsec = 5;
-		if (elapsedMS < kTotalSleepTimeMsec)
+		if (elapsedMS < MinLoopSleepTimeMsec)
 		{
-			FMediaRunnable::SleepMilliseconds(kTotalSleepTimeMsec - elapsedMS);
+			FMediaRunnable::SleepMilliseconds(MinLoopSleepTimeMsec - elapsedMS);
+		}
+		else
+		{
+			FPlatformProcess::YieldThread();
 		}
 		TimeLast = TimeNow;
+		MinLoopSleepTimeMsec = kDefaultMinLoopSleepTimeMS;
+
 
 		// Get the next access unit to decode.
 		EAUChangeFlags NewAUFlags = GetAndPrepareInputAU();
