ElectraPlayer: Added audio stream metric collection

#rnx

diff --git a/Engine/Plugins/Media/ElectraPlayer/Source/ElectraPlayerRuntime/Private/ElectraPlayer.cpp b/Engine/Plugins/Media/ElectraPlayer/Source/ElectraPlayerRuntime/Private/ElectraPlayer.cpp
--- a/Engine/Plugins/Media/ElectraPlayer/Source/ElectraPlayerRuntime/Private/ElectraPlayer.cpp
+++ b/Engine/Plugins/Media/ElectraPlayer/Source/ElectraPlayerRuntime/Private/ElectraPlayer.cpp
@@ -1883,6 +1883,12 @@ void FElectraPlayer::HandleDeferredPlayerEvents()
 				HandlePlayerEventVideoQualityChange(Ev->NewBitrate, Ev->PreviousBitrate, Ev->bIsDrasticDownswitch);
 				break;
 			}
+			case FPlayerMetricEventBase::EType::AudioQualityChange:
+			{
+				FPlayerMetricEvent_AudioQualityChange* Ev = static_cast<FPlayerMetricEvent_AudioQualityChange*>(Event.Get());
+				HandlePlayerEventAudioQualityChange(Ev->NewBitrate, Ev->PreviousBitrate, Ev->bIsDrasticDownswitch);
+				break;
+			}
 			case FPlayerMetricEventBase::EType::CodecFormatChange:
 			{
 				FPlayerMetricEvent_CodecFormatChange* Ev = static_cast<FPlayerMetricEvent_CodecFormatChange*>(Event.Get());
@@ -2047,6 +2053,7 @@ void FElectraPlayer::HandlePlayerEventReceivedPlaylists()
 	Statistics.MediaTimelineAtStart = MediaTimeline;
 	Statistics.MediaTimelineAtEnd = MediaTimeline;
 	Statistics.MediaDuration = MediaDuration.IsInfinity() ? -1.0 : MediaDuration.GetAsSeconds();
+	StatisticsLock.Unlock();
 	// Get the video bitrates and populate our number of segments per bitrate map.
 	TArray<Electra::FTrackMetadata> VideoStreamMetaData;
 	CurrentPlayer->AdaptivePlayer->GetTrackMetadata(VideoStreamMetaData, Electra::EStreamType::Video);
@@ -2061,7 +2068,6 @@ void FElectraPlayer::HandlePlayerEventReceivedPlaylists()
 												VideoStreamMetaData[0].StreamDetails[i].Bandwidth);
 		}
 	}
-	StatisticsLock.Unlock();
 
 	SelectedVideoTrackIndex = NumTracksVideo ? 0 : -1;
 
@@ -2073,9 +2079,18 @@ void FElectraPlayer::HandlePlayerEventReceivedPlaylists()
 		EnqueueAnalyticsEvent(AnalyticEvent);
 	}
 
+	// Get the audio bitrates and populate our number of segments per bitrate map.
 	TArray<Electra::FTrackMetadata> AudioStreamMetaData;
 	CurrentPlayer->AdaptivePlayer->GetTrackMetadata(AudioStreamMetaData, Electra::EStreamType::Audio);
 	NumTracksAudio = AudioStreamMetaData.Num();
+	if (NumTracksAudio)
+	{
+		for(int32 i=0; i<AudioStreamMetaData[0].StreamDetails.Num(); ++i)
+		{
+			UE_LOG(LogElectraPlayer, Log, TEXT("[%p][%p] Found audio stream at bitrate %d"), this, CurrentPlayer.Get(),
+												AudioStreamMetaData[0].StreamDetails[i].Bandwidth);
+		}
+	}
 
 	TArray<Electra::FTrackMetadata> SubtitleStreamMetaData;
 	CurrentPlayer->AdaptivePlayer->GetTrackMetadata(SubtitleStreamMetaData, Electra::EStreamType::Subtitle);
@@ -2121,6 +2136,14 @@ void FElectraPlayer::HandlePlayerEventTracksChanged()
 	TArray<Electra::FTrackMetadata> AudioStreamMetaData;
 	CurrentPlayer->AdaptivePlayer->GetTrackMetadata(AudioStreamMetaData, Electra::EStreamType::Audio);
 	NumTracksAudio = AudioStreamMetaData.Num();
+	if (NumTracksAudio)
+	{
+		for(int32 i=0; i<AudioStreamMetaData[0].StreamDetails.Num(); ++i)
+		{
+			UE_LOG(LogElectraPlayer, Log, TEXT("[%p][%p] Found audio stream at bitrate %d"), this, CurrentPlayer.Get(),
+												AudioStreamMetaData[0].StreamDetails[i].Bandwidth);
+		}
+	}
 
 	TArray<Electra::FTrackMetadata> SubtitleStreamMetaData;
 	CurrentPlayer->AdaptivePlayer->GetTrackMetadata(SubtitleStreamMetaData, Electra::EStreamType::Subtitle);
@@ -2395,30 +2418,30 @@ void FElectraPlayer::HandlePlayerEventVideoQualityChange(int32 NewBitrate, int32
 	FScopeLock Lock(&StatisticsLock);
 	if (PreviousBitrate == 0)
 	{
-		Statistics.InitialStreamBitrate = NewBitrate;
+		Statistics.InitialVideoStreamBitrate = NewBitrate;
 	}
 	else
 	{
 		if (bIsDrasticDownswitch)
 		{
-			++Statistics.NumQualityDrasticDownswitches;
+			++Statistics.NumVideoQualityDrasticDownswitches;
 		}
 		if (NewBitrate > PreviousBitrate)
 		{
-			++Statistics.NumQualityUpswitches;
+			++Statistics.NumVideoQualityUpswitches;
 		}
 		else
 		{
-			++Statistics.NumQualityDownswitches;
+			++Statistics.NumVideoQualityDownswitches;
 		}
 	}
 	if (bIsDrasticDownswitch)
 	{
-		UE_LOG(LogElectraPlayer, Log, TEXT("[%p][%p] Player switched video quality drastically down to %d bps from %d bps. %d upswitches, %d downswitches (%d drastic ones)"), this, CurrentPlayer.Get(), NewBitrate, PreviousBitrate, Statistics.NumQualityUpswitches, Statistics.NumQualityDownswitches, Statistics.NumQualityDrasticDownswitches);
+		UE_LOG(LogElectraPlayer, Log, TEXT("[%p][%p] Player switched video quality drastically down to %d bps from %d bps. %d upswitches, %d downswitches (%d drastic ones)"), this, CurrentPlayer.Get(), NewBitrate, PreviousBitrate, Statistics.NumVideoQualityUpswitches, Statistics.NumVideoQualityDownswitches, Statistics.NumVideoQualityDrasticDownswitches);
 	}
 	else
 	{
-		UE_LOG(LogElectraPlayer, Log, TEXT("[%p][%p] Player switched video quality to %d bps from %d bps. %d upswitches, %d downswitches (%d drastic ones)"), this, CurrentPlayer.Get(), NewBitrate, PreviousBitrate, Statistics.NumQualityUpswitches, Statistics.NumQualityDownswitches, Statistics.NumQualityDrasticDownswitches);
+		UE_LOG(LogElectraPlayer, Log, TEXT("[%p][%p] Player switched video quality to %d bps from %d bps. %d upswitches, %d downswitches (%d drastic ones)"), this, CurrentPlayer.Get(), NewBitrate, PreviousBitrate, Statistics.NumVideoQualityUpswitches, Statistics.NumVideoQualityDownswitches, Statistics.NumVideoQualityDrasticDownswitches);
 	}
 
 	int32 prvWidth  = Statistics.CurrentlyActiveResolutionWidth;
@@ -2456,7 +2479,55 @@ void FElectraPlayer::HandlePlayerEventVideoQualityChange(int32 NewBitrate, int32
 
 	Statistics.AddMessageToHistory(FString::Printf(TEXT("Video bitrate change from %d to %d"), PreviousBitrate, NewBitrate));
 
-	CSV_EVENT(ElectraPlayer, TEXT("QualityChange %d -> %d"), PreviousBitrate, NewBitrate);
+	CSV_EVENT(ElectraPlayer, TEXT("VideoQualityChange %d -> %d"), PreviousBitrate, NewBitrate);
+}
+
+void FElectraPlayer::HandlePlayerEventAudioQualityChange(int32 NewBitrate, int32 PreviousBitrate, bool bIsDrasticDownswitch)
+{
+	// Update statistics
+	FScopeLock Lock(&StatisticsLock);
+	if (PreviousBitrate == 0)
+	{
+		Statistics.InitialAudioStreamBitrate = NewBitrate;
+	}
+	else
+	{
+		if (bIsDrasticDownswitch)
+		{
+			++Statistics.NumAudioQualityDrasticDownswitches;
+		}
+		if (NewBitrate > PreviousBitrate)
+		{
+			++Statistics.NumAudioQualityUpswitches;
+		}
+		else
+		{
+			++Statistics.NumAudioQualityDownswitches;
+		}
+	}
+	if (bIsDrasticDownswitch)
+	{
+		UE_LOG(LogElectraPlayer, Log, TEXT("[%p][%p] Player switched audio quality drastically down to %d bps from %d bps. %d upswitches, %d downswitches (%d drastic ones)"), this, CurrentPlayer.Get(), NewBitrate, PreviousBitrate, Statistics.NumAudioQualityUpswitches, Statistics.NumAudioQualityDownswitches, Statistics.NumAudioQualityDrasticDownswitches);
+	}
+	else
+	{
+		UE_LOG(LogElectraPlayer, Log, TEXT("[%p][%p] Player switched audio quality to %d bps from %d bps. %d upswitches, %d downswitches (%d drastic ones)"), this, CurrentPlayer.Get(), NewBitrate, PreviousBitrate, Statistics.NumAudioQualityUpswitches, Statistics.NumAudioQualityDownswitches, Statistics.NumAudioQualityDrasticDownswitches);
+	}
+
+	// Enqueue a "AudioQualityChange" event.
+	static const FString kEventNameElectraAudioQualityChange(TEXT("Electra.AudioQualityChange"));
+	if (Electra::IsAnalyticsEventEnabled(kEventNameElectraAudioQualityChange))
+	{
+		TSharedPtr<FAnalyticsEvent> AnalyticEvent = CreateAnalyticsEvent(kEventNameElectraAudioQualityChange);
+		AnalyticEvent->ParamArray.Add(FAnalyticsEventAttribute(TEXT("OldBitrate"), PreviousBitrate));
+		AnalyticEvent->ParamArray.Add(FAnalyticsEventAttribute(TEXT("NewBitrate"), NewBitrate));
+		AnalyticEvent->ParamArray.Add(FAnalyticsEventAttribute(TEXT("bIsDrasticDownswitch"), bIsDrasticDownswitch));
+		EnqueueAnalyticsEvent(AnalyticEvent);
+	}
+
+	Statistics.AddMessageToHistory(FString::Printf(TEXT("Audio bitrate change from %d to %d"), PreviousBitrate, NewBitrate));
+
+	CSV_EVENT(ElectraPlayer, TEXT("AudioQualityChange %d -> %d"), PreviousBitrate, NewBitrate);
 }
 
 void FElectraPlayer::HandlePlayerEventCodecFormatChange(const Electra::FStreamCodecInformation& NewDecodingFormat)
@@ -2859,6 +2930,7 @@ void FElectraPlayer::LogStatistics()
 			"Time after stream playlists loaded: %.3fs\n"\
 			"Time for initial buffering: %.3fs\n"\
 			"Initial video stream bitrate: %d bps\n"\
+			"Initial audio stream bitrate: %d bps\n"\
 			"Initial buffering bandwidth bps: %.3f\n"\
 			"Initial buffering latency: %.3fs\n"\
 			"Time for initial preroll: %.3fs\n"\
@@ -2875,9 +2947,12 @@ void FElectraPlayer::LogStatistics()
 			"Media duration: %.3fs\n"\
 			"Play position at start: %.3fs\n"\
 			"Play position at end: %.3fs\n"\
-			"Number of quality upswitches: %d\n"\
-			"Number of quality downswitches: %d\n"\
-			"Number of drastic downswitches: %d\n"\
+			"Number of video quality upswitches: %d\n"\
+			"Number of video quality downswitches: %d\n"\
+			"Number of video drastic downswitches: %d\n"\
+			"Number of audio quality upswitches: %d\n"\
+			"Number of audio quality downswitches: %d\n"\
+			"Number of audio drastic downswitches: %d\n"\
 			"Bytes of video data streamed: %lld\n"\
 			"Bytes of audio data streamed: %lld\n"\
 			"Number of segments fetched across all quality levels:\n%s"\
@@ -2893,7 +2968,8 @@ void FElectraPlayer::LogStatistics()
 			Statistics.TimeToLoadMasterPlaylist,
 			Statistics.TimeToLoadStreamPlaylists,
 			Statistics.InitialBufferingDuration,
-			Statistics.InitialStreamBitrate,
+			Statistics.InitialVideoStreamBitrate,
+			Statistics.InitialAudioStreamBitrate,
 			Statistics.InitialBufferingBandwidth.GetAverageBandwidth(),
 			Statistics.InitialBufferingBandwidth.GetAverageLatency(),
 			Statistics.TimeForInitialPreroll,
@@ -2910,9 +2986,12 @@ void FElectraPlayer::LogStatistics()
 			Statistics.MediaDuration,
 			Statistics.PlayPosAtStart,
 			Statistics.PlayPosAtEnd,
-			Statistics.NumQualityUpswitches,
-			Statistics.NumQualityDownswitches,
-			Statistics.NumQualityDrasticDownswitches,
+			Statistics.NumVideoQualityUpswitches,
+			Statistics.NumVideoQualityDownswitches,
+			Statistics.NumVideoQualityDrasticDownswitches,
+			Statistics.NumAudioQualityUpswitches,
+			Statistics.NumAudioQualityDownswitches,
+			Statistics.NumAudioQualityDrasticDownswitches,
 			(long long int)Statistics.NumVideoDatabytesStreamed,
 			(long long int)Statistics.NumAudioDatabytesStreamed,
 			*SegsPerStream,
@@ -2981,7 +3060,8 @@ void FElectraPlayer::SendAnalyticMetrics(const TSharedPtr<IAnalyticsProviderET>&
 	ParamArray.Add(FAnalyticsEventAttribute(TEXT("TimeElapsedToPlaylists"), Statistics.TimeToLoadStreamPlaylists));
 	ParamArray.Add(FAnalyticsEventAttribute(TEXT("InitialAvgBufferingBandwidth"), Statistics.InitialBufferingBandwidth.GetAverageBandwidth()));
 	ParamArray.Add(FAnalyticsEventAttribute(TEXT("InitialAvgBufferingLatency"), Statistics.InitialBufferingBandwidth.GetAverageLatency()));
-	ParamArray.Add(FAnalyticsEventAttribute(TEXT("InitialVideoBitrate"), Statistics.InitialStreamBitrate));
+	ParamArray.Add(FAnalyticsEventAttribute(TEXT("InitialVideoBitrate"), Statistics.InitialVideoStreamBitrate));
+	ParamArray.Add(FAnalyticsEventAttribute(TEXT("InitialAudioBitrate"), Statistics.InitialAudioStreamBitrate));
 	ParamArray.Add(FAnalyticsEventAttribute(TEXT("InitialBufferingDuration"), Statistics.InitialBufferingDuration));
 	ParamArray.Add(FAnalyticsEventAttribute(TEXT("InitialPrerollDuration"), Statistics.TimeForInitialPreroll));
 	ParamArray.Add(FAnalyticsEventAttribute(TEXT("TimeElapsedUntilReady"), Statistics.TimeForInitialPreroll + Statistics.TimeAtPrerollBegin - Statistics.TimeAtOpen));
@@ -2998,9 +3078,12 @@ void FElectraPlayer::SendAnalyticMetrics(const TSharedPtr<IAnalyticsProviderET>&
 	ParamArray.Add(FAnalyticsEventAttribute(TEXT("NumTimesMovedBackward"), (uint32) Statistics.NumTimesRewound));
 	ParamArray.Add(FAnalyticsEventAttribute(TEXT("NumTimesLooped"), (uint32) Statistics.NumTimesLooped));
 	ParamArray.Add(FAnalyticsEventAttribute(TEXT("AbortedSegmentDownloads"), (uint32) Statistics.NumSegmentDownloadsAborted));
-	ParamArray.Add(FAnalyticsEventAttribute(TEXT("NumQualityUpswitches"), (uint32) Statistics.NumQualityUpswitches));
-	ParamArray.Add(FAnalyticsEventAttribute(TEXT("NumQualityDownswitches"), (uint32) Statistics.NumQualityDownswitches));
-	ParamArray.Add(FAnalyticsEventAttribute(TEXT("NumQualityDrasticDownswitches"), (uint32) Statistics.NumQualityDrasticDownswitches));
+	ParamArray.Add(FAnalyticsEventAttribute(TEXT("NumQualityUpswitches"), (uint32) Statistics.NumVideoQualityUpswitches));
+	ParamArray.Add(FAnalyticsEventAttribute(TEXT("NumQualityDownswitches"), (uint32) Statistics.NumVideoQualityDownswitches));
+	ParamArray.Add(FAnalyticsEventAttribute(TEXT("NumQualityDrasticDownswitches"), (uint32) Statistics.NumVideoQualityDrasticDownswitches));
+	ParamArray.Add(FAnalyticsEventAttribute(TEXT("AudioQualityUpswitches"), (uint32) Statistics.NumAudioQualityUpswitches));
+	ParamArray.Add(FAnalyticsEventAttribute(TEXT("AudioQualityDownswitches"), (uint32) Statistics.NumAudioQualityDownswitches));
+	ParamArray.Add(FAnalyticsEventAttribute(TEXT("AudioQualityDrasticDownswitches"), (uint32) Statistics.NumAudioQualityDrasticDownswitches));
 	ParamArray.Add(FAnalyticsEventAttribute(TEXT("Rebuffering.Num"), (uint32)Statistics.NumTimesRebuffered));
 	ParamArray.Add(FAnalyticsEventAttribute(TEXT("Rebuffering.AvgDuration"), Statistics.NumTimesRebuffered > 0 ? Statistics.TotalRebufferingDuration / Statistics.NumTimesRebuffered : 0.0));
 	ParamArray.Add(FAnalyticsEventAttribute(TEXT("Rebuffering.MaxDuration"), Statistics.LongestRebufferingDuration));
diff --git a/Engine/Plugins/Media/ElectraPlayer/Source/ElectraPlayerRuntime/Private/ElectraPlayer.h b/Engine/Plugins/Media/ElectraPlayer/Source/ElectraPlayerRuntime/Private/ElectraPlayer.h
--- a/Engine/Plugins/Media/ElectraPlayer/Source/ElectraPlayerRuntime/Private/ElectraPlayer.h
+++ b/Engine/Plugins/Media/ElectraPlayer/Source/ElectraPlayerRuntime/Private/ElectraPlayer.h
@@ -195,6 +195,7 @@ private:
 			LicenseKey,
 			DataAvailabilityChange,
 			VideoQualityChange,
+			AudioQualityChange,
 			CodecFormatChange,
 			PrerollStart,
 			PrerollEnd,
@@ -274,6 +275,13 @@ private:
 		int32 PreviousBitrate;
 		bool bIsDrasticDownswitch;
 	};
+	struct FPlayerMetricEvent_AudioQualityChange : public FPlayerMetricEventBase
+	{
+		FPlayerMetricEvent_AudioQualityChange(int32 InNewBitrate, int32 InPreviousBitrate, bool bInIsDrasticDownswitch) : FPlayerMetricEventBase(EType::AudioQualityChange), NewBitrate(InNewBitrate), PreviousBitrate(InPreviousBitrate), bIsDrasticDownswitch(bInIsDrasticDownswitch) {}
+		int32 NewBitrate;
+		int32 PreviousBitrate;
+		bool bIsDrasticDownswitch;
+	};
 	struct FPlayerMetricEvent_CodecFormatChange : public FPlayerMetricEventBase
 	{
 		FPlayerMetricEvent_CodecFormatChange(const FStreamCodecInformation& InNewDecodingFormat) : FPlayerMetricEventBase(EType::CodecFormatChange), NewDecodingFormat(InNewDecodingFormat) {}
@@ -347,6 +355,8 @@ private:
 	{ DeferredPlayerEvents.Enqueue(MakeSharedTS<FPlayerMetricEvent_DataAvailabilityChange>(DataAvailability)); }
 	virtual void ReportVideoQualityChange(int32 NewBitrate, int32 PreviousBitrate, bool bIsDrasticDownswitch) override
 	{ DeferredPlayerEvents.Enqueue(MakeSharedTS<FPlayerMetricEvent_VideoQualityChange>(NewBitrate, PreviousBitrate, bIsDrasticDownswitch)); }
+	virtual void ReportAudioQualityChange(int32 NewBitrate, int32 PreviousBitrate, bool bIsDrasticDownswitch) override
+	{ DeferredPlayerEvents.Enqueue(MakeSharedTS<FPlayerMetricEvent_AudioQualityChange>(NewBitrate, PreviousBitrate, bIsDrasticDownswitch)); }
 	virtual void ReportDecodingFormatChange(const FStreamCodecInformation& NewDecodingFormat) override
 	{ DeferredPlayerEvents.Enqueue(MakeSharedTS<FPlayerMetricEvent_CodecFormatChange>(NewDecodingFormat)); }
 	virtual void ReportPrerollStart() override
@@ -677,36 +687,40 @@ private:
 			CurrentlyActivePlaylistURL.Empty();
 			LastError.Empty();
 			LastState = "Empty";
-			TimeAtOpen  					= -1.0;
-			TimeToLoadMasterPlaylist		= -1.0;
-			TimeToLoadStreamPlaylists   	= -1.0;
-			InitialBufferingDuration		= -1.0;
-			InitialStreamBitrate			= 0;
-			TimeAtPrerollBegin  			= -1.0;
-			TimeForInitialPreroll   		= -1.0;
-			NumTimesRebuffered  			= 0;
-			NumTimesForwarded   			= 0;
-			NumTimesRewound 				= 0;
-			NumTimesLooped  				= 0;
-			TimeAtBufferingBegin			= 0.0;
-			TotalRebufferingDuration		= 0.0;
-			LongestRebufferingDuration  	= 0.0;
-			PlayPosAtStart  				= -1.0;
-			PlayPosAtEnd					= -1.0;
-			NumQualityUpswitches			= 0;
-			NumQualityDownswitches  		= 0;
-			NumQualityDrasticDownswitches   = 0;
-			NumVideoDatabytesStreamed   	= 0;
-			NumAudioDatabytesStreamed   	= 0;
-			NumSegmentDownloadsAborted  	= 0;
-			CurrentlyActiveResolutionWidth  = 0;
+			TimeAtOpen = -1.0;
+			TimeToLoadMasterPlaylist = -1.0;
+			TimeToLoadStreamPlaylists = -1.0;
+			InitialBufferingDuration = -1.0;
+			InitialVideoStreamBitrate = 0;
+			InitialAudioStreamBitrate = 0;
+			TimeAtPrerollBegin = -1.0;
+			TimeForInitialPreroll = -1.0;
+			NumTimesRebuffered = 0;
+			NumTimesForwarded = 0;
+			NumTimesRewound = 0;
+			NumTimesLooped = 0;
+			TimeAtBufferingBegin = 0.0;
+			TotalRebufferingDuration = 0.0;
+			LongestRebufferingDuration = 0.0;
+			PlayPosAtStart = -1.0;
+			PlayPosAtEnd = -1.0;
+			NumVideoQualityUpswitches = 0;
+			NumVideoQualityDownswitches = 0;
+			NumVideoQualityDrasticDownswitches = 0;
+			NumAudioQualityUpswitches = 0;
+			NumAudioQualityDownswitches = 0;
+			NumAudioQualityDrasticDownswitches = 0;
+			NumVideoDatabytesStreamed = 0;
+			NumAudioDatabytesStreamed = 0;
+			NumSegmentDownloadsAborted = 0;
+			CurrentlyActiveResolutionWidth = 0;
 			CurrentlyActiveResolutionHeight = 0;
 			VideoSegmentBitratesStreamed.Empty();
 			InitialBufferingBandwidth.Reset();
-			bIsInitiallyDownloading 		= false;
-			bDidPlaybackEnd 				= false;
+			bIsInitiallyDownloading = false;
+			bDidPlaybackEnd = false;
 			SubtitlesURL.Empty();
-			SubtitlesResponseTime			= 0.0;
+			SubtitlesResponseTime = 0.0;
 			SubtitlesLastError.Empty();
 			DroppedVideoFrames.Reset();
 			DroppedAudioFrames.Reset();
@@ -725,7 +739,8 @@ private:
 		double					TimeToLoadMasterPlaylist;
 		double					TimeToLoadStreamPlaylists;
 		double					InitialBufferingDuration;
-		int32					InitialStreamBitrate;
+		int32					InitialVideoStreamBitrate;
+		int32					InitialAudioStreamBitrate;
 		double					TimeAtPrerollBegin;
 		double					TimeForInitialPreroll;
 		int32					NumTimesRebuffered;
@@ -737,9 +752,12 @@ private:
 		double					LongestRebufferingDuration;
 		double					PlayPosAtStart;
 		double					PlayPosAtEnd;
-		int32					NumQualityUpswitches;
-		int32					NumQualityDownswitches;
-		int32					NumQualityDrasticDownswitches;
+		int32					NumVideoQualityUpswitches;
+		int32					NumVideoQualityDownswitches;
+		int32					NumVideoQualityDrasticDownswitches;
+		int32					NumAudioQualityUpswitches;
+		int32					NumAudioQualityDownswitches;
+		int32					NumAudioQualityDrasticDownswitches;
 		int64					NumVideoDatabytesStreamed;
 		int64					NumAudioDatabytesStreamed;
 		int32					NumSegmentDownloadsAborted;
@@ -805,6 +823,7 @@ private:
 	void HandlePlayerEventLicenseKey(const Metrics::FLicenseKeyStats& LicenseKeyStats);
 	void HandlePlayerEventDataAvailabilityChange(const Metrics::FDataAvailabilityChange& DataAvailability);
 	void HandlePlayerEventVideoQualityChange(int32 NewBitrate, int32 PreviousBitrate, bool bIsDrasticDownswitch);
+	void HandlePlayerEventAudioQualityChange(int32 NewBitrate, int32 PreviousBitrate, bool bIsDrasticDownswitch);
 	void HandlePlayerEventCodecFormatChange(const Electra::FStreamCodecInformation& NewDecodingFormat);
 	void HandlePlayerEventPrerollStart();
 	void HandlePlayerEventPrerollEnd();
diff --git a/Engine/Plugins/Media/ElectraPlayer/Source/ElectraPlayerRuntime/Private/Runtime/Player/AdaptiveStreamingPlayer.cpp b/Engine/Plugins/Media/ElectraPlayer/Source/ElectraPlayerRuntime/Private/Runtime/Player/AdaptiveStreamingPlayer.cpp
--- a/Engine/Plugins/Media/ElectraPlayer/Source/ElectraPlayerRuntime/Private/Runtime/Player/AdaptiveStreamingPlayer.cpp
+++ b/Engine/Plugins/Media/ElectraPlayer/Source/ElectraPlayerRuntime/Private/Runtime/Player/AdaptiveStreamingPlayer.cpp
@@ -448,6 +448,9 @@ void FAdaptiveStreamingPlayer::FireSyncEvent(TSharedPtrTS<FMetricEvent> Event)
 			case FMetricEvent::EType::VideoQualityChange:
 				Listeners[i]->ReportVideoQualityChange(Event->Param.QualityChange.NewBitrate, Event->Param.QualityChange.PrevBitrate, Event->Param.QualityChange.bIsDrastic);
 				break;
+			case FMetricEvent::EType::AudioQualityChange:
+				Listeners[i]->ReportAudioQualityChange(Event->Param.QualityChange.NewBitrate, Event->Param.QualityChange.PrevBitrate, Event->Param.QualityChange.bIsDrastic);
+				break;
 			case FMetricEvent::EType::CodecFormatChange:
 				Listeners[i]->ReportDecodingFormatChange(Event->Param.CodecFormatChange);
 				break;
@@ -1569,6 +1572,19 @@ bool FAdaptiveStreamingPlayer::InternalHandleThreadMessages()
 							CurrentVideoStreamBitrate.QualityLevel = SegmentQualityLevel;
 						}
 					}
+					// Audio bitrate change?
+					else if (ev.Request->GetType() == EStreamType::Audio)
+					{
+						int32 SegmentBitrate = pRequest->GetBitrate();
+						int32 SegmentQualityLevel = pRequest->GetQualityIndex();
+						if (SegmentBitrate != CurrentAudioStreamBitrate.Bitrate)
+						{
+							bool bDrastic = CurrentAudioStreamBitrate.Bitrate && SegmentQualityLevel < CurrentAudioStreamBitrate.QualityLevel-1;
+							DispatchEvent(FMetricEvent::ReportAudioQualityChange(SegmentBitrate, CurrentAudioStreamBitrate.Bitrate, bDrastic));
+							CurrentAudioStreamBitrate.Bitrate      = SegmentBitrate;
+							CurrentAudioStreamBitrate.QualityLevel = SegmentQualityLevel;
+						}
+					}
 				}
 				break;
 			}
diff --git a/Engine/Plugins/Media/ElectraPlayer/Source/ElectraPlayerRuntime/Private/Runtime/Player/AdaptiveStreamingPlayerABR.cpp b/Engine/Plugins/Media/ElectraPlayer/Source/ElectraPlayerRuntime/Private/Runtime/Player/AdaptiveStreamingPlayerABR.cpp
--- a/Engine/Plugins/Media/ElectraPlayer/Source/ElectraPlayerRuntime/Private/Runtime/Player/AdaptiveStreamingPlayerABR.cpp
+++ b/Engine/Plugins/Media/ElectraPlayer/Source/ElectraPlayerRuntime/Private/Runtime/Player/AdaptiveStreamingPlayerABR.cpp
@@ -60,6 +60,7 @@ namespace Electra
 		void ReportLicenseKey(const Metrics::FLicenseKeyStats& LicenseKeyStats) override {}
 		void ReportDataAvailabilityChange(const Metrics::FDataAvailabilityChange& DataAvailability) override {}
 		void ReportVideoQualityChange(int32 NewBitrate, int32 PreviousBitrate, bool bIsDrasticDownswitch) override {}
+		void ReportAudioQualityChange(int32 NewBitrate, int32 PreviousBitrate, bool bIsDrasticDownswitch) override {}
 		void ReportDecodingFormatChange(const FStreamCodecInformation& NewDecodingFormat) override {}
 		void ReportPrerollStart() override {}
 		void ReportPrerollEnd() override {}
diff --git a/Engine/Plugins/Media/ElectraPlayer/Source/ElectraPlayerRuntime/Private/Runtime/Player/AdaptiveStreamingPlayerInternal.h b/Engine/Plugins/Media/ElectraPlayer/Source/ElectraPlayerRuntime/Private/Runtime/Player/AdaptiveStreamingPlayerInternal.h
--- a/Engine/Plugins/Media/ElectraPlayer/Source/ElectraPlayerRuntime/Private/Runtime/Player/AdaptiveStreamingPlayerInternal.h
+++ b/Engine/Plugins/Media/ElectraPlayer/Source/ElectraPlayerRuntime/Private/Runtime/Player/AdaptiveStreamingPlayerInternal.h
@@ -617,6 +617,7 @@ struct FMetricEvent
 		SegmentDownload,
 		DataAvailabilityChange,
 		VideoQualityChange,
+		AudioQualityChange,
 		CodecFormatChange,
 		PrerollStart,
 		PrerollEnd,
@@ -789,6 +790,15 @@ struct FMetricEvent
 		Evt->Param.QualityChange.bIsDrastic = bIsDrasticDownswitch;
 		return Evt;
 	}
+	static TSharedPtrTS<FMetricEvent> ReportAudioQualityChange(int32 NewBitrate, int32 PreviousBitrate, bool bIsDrasticDownswitch)
+	{
+		TSharedPtrTS<FMetricEvent> Evt = MakeSharedTS<FMetricEvent>();
+		Evt->Type = EType::AudioQualityChange;
+		Evt->Param.QualityChange.NewBitrate = NewBitrate;
+		Evt->Param.QualityChange.PrevBitrate = PreviousBitrate;
+		Evt->Param.QualityChange.bIsDrastic = bIsDrasticDownswitch;
+		return Evt;
+	}
 	static TSharedPtrTS<FMetricEvent> ReportCodecFormatChange(const FStreamCodecInformation& NewDecodingFormat)
 	{
 		TSharedPtrTS<FMetricEvent> Evt = MakeSharedTS<FMetricEvent>();
@@ -2235,6 +2245,7 @@ private:
 	int32																VideoResolutionLimitHeight;
 
 	FStreamBitrateInfo													CurrentVideoStreamBitrate;
+	FStreamBitrateInfo													CurrentAudioStreamBitrate;
 
 	bool																bShouldBePaused;
 	bool																bShouldBePlaying;
diff --git a/Engine/Plugins/Media/ElectraPlayer/Source/ElectraPlayerRuntime/Private/Runtime/Player/AdaptiveStreamingPlayerMetrics.h b/Engine/Plugins/Media/ElectraPlayer/Source/ElectraPlayerRuntime/Private/Runtime/Player/AdaptiveStreamingPlayerMetrics.h
--- a/Engine/Plugins/Media/ElectraPlayer/Source/ElectraPlayerRuntime/Private/Runtime/Player/AdaptiveStreamingPlayerMetrics.h
+++ b/Engine/Plugins/Media/ElectraPlayer/Source/ElectraPlayerRuntime/Private/Runtime/Player/AdaptiveStreamingPlayerMetrics.h
@@ -273,6 +273,12 @@ public:
 	 */
 	virtual void ReportVideoQualityChange(int32 NewBitrate, int32 PreviousBitrate, bool bIsDrasticDownswitch) = 0;
 
+	/**
+	 * Called when a new audio stream segment is fetched at a different bitrate than before.
+	 * A drastic change is one where quality _drops_ more than _one_ level.
+	 */
+	virtual void ReportAudioQualityChange(int32 NewBitrate, int32 PreviousBitrate, bool bIsDrasticDownswitch) = 0;
+
 	/**
 	 * Called when stream data availability changes when feeding the decoder.
 	 */
diff --git a/Engine/Plugins/Media/ElectraPlayer/Source/ElectraPlayerRuntime/Private/Runtime/Player/SimpleElectraAudioPlayer.cpp b/Engine/Plugins/Media/ElectraPlayer/Source/ElectraPlayerRuntime/Private/Runtime/Player/SimpleElectraAudioPlayer.cpp
--- a/Engine/Plugins/Media/ElectraPlayer/Source/ElectraPlayerRuntime/Private/Runtime/Player/SimpleElectraAudioPlayer.cpp
+++ b/Engine/Plugins/Media/ElectraPlayer/Source/ElectraPlayerRuntime/Private/Runtime/Player/SimpleElectraAudioPlayer.cpp
@@ -415,6 +415,8 @@ public:
 	{ }
 	void ReportVideoQualityChange(int32 InNewBitrate, int32 InPreviousBitrate, bool bInIsDrasticDownswitch) override
 	{ }
+	void ReportAudioQualityChange(int32 InNewBitrate, int32 InPreviousBitrate, bool bInIsDrasticDownswitch) override
+	{ }
 	void ReportDataAvailabilityChange(const Metrics::FDataAvailabilityChange& InDataAvailability) override
 	{ }
 	void ReportDecodingFormatChange(const FStreamCodecInformation& InNewDecodingFormat) override
