ElectraPlayer: Create HTTP manager on module startup on some platforms

#rnx

diff --git a/Engine/Plugins/Media/ElectraPlayer/Source/ElectraPlayerRuntime/Private/ElectraPlayerRuntimeModule.cpp b/Engine/Plugins/Media/ElectraPlayer/Source/ElectraPlayerRuntime/Private/ElectraPlayerRuntimeModule.cpp
--- a/Engine/Plugins/Media/ElectraPlayer/Source/ElectraPlayerRuntime/Private/ElectraPlayerRuntimeModule.cpp
+++ b/Engine/Plugins/Media/ElectraPlayer/Source/ElectraPlayerRuntime/Private/ElectraPlayerRuntimeModule.cpp
@@ -16,8 +16,6 @@
 #include "ElectraPlayer.h"
 #include "ElectraPlayerPlatform.h"
 
-#include "IElectraHTTPStreamModule.h"
-
 DEFINE_LOG_CATEGORY(LogElectraPlayer);
 DEFINE_LOG_CATEGORY(LogElectraPlayerStats);
 
@@ -69,9 +67,6 @@ public:
 		FElectraRendererAudio::SystemConfiguration SysCfgAudio;
 		FElectraRendererAudio::Startup(SysCfgAudio);
 
-		// The HTTP stream module is required.
-		FModuleManager::LoadModuleChecked<IElectraHTTPStreamModule>("ElectraHTTPStream");
-
 		bInitialized = true;
 	}
 
diff --git a/Engine/Plugins/Media/ElectraPlayer/Source/ElectraPlayerRuntime/Private/Runtime/PlayerRuntimeGlobal.cpp b/Engine/Plugins/Media/ElectraPlayer/Source/ElectraPlayerRuntime/Private/Runtime/PlayerRuntimeGlobal.cpp
--- a/Engine/Plugins/Media/ElectraPlayer/Source/ElectraPlayerRuntime/Private/Runtime/PlayerRuntimeGlobal.cpp
+++ b/Engine/Plugins/Media/ElectraPlayer/Source/ElectraPlayerRuntime/Private/Runtime/PlayerRuntimeGlobal.cpp
@@ -5,6 +5,15 @@
 #include "Misc/CoreDelegates.h"
 #include "Modules/ModuleManager.h"
 
+#ifndef ELECTRA_PRECREATE_HTTP_MANAGER
+#define ELECTRA_PRECREATE_HTTP_MANAGER 0
+#endif
+
+#if ELECTRA_PRECREATE_HTTP_MANAGER
+#include "HTTP/HTTPManager.h"
+#endif
+
+
 
 namespace Electra
 {
@@ -19,6 +28,10 @@ namespace Electra
 		TArray<TWeakPtrTS<FFGBGNotificationHandlers>>		ApplicationBGFGHandlers;
 		volatile bool										bIsInBackground = false;
 		volatile bool										bAppIsTerminating = false;
+
+#if ELECTRA_PRECREATE_HTTP_MANAGER
+		TSharedPtrTS<IElectraHttpManager>					HttpManager;
+#endif
 	}
 	using namespace Global;
 
@@ -107,6 +120,14 @@ namespace Electra
 		FModuleManager::Get().LoadModule(TEXT("ElectraCDM"));
 
 		EnabledAnalyticsEvents = InConfiguration.EnabledAnalyticsEvents;
+
+#if ELECTRA_PRECREATE_HTTP_MANAGER
+		HttpManager = IElectraHttpManager::Create();
+		if (!HttpManager.IsValid())
+		{
+			return false;
+		}
+#endif
 		return true;
 	}
 
@@ -117,6 +138,10 @@ namespace Electra
 	 */
 	void Shutdown(void)
 	{
+#if ELECTRA_PRECREATE_HTTP_MANAGER
+		HttpManager.Reset();
+#endif
+
 		if (ApplicationSuspendedDelegate.IsValid())
 		{
 			FCoreDelegates::ApplicationWillEnterBackgroundDelegate.Remove(ApplicationSuspendedDelegate);
