Adding support for CPU side task scheduling based on decoder data availability to Electra's output pipe
#tests Test on local machine with various codecs and media files

diff --git a/Engine/Plugins/Media/AvidDNxHDDecoderElectra/Source/Private/Windows/ElectraMediaAvidDNxHDDecoder.cpp b/Engine/Plugins/Media/AvidDNxHDDecoderElectra/Source/Private/Windows/ElectraMediaAvidDNxHDDecoder.cpp
--- a/Engine/Plugins/Media/AvidDNxHDDecoderElectra/Source/Private/Windows/ElectraMediaAvidDNxHDDecoder.cpp
+++ b/Engine/Plugins/Media/AvidDNxHDDecoderElectra/Source/Private/Windows/ElectraMediaAvidDNxHDDecoder.cpp
@@ -152,7 +152,7 @@ public:
 #if ELECTRA_MEDIAGPUBUFFER_DX12
 		if (InBufferIndex == 0)
 		{
-			SyncObject = { GPUBuffer.Fence.GetReference(), GPUBuffer.FenceValue };
+			SyncObject = { GPUBuffer.Fence.GetReference(), GPUBuffer.FenceValue, nullptr, GPUBuffer_TaskSync };
 			return true;
 		}
 #endif
@@ -207,6 +207,7 @@ public:
 	EElectraDecoderPlatformPixelEncoding BufferEncoding = EElectraDecoderPlatformPixelEncoding::Native;
 #if ELECTRA_MEDIAGPUBUFFER_DX12
 	FElectraMediaDecoderOutputBufferPool_DX12::FOutputData GPUBuffer;
+	TSharedPtr<IElectraDecoderResourceDelegateBase::IAsyncConsecutiveTaskSync> GPUBuffer_TaskSync;
 #endif
 };
 
@@ -297,7 +298,7 @@ private:
 	uint32 MaxOutputBuffers;
 #if ELECTRA_MEDIAGPUBUFFER_DX12
 	mutable TSharedPtr<FElectraMediaDecoderOutputBufferPool_DX12> D3D12ResourcePool;
-	mutable TUniquePtr<IElectraDecoderResourceDelegateBase::IAsyncConsecutiveTaskSync> TaskSync;
+	mutable TSharedPtr<IElectraDecoderResourceDelegateBase::IAsyncConsecutiveTaskSync> TaskSync;
 #endif
 };
 
@@ -723,7 +724,7 @@ IElectraDecoder::EDecoderError FVideoDecoderAvidDNxHDElectra::DecodeAccessUnit(c
 			// Create a tasksync instance so we can have our own async jobs run in consecutive order
 			if (!TaskSync.IsValid())
 			{
-				TaskSync.Reset(PinnedResourceDelegate->CreateAsyncConsecutiveTaskSync());
+				TaskSync = PinnedResourceDelegate->CreateAsyncConsecutiveTaskSync();
 			}
 
 			// Request resource and fence...
@@ -772,11 +773,15 @@ IElectraDecoder::EDecoderError FVideoDecoderAvidDNxHDElectra::DecodeAccessUnit(c
 				// If we have a temp buffer, we need to copy the data first...
 				if (!TempBuffer.IsEmpty())
 				{
-					if (!PinnedResourceDelegate->RunCodeAsync([BufferResource=NewOutput->GPUBuffer.Resource, Width=NewOutput->Width, Height=NewOutput->Height, ResourcePitch, BPP, TempBuffer, BufferFence=NewOutput->GPUBuffer.Fence, BufferFenceValue=NewOutput->GPUBuffer.FenceValue]()
+					if (PinnedResourceDelegate->RunCodeAsync([BufferResource = NewOutput->GPUBuffer.Resource, Width = NewOutput->Width, Height = NewOutput->Height, ResourcePitch, BPP, TempBuffer, BufferFence = NewOutput->GPUBuffer.Fence, BufferFenceValue = NewOutput->GPUBuffer.FenceValue]()
 						{
 							CopyData(BufferResource, ResourcePitch, TempBuffer.GetData(), Width * BPP, Height);
 							BufferFence->Signal(BufferFenceValue);
 						}, TaskSync.Get()))
+					{
+						NewOutput->GPUBuffer_TaskSync = TaskSync;
+					}
+					else
 					{
 						// Async copy failed, do it synchronously...
 						CopyData(NewOutput->GPUBuffer.Resource, ResourcePitch, TempBuffer.GetData(), NewOutput->Width * BPP, NewOutput->Height);
diff --git a/Engine/Plugins/Media/ElectraCodecs/Source/ElectraDecoders/Private/DX/VideoDecoderH264_DX.cpp b/Engine/Plugins/Media/ElectraCodecs/Source/ElectraDecoders/Private/DX/VideoDecoderH264_DX.cpp
--- a/Engine/Plugins/Media/ElectraCodecs/Source/ElectraDecoders/Private/DX/VideoDecoderH264_DX.cpp
+++ b/Engine/Plugins/Media/ElectraCodecs/Source/ElectraDecoders/Private/DX/VideoDecoderH264_DX.cpp
@@ -170,7 +170,7 @@ public:
 			SyncObject = { Buffer.DecoderSync.GetReference(), 0, Buffer.MFSample };
 			// note: a reference to the MFSample remains with this instance, but is assumed to soon be released as the instance will be destroyed shortly after this was called
 #else
-			SyncObject = { Buffer.Fence.GetReference(), Buffer.FenceValue, nullptr };
+			SyncObject = { Buffer.Fence.GetReference(), Buffer.FenceValue, nullptr, Buffer_TaskSync };
 #endif
 			return true;
 		}
@@ -236,6 +236,7 @@ public:
 	} Buffer;
 #else
 	FElectraMediaDecoderOutputBufferPool_DX12::FOutputData Buffer;
+	TSharedPtr<IElectraDecoderResourceDelegateWindows::IAsyncConsecutiveTaskSync, ESPMode::ThreadSafe> Buffer_TaskSync;
 #endif
 };
 
@@ -349,7 +350,6 @@ private:
 	bool ConvertDecoderOutput();
 	static void CopyBufferData(TRefCountPtr<IMF2DBuffer> Buffer2D, TRefCountPtr<ID3D12Resource> BufferResource, TRefCountPtr<ID3D12Fence> BufferFence, uint64 BufferFenceValue, uint32 BufferPitch, int32 DecodedHeight);
 
-
 	TWeakPtr<IElectraDecoderResourceDelegate, ESPMode::ThreadSafe> ResourceDelegate;
 
 	IElectraDecoder::FError LastError;
@@ -370,7 +370,7 @@ private:
 	uint32 MaxHeight;
 #if !ALLOW_MFSAMPLE_WITH_DX12
 	mutable TSharedPtr<FElectraMediaDecoderOutputBufferPool_DX12> D3D12ResourcePool;
-	mutable TUniquePtr<IElectraDecoderResourceDelegateBase::IAsyncConsecutiveTaskSync> TaskSync;
+	mutable TSharedPtr<IElectraDecoderResourceDelegateBase::IAsyncConsecutiveTaskSync> TaskSync;
 #endif
 
 	static const GUID MFTms_AVLowLatencyMode;
@@ -1154,7 +1154,7 @@ bool FElectraVideoDecoderH264_DX::ConvertDecoderOutput()
 		// Create a tasksync instance so we can have our own async jobs run in consecutive order
 		if (!TaskSync.IsValid())
 		{
-			TaskSync.Reset(PinnedResourceDelegate->CreateAsyncConsecutiveTaskSync());
+			TaskSync = PinnedResourceDelegate->CreateAsyncConsecutiveTaskSync();
 		}
 
 		// Get data from MFSample...
@@ -1170,10 +1170,14 @@ bool FElectraVideoDecoderH264_DX::ConvertDecoderOutput()
 
 		NewOutput->Pitch = BufferPitch;
 
-		if (!PinnedResourceDelegate->RunCodeAsync([DecodedOutputSample, Buffer2D, BufferResource = NewOutput->Buffer.Resource, BufferFence = NewOutput->Buffer.Fence, BufferFenceValue = NewOutput->Buffer.FenceValue, BufferPitch, DecodedHeight = NewOutput->DecodedHeight]()
+		if (PinnedResourceDelegate->RunCodeAsync([DecodedOutputSample, Buffer2D, BufferResource = NewOutput->Buffer.Resource, BufferFence = NewOutput->Buffer.Fence, BufferFenceValue = NewOutput->Buffer.FenceValue, BufferPitch, DecodedHeight = NewOutput->DecodedHeight]()
 			{
 				CopyBufferData(Buffer2D, BufferResource, BufferFence, BufferFenceValue, BufferPitch, DecodedHeight);
 			}, TaskSync.Get()))
+		{
+			NewOutput->Buffer_TaskSync = TaskSync;
+		}
+		else
 		{
 			CopyBufferData(Buffer2D, NewOutput->Buffer.Resource, NewOutput->Buffer.Fence, NewOutput->Buffer.FenceValue, BufferPitch, NewOutput->DecodedHeight);
 		}
diff --git a/Engine/Plugins/Media/ElectraCodecs/Source/ElectraDecoders/Private/DX/VideoDecoderH265_DX.cpp b/Engine/Plugins/Media/ElectraCodecs/Source/ElectraDecoders/Private/DX/VideoDecoderH265_DX.cpp
--- a/Engine/Plugins/Media/ElectraCodecs/Source/ElectraDecoders/Private/DX/VideoDecoderH265_DX.cpp
+++ b/Engine/Plugins/Media/ElectraCodecs/Source/ElectraDecoders/Private/DX/VideoDecoderH265_DX.cpp
@@ -190,7 +190,7 @@ public:
 			SyncObject = { Buffer.DecoderSync.GetReference(), 0, Buffer.MFSample };
 			// note: a reference to the MFSample remains with this instance, but is assumed to soon be released as the instance will be destroyed shortly after this was called
 #else
-			SyncObject = { Buffer.Fence.GetReference(), Buffer.FenceValue, nullptr };
+			SyncObject = { Buffer.Fence.GetReference(), Buffer.FenceValue, nullptr, Buffer_TaskSync };
 #endif
 			return true;
 		}
@@ -256,6 +256,7 @@ public:
 	} Buffer;
 #else
 	FElectraMediaDecoderOutputBufferPool_DX12::FOutputData Buffer;
+	TSharedPtr<IElectraDecoderResourceDelegateWindows::IAsyncConsecutiveTaskSync, ESPMode::ThreadSafe> Buffer_TaskSync;
 #endif
 };
 
@@ -389,7 +390,7 @@ private:
 	uint32 MaxHeight;
 #if !ALLOW_MFSAMPLE_WITH_DX12
 	mutable TSharedPtr<FElectraMediaDecoderOutputBufferPool_DX12> D3D12ResourcePool;
-	mutable TUniquePtr<IElectraDecoderResourceDelegateBase::IAsyncConsecutiveTaskSync> TaskSync;
+	mutable TSharedPtr<IElectraDecoderResourceDelegateBase::IAsyncConsecutiveTaskSync> TaskSync;
 #endif
 
 	static const GUID MFTms_AVLowLatencyMode;
@@ -1209,7 +1210,7 @@ bool FElectraVideoDecoderH265_DX::ConvertDecoderOutput()
 		// Create a tasksync instance so we can have our own async jobs run in consecutive order
 		if (!TaskSync.IsValid())
 		{
-			TaskSync.Reset(PinnedResourceDelegate->CreateAsyncConsecutiveTaskSync());
+			TaskSync = PinnedResourceDelegate->CreateAsyncConsecutiveTaskSync();
 		}
 
 		// Get data from MFSample...
@@ -1225,10 +1226,14 @@ bool FElectraVideoDecoderH265_DX::ConvertDecoderOutput()
 
 		NewOutput->Pitch = BufferPitch;
 
-		if (!PinnedResourceDelegate->RunCodeAsync([DecodedOutputSample, Buffer2D, BufferResource = NewOutput->Buffer.Resource, BufferFence = NewOutput->Buffer.Fence, BufferFenceValue = NewOutput->Buffer.FenceValue, BufferPitch, DecodedHeight = NewOutput->DecodedHeight]()
+		if (PinnedResourceDelegate->RunCodeAsync([DecodedOutputSample, Buffer2D, BufferResource = NewOutput->Buffer.Resource, BufferFence = NewOutput->Buffer.Fence, BufferFenceValue = NewOutput->Buffer.FenceValue, BufferPitch, DecodedHeight = NewOutput->DecodedHeight]()
 			{
 				CopyBufferData(Buffer2D, BufferResource, BufferFence, BufferFenceValue, BufferPitch, DecodedHeight);
 			}, TaskSync.Get()))
+		{
+			NewOutput->Buffer_TaskSync = TaskSync;
+		}
+		else
 		{
 			CopyBufferData(Buffer2D, NewOutput->Buffer.Resource, NewOutput->Buffer.Fence, NewOutput->Buffer.FenceValue, BufferPitch, NewOutput->DecodedHeight);
 		}
diff --git a/Engine/Plugins/Media/ElectraCodecs/Source/ElectraDecoders/Public/IElectraDecoderResourceDelegateBase.h b/Engine/Plugins/Media/ElectraCodecs/Source/ElectraDecoders/Public/IElectraDecoderResourceDelegateBase.h
--- a/Engine/Plugins/Media/ElectraCodecs/Source/ElectraDecoders/Public/IElectraDecoderResourceDelegateBase.h
+++ b/Engine/Plugins/Media/ElectraCodecs/Source/ElectraDecoders/Public/IElectraDecoderResourceDelegateBase.h
@@ -39,7 +39,7 @@ public:
 		IAsyncConsecutiveTaskSync() = default;
 	};
 
-	virtual IAsyncConsecutiveTaskSync* CreateAsyncConsecutiveTaskSync()
+	virtual TSharedPtr<IAsyncConsecutiveTaskSync, ESPMode::ThreadSafe> CreateAsyncConsecutiveTaskSync()
 	{ return nullptr; }
 	virtual bool RunCodeAsync(TFunction<void()>&& CodeToRun, IAsyncConsecutiveTaskSync* TaskSync = nullptr)
 	{ return false; }
diff --git a/Engine/Plugins/Media/ElectraCodecs/Source/ElectraDecoders/Public/Windows/WindowsElectraDecoderPlatformOutputHandleTypes.h b/Engine/Plugins/Media/ElectraCodecs/Source/ElectraDecoders/Public/Windows/WindowsElectraDecoderPlatformOutputHandleTypes.h
--- a/Engine/Plugins/Media/ElectraCodecs/Source/ElectraDecoders/Public/Windows/WindowsElectraDecoderPlatformOutputHandleTypes.h
+++ b/Engine/Plugins/Media/ElectraCodecs/Source/ElectraDecoders/Public/Windows/WindowsElectraDecoderPlatformOutputHandleTypes.h
@@ -3,6 +3,7 @@
 #pragma once
 
 #include "CoreMinimal.h"
+#include "IElectraDecoderResourceDelegateBase.h"
 
 #include "Windows/AllowWindowsPlatformTypes.h"
 THIRD_PARTY_INCLUDES_START
@@ -54,21 +55,22 @@ class FElectraDecoderOutputSync final
 {
 public:
 	FElectraDecoderOutputSync() : SyncValue(0), CopyDoneSyncValue(0) {}
-	FElectraDecoderOutputSync(TRefCountPtr<IUnknown> InSync, uint64 InSyncValue, TRefCountPtr<IMFSample> InMFSample = nullptr)
-		: Sync(InSync), SyncValue(InSyncValue), MFSample(InMFSample)
+	FElectraDecoderOutputSync(TRefCountPtr<IUnknown> InSync, uint64 InSyncValue, TRefCountPtr<IMFSample> InMFSample = nullptr, TSharedPtr<IElectraDecoderResourceDelegateBase::IAsyncConsecutiveTaskSync> InTaskSync = nullptr)
+		: Sync(InSync), SyncValue(InSyncValue), TaskSync(InTaskSync), MFSample(InMFSample)
 	{}
 	FElectraDecoderOutputSync(TRefCountPtr<IUnknown> InSync, uint64 InSyncValue, TRefCountPtr<ID3D12Fence> InCopyDoneSync, uint64 InCopyDoneSyncValue)
 		: Sync(InSync), SyncValue(InSyncValue), CopyDoneSync(InCopyDoneSync), CopyDoneSyncValue(InCopyDoneSyncValue)
 	{}
 
-	TRefCountPtr<IUnknown>		Sync;				// An IUnknown based sync primitive (fence, or DX12 specific WMF object)
-	uint64						SyncValue;			// Optionally needed sync value
+	TRefCountPtr<IUnknown>		Sync;														// Optional IUnknown based sync primitive (fence, or DX12 specific WMF object)
+	uint64						SyncValue;													// Optionally needed sync value
+	TSharedPtr<IElectraDecoderResourceDelegateBase::IAsyncConsecutiveTaskSync> TaskSync;	// Optional Sync/Perquisite object to use to run code async after the decoder data arrives
 
-	TRefCountPtr<ID3D12Fence>	CopyDoneSync;		// Fence to signal end of copy from this buffer back to decoder
-	uint64						CopyDoneSyncValue;	// Value to be used with CopyDoneSync fence
+	TRefCountPtr<ID3D12Fence>	CopyDoneSync;												// Optional fence to signal end of copy from this buffer back to decoder
+	uint64						CopyDoneSyncValue;											// Value to be used with CopyDoneSync fence
 
 private:
-	TRefCountPtr<IMFSample>		MFSample;			// The sample reference can ride along to ensure that the DX12 WMF interface properly keeps reference to the sample until we enqueue their sync primitives with the a D3D queue
+	TRefCountPtr<IMFSample>		MFSample;													// The sample reference can ride along to ensure that the DX12 WMF interface properly keeps reference to the sample until we enqueue their sync primitives with the a D3D queue
 };
 
 class IElectraDecoderVideoOutputImageBuffers
diff --git a/Engine/Plugins/Media/ElectraPlayer/Source/ElectraPlayerPlugin/Private/ElectraPlayerPluginModule.cpp b/Engine/Plugins/Media/ElectraPlayer/Source/ElectraPlayerPlugin/Private/ElectraPlayerPluginModule.cpp
--- a/Engine/Plugins/Media/ElectraPlayer/Source/ElectraPlayerPlugin/Private/ElectraPlayerPluginModule.cpp
+++ b/Engine/Plugins/Media/ElectraPlayer/Source/ElectraPlayerPlugin/Private/ElectraPlayerPluginModule.cpp
@@ -92,9 +92,9 @@ public:
 		FGraphEventRef	GraphEvent;
 	};
 
-	static IElectraDecoderResourceDelegateBase::IAsyncConsecutiveTaskSync* CreateAsyncConsecutiveTaskSync()
+	static TSharedPtr<IElectraDecoderResourceDelegateBase::IAsyncConsecutiveTaskSync, ESPMode::ThreadSafe> CreateAsyncConsecutiveTaskSync()
 	{
-		return new FAsyncConsecutiveTaskSync();
+		return MakeShared<FAsyncConsecutiveTaskSync, ESPMode::ThreadSafe>();
 	}
 
 	static void RunCodeAsync(TFunction<void()>&& CodeToRun, IElectraDecoderResourceDelegateBase::IAsyncConsecutiveTaskSync* TaskSync)
diff --git a/Engine/Plugins/Media/ElectraPlayer/Source/ElectraPlayerRuntime/Private/Runtime/Decoder/Windows/ElectraVideoDecoder_PC.cpp b/Engine/Plugins/Media/ElectraPlayer/Source/ElectraPlayerRuntime/Private/Runtime/Decoder/Windows/ElectraVideoDecoder_PC.cpp
--- a/Engine/Plugins/Media/ElectraPlayer/Source/ElectraPlayerRuntime/Private/Runtime/Decoder/Windows/ElectraVideoDecoder_PC.cpp
+++ b/Engine/Plugins/Media/ElectraPlayer/Source/ElectraPlayerRuntime/Private/Runtime/Decoder/Windows/ElectraVideoDecoder_PC.cpp
@@ -4,7 +4,7 @@
 #include "Renderer/RendererVideo.h"
 #include COMPILED_PLATFORM_HEADER(ElectraDecoderGPUBufferHelpers.h)
 #include COMPILED_PLATFORM_HEADER(ElectraDecoderPlatformOutputHandleTypes.h)
-//#include COMPILED_PLATFORM_HEADER(PlatformHeaders_Video_DX.h) // PRIVATE FOR CODECS
+#include "WindowsElectraDecoderResourceManager.h"
 
 #include "Windows/AllowWindowsPlatformTypes.h"
 THIRD_PARTY_INCLUDES_START
@@ -365,49 +365,71 @@ void FElectraPlayerVideoDecoderOutputPC::InitializeWithResource(const TRefCountP
 
 	// - - - - - - - - - - - - - - - - - - - - - - - - - - - - - - - - - - - - - - - - - - - - - - - - - - - - - - - -
 
-	// Trigger copy (this will eventually execute on the submission thread of RHI if running in UE)
-	// (note: we pass in all of FElectraDecoderOutputSync to guarantee any references needed to make the decoder output sync work are passed along, too!)
-	InResourceDelegate->ExecuteCodeWithCopyCommandQueueUsage([CmdList = D3DCmdList, DestFence = D3DFence, DestFenceValue = FenceValue, 
-#if ALLOW_MFSAMPLE_WITH_DX12
-															  DecoderSync, /*SAMPLE ?*/
-#endif
-															  ResourceFence, OutputSync](ID3D12CommandQueue* D3DCmdQueue)
+	// Can we wait for the decoder to have the data ready using an async CPU job?
+	bool bTriggerOk = false;
+	if (OutputSync.TaskSync.IsValid())
 	{
-#if ALLOW_MFSAMPLE_WITH_DX12
-		if (DecoderSync)
-		{
-			// Sync queue to make sure decoder output is ready for us (WMF case)
-			DecoderSync->EnqueueResourceReadyWait(D3DCmdQueue);
-		}
-#endif
-		if (ResourceFence)
-		{
-			// Sync queue to make sure decoder output is ready for us
-			D3DCmdQueue->Wait(ResourceFence, OutputSync.SyncValue);
-		}
+		// Yes. Trigger the GPU copy once the CPU is ready, so we can avoid any stalls on the copy-queue waiting for decoder data...
+		auto ElectraDecoderResourceDelegate = Electra::FElectraDecoderResourceManagerWindows::GetDelegate();
+
+		// Note: we capture "this" as we ensure that this instance only dies once the copy triggered here is actually done, hence ensuring any reference to "this" is done
+		bTriggerOk = ElectraDecoderResourceDelegate->RunCodeAsync([this, DecoderSync, ResourceFence, OutputSync, InResourceDelegate]()
+			{
+				TriggerDataCopy(DecoderSync, ResourceFence, OutputSync, InResourceDelegate);
+			}, OutputSync.TaskSync.Get());
+	}
 
-		// Execute copy
-		ID3D12CommandList* CmdLists[1] = { CmdList.GetReference() };
-		D3DCmdQueue->ExecuteCommandLists(1, CmdLists);
+	if (!bTriggerOk)
+	{
+		// We could not run the trigger async. Schedule the copy right away. Any needed synchronization will be done in the copy-queue by the GPU
+		TriggerDataCopy(DecoderSync, ResourceFence, OutputSync, InResourceDelegate);
+	}
+}
 
 #if ALLOW_MFSAMPLE_WITH_DX12
-		if (DecoderSync)
-		{
-			// Sync to end of copy operation and release MFSample once reached
-			DecoderSync->EnqueueResourceRelease(D3DCmdQueue);
-		}
+void FElectraPlayerVideoDecoderOutputPC::TriggerDataCopy(TRefCountPtr<IMFD3D12SynchronizationObjectCommands> DecoderSync, TRefCountPtr<ID3D12Fence> ResourceFence, const FElectraDecoderOutputSync& OutputSync, Electra::IVideoDecoderResourceDelegate* InResourceDelegate) const
+#else
+void FElectraPlayerVideoDecoderOutputPC::TriggerDataCopy(TRefCountPtr<IUnknown> DecoderSync, TRefCountPtr<ID3D12Fence> ResourceFence, const FElectraDecoderOutputSync& OutputSync, Electra::IVideoDecoderResourceDelegate* InResourceDelegate) const
 #endif
-		// Trigger optional notification back to the decoder, so it could reuse its buffer after the copy
-		if (OutputSync.CopyDoneSync)
+{
+	// Trigger copy (this will eventually execute on the submission thread of RHI if running in UE)
+	// (note: we pass in all of FElectraDecoderOutputSync to guarantee any references needed to make the decoder output sync work are passed along, too!)
+	InResourceDelegate->ExecuteCodeWithCopyCommandQueueUsage([CmdList = D3DCmdList, DestFence = D3DFence, DestFenceValue = FenceValue, DecoderSync, ResourceFence, OutputSync](ID3D12CommandQueue* D3DCmdQueue)
 		{
-			D3DCmdQueue->Signal(OutputSync.CopyDoneSync, OutputSync.CopyDoneSyncValue);
-		}
+	#if ALLOW_MFSAMPLE_WITH_DX12
+			if (DecoderSync)
+			{
+				// Sync queue to make sure decoder output is ready for us (WMF case)
+				DecoderSync->EnqueueResourceReadyWait(D3DCmdQueue);
+			}
+	#endif
+			if (ResourceFence)
+			{
+				// Sync queue to make sure decoder output is ready for us
+				D3DCmdQueue->Wait(ResourceFence, OutputSync.SyncValue);
+			}
 
-		// Notify user of the output data we just copied of its arrival
-		D3DCmdQueue->Signal(DestFence, DestFenceValue);
-	});
-}
+			// Execute copy
+			ID3D12CommandList* CmdLists[1] = { CmdList.GetReference() };
+			D3DCmdQueue->ExecuteCommandLists(1, CmdLists);
+
+	#if ALLOW_MFSAMPLE_WITH_DX12
+			if (DecoderSync)
+			{
+				// Sync to end of copy operation and release MFSample once reached
+				DecoderSync->EnqueueResourceRelease(D3DCmdQueue);
+			}
+	#endif
+			// Trigger optional notification back to the decoder, so it could reuse its buffer after the copy
+			if (OutputSync.CopyDoneSync)
+			{
+				D3DCmdQueue->Signal(OutputSync.CopyDoneSync, OutputSync.CopyDoneSyncValue);
+			}
 
+			// Notify user of the output data we just copied of its arrival
+			D3DCmdQueue->Signal(DestFence, DestFenceValue);
+		});
+}
 
 void FElectraPlayerVideoDecoderOutputPC::InitializeWithSharedTexture(const TRefCountPtr<ID3D11Device>& InD3D11Device, const TRefCountPtr<ID3D11DeviceContext> InDeviceContext, const TRefCountPtr<IMFSample> InMFSample, const FIntPoint& OutputDim, TSharedPtr<Electra::FParamDict, ESPMode::ThreadSafe> InParamDict)
 {
diff --git a/Engine/Plugins/Media/ElectraPlayer/Source/ElectraPlayerRuntime/Private/Runtime/Decoder/Windows/ElectraVideoDecoder_PC.h b/Engine/Plugins/Media/ElectraPlayer/Source/ElectraPlayerRuntime/Private/Runtime/Decoder/Windows/ElectraVideoDecoder_PC.h
--- a/Engine/Plugins/Media/ElectraPlayer/Source/ElectraPlayerRuntime/Private/Runtime/Decoder/Windows/ElectraVideoDecoder_PC.h
+++ b/Engine/Plugins/Media/ElectraPlayer/Source/ElectraPlayerRuntime/Private/Runtime/Decoder/Windows/ElectraVideoDecoder_PC.h
@@ -21,6 +21,13 @@ THIRD_PARTY_INCLUDES_START
 #include "mfidl.h"
 THIRD_PARTY_INCLUDES_END
 
+#if defined(NTDDI_WIN10_NI)
+#include <mfd3d12.h>
+#define HAVE_MFSAMPLE_WITH_DX12 1
+#else
+#define HAVE_MFSAMPLE_WITH_DX12 0
+#endif
+
 #include "Windows/HideWindowsPlatformTypes.h"
 
 struct ID3D12Resource;
@@ -65,6 +72,12 @@ public:
 	TRefCountPtr<IUnknown> GetSync(uint64& SyncValue) const override;
 
 private:
+#if HAVE_MFSAMPLE_WITH_DX12
+	void TriggerDataCopy(TRefCountPtr<IMFD3D12SynchronizationObjectCommands> DecoderSync, TRefCountPtr<ID3D12Fence> ResourceFence, const FElectraDecoderOutputSync& OutputSync, Electra::IVideoDecoderResourceDelegate* InResourceDelegate) const;
+#else
+	void TriggerDataCopy(TRefCountPtr<IUnknown> DecoderSync, TRefCountPtr<ID3D12Fence> ResourceFence, const FElectraDecoderOutputSync& OutputSync, Electra::IVideoDecoderResourceDelegate* InResourceDelegate) const;
+#endif
+
 	// Decoder output type
 	EOutputType OutputType = EOutputType::Unknown;
 
diff --git a/Engine/Plugins/Media/ElectraPlayer/Source/ElectraPlayerRuntime/Private/Runtime/Decoder/Windows/WindowsElectraDecoderResourceManager.cpp b/Engine/Plugins/Media/ElectraPlayer/Source/ElectraPlayerRuntime/Private/Runtime/Decoder/Windows/WindowsElectraDecoderResourceManager.cpp
--- a/Engine/Plugins/Media/ElectraPlayer/Source/ElectraPlayerRuntime/Private/Runtime/Decoder/Windows/WindowsElectraDecoderResourceManager.cpp
+++ b/Engine/Plugins/Media/ElectraPlayer/Source/ElectraPlayerRuntime/Private/Runtime/Decoder/Windows/WindowsElectraDecoderResourceManager.cpp
@@ -129,7 +129,7 @@ bool FElectraDecoderResourceManagerWindows::GetD3DDevice(void **OutD3DDevice, in
 	return false;
 }
 
-IElectraDecoderResourceDelegateWindows::IAsyncConsecutiveTaskSync* FElectraDecoderResourceManagerWindows::CreateAsyncConsecutiveTaskSync()
+TSharedPtr<IElectraDecoderResourceDelegateBase::IAsyncConsecutiveTaskSync, ESPMode::ThreadSafe> FElectraDecoderResourceManagerWindows::CreateAsyncConsecutiveTaskSync()
 {
 	if (WindowsDecoderResources::Callbacks.CreateAsyncConsecutiveTaskSync)
 	{
diff --git a/Engine/Plugins/Media/ElectraPlayer/Source/ElectraPlayerRuntime/Private/Runtime/Decoder/Windows/WindowsElectraDecoderResourceManager.h b/Engine/Plugins/Media/ElectraPlayer/Source/ElectraPlayerRuntime/Private/Runtime/Decoder/Windows/WindowsElectraDecoderResourceManager.h
--- a/Engine/Plugins/Media/ElectraPlayer/Source/ElectraPlayerRuntime/Private/Runtime/Decoder/Windows/WindowsElectraDecoderResourceManager.h
+++ b/Engine/Plugins/Media/ElectraPlayer/Source/ElectraPlayerRuntime/Private/Runtime/Decoder/Windows/WindowsElectraDecoderResourceManager.h
@@ -19,7 +19,7 @@ public:
 	{
 		TFunction<bool(void ** /*OutD3DDevice*/, int32* /*OutD3DVersionTimes1000*/, void* /*UserValue*/)> GetD3DDevice;
 		TFunction<bool(TFunction<void()>&& CodeToRun, IAsyncConsecutiveTaskSync* TaskSync)> RunCodeAsync;
-		TFunction<IAsyncConsecutiveTaskSync* ()> CreateAsyncConsecutiveTaskSync;
+		TFunction<TSharedPtr<IElectraDecoderResourceDelegateBase::IAsyncConsecutiveTaskSync, ESPMode::ThreadSafe>()> CreateAsyncConsecutiveTaskSync;
 
 		void* UserValue = nullptr;
 	};
@@ -35,7 +35,7 @@ public:
 
 	virtual bool GetD3DDevice(void **OutD3DDevice, int32* OutD3DVersionTimes1000) override;
 
-	virtual IAsyncConsecutiveTaskSync* CreateAsyncConsecutiveTaskSync() override;
+	virtual TSharedPtr<IElectraDecoderResourceDelegateBase::IAsyncConsecutiveTaskSync, ESPMode::ThreadSafe> CreateAsyncConsecutiveTaskSync() override;
 	virtual bool RunCodeAsync(TFunction<void()>&& CodeToRun, IAsyncConsecutiveTaskSync* TaskSync) override;
 
 	virtual ~FElectraDecoderResourceManagerWindows();
diff --git a/Engine/Plugins/Media/ElectraPlayer/Source/ElectraPlayerRuntime/Private/Windows/ElectraPlayer_Windows.cpp b/Engine/Plugins/Media/ElectraPlayer/Source/ElectraPlayerRuntime/Private/Windows/ElectraPlayer_Windows.cpp
--- a/Engine/Plugins/Media/ElectraPlayer/Source/ElectraPlayerRuntime/Private/Windows/ElectraPlayer_Windows.cpp
+++ b/Engine/Plugins/Media/ElectraPlayer/Source/ElectraPlayerRuntime/Private/Windows/ElectraPlayer_Windows.cpp
@@ -116,9 +116,9 @@ bool FElectraPlayerPlatform::StartupPlatformResources(const Electra::FParamDict&
 		return false;
 	};
 
-	auto CreateAsyncConsecutiveTaskSyncCallback = reinterpret_cast<IElectraDecoderResourceDelegateWindows::IAsyncConsecutiveTaskSync*(*)()>(Params.GetValue(FName(TEXT("CreateAsyncConsecutiveTaskSync"))).SafeGetPointer());
+	auto CreateAsyncConsecutiveTaskSyncCallback = reinterpret_cast<TSharedPtr<IElectraDecoderResourceDelegateWindows::IAsyncConsecutiveTaskSync, ESPMode::ThreadSafe>(*)()>(Params.GetValue(FName(TEXT("CreateAsyncConsecutiveTaskSync"))).SafeGetPointer());
 
-	Callbacks.CreateAsyncConsecutiveTaskSync = [CreateAsyncConsecutiveTaskSyncCallback]() -> IElectraDecoderResourceDelegateWindows::IAsyncConsecutiveTaskSync*
+	Callbacks.CreateAsyncConsecutiveTaskSync = [CreateAsyncConsecutiveTaskSyncCallback]() -> TSharedPtr<IElectraDecoderResourceDelegateWindows::IAsyncConsecutiveTaskSync, ESPMode::ThreadSafe>
 	{
 		if (CreateAsyncConsecutiveTaskSyncCallback)
 		{
diff --git a/Engine/Plugins/Media/HAPDecoderElectra/Source/HAPDecoderElectra/Private/HAPDecoder/ElectraMediaHAPDecoder.cpp b/Engine/Plugins/Media/HAPDecoderElectra/Source/HAPDecoderElectra/Private/HAPDecoder/ElectraMediaHAPDecoder.cpp
--- a/Engine/Plugins/Media/HAPDecoderElectra/Source/HAPDecoderElectra/Private/HAPDecoder/ElectraMediaHAPDecoder.cpp
+++ b/Engine/Plugins/Media/HAPDecoderElectra/Source/HAPDecoderElectra/Private/HAPDecoder/ElectraMediaHAPDecoder.cpp
@@ -155,12 +155,12 @@ public:
 #if ELECTRA_MEDIAGPUBUFFER_DX12
 		if (InBufferIndex == 0)
 		{
-			SyncObject = { Buffer[Buffer_Color].GPUBuffer.Fence.GetReference(), Buffer[Buffer_Color].GPUBuffer.FenceValue };
+			SyncObject = { Buffer[Buffer_Color].GPUBuffer.Fence.GetReference(), Buffer[Buffer_Color].GPUBuffer.FenceValue, nullptr, Buffer_TaskSync };
 			return true;
 		}
 		if (InBufferIndex == 1)
 		{
-			SyncObject = { Buffer[Buffer_Alpha].GPUBuffer.Fence.GetReference(), Buffer[Buffer_Alpha].GPUBuffer.FenceValue };
+			SyncObject = { Buffer[Buffer_Alpha].GPUBuffer.Fence.GetReference(), Buffer[Buffer_Alpha].GPUBuffer.FenceValue, nullptr, Buffer_TaskSync };
 			return true;
 		}
 #endif
@@ -238,6 +238,9 @@ public:
 		FElectraMediaDecoderOutputBufferPool_DX12::FOutputData GPUBuffer;
 #endif
 	} Buffer[Buffer_Max];
+#if ELECTRA_MEDIAGPUBUFFER_DX12
+	TSharedPtr<IElectraDecoderResourceDelegateBase::IAsyncConsecutiveTaskSync> Buffer_TaskSync;
+#endif
 };
 
 
@@ -320,7 +323,7 @@ private:
 		mutable TSharedPtr<FElectraMediaDecoderOutputBufferPool_DX12> D3D12ResourcePool;
 #endif
 	} InstanceBufferData[FVideoDecoderOutputHAPElectra::Buffer_Max];
-	mutable TUniquePtr<IElectraDecoderResourceDelegateBase::IAsyncConsecutiveTaskSync> TaskSync;
+	mutable TSharedPtr<IElectraDecoderResourceDelegateBase::IAsyncConsecutiveTaskSync> TaskSync;
 };
 
 
@@ -641,7 +644,7 @@ IElectraDecoder::EDecoderError FVideoDecoderHAPElectra::DecodeAccessUnit(const F
 				// Create a tasksync instance so we can have our own async jobs run in consecutive order
 				if (!TaskSync.IsValid())
 				{
-					TaskSync.Reset(PinnedResourceDelegate->CreateAsyncConsecutiveTaskSync());
+					TaskSync = PinnedResourceDelegate->CreateAsyncConsecutiveTaskSync();
 				}
 
 				InstanceBufferData[BufIdx].D3D12ResourcePool->AllocateOutputDataAsBuffer(BufferInfo.GPUBuffer, ResourcePitch);
@@ -686,7 +689,7 @@ IElectraDecoder::EDecoderError FVideoDecoderHAPElectra::DecodeAccessUnit(const F
 			{
 				if (Result == HapResult_No_Error)
 				{
-					if (!PinnedResourceDelegate->RunCodeAsync([BufferResource = BufferInfo.GPUBuffer.Resource, BufferFence = BufferInfo.GPUBuffer.Fence, BufferFenceValue = BufferInfo.GPUBuffer.FenceValue, ResourcePitch, BufferAddr, BufferPitch, DecodedPitch = DecodedHeight / BlockSizeY, TempBuffer]()
+					if (PinnedResourceDelegate->RunCodeAsync([BufferResource = BufferInfo.GPUBuffer.Resource, BufferFence = BufferInfo.GPUBuffer.Fence, BufferFenceValue = BufferInfo.GPUBuffer.FenceValue, ResourcePitch, BufferAddr, BufferPitch, DecodedPitch = DecodedHeight / BlockSizeY, TempBuffer]()
 						{
 							uint8* DX12BufferAddr;
 							HRESULT Res = BufferResource->Map(0, nullptr, (void**)&DX12BufferAddr);
@@ -698,6 +701,11 @@ IElectraDecoder::EDecoderError FVideoDecoderHAPElectra::DecodeAccessUnit(const F
 
 							BufferFence->Signal(BufferFenceValue);
 						}, TaskSync.Get()))
+					{
+						// note: we share one task sync for all buffers -> this will lead any external task to be executed after ALL tasks we trigger have been done (aka: after the "last" buffer is ready)
+						NewOutput->Buffer_TaskSync = TaskSync;
+					}
+					else	
 					{
 						uint8* DX12BufferAddr;
 						HRESULT Res = BufferInfo.GPUBuffer.Resource->Map(0, nullptr, (void**)&DX12BufferAddr);
