Removing outdated debug check in Electra output video sample handling
#rnx

diff --git a/Engine/Plugins/Media/ElectraPlayer/Source/ElectraPlayerRuntime/Private/Runtime/Decoder/Windows/ElectraVideoDecoder_PC.cpp b/Engine/Plugins/Media/ElectraPlayer/Source/ElectraPlayerRuntime/Private/Runtime/Decoder/Windows/ElectraVideoDecoder_PC.cpp
--- a/Engine/Plugins/Media/ElectraPlayer/Source/ElectraPlayerRuntime/Private/Runtime/Decoder/Windows/ElectraVideoDecoder_PC.cpp
+++ b/Engine/Plugins/Media/ElectraPlayer/Source/ElectraPlayerRuntime/Private/Runtime/Decoder/Windows/ElectraVideoDecoder_PC.cpp
@@ -203,16 +203,6 @@ void FElectraPlayerVideoDecoderOutputPC::InitializeWithResource(const TRefCountP
 		// We get here only after the instance came back from the pool! Hence we can be sure any old texture is no longer actively used.
 		TextureDX12 = nullptr;
 
-#if !UE_BUILD_SHIPPING
-		if (bIsTexture)
-		{
-			// A little bit of sanity checking
-			check(DXGIFmt == ResourceDesc.Format);
-			check(SampleDim.X == ResourceDesc.Width);
-			check(SampleDim.Y == ResourceDesc.Height);
-		}
-#endif
-
 		if (!D3D12ResourcePool.IsValid())
 		{
 			D3D12ResourcePool = MakeShared<FElectraMediaDecoderOutputBufferPool_DX12>(InD3D12Device, MaxOutputBuffers + kElectraDecoderPipelineExtraFrames, MaxWidth, MaxHeight, DXGIFmt, D3D12_HEAP_TYPE_DEFAULT);
