Electra: Fixing up sequence index updates on seeks when the seeks are actually subpressed by an internal seek request
#tests Local machine tests

diff --git a/Engine/Plugins/Media/ElectraPlayer/Source/ElectraPlayerRuntime/Private/Runtime/Player/AdaptiveStreamingPlayer.cpp b/Engine/Plugins/Media/ElectraPlayer/Source/ElectraPlayerRuntime/Private/Runtime/Player/AdaptiveStreamingPlayer.cpp
--- a/Engine/Plugins/Media/ElectraPlayer/Source/ElectraPlayerRuntime/Private/Runtime/Player/AdaptiveStreamingPlayer.cpp
+++ b/Engine/Plugins/Media/ElectraPlayer/Source/ElectraPlayerRuntime/Private/Runtime/Player/AdaptiveStreamingPlayer.cpp
@@ -1748,12 +1748,11 @@ void FAdaptiveStreamingPlayer::HandleSeeking()
 		// On a user-induced seek the sequence index is supposed to be increased.
 		// We need to do this even when we would not actually issue the seek.
 
-		if (!SeekVars.PendingRequest.GetValue().bIgnoreForSequenceIndex.Get(false))
-		{
-			CurrentPlaybackSequenceState.PrimaryIndex += SeekVars.NumSeekToCallsSinceLastSeen;
-			CurrentPlaybackSequenceState.SecondaryIndex = 0;
-			SeekVars.NumSeekToCallsSinceLastSeen = 0;
-		}
+		// Adjust seek index as indicated
+		// (we even do this for seeks marked as not affecting the seek index as we must still add any priovious accumulated calls that had not been flagged in such a way)
+		CurrentPlaybackSequenceState.PrimaryIndex += SeekVars.NumSeekToCallsSinceLastSeen;
+		CurrentPlaybackSequenceState.SecondaryIndex = 0;
+		SeekVars.NumSeekToCallsSinceLastSeen = 0;
 
 		// And since it is a seek on purpose the loop counter is reset as well.
 		FInternalLoopState LoopStateNow;
