Fixed missing refernce keeping that could cause crashes when Electra shutudown while certin texture sample operations took place
#rnx
#tests Tesed on local machine

diff --git a/Engine/Plugins/Media/ElectraPlayer/Source/ElectraPlayerRuntime/Private/Runtime/Decoder/Windows/ElectraVideoDecoder_PC.cpp b/Engine/Plugins/Media/ElectraPlayer/Source/ElectraPlayerRuntime/Private/Runtime/Decoder/Windows/ElectraVideoDecoder_PC.cpp
--- a/Engine/Plugins/Media/ElectraPlayer/Source/ElectraPlayerRuntime/Private/Runtime/Decoder/Windows/ElectraVideoDecoder_PC.cpp
+++ b/Engine/Plugins/Media/ElectraPlayer/Source/ElectraPlayerRuntime/Private/Runtime/Decoder/Windows/ElectraVideoDecoder_PC.cpp
@@ -113,7 +113,7 @@ void FElectraPlayerVideoDecoderOutputPC::InitializeWithBuffer(TSharedPtr<TArray<
 	Stride = InStride;
 }
 
-void FElectraPlayerVideoDecoderOutputPC::InitializeWithResource(const TRefCountPtr<ID3D12Device>& InD3D12Device, const TRefCountPtr<ID3D12Resource> InResource, uint32 ResourcePitch, const FElectraDecoderOutputSync& OutputSync, const FIntPoint& InOutputDim, TSharedPtr<Electra::FParamDict, ESPMode::ThreadSafe> InParamDict, Electra::IVideoDecoderResourceDelegate* InResourceDelegate,
+void FElectraPlayerVideoDecoderOutputPC::InitializeWithResource(const TRefCountPtr<ID3D12Device>& InD3D12Device, const TRefCountPtr<ID3D12Resource> InResource, uint32 ResourcePitch, const FElectraDecoderOutputSync& OutputSync, const FIntPoint& InOutputDim, TSharedPtr<Electra::FParamDict, ESPMode::ThreadSafe> InParamDict, TWeakPtr<Electra::IVideoDecoderResourceDelegate, ESPMode::ThreadSafe> InResourceDelegate,
 																uint32 MaxWidth, uint32 MaxHeight, uint32 MaxOutputBuffers)
 {
 	// General initialization
@@ -363,16 +363,22 @@ void FElectraPlayerVideoDecoderOutputPC::InitializeWithResource(const TRefCountP
 		auto ElectraDecoderResourceDelegate = Electra::FElectraDecoderResourceManagerWindows::GetDelegate();
 
 		// Note: we capture "this" as we ensure that this instance only dies once the copy triggered here is actually done, hence ensuring any reference to "this" is done
-		bTriggerOk = ElectraDecoderResourceDelegate->RunCodeAsync([this, DecoderSync, ResourceFence, OutputSync, InResourceDelegate]()
+		bTriggerOk = ElectraDecoderResourceDelegate->RunCodeAsync([WeakThis=AsWeak(), DecoderSync, ResourceFence, OutputSync, InResourceDelegate]()
 			{
-				TriggerDataCopy(DecoderSync, ResourceFence, OutputSync, InResourceDelegate);
+				if (auto This = StaticCastSharedPtr<FElectraPlayerVideoDecoderOutputPC, IDecoderOutputPoolable, ESPMode::ThreadSafe>(WeakThis.Pin()))
+				{
+					if (auto ResourceDelegate = InResourceDelegate.Pin())
+					{
+						This->TriggerDataCopy(DecoderSync, ResourceFence, OutputSync, ResourceDelegate.Get());
+					}
+				}
 			}, OutputSync.TaskSync.Get());
 	}
 
 	if (!bTriggerOk)
 	{
 		// We could not run the trigger async. Schedule the copy right away. Any needed synchronization will be done in the copy-queue by the GPU
-		TriggerDataCopy(DecoderSync, ResourceFence, OutputSync, InResourceDelegate);
+		TriggerDataCopy(DecoderSync, ResourceFence, OutputSync, InResourceDelegate.Pin().Get());
 	}
 }
 
diff --git a/Engine/Plugins/Media/ElectraPlayer/Source/ElectraPlayerRuntime/Private/Runtime/Decoder/Windows/ElectraVideoDecoder_PC.h b/Engine/Plugins/Media/ElectraPlayer/Source/ElectraPlayerRuntime/Private/Runtime/Decoder/Windows/ElectraVideoDecoder_PC.h
--- a/Engine/Plugins/Media/ElectraPlayer/Source/ElectraPlayerRuntime/Private/Runtime/Decoder/Windows/ElectraVideoDecoder_PC.h
+++ b/Engine/Plugins/Media/ElectraPlayer/Source/ElectraPlayerRuntime/Private/Runtime/Decoder/Windows/ElectraVideoDecoder_PC.h
@@ -51,7 +51,7 @@ public:
 
 	void InitializeWithBuffer(TSharedPtr<TArray<uint8>, ESPMode::ThreadSafe> InBuffer, uint32 InStride, FIntPoint Dim, TSharedPtr<Electra::FParamDict, ESPMode::ThreadSafe> InParamDict);
 
-	void InitializeWithResource(const TRefCountPtr<ID3D12Device>& InD3D12Device, const TRefCountPtr<ID3D12Resource> Resource, uint32 ResourcePitch, const FElectraDecoderOutputSync& OutputSync, const FIntPoint& OutputDim, TSharedPtr<Electra::FParamDict, ESPMode::ThreadSafe> InParamDict, Electra::IVideoDecoderResourceDelegate* ResourceDelegate,
+	void InitializeWithResource(const TRefCountPtr<ID3D12Device>& InD3D12Device, const TRefCountPtr<ID3D12Resource> Resource, uint32 ResourcePitch, const FElectraDecoderOutputSync& OutputSync, const FIntPoint& OutputDim, TSharedPtr<Electra::FParamDict, ESPMode::ThreadSafe> InParamDict, TWeakPtr<Electra::IVideoDecoderResourceDelegate, ESPMode::ThreadSafe> ResourceDelegate,
 								uint32 MaxWidth, uint32 MaxHeight, uint32 MaxOutputBuffers);
 
 	// Hardware decode to shared DX11 texture (Win8+) from IMFSample
diff --git a/Engine/Plugins/Media/ElectraPlayer/Source/ElectraPlayerRuntime/Private/Runtime/Decoder/Windows/WindowsElectraDecoderResourceManager.cpp b/Engine/Plugins/Media/ElectraPlayer/Source/ElectraPlayerRuntime/Private/Runtime/Decoder/Windows/WindowsElectraDecoderResourceManager.cpp
--- a/Engine/Plugins/Media/ElectraPlayer/Source/ElectraPlayerRuntime/Private/Runtime/Decoder/Windows/WindowsElectraDecoderResourceManager.cpp
+++ b/Engine/Plugins/Media/ElectraPlayer/Source/ElectraPlayerRuntime/Private/Runtime/Decoder/Windows/WindowsElectraDecoderResourceManager.cpp
@@ -530,7 +530,7 @@ bool FElectraDecoderResourceManagerWindows::SetupRenderBufferFromDecoderOutput(I
 							FElectraDecoderOutputSync OutputSync;
 							ImageBuffers->GetBufferTextureSyncByIndex(0, OutputSync);
 
-							DecoderOutput->InitializeWithResource(D3D12Device, Resource, Pitch, OutputSync, FIntPoint(InDecoderOutput->GetDecodedWidth(), InDecoderOutput->GetDecodedHeight()), InOutBufferPropertes, Vars->VideoDecoderResourceDelegate.Pin().Get(), Vars->MaxWidth, Vars->MaxHeight, Vars->MaxOutputBuffers);
+							DecoderOutput->InitializeWithResource(D3D12Device, Resource, Pitch, OutputSync, FIntPoint(InDecoderOutput->GetDecodedWidth(), InDecoderOutput->GetDecodedHeight()), InOutBufferPropertes, Vars->VideoDecoderResourceDelegate, Vars->MaxWidth, Vars->MaxHeight, Vars->MaxOutputBuffers);
 							return true;
 						}
 					}
