ElectraPlayer: fixed incorrect version handling of sgpd box

#rnx

diff --git a/Engine/Plugins/Media/ElectraPlayer/Source/ElectraPlayerRuntime/Private/Runtime/Demuxer/ParserISO14496-12.cpp b/Engine/Plugins/Media/ElectraPlayer/Source/ElectraPlayerRuntime/Private/Runtime/Demuxer/ParserISO14496-12.cpp
--- a/Engine/Plugins/Media/ElectraPlayer/Source/ElectraPlayerRuntime/Private/Runtime/Demuxer/ParserISO14496-12.cpp
+++ b/Engine/Plugins/Media/ElectraPlayer/Source/ElectraPlayerRuntime/Private/Runtime/Demuxer/ParserISO14496-12.cpp
@@ -4975,7 +4975,7 @@ private:
 
 			RETURN_IF_ERROR(ParseInfo->Reader()->Read(GroupingType));							// grouping_type
 
-			if (Version == 1)
+			if (Version >= 1)
 			{
 				RETURN_IF_ERROR(ParseInfo->Reader()->Read(DefaultLength));						// default_length
 			}
