ElectraPlayer: Checking if across DASH periods the largest resolution video stream is larger than the previous and whether or not the decoder needs to be re-created in order to decode it.

#rnx

diff --git a/Engine/Plugins/Media/ElectraPlayer/Source/ElectraPlayerRuntime/Private/Runtime/Decoder/VideoDecoder.cpp b/Engine/Plugins/Media/ElectraPlayer/Source/ElectraPlayerRuntime/Private/Runtime/Decoder/VideoDecoder.cpp
--- a/Engine/Plugins/Media/ElectraPlayer/Source/ElectraPlayerRuntime/Private/Runtime/Decoder/VideoDecoder.cpp
+++ b/Engine/Plugins/Media/ElectraPlayer/Source/ElectraPlayerRuntime/Private/Runtime/Decoder/VideoDecoder.cpp
@@ -55,6 +55,7 @@ public:
 
 	void SetPlayerSessionServices(IPlayerSessionServices* SessionServices) override;
 	void Open(TSharedPtrTS<FAccessUnit::CodecData> InCodecData, FParamDict&& InAdditionalOptions, const FStreamCodecInformation* InMaxStreamConfiguration) override;
+	bool Reopen(TSharedPtrTS<FAccessUnit::CodecData> InCodecData, const FParamDict& InAdditionalOptions, const FStreamCodecInformation* InMaxStreamConfiguration) override;
 	void Close() override;
 	void DrainForCodecChange() override;
 	void SetVideoResourceDelegate(TWeakPtr<IVideoDecoderResourceDelegate, ESPMode::ThreadSafe> InVideoResourceDelegate) override;
@@ -256,7 +257,20 @@ TSharedPtr<IElectraCodecFactory, ESPMode::ThreadSafe> FVideoDecoderImpl::GetDeco
 	OutAddtlCfg.Add(TEXT("width"), FVariant((uint32)InCodecInfo.GetResolution().Width));
 	OutAddtlCfg.Add(TEXT("height"), FVariant((uint32)InCodecInfo.GetResolution().Height));
 	OutAddtlCfg.Add(TEXT("bitrate"), FVariant((int64)InCodecInfo.GetBitrate()));
-	OutAddtlCfg.Add(TEXT("fps"), FVariant((double)(InCodecInfo.GetFrameRate().IsValid() ? InCodecInfo.GetFrameRate().GetAsDouble() : 0.0)));
+	Electra::FTimeFraction Framerate = InCodecInfo.GetFrameRate();
+	if (Framerate.IsValid())
+	{
+		OutAddtlCfg.Add(TEXT("fps"), FVariant((double)Framerate.GetAsDouble()));
+		OutAddtlCfg.Add(TEXT("fps_n"), FVariant((int64)Framerate.GetNumerator()));
+		OutAddtlCfg.Add(TEXT("fps_d"), FVariant((uint32)Framerate.GetDenominator()));
+	}
+	else
+	{
+		OutAddtlCfg.Add(TEXT("fps"), FVariant((double)0.0));
+		OutAddtlCfg.Add(TEXT("fps_n"), FVariant((int64)0));
+		OutAddtlCfg.Add(TEXT("fps_d"), FVariant((uint32)1));
+	}
+
 	OutAddtlCfg.Add(TEXT("aspect_w"), FVariant((uint32)InCodecInfo.GetAspectRatio().Width));
 	OutAddtlCfg.Add(TEXT("aspect_h"), FVariant((uint32)InCodecInfo.GetAspectRatio().Height));
 	if (InCodecData.IsValid() && InCodecData->CodecSpecificData.Num())
@@ -309,6 +323,41 @@ void FVideoDecoderImpl::Open(TSharedPtrTS<FAccessUnit::CodecData> InCodecData, F
 	StartThread();
 }
 
+bool FVideoDecoderImpl::Reopen(TSharedPtrTS<FAccessUnit::CodecData> InCodecData, const FParamDict& InAdditionalOptions, const FStreamCodecInformation* InMaxStreamConfiguration)
+{
+	// Check if we can be used to decode the next set of streams.
+	// If no new information is provided, err on the safe side and say we can't be used for this.
+	if (!InCodecData.IsValid() || !InMaxStreamConfiguration)
+	{
+		return false;
+	}
+	// Check new against old limits.
+	if (InitialMaxStreamProperties.IsSet() && InMaxStreamConfiguration)
+	{
+		// If the codec has suddenly changed, we cannot be used.
+		if (InitialMaxStreamProperties.GetValue().GetCodec() != InMaxStreamConfiguration->GetCodec())
+		{
+			return false;
+		}
+		// If this is a H.265 stream of different profile (Main vs. Main10) we cannot be used.
+		if (InMaxStreamConfiguration->GetCodec() == FStreamCodecInformation::ECodec::H265 &&
+			InMaxStreamConfiguration->GetProfile() != InitialMaxStreamProperties.GetValue().GetProfile())
+		{
+			return false;
+		}
+		// If the current maximum resolution is less than what is required now, we cannot be used.
+		if (InitialMaxStreamProperties.GetValue().GetResolution().Width  < InMaxStreamConfiguration->GetResolution().Width ||
+			InitialMaxStreamProperties.GetValue().GetResolution().Height < InMaxStreamConfiguration->GetResolution().Height)
+		{
+			return false;
+		}
+		// Assume at this point that we can be used.
+		return true;
+	}
+
+	return false;
+}
+
 void FVideoDecoderImpl::Close()
 {
 	StopThread();
@@ -807,7 +856,14 @@ IElectraDecoder::EOutputStatus FVideoDecoderImpl::HandleOutput()
 				CurrentDecoderOutput->GetTransferHandle()->ReleaseHandle();
 			}
 
-			NotifyReadyBufferListener(true);
+			// Check if the output can actually be output or if the decoder says this is not to be output (incorrectly decoded)
+			//bool bUseOutput = CurrentDecoderOutput->GetOutputType() == IElectraDecoderVideoOutput::EOutputType::Output;
+			bool bUseOutput = true;
+			if (bUseOutput)
+			{
+				NotifyReadyBufferListener(true);
+			}
+			if (1)
 			{
 				SCOPE_CYCLE_COUNTER(STAT_ElectraPlayer_VideoConvertOutput);
 				CSV_SCOPED_TIMING_STAT(ElectraPlayer, VideoConvertOutput);
@@ -829,6 +885,9 @@ IElectraDecoder::EOutputStatus FVideoDecoderImpl::HandleOutput()
 							with PTS's increasing from there. On an open GOP or SAP type 2 or worse there may be frames with
 							PTS's earlier than the starting frame.
 
+							It may also be that the decoder could not produce valid output for some of the earlier input because
+							of a broken frame or a frame that needed nonexisting frames as references.
+
 							We check if there is a precise match somewhere in our list and use it.
 							Any elements in the list that are far too old we remove since it is not likely for the decoder to
 							emit those frames at all and we don't want our list to grow too long.
@@ -864,13 +923,14 @@ IElectraDecoder::EOutputStatus FVideoDecoderImpl::HandleOutput()
 				// Set properties from the bitstream messages.
 				BitstreamProcessor->SetPropertiesOnOutput(CurrentDecoderOutput, BufferProperties.Get(), MatchingInput->BitstreamInfo);
 
-				if (!FPlatformElectraDecoderResourceManager::SetupRenderBufferFromDecoderOutput(CurrentOutputBuffer, BufferProperties, CurrentDecoderOutput, PlatformResource))
+				if (bUseOutput && !FPlatformElectraDecoderResourceManager::SetupRenderBufferFromDecoderOutput(CurrentOutputBuffer, BufferProperties, CurrentDecoderOutput, PlatformResource))
 				{
 					PostError(0, TEXT("Failed to set up the decoder output!"), ERRCODE_VIDEO_INTERNAL_FAILED_TO_CONVERT_OUTPUT);
 					return IElectraDecoder::EOutputStatus::Error;
 				}
 
-				Renderer->ReturnBuffer(CurrentOutputBuffer, MatchingInput->AdjustedPTS.IsValid(), *BufferProperties);
+				bUseOutput = bUseOutput ? MatchingInput->AdjustedPTS.IsValid() : false;
+				Renderer->ReturnBuffer(CurrentOutputBuffer, bUseOutput, *BufferProperties);
 				CurrentDecoderOutput.Reset();
 				CurrentOutputBuffer = nullptr;
 			}
@@ -1283,7 +1343,18 @@ void FVideoDecoderImpl::WorkerThread()
 			DecoderFactoryAddtlCfg.Add(TEXT("max_width"), FVariant((uint32)InitialMaxStreamProperties.GetValue().GetResolution().Width));
 			DecoderFactoryAddtlCfg.Add(TEXT("max_height"), FVariant((uint32)InitialMaxStreamProperties.GetValue().GetResolution().Height));
 			DecoderFactoryAddtlCfg.Add(TEXT("max_bitrate"), FVariant((int64)InitialMaxStreamProperties.GetValue().GetBitrate()));
-			DecoderFactoryAddtlCfg.Add(TEXT("max_fps"), FVariant((double)(InitialMaxStreamProperties.GetValue().GetFrameRate().IsValid() ? InitialMaxStreamProperties.GetValue().GetFrameRate().GetAsDouble() : 0.0)));
+			if (InitialMaxStreamProperties.GetValue().GetFrameRate().IsValid())
+			{
+				DecoderFactoryAddtlCfg.Add(TEXT("max_fps"), FVariant((double)InitialMaxStreamProperties.GetValue().GetFrameRate().GetAsDouble()));
+				DecoderFactoryAddtlCfg.Add(TEXT("max_fps_n"), FVariant((int64)InitialMaxStreamProperties.GetValue().GetFrameRate().GetNumerator()));
+				DecoderFactoryAddtlCfg.Add(TEXT("max_fps_d"), FVariant((uint32)InitialMaxStreamProperties.GetValue().GetFrameRate().GetDenominator()));
+			}
+			else
+			{
+				DecoderFactoryAddtlCfg.Add(TEXT("max_fps"), FVariant((double)0.0));
+				DecoderFactoryAddtlCfg.Add(TEXT("max_fps_n"), FVariant((int64)0));
+				DecoderFactoryAddtlCfg.Add(TEXT("max_fps_d"), FVariant((uint32)0));
+			}
 			DecoderFactoryAddtlCfg.Add(TEXT("max_codecprofile"), FVariant(InitialMaxStreamProperties.GetValue().GetCodecSpecifierRFC6381()));
 		}
 		if (DecoderFactory.IsValid())
diff --git a/Engine/Plugins/Media/ElectraPlayer/Source/ElectraPlayerRuntime/Private/Runtime/Decoder/VideoDecoder.h b/Engine/Plugins/Media/ElectraPlayer/Source/ElectraPlayerRuntime/Private/Runtime/Decoder/VideoDecoder.h
--- a/Engine/Plugins/Media/ElectraPlayer/Source/ElectraPlayerRuntime/Private/Runtime/Decoder/VideoDecoder.h
+++ b/Engine/Plugins/Media/ElectraPlayer/Source/ElectraPlayerRuntime/Private/Runtime/Decoder/VideoDecoder.h
@@ -27,6 +27,7 @@ namespace Electra
 		virtual void SetPlayerSessionServices(IPlayerSessionServices* SessionServices) = 0;
 
 		virtual void Open(TSharedPtrTS<FAccessUnit::CodecData> InCodecData, FParamDict&& InAdditionalOptions, const FStreamCodecInformation* InMaxStreamConfiguration) = 0;
+		virtual bool Reopen(TSharedPtrTS<FAccessUnit::CodecData> InCodecData, const FParamDict& InAdditionalOptions, const FStreamCodecInformation* InMaxStreamConfiguration) = 0;
 		virtual void Close() = 0;
 		virtual void DrainForCodecChange() = 0;
 
diff --git a/Engine/Plugins/Media/ElectraPlayer/Source/ElectraPlayerRuntime/Private/Runtime/Player/AdaptiveStreamingPlayer.cpp b/Engine/Plugins/Media/ElectraPlayer/Source/ElectraPlayerRuntime/Private/Runtime/Player/AdaptiveStreamingPlayer.cpp
--- a/Engine/Plugins/Media/ElectraPlayer/Source/ElectraPlayerRuntime/Private/Runtime/Player/AdaptiveStreamingPlayer.cpp
+++ b/Engine/Plugins/Media/ElectraPlayer/Source/ElectraPlayerRuntime/Private/Runtime/Player/AdaptiveStreamingPlayer.cpp
@@ -4011,7 +4011,7 @@ void FAdaptiveStreamingPlayer::UpdateDataAvailabilityState(Metrics::FDataAvailab
 
 //-----------------------------------------------------------------------------
 /**
- * Updates the ABR and video decoder with maximum stream resolution limits.
+ * Updates the ABR with maximum stream resolution limits.
  */
 void FAdaptiveStreamingPlayer::UpdateStreamResolutionLimit()
 {
diff --git a/Engine/Plugins/Media/ElectraPlayer/Source/ElectraPlayerRuntime/Private/Runtime/Player/AdaptiveStreamingPlayerDecoders.cpp b/Engine/Plugins/Media/ElectraPlayer/Source/ElectraPlayerRuntime/Private/Runtime/Player/AdaptiveStreamingPlayerDecoders.cpp
--- a/Engine/Plugins/Media/ElectraPlayer/Source/ElectraPlayerRuntime/Private/Runtime/Player/AdaptiveStreamingPlayerDecoders.cpp
+++ b/Engine/Plugins/Media/ElectraPlayer/Source/ElectraPlayerRuntime/Private/Runtime/Player/AdaptiveStreamingPlayerDecoders.cpp
@@ -254,6 +254,7 @@ int32 FAdaptiveStreamingPlayer::CreateDecoder(EStreamType type)
 				check(AccessUnit->BufferSourceInfo.IsValid());
 				FString PeriodID = AccessUnit->BufferSourceInfo.IsValid() ? AccessUnit->BufferSourceInfo->PeriodID : FString();
 				VideoDecoder.CurrentCodecInfo.Clear();
+				VideoDecoder.LastBufferSourceInfo = AccessUnit->BufferSourceInfo;
 				if (AccessUnit->AUCodecData.IsValid())
 				{
 					VideoDecoder.CurrentCodecInfo = AccessUnit->AUCodecData->ParsedInfo;
@@ -296,8 +297,7 @@ int32 FAdaptiveStreamingPlayer::CreateDecoder(EStreamType type)
 				VideoDecoder.CheckIfNewDecoderMustBeSuspendedImmediately();
 				if (VideoDecoder.Decoder)
 				{
-					// Now we get the currently limited stream resolution and let the decoder now what we will be using
-					// at most right now. This allows the decoder to be created with a smaller memory footprint at first.
+					// Apply the current limit.
 					UpdateStreamResolutionLimit();
 				}
 				else
@@ -574,32 +574,45 @@ void FAdaptiveStreamingPlayer::FeedDecoder(EStreamType Type, IAccessUnitBufferIn
 	FBufferStats* pStats = nullptr;
 	FStreamCodecInformation* CurrentCodecInfo = nullptr;
 	TSharedPtrTS<FAccessUnit::CodecData>* LastSentAUCodecData = nullptr;
+	TSharedPtrTS<const FBufferSourceInfo>* LastBufferSourceInfo = nullptr;
+
 	bool bCodecChangeDetected = false;
 	bool bIsDeselected = false;
 
 	switch(Type)
 	{
 		case EStreamType::Video:
+		{
 			pStats = &VideoBufferStats;
 			CurrentCodecInfo = &VideoDecoder.CurrentCodecInfo;
 			LastSentAUCodecData = &VideoDecoder.LastSentAUCodecData;
+			LastBufferSourceInfo = &VideoDecoder.LastBufferSourceInfo;
 			bIsDeselected = bIsVideoDeselected;
 			break;
+		}
 		case EStreamType::Audio:
+		{
 			pStats = &AudioBufferStats;
 			CurrentCodecInfo = &AudioDecoder.CurrentCodecInfo;
 			LastSentAUCodecData = &AudioDecoder.LastSentAUCodecData;
+			//LastBufferSourceInfo = &AudioDecoder.LastBufferSourceInfo;
 			bIsDeselected = bIsAudioDeselected;
 			break;
+		}
 		case EStreamType::Subtitle:
+		{
 			pStats = &TextBufferStats;
 			CurrentCodecInfo = &SubtitleDecoder.CurrentCodecInfo;
 			LastSentAUCodecData = &SubtitleDecoder.LastSentAUCodecData;
+			//LastBufferSourceInfo = &SubtitleDecoder.LastBufferSourceInfo;
 			bIsDeselected = bIsTextDeselected;
 			break;
+		}
 		default:
+		{
 			checkNoEntry();
 			return;
+		}
 	}
 
 	// Lock the AU buffer for the duration of this function to ensure this can never clash with a Flush() call
@@ -659,10 +672,24 @@ void FAdaptiveStreamingPlayer::FeedDecoder(EStreamType Type, IAccessUnitBufferIn
 					bCodecChangeDetected = true;
 				}
 				// Mimetype change in subtitles?
-				else if (Type == EStreamType::Subtitle && !PeekedAU->AUCodecData->ParsedInfo.GetMimeType().Equals(CurrentCodecInfo->GetMimeType()))
+				if (Type == EStreamType::Subtitle && !PeekedAU->AUCodecData->ParsedInfo.GetMimeType().Equals(CurrentCodecInfo->GetMimeType()))
 				{
 					bCodecChangeDetected = true;
 				}
+				// Change in period of the video stream that may have a new maximum resolution from the one before?
+				if (!bCodecChangeDetected && Type == EStreamType::Video && VideoDecoder.Decoder && LastBufferSourceInfo && PeekedAU->BufferSourceInfo.IsValid() &&
+					(((*LastBufferSourceInfo)->PeriodID != PeekedAU->BufferSourceInfo->PeriodID) || ((*LastBufferSourceInfo)->PeriodAdaptationSetID != PeekedAU->BufferSourceInfo->PeriodAdaptationSetID)))
+				{
+					// Get the new highest stream properties and check if the decoder can still handle those.
+					FStreamCodecInformation HighestStream;
+					if (FindMatchingStreamInfo(HighestStream, PeekedAU->BufferSourceInfo->PeriodID, PeekedAU->EarliestPTS.IsValid() ? PeekedAU->EarliestPTS : PeekedAU->PTS, PeekedAU->AUCodecData->ParsedInfo))
+					{
+						if (!VideoDecoder.Decoder->Reopen(PeekedAU->AUCodecData, PlayerOptions.GetDictionary(), &HighestStream))
+						{
+							bCodecChangeDetected = true;
+						}
+					}
+				}
 			}
 
 			// Is there a change?
@@ -714,6 +741,11 @@ void FAdaptiveStreamingPlayer::FeedDecoder(EStreamType Type, IAccessUnitBufferIn
 						Decoder->AUdataClearEOD();
 					}
 					Decoder->AUdataPushAU(AccessUnit);
+					// Remember the buffer info for period transition checks.
+					if (Type == EStreamType::Video)
+					{
+						VideoDecoder.LastBufferSourceInfo = AccessUnit->BufferSourceInfo;
+					}
 				}
 
 				// If there is any pertinent format change, emit an event.
diff --git a/Engine/Plugins/Media/ElectraPlayer/Source/ElectraPlayerRuntime/Private/Runtime/Player/AdaptiveStreamingPlayerInternal.h b/Engine/Plugins/Media/ElectraPlayer/Source/ElectraPlayerRuntime/Private/Runtime/Player/AdaptiveStreamingPlayerInternal.h
--- a/Engine/Plugins/Media/ElectraPlayer/Source/ElectraPlayerRuntime/Private/Runtime/Player/AdaptiveStreamingPlayerInternal.h
+++ b/Engine/Plugins/Media/ElectraPlayer/Source/ElectraPlayerRuntime/Private/Runtime/Player/AdaptiveStreamingPlayerInternal.h
@@ -1237,6 +1237,7 @@ private:
 			delete Decoder;
 			Decoder = nullptr;
 			LastSentAUCodecData.Reset();
+			LastBufferSourceInfo.Reset();
 		}
 		void Flush()
 		{
@@ -1285,6 +1286,7 @@ private:
 
 		FStreamCodecInformation CurrentCodecInfo;
 		TSharedPtrTS<FAccessUnit::CodecData> LastSentAUCodecData;
+		TSharedPtrTS<const FBufferSourceInfo> LastBufferSourceInfo;
 		FAdaptiveStreamingPlayer* Parent = nullptr;
 		IVideoDecoder* Decoder = nullptr;
 		bool bDrainingForCodecChange = false;
@@ -1387,7 +1389,7 @@ private:
 				Decoder->AUdataClearEOD();
 			}
 		}
-		
+
 		void Start()
 		{
 			if (Decoder)
@@ -1399,7 +1401,7 @@ private:
 				}
 			}
 		}
-		
+
 		void Stop()
 		{
 			if (bIsRunning)
@@ -1981,8 +1983,8 @@ private:
 		}
 		bool Handle(const FTimeValue& InAtTime);
 		TSharedPtrTS<UtilsMP4::FMetadataParser> GetActive() const
-		{ 
-			return ActiveMetadata; 
+		{
+			return ActiveMetadata;
 		}
 
 		struct FEntry
