Cleaning up D3D dependencies in Electra.

diff --git a/Engine/Plugins/Media/ElectraPlayer/Source/ElectraPlayerRuntime/ElectraPlayerRuntime.Build.cs b/Engine/Plugins/Media/ElectraPlayer/Source/ElectraPlayerRuntime/ElectraPlayerRuntime.Build.cs
--- a/Engine/Plugins/Media/ElectraPlayer/Source/ElectraPlayerRuntime/ElectraPlayerRuntime.Build.cs
+++ b/Engine/Plugins/Media/ElectraPlayer/Source/ElectraPlayerRuntime/ElectraPlayerRuntime.Build.cs
@@ -1,6 +1,7 @@
 // Copyright Epic Games, Inc. All Rights Reserved.
 using UnrealBuildTool;
 using System.IO;
+using EpicGames.Core;
 
 namespace UnrealBuildTool.Rules
 {
@@ -48,7 +49,9 @@ namespace UnrealBuildTool.Rules
 				PrivateDefinitions.Add("_CRT_SECURE_NO_WARNINGS=1");
 
 				AddEngineThirdPartyPrivateStaticDependencies(Target, "WinHttp");
-				AddEngineThirdPartyPrivateStaticDependencies(Target, "DX12");
+
+				PrivateDependencyModuleNames.AddAll("D3D11RHI", "D3D12RHI");
+				AddEngineThirdPartyPrivateStaticDependencies(Target, "DX11", "DX12");
 
 				PrivateDefinitions.Add("ELECTRA_HAVE_DX11");	// video decoding for DX11 enabled (Win8+)
 
diff --git a/Engine/Plugins/Media/ElectraUtil/Source/ElectraSamples/ElectraSamples.Build.cs b/Engine/Plugins/Media/ElectraUtil/Source/ElectraSamples/ElectraSamples.Build.cs
--- a/Engine/Plugins/Media/ElectraUtil/Source/ElectraSamples/ElectraSamples.Build.cs
+++ b/Engine/Plugins/Media/ElectraUtil/Source/ElectraSamples/ElectraSamples.Build.cs
@@ -1,6 +1,7 @@
 // Copyright Epic Games, Inc. All Rights Reserved.
 using UnrealBuildTool;
 using System.IO;
+using EpicGames.Core;
 
 namespace UnrealBuildTool.Rules
 {
@@ -49,9 +50,8 @@ namespace UnrealBuildTool.Rules
 					});
 				}
 
-				PrivateDependencyModuleNames.Add("D3D11RHI");
-				PrivateDependencyModuleNames.Add("D3D12RHI");
-				AddEngineThirdPartyPrivateStaticDependencies(Target, "DX12");
+				PrivateDependencyModuleNames.AddAll("D3D11RHI", "D3D12RHI");
+				AddEngineThirdPartyPrivateStaticDependencies(Target, "DX11", "DX12");
 
 				PrivateDefinitions.Add("ELECTRA_SUPPORT_PREWIN8");
 				PrivateDefinitions.Add("ELECTRA_HAVE_DX11");
diff --git a/Engine/Plugins/Media/ElectraUtil/Source/ElectraSamples/Private/ElectraTextureSamplePC.cpp b/Engine/Plugins/Media/ElectraUtil/Source/ElectraSamples/Private/ElectraTextureSamplePC.cpp
--- a/Engine/Plugins/Media/ElectraUtil/Source/ElectraSamples/Private/ElectraTextureSamplePC.cpp
+++ b/Engine/Plugins/Media/ElectraUtil/Source/ElectraSamples/Private/ElectraTextureSamplePC.cpp
@@ -12,8 +12,7 @@
 #include "ID3D12DynamicRHI.h"
 
 #ifdef ELECTRA_HAVE_DX11
-#include "D3D11State.h"
-#include "D3D11Resources.h"
+#include "ID3D11DynamicRHI.h"
 #endif
 
 /*
