ElectraCodecs: Added support for D3D12 video decoding on Windows for H.264 and H.265

#rnx

diff --git a/Engine/Plugins/Media/AppleProResDecoderElectra/Source/AppleProResDecoderElectra/Private/ProResDecoder/Mac/ElectraMediaProResDecoder.cpp b/Engine/Plugins/Media/AppleProResDecoderElectra/Source/AppleProResDecoderElectra/Private/ProResDecoder/Mac/ElectraMediaProResDecoder.cpp
--- a/Engine/Plugins/Media/AppleProResDecoderElectra/Source/AppleProResDecoderElectra/Private/ProResDecoder/Mac/ElectraMediaProResDecoder.cpp
+++ b/Engine/Plugins/Media/AppleProResDecoderElectra/Source/AppleProResDecoderElectra/Private/ProResDecoder/Mac/ElectraMediaProResDecoder.cpp
@@ -62,6 +62,8 @@ public:
 	uint64 GetUserValue() const override
 	{ return UserValue; }
 
+	EOutputType GetOutputType() const
+	{ return EOutputType::Output; }
 	int32 GetWidth() const override
 	{ return Width - Crop.Left - Crop.Right; }
 	int32 GetHeight() const override
diff --git a/Engine/Plugins/Media/AppleProResDecoderElectra/Source/AppleProResDecoderElectra/Private/ProResDecoder/Windows/ElectraMediaProResDecoder.cpp b/Engine/Plugins/Media/AppleProResDecoderElectra/Source/AppleProResDecoderElectra/Private/ProResDecoder/Windows/ElectraMediaProResDecoder.cpp
--- a/Engine/Plugins/Media/AppleProResDecoderElectra/Source/AppleProResDecoderElectra/Private/ProResDecoder/Windows/ElectraMediaProResDecoder.cpp
+++ b/Engine/Plugins/Media/AppleProResDecoderElectra/Source/AppleProResDecoderElectra/Private/ProResDecoder/Windows/ElectraMediaProResDecoder.cpp
@@ -56,6 +56,11 @@ public:
 		return UserValue;
 	}
 
+	EOutputType GetOutputType() const
+	{
+		return EOutputType::Output;
+	}
+
 	int32 GetWidth() const override
 	{
 		return ImageWidth;
diff --git a/Engine/Plugins/Media/AvidDNxHDDecoderElectra/Source/Private/Windows/ElectraMediaAvidDNxHDDecoder.cpp b/Engine/Plugins/Media/AvidDNxHDDecoderElectra/Source/Private/Windows/ElectraMediaAvidDNxHDDecoder.cpp
--- a/Engine/Plugins/Media/AvidDNxHDDecoderElectra/Source/Private/Windows/ElectraMediaAvidDNxHDDecoder.cpp
+++ b/Engine/Plugins/Media/AvidDNxHDDecoderElectra/Source/Private/Windows/ElectraMediaAvidDNxHDDecoder.cpp
@@ -59,6 +59,11 @@ public:
 		return UserValue;
 	}
 
+	EOutputType GetOutputType() const
+	{
+		return EOutputType::Output;
+	}
+
 	int32 GetWidth() const override
 	{
 		return ImageWidth;
@@ -1015,9 +1020,9 @@ int32 FVideoDecoderAvidDNxHDElectra::FDecoderHandle::GetNumberOfOutputColorBitsP
 			return 8;
 		case DNX_ComponentType_t::DNX_CT_USHORT_10_6:	// 10 bit
 		case DNX_ComponentType_t::DNX_CT_10Bit_2_8:		// 10 bit in 2_8 format. Byte ordering is fixed. This is to be used with 10-bit 4:2:2 YCbCr components.
-		case DNX_ComponentType_t::DNX_CT_V210:			// Apple's V210 
+		case DNX_ComponentType_t::DNX_CT_V210:			// Apple's V210
 			return 10;
-		case DNX_ComponentType_t::DNX_CT_SHORT_2_14:	// Fixed point 
+		case DNX_ComponentType_t::DNX_CT_SHORT_2_14:	// Fixed point
 		case DNX_ComponentType_t::DNX_CT_SHORT:			// 16 bit. Premultiplied by 257. Byte ordering is machine dependent.
 			return 16;
 		case DNX_ComponentType_t::DNX_CT_USHORT_12_4:	// 12 bit
diff --git a/Engine/Plugins/Media/ElectraCodecs/Source/ElectraDecoders/Private/Android/h264/VideoDecoderH264_Android.cpp b/Engine/Plugins/Media/ElectraCodecs/Source/ElectraDecoders/Private/Android/h264/VideoDecoderH264_Android.cpp
--- a/Engine/Plugins/Media/ElectraCodecs/Source/ElectraDecoders/Private/Android/h264/VideoDecoderH264_Android.cpp
+++ b/Engine/Plugins/Media/ElectraCodecs/Source/ElectraDecoders/Private/Android/h264/VideoDecoderH264_Android.cpp
@@ -53,6 +53,8 @@ public:
 	uint64 GetUserValue() const override
 	{ return UserValue; }
 
+	EOutputType GetOutputType() const
+	{ return EOutputType::Output; }
 	int32 GetWidth() const override
 	{ return Width - Crop.Left - Crop.Right; }
 	int32 GetHeight() const override
@@ -409,7 +411,7 @@ bool FElectraVideoDecoderH264_Android::ResetToCleanStart()
 	DecodeState = EDecodeState::Decoding;
 	LastPushedPresentationTime = -1;
 	bDidSendEOS = false;
-	
+
 	return !LastError.IsSet();
 }
 
@@ -1102,7 +1104,7 @@ void FElectraVideoDecoderOutputH264_Android::ReleaseOutputBuffer()
 }
 
 IElectraDecoderVideoOutput::EImageCopyResult FElectraVideoDecoderOutputH264_Android::CopyPlatformImage(IElectraDecoderVideoOutputCopyResources* InCopyResources) const
-{ 
+{
 	TSharedPtr<FElectraVideoDecoderH264_Android, ESPMode::ThreadSafe> Decoder = OwningDecoder.Pin();
 	if (Decoder.IsValid())
 	{
@@ -1117,7 +1119,7 @@ IElectraDecoderVideoOutput::EImageCopyResult FElectraVideoDecoderOutputH264_Andr
 			OwningDecoderBufferInfo.BufferIndex = -1;
 		}
 
-		return IElectraDecoderVideoOutput::EImageCopyResult::Succeeded; 
+		return IElectraDecoderVideoOutput::EImageCopyResult::Succeeded;
 	}
-	return IElectraDecoderVideoOutput::EImageCopyResult::Failed; 
+	return IElectraDecoderVideoOutput::EImageCopyResult::Failed;
 }
diff --git a/Engine/Plugins/Media/ElectraCodecs/Source/ElectraDecoders/Private/Android/h265/VideoDecoderH265_Android.cpp b/Engine/Plugins/Media/ElectraCodecs/Source/ElectraDecoders/Private/Android/h265/VideoDecoderH265_Android.cpp
--- a/Engine/Plugins/Media/ElectraCodecs/Source/ElectraDecoders/Private/Android/h265/VideoDecoderH265_Android.cpp
+++ b/Engine/Plugins/Media/ElectraCodecs/Source/ElectraDecoders/Private/Android/h265/VideoDecoderH265_Android.cpp
@@ -53,6 +53,8 @@ public:
 	uint64 GetUserValue() const override
 	{ return UserValue; }
 
+	EOutputType GetOutputType() const
+	{ return EOutputType::Output; }
 	int32 GetWidth() const override
 	{ return Width - Crop.Left - Crop.Right; }
 	int32 GetHeight() const override
@@ -436,7 +438,7 @@ bool FElectraVideoDecoderH265_Android::ResetToCleanStart()
 	DecodeState = EDecodeState::Decoding;
 	LastPushedPresentationTime = -1;
 	bDidSendEOS = false;
-	
+
 	return !LastError.IsSet();
 }
 
@@ -1130,7 +1132,7 @@ void FElectraVideoDecoderOutputH265_Android::ReleaseOutputBuffer()
 }
 
 IElectraDecoderVideoOutput::EImageCopyResult FElectraVideoDecoderOutputH265_Android::CopyPlatformImage(IElectraDecoderVideoOutputCopyResources* InCopyResources) const
-{ 
+{
 	TSharedPtr<FElectraVideoDecoderH265_Android, ESPMode::ThreadSafe> Decoder = OwningDecoder.Pin();
 	if (Decoder.IsValid())
 	{
@@ -1145,7 +1147,7 @@ IElectraDecoderVideoOutput::EImageCopyResult FElectraVideoDecoderOutputH265_Andr
 			OwningDecoderBufferInfo.BufferIndex = -1;
 		}
 
-		return IElectraDecoderVideoOutput::EImageCopyResult::Succeeded; 
+		return IElectraDecoderVideoOutput::EImageCopyResult::Succeeded;
 	}
-	return IElectraDecoderVideoOutput::EImageCopyResult::Failed; 
+	return IElectraDecoderVideoOutput::EImageCopyResult::Failed;
 }
diff --git a/Engine/Plugins/Media/ElectraCodecs/Source/ElectraDecoders/Private/Apple/h264/VideoDecoderH264_Apple.cpp b/Engine/Plugins/Media/ElectraCodecs/Source/ElectraDecoders/Private/Apple/h264/VideoDecoderH264_Apple.cpp
--- a/Engine/Plugins/Media/ElectraCodecs/Source/ElectraDecoders/Private/Apple/h264/VideoDecoderH264_Apple.cpp
+++ b/Engine/Plugins/Media/ElectraCodecs/Source/ElectraDecoders/Private/Apple/h264/VideoDecoderH264_Apple.cpp
@@ -44,6 +44,8 @@ public:
 	uint64 GetUserValue() const override
 	{ return UserValue; }
 
+	EOutputType GetOutputType() const
+	{ return EOutputType::Output; }
 	int32 GetWidth() const override
 	{ return Width - Crop.Left - Crop.Right; }
 	int32 GetHeight() const override
diff --git a/Engine/Plugins/Media/ElectraCodecs/Source/ElectraDecoders/Private/DX/VideoDecoderH265_DX.cpp b/Engine/Plugins/Media/ElectraCodecs/Source/ElectraDecoders/Private/DX/VideoDecoderH265_DX.cpp
--- a/Engine/Plugins/Media/ElectraCodecs/Source/ElectraDecoders/Private/DX/VideoDecoderH265_DX.cpp
+++ b/Engine/Plugins/Media/ElectraCodecs/Source/ElectraDecoders/Private/DX/VideoDecoderH265_DX.cpp
@@ -70,6 +70,10 @@ public:
 		return UserValue;
 	}
 
+	EOutputType GetOutputType() const
+	{
+		return EOutputType::Output;
+	}
 	int32 GetWidth() const override
 	{
 		return Width;
diff --git a/Engine/Plugins/Media/ElectraCodecs/Source/ElectraDecoders/Private/Utils/MPEG/ElectraUtilsMPEGVideo.cpp b/Engine/Plugins/Media/ElectraCodecs/Source/ElectraDecoders/Private/Utils/MPEG/ElectraUtilsMPEGVideo.cpp
--- a/Engine/Plugins/Media/ElectraCodecs/Source/ElectraDecoders/Private/Utils/MPEG/ElectraUtilsMPEGVideo.cpp
+++ b/Engine/Plugins/Media/ElectraCodecs/Source/ElectraDecoders/Private/Utils/MPEG/ElectraUtilsMPEGVideo.cpp
@@ -4,7 +4,6 @@
 #include "Utils/ElectraBitstreamReader.h"
 #include "ElectraDecodersUtils.h"
 
-
 namespace ElectraDecodersUtil
 {
 	namespace MPEG
@@ -503,8 +502,6 @@ namespace ElectraDecodersUtil
 		};
 
 
-
-
 		bool ParseH264SPS(FISO14496_10_seq_parameter_set_data& OutSPS, const void* Data, int32 Size)
 		{
 			struct FSyntaxElement
@@ -553,7 +550,8 @@ namespace ElectraDecodersUtil
 			OutSPS.level_idc = (uint8)BitReader.GetBits(8);
 			OutSPS.seq_parameter_set_id = FSyntaxElement::ue_v(BitReader);
 			if (OutSPS.profile_idc == 100 || OutSPS.profile_idc == 110 || OutSPS.profile_idc == 122 || OutSPS.profile_idc == 244 ||
-				OutSPS.profile_idc == 44 || OutSPS.profile_idc == 83 || OutSPS.profile_idc == 86 || OutSPS.profile_idc == 118 || OutSPS.profile_idc == 128)
+				OutSPS.profile_idc == 44 || OutSPS.profile_idc == 83 || OutSPS.profile_idc == 86 || OutSPS.profile_idc == 118 || OutSPS.profile_idc == 128 ||
+				OutSPS.profile_idc == 138 || OutSPS.profile_idc == 139 || OutSPS.profile_idc == 134 || OutSPS.profile_idc == 135)
 			{
 				OutSPS.chroma_format_idc = FSyntaxElement::ue_v(BitReader);
 				if (OutSPS.chroma_format_idc == 3)
@@ -721,9 +719,13 @@ namespace ElectraDecodersUtil
 			FMemory::Memzero(OutSPS);
 
 			if (BitReader.GetBits(1) != 0)	// forbidden_zero_bit
+			{
 				return false;
+			}
 			if (BitReader.GetBits(6) != 33)	// sps_nut ?
+			{
 				return false;
+			}
 			BitReader.SkipBits(6);			// nuh_layer_id
 			BitReader.SkipBits(3);			// nuh_temporal_id_plus1
 
@@ -1006,9 +1008,8 @@ namespace ElectraDecodersUtil
 // None of the following is of interest to us now, so we stop parsing at this point.
 					}
 				}
-			return true;
 			}
-		return false;
+			return true;
 		}
 
 
@@ -1178,5 +1179,4 @@ namespace ElectraDecodersUtil
 		}
 
 	} // namespace MPEG
-
 } // namespace ElectraDecodersUtil
diff --git a/Engine/Plugins/Media/ElectraCodecs/Source/ElectraDecoders/Private/Windows/WindowsPlatformHeaders_Video_DX.h b/Engine/Plugins/Media/ElectraCodecs/Source/ElectraDecoders/Private/Windows/WindowsPlatformHeaders_Video_DX.h
--- a/Engine/Plugins/Media/ElectraCodecs/Source/ElectraDecoders/Private/Windows/WindowsPlatformHeaders_Video_DX.h
+++ b/Engine/Plugins/Media/ElectraCodecs/Source/ElectraDecoders/Private/Windows/WindowsPlatformHeaders_Video_DX.h
@@ -5,8 +5,8 @@
 
 #include "CoreMinimal.h"
 
-#include "Windows/AllowWindowsPlatformTypes.h"
 THIRD_PARTY_INCLUDES_START
+#include "Microsoft/AllowMicrosoftPlatformTypes.h"
 
 #include <d3d11.h>
 #include <d3d11_1.h>
@@ -25,8 +25,8 @@ THIRD_PARTY_INCLUDES_START
 #define ALLOW_MFSAMPLE_WITH_DX12	0
 #endif
 
+#include "Microsoft/HideMicrosoftPlatformTypes.h"
 THIRD_PARTY_INCLUDES_END
-#include "Windows/HideWindowsPlatformTypes.h"
 
 struct FElectraVideoDecoderDXDeviceContext
 {
diff --git a/Engine/Plugins/Media/ElectraCodecs/Source/ElectraDecoders/Public/ElectraDecodersUtils.h b/Engine/Plugins/Media/ElectraCodecs/Source/ElectraDecoders/Public/ElectraDecodersUtils.h
--- a/Engine/Plugins/Media/ElectraCodecs/Source/ElectraDecoders/Public/ElectraDecodersUtils.h
+++ b/Engine/Plugins/Media/ElectraCodecs/Source/ElectraDecoders/Public/ElectraDecodersUtils.h
@@ -8,6 +8,15 @@
 
 namespace ElectraDecodersUtil
 {
+	struct FFractionalValue
+	{
+		FFractionalValue(int64 InNumerator=0, uint32 InDenominator=0)
+			: Num(InNumerator), Denom(InDenominator)
+		{ }
+		int64 Num = 0;
+		uint32 Denom = 0;
+	};
+
 	// Advance a pointer by a number of bytes.
 	template <typename T, typename C>
 	inline T AdvancePointer(T pPointer, C numBytes)
diff --git a/Engine/Plugins/Media/ElectraCodecs/Source/ElectraDecoders/Public/IElectraDecoderOutputVideo.h b/Engine/Plugins/Media/ElectraCodecs/Source/ElectraDecoders/Public/IElectraDecoderOutputVideo.h
--- a/Engine/Plugins/Media/ElectraCodecs/Source/ElectraDecoders/Public/IElectraDecoderOutputVideo.h
+++ b/Engine/Plugins/Media/ElectraCodecs/Source/ElectraDecoders/Public/IElectraDecoderOutputVideo.h
@@ -61,6 +61,23 @@ public:
 		return EType::Video;
 	}
 
+	enum class EOutputType
+	{
+		// Output this output.
+		Output,
+
+		// Do not output this. Either the decoder did not produce an actual output for the given
+		// input, or the output has decoding errors that makes it unusable.
+		DoNotOutput
+	};
+
+	/**
+	 * Returns the state of this output.
+	 * Depending on the input bitstream the decoder might not have been able to produce
+	 * a valid output or it may be partially corrupted.
+	 * If the output should not be used for display `DoNotOutput` is returned.
+	 */
+	virtual EOutputType GetOutputType() const = 0;
 
 	/**
 	 * Returns the active number of horizontal pixels.
@@ -75,7 +92,7 @@ public:
 	 * pitch requirements.
 	 */
 	virtual int32 GetHeight() const = 0;
-	
+
 	/**
 	 * Returns the number of decoded horizontal pixels.
 	 * This value may be different from GetWidth() since it does not take any
@@ -95,7 +112,7 @@ public:
 	 * (e.g. if the data is returned as RGBA pixels, but represents a YUYV 4:2:2 format)
 	 */
 	virtual int32 GetDecodedHeight() const = 0;
-	
+
 	/**
 	 * Returns the cropping values.
 	 * This is useful in combination with GetDecodedWidth() and GetDecodedHeight()
@@ -129,7 +146,7 @@ public:
 	 * Returns the number of bits. Usually 8, 10, or 12.
 	 */
 	virtual int32 GetNumberOfBits() const = 0;
-	
+
 	/**
 	 * Returns additional values specific to the decoder and format that are used
 	 * in handling this output in a platform specific way.
@@ -162,7 +179,7 @@ public:
 	 * This is a value particular to the target hardware.
 	 * It is usually a resource handle of sorts and is owned by the
 	 * platform's resource manager.
-	 * 
+	 *
 	 * USAGE TBD
 	 */
 	virtual IElectraDecoderVideoOutputTransferHandle* GetTransferHandle() const = 0;
diff --git a/Engine/Plugins/Media/ElectraCodecs/Source/ElectraDecoders/Public/Utils/ElectraBitstreamReader.h b/Engine/Plugins/Media/ElectraCodecs/Source/ElectraDecoders/Public/Utils/ElectraBitstreamReader.h
--- a/Engine/Plugins/Media/ElectraCodecs/Source/ElectraDecoders/Public/Utils/ElectraBitstreamReader.h
+++ b/Engine/Plugins/Media/ElectraCodecs/Source/ElectraDecoders/Public/Utils/ElectraBitstreamReader.h
@@ -45,9 +45,14 @@ public:
 		return DataSize - BytePosition;
 	}
 
+	uint32 GetBitPosition() const
+	{
+		return BitPosition;
+	}
+
 	uint64 GetRemainingBits() const
 	{
-		return GetRemainingByteLength() * 8 + (BitPosition ? 8 - BitPosition : 0);
+		return GetRemainingByteLength() * 8 - BitPosition;
 	}
 
 	bool GetAlignedBytes(uint8* To, uint64 nBytes)
@@ -65,6 +70,10 @@ public:
 	void SkipBytes(uint64 nBytes)
 	{
 		BytePosition += nBytes;
+		if (BytePosition > DataSize)
+		{
+			BytePosition = DataSize;
+		}
 	}
 
 	void SkipBits(uint64 nBits)
@@ -75,11 +84,24 @@ public:
 
 	uint32 GetBits(uint64 nBits)
 	{
-		uint32 rv = PeekBits(nBits);
-		uint32 bp = BitPosition + nBits;
-		BytePosition += bp >> 3;
-		BitPosition = bp & 7;
-		return rv;
+		if (nBits == 0)
+		{
+			return 0;
+		}
+		else if (nBits > GetRemainingBits())
+		{
+			BytePosition = DataSize;
+			BitPosition = 0;
+			return 0;
+		}
+		else
+		{
+			uint32 rv = PeekBits(nBits);
+			uint32 bp = BitPosition + nBits;
+			BytePosition += bp >> 3;
+			BitPosition = bp & 7;
+			return rv;
+		}
 	}
 
 	uint32 PeekBits(uint64 nBits)
@@ -89,6 +111,10 @@ public:
 		{
 			return 0;
 		}
+		else if (nBits > GetRemainingBits())
+		{
+			return 0;
+		}
 		else
 		{
 			uint32 bytePos = BytePosition + (BitPosition >> 3);
diff --git a/Engine/Plugins/Media/ElectraCodecs/Source/ElectraDecoders/Public/Utils/MPEG/ElectraUtilsMPEGVideo.h b/Engine/Plugins/Media/ElectraCodecs/Source/ElectraDecoders/Public/Utils/MPEG/ElectraUtilsMPEGVideo.h
--- a/Engine/Plugins/Media/ElectraCodecs/Source/ElectraDecoders/Public/Utils/MPEG/ElectraUtilsMPEGVideo.h
+++ b/Engine/Plugins/Media/ElectraCodecs/Source/ElectraDecoders/Public/Utils/MPEG/ElectraUtilsMPEGVideo.h
@@ -5,18 +5,13 @@
 #include <CoreMinimal.h>
 #include <Containers/Array.h>
 
+#include "ElectraDecodersUtils.h"
+
 namespace ElectraDecodersUtil
 {
 	namespace MPEG
 	{
-		struct FFractionalValue
-		{
-			FFractionalValue(int64 InNumerator=0, uint32 InDenominator=0)
-				: Num(InNumerator), Denom(InDenominator)
-			{ }
-			int64 Num = 0;
-			uint32 Denom = 0;
-		};
+		using ElectraDecodersUtil::FFractionalValue;
 
 		/**
 		 * Sequence parameter set as per ISO/IEC 14496-10:2012, section 7.3.2.1.1
@@ -302,20 +297,20 @@ namespace ElectraDecodersUtil
 					{
 						default:
 						case 0:
-							SubWidthC = 1; 
-							SubHeightC = 1;	
+							SubWidthC = 1;
+							SubHeightC = 1;
 							break;
-						case 1:	
-							SubWidthC = 2; 
-							SubHeightC = 2;	
+						case 1:
+							SubWidthC = 2;
+							SubHeightC = 2;
 							break;
-						case 2:	
-							SubWidthC = 2; 
-							SubHeightC = 1;	
+						case 2:
+							SubWidthC = 2;
+							SubHeightC = 1;
 							break;
-						case 3:	
-							SubWidthC = 1; 
-							SubHeightC = 1;	
+						case 3:
+							SubWidthC = 1;
+							SubHeightC = 1;
 							break;
 					}
 					Left = conf_win_left_offset * SubWidthC;
@@ -390,6 +385,16 @@ namespace ElectraDecodersUtil
 			int32 GetNumberOfSPS() const;
 			const FISO23008_2_seq_parameter_set_data& GetParsedSPS(int32 SpsIndex) const;
 
+			uint8 GetGeneralProfileSpace() const { return GeneralProfileSpace; }
+			uint8 GetGeneralTierFlag() const { return GeneralTierFlag; }
+			uint8 GetGeneralProfileIDC() const { return GeneralProfileIDC; }
+			uint32 GetGeneralProfileCompatibilityFlags() const { return GeneralProfileCompatibilityFlags; }
+			uint64 GetGeneralConstraintIndicatorFlags() const { return GeneralConstraintIndicatorFlags; }
+			uint8 GetGeneralLevelIDC() const { return GeneralLevelIDC; }
+			uint8 GetChromaFormat() const { return ChromaFormat; }
+			uint8 GetBitDepthLumaMinus8() const { return BitDepthLumaMinus8; }
+			uint8 GetBitDepthChromaMinus8() const { return BitDepthChromaMinus8; }
+
 		private:
 			struct FArray
 			{
diff --git a/Engine/Plugins/Media/ElectraPlayer/Source/ElectraPlayerRuntime/Private/Runtime/Decoder/BitstreamProcessors/VideoDecoderInputBitstreamProcessorH265.cpp b/Engine/Plugins/Media/ElectraPlayer/Source/ElectraPlayerRuntime/Private/Runtime/Decoder/BitstreamProcessors/VideoDecoderInputBitstreamProcessorH265.cpp
--- a/Engine/Plugins/Media/ElectraPlayer/Source/ElectraPlayerRuntime/Private/Runtime/Decoder/BitstreamProcessors/VideoDecoderInputBitstreamProcessorH265.cpp
+++ b/Engine/Plugins/Media/ElectraPlayer/Source/ElectraPlayerRuntime/Private/Runtime/Decoder/BitstreamProcessors/VideoDecoderInputBitstreamProcessorH265.cpp
@@ -155,13 +155,14 @@ IVideoDecoderInputBitstreamProcessor::EProcessResult FVideoDecoderInputBitstream
 
 		uint8 nut = *(const uint8 *)(NALU + 1);
 		nut >>= 1;
-		// IDR frame?
-		if (nut == 19 /*IDR_W_RADL*/ || nut == 20 /*IDR_N_LP*/ || nut == 21 /*CRA_NUT*/)
+
+		// IDR, CRA or BLA frame?
+		if (nut >= 16 && nut <= 21)
 		{
 			OutBSI.bIsSyncFrame = true;
 		}
 		// One of TRAIL_N, TSA_N, STSA_N, RADL_N, RASL_N, RSV_VCL_N10, RSV_VCL_N12 or RSV_VCL_N14 ?
-		else if (nut == 0 || nut == 2 || nut == 4 || nut == 6 || nut == 8 || nut == 10 || nut == 12 || nut == 14)
+		else if (nut <= 14 && (nut & 1) == 0)
 		{
 			OutBSI.bIsDiscardable = true;
 		}
diff --git a/Engine/Plugins/Media/ElectraPlayer/Source/ElectraPlayerRuntime/Private/Runtime/Decoder/VideoDecoder.cpp b/Engine/Plugins/Media/ElectraPlayer/Source/ElectraPlayerRuntime/Private/Runtime/Decoder/VideoDecoder.cpp
--- a/Engine/Plugins/Media/ElectraPlayer/Source/ElectraPlayerRuntime/Private/Runtime/Decoder/VideoDecoder.cpp
+++ b/Engine/Plugins/Media/ElectraPlayer/Source/ElectraPlayerRuntime/Private/Runtime/Decoder/VideoDecoder.cpp
@@ -857,8 +857,7 @@ IElectraDecoder::EOutputStatus FVideoDecoderImpl::HandleOutput()
 			}
 
 			// Check if the output can actually be output or if the decoder says this is not to be output (incorrectly decoded)
-			//bool bUseOutput = CurrentDecoderOutput->GetOutputType() == IElectraDecoderVideoOutput::EOutputType::Output;
-			bool bUseOutput = true;
+			bool bUseOutput = CurrentDecoderOutput->GetOutputType() == IElectraDecoderVideoOutput::EOutputType::Output;
 			if (bUseOutput)
 			{
 				NotifyReadyBufferListener(true);
diff --git a/Engine/Plugins/Media/ElectraPlayer/Source/ElectraPlayerRuntime/Private/Runtime/Demuxer/ParserISO14496-12.cpp b/Engine/Plugins/Media/ElectraPlayer/Source/ElectraPlayerRuntime/Private/Runtime/Demuxer/ParserISO14496-12.cpp
--- a/Engine/Plugins/Media/ElectraPlayer/Source/ElectraPlayerRuntime/Private/Runtime/Demuxer/ParserISO14496-12.cpp
+++ b/Engine/Plugins/Media/ElectraPlayer/Source/ElectraPlayerRuntime/Private/Runtime/Demuxer/ParserISO14496-12.cpp
@@ -269,36 +269,36 @@ namespace Electra
 	};
 
 
-class FDataBufferReader : public IParserISO14496_12::IReader
-{
-public:
-	FDataBufferReader(const TArray<uint8>& InDataBufferToReadFrom) : DataBufferRef(InDataBufferToReadFrom)
-	{}
-	virtual ~FDataBufferReader() = default;
-	int64 ReadData(void* IntoBuffer, int64 NumBytesToRead) override
+	class FDataBufferReader : public IParserISO14496_12::IReader
 	{
-		int64 NumAvail = DataBufferRef.Num() - CurrentOffset;
-		if (NumAvail >= NumBytesToRead)
+	public:
+		FDataBufferReader(const TArray<uint8>& InDataBufferToReadFrom) : DataBufferRef(InDataBufferToReadFrom)
+		{}
+		virtual ~FDataBufferReader() = default;
+		int64 ReadData(void* IntoBuffer, int64 NumBytesToRead) override
 		{
-			if (IntoBuffer)
+			int64 NumAvail = DataBufferRef.Num() - CurrentOffset;
+			if (NumAvail >= NumBytesToRead)
 			{
-				FMemory::Memcpy(IntoBuffer, DataBufferRef.GetData() + CurrentOffset, NumBytesToRead);
+				if (IntoBuffer)
+				{
+					FMemory::Memcpy(IntoBuffer, DataBufferRef.GetData() + CurrentOffset, NumBytesToRead);
+				}
+				CurrentOffset += NumBytesToRead;
+				return NumBytesToRead;
 			}
-			CurrentOffset += NumBytesToRead;
-			return NumBytesToRead;
+			return -1;
 		}
-		return -1;
-	}
-	bool HasReachedEOF() const override
-	{ return CurrentOffset >= DataBufferRef.Num(); }
-	bool HasReadBeenAborted() const override
-	{ return false;	}
-	int64 GetCurrentOffset() const override
-	{ return CurrentOffset;	}
-private:
-	const TArray<uint8>& DataBufferRef;
-	int64 CurrentOffset = 0;
-};
+		bool HasReachedEOF() const override
+		{ return CurrentOffset >= DataBufferRef.Num(); }
+		bool HasReadBeenAborted() const override
+		{ return false;	}
+		int64 GetCurrentOffset() const override
+		{ return CurrentOffset;	}
+	private:
+		const TArray<uint8>& DataBufferRef;
+		int64 CurrentOffset = 0;
+	};
 
 
 #define RETURN_IF_ERROR(expr)											\
@@ -791,7 +791,8 @@ private:
 		/**
 		 * Sample grouping types
 		 */
-		static const IParserISO14496_12::FBoxType kGrouping_seig = MAKE_BOX_ATOM('s', 'e', 'i', 'g');
+		//static const IParserISO14496_12::FBoxType kGrouping_seig = MAKE_BOX_ATOM('s', 'e', 'i', 'g');
+		//static const IParserISO14496_12::FBoxType kGrouping_rap  = MAKE_BOX_ATOM('r', 'a', 'p', ' ');
 
 		/**
 		 * Encryption scheme types
@@ -879,6 +880,11 @@ private:
 		{
 		}
 
+		uint8 GetBoxVersion() const
+		{
+			return Version;
+		}
+
 		uint32 GetFlags() const
 		{
 			return Flags;
@@ -4960,6 +4966,15 @@ private:
 		{
 			UEMediaError Error = UEMEDIA_ERROR_OK;
 			RETURN_IF_ERROR(FMP4BoxFull::ReadAndParseAttributes(ParseInfo));
+			/*
+				Flags field must be zero when this box is in a `traf` box, otherwise, if it appears in
+				a `stbl` box flag bit values mean:
+					Bit 0: when set to 1, this flag indicates that there are no SampleGroupDescriptionBoxes
+						   of this grouping_type in any TrackFragmentBox of this track.
+					Bit 1: when set to 1, this flag indicates that there are no SampleToGroupBoxes
+						   of this grouping_type in this track (in neither the SampleTableBox nor any TrackFragmentBox of this
+						   track); all samples therefore map to the default.
+			*/
 
 			// Version 0 of this box type has been deprecated.
 			if (Version == 0)
@@ -5757,6 +5772,14 @@ private:
 				return Track != nullptr;
 			}
 		private:
+			typedef uint32 FGroupType;
+			#define MAKE_MP4_GROUPTYPE(a,b,c,d) (FGroupType)((uint32)a << 24) | ((uint32)b << 16) | ((uint32)c << 8) | ((uint32)d)
+			enum EGroupType
+			{
+				Grouping_rap = MAKE_MP4_GROUPTYPE('r', 'a', 'p', ' '),
+			};
+			#undef MAKE_MP4_GROUPTYPE
+
 			FTrackIterator& AssignFrom(const FTrackIterator& Other)
 			{
 				if (this != &Other)
@@ -5792,6 +5815,8 @@ private:
 					SampleDuration = Other.SampleDuration;
 					SampleDescriptionIndex = Other.SampleDescriptionIndex;
 					bEOS = Other.bEOS;
+					GroupIterators = Other.GroupIterators;
+					bHaveRAPGroup = Other.bHaveRAPGroup;
 				}
 				return *this;
 			}
@@ -5835,6 +5860,52 @@ private:
 			uint32				SampleDuration = 0;
 			uint32				SampleDescriptionIndex = 0;
 			bool				bEOS = true;
+
+			struct FSampleGroupIterator
+			{
+				FBoxType GroupingType = 0;
+				const FMP4BoxSBGP* SBGPBox = nullptr;
+				uint32 DefaultDescriptionIndex = 0;
+				uint32 CurrentSampleIdxInEntry = 0;
+				int32 CurrentEntryIndex = 0;
+				uint32 CurrentGroupDescriptionIndex = 0;
+				uint32 GetCurrentGroupDescriptionIndex() const
+				{
+					return CurrentGroupDescriptionIndex;
+				}
+				void SetFirst()
+				{
+					check(SBGPBox);
+					CurrentSampleIdxInEntry = 0;
+					CurrentEntryIndex = 0;
+					CurrentGroupDescriptionIndex = SBGPBox->GetNumberOfEntries() ? SBGPBox->GetEntry(CurrentEntryIndex).GroupDescriptionIndex : DefaultDescriptionIndex;
+				}
+				void Advance()
+				{
+					if (CurrentEntryIndex < SBGPBox->GetNumberOfEntries())
+					{
+						// Need to move onto the next entry?
+						if (++CurrentSampleIdxInEntry >= SBGPBox->GetEntry(CurrentEntryIndex).SampleCount)
+						{
+							CurrentSampleIdxInEntry = 0;
+							if (++CurrentEntryIndex < SBGPBox->GetNumberOfEntries())
+							{
+								CurrentGroupDescriptionIndex = SBGPBox->GetEntry(CurrentEntryIndex).GroupDescriptionIndex;
+							}
+							else
+							{
+								CurrentGroupDescriptionIndex = DefaultDescriptionIndex;
+							}
+						}
+					}
+					else
+					{
+						CurrentGroupDescriptionIndex = DefaultDescriptionIndex;
+					}
+				}
+			};
+			TMap<FBoxType, FSampleGroupIterator> GroupIterators;
+			bool bHaveRAPGroup = false;
 		};
 
 
@@ -5875,6 +5946,9 @@ private:
 			const FMP4BoxSTCO* STCOBox = nullptr;
 			const FMP4BoxSTSS* STSSBox = nullptr;
 			const FMP4BoxUDTA* UDTABox = nullptr;
+			// Groupings
+			TArray<const FMP4BoxSGPD*> SGPDBoxes;
+			TArray<const FMP4BoxSBGP*> SBGPBoxes;
 			// Encryption
 			TArray<const FMP4BoxPSSH*>	PSSHBoxesFromInit;
 			TArray<const FMP4BoxPSSH*>	PSSHBoxes;
@@ -6000,7 +6074,7 @@ private:
 		FMP4ParseInfo* ParsedData;
 
 		FTrackInfo* ParsedTrackInfo;
-		
+
 		IElectraCodecFactoryModule* DecoderFactoryModule;
 	};
 
@@ -6358,6 +6432,44 @@ private:
 			SampleDescriptionIndex = Track->TFHDBox->HasSampleDescriptionIndex() ? Track->TFHDBox->GetSampleDescriptionIndex() :
 				Track->TREXBox->GetDefaultSampleDescriptionIndex();
 
+			// Sample groups
+			for(int32 i=0; i<Track->SBGPBoxes.Num(); ++i)
+			{
+				switch(Track->SBGPBoxes[i]->GetGroupingType())
+				{
+					// Random access
+					case EGroupType::Grouping_rap:
+					{
+						// ISO/IEC 14496-15 disallows `rap ` boxes greater than version 0.
+						if (Track->SBGPBoxes[i]->GetBoxVersion() == 0)
+						{
+							// There needs to be a matching entry
+							const FMP4BoxSGPD* const* SGPDBox = Track->SGPDBoxes.FindByPredicate([](const FMP4BoxSGPD* In) { return In->GetGroupingType() == EGroupType::Grouping_rap; });
+							if (SGPDBox && (*SGPDBox)->GetNumberOfEntries() > 0)
+							{
+								// We do not care about the actual `VisualRandomAccessEntry` for now, only the fact that there is one.
+								FSampleGroupIterator& sgit = GroupIterators.Add(EGroupType::Grouping_rap);
+								sgit.GroupingType = EGroupType::Grouping_rap;
+								sgit.SBGPBox = Track->SBGPBoxes[i];
+								sgit.DefaultDescriptionIndex = (*SGPDBox)->GetDefaultSampleDescriptionIndex();
+								sgit.SetFirst();
+								// If the sample group says this sample is a RAP then clear the `sample_is_non_sync_sample` flag.
+								if (sgit.GetCurrentGroupDescriptionIndex() != 0)
+								{
+									SampleFlag &= ~0x10000U;
+								}
+								bHaveRAPGroup = true;
+							}
+						}
+						break;
+					}
+					default:
+					{
+						break;
+					}
+				}
+			}
+
 			RemainingSamplesInTRUN = TRUNBox->GetNumberOfSamples();
 			check(RemainingSamplesInTRUN != 0);
 			// Run could be empty!
@@ -6599,6 +6711,21 @@ private:
 				SamplePTS += (TRUNBox->GetSampleCompositionTimeOffsets().GetData())[SampleNumberInTRUN];
 			}
 
+			// Advance the grouping iterators
+			for(auto &grpit : GroupIterators)
+			{
+				grpit.Value.Advance();
+			}
+			// If we have a RAP grouping we need to look at it now
+			if (bHaveRAPGroup)
+			{
+				// If the sample group says this sample is a RAP then clear the `sample_is_non_sync_sample` flag.
+				if (GroupIterators[EGroupType::Grouping_rap].GetCurrentGroupDescriptionIndex() != 0)
+				{
+					SampleFlag &= ~0x10000U;
+				}
+			}
+
 			return UEMEDIA_ERROR_OK;
 		}
 		else
@@ -6832,6 +6959,8 @@ private:
 
 	bool FParserISO14496_12::FTrackIterator::IsSyncSample() const
 	{
+		// Note: Depending on the codec an I-frame might also be a sync sample, in which
+		//       case `sample_depends_on` (mask 0x03000000) would be set to 2.
 		return (SampleFlag & 0x10000) == 0;
 	}
 
@@ -7200,30 +7329,56 @@ private:
 						check(bOk);
 						if (bOk)
 						{
+							ElectraDecodersUtil::MPEG::FISO23008_2_seq_parameter_set_data TempSPS;
 							Track->CodecInformation.SetStreamType(EStreamType::Video);
 							Track->CodecInformation.SetCodec(FStreamCodecInformation::ECodec::H265);
 							Track->CodecInformation.SetCodecSpecificData(Track->CodecSpecificDataRAW);
 							Track->CodecInformation.SetStreamLanguageCode(Track->GetLanguage());
+							const ElectraDecodersUtil::MPEG::FISO23008_2_seq_parameter_set_data* sps = nullptr;
+							const TCHAR* HevType = TEXT("hvc1");
 							if (Track->CodecSpecificDataHEVC.GetNumberOfSPS() == 0)
 							{
-								return UEMEDIA_ERROR_FORMAT_ERROR;
+								TempSPS.general_profile_space = Track->CodecSpecificDataHEVC.GetGeneralProfileSpace();
+								TempSPS.general_tier_flag = Track->CodecSpecificDataHEVC.GetGeneralTierFlag();
+								TempSPS.general_profile_idc = Track->CodecSpecificDataHEVC.GetGeneralProfileIDC();
+								TempSPS.general_level_idc = Track->CodecSpecificDataHEVC.GetGeneralLevelIDC();
+								TempSPS.general_profile_compatibility_flag = Track->CodecSpecificDataHEVC.GetGeneralProfileCompatibilityFlags();
+								uint64 flags = Track->CodecSpecificDataHEVC.GetGeneralConstraintIndicatorFlags();
+								TempSPS.general_reserved_zero_44bits = (flags & 0xfffffffffffUL);
+								flags >>= 44;
+								TempSPS.general_frame_only_constraint_flag = flags & 1;
+								flags >>= 1;
+								TempSPS.general_non_packed_constraint_flag = flags & 1;
+								flags >>= 1;
+								TempSPS.general_interlaced_source_flag = flags & 1;
+								flags >>= 1;
+								TempSPS.general_progressive_source_flag = flags & 1;
+								TempSPS.pic_width_in_luma_samples = Track->TKHDBox->GetWidth();
+								TempSPS.pic_height_in_luma_samples = Track->TKHDBox->GetHeight();
+								TempSPS.bit_depth_luma_minus8 = Track->CodecSpecificDataHEVC.GetBitDepthLumaMinus8();
+								TempSPS.bit_depth_chroma_minus8 = Track->CodecSpecificDataHEVC.GetBitDepthChromaMinus8();
+								sps = &TempSPS;
+								HevType = TEXT("hev1");
+							}
+							else
+							{
+								sps = &Track->CodecSpecificDataHEVC.GetParsedSPS(0);
 							}
-							const ElectraDecodersUtil::MPEG::FISO23008_2_seq_parameter_set_data& sps = Track->CodecSpecificDataHEVC.GetParsedSPS(0);
 							int32 CropL, CropR, CropT, CropB;
-							sps.GetCrop(CropL, CropR, CropT, CropB);
-							Track->CodecInformation.SetResolution(FStreamCodecInformation::FResolution(sps.GetWidth() - CropL - CropR, sps.GetHeight() - CropT - CropB));
+							sps->GetCrop(CropL, CropR, CropT, CropB);
+							Track->CodecInformation.SetResolution(FStreamCodecInformation::FResolution(sps->GetWidth() - CropL - CropR, sps->GetHeight() - CropT - CropB));
 							Track->CodecInformation.SetCrop(FStreamCodecInformation::FCrop(CropL, CropT, CropR, CropB));
 							FStreamCodecInformation::FAspectRatio ar;
-							sps.GetAspect(ar.Width, ar.Height);
+							sps->GetAspect(ar.Width, ar.Height);
 							Track->CodecInformation.SetAspectRatio(ar);
-							Track->CodecInformation.SetFrameRate(sps.GetTiming().Denom ? FTimeFraction(sps.GetTiming().Num, sps.GetTiming().Denom) : FTimeFraction());
-							Track->CodecInformation.SetProfileSpace(sps.general_profile_space);
-							Track->CodecInformation.SetProfileTier(sps.general_tier_flag);
-							Track->CodecInformation.SetProfile(sps.general_profile_idc);
-							Track->CodecInformation.SetProfileLevel(sps.general_level_idc);
-							Track->CodecInformation.SetProfileConstraints(sps.GetConstraintFlags());
-							Track->CodecInformation.SetProfileCompatibilityFlags(sps.general_profile_compatibility_flag);
-							Track->CodecInformation.SetCodecSpecifierRFC6381(sps.GetRFC6381(TEXT("hvc1")));
+							Track->CodecInformation.SetFrameRate(sps->GetTiming().Denom ? FTimeFraction(sps->GetTiming().Num, sps->GetTiming().Denom) : FTimeFraction());
+							Track->CodecInformation.SetProfileSpace(sps->general_profile_space);
+							Track->CodecInformation.SetProfileTier(sps->general_tier_flag);
+							Track->CodecInformation.SetProfile(sps->general_profile_idc);
+							Track->CodecInformation.SetProfileLevel(sps->general_level_idc);
+							Track->CodecInformation.SetProfileConstraints(sps->GetConstraintFlags());
+							Track->CodecInformation.SetProfileCompatibilityFlags(sps->general_profile_compatibility_flag);
+							Track->CodecInformation.SetCodecSpecifierRFC6381(sps->GetRFC6381(HevType));
 							bGotVideoFormat = true;
 						}
 						else
@@ -8072,7 +8227,7 @@ private:
 				{
 					case FMP4Box::kBox_traf:
 					{
-						// FIXME: for encrypted content we need to also handle: sbgp, sgpd, subs, saiz, saio
+						// FIXME: for encrypted content we may need to also handle: `subs`
 
 						// Get the tfhd box.
 						const FMP4BoxTFHD* TFHDBox = static_cast<const FMP4BoxTFHD*>(Box->FindBox(FMP4Box::kBox_tfhd, 0));
@@ -8090,12 +8245,25 @@ private:
 									Track->MOOFBox = MOOFBox;		// need this to get to the absolute file position for default-base-is-moof addressing.
 									Track->TFHDBox = TFHDBox;
 									Track->TFDTBox = static_cast<const FMP4BoxTFDT*>(Box->FindBox(FMP4Box::kBox_tfdt, 0));
+									TArray<const FMP4Box*> ListOfBoxes;
 									// Locate all TRUN boxes.
-									TArray<const FMP4Box*> AllTRUNBoxes;
-									Box->GetAllBoxInstances(AllTRUNBoxes, FMP4Box::kBox_trun);
-									for(int32 nTruns = 0; nTruns < AllTRUNBoxes.Num(); ++nTruns)
+									Box->GetAllBoxInstances(ListOfBoxes, FMP4Box::kBox_trun);
+									for(int32 nBox=0; nBox<ListOfBoxes.Num(); ++nBox)
+									{
+										Track->TRUNBoxes.Add(static_cast<const FMP4BoxTRUN*>(ListOfBoxes[nBox]));
+									}
+									// Locate all sample grouping boxes. At present we do this only for fragmented streams.
+									ListOfBoxes.Empty();
+									Box->GetAllBoxInstances(ListOfBoxes, FMP4Box::kBox_sgpd);
+									for(int32 nBox=0; nBox<ListOfBoxes.Num(); ++nBox)
+									{
+										Track->SGPDBoxes.Add(static_cast<const FMP4BoxSGPD*>(ListOfBoxes[nBox]));
+									}
+									ListOfBoxes.Empty();
+									Box->GetAllBoxInstances(ListOfBoxes, FMP4Box::kBox_sbgp);
+									for(int32 nBox=0; nBox<ListOfBoxes.Num(); ++nBox)
 									{
-										Track->TRUNBoxes.Add(static_cast<const FMP4BoxTRUN*>(AllTRUNBoxes[nTruns]));
+										Track->SBGPBoxes.Add(static_cast<const FMP4BoxSBGP*>(ListOfBoxes[nBox]));
 									}
 
 									// Locate encryption related boxes. At present we do this only for fragmented streams.
diff --git a/Engine/Plugins/Media/ElectraUtil/Source/ElectraBase/Private/URLParser.cpp b/Engine/Plugins/Media/ElectraUtil/Source/ElectraBase/Private/URLParser.cpp
--- a/Engine/Plugins/Media/ElectraUtil/Source/ElectraBase/Private/URLParser.cpp
+++ b/Engine/Plugins/Media/ElectraUtil/Source/ElectraBase/Private/URLParser.cpp
@@ -124,6 +124,17 @@ namespace Electra
 					if (it && IsPathSeparator(*it))
 					{
 						++it;
+
+						// Check for a third '/' in Windows filenames like "file:///c:/autoexec.bat"
+						if (Scheme.Equals(URLComponents::SchemeFile) && IsPathSeparator(*it) && it.GetRemainingLength() > 3)
+						{
+							const TCHAR* Next = it.GetRemainder();
+							if (Next[2] == TCHAR(':') && Next[3] == TCHAR('/'))
+							{
+								++it;
+							}
+						}
+
 						// Parse out the authority and if successful continue parsing out the path.
 						if (!ParseAuthority(it))
 						{
diff --git a/Engine/Plugins/Media/HAPDecoderElectra/Source/HAPDecoderElectra/Private/HAPDecoder/ElectraMediaHAPDecoder.cpp b/Engine/Plugins/Media/HAPDecoderElectra/Source/HAPDecoderElectra/Private/HAPDecoder/ElectraMediaHAPDecoder.cpp
--- a/Engine/Plugins/Media/HAPDecoderElectra/Source/HAPDecoderElectra/Private/HAPDecoder/ElectraMediaHAPDecoder.cpp
+++ b/Engine/Plugins/Media/HAPDecoderElectra/Source/HAPDecoderElectra/Private/HAPDecoder/ElectraMediaHAPDecoder.cpp
@@ -54,6 +54,11 @@ public:
 		return UserValue;
 	}
 
+	EOutputType GetOutputType() const
+	{
+		return EOutputType::Output;
+	}
+
 	int32 GetWidth() const override
 	{
 		return Width - Crop.Left - Crop.Right;
@@ -705,7 +710,7 @@ IElectraDecoder::EDecoderError FVideoDecoderHAPElectra::DecodeAccessUnit(const F
 						// note: we share one task sync for all buffers -> this will lead any external task to be executed after ALL tasks we trigger have been done (aka: after the "last" buffer is ready)
 						NewOutput->Buffer_TaskSync = TaskSync;
 					}
-					else	
+					else
 					{
 						uint8* DX12BufferAddr;
 						HRESULT Res = BufferInfo.GPUBuffer.Resource->Map(0, nullptr, (void**)&DX12BufferAddr);
