ElectraPlayer: URL fragment parameter #t,e creates a temporary playback range only now instead of a locked range

#rnx

diff --git a/Engine/Plugins/Media/ElectraPlayer/Source/ElectraPlayerRuntime/Private/Runtime/Player/AdaptiveStreamingPlayer.cpp b/Engine/Plugins/Media/ElectraPlayer/Source/ElectraPlayerRuntime/Private/Runtime/Player/AdaptiveStreamingPlayer.cpp
--- a/Engine/Plugins/Media/ElectraPlayer/Source/ElectraPlayerRuntime/Private/Runtime/Player/AdaptiveStreamingPlayer.cpp
+++ b/Engine/Plugins/Media/ElectraPlayer/Source/ElectraPlayerRuntime/Private/Runtime/Player/AdaptiveStreamingPlayer.cpp
@@ -801,7 +801,7 @@ void FAdaptiveStreamingPlayer::GetLoopState(FLoopState& OutLoopState) const
 
 	// Check if there is a pending seek. A seek resets the loop count, but the seek itself may not be executing
 	// and the count not being reset yet.
-	
+
 	/* Not needed since we are only accessing a bool
 	FScopeLock lock(&SeekVars.Lock);
 	*/
@@ -1799,6 +1799,13 @@ void FAdaptiveStreamingPlayer::HandleSeeking()
 			SeekVars.Lock.Lock();
 
 			SeekVars.bForScrubbing = bIsForScrubbing;
+
+			// When seeking (other than the initial playstart seek) any default end time that
+			// was set via the URL #t,e parameter no longer applies.
+			if (Manifest.IsValid())
+			{
+				Manifest->ClearDefaultEndTime();
+			}
 		}
 		SeekVars.ActiveRequest = SeekParam;
 		lock.Unlock();
@@ -2058,7 +2065,7 @@ double FAdaptiveStreamingPlayer::GetMinBufferTimeBeforePlayback()
 										LastBufferingState == EPlayerState::eState_Rebuffering ? PlayerConfig.RebufferMinTimeAvailBeforePlayback : PlayerConfig.InitialBufferMinTimeAvailBeforePlayback;
 
 	FTimeValue mbtAbr = StreamSelector.IsValid() ? StreamSelector->GetMinBufferTimeForPlayback(LastBufferingState == EPlayerState::eState_Seeking ? IAdaptiveStreamSelector::EMinBufferType::Seeking
-																							 : LastBufferingState == EPlayerState::eState_Rebuffering ? IAdaptiveStreamSelector::EMinBufferType::Rebuffering 
+																							 : LastBufferingState == EPlayerState::eState_Rebuffering ? IAdaptiveStreamSelector::EMinBufferType::Rebuffering
 																							 :	IAdaptiveStreamSelector::EMinBufferType::Initial, Manifest->GetMinBufferTime()) : FTimeValue();
 	kMinBufferBeforePlayback = mbtAbr.IsValid() ? mbtAbr.GetAsSeconds() : kMinBufferBeforePlayback;
 	return kMinBufferBeforePlayback;
@@ -2858,7 +2865,7 @@ void FAdaptiveStreamingPlayer::InternalHandlePendingStartRequest(const FTimeValu
 								case IManifest::FResult::EType::PastEOS:
 								{
 									PostLog(Facility::EFacility::Player, IInfoLog::ELevel::Info, FString::Printf(TEXT("miss(%d,a=%d,l=%d) for %lld/%d in range %lld/%d - %lld/%d"), (int32)Result.GetType(), StartAt.Options.bFrameAccuracy, PendingStartRequest->StartType == FPendingStartRequest::EStartType::LoopPoint,
-											(long long int) StartAt.Time.GetAsHNS(), StartAt.Time.IsValid(), 
+											(long long int) StartAt.Time.GetAsHNS(), StartAt.Time.IsValid(),
 											(long long int) StartAt.Options.PlaybackRange.Start.GetAsHNS(),
 											StartAt.Options.PlaybackRange.Start.IsValid(),
 											(long long int) StartAt.Options.PlaybackRange.End.GetAsHNS(),
@@ -4031,6 +4038,15 @@ void FAdaptiveStreamingPlayer::CheckForStreamEnd()
 		{
 			// First check if there is an end time set at which we need to stop.
 			FTimeValue EndAtTime = PlaybackState.GetPlaybackEndAtTime();
+			FTimeValue InitialEndTime = Manifest.IsValid() ? Manifest->GetDefaultEndTime() : FTimeValue();
+			if (EndAtTime.IsValid() && InitialEndTime.IsValid())
+			{
+				EndAtTime = EndAtTime < InitialEndTime ? EndAtTime : InitialEndTime;
+			}
+			else if (InitialEndTime.IsValid())
+			{
+				EndAtTime = InitialEndTime;
+			}
 			if (EndAtTime.IsValid())
 			{
 				if (PlaybackState.GetPlayPosition() >= EndAtTime)
@@ -4301,6 +4317,14 @@ void FAdaptiveStreamingPlayer::InternalSetPlaybackEnded()
 	{
 		TimelineRange.End = RangeEnd;
 	}
+	if (Manifest.IsValid() && Manifest->GetDefaultEndTime().IsValid())
+	{
+		if (Manifest->GetDefaultEndTime() < TimelineRange.End)
+		{
+			TimelineRange.End = Manifest->GetDefaultEndTime();
+		}
+	}
+
 	PlaybackState.SetPlayPosition(TimelineRange.End);
 
 	if (LastBufferingState == EPlayerState::eState_Seeking)
diff --git a/Engine/Plugins/Media/ElectraPlayer/Source/ElectraPlayerRuntime/Private/Runtime/Player/AdaptiveStreamingPlayerManifest.cpp b/Engine/Plugins/Media/ElectraPlayer/Source/ElectraPlayerRuntime/Private/Runtime/Player/AdaptiveStreamingPlayerManifest.cpp
--- a/Engine/Plugins/Media/ElectraPlayer/Source/ElectraPlayerRuntime/Private/Runtime/Player/AdaptiveStreamingPlayerManifest.cpp
+++ b/Engine/Plugins/Media/ElectraPlayer/Source/ElectraPlayerRuntime/Private/Runtime/Player/AdaptiveStreamingPlayerManifest.cpp
@@ -285,8 +285,6 @@ bool FAdaptiveStreamingPlayer::SelectManifest()
 		if (ManifestType != EMediaFormatType::Unknown)
 		{
 			TArray<FTimespan> SeekablePositions;
-			FPlaybackRange PlaybackRange;
-			FTimeRange RestrictedPlaybackRange;
 			TSharedPtrTS<IManifest> NewPresentation = ManifestReader->GetManifest();
 			check(NewPresentation.IsValid());
 
@@ -295,23 +293,31 @@ bool FAdaptiveStreamingPlayer::SelectManifest()
 			PlaybackState.SetSeekablePositions(SeekablePositions);
 			PlaybackState.SetTimelineRange(NewPresentation->GetTotalTimeRange());
 			PlaybackState.SetDuration(NewPresentation->GetDuration());
-			
-			// Check for playback range restriction. This is currently assumed to come from URL fragment parameters
-			// like example.mp4#t=10.8,18.4
-			// These do not override any user defined range values!
-			RestrictedPlaybackRange = NewPresentation->GetPlaybackRange();
-			PlaybackState.GetPlayRange(PlaybackRange);
-			if (RestrictedPlaybackRange.Start.IsValid() && !PlaybackRange.Start.Get(FTimeValue()).IsValid())
-			{
-				PlaybackRange.Start = RestrictedPlaybackRange.Start;
-			}
-			if (RestrictedPlaybackRange.End.IsValid() && !PlaybackRange.End.Get(FTimeValue()).IsValid())
-			{
-				PlaybackRange.End = RestrictedPlaybackRange.End;
-			}
-			// Set the new range, but say it did not change. This is the first start and we do not need to issue a seek.
-			PlaybackState.SetPlayRange(PlaybackRange);
-			PlaybackState.ActivateNewPlayRange(nullptr);
+
+			#if 0
+			// Removed for now. This can be put back in if necessary, but not through the #t URL fragment parameter
+			// which is not really intended to lock down a playback range.
+			// We could add a #r[,e] parameter instead to create a range if needed.
+
+				// Check for playback range restriction. This is currently assumed to come from URL fragment parameters
+				// like example.mp4#t=10.8,18.4
+				// These do not override any user defined range values!
+				FPlaybackRange PlaybackRange;
+				FTimeRange RestrictedPlaybackRange;
+				RestrictedPlaybackRange = NewPresentation->GetPlaybackRange();
+				PlaybackState.GetPlayRange(PlaybackRange);
+				if (RestrictedPlaybackRange.Start.IsValid() && !PlaybackRange.Start.Get(FTimeValue()).IsValid())
+				{
+					PlaybackRange.Start = RestrictedPlaybackRange.Start;
+				}
+				if (RestrictedPlaybackRange.End.IsValid() && !PlaybackRange.End.Get(FTimeValue()).IsValid())
+				{
+					PlaybackRange.End = RestrictedPlaybackRange.End;
+				}
+				// Set the new range, but say it did not change. This is the first start and we do not need to issue a seek.
+				PlaybackState.SetPlayRange(PlaybackRange);
+				PlaybackState.ActivateNewPlayRange(nullptr);
+			#endif
 
 			TArray<FTrackMetadata> VideoTrackMetadata;
 			TArray<FTrackMetadata> AudioTrackMetadata;
diff --git a/Engine/Plugins/Media/ElectraPlayer/Source/ElectraPlayerRuntime/Private/Runtime/Player/DASH/ManifestBuilderDASH.cpp b/Engine/Plugins/Media/ElectraPlayer/Source/ElectraPlayerRuntime/Private/Runtime/Player/DASH/ManifestBuilderDASH.cpp
--- a/Engine/Plugins/Media/ElectraPlayer/Source/ElectraPlayerRuntime/Private/Runtime/Player/DASH/ManifestBuilderDASH.cpp
+++ b/Engine/Plugins/Media/ElectraPlayer/Source/ElectraPlayerRuntime/Private/Runtime/Player/DASH/ManifestBuilderDASH.cpp
@@ -97,7 +97,7 @@ namespace
 
 	const TCHAR* const DASHRole = TEXT("urn:mpeg:dash:role:2011");
 
-	const TCHAR* const SupportedEssentialProperties[] = 
+	const TCHAR* const SupportedEssentialProperties[] =
 	{
 		TEXT("urn:mpeg:dash:urlparam:2014"),
 		TEXT("urn:mpeg:dash:urlparam:2016"),
@@ -733,7 +733,7 @@ class FManifestBuilderDASH : public IManifestBuilderDASH
 public:
 	FManifestBuilderDASH(IPlayerSessionServices* InPlayerSessionServices);
 	virtual ~FManifestBuilderDASH() = default;
-	
+
 	virtual FErrorDetail BuildFromMPD(TSharedPtrTS<FManifestDASHInternal>& OutMPD, TCHAR* InOutMPDXML, const FString& EffectiveURL, const FString& ETag) override;
 
 private:
@@ -1039,7 +1039,7 @@ FErrorDetail FManifestDASHInternal::Build(IPlayerSessionServices* InPlayerSessio
 	FErrorDetail Error;
 
 	MPDRoot = InMPDRoot;
-	
+
 	// What type of presentation is this?
 	// We do not validate the MPD@type here. Anything not 'static' is handled as 'dynamic'.
 	PresentationType = MPDRoot->GetType().Equals(TEXT("static")) ? EPresentationType::Static : EPresentationType::Dynamic;
@@ -1203,7 +1203,7 @@ FErrorDetail FManifestDASHInternal::BuildAfterInitialRemoteElementDownload()
 					return CreateErrorAndLog(PlayerSessionServices, FString::Printf(TEXT("Dynamic presentations with the last Period@duration missing and no MPD@minimumUpdateTime need to have MPD@mediaPresentationDuration set!")), ERRCODE_DASH_MPD_BUILDER_MEDIAPRESENTATIONDURATION_NEEDED);
 				}
 			}
-			
+
 			// Check for potential period overlap or gaps.
 			if (Periods.Num())
 			{
@@ -1653,7 +1653,7 @@ void FManifestDASHInternal::PreparePeriodAdaptationSets(TSharedPtrTS<FPeriod> Pe
 					{
 						Representation->CodecInfo.SetSamplingRate((int32) MPDAdaptationSet->GetAudioSamplingRate()[0]);
 					}
-					
+
 					// Get the audio channel configurations from both Representation and AdaptationSet.
 					TArray<TSharedPtrTS<FDashMPD_DescriptorType>> AudioChannelConfigurations(MPDRepresentation->GetAudioChannelConfigurations());
 					AudioChannelConfigurations.Append(MPDAdaptationSet->GetAudioChannelConfigurations());
@@ -1689,7 +1689,7 @@ void FManifestDASHInternal::PreparePeriodAdaptationSets(TSharedPtrTS<FPeriod> Pe
 								   tag:dolby.com,2014:dash:DolbyDigitalPlusExtensionType:2014
 								   tag:dolby.com,2014:dash:complexityIndexTypeA:2014"
 								and more.
-							
+
 								There is also information in the Accessibility descriptors like
 								   urn:tva:metadata:cs:AudioPurposeCS:2007
 							*/
@@ -2133,7 +2133,7 @@ void FManifestDASHInternal::SendEventsFromAllPeriodEventStreams(TSharedPtrTS<FPe
 			const TSharedPtrTS<FDashMPD_EventStreamType>& EvS = MPDEventStreams[nEvS];
 			uint32 Timescale = EvS->GetTimescale().GetWithDefault(1);
 			int64 PTO = (int64) EvS->GetPresentationTimeOffset().GetWithDefault(0);
-			
+
 			const TArray<TSharedPtrTS<FDashMPD_EventType>>& Events = EvS->GetEvents();
 			for(int32 i=0; i<Events.Num(); ++i)
 			{
@@ -2192,9 +2192,9 @@ FErrorDetail FManifestDASHInternal::ResolveInitialRemoteElementRequest(TSharedPt
 	PendingRemoteElementLoadRequests.Remove(RequestResponse);
 	// Likewise for the pending element list.
 	RemoteElementsToResolve.Remove(RequestResponse->XLinkElement);
-	
+
 	TSharedPtrTS<IDashMPDElement> XLinkElement = RequestResponse->XLinkElement.Pin();
-	if (XLinkElement.IsValid())	
+	if (XLinkElement.IsValid())
 	{
 		int64 LastResolveID = XLinkElement->GetXLink().LastResolveID;
 		int64 NewResolveID = LastResolveID;
@@ -2579,7 +2579,7 @@ FTimeRange FManifestDASHInternal::GetSeekableTimeRange() const
 			}
 			/*
 				The start must be covered by a Period. See: https://dashif-documents.azurewebsites.net/Guidelines-TimingModel/master/Guidelines-TimingModel.html#timing-timeshift
-					"Clients SHALL NOT allow seeking into regions of the time shift buffer that are not covered by periods, 
+					"Clients SHALL NOT allow seeking into regions of the time shift buffer that are not covered by periods,
 					regardless of whether such regions are before or after the periods described by the MPD."
 			*/
 			FTimeValue PST = ast + GetPeriods()[0]->GetStart();
@@ -2688,13 +2688,14 @@ void FManifestDASHInternal::PrepareDefaultStartTime()
 {
 	FTimeRange PlaybackRange = GetPlayTimesFromURI();
 	DefaultStartTime = PlaybackRange.Start;
+	DefaultEndTime = PlaybackRange.End;
 }
 
 
 FTimeRange FManifestDASHInternal::GetPlayTimesFromURI() const
 {
 	FTimeRange FromTo;
-	
+
 	// We are interested in the 't' and 'period' fragment values here.
 	FString Time, PeriodID;
 	for(int32 i=0,iMax=URLFragmentComponents.Num(); i<iMax; ++i)
@@ -2844,6 +2845,15 @@ void FManifestDASHInternal::ClearDefaultStartTime()
 	DefaultStartTime.SetToInvalid();
 }
 
+FTimeValue FManifestDASHInternal::GetDefaultEndTime() const
+{
+	return DefaultEndTime;
+}
+
+void FManifestDASHInternal::ClearDefaultEndTime()
+{
+	DefaultEndTime.SetToInvalid();
+}
 
 
 } // namespace Electra
diff --git a/Engine/Plugins/Media/ElectraPlayer/Source/ElectraPlayerRuntime/Private/Runtime/Player/DASH/ManifestDASH.cpp b/Engine/Plugins/Media/ElectraPlayer/Source/ElectraPlayerRuntime/Private/Runtime/Player/DASH/ManifestDASH.cpp
--- a/Engine/Plugins/Media/ElectraPlayer/Source/ElectraPlayerRuntime/Private/Runtime/Player/DASH/ManifestDASH.cpp
+++ b/Engine/Plugins/Media/ElectraPlayer/Source/ElectraPlayerRuntime/Private/Runtime/Player/DASH/ManifestDASH.cpp
@@ -227,7 +227,7 @@ private:
 		FManifestDASHInternal::FSegmentInformation InitSegmentInfo;
 		bool bRequested = false;
 		TSharedPtrTS<FMPDLoadRequestDASH> LoadRequest;
-		
+
 		class FAcceptBoxes : public IParserISO14496_12::IBoxCallback
 		{
 		public:
@@ -344,7 +344,7 @@ IManifest::EType FManifestDASH::GetPresentationType() const
 }
 
 TSharedPtrTS<const FLowLatencyDescriptor> FManifestDASH::GetLowLatencyDescriptor() const
-{ 
+{
 	TSharedPtrTS<FManifestDASHInternal> Manifest(CurrentManifest);
 	return Manifest.IsValid() ? Manifest->GetLowLatencyDescriptor() : nullptr;
 }
@@ -403,6 +403,21 @@ void FManifestDASH::ClearDefaultStartTime()
 	}
 }
 
+FTimeValue FManifestDASH::GetDefaultEndTime() const
+{
+	TSharedPtrTS<FManifestDASHInternal> Manifest(CurrentManifest);
+	return Manifest.IsValid() ? Manifest->GetDefaultEndTime() : FTimeValue();
+}
+
+void FManifestDASH::ClearDefaultEndTime()
+{
+	TSharedPtrTS<FManifestDASHInternal> Manifest(CurrentManifest);
+	if (Manifest.IsValid())
+	{
+		Manifest->ClearDefaultEndTime();
+	}
+}
+
 FTimeValue FManifestDASH::GetMinBufferTime() const
 {
 	TSharedPtrTS<FManifestDASHInternal> Manifest(CurrentManifest);
@@ -1748,7 +1763,7 @@ IManifest::FResult FDASHPlayPeriod::GetRetrySegment(TSharedPtrTS<IStreamSegment>
 		OutSegment = NewRequest;
 		return IManifest::FResult(IManifest::FResult::EType::Found);
 	}
-	
+
 	// Pass the download stats bWaitingForRemoteRetryElement to convey if the retry segment needs to wait for a remote element,
 	// which is either some xlink or an index segment.
 	FStreamSegmentRequestDASH* CurrentRequest = const_cast<FStreamSegmentRequestDASH*>(static_cast<const FStreamSegmentRequestDASH*>(InCurrentSegment.Get()));
diff --git a/Engine/Plugins/Media/ElectraPlayer/Source/ElectraPlayerRuntime/Private/Runtime/Player/DASH/ManifestDASH.h b/Engine/Plugins/Media/ElectraPlayer/Source/ElectraPlayerRuntime/Private/Runtime/Player/DASH/ManifestDASH.h
--- a/Engine/Plugins/Media/ElectraPlayer/Source/ElectraPlayerRuntime/Private/Runtime/Player/DASH/ManifestDASH.h
+++ b/Engine/Plugins/Media/ElectraPlayer/Source/ElectraPlayerRuntime/Private/Runtime/Player/DASH/ManifestDASH.h
@@ -28,6 +28,8 @@ public:
 	FTimeValue GetDuration() const override;
 	FTimeValue GetDefaultStartTime() const override;
 	void ClearDefaultStartTime() override;
+	FTimeValue GetDefaultEndTime() const override;
+	void ClearDefaultEndTime() override;
 	FTimeValue GetMinBufferTime() const override;
 	FTimeValue GetDesiredLiveLatency() const override;
 	TRangeSet<double> GetPossiblePlaybackRates(EPlayRateType InForType) const override;
diff --git a/Engine/Plugins/Media/ElectraPlayer/Source/ElectraPlayerRuntime/Private/Runtime/Player/DASH/PlaylistReaderDASH_Internal.h b/Engine/Plugins/Media/ElectraPlayer/Source/ElectraPlayerRuntime/Private/Runtime/Player/DASH/PlaylistReaderDASH_Internal.h
--- a/Engine/Plugins/Media/ElectraPlayer/Source/ElectraPlayerRuntime/Private/Runtime/Player/DASH/PlaylistReaderDASH_Internal.h
+++ b/Engine/Plugins/Media/ElectraPlayer/Source/ElectraPlayerRuntime/Private/Runtime/Player/DASH/PlaylistReaderDASH_Internal.h
@@ -895,6 +895,8 @@ public:
 	void PrepareDefaultStartTime();
 	FTimeValue GetDefaultStartTime() const;
 	void ClearDefaultStartTime();
+	FTimeValue GetDefaultEndTime() const;
+	void ClearDefaultEndTime();
 	FTimeRange GetPlayTimesFromURI() const;
 	FTimeValue GetDesiredLiveLatency() const;
 
@@ -953,6 +955,7 @@ private:
 	mutable FTimeRange TotalTimeRange;
 	mutable FTimeRange SeekableTimeRange;
 	FTimeValue DefaultStartTime;
+	FTimeValue DefaultEndTime;
 	mutable bool bWarnedAboutTooSmallSuggestedPresentationDelay = false;
 
 	bool bDidInjectUTCTimingElements = false;
diff --git a/Engine/Plugins/Media/ElectraPlayer/Source/ElectraPlayerRuntime/Private/Runtime/Player/HLS/ManifestHLS.cpp b/Engine/Plugins/Media/ElectraPlayer/Source/ElectraPlayerRuntime/Private/Runtime/Player/HLS/ManifestHLS.cpp
--- a/Engine/Plugins/Media/ElectraPlayer/Source/ElectraPlayerRuntime/Private/Runtime/Player/HLS/ManifestHLS.cpp
+++ b/Engine/Plugins/Media/ElectraPlayer/Source/ElectraPlayerRuntime/Private/Runtime/Player/HLS/ManifestHLS.cpp
@@ -196,6 +196,13 @@ FTimeValue FManifestHLS::GetDefaultStartTime() const
 void FManifestHLS::ClearDefaultStartTime()
 {
 }
+FTimeValue FManifestHLS::GetDefaultEndTime() const
+{
+	return FTimeValue::GetInvalid();
+}
+void FManifestHLS::ClearDefaultEndTime()
+{
+}
 
 
 
diff --git a/Engine/Plugins/Media/ElectraPlayer/Source/ElectraPlayerRuntime/Private/Runtime/Player/HLS/ManifestHLS.h b/Engine/Plugins/Media/ElectraPlayer/Source/ElectraPlayerRuntime/Private/Runtime/Player/HLS/ManifestHLS.h
--- a/Engine/Plugins/Media/ElectraPlayer/Source/ElectraPlayerRuntime/Private/Runtime/Player/HLS/ManifestHLS.h
+++ b/Engine/Plugins/Media/ElectraPlayer/Source/ElectraPlayerRuntime/Private/Runtime/Player/HLS/ManifestHLS.h
@@ -26,6 +26,8 @@ public:
 	FTimeValue GetDuration() const override;
 	FTimeValue GetDefaultStartTime() const override;
 	void ClearDefaultStartTime() override;
+	FTimeValue GetDefaultEndTime() const override;
+	void ClearDefaultEndTime() override;
 	FTimeValue GetMinBufferTime() const override;
 	FTimeValue GetDesiredLiveLatency() const override;
 	TRangeSet<double> GetPossiblePlaybackRates(EPlayRateType InForType) const override;
diff --git a/Engine/Plugins/Media/ElectraPlayer/Source/ElectraPlayerRuntime/Private/Runtime/Player/Manifest.h b/Engine/Plugins/Media/ElectraPlayer/Source/ElectraPlayerRuntime/Private/Runtime/Player/Manifest.h
--- a/Engine/Plugins/Media/ElectraPlayer/Source/ElectraPlayerRuntime/Private/Runtime/Player/Manifest.h
+++ b/Engine/Plugins/Media/ElectraPlayer/Source/ElectraPlayerRuntime/Private/Runtime/Player/Manifest.h
@@ -308,6 +308,17 @@ namespace Electra
 		 */
 		virtual void ClearDefaultStartTime() = 0;
 
+		/**
+		 * Returns the playback end time as defined by the presentation itself.
+		 * If the presentation has no preferred end time an invalid value is returned.
+		 */
+		virtual FTimeValue GetDefaultEndTime() const = 0;
+
+		/**
+		 * Clears the internal default end time so it will not be used again.
+		 */
+		virtual void ClearDefaultEndTime() = 0;
+
 		//! Returns track metadata. For period based presentations the streams can be different per period in which case the metadata of the first period is returned.
 		virtual void GetTrackMetadata(TArray<FTrackMetadata>& OutMetadata, EStreamType StreamType) const = 0;
 
diff --git a/Engine/Plugins/Media/ElectraPlayer/Source/ElectraPlayerRuntime/Private/Runtime/Player/mkv/ManifestMKV.cpp b/Engine/Plugins/Media/ElectraPlayer/Source/ElectraPlayerRuntime/Private/Runtime/Player/mkv/ManifestMKV.cpp
--- a/Engine/Plugins/Media/ElectraPlayer/Source/ElectraPlayerRuntime/Private/Runtime/Player/mkv/ManifestMKV.cpp
+++ b/Engine/Plugins/Media/ElectraPlayer/Source/ElectraPlayerRuntime/Private/Runtime/Player/mkv/ManifestMKV.cpp
@@ -63,6 +63,7 @@ FErrorDetail FManifestMKVInternal::Build(TSharedPtrTS<IParserMKV> MKVParser, con
 	FErrorDetail Error = MediaAsset->Build(PlayerSessionServices, MKVParser, URL);
 	FTimeRange PlaybackRange = GetPlaybackRange();
 	DefaultStartTime = PlaybackRange.Start;
+	DefaultEndTime = PlaybackRange.End;
 	return Error;
 }
 
diff --git a/Engine/Plugins/Media/ElectraPlayer/Source/ElectraPlayerRuntime/Private/Runtime/Player/mkv/ManifestMKV.h b/Engine/Plugins/Media/ElectraPlayer/Source/ElectraPlayerRuntime/Private/Runtime/Player/mkv/ManifestMKV.h
--- a/Engine/Plugins/Media/ElectraPlayer/Source/ElectraPlayerRuntime/Private/Runtime/Player/mkv/ManifestMKV.h
+++ b/Engine/Plugins/Media/ElectraPlayer/Source/ElectraPlayerRuntime/Private/Runtime/Player/mkv/ManifestMKV.h
@@ -52,6 +52,10 @@ public:
 	{ return DefaultStartTime; }
 	void ClearDefaultStartTime() override
 	{ DefaultStartTime.SetToInvalid(); }
+	FTimeValue GetDefaultEndTime() const override
+	{ return DefaultEndTime; }
+	void ClearDefaultEndTime() override
+	{ DefaultEndTime.SetToInvalid(); }
 	void GetTrackMetadata(TArray<FTrackMetadata>& OutMetadata, EStreamType StreamType) const override;
 	FTimeValue GetMinBufferTime() const override;
 	FTimeValue GetDesiredLiveLatency() const override
@@ -322,6 +326,7 @@ public:
 	HTTP::FConnectionInfo ConnectionInfo;
 	TArray<FURL_RFC3986::FQueryParam> URLFragmentComponents;
 	FTimeValue DefaultStartTime;
+	FTimeValue DefaultEndTime;
 };
 
 } // namespace Electra
diff --git a/Engine/Plugins/Media/ElectraPlayer/Source/ElectraPlayerRuntime/Private/Runtime/Player/mp4/ManifestMP4.cpp b/Engine/Plugins/Media/ElectraPlayer/Source/ElectraPlayerRuntime/Private/Runtime/Player/mp4/ManifestMP4.cpp
--- a/Engine/Plugins/Media/ElectraPlayer/Source/ElectraPlayerRuntime/Private/Runtime/Player/mp4/ManifestMP4.cpp
+++ b/Engine/Plugins/Media/ElectraPlayer/Source/ElectraPlayerRuntime/Private/Runtime/Player/mp4/ManifestMP4.cpp
@@ -67,6 +67,7 @@ FErrorDetail FManifestMP4Internal::Build(TSharedPtrTS<IParserISO14496_12> MP4Par
 	FErrorDetail Result = MediaAsset->Build(PlayerSessionServices, MP4Parser, URL);
 	FTimeRange PlaybackRange = GetPlaybackRange();
 	DefaultStartTime = PlaybackRange.Start;
+	DefaultEndTime = PlaybackRange.End;
 	return Result;
 }
 
diff --git a/Engine/Plugins/Media/ElectraPlayer/Source/ElectraPlayerRuntime/Private/Runtime/Player/mp4/ManifestMP4.h b/Engine/Plugins/Media/ElectraPlayer/Source/ElectraPlayerRuntime/Private/Runtime/Player/mp4/ManifestMP4.h
--- a/Engine/Plugins/Media/ElectraPlayer/Source/ElectraPlayerRuntime/Private/Runtime/Player/mp4/ManifestMP4.h
+++ b/Engine/Plugins/Media/ElectraPlayer/Source/ElectraPlayerRuntime/Private/Runtime/Player/mp4/ManifestMP4.h
@@ -53,6 +53,10 @@ public:
 	{ return DefaultStartTime; }
 	void ClearDefaultStartTime() override
 	{ DefaultStartTime.SetToInvalid(); }
+	FTimeValue GetDefaultEndTime() const override
+	{ return DefaultEndTime; }
+	void ClearDefaultEndTime() override
+	{ DefaultEndTime.SetToInvalid(); }
 	void GetTrackMetadata(TArray<FTrackMetadata>& OutMetadata, EStreamType StreamType) const override;
 	FTimeValue GetMinBufferTime() const override;
 	FTimeValue GetDesiredLiveLatency() const override
@@ -331,6 +335,7 @@ public:
 	// The MPD URL fragment components
 	TArray<FURL_RFC3986::FQueryParam>	URLFragmentComponents;
 	FTimeValue							DefaultStartTime;
+	FTimeValue							DefaultEndTime;
 };
 
 
