ElectraPlayer: Fixed seeking to the very last frame (or beyond) of an .mp4 file

#rnx

diff --git a/Engine/Plugins/Media/ElectraPlayer/Source/ElectraPlayerRuntime/Private/Runtime/Player/mp4/ManifestMP4.cpp b/Engine/Plugins/Media/ElectraPlayer/Source/ElectraPlayerRuntime/Private/Runtime/Player/mp4/ManifestMP4.cpp
--- a/Engine/Plugins/Media/ElectraPlayer/Source/ElectraPlayerRuntime/Private/Runtime/Player/mp4/ManifestMP4.cpp
+++ b/Engine/Plugins/Media/ElectraPlayer/Source/ElectraPlayerRuntime/Private/Runtime/Player/mp4/ManifestMP4.cpp
@@ -235,9 +235,15 @@ IStreamReader* FManifestMP4Internal::CreateStreamReaderHandler()
  */
 IManifest::FResult FManifestMP4Internal::FindPlayPeriod(TSharedPtrTS<IPlayPeriod>& OutPlayPeriod, const FPlayStartPosition& StartPosition, ESearchType SearchType)
 {
-	// FIXME: We could however check if the start position falls into the duration of the asset. Not sure why it wouldn't or why we would want to do that.
-	OutPlayPeriod = MakeSharedTS<FPlayPeriodMP4>(MediaAsset);
-	return IManifest::FResult(IManifest::FResult::EType::Found);
+	if (MediaAsset.IsValid() && StartPosition.Time.IsValid() && StartPosition.Time < MediaAsset->GetDuration())
+	{
+		OutPlayPeriod = MakeSharedTS<FPlayPeriodMP4>(MediaAsset);
+		return IManifest::FResult(IManifest::FResult::EType::Found);
+	}
+	else
+	{
+		return IManifest::FResult(IManifest::FResult::EType::PastEOS);
+	}
 }
 
 IManifest::FResult FManifestMP4Internal::FindNextPlayPeriod(TSharedPtrTS<IPlayPeriod>& OutPlayPeriod, TSharedPtrTS<const IStreamSegment> CurrentSegment)
@@ -927,10 +933,10 @@ void FManifestMP4Internal::FTimelineAssetMP4::UpdatePlayRangeEndInfo(const FTime
 							TrackLocalTime = PlayRangeEndInfo.Time.GetAsTimebase(TrkIt->GetTimescale());
 						}
 						// Look at the DTS, not the PTS since a frame with a larger PTS may still be referenced
-						// by earlier frames (B frame reordering).
+						// by earlier frames (In case of video: B frame reordering).
 						if (TrkIt->GetDTS() >= TrackLocalTime)
 						{
-							TrackEndBytePos = TrkIt->GetSampleFileOffset();
+							TrackEndBytePos = TrkIt->GetSampleFileOffset() + TrkIt->GetSampleSize();
 							if (TrackEndBytePos > PlayRangeEndInfo.TotalFileEndOffset)
 							{
 								PlayRangeEndInfo.TotalFileEndOffset = TrackEndBytePos;
