ElectraPlayer: Ensuring the first buffering of data is tracked as the initial buffering, even if a seek was issued right away. We need this for correctly dispatching the tracks changed event and for the initial buffering analytics.

#rnx

diff --git a/Engine/Plugins/Media/ElectraPlayer/Source/ElectraPlayerRuntime/Private/ElectraPlayer.cpp b/Engine/Plugins/Media/ElectraPlayer/Source/ElectraPlayerRuntime/Private/ElectraPlayer.cpp
--- a/Engine/Plugins/Media/ElectraPlayer/Source/ElectraPlayerRuntime/Private/ElectraPlayer.cpp
+++ b/Engine/Plugins/Media/ElectraPlayer/Source/ElectraPlayerRuntime/Private/ElectraPlayer.cpp
@@ -213,6 +213,7 @@ void FElectraPlayer::ClearToDefaultState()
 	bSubtitleTrackIndexDirty = true;
 	bInitialSeekPerformed = false;
 	bDiscardOutputUntilCleanStart = false;
+	bIsFirstBuffering = true;
 	LastPresentedFrameDimension = FIntPoint::ZeroValue;
 	CurrentStreamMetadata.Reset();
 	CurrentlyActiveVideoStreamFormat.Reset();
@@ -2234,6 +2235,13 @@ void FElectraPlayer::HandlePlayerEventBufferingStart(Electra::Metrics::EBufferin
 {
 	PlayerState.Status = PlayerState.Status | EPlayerStatus::Buffering;
 
+	// In case a seek was performed right away the reason would be `Seeking`, but we want to
+	// track it as `Initial` for statistics reasons and to make sure we won't miss sending `TracksChanged`.
+	if (bIsFirstBuffering)
+	{
+		BufferingReason = Electra::Metrics::EBufferingReason::Initial;
+	}
+
 	// Send TracksChanged on the initial buffering event. Prior to that we do not know where in the stream
 	// playback will begin and what tracks are available there.
 	if (BufferingReason == Electra::Metrics::EBufferingReason::Initial)
@@ -2283,6 +2291,13 @@ void FElectraPlayer::HandlePlayerEventBufferingEnd(Electra::Metrics::EBufferingR
 	//       To avoid this we keep the state as buffering until the pre-rolling phase has also completed.
 	//PlayerState.Status = PlayerState.Status & ~EPlayerStatus::Buffering;
 
+	// In case a seek was performed right away the reason would be `Seeking`, but we want to track it as `Initial` for statistics.
+	if (bIsFirstBuffering)
+	{
+		BufferingReason = Electra::Metrics::EBufferingReason::Initial;
+		bIsFirstBuffering = false;
+	}
+
 	// Update statistics
 	FScopeLock Lock(&StatisticsLock);
 	double BufferingDuration = FPlatformTime::Seconds() - Statistics.TimeAtBufferingBegin;
diff --git a/Engine/Plugins/Media/ElectraPlayer/Source/ElectraPlayerRuntime/Private/ElectraPlayer.h b/Engine/Plugins/Media/ElectraPlayer/Source/ElectraPlayerRuntime/Private/ElectraPlayer.h
--- a/Engine/Plugins/Media/ElectraPlayer/Source/ElectraPlayerRuntime/Private/ElectraPlayer.h
+++ b/Engine/Plugins/Media/ElectraPlayer/Source/ElectraPlayerRuntime/Private/ElectraPlayer.h
@@ -418,6 +418,7 @@ private:
 
 	bool											bInitialSeekPerformed;
 	bool											bDiscardOutputUntilCleanStart;
+	bool											bIsFirstBuffering;
 
 	FPlaybackRange									CurrentPlaybackRange;
 	TOptional<bool>									bFrameAccurateSeeking;
