ElectraPlayer: Passing seek time in DASH segment requests for frame accuracy as converting time into media local time and back loses precision.

#rnx

diff --git a/Engine/Plugins/Media/ElectraPlayer/Source/ElectraPlayerRuntime/Private/Runtime/Player/DASH/ManifestDASH.cpp b/Engine/Plugins/Media/ElectraPlayer/Source/ElectraPlayerRuntime/Private/Runtime/Player/DASH/ManifestDASH.cpp
--- a/Engine/Plugins/Media/ElectraPlayer/Source/ElectraPlayerRuntime/Private/Runtime/Player/DASH/ManifestDASH.cpp
+++ b/Engine/Plugins/Media/ElectraPlayer/Source/ElectraPlayerRuntime/Private/Runtime/Player/DASH/ManifestDASH.cpp
@@ -1426,6 +1426,10 @@ IManifest::FResult FDASHPlayPeriod::GetStartingSegment(TSharedPtrTS<IStreamSegme
 				}
 				SegmentRequest->Segment = MoveTemp(SegmentInfo);
 				SegmentRequest->TimestampSequenceIndex = InSequenceState.GetSequenceIndex();
+				if (bFrameAccurateSearch)
+				{
+					SegmentRequest->FrameAccurateStartTime = StartPosition.Time;
+				}
 
 				// The start segment request needs to be able to return a valid first PTS which is what the player sets
 				// the playback position to. If not valid yet update it with the current stream values.
diff --git a/Engine/Plugins/Media/ElectraPlayer/Source/ElectraPlayerRuntime/Private/Runtime/Player/DASH/StreamReaderDASH.cpp b/Engine/Plugins/Media/ElectraPlayer/Source/ElectraPlayerRuntime/Private/Runtime/Player/DASH/StreamReaderDASH.cpp
--- a/Engine/Plugins/Media/ElectraPlayer/Source/ElectraPlayerRuntime/Private/Runtime/Player/DASH/StreamReaderDASH.cpp
+++ b/Engine/Plugins/Media/ElectraPlayer/Source/ElectraPlayerRuntime/Private/Runtime/Player/DASH/StreamReaderDASH.cpp
@@ -196,7 +196,7 @@ UEMediaError FStreamReaderDASH::Create(IPlayerSessionServices* InPlayerSessionSe
 			StreamHandlers[i].bRunOnThreadPool = true;
 		}
 #endif
-		StreamHandlers[i].ThreadSetName(i==0 ? "ElectraPlayer::DASH Video" : 
+		StreamHandlers[i].ThreadSetName(i==0 ? "ElectraPlayer::DASH Video" :
 										i==1 ? "ElectraPlayer::DASH Audio" :
 											   "ElectraPlayer::DASH Subtitle");
 	}
@@ -510,7 +510,7 @@ FErrorDetail FStreamReaderDASH::FStreamHandler::LoadInitSegment(TSharedPtrTS<FMP
 	{
 		return FErrorDetail();
 	}
-	
+
 	// Set up download stats to be sent to the stream selector.
 	const HTTP::FConnectionInfo* ci = OutLoadRequest->GetConnectionInfo();
 	ds.StatsID = FMediaInterlockedIncrement(UniqueDownloadID);
@@ -762,7 +762,7 @@ FErrorDetail FStreamReaderDASH::FStreamHandler::RetrieveSideloadedFile(TSharedPt
 	{
 		return FErrorDetail();
 	}
-	
+
 	// Set up download stats to be sent to the stream selector.
 	const HTTP::FConnectionInfo* ci = LoadReq->GetConnectionInfo();
 	Metrics::FSegmentDownloadStats ds;
@@ -989,7 +989,7 @@ void FStreamReaderDASH::FStreamHandler::HandleRequestMP4()
 		A segment that is already at EOS is not meant to be loaded as it does not exist and there
 		would not be another segment following it either. They are meant to indicate to the player
 		that a stream has ended and will not be delivering any more data.
-		
+
 		NOTE:
 		  We had to handle this request up to detecting the use of inband event streams in order to
 		  signal that this stream has now ended and will not be receiving any further inband events!
@@ -1344,8 +1344,15 @@ void FStreamReaderDASH::FStreamHandler::HandleRequestMP4()
 
 									if (Request->Segment.bFrameAccuracyRequired)
 									{
-										AccessUnit->EarliestPTS.SetFromND(MediaLocalFirstAUTime - PTO, TrackTimescale);
-										AccessUnit->EarliestPTS += TimeOffset;
+										if (Request->FrameAccurateStartTime.IsValid())
+										{
+											AccessUnit->EarliestPTS = Request->FrameAccurateStartTime;
+										}
+										else
+										{
+											AccessUnit->EarliestPTS.SetFromND(MediaLocalFirstAUTime - PTO, TrackTimescale);
+											AccessUnit->EarliestPTS += TimeOffset;
+										}
 										AccessUnit->EarliestPTS.SetSequenceIndex(Request->TimestampSequenceIndex);
 									}
 									if (MediaLocalLastAUTime != TNumericLimits<int64>::Max())
@@ -1412,7 +1419,7 @@ void FStreamReaderDASH::FStreamHandler::HandleRequestMP4()
 										// If we need to decrypt we have to wait for the decrypter to become ready.
 										if (bIsSampleEncrypted && Decrypter.IsValid())
 										{
-											while(!bTerminate && !HasReadBeenAborted() && 
+											while(!bTerminate && !HasReadBeenAborted() &&
 												  (Decrypter->GetState() == ElectraCDM::ECDMState::WaitingForKey || Decrypter->GetState() == ElectraCDM::ECDMState::Idle))
 											{
 												FMediaRunnable::SleepMilliseconds(100);
@@ -1854,7 +1861,7 @@ void FStreamReaderDASH::FStreamHandler::HandleRequestMKV()
 		A segment that is already at EOS is not meant to be loaded as it does not exist and there
 		would not be another segment following it either. They are meant to indicate to the player
 		that a stream has ended and will not be delivering any more data.
-		
+
 		NOTE:
 		  We had to handle this request up to detecting the use of inband event streams in order to
 		  signal that this stream has now ended and will not be receiving any further inband events!
diff --git a/Engine/Plugins/Media/ElectraPlayer/Source/ElectraPlayerRuntime/Private/Runtime/Player/DASH/StreamReaderDASH.h b/Engine/Plugins/Media/ElectraPlayer/Source/ElectraPlayerRuntime/Private/Runtime/Player/DASH/StreamReaderDASH.h
--- a/Engine/Plugins/Media/ElectraPlayer/Source/ElectraPlayerRuntime/Private/Runtime/Player/DASH/StreamReaderDASH.h
+++ b/Engine/Plugins/Media/ElectraPlayer/Source/ElectraPlayerRuntime/Private/Runtime/Player/DASH/StreamReaderDASH.h
@@ -57,6 +57,7 @@ public:
 	FTimeValue												AdditionalAdjustmentTime = FTimeValue::GetZero();	//!< Sum of any other time corrections
 	bool													bInsertFillerData = false;							//!< true to insert empty access units into the buffer instead of reading actual data.
 	int64													TimestampSequenceIndex = 0;							//!< Sequence index to set in all timestamp values of the decoded access unit.
+	FTimeValue												FrameAccurateStartTime;								//!< If set, the start time as was requested in a Seek() (not in media local time)
 
 	// UTC wallclock times during which this segment can be fetched;
 	FTimeValue												ASAST;
@@ -96,7 +97,7 @@ public:
 	//! Adds a request to read from a stream
 	virtual EAddResult AddRequest(uint32 CurrentPlaybackSequenceID, TSharedPtrTS<IStreamSegment> Request) override;
 
-	//! Cancels any ongoing requests of the given stream type. Silent cancellation will not notify OnFragmentClose() or OnFragmentReachedEOS(). 
+	//! Cancels any ongoing requests of the given stream type. Silent cancellation will not notify OnFragmentClose() or OnFragmentReachedEOS().
 	virtual void CancelRequest(EStreamType StreamType, bool bSilent) override;
 
 	//! Cancels all pending requests.
@@ -309,7 +310,7 @@ private:
 		// Methods from IParserISO14496_12::IBoxCallback
 		IParserISO14496_12::IBoxCallback::EParseContinuation OnFoundBox(IParserISO14496_12::FBoxType Box, int64 BoxSizeInBytes, int64 FileDataOffset, int64 BoxDataOffset) override;
 		IParserISO14496_12::IBoxCallback::EParseContinuation OnEndOfBox(IParserISO14496_12::FBoxType Box, int64 BoxSizeInBytes, int64 FileDataOffset, int64 BoxDataOffset) override;
-	
+
 		// Methods from IParserMKV::IReader
 		int64 MKVReadData(void* InDestinationBuffer, int64 InNumBytesToRead, int64 InFromOffset) override;
 		int64 MKVGetCurrentFileOffset() const override;
@@ -317,7 +318,7 @@ private:
 		bool MKVHasReadBeenAborted() const override;
 	};
 
-	FStreamHandler						StreamHandlers[3];		// 0 = video, 1 = audio, 2 = subtitle 
+	FStreamHandler						StreamHandlers[3];		// 0 = video, 1 = audio, 2 = subtitle
 	IPlayerSessionServices*				PlayerSessionService = nullptr;
 	bool								bIsStarted = false;
 	FErrorDetail						ErrorDetail;
