Fix unreachable code warnings

diff --git a/Engine/Plugins/Experimental/MotoSynth/Source/MotoSynth/Private/SynthComponentMoto.cpp b/Engine/Plugins/Experimental/MotoSynth/Source/MotoSynth/Private/SynthComponentMoto.cpp
--- a/Engine/Plugins/Experimental/MotoSynth/Source/MotoSynth/Private/SynthComponentMoto.cpp
+++ b/Engine/Plugins/Experimental/MotoSynth/Source/MotoSynth/Private/SynthComponentMoto.cpp
@@ -151,6 +151,4 @@ ISoundGeneratorPtr USynthComponentMoto::CreateSoundGenerator(const FSoundGenerat
  		UE_LOG(LogMotoSynth, Verbose, TEXT("Can't play moto synth without a preset UMotoSynthPreset object and both acceleration source and deceleration source set."));
  		return ISoundGeneratorPtr(new FSoundGeneratorNull());
  	}
- 
-	return nullptr;
 }
diff --git a/Engine/Plugins/Media/ElectraPlayer/Source/ElectraPlayerRuntime/Private/Runtime/Decoder/BitstreamProcessors/VideoDecoderInputBitstreamProcessorH264.cpp b/Engine/Plugins/Media/ElectraPlayer/Source/ElectraPlayerRuntime/Private/Runtime/Decoder/BitstreamProcessors/VideoDecoderInputBitstreamProcessorH264.cpp
--- a/Engine/Plugins/Media/ElectraPlayer/Source/ElectraPlayerRuntime/Private/Runtime/Decoder/BitstreamProcessors/VideoDecoderInputBitstreamProcessorH264.cpp
+++ b/Engine/Plugins/Media/ElectraPlayer/Source/ElectraPlayerRuntime/Private/Runtime/Decoder/BitstreamProcessors/VideoDecoderInputBitstreamProcessorH264.cpp
@@ -169,28 +169,27 @@ void FVideoDecoderInputBitstreamProcessorH264::SetPropertiesOnOutput(TSharedPtr<
 	TSharedPtr<FCodecSpecificMessages> Msg = StaticCastSharedPtr<FCodecSpecificMessages>(InFromBSI.CodecSpecificMessages);
 
 	const TArray<ElectraDecodersUtil::MPEG::FISO14496_10_seq_parameter_set_data>& SPSList = Msg.IsValid() ? Msg->SPSs : SPSs;
-	for(int32 nMsg=0; nMsg<SPSList.Num(); ++nMsg)
+
+	// We only interact with the first SPS.
+	if (SPSList.Num() > 0)
 	{
 		// Set the bit depth and the colorimetry.
 		uint8 colour_primaries=2, transfer_characteristics=2, matrix_coeffs=2;
 		uint8 video_full_range_flag=0, video_format=5;
-		if (SPSList[nMsg].colour_description_present_flag)
+		if (SPSList[0].colour_description_present_flag)
 		{
-			colour_primaries = SPSList[nMsg].colour_primaries;
-			transfer_characteristics = SPSList[nMsg].transfer_characteristics;
-			matrix_coeffs = SPSList[nMsg].matrix_coefficients;
+			colour_primaries = SPSList[0].colour_primaries;
+			transfer_characteristics = SPSList[0].transfer_characteristics;
+			matrix_coeffs = SPSList[0].matrix_coefficients;
 		}
-		if (SPSList[nMsg].video_signal_type_present_flag)
+		if (SPSList[0].video_signal_type_present_flag)
 		{
-			video_full_range_flag = SPSList[nMsg].video_full_range_flag;
-			video_format = SPSList[nMsg].video_format;
+			video_full_range_flag = SPSList[0].video_full_range_flag;
+			video_format = SPSList[0].video_format;
 		}
 
 		Colorimetry.Update(colour_primaries, transfer_characteristics, matrix_coeffs, video_full_range_flag, video_format);
 		Colorimetry.UpdateParamDict(*InOutProperties);
-
-		// We stop after the first SPS.
-		break;
 	}
 }
 
diff --git a/Engine/Plugins/Media/ElectraPlayer/Source/ElectraPlayerRuntime/Private/Runtime/Player/DASH/ManifestParserDASH.cpp b/Engine/Plugins/Media/ElectraPlayer/Source/ElectraPlayerRuntime/Private/Runtime/Player/DASH/ManifestParserDASH.cpp
--- a/Engine/Plugins/Media/ElectraPlayer/Source/ElectraPlayerRuntime/Private/Runtime/Player/DASH/ManifestParserDASH.cpp
+++ b/Engine/Plugins/Media/ElectraPlayer/Source/ElectraPlayerRuntime/Private/Runtime/Player/DASH/ManifestParserDASH.cpp
@@ -2866,7 +2866,6 @@ bool FDashMPD_SegmentTimelineType::ProcessElement(FManifestParserDASH* Builder,
 	{
 		return Super::ProcessElement(Builder, ElementName, ElementData, XmlFileLineNumber);
 	}
-	return true;
 }
 
 bool FDashMPD_SegmentTimelineType::ProcessAttribute(FManifestParserDASH* Builder, const TCHAR* AttributeName, const TCHAR* AttributeValue)
diff --git a/Engine/Plugins/Media/ElectraPlayer/Source/ElectraPlayerRuntime/Private/Runtime/Player/HLS/ManifestHLS.cpp b/Engine/Plugins/Media/ElectraPlayer/Source/ElectraPlayerRuntime/Private/Runtime/Player/HLS/ManifestHLS.cpp
--- a/Engine/Plugins/Media/ElectraPlayer/Source/ElectraPlayerRuntime/Private/Runtime/Player/HLS/ManifestHLS.cpp
+++ b/Engine/Plugins/Media/ElectraPlayer/Source/ElectraPlayerRuntime/Private/Runtime/Player/HLS/ManifestHLS.cpp
@@ -864,7 +864,6 @@ IManifest::FResult FPlayPeriodHLS::FindSegment(TSharedPtrTS<FStreamSegmentReques
 		OutRequest = Req;
 		return IManifest::FResult(IManifest::FResult::EType::PastEOS);
 	}
-	return IManifest::FResult(IManifest::FResult::EType::NotFound);
 }
 
 
@@ -1199,8 +1198,6 @@ IManifest::FResult FPlayPeriodHLS::GetNextOrRetrySegment(TSharedPtrTS<IStreamSeg
 		}
 		return Result;
 	}
-
-	return IManifest::FResult(IManifest::FResult::EType::PastEOS);
 }
 
 /**
diff --git a/Engine/Plugins/Media/ElectraPlayer/Source/ElectraPlayerRuntime/Private/Runtime/Player/mp4/ManifestMP4.h b/Engine/Plugins/Media/ElectraPlayer/Source/ElectraPlayerRuntime/Private/Runtime/Player/mp4/ManifestMP4.h
--- a/Engine/Plugins/Media/ElectraPlayer/Source/ElectraPlayerRuntime/Private/Runtime/Player/mp4/ManifestMP4.h
+++ b/Engine/Plugins/Media/ElectraPlayer/Source/ElectraPlayerRuntime/Private/Runtime/Player/mp4/ManifestMP4.h
@@ -185,8 +185,6 @@ public:
 					return AdaptationSetIndex < AudioAdaptationSets.Num() ? AudioAdaptationSets[AdaptationSetIndex] : TSharedPtrTS<IPlaybackAssetAdaptationSet>();
 				case EStreamType::Subtitle:
 					return AdaptationSetIndex < SubtitleAdaptationSets.Num() ? SubtitleAdaptationSets[AdaptationSetIndex] : TSharedPtrTS<IPlaybackAssetAdaptationSet>();
-				default:
-					return 0;
 			}
 			return TSharedPtrTS<IPlaybackAssetAdaptationSet>();
 		}
diff --git a/Engine/Plugins/Media/ElectraSubtitles/Source/ElectraSubtitles/Private/ttml/TTMLSubtitleList.cpp b/Engine/Plugins/Media/ElectraSubtitles/Source/ElectraSubtitles/Private/ttml/TTMLSubtitleList.cpp
--- a/Engine/Plugins/Media/ElectraSubtitles/Source/ElectraSubtitles/Private/ttml/TTMLSubtitleList.cpp
+++ b/Engine/Plugins/Media/ElectraSubtitles/Source/ElectraSubtitles/Private/ttml/TTMLSubtitleList.cpp
@@ -335,7 +335,6 @@ bool FTTMLSubtitleList::UpdateActiveRange(const Electra::FTimeRange& InRange)
 			return false;
 		}
 	}
-	return false;
 }
 
 void FTTMLSubtitleList::GetActiveSubtitles(TArray<FActiveSubtitle>& OutSubtitles) const
diff --git a/Engine/Plugins/Online/OnlineServicesEOSGS/Source/Private/Online/OnlineServicesEOSGS.cpp b/Engine/Plugins/Online/OnlineServicesEOSGS/Source/Private/Online/OnlineServicesEOSGS.cpp
--- a/Engine/Plugins/Online/OnlineServicesEOSGS/Source/Private/Online/OnlineServicesEOSGS.cpp
+++ b/Engine/Plugins/Online/OnlineServicesEOSGS/Source/Private/Online/OnlineServicesEOSGS.cpp
@@ -177,9 +177,6 @@ TOnlineResult<FGetResolvedConnectString> FOnlineServicesEOSGS::GetResolvedConnec
 		{
 			return TOnlineResult<FGetResolvedConnectString>(Result.GetErrorValue());
 		}
-
-		// No matching session
-		return TOnlineResult<FGetResolvedConnectString>(Errors::NotFound());
 	}
 
 	// No valid lobby or session id set
diff --git a/Engine/Plugins/Online/OnlineServicesOSSAdapter/Source/Private/Online/SocialOSSAdapter.cpp b/Engine/Plugins/Online/OnlineServicesOSSAdapter/Source/Private/Online/SocialOSSAdapter.cpp
--- a/Engine/Plugins/Online/OnlineServicesOSSAdapter/Source/Private/Online/SocialOSSAdapter.cpp
+++ b/Engine/Plugins/Online/OnlineServicesOSSAdapter/Source/Private/Online/SocialOSSAdapter.cpp
@@ -131,8 +131,6 @@ TOnlineResult<FGetFriends> FSocialOSSAdapter::GetFriends(FGetFriends::Params&& P
 		{
 			return TOnlineResult<FGetFriends>(Errors::InvalidUser());
 		}
-
-		return TOnlineResult<FGetFriends>(Result);
 	}
 	else
 	{
diff --git a/Engine/Source/Runtime/Core/Private/Windows/WindowsPlatformMisc.cpp b/Engine/Source/Runtime/Core/Private/Windows/WindowsPlatformMisc.cpp
--- a/Engine/Source/Runtime/Core/Private/Windows/WindowsPlatformMisc.cpp
+++ b/Engine/Source/Runtime/Core/Private/Windows/WindowsPlatformMisc.cpp
@@ -1342,7 +1342,7 @@ void FWindowsPlatformMisc::RequestExitWithStatus(bool Force, uint8 ReturnCode, c
 	PGO_WriteFile();
 	TerminateProcess(GetCurrentProcess(), 0);
 	return;
-#endif
+#else
 
 	RequestEngineExit(TEXT("Win RequestExit"));
 
@@ -1373,6 +1373,7 @@ void FWindowsPlatformMisc::RequestExitWithStatus(bool Force, uint8 ReturnCode, c
 		// Tell the platform specific code we want to exit cleanly from the main loop.
 		PostQuitMessage(ReturnCode);
 	}
+#endif
 }
 
 const TCHAR* FWindowsPlatformMisc::GetSystemErrorMessage(TCHAR* OutBuffer, int32 BufferCount, int32 Error)
