Fix Unreachable code warnings

diff --git a/Engine/Plugins/Media/ElectraPlayer/Source/ElectraPlayerRuntime/Private/Runtime/Decoder/BitstreamProcessors/VideoDecoderInputBitstreamProcessorH265.cpp b/Engine/Plugins/Media/ElectraPlayer/Source/ElectraPlayerRuntime/Private/Runtime/Decoder/BitstreamProcessors/VideoDecoderInputBitstreamProcessorH265.cpp
--- a/Engine/Plugins/Media/ElectraPlayer/Source/ElectraPlayerRuntime/Private/Runtime/Decoder/BitstreamProcessors/VideoDecoderInputBitstreamProcessorH265.cpp
+++ b/Engine/Plugins/Media/ElectraPlayer/Source/ElectraPlayerRuntime/Private/Runtime/Decoder/BitstreamProcessors/VideoDecoderInputBitstreamProcessorH265.cpp
@@ -193,29 +193,27 @@ void FVideoDecoderInputBitstreamProcessorH265::SetPropertiesOnOutput(TSharedPtr<
 
 	TSharedPtr<FCodecSpecificMessages> Msg = StaticCastSharedPtr<FCodecSpecificMessages>(InFromBSI.CodecSpecificMessages);
 	uint8 num_bits = 8;
-	for(int32 nMsg=0; nMsg<Msg->SPSs.Num(); ++nMsg)
+	// We only process the first SPS.
+	if (Msg->SPSs.Num() > 0)
 	{
 		// Set the bit depth and the colorimetry.
 		uint8 colour_primaries=2, transfer_characteristics=2, matrix_coeffs=2;
 		uint8 video_full_range_flag=0, video_format=5;
 		num_bits = Msg->SPSs[0].bit_depth_luma_minus8 + 8;
-		if (Msg->SPSs[nMsg].colour_description_present_flag)
+		if (Msg->SPSs[0].colour_description_present_flag)
 		{
-			colour_primaries = Msg->SPSs[nMsg].colour_primaries;
-			transfer_characteristics = Msg->SPSs[nMsg].transfer_characteristics;
-			matrix_coeffs = Msg->SPSs[nMsg].matrix_coeffs;
+			colour_primaries = Msg->SPSs[0].colour_primaries;
+			transfer_characteristics = Msg->SPSs[0].transfer_characteristics;
+			matrix_coeffs = Msg->SPSs[0].matrix_coeffs;
 		}
-		if (Msg->SPSs[nMsg].video_signal_type_present_flag)
+		if (Msg->SPSs[0].video_signal_type_present_flag)
 		{
-			video_full_range_flag = Msg->SPSs[nMsg].video_full_range_flag;
-			video_format = Msg->SPSs[nMsg].video_format;
+			video_full_range_flag = Msg->SPSs[0].video_full_range_flag;
+			video_format = Msg->SPSs[0].video_format;
 		}
 
 		Colorimetry.Update(colour_primaries, transfer_characteristics, matrix_coeffs, video_full_range_flag, video_format);
 		Colorimetry.UpdateParamDict(*InOutProperties);
-
-		// We stop after the first SPS.
-		break;
 	}
 
 	// Extract HDR metadata from SEI messages
