ElectraPlayer: Fixing an issue where the last access units of an mp4 stream may be incorrectly skipped over when a segment end time has been specified.

#rnx

diff --git a/Engine/Plugins/Media/ElectraPlayer/Source/ElectraPlayerRuntime/Private/Runtime/Player/mp4/StreamReaderMP4.cpp b/Engine/Plugins/Media/ElectraPlayer/Source/ElectraPlayerRuntime/Private/Runtime/Player/mp4/StreamReaderMP4.cpp
--- a/Engine/Plugins/Media/ElectraPlayer/Source/ElectraPlayerRuntime/Private/Runtime/Player/mp4/StreamReaderMP4.cpp
+++ b/Engine/Plugins/Media/ElectraPlayer/Source/ElectraPlayerRuntime/Private/Runtime/Player/mp4/StreamReaderMP4.cpp
@@ -580,7 +580,7 @@ void FStreamReaderMP4::HandleRequest()
 							// retry on the _next_ AU and not this one again!
 							Request->CurrentIteratorBytePos = SampleFileOffset + SampleSize;
 						}
-						else if (AccessUnit->PTS >= AccessUnit->LatestPTS)
+						else if (AccessUnit->DTS >= AccessUnit->LatestPTS && AccessUnit->PTS >= AccessUnit->LatestPTS)
 						{
 							// Tag the last one and send it off, but stop doing so for the remainder of the segment.
 							// Note: we continue reading this segment all the way to the end on purpose in case there are further 'emsg' boxes.
