ElectraPlayer: Adding the URL fragment parameter #r=s,e to optionally set a range to which playback will be constrained. The #t=s,e parameter only sets an initial range for first playstart and is cleared when a seek is performed. The new #r range is sticky and if set, seeking can only be performed in that range. A #t range can be used alongside an #r range, provided the #t range is inside the #r

#rnx

diff --git a/Engine/Plugins/Media/ElectraPlayer/Source/ElectraPlayerRuntime/Private/Runtime/Player/AdaptiveStreamingPlayerManifest.cpp b/Engine/Plugins/Media/ElectraPlayer/Source/ElectraPlayerRuntime/Private/Runtime/Player/AdaptiveStreamingPlayerManifest.cpp
--- a/Engine/Plugins/Media/ElectraPlayer/Source/ElectraPlayerRuntime/Private/Runtime/Player/AdaptiveStreamingPlayerManifest.cpp
+++ b/Engine/Plugins/Media/ElectraPlayer/Source/ElectraPlayerRuntime/Private/Runtime/Player/AdaptiveStreamingPlayerManifest.cpp
@@ -294,30 +294,24 @@ bool FAdaptiveStreamingPlayer::SelectManifest()
 			PlaybackState.SetTimelineRange(NewPresentation->GetTotalTimeRange());
 			PlaybackState.SetDuration(NewPresentation->GetDuration());
 
-			#if 0
-			// Removed for now. This can be put back in if necessary, but not through the #t URL fragment parameter
-			// which is not really intended to lock down a playback range.
-			// We could add a #r[,e] parameter instead to create a range if needed.
-
-				// Check for playback range restriction. This is currently assumed to come from URL fragment parameters
-				// like example.mp4#t=10.8,18.4
-				// These do not override any user defined range values!
-				FPlaybackRange PlaybackRange;
-				FTimeRange RestrictedPlaybackRange;
-				RestrictedPlaybackRange = NewPresentation->GetPlaybackRange();
-				PlaybackState.GetPlayRange(PlaybackRange);
-				if (RestrictedPlaybackRange.Start.IsValid() && !PlaybackRange.Start.Get(FTimeValue()).IsValid())
-				{
-					PlaybackRange.Start = RestrictedPlaybackRange.Start;
-				}
-				if (RestrictedPlaybackRange.End.IsValid() && !PlaybackRange.End.Get(FTimeValue()).IsValid())
-				{
-					PlaybackRange.End = RestrictedPlaybackRange.End;
-				}
-				// Set the new range, but say it did not change. This is the first start and we do not need to issue a seek.
-				PlaybackState.SetPlayRange(PlaybackRange);
-				PlaybackState.ActivateNewPlayRange(nullptr);
-			#endif
+			// Check for playback range restriction. This is currently assumed to come from URL fragment parameters
+			// like example.mp4#r=10.8,18.4
+			// These do not override any user defined range values!
+			FPlaybackRange PlaybackRange;
+			FTimeRange RestrictedPlaybackRange;
+			RestrictedPlaybackRange = NewPresentation->GetPlaybackRange(IManifest::EPlaybackRangeType::LockedPlaybackRange);
+			PlaybackState.GetPlayRange(PlaybackRange);
+			if (RestrictedPlaybackRange.Start.IsValid() && !PlaybackRange.Start.Get(FTimeValue()).IsValid())
+			{
+				PlaybackRange.Start = RestrictedPlaybackRange.Start;
+			}
+			if (RestrictedPlaybackRange.End.IsValid() && !PlaybackRange.End.Get(FTimeValue()).IsValid())
+			{
+				PlaybackRange.End = RestrictedPlaybackRange.End;
+			}
+			// Set the new range, but say it did not change. This is the first start and we do not need to issue a seek.
+			PlaybackState.SetPlayRange(PlaybackRange);
+			PlaybackState.ActivateNewPlayRange(nullptr);
 
 			TArray<FTrackMetadata> VideoTrackMetadata;
 			TArray<FTrackMetadata> AudioTrackMetadata;
diff --git a/Engine/Plugins/Media/ElectraPlayer/Source/ElectraPlayerRuntime/Private/Runtime/Player/DASH/ManifestBuilderDASH.cpp b/Engine/Plugins/Media/ElectraPlayer/Source/ElectraPlayerRuntime/Private/Runtime/Player/DASH/ManifestBuilderDASH.cpp
--- a/Engine/Plugins/Media/ElectraPlayer/Source/ElectraPlayerRuntime/Private/Runtime/Player/DASH/ManifestBuilderDASH.cpp
+++ b/Engine/Plugins/Media/ElectraPlayer/Source/ElectraPlayerRuntime/Private/Runtime/Player/DASH/ManifestBuilderDASH.cpp
@@ -2686,27 +2686,37 @@ void FManifestDASHInternal::EndPresentationAt(const FTimeValue& EndsAt, const FS
 
 void FManifestDASHInternal::PrepareDefaultStartTime()
 {
-	FTimeRange PlaybackRange = GetPlayTimesFromURI();
+	FTimeRange PlaybackRange = GetPlayTimesFromURI(IManifest::EPlaybackRangeType::TemporaryPlaystartRange);
 	DefaultStartTime = PlaybackRange.Start;
 	DefaultEndTime = PlaybackRange.End;
 }
 
 
-FTimeRange FManifestDASHInternal::GetPlayTimesFromURI() const
+FTimeRange FManifestDASHInternal::GetPlayTimesFromURI(IManifest::EPlaybackRangeType InRangeType) const
 {
 	FTimeRange FromTo;
 
-	// We are interested in the 't' and 'period' fragment values here.
+	// We are interested in the 't', 'r' and 'period' fragment values here.
 	FString Time, PeriodID;
 	for(int32 i=0,iMax=URLFragmentComponents.Num(); i<iMax; ++i)
 	{
-		if (URLFragmentComponents[i].Name.Equals(TEXT("t")))
+		if (InRangeType == IManifest::EPlaybackRangeType::TemporaryPlaystartRange)
 		{
-			Time = URLFragmentComponents[i].Value;
+			if (URLFragmentComponents[i].Name.Equals(TEXT("t")))
+			{
+				Time = URLFragmentComponents[i].Value;
+			}
+			else if (URLFragmentComponents[i].Name.Equals(TEXT("period")))
+			{
+				PeriodID = URLFragmentComponents[i].Value;
+			}
 		}
-		else if (URLFragmentComponents[i].Name.Equals(TEXT("period")))
+		else if (InRangeType == IManifest::EPlaybackRangeType::LockedPlaybackRange)
 		{
-			PeriodID = URLFragmentComponents[i].Value;
+			if (URLFragmentComponents[i].Name.Equals(TEXT("r")))
+			{
+				Time = URLFragmentComponents[i].Value;
+			}
 		}
 	}
 	if (Time.IsEmpty() && PeriodID.IsEmpty())
diff --git a/Engine/Plugins/Media/ElectraPlayer/Source/ElectraPlayerRuntime/Private/Runtime/Player/DASH/ManifestDASH.cpp b/Engine/Plugins/Media/ElectraPlayer/Source/ElectraPlayerRuntime/Private/Runtime/Player/DASH/ManifestDASH.cpp
--- a/Engine/Plugins/Media/ElectraPlayer/Source/ElectraPlayerRuntime/Private/Runtime/Player/DASH/ManifestDASH.cpp
+++ b/Engine/Plugins/Media/ElectraPlayer/Source/ElectraPlayerRuntime/Private/Runtime/Player/DASH/ManifestDASH.cpp
@@ -367,10 +367,10 @@ FTimeRange FManifestDASH::GetSeekableTimeRange() const
 	return Manifest.IsValid() ? Manifest->GetSeekableTimeRange() : FTimeRange();
 }
 
-FTimeRange FManifestDASH::GetPlaybackRange() const
+FTimeRange FManifestDASH::GetPlaybackRange(EPlaybackRangeType InRangeType) const
 {
 	TSharedPtrTS<FManifestDASHInternal> Manifest(CurrentManifest);
-	return Manifest.IsValid() ? Manifest->GetPlayTimesFromURI() : FTimeRange();
+	return Manifest.IsValid() ? Manifest->GetPlayTimesFromURI(InRangeType) : FTimeRange();
 }
 
 void FManifestDASH::GetSeekablePositions(TArray<FTimespan>& OutPositions) const
diff --git a/Engine/Plugins/Media/ElectraPlayer/Source/ElectraPlayerRuntime/Private/Runtime/Player/DASH/ManifestDASH.h b/Engine/Plugins/Media/ElectraPlayer/Source/ElectraPlayerRuntime/Private/Runtime/Player/DASH/ManifestDASH.h
--- a/Engine/Plugins/Media/ElectraPlayer/Source/ElectraPlayerRuntime/Private/Runtime/Player/DASH/ManifestDASH.h
+++ b/Engine/Plugins/Media/ElectraPlayer/Source/ElectraPlayerRuntime/Private/Runtime/Player/DASH/ManifestDASH.h
@@ -23,7 +23,7 @@ public:
 	FTimeValue GetAnchorTime() const override;
 	FTimeRange GetTotalTimeRange() const override;
 	FTimeRange GetSeekableTimeRange() const override;
-	FTimeRange GetPlaybackRange() const override;
+	FTimeRange GetPlaybackRange(EPlaybackRangeType InRangeType) const override;
 	void GetSeekablePositions(TArray<FTimespan>& OutPositions) const override;
 	FTimeValue GetDuration() const override;
 	FTimeValue GetDefaultStartTime() const override;
diff --git a/Engine/Plugins/Media/ElectraPlayer/Source/ElectraPlayerRuntime/Private/Runtime/Player/DASH/PlaylistReaderDASH_Internal.h b/Engine/Plugins/Media/ElectraPlayer/Source/ElectraPlayerRuntime/Private/Runtime/Player/DASH/PlaylistReaderDASH_Internal.h
--- a/Engine/Plugins/Media/ElectraPlayer/Source/ElectraPlayerRuntime/Private/Runtime/Player/DASH/PlaylistReaderDASH_Internal.h
+++ b/Engine/Plugins/Media/ElectraPlayer/Source/ElectraPlayerRuntime/Private/Runtime/Player/DASH/PlaylistReaderDASH_Internal.h
@@ -897,7 +897,7 @@ public:
 	void ClearDefaultStartTime();
 	FTimeValue GetDefaultEndTime() const;
 	void ClearDefaultEndTime();
-	FTimeRange GetPlayTimesFromURI() const;
+	FTimeRange GetPlayTimesFromURI(IManifest::EPlaybackRangeType InRangeType) const;
 	FTimeValue GetDesiredLiveLatency() const;
 
 	FTimeValue GetMPDValidityEndTime() const;
diff --git a/Engine/Plugins/Media/ElectraPlayer/Source/ElectraPlayerRuntime/Private/Runtime/Player/HLS/ManifestHLS.cpp b/Engine/Plugins/Media/ElectraPlayer/Source/ElectraPlayerRuntime/Private/Runtime/Player/HLS/ManifestHLS.cpp
--- a/Engine/Plugins/Media/ElectraPlayer/Source/ElectraPlayerRuntime/Private/Runtime/Player/HLS/ManifestHLS.cpp
+++ b/Engine/Plugins/Media/ElectraPlayer/Source/ElectraPlayerRuntime/Private/Runtime/Player/HLS/ManifestHLS.cpp
@@ -168,7 +168,7 @@ FTimeRange FManifestHLS::GetSeekableTimeRange() const
 	}
 	return FTimeRange();
 }
-FTimeRange FManifestHLS::GetPlaybackRange() const
+FTimeRange FManifestHLS::GetPlaybackRange(EPlaybackRangeType /*InRangeType*/) const
 {
 	return FTimeRange();
 }
diff --git a/Engine/Plugins/Media/ElectraPlayer/Source/ElectraPlayerRuntime/Private/Runtime/Player/HLS/ManifestHLS.h b/Engine/Plugins/Media/ElectraPlayer/Source/ElectraPlayerRuntime/Private/Runtime/Player/HLS/ManifestHLS.h
--- a/Engine/Plugins/Media/ElectraPlayer/Source/ElectraPlayerRuntime/Private/Runtime/Player/HLS/ManifestHLS.h
+++ b/Engine/Plugins/Media/ElectraPlayer/Source/ElectraPlayerRuntime/Private/Runtime/Player/HLS/ManifestHLS.h
@@ -21,7 +21,7 @@ public:
 	FTimeValue GetAnchorTime() const override;
 	FTimeRange GetTotalTimeRange() const override;
 	FTimeRange GetSeekableTimeRange() const override;
-	FTimeRange GetPlaybackRange() const override;
+	FTimeRange GetPlaybackRange(EPlaybackRangeType InRangeType) const override;
 	void GetSeekablePositions(TArray<FTimespan>& OutPositions) const override;
 	FTimeValue GetDuration() const override;
 	FTimeValue GetDefaultStartTime() const override;
diff --git a/Engine/Plugins/Media/ElectraPlayer/Source/ElectraPlayerRuntime/Private/Runtime/Player/Manifest.h b/Engine/Plugins/Media/ElectraPlayer/Source/ElectraPlayerRuntime/Private/Runtime/Player/Manifest.h
--- a/Engine/Plugins/Media/ElectraPlayer/Source/ElectraPlayerRuntime/Private/Runtime/Player/Manifest.h
+++ b/Engine/Plugins/Media/ElectraPlayer/Source/ElectraPlayerRuntime/Private/Runtime/Player/Manifest.h
@@ -272,15 +272,29 @@ namespace Electra
 		 */
 		virtual FTimeRange GetSeekableTimeRange() const = 0;
 
+		enum class EPlaybackRangeType
+		{
+			/**
+			 * Initial playback range as may be defined using `#t=s,e` URL fragment parameter.
+			 * This is used only on first playstart and is canceled when a Seek() is performed.
+			 */
+			TemporaryPlaystartRange,
+			/**
+			 * Fixed playback range that may be defined using `#r=s,e` URL fragment parameter.
+			 * This is a non-standard parameter. The specified range will be locked in place
+			 * and any Seek() can only be performed inside this range.
+			 */
+			LockedPlaybackRange
+		};
 		/**
 		 * Returns the playback range on the timeline, which is a subset of the total
 		 * time range. This may be set through manifest internal means or by URL fragment
 		 * parameters where permissable (eg. example.mp4#t=22,50).
 		 * If start or end are not specified they will be set to invalid.
-		 * 
+		 *
 		 * @return Optionally set time range to which playback is restricted.
 		 */
-		virtual FTimeRange GetPlaybackRange() const = 0;
+		virtual FTimeRange GetPlaybackRange(EPlaybackRangeType InRangeType) const = 0;
 
 		/**
 		 * Returns the timestamps of the segments from the video or audio track (if no video is present).
@@ -556,7 +570,7 @@ namespace Electra
 		{
 			return ValidFrom;
 		}
-		
+
 		TSharedPtrTS<UtilsMP4::FMetadataParser> GetMetadata() const
 		{
 			return Metadata;
diff --git a/Engine/Plugins/Media/ElectraPlayer/Source/ElectraPlayerRuntime/Private/Runtime/Player/mkv/ManifestMKV.cpp b/Engine/Plugins/Media/ElectraPlayer/Source/ElectraPlayerRuntime/Private/Runtime/Player/mkv/ManifestMKV.cpp
--- a/Engine/Plugins/Media/ElectraPlayer/Source/ElectraPlayerRuntime/Private/Runtime/Player/mkv/ManifestMKV.cpp
+++ b/Engine/Plugins/Media/ElectraPlayer/Source/ElectraPlayerRuntime/Private/Runtime/Player/mkv/ManifestMKV.cpp
@@ -61,7 +61,7 @@ FErrorDetail FManifestMKVInternal::Build(TSharedPtrTS<IParserMKV> MKVParser, con
 	ConnectionInfo = InConnectionInfo;
 	MediaAsset = MakeSharedTS<FTimelineAssetMKV>();
 	FErrorDetail Error = MediaAsset->Build(PlayerSessionServices, MKVParser, URL);
-	FTimeRange PlaybackRange = GetPlaybackRange();
+	FTimeRange PlaybackRange = GetPlaybackRange(IManifest::EPlaybackRangeType::TemporaryPlaystartRange);
 	DefaultStartTime = PlaybackRange.Start;
 	DefaultEndTime = PlaybackRange.End;
 	return Error;
@@ -106,15 +106,16 @@ void FManifestMKVInternal::GetTrackMetadata(TArray<FTrackMetadata>& OutMetadata,
  *
  * @return Optionally set time range to which playback is restricted.
  */
-FTimeRange FManifestMKVInternal::GetPlaybackRange() const
+FTimeRange FManifestMKVInternal::GetPlaybackRange(EPlaybackRangeType InRangeType) const
 {
 	FTimeRange FromTo;
 
-	// We are interested in the 't' fragment value here.
+	// We are interested in the 't' or 'r' fragment value here.
 	FString Time;
 	for(auto& Fragment : URLFragmentComponents)
 	{
-		if (Fragment.Name.Equals(TEXT("t")))
+		if ((InRangeType == IManifest::EPlaybackRangeType::TemporaryPlaystartRange && Fragment.Name.Equals(TEXT("t"))) ||
+			(InRangeType == IManifest::EPlaybackRangeType::LockedPlaybackRange && Fragment.Name.Equals(TEXT("r"))))
 		{
 			Time = Fragment.Value;
 		}
diff --git a/Engine/Plugins/Media/ElectraPlayer/Source/ElectraPlayerRuntime/Private/Runtime/Player/mkv/ManifestMKV.h b/Engine/Plugins/Media/ElectraPlayer/Source/ElectraPlayerRuntime/Private/Runtime/Player/mkv/ManifestMKV.h
--- a/Engine/Plugins/Media/ElectraPlayer/Source/ElectraPlayerRuntime/Private/Runtime/Player/mkv/ManifestMKV.h
+++ b/Engine/Plugins/Media/ElectraPlayer/Source/ElectraPlayerRuntime/Private/Runtime/Player/mkv/ManifestMKV.h
@@ -41,7 +41,7 @@ public:
 		// For the time being the seekable range equals the total range.
 		return GetTotalTimeRange();
 	}
-	FTimeRange GetPlaybackRange() const override;
+	FTimeRange GetPlaybackRange(EPlaybackRangeType InRangeType) const override;
 	void GetSeekablePositions(TArray<FTimespan>& OutPositions) const override
 	{
 		// For the time being we do not return anything here as that would require to iterate the tracks.
diff --git a/Engine/Plugins/Media/ElectraPlayer/Source/ElectraPlayerRuntime/Private/Runtime/Player/mp4/ManifestMP4.cpp b/Engine/Plugins/Media/ElectraPlayer/Source/ElectraPlayerRuntime/Private/Runtime/Player/mp4/ManifestMP4.cpp
--- a/Engine/Plugins/Media/ElectraPlayer/Source/ElectraPlayerRuntime/Private/Runtime/Player/mp4/ManifestMP4.cpp
+++ b/Engine/Plugins/Media/ElectraPlayer/Source/ElectraPlayerRuntime/Private/Runtime/Player/mp4/ManifestMP4.cpp
@@ -65,7 +65,7 @@ FErrorDetail FManifestMP4Internal::Build(TSharedPtrTS<IParserISO14496_12> MP4Par
 {
 	MediaAsset = MakeSharedTS<FTimelineAssetMP4>();
 	FErrorDetail Result = MediaAsset->Build(PlayerSessionServices, MP4Parser, URL);
-	FTimeRange PlaybackRange = GetPlaybackRange();
+	FTimeRange PlaybackRange = GetPlaybackRange(IManifest::EPlaybackRangeType::TemporaryPlaystartRange);
 	DefaultStartTime = PlaybackRange.Start;
 	DefaultEndTime = PlaybackRange.End;
 	return Result;
@@ -126,15 +126,16 @@ void FManifestMP4Internal::GetTrackMetadata(TArray<FTrackMetadata>& OutMetadata,
  *
  * @return Optionally set time range to which playback is restricted.
  */
-FTimeRange FManifestMP4Internal::GetPlaybackRange() const
+FTimeRange FManifestMP4Internal::GetPlaybackRange(EPlaybackRangeType InRangeType) const
 {
 	FTimeRange FromTo;
 
-	// We are interested in the 't' fragment value here.
+	// We are interested in the 't' or 'r' fragment value here.
 	FString Time;
 	for(auto& Fragment : URLFragmentComponents)
 	{
-		if (Fragment.Name.Equals(TEXT("t")))
+		if ((InRangeType == IManifest::EPlaybackRangeType::TemporaryPlaystartRange && Fragment.Name.Equals(TEXT("t"))) ||
+			(InRangeType == IManifest::EPlaybackRangeType::LockedPlaybackRange && Fragment.Name.Equals(TEXT("r"))))
 		{
 			Time = Fragment.Value;
 		}
diff --git a/Engine/Plugins/Media/ElectraPlayer/Source/ElectraPlayerRuntime/Private/Runtime/Player/mp4/ManifestMP4.h b/Engine/Plugins/Media/ElectraPlayer/Source/ElectraPlayerRuntime/Private/Runtime/Player/mp4/ManifestMP4.h
--- a/Engine/Plugins/Media/ElectraPlayer/Source/ElectraPlayerRuntime/Private/Runtime/Player/mp4/ManifestMP4.h
+++ b/Engine/Plugins/Media/ElectraPlayer/Source/ElectraPlayerRuntime/Private/Runtime/Player/mp4/ManifestMP4.h
@@ -42,7 +42,7 @@ public:
 		// For the time being the seekable range equals the total range.
 		return GetTotalTimeRange();
 	}
-	FTimeRange GetPlaybackRange() const override;
+	FTimeRange GetPlaybackRange(EPlaybackRangeType InRangeType) const override;
 	void GetSeekablePositions(TArray<FTimespan>& OutPositions) const override
 	{
 		// For the time being we do not return anything here as that would require to iterate the tracks.
diff --git a/Engine/Plugins/Media/ElectraPlayer/Source/ElectraPlayerRuntime/Private/Runtime/StreamAccessUnitBuffer.h b/Engine/Plugins/Media/ElectraPlayer/Source/ElectraPlayerRuntime/Private/Runtime/StreamAccessUnitBuffer.h
--- a/Engine/Plugins/Media/ElectraPlayer/Source/ElectraPlayerRuntime/Private/Runtime/StreamAccessUnitBuffer.h
+++ b/Engine/Plugins/Media/ElectraPlayer/Source/ElectraPlayerRuntime/Private/Runtime/StreamAccessUnitBuffer.h
@@ -6,6 +6,7 @@
 #include "PlayerTime.h"
 #include "StreamTypes.h"
 #include "Containers/Queue.h"
+#include "HAL/PlatformMisc.h"
 
 
 
@@ -943,6 +944,7 @@ public:
 	void ClearEOD()
 	{
 		bIsEOD = false;
+		FPlatformMisc::MemoryBarrier();
 		bReachedEOD = false;
 	}
 
@@ -950,7 +952,7 @@ public:
 	{
 		return bIsEOD;
 	}
-	
+
 	bool ReachedEOD() const
 	{
 		return bReachedEOD;
@@ -961,8 +963,8 @@ private:
 
 	FMediaSemaphore AvailSema;
 	TQueue<T> Elements;
-	bool bIsEOD = false;
-	bool bReachedEOD = false;
+	volatile bool bIsEOD = false;
+	volatile bool bReachedEOD = false;
 };
 
 
