Media: Reworked media player options; allow EMediaControl::Scrub with Electra

#rnx

diff --git a/Engine/Plugins/Media/ElectraPlayer/Source/ElectraPlayerPlugin/Private/ElectraPlayerPlugin.cpp b/Engine/Plugins/Media/ElectraPlayer/Source/ElectraPlayerPlugin/Private/ElectraPlayerPlugin.cpp
--- a/Engine/Plugins/Media/ElectraPlayer/Source/ElectraPlayerPlugin/Private/ElectraPlayerPlugin.cpp
+++ b/Engine/Plugins/Media/ElectraPlayer/Source/ElectraPlayerPlugin/Private/ElectraPlayerPlugin.cpp
@@ -135,21 +135,21 @@ public:
 			}
 		}
 	}
-	
+
 	virtual EOrigin GetOrigin() const override
 	{
 		switch(Metadata->GetOrigin())
 		{
 			default:
-			case IMetaDataDecoderOutput::EOrigin::TimedMetadata:		
+			case IMetaDataDecoderOutput::EOrigin::TimedMetadata:
 			{
 				return FElectraBinarySample::EOrigin::TimedMetadata;
 			}
-			case IMetaDataDecoderOutput::EOrigin::EventStream:		
+			case IMetaDataDecoderOutput::EOrigin::EventStream:
 			{
 				return FElectraBinarySample::EOrigin::EventStream;
 			}
-			case IMetaDataDecoderOutput::EOrigin::InbandEventStream:	
+			case IMetaDataDecoderOutput::EOrigin::InbandEventStream:
 			{
 				return FElectraBinarySample::EOrigin::InbandEventStream;
 			}
@@ -157,11 +157,11 @@ public:
 	}
 
 	virtual FMediaTimeStamp GetTime() const override
-	{ 
-		FDecoderTimeStamp ts = Metadata->GetTime(); 
+	{
+		FDecoderTimeStamp ts = Metadata->GetTime();
 		return FMediaTimeStamp(ts.Time, ts.SequenceIndex);
 	}
-	
+
 	virtual FTimespan GetDuration() const override
 	{
 		FTimespan Duration = Metadata->GetDuration();
@@ -175,9 +175,9 @@ public:
 	}
 
 	virtual TOptional<FMediaTimeStamp> GetTrackBaseTime() const	override
-	{ 
+	{
 		TOptional<FMediaTimeStamp> ms;
-		TOptional<FDecoderTimeStamp> ts = Metadata->GetTime(); 
+		TOptional<FDecoderTimeStamp> ts = Metadata->GetTime();
 		if (ts.IsSet())
 		{
 			ms = FMediaTimeStamp(ts.GetValue().Time, ts.GetValue().SequenceIndex);
@@ -194,16 +194,16 @@ class FElectraSubtitleSample : public IElectraSubtitleSample
 {
 public:
 	virtual FGuid GetGUID() const override
-	{ 
-		return IElectraSubtitleSample::GetSampleTypeGUID(); 
+	{
+		return IElectraSubtitleSample::GetSampleTypeGUID();
 	}
 
 	virtual FMediaTimeStamp GetTime() const override
-	{ 
-		FDecoderTimeStamp ts = Subtitle->GetTime(); 
+	{
+		FDecoderTimeStamp ts = Subtitle->GetTime();
 		return FMediaTimeStamp(ts.Time, ts.SequenceIndex);
 	}
-	
+
 	virtual FTimespan GetDuration() const override
 	{
 		return Subtitle->GetDuration();
@@ -213,14 +213,14 @@ public:
 	{
 		return TOptional<FVector2D>();
 	}
-	
+
 	virtual FText GetText() const override
 	{
 		FUTF8ToTCHAR cnv((const ANSICHAR*)Subtitle->GetData().GetData(), Subtitle->GetData().Num());
 		FString UTF8Text(cnv.Length(), cnv.Get());
 		return FText::FromString(UTF8Text);
 	}
-	
+
 	virtual EMediaOverlaySampleType GetType() const override
 	{
 		return EMediaOverlaySampleType::Subtitle;
@@ -579,11 +579,20 @@ bool FElectraPlayerPlugin::Open(const FString& Url, const IMediaOptions* Options
 	IElectraPlayerInterface::FPlaystartOptions LocalPlaystartOptions;
 
 	// Get playstart options from passed options, if they exist.
+	FName Environment;
 	if (InPlayerOptions)
 	{
-		LocalPlaystartOptions.TimeOffset = InPlayerOptions->SeekTime;
-		LocalPlaystartOptions.InitialAudioTrackAttributes.TrackIndexOverride = InPlayerOptions->Tracks.Audio;
-		LocalPlaystartOptions.InitialSubtitleTrackAttributes.TrackIndexOverride = InPlayerOptions->Tracks.Subtitle;
+		if (InPlayerOptions->SeekTimeType != EMediaPlayerOptionSeekTimeType::Ignored)
+		{
+			LocalPlaystartOptions.TimeOffset = InPlayerOptions->SeekTime;
+		}
+		if (InPlayerOptions->TrackSelection == EMediaPlayerOptionTrackSelectMode::UseTrackOptionIndices)
+		{
+			LocalPlaystartOptions.InitialAudioTrackAttributes.TrackIndexOverride = InPlayerOptions->Tracks.Audio;
+			LocalPlaystartOptions.InitialSubtitleTrackAttributes.TrackIndexOverride = InPlayerOptions->Tracks.Subtitle;
+		}
+		const FVariant* Env = InPlayerOptions->InternalCustomOptions.Find(MediaPlayerOptionValues::Environment());
+		Environment = Env ? Env->GetValue<FName>() : Environment;
 	}
 	FString InitialAudioLanguage = Options->GetMediaOption(TEXT("InitialAudioLanguage"), FString());
 	if (InitialAudioLanguage.Len())
@@ -677,23 +686,36 @@ bool FElectraPlayerPlugin::Open(const FString& Url, const IMediaOptions* Options
 			UE_LOG(LogElectraPlayerPlugin, Log, TEXT("[%p] IMediaPlayer::Open: CDN HTTP status %d will deny a stream permanently"), this, HTTPStatus);
 		}
 	}
-	// Check if there are options to tweak for scrubbing
-	bool bUseScrubOptimizations = Options->GetMediaOption(TEXT("ElectraScrubOptimization"), (bool)false);
-	if (bUseScrubOptimizations)
+
+	// Check if there is an environment specified in which this player is used.
+	// Certain optimization settings apply for dedicated environments.
+	if (Environment == MediaPlayerOptionValues::Environment_Preview() || Environment == MediaPlayerOptionValues::Environment_Sequencer())
 	{
-		int64 ScrubSeekBitrate = Options->GetMediaOption(TEXT("ElectraScrubBitrate"), (int64)-1);
-		int64 ScrubCacheSizeKiB = Options->GetMediaOption(TEXT("ElectraScrubCacheSizeKiB"), (int64)-1);
-		if (ScrubSeekBitrate >= 0)
-		{
-			PlayerOptions.Set(TEXT("seekstart_bitrate"), Electra::FVariantValue(ScrubSeekBitrate));
-		}
-		if (ScrubCacheSizeKiB > 0)
-		{
-			PlayerOptions.Set(TEXT("httpcache_max_bytesize"), Electra::FVariantValue(ScrubCacheSizeKiB << 10));
-			PlayerOptions.Set(TEXT("httpcache_max_entries"), Electra::FVariantValue((int64)10000));
-		}
 		PlayerOptions.Set(TEXT("optimize_seek_for_scrubbing"), Electra::FVariantValue(true));
+		PlayerOptions.Set(TEXT("new_scrubbing_seek_cancels_current"), Electra::FVariantValue(true));
 		PlayerOptions.Set(TEXT("do_not_hold_back_first_frame"), Electra::FVariantValue(true));
+		PlayerOptions.Set(TEXT("worker_threads"), Electra::FVariantValue(FString(TEXT("worker"))));
+	}
+	else
+	{
+		// Check if there are options to tweak for scrubbing
+		bool bUseScrubOptimizations = Options->GetMediaOption(TEXT("ElectraScrubOptimization"), (bool)false);
+		if (bUseScrubOptimizations)
+		{
+			int64 ScrubSeekBitrate = Options->GetMediaOption(TEXT("ElectraScrubBitrate"), (int64)-1);
+			int64 ScrubCacheSizeKiB = Options->GetMediaOption(TEXT("ElectraScrubCacheSizeKiB"), (int64)-1);
+			if (ScrubSeekBitrate >= 0)
+			{
+				PlayerOptions.Set(TEXT("seekstart_bitrate"), Electra::FVariantValue(ScrubSeekBitrate));
+			}
+			if (ScrubCacheSizeKiB > 0)
+			{
+				PlayerOptions.Set(TEXT("httpcache_max_bytesize"), Electra::FVariantValue(ScrubCacheSizeKiB << 10));
+				PlayerOptions.Set(TEXT("httpcache_max_entries"), Electra::FVariantValue((int64)10000));
+			}
+			PlayerOptions.Set(TEXT("optimize_seek_for_scrubbing"), Electra::FVariantValue(true));
+			PlayerOptions.Set(TEXT("do_not_hold_back_first_frame"), Electra::FVariantValue(true));
+		}
 	}
 
 	// Check for options that can be changed during playback and apply them at startup already.
@@ -875,7 +897,7 @@ bool FElectraPlayerPlugin::CanControl(EMediaControl Control) const
 	{
 		return CurrentState == EMediaState::Paused || CurrentState == EMediaState::Stopped;
 	}
-	else if (Control == EMediaControl::Seek)
+	else if (Control == EMediaControl::Seek || Control == EMediaControl::Scrub)
 	{
 		return CurrentState == EMediaState::Playing || CurrentState == EMediaState::Paused || CurrentState == EMediaState::Stopped;
 	}
diff --git a/Engine/Plugins/Media/ElectraPlayer/Source/ElectraPlayerRuntime/Private/ElectraPlayer.cpp b/Engine/Plugins/Media/ElectraPlayer/Source/ElectraPlayerRuntime/Private/ElectraPlayer.cpp
--- a/Engine/Plugins/Media/ElectraPlayer/Source/ElectraPlayerRuntime/Private/ElectraPlayer.cpp
+++ b/Engine/Plugins/Media/ElectraPlayer/Source/ElectraPlayerRuntime/Private/ElectraPlayer.cpp
@@ -286,6 +286,10 @@ bool FElectraPlayer::OpenInternal(const FString& Url, const FParamDict& InPlayer
 		CreateParams.VideoRenderer = NewPlayer->RendererVideo;
 		CreateParams.AudioRenderer = NewPlayer->RendererAudio;
 		CreateParams.ExternalPlayerGUID = PlayerGuid;
+		FString WorkerThreadOption = PlayerOptions.GetValue(Electra::OptionKeyWorkerThreads).SafeGetFString(TEXT("shared"));
+		CreateParams.WorkerThreads = WorkerThreadOption.Equals(TEXT("worker"), ESearchCase::IgnoreCase) ? IAdaptiveStreamingPlayer::FCreateParam::EWorkerThreads::DedicatedWorker :
+									 WorkerThreadOption.Equals(TEXT("worker_and_events"), ESearchCase::IgnoreCase) ? IAdaptiveStreamingPlayer::FCreateParam::EWorkerThreads::DedicatedWorkerAndEventDispatch :
+									 IAdaptiveStreamingPlayer::FCreateParam::EWorkerThreads::Shared;
 		NewPlayer->AdaptivePlayer = IAdaptiveStreamingPlayer::Create(CreateParams);
 		NewPlayer->AdaptivePlayer->AddMetricsReceiver(this);
 		NewPlayer->AdaptivePlayer->SetStaticResourceProviderCallback(StaticResourceProvider);
diff --git a/Engine/Plugins/Media/ElectraPlayer/Source/ElectraPlayerRuntime/Private/Runtime/Player/AdaptivePlayerOptionKeynames.h b/Engine/Plugins/Media/ElectraPlayer/Source/ElectraPlayerRuntime/Private/Runtime/Player/AdaptivePlayerOptionKeynames.h
--- a/Engine/Plugins/Media/ElectraPlayer/Source/ElectraPlayerRuntime/Private/Runtime/Player/AdaptivePlayerOptionKeynames.h
+++ b/Engine/Plugins/Media/ElectraPlayer/Source/ElectraPlayerRuntime/Private/Runtime/Player/AdaptivePlayerOptionKeynames.h
@@ -6,8 +6,11 @@
 
 namespace Electra
 {
+	//! (FString) Use of worker threads. ("shared", "worker" or "worker_and_events")
+	const FName OptionKeyWorkerThreads(TEXT("worker_threads"));
+
 	//! (FString) mime type of URL to load
-	const FName OptionKeyMimeType(TEXT("mime_type"));										
+	const FName OptionKeyMimeType(TEXT("mime_type"));
 
 	//! (int64) value indicating the bitrate to start playback with (initial start).
 	const FName OptionKeyInitialBitrate(TEXT("initial_bitrate"));
@@ -32,13 +35,16 @@ namespace Electra
 
 	//! (bool) true to optimize seeking for faster frame scrubbing, false to optimize for playback.
 	const FName OptionKeyFrameOptimizeSeekForScrubbing(TEXT("optimize_seek_for_scrubbing"));
-	
+
+	//! (bool) true to allow a new scrubbing seek to cancel an ongoing scrubbing seek. Non-scrubbing seeks always cancel pending seeks.
+	const FName OptionKeyNewScrubbingSeekCancelsCurrent(TEXT("new_scrubbing_seek_cancels_current"));
+
 	//! (bool) true to emit the first decoded video frame while prerolling so it can be displayed while scrubbing.
 	const FName OptionKeyDoNotHoldBackFirstVideoFrame(TEXT("do_not_hold_back_first_frame"));
 
 	//! (bool) true to not truncate the media segment access units at the end of the presentation. Must only be used without a set playback range end!
 	const FName OptionKeyDoNotTruncateAtPresentationEnd(TEXT("do_not_truncate_at_presentation_end"));
-	
+
 	//! (bool) true to have every request to read data break out to an external data reader.
 	const FName OptionKeyUseExternalDataReader(TEXT("use_external_data_reader"));
 
@@ -49,7 +55,7 @@ namespace Electra
 	const FName OptionKeyExcludedCodecsVideo(TEXT("excluded_codecs_video"));
 	const FName OptionKeyExcludedCodecsAudio(TEXT("excluded_codecs_audio"));
 	const FName OptionKeyExcludedCodecsSubtitles(TEXT("excluded_codecs_subtitles"));
-	
+
 	const FName OptionKeyPreferredCodecsVideo(TEXT("preferred_codecs_video"));
 	const FName OptionKeyPreferredCodecsAudio(TEXT("preferred_codecs_audio"));
 	const FName OptionKeyPreferredCodecsSubtitles(TEXT("preferred_codecs_subtitles"));
diff --git a/Engine/Plugins/Media/ElectraPlayer/Source/ElectraPlayerRuntime/Private/Runtime/Player/AdaptiveStreamingPlayer.cpp b/Engine/Plugins/Media/ElectraPlayer/Source/ElectraPlayerRuntime/Private/Runtime/Player/AdaptiveStreamingPlayer.cpp
--- a/Engine/Plugins/Media/ElectraPlayer/Source/ElectraPlayerRuntime/Private/Runtime/Player/AdaptiveStreamingPlayer.cpp
+++ b/Engine/Plugins/Media/ElectraPlayer/Source/ElectraPlayerRuntime/Private/Runtime/Player/AdaptiveStreamingPlayer.cpp
@@ -26,7 +26,7 @@ FAdaptiveStreamingPlayer *FAdaptiveStreamingPlayer::PointerToLatestPlayer;
 
 TSharedPtr<IAdaptiveStreamingPlayer, ESPMode::ThreadSafe> IAdaptiveStreamingPlayer::Create(const IAdaptiveStreamingPlayer::FCreateParam& InCreateParameters)
 {
-	return MakeShared<FAdaptiveStreamingPlayer, ESPMode::ThreadSafe>(InCreateParameters);
+	return MakeShareable(new FAdaptiveStreamingPlayer(InCreateParameters), FAdaptiveStreamingPlayer::TDeleter());
 }
 
 void IAdaptiveStreamingPlayer::DebugHandle(void* pPlayer, void (*debugDrawPrintf)(void* pPlayer, const char *pFmt, ...))
@@ -74,6 +74,7 @@ FAdaptiveStreamingPlayer::FAdaptiveStreamingPlayer(const IAdaptiveStreamingPlaye
 		AudioRender.Renderer = CreateWrappedRenderer(InCreateParameters.AudioRenderer, EStreamType::Audio);
 	}
 	ExternalPlayerGUID		= InCreateParameters.ExternalPlayerGUID;
+	UseSharedWorkerThreads  = InCreateParameters.WorkerThreads;
 	ManifestType			= EMediaFormatType::Unknown;
 	CurrentState   		    = EPlayerState::eState_Idle;
 	PipelineState  		    = EPipelineState::ePipeline_Stopped;
@@ -134,7 +135,6 @@ FAdaptiveStreamingPlayer::FAdaptiveStreamingPlayer(const IAdaptiveStreamingPlaye
 
 	TMediaInterlockedExchangePointer(PointerToLatestPlayer, this);
 
-	bUseSharedWorkerThread = true;
 	StartWorkerThread();
 }
 
@@ -509,11 +509,11 @@ void FAdaptiveStreamingPlayer::StartWorkerThread()
 	// Get us an event dispatcher and add ourselves to the shared worker thread.
 	if (!EventDispatcher.IsValid())
 	{
-		EventDispatcher = FAdaptiveStreamingPlayerEventHandler::Create();
+		EventDispatcher = FAdaptiveStreamingPlayerEventHandler::Create(UseSharedWorkerThreads != IAdaptiveStreamingPlayer::FCreateParam::EWorkerThreads::DedicatedWorkerAndEventDispatch);
 	}
 	if (!SharedWorkerThread.IsValid())
 	{
-		SharedWorkerThread = FAdaptiveStreamingPlayerWorkerThread::Create(bUseSharedWorkerThread);
+		SharedWorkerThread = FAdaptiveStreamingPlayerWorkerThread::Create(UseSharedWorkerThreads == IAdaptiveStreamingPlayer::FCreateParam::EWorkerThreads::Shared);
 		WorkerThread.SetSharedWorkerThread(SharedWorkerThread);
 		SharedWorkerThread->AddPlayerInstance(this);
 	}
@@ -1726,7 +1726,7 @@ void FAdaptiveStreamingPlayer::HandleSeeking()
 		return;
 	}
 
-	// When playing the last successfully seeked position is irrelevant and cannot be used as a reference any more.
+	// When playing, the last successfully seeked position is irrelevant and cannot be used as a reference any more.
 	if (PlaybackState.GetIsPlaying())
 	{
 		SeekVars.InvalidateLastFinished();
@@ -1738,7 +1738,8 @@ void FAdaptiveStreamingPlayer::HandleSeeking()
 	{
 		// If there is an active request and the new request is for scrubbing we let the active request finish first.
 		bool bIsForScrubbing = SeekVars.PendingRequest.GetValue().bOptimizeForScrubbing.Get(PlayerOptions.GetValue(OptionKeyFrameOptimizeSeekForScrubbing).SafeGetBool(false));
-		if (SeekVars.ActiveRequest.IsSet() && bIsForScrubbing)
+		bool bNewScrubSeekCancelsCurrent = PlayerOptions.GetValue(OptionKeyNewScrubbingSeekCancelsCurrent).SafeGetBool(false);
+		if (SeekVars.ActiveRequest.IsSet() && bIsForScrubbing && !bNewScrubSeekCancelsCurrent)
 		{
 			return;
 		}
@@ -5057,18 +5058,27 @@ void FAdaptiveStreamingPlayerWorkerThread::WorkerThreadFN()
 TWeakPtrTS<FAdaptiveStreamingPlayerEventHandler>	FAdaptiveStreamingPlayerEventHandler::SingletonSelf;
 FCriticalSection									FAdaptiveStreamingPlayerEventHandler::SingletonLock;
 
-TSharedPtrTS<FAdaptiveStreamingPlayerEventHandler> FAdaptiveStreamingPlayerEventHandler::Create()
+TSharedPtrTS<FAdaptiveStreamingPlayerEventHandler> FAdaptiveStreamingPlayerEventHandler::Create(bool bUseSharedWorkerThread)
 {
-	FScopeLock lock(&SingletonLock);
-	TSharedPtrTS<FAdaptiveStreamingPlayerEventHandler> Self = SingletonSelf.Pin();
-	if (!Self.IsValid())
+	if (bUseSharedWorkerThread)
+	{
+		FScopeLock lock(&SingletonLock);
+		TSharedPtrTS<FAdaptiveStreamingPlayerEventHandler> Self = SingletonSelf.Pin();
+		if (!Self.IsValid())
+		{
+			FAdaptiveStreamingPlayerEventHandler* Handler = new FAdaptiveStreamingPlayerEventHandler;
+			Handler->StartWorkerThread();
+			Self = MakeShareable(Handler);
+			SingletonSelf = Self;
+		}
+		return Self;
+	}
+	else
 	{
 		FAdaptiveStreamingPlayerEventHandler* Handler = new FAdaptiveStreamingPlayerEventHandler;
 		Handler->StartWorkerThread();
-		Self = MakeShareable(Handler);
-		SingletonSelf = Self;
+		return MakeShareable(Handler);
 	}
-	return Self;
 }
 
 void FAdaptiveStreamingPlayerEventHandler::DispatchEvent(TSharedPtrTS<FMetricEvent> InEvent)
diff --git a/Engine/Plugins/Media/ElectraPlayer/Source/ElectraPlayerRuntime/Private/Runtime/Player/AdaptiveStreamingPlayer.h b/Engine/Plugins/Media/ElectraPlayer/Source/ElectraPlayerRuntime/Private/Runtime/Player/AdaptiveStreamingPlayer.h
--- a/Engine/Plugins/Media/ElectraPlayer/Source/ElectraPlayerRuntime/Private/Runtime/Player/AdaptiveStreamingPlayer.h
+++ b/Engine/Plugins/Media/ElectraPlayer/Source/ElectraPlayerRuntime/Private/Runtime/Player/AdaptiveStreamingPlayer.h
@@ -30,9 +30,16 @@ class IAdaptiveStreamingPlayer : public TMediaNoncopyable<IAdaptiveStreamingPlay
 public:
 	struct FCreateParam
 	{
-		TSharedPtr<IMediaRenderer, ESPMode::ThreadSafe>					VideoRenderer;
-		TSharedPtr<IMediaRenderer, ESPMode::ThreadSafe>					AudioRenderer;
-		FGuid															ExternalPlayerGUID;
+		TSharedPtr<IMediaRenderer, ESPMode::ThreadSafe> VideoRenderer;
+		TSharedPtr<IMediaRenderer, ESPMode::ThreadSafe> AudioRenderer;
+		FGuid ExternalPlayerGUID;
+		enum class EWorkerThreads
+		{
+			Shared,
+			DedicatedWorker,
+			DedicatedWorkerAndEventDispatch
+		};
+		EWorkerThreads WorkerThreads = EWorkerThreads::Shared;
 	};
 
 	static TSharedPtr<IAdaptiveStreamingPlayer, ESPMode::ThreadSafe> Create(const FCreateParam& InCreateParameters);
diff --git a/Engine/Plugins/Media/ElectraPlayer/Source/ElectraPlayerRuntime/Private/Runtime/Player/AdaptiveStreamingPlayerInternal.h b/Engine/Plugins/Media/ElectraPlayer/Source/ElectraPlayerRuntime/Private/Runtime/Player/AdaptiveStreamingPlayerInternal.h
--- a/Engine/Plugins/Media/ElectraPlayer/Source/ElectraPlayerRuntime/Private/Runtime/Player/AdaptiveStreamingPlayerInternal.h
+++ b/Engine/Plugins/Media/ElectraPlayer/Source/ElectraPlayerRuntime/Private/Runtime/Player/AdaptiveStreamingPlayerInternal.h
@@ -946,7 +946,7 @@ private:
 class FAdaptiveStreamingPlayerEventHandler
 {
 public:
-	static TSharedPtrTS<FAdaptiveStreamingPlayerEventHandler> Create();
+	static TSharedPtrTS<FAdaptiveStreamingPlayerEventHandler> Create(bool bUseSharedWorkerThread);
 
 	void DispatchEvent(TSharedPtrTS<FMetricEvent> InEvent);
 
@@ -973,6 +973,19 @@ class FAdaptiveStreamingPlayer : public IAdaptiveStreamingPlayer
 							   , public IAdaptiveStreamSelector::IPlayerLiveControl
 {
 public:
+	class TDeleter
+	{
+	public:
+		void operator()(FAdaptiveStreamingPlayer* InInstanceToDelete)
+		{
+			TFunction<void()> DeleteTask = [InInstanceToDelete]()
+			{
+				delete InInstanceToDelete;
+			};
+			FMediaRunnable::EnqueueAsyncTask(MoveTemp(DeleteTask));
+		}
+	};
+
 	FAdaptiveStreamingPlayer(const IAdaptiveStreamingPlayer::FCreateParam& InCreateParameters);
 	virtual ~FAdaptiveStreamingPlayer();
 
@@ -2169,7 +2182,7 @@ private:
 	TSharedPtrTS<FAdaptiveStreamingPlayerEventHandler>					EventDispatcher;
 	TSharedPtrTS<FAdaptiveStreamingPlayerWorkerThread>					SharedWorkerThread;
 	FWorkerThreadMessages												WorkerThread;
-	bool																bUseSharedWorkerThread;
+	IAdaptiveStreamingPlayer::FCreateParam::EWorkerThreads				UseSharedWorkerThreads;
 
 	FParamDictTS														PlayerOptions;
 	FPlaybackState														PlaybackState;
diff --git a/Engine/Plugins/Media/MediaCompositing/Source/MediaCompositing/Private/MovieScene/MovieSceneMediaTemplate.cpp b/Engine/Plugins/Media/MediaCompositing/Source/MediaCompositing/Private/MovieScene/MovieSceneMediaTemplate.cpp
--- a/Engine/Plugins/Media/MediaCompositing/Source/MediaCompositing/Private/MovieScene/MovieSceneMediaTemplate.cpp
+++ b/Engine/Plugins/Media/MediaCompositing/Source/MediaCompositing/Private/MovieScene/MovieSceneMediaTemplate.cpp
@@ -87,13 +87,17 @@ struct FMediaSectionPreRollExecutionToken
 		// open the media source if necessary
 		if (MediaPlayer->GetUrl().IsEmpty())
 		{
+			FMediaPlayerOptions Options;
+			Options.SetAllAsOptional();
+
 			if (PlayerProxyInterface != nullptr)
 			{
 				MediaSource->SetCacheSettings(PlayerProxyInterface->GetCacheSettings());
 			}
-			
 			SectionData.SeekOnOpen(StartTime);
-			MediaPlayer->OpenSource(MediaSource);
+
+			Options.InternalCustomOptions.Emplace(MediaPlayerOptionValues::Environment(), MediaPlayerOptionValues::Environment_Sequencer());
+			MediaPlayer->OpenSourceWithOptions(MediaSource, Options);
 			return;
 		}
 
@@ -176,15 +180,19 @@ struct FMediaSectionExecutionToken
 		// open the media source if necessary
 		if (MediaPlayer->GetUrl().IsEmpty() || bCacheSettingsChanged)
 		{
+			FMediaPlayerOptions Options;
+			Options.SetAllAsOptional();
+
 			if (PlayerProxyInterface != nullptr)
 			{
 				MediaSource->SetCacheSettings(PlayerProxyInterface->GetCacheSettings());
 			}
 			SectionData.SeekOnOpen(CurrentTime);
+
+			Options.InternalCustomOptions.Emplace(MediaPlayerOptionValues::Environment(), MediaPlayerOptionValues::Environment_Sequencer());
 			// Setup an initial blocking range - MediaFramework will block (even through the opening process) in its next tick...
 			MediaPlayer->SetBlockOnTimeRange(TRange<FTimespan>(CurrentTime, CurrentTime + FrameDuration));
-			MediaPlayer->OpenSource(MediaSource);
-
+			MediaPlayer->OpenSourceWithOptions(MediaSource, Options);
 			return;
 		}
 
diff --git a/Engine/Plugins/Media/MediaPlate/Source/MediaPlate/Private/MediaPlateComponent.cpp b/Engine/Plugins/Media/MediaPlate/Source/MediaPlate/Private/MediaPlateComponent.cpp
--- a/Engine/Plugins/Media/MediaPlate/Source/MediaPlate/Private/MediaPlateComponent.cpp
+++ b/Engine/Plugins/Media/MediaPlate/Source/MediaPlate/Private/MediaPlateComponent.cpp
@@ -530,10 +530,9 @@ bool UMediaPlateComponent::PlayMediaSource(UMediaSource* InMediaSource, bool bIn
 			// Play the source.
 			FMediaPlayerOptions Options;
 			Options.SeekTime = FTimespan::FromSeconds(StartTime);
-			Options.PlayOnOpen = bInPlayOnOpen ? EMediaPlayerOptionBooleanOverride::Enabled :
-				EMediaPlayerOptionBooleanOverride::Disabled;
-			Options.Loop = (bLoop && (bIsPlaylist == false)) ?
-				EMediaPlayerOptionBooleanOverride::Enabled : EMediaPlayerOptionBooleanOverride::Disabled;
+			Options.PlayOnOpen = bInPlayOnOpen ? EMediaPlayerOptionBooleanOverride::Enabled : EMediaPlayerOptionBooleanOverride::Disabled;
+			Options.Loop = (bLoop && (bIsPlaylist == false)) ? EMediaPlayerOptionBooleanOverride::Enabled : EMediaPlayerOptionBooleanOverride::Disabled;
+			Options.InternalCustomOptions.Emplace(MediaPlayerOptionValues::Environment(), MediaPlayerOptionValues::Environment_Preview());
 			bIsPlaying = MediaPlayer->OpenSourceWithOptions(InMediaSource, Options);
 
 			// Did we play anything?
diff --git a/Engine/Plugins/Media/MediaPlayerEditor/Source/MediaPlayerEditor/Private/Toolkits/MediaSourceEditorToolkit.cpp b/Engine/Plugins/Media/MediaPlayerEditor/Source/MediaPlayerEditor/Private/Toolkits/MediaSourceEditorToolkit.cpp
--- a/Engine/Plugins/Media/MediaPlayerEditor/Source/MediaPlayerEditor/Private/Toolkits/MediaSourceEditorToolkit.cpp
+++ b/Engine/Plugins/Media/MediaPlayerEditor/Source/MediaPlayerEditor/Private/Toolkits/MediaSourceEditorToolkit.cpp
@@ -105,7 +105,7 @@ void FMediaSourceEditorToolkit::Initialize(UMediaSource* InMediaSource, const ET
 						->AddTab(MediaSourceEditorToolkit::ViewerTabId, ETabState::OpenedTab)
 						->SetHideTabWell(true)
 						->SetSizeCoefficient(0.6f)
-						
+
 				)
 				->Split
 				(
@@ -140,7 +140,7 @@ void FMediaSourceEditorToolkit::Initialize(UMediaSource* InMediaSource, const ET
 		true /*bCreateDefaultToolbar*/,
 		InMediaSource
 	);
-	
+
 	ExtendToolBar();
 	RegenerateMenusAndToolbars();
 }
@@ -289,7 +289,13 @@ void FMediaSourceEditorToolkit::BindCommands()
 
 	ToolkitCommands->MapAction(
 		Commands.OpenMedia,
-		FExecuteAction::CreateLambda([this] { MediaPlayer->OpenSource(MediaSource); }),
+		FExecuteAction::CreateLambda([this]
+		{
+			FMediaPlayerOptions Options;
+			Options.SetAllAsOptional();
+			Options.InternalCustomOptions.Emplace(MediaPlayerOptionValues::Environment(), MediaPlayerOptionValues::Environment_Preview());
+			MediaPlayer->OpenSourceWithOptions(MediaSource, Options);
+		}),
 		FCanExecuteAction::CreateLambda([this] { return true; })
 	);
 
diff --git a/Engine/Plugins/Media/MediaPlayerEditor/Source/MediaPlayerEditor/Private/Widgets/SMediaPlayerEditorMedia.cpp b/Engine/Plugins/Media/MediaPlayerEditor/Source/MediaPlayerEditor/Private/Widgets/SMediaPlayerEditorMedia.cpp
--- a/Engine/Plugins/Media/MediaPlayerEditor/Source/MediaPlayerEditor/Private/Widgets/SMediaPlayerEditorMedia.cpp
+++ b/Engine/Plugins/Media/MediaPlayerEditor/Source/MediaPlayerEditor/Private/Widgets/SMediaPlayerEditorMedia.cpp
@@ -79,7 +79,10 @@ void SMediaPlayerEditorMedia::OpenMediaAsset(UObject* Asset)
 
 	if (MediaSource != nullptr)
 	{
-		if (!MediaPlayer->OpenSource(MediaSource))
+		FMediaPlayerOptions Options;
+		Options.SetAllAsOptional();
+		Options.InternalCustomOptions.Emplace(MediaPlayerOptionValues::Environment(), MediaPlayerOptionValues::Environment_Preview());
+		if (!MediaPlayer->OpenSourceWithOptions(MediaSource, Options))
 		{
 			ShowMediaOpenFailedMessage();
 		}
diff --git a/Engine/Plugins/Media/MfMedia/Source/MfMedia/Private/Player/MfMediaTracks.cpp b/Engine/Plugins/Media/MfMedia/Source/MfMedia/Private/Player/MfMediaTracks.cpp
--- a/Engine/Plugins/Media/MfMedia/Source/MfMedia/Private/Player/MfMediaTracks.cpp
+++ b/Engine/Plugins/Media/MfMedia/Source/MfMedia/Private/Player/MfMediaTracks.cpp
@@ -108,7 +108,7 @@ void FMfMediaTracks::AppendStats(FString &OutStats) const
 void FMfMediaTracks::ClearFlags()
 {
 	FScopeLock Lock(&CriticalSection);
-	
+
 	MediaSourceChanged = false;
 	SelectionChanged = false;
 }
@@ -245,7 +245,7 @@ void FMfMediaTracks::Initialize(IMFMediaSource* InMediaSource, IMFSourceReaderCa
 	Algo::Reverse(CaptionTracks);
 	Algo::Reverse(VideoTracks);
 
-	if (PlayerOptions)
+	if (PlayerOptions && PlayerOptions->TrackSelection == EMediaPlayerOptionTrackSelectMode::UseTrackOptionIndices)
 	{
 		// Select tracks based on the options provided
 		SelectTrack(EMediaTrackType::Audio, PlayerOptions->Tracks.Audio);
@@ -365,7 +365,7 @@ void FMfMediaTracks::ProcessSample(IMFSample* Sample, HRESULT Status, DWORD Stre
 			UpdateCaptions();
 		}
 	}
-	
+
 	// process video sample
 	if (SelectedVideoTrack != INDEX_NONE)
 	{
@@ -835,7 +835,7 @@ bool FMfMediaTracks::SelectTrack(EMediaTrackType TrackType, int32 TrackIndex)
 		}
 
 		UE_LOG(LogMfMedia, Verbose, TEXT("Tracks %p: Enabled stream %i"), this, StreamIndex);
-			
+
 		*SelectedTrack = TrackIndex;
 		SelectionChanged = true;
 	}
diff --git a/Engine/Plugins/Media/WmfMedia/Source/WmfMedia/Private/Player/WmfMediaTracks.cpp b/Engine/Plugins/Media/WmfMedia/Source/WmfMedia/Private/Player/WmfMediaTracks.cpp
--- a/Engine/Plugins/Media/WmfMedia/Source/WmfMedia/Private/Player/WmfMediaTracks.cpp
+++ b/Engine/Plugins/Media/WmfMedia/Source/WmfMedia/Private/Player/WmfMediaTracks.cpp
@@ -297,7 +297,7 @@ void FWmfMediaTracks::Initialize(IMFMediaSource* InMediaSource, const FString& U
 
 	// Tracks must be selected before Session->SetTopology() is called
 	FMediaPlayerTrackOptions TrackOptions;
-	if (PlayerOptions)
+	if (PlayerOptions && PlayerOptions->TrackSelection == EMediaPlayerOptionTrackSelectMode::UseTrackOptionIndices)
 	{
 		TrackOptions = PlayerOptions->Tracks;
 	}
diff --git a/Engine/Source/Runtime/MediaAssets/Private/Assets/MediaPlayer.cpp b/Engine/Source/Runtime/MediaAssets/Private/Assets/MediaPlayer.cpp
--- a/Engine/Source/Runtime/MediaAssets/Private/Assets/MediaPlayer.cpp
+++ b/Engine/Source/Runtime/MediaAssets/Private/Assets/MediaPlayer.cpp
@@ -117,7 +117,7 @@ void UMediaPlayer::SetPlaylistInternal(UMediaPlaylist* InPlaylist) const
 {
 	if (Playlist && Playlist != InPlaylist && Playlist->IsRooted())
 	{
-		// To avoid leaking UObjects we need to remove the old playlist from root set 
+		// To avoid leaking UObjects we need to remove the old playlist from root set
 		// (which has been most likely rooted because this MediaPlayer is in disregard for GC set)
 		Playlist->RemoveFromRoot();
 	}
@@ -249,7 +249,7 @@ TMap<FString, FMediaMetadataItemsBPT> UMediaPlayer::GetMediaMetadataItems() cons
 					OutItem.LanguageCode = ItemPtr->GetLanguageCode();
 					OutItem.MimeType = ItemPtr->GetMimeType();
 					FVariant TempItem = ItemPtr->GetValue();
-			
+
 					switch (TempItem.GetType())
 					{
 					case EVariantTypes::String:
@@ -263,7 +263,7 @@ TMap<FString, FMediaMetadataItemsBPT> UMediaPlayer::GetMediaMetadataItems() cons
 							OutItem.BinaryData = TempItem.GetValue<TArray<uint8>>();
 							break;
 						}
-			
+
 					case EVariantTypes::Bool:
 						{
 							OutItem.StringData = TempItem.GetValue<bool>() ? TEXT("true") : TEXT("false");
@@ -602,7 +602,7 @@ bool UMediaPlayer::OpenPlaylistIndex(UMediaPlaylist* InPlaylist, int32 Index)
 	UE_LOG(LogMediaAssets, Verbose, TEXT("%s.OpenSource %s %i"), *GetFName().ToString(), *InPlaylist->GetFName().ToString(), Index);
 
 	SetPlaylistInternal(InPlaylist);
-	
+
 	if (Index == INDEX_NONE)
 	{
 		return true;
@@ -669,10 +669,10 @@ bool UMediaPlayer::OpenSourceInternal(UMediaSource* MediaSource, const FMediaPla
 	{
 		TArray<FAnalyticsEventAttribute> EventAttributes;
 		EventAttributes.Add(FAnalyticsEventAttribute(TEXT("MediaSourceType"), MediaSource->GetClass()->GetName()));
-		FEngineAnalytics::GetProvider().RecordEvent(TEXT("MediaFramework.MediaSourceOpened"), EventAttributes);	
+		FEngineAnalytics::GetProvider().RecordEvent(TEXT("MediaFramework.MediaSourceOpened"), EventAttributes);
 	}
 #endif
-	
+
 	RegisterWithMediaModule();
 	return PlayerFacade->Open(MediaSource->GetUrl(), MediaSource, PlayerOptions);
 }
@@ -722,7 +722,7 @@ void UMediaPlayer::PlayAndSeek()
 	PlayOnNext = false;
 	if (Play())
 	{
-		if (PlayerFacade->ActivePlayerOptions.IsSet() && !PlayerFacade->ActivePlayerOptions->SeekTime.IsZero() && SupportsSeeking())
+		if (PlayerFacade->ActivePlayerOptions.IsSet() && !PlayerFacade->ActivePlayerOptions->SeekTime.IsZero() && SupportsSeeking() && PlayerFacade->ActivePlayerOptions->SeekTimeType == EMediaPlayerOptionSeekTimeType::RelativeToStartTime)
 		{
 			Seek(PlayerFacade->ActivePlayerOptions->SeekTime);
 		}
@@ -1299,7 +1299,7 @@ public:
 					{
 						if (Options.SeekTime < FTimespan::FromSeconds(0) || Options.SeekTime > MediaPlayer->GetDuration())
 						{
-							UE_LOG(LogMediaAssets, Warning, TEXT("Open Media Latent: Media player seeking to time out of bounds. Seek: %s, Duration: %s, URL: %s"), 
+							UE_LOG(LogMediaAssets, Warning, TEXT("Open Media Latent: Media player seeking to time out of bounds. Seek: %s, Duration: %s, URL: %s"),
 								*Options.SeekTime.ToString(), *MediaPlayer->GetDuration().ToString(), *URL);
 							FailedOperation(Response);
 							return;
diff --git a/Engine/Source/Runtime/MediaUtils/Private/MediaPlayerFacade.cpp b/Engine/Source/Runtime/MediaUtils/Private/MediaPlayerFacade.cpp
--- a/Engine/Source/Runtime/MediaUtils/Private/MediaPlayerFacade.cpp
+++ b/Engine/Source/Runtime/MediaUtils/Private/MediaPlayerFacade.cpp
@@ -1400,7 +1400,7 @@ const TRange<FMediaTimeStamp>& FMediaPlayerFacade::FBlockOnRange::GetRange() con
 					// Yes. Adjust the base secondary index value to guarantee a simple continuation of the secondary index values
 					int32 LastLoopIdxS = (int32)FMath::FloorToInt(LastTimeRange.GetLowerBoundValue().GetTotalSeconds() / Duration.GetTotalSeconds());
 					// -LoopIdxS -> Compensate for new start index
-					// +LastLoopIdxS + 1 -> Move index one beyond last 
+					// +LastLoopIdxS + 1 -> Move index one beyond last
 					OnBlockSecondaryIndexOffset = -LoopIdxS + LastLoopIdxS - 1;
 				}
 			}
@@ -1994,13 +1994,16 @@ void FMediaPlayerFacade::SelectDefaultTracks()
 		// If overrides are set, use them.
 		if (ActivePlayerOptions.IsSet())
 		{
-			FMediaPlayerTrackOptions TrackOptions;
-			TrackOptions = ActivePlayerOptions.GetValue().Tracks;
-			TrackSelection.UserSelection[(int32)EMediaTrackType::Audio] = TrackOptions.Audio;
-			TrackSelection.UserSelection[(int32)EMediaTrackType::Caption] = TrackOptions.Caption;
-			TrackSelection.UserSelection[(int32)EMediaTrackType::Metadata] = TrackOptions.Metadata;
-			TrackSelection.UserSelection[(int32)EMediaTrackType::Subtitle] = TrackOptions.Subtitle;
-			TrackSelection.UserSelection[(int32)EMediaTrackType::Video] = TrackOptions.Video;
+			if (ActivePlayerOptions.GetValue().TrackSelection == EMediaPlayerOptionTrackSelectMode::UseTrackOptionIndices)
+			{
+				FMediaPlayerTrackOptions TrackOptions;
+				TrackOptions = ActivePlayerOptions.GetValue().Tracks;
+				TrackSelection.UserSelection[(int32)EMediaTrackType::Audio] = TrackOptions.Audio;
+				TrackSelection.UserSelection[(int32)EMediaTrackType::Caption] = TrackOptions.Caption;
+				TrackSelection.UserSelection[(int32)EMediaTrackType::Metadata] = TrackOptions.Metadata;
+				TrackSelection.UserSelection[(int32)EMediaTrackType::Subtitle] = TrackOptions.Subtitle;
+				TrackSelection.UserSelection[(int32)EMediaTrackType::Video] = TrackOptions.Video;
+			}
 		}
 	}
 	else
@@ -2008,7 +2011,10 @@ void FMediaPlayerFacade::SelectDefaultTracks()
 		FMediaPlayerTrackOptions TrackOptions;
 		if (ActivePlayerOptions.IsSet())
 		{
-			TrackOptions = ActivePlayerOptions.GetValue().Tracks;
+			if (ActivePlayerOptions.GetValue().TrackSelection == EMediaPlayerOptionTrackSelectMode::UseTrackOptionIndices)
+			{
+				TrackOptions = ActivePlayerOptions.GetValue().Tracks;
+			}
 		}
 
 		TrackSelection.UserSelection[(int32)EMediaTrackType::Audio] = TrackOptions.Audio;
@@ -3020,12 +3026,12 @@ bool FMediaPlayerFacade::IsVideoSampleStillGood(const TRange<FMediaTimeStamp>& L
 																	: TRange<FMediaTimeStamp>(LastSampleTimeRange0.GetLowerBoundValue() - LastSampleTimeRange0.Size<FMediaTimeStamp>().Time, LastSampleTimeRange0.GetLowerBoundValue());
 
 			// Note: Loops (or the end of the time line in non-looping setups)
-			// 
+			//
 			// - We could check for them and generate proper changes to the sequence index
 			// - Doing this would leave us with quite complex setups to compute the coverage
 			// - We opt for a cleaner, simpler approach: as we are NOT interested into proper PTS values, we can safely work with an "infinite" time line when computing any overlaps, coverage and such
 			//   (note: we DO need to restrict the range to the actual media duration if not looping - the code above does this)
-			// 
+			//
 			// --> we simply keep what we compute above!
 			//
 
diff --git a/Engine/Source/Runtime/MediaUtils/Public/MediaPlayerOptions.h b/Engine/Source/Runtime/MediaUtils/Public/MediaPlayerOptions.h
--- a/Engine/Source/Runtime/MediaUtils/Public/MediaPlayerOptions.h
+++ b/Engine/Source/Runtime/MediaUtils/Public/MediaPlayerOptions.h
@@ -5,6 +5,8 @@
 #include "CoreTypes.h"
 #include "UObject/ObjectMacros.h"
 #include "UObject/ScriptMacros.h"
+#include "Containers/Map.h"
+#include "Misc/Variant.h"
 
 #include "MediaPlayerOptions.generated.h"
 
@@ -16,6 +18,27 @@ enum class EMediaPlayerOptionBooleanOverride : uint8
 	Disabled
 };
 
+
+UENUM(BlueprintType)
+enum class EMediaPlayerOptionSeekTimeType : uint8
+{
+	// Ignore the given value and lets the media player choose.
+	Ignored,
+	// Given seek time is relative to the start of the media.
+	RelativeToStartTime
+};
+
+
+UENUM(BlueprintType)
+enum class EMediaPlayerOptionTrackSelectMode : uint8
+{
+	// Let the media player choose defaults.
+	UseMediaPlayerDefaults UMETA(DisplayName="Media player selects default tracks"),
+	// Use fixed track indices as specified with MediaPlayerTrackOptions
+	UseTrackOptionIndices UMETA(DisplayName="Uses provided track indices")
+};
+
+
 USTRUCT(BlueprintType)
 struct FMediaPlayerTrackOptions
 {
@@ -44,25 +67,25 @@ struct FMediaPlayerTrackOptions
 			|| Video != Other.Video);
 	}
 
-	UPROPERTY(BlueprintReadWrite, Category = "Tracks")
+	UPROPERTY(BlueprintReadWrite, EditAnywhere, Category = "Tracks")
 	int32 Audio;
 
-	UPROPERTY(BlueprintReadWrite, Category = "Tracks")
+	UPROPERTY(BlueprintReadWrite, EditAnywhere, Category = "Tracks")
 	int32 Caption;
 
-	UPROPERTY(BlueprintReadWrite, Category = "Tracks")
+	UPROPERTY(BlueprintReadWrite, EditAnywhere, Category = "Tracks")
 	int32 Metadata;
 
-	UPROPERTY(BlueprintReadWrite, Category = "Tracks")
+	UPROPERTY(BlueprintReadWrite, EditAnywhere, Category = "Tracks")
 	int32 Script;
 
-	UPROPERTY(BlueprintReadWrite, Category = "Tracks")
+	UPROPERTY(BlueprintReadWrite, EditAnywhere, Category = "Tracks")
 	int32 Subtitle;
 
-	UPROPERTY(BlueprintReadWrite, Category = "Tracks")
+	UPROPERTY(BlueprintReadWrite, EditAnywhere, Category = "Tracks")
 	int32 Text;
 
-	UPROPERTY(BlueprintReadWrite, Category = "Tracks")
+	UPROPERTY(BlueprintReadWrite, EditAnywhere, Category = "Tracks")
 	int32 Video;
 };
 
@@ -72,30 +95,69 @@ struct FMediaPlayerOptions
 	GENERATED_BODY()
 
 	FMediaPlayerOptions() :
+		TrackSelection(EMediaPlayerOptionTrackSelectMode::UseTrackOptionIndices),
 		SeekTime(0),
+		SeekTimeType(EMediaPlayerOptionSeekTimeType::RelativeToStartTime),
 		PlayOnOpen(EMediaPlayerOptionBooleanOverride::UseMediaPlayerSetting),
 		Loop(EMediaPlayerOptionBooleanOverride::UseMediaPlayerSetting)
 	{
 	}
 
+	void SetAllAsOptional()
+	{
+		SeekTime = FTimespan::MinValue();
+		TrackSelection = EMediaPlayerOptionTrackSelectMode::UseMediaPlayerDefaults;
+		SeekTimeType = EMediaPlayerOptionSeekTimeType::Ignored;
+		PlayOnOpen = EMediaPlayerOptionBooleanOverride::UseMediaPlayerSetting;
+		Loop = EMediaPlayerOptionBooleanOverride::UseMediaPlayerSetting;
+	}
+
 	/** Check if this Media Options is not equal to another one. */
 	bool operator!=(const FMediaPlayerOptions& Other) const
 	{
 		return (Tracks != Other.Tracks
+			|| TrackSelection != Other.TrackSelection
 			|| SeekTime != Other.SeekTime
+			|| SeekTimeType != Other.SeekTimeType
 			|| PlayOnOpen != Other.PlayOnOpen
 			|| Loop != Other.Loop);
 	}
 
-	UPROPERTY(BlueprintReadWrite, Category = "Tracks")
+	/** Fixed indices of media tracks to select. */
+	UPROPERTY(BlueprintReadWrite, EditAnywhere, Category = "Tracks", meta=(DisplayName="Initial track indices", ToolTip="Indices of the media tracks to select for playback", EditCondition="TrackSelection==EMediaPlayerOptionTrackSelectMode::UseTrackOptionIndices"))
 	FMediaPlayerTrackOptions Tracks;
 
-	UPROPERTY(BlueprintReadWrite, Category = "Misc")
+	/** How to select the initial media tracks. */
+	UPROPERTY(BlueprintReadWrite, EditAnywhere, Category = "Tracks", meta=(DisplayName="Initial track selection mode", ToolTip="How the initial media tracks for playback are selected"))
+	EMediaPlayerOptionTrackSelectMode TrackSelection;
+
+	/** Initial media time to start playback at. */
+	UPROPERTY(BlueprintReadWrite, EditAnywhere, Category = "Misc", meta=(EditCondition="SeekTimeType!=EMediaPlayerOptionSeekTimeType::Ignored"))
 	FTimespan SeekTime;
 
-	UPROPERTY(BlueprintReadWrite, Category = "Misc")
+	/** How to interpret the initial seek time. */
+	UPROPERTY(BlueprintReadWrite, EditAnywhere, Category = "Misc", meta=(DisplayName="Seek time interpretation"))
+	EMediaPlayerOptionSeekTimeType SeekTimeType;
+
+	/** How to handle automatic playback when media opens. */
+	UPROPERTY(BlueprintReadWrite, EditAnywhere, Category = "Misc")
 	EMediaPlayerOptionBooleanOverride PlayOnOpen;
 
-	UPROPERTY(BlueprintReadWrite, Category = "Misc")
+	/** How to initially select looping of the media. */
+	UPROPERTY(BlueprintReadWrite, EditAnywhere, Category = "Misc")
 	EMediaPlayerOptionBooleanOverride Loop;
+
+
+	/** Custom options used internally. Must not be serialized or editable via blueprint. */
+	TMap<FName, FVariant> InternalCustomOptions;
 };
+
+namespace MediaPlayerOptionValues
+{
+	inline const FName& Environment() 
+	{ static FName OptName(TEXT("Environment")); return OptName; }
+	inline const FName& Environment_Sequencer() 
+	{ static FName OptName(TEXT("Sequencer")); return OptName; }
+	inline const FName& Environment_Preview() 
+	{ static FName OptName(TEXT("Preview")); return OptName; }
+}
