ElectraPlayer: tracking dropped decoded frames for frame accurate seeking

#rnx

diff --git a/Engine/Plugins/Media/ElectraPlayer/Source/ElectraPlayerRuntime/Private/Runtime/Decoder/VideoDecoder.cpp b/Engine/Plugins/Media/ElectraPlayer/Source/ElectraPlayerRuntime/Private/Runtime/Decoder/VideoDecoder.cpp
--- a/Engine/Plugins/Media/ElectraPlayer/Source/ElectraPlayerRuntime/Private/Runtime/Decoder/VideoDecoder.cpp
+++ b/Engine/Plugins/Media/ElectraPlayer/Source/ElectraPlayerRuntime/Private/Runtime/Decoder/VideoDecoder.cpp
@@ -178,6 +178,10 @@ private:
 	bool																	bWaitForSyncSample = true;
 	bool																	bWarnedMissingSyncSample = false;
 
+	int32																	NumInitialSkippedFrames = 0;
+	int32																	NumInitialSkippedDecodingFrames = 0;
+	bool																	bIsStartOfSequence = true;
+
 	bool																	bError = false;
 
 	FMediaEvent																TerminateThreadSignal;
@@ -849,9 +853,10 @@ IElectraDecoder::EOutputStatus FVideoDecoderImpl::HandleOutput()
 			}
 			else
 			{
-				// Get the transfer handle as the current decoder output.
-				check(CurrentOutputBuffer == nullptr);
+				// If we did not return the last buffer to the renderer, do it now.
+				// Transfer handles are unique to the buffer and cannot be used for different frames.
 				ReturnUnusedOutputBuffer();
+				// Get the transfer handle as the current decoder output.
 				CurrentOutputBuffer = reinterpret_cast<IMediaRenderer::IBuffer*>(CurrentDecoderOutput->GetTransferHandle()->GetHandle());
 				CurrentDecoderOutput->GetTransferHandle()->ReleaseHandle();
 			}
@@ -914,24 +919,32 @@ IElectraDecoder::EOutputStatus FVideoDecoderImpl::HandleOutput()
 					return IElectraDecoder::EOutputStatus::Error;
 				}
 
-				// Create the platform specific decoder output.
-				TSharedPtr<FParamDict, ESPMode::ThreadSafe> BufferProperties(new FParamDict);
-				BufferProperties->Set(RenderOptionKeys::PTS, FVariantValue(MatchingInput->AdjustedPTS));
-				BufferProperties->Set(RenderOptionKeys::Duration, FVariantValue(MatchingInput->AdjustedDuration));
+				bUseOutput = bUseOutput ? MatchingInput->AdjustedPTS.IsValid() : false;
+				if (bUseOutput)
+				{
+					// Create the platform specific decoder output.
+					TSharedPtr<FParamDict, ESPMode::ThreadSafe> BufferProperties(new FParamDict);
+					BufferProperties->Set(RenderOptionKeys::PTS, FVariantValue(MatchingInput->AdjustedPTS));
+					BufferProperties->Set(RenderOptionKeys::Duration, FVariantValue(MatchingInput->AdjustedDuration));
 
-				// Set properties from the bitstream messages.
-				BitstreamProcessor->SetPropertiesOnOutput(CurrentDecoderOutput, BufferProperties.Get(), MatchingInput->BitstreamInfo);
+					// Set properties from the bitstream messages.
+					BitstreamProcessor->SetPropertiesOnOutput(CurrentDecoderOutput, BufferProperties.Get(), MatchingInput->BitstreamInfo);
 
-				if (bUseOutput && !FPlatformElectraDecoderResourceManager::SetupRenderBufferFromDecoderOutput(CurrentOutputBuffer, BufferProperties, CurrentDecoderOutput, PlatformResource))
-				{
-					PostError(0, TEXT("Failed to set up the decoder output!"), ERRCODE_VIDEO_INTERNAL_FAILED_TO_CONVERT_OUTPUT);
-					return IElectraDecoder::EOutputStatus::Error;
-				}
+					if (!FPlatformElectraDecoderResourceManager::SetupRenderBufferFromDecoderOutput(CurrentOutputBuffer, BufferProperties, CurrentDecoderOutput, PlatformResource))
+					{
+						PostError(0, TEXT("Failed to set up the decoder output!"), ERRCODE_VIDEO_INTERNAL_FAILED_TO_CONVERT_OUTPUT);
+						return IElectraDecoder::EOutputStatus::Error;
+					}
+					Renderer->ReturnBuffer(CurrentOutputBuffer, bUseOutput, *BufferProperties);
+					CurrentOutputBuffer = nullptr;
 
-				bUseOutput = bUseOutput ? MatchingInput->AdjustedPTS.IsValid() : false;
-				Renderer->ReturnBuffer(CurrentOutputBuffer, bUseOutput, *BufferProperties);
+					if (bIsStartOfSequence && NumInitialSkippedFrames)
+					{
+						UE_LOG(LogElectraPlayer, Verbose, TEXT("Frame accurate seek skipped %d leading frames of which %d had to be decoded"), NumInitialSkippedFrames, NumInitialSkippedFrames-NumInitialSkippedDecodingFrames);
+						bIsStartOfSequence = false;
+					}
+				}
 				CurrentDecoderOutput.Reset();
-				CurrentOutputBuffer = nullptr;
 			}
 		}
 	}
@@ -956,6 +969,8 @@ FVideoDecoderImpl::ENextDecodingState FVideoDecoderImpl::HandleDecoding()
 			// not decoding dummy data the decoder must be drained to get the last decoded data out.
 			bDrainAfterDecode = CurrentAccessUnit->AccessUnit->bIsLastInPeriod && !bInDummyDecodeMode;
 			CurrentAccessUnit.Reset();
+			NumInitialSkippedFrames += bIsStartOfSequence ? 1 : 0;
+			NumInitialSkippedDecodingFrames += bIsStartOfSequence ? 1 : 0;
 			return ENextDecodingState::NormalDecoding;
 		}
 
@@ -996,6 +1011,8 @@ FVideoDecoderImpl::ENextDecodingState FVideoDecoderImpl::HandleDecoding()
 			}
 			if (bSupportsDroppingOutput && !CurrentAccessUnit->AdjustedPTS.IsValid())
 			{
+				NumInitialSkippedFrames += bIsStartOfSequence ? 1 : 0;
+				NumInitialSkippedDecodingFrames += CurrentAccessUnit->BitstreamInfo.bIsDiscardable ? 1 : 0;
 				DecAU.Flags |= EElectraDecoderFlags::DoNotOutput;
 			}
 			TMap<FString, FVariant> CSDOptions;
@@ -1279,6 +1296,9 @@ bool FVideoDecoderImpl::CheckForFlush()
 		bWaitForSyncSample = true;
 		bWarnedMissingSyncSample = false;
 		CurrentDecodingState = EDecodingState::NormalDecoding;
+		NumInitialSkippedFrames = 0;
+		NumInitialSkippedDecodingFrames = 0;
+		bIsStartOfSequence = true;
 		BitstreamProcessor->Clear();
 		FlushDecoderSignal.Reset();
 		DecoderFlushedSignal.Signal();
@@ -1333,6 +1353,11 @@ void FVideoDecoderImpl::WorkerThread()
 	bWarnedMissingSyncSample = false;
 	CurrentDecodingState = EDecodingState::NormalDecoding;
 
+	// Clear initial skip frame stats
+	NumInitialSkippedFrames = 0;
+	NumInitialSkippedDecodingFrames = 0;
+	bIsStartOfSequence = true;
+
 	check(InitialCodecSpecificData.IsValid());
 	if (InitialCodecSpecificData.IsValid())
 	{
