ElectraPlayer: Added human-readable name for some codecs displayed in the editor.

#rnx

diff --git a/Engine/Plugins/Media/ElectraPlayer/Source/ElectraPlayerRuntime/Private/Runtime/StreamTypes.cpp b/Engine/Plugins/Media/ElectraPlayer/Source/ElectraPlayerRuntime/Private/Runtime/StreamTypes.cpp
--- a/Engine/Plugins/Media/ElectraPlayer/Source/ElectraPlayerRuntime/Private/Runtime/StreamTypes.cpp
+++ b/Engine/Plugins/Media/ElectraPlayer/Source/ElectraPlayerRuntime/Private/Runtime/StreamTypes.cpp
@@ -366,6 +366,154 @@ namespace Electra
 		}
 	}
 
+	const FString& FStreamCodecInformation::GetHumanReadableCodecName() const
+	{
+		if (HumanReadableCodecName.IsEmpty())
+		{
+			if (!TryConstructHumanReadableCodecName())
+			{
+				HumanReadableCodecName = CodecSpecifier;
+			}
+		}
+		return HumanReadableCodecName;
+	}
+
+	bool FStreamCodecInformation::TryConstructHumanReadableCodecName() const
+	{
+		switch(GetCodec())
+		{
+			case ECodec::H264:
+			{
+				HumanReadableCodecName = TEXT("AVC (H.264)");
+				if (ProfileLevel.Profile == 66)
+				{
+					HumanReadableCodecName.Append(TEXT(", Baseline"));
+				}
+				else if (ProfileLevel.Profile == 77)
+				{
+					HumanReadableCodecName.Append(TEXT(", Main"));
+				}
+				else if (ProfileLevel.Profile == 100)
+				{
+					HumanReadableCodecName.Append(TEXT(", High"));
+				}
+				else
+				{
+					HumanReadableCodecName.Append(TEXT(", Unknown profile"));
+				}
+				HumanReadableCodecName.Append(FString::Printf(TEXT(", level %d.%d"), ProfileLevel.Level/10, ProfileLevel.Level%10));
+				return true;
+			}
+			case ECodec::H265:
+			{
+				HumanReadableCodecName = TEXT("HEVC (H.265)");
+				if (ProfileLevel.Profile == 1)
+				{
+					HumanReadableCodecName.Append(TEXT(", Main"));
+				}
+				else if (ProfileLevel.Profile == 2)
+				{
+					HumanReadableCodecName.Append(TEXT(", Main10"));
+				}
+				else
+				{
+					HumanReadableCodecName.Append(TEXT(", Unknown profile"));
+				}
+				HumanReadableCodecName.Append(FString::Printf(TEXT(", level %d.%d"), ProfileLevel.Level/30, ProfileLevel.Level%30));
+				return true;
+			}
+			case ECodec::Video4CC:
+			{
+				switch(GetCodec4CC())
+				{
+					case Make4CC('v','p','0','8'):
+						HumanReadableCodecName = TEXT("VP8");
+						return true;
+					case Make4CC('v','p','0','9'):
+						HumanReadableCodecName = TEXT("VP9");
+						return true;
+
+					case Make4CC('a','p','c','h'):
+						HumanReadableCodecName = TEXT("Apple ProRes 422 High Quality");
+						return true;
+					case Make4CC('a','p','c','n'):
+						HumanReadableCodecName = TEXT("Apple ProRes 422 Standard Definition");
+						return true;
+					case Make4CC('a','p','c','s'):
+						HumanReadableCodecName = TEXT("Apple ProRes 422 LT");
+						return true;
+					case Make4CC('a','p','c','o'):
+						HumanReadableCodecName = TEXT("Apple ProRes 422 Proxy");
+						return true;
+					case Make4CC('a','p','4','h'):
+						HumanReadableCodecName = TEXT("Apple ProRes 4444");
+						return true;
+
+					case Make4CC('H','a','p','1'):
+						HumanReadableCodecName = TEXT("Hap");
+						return true;
+					case Make4CC('H','a','p','5'):
+						HumanReadableCodecName = TEXT("Hap Alpha");
+						return true;
+					case Make4CC('H','a','p','Y'):
+						HumanReadableCodecName = TEXT("Hap Q");
+						return true;
+					case Make4CC('H','a','p','M'):
+						HumanReadableCodecName = TEXT("Hap Q Alpha");
+						return true;
+					case Make4CC('H','a','p','7'):
+						HumanReadableCodecName = TEXT("Hap R");
+						return true;
+					case Make4CC('H','a','p','H'):
+						HumanReadableCodecName = TEXT("Hap HDR");
+						return true;
+
+					case Make4CC('A','V','d','h'):
+						HumanReadableCodecName = TEXT("Avid DNxHD");
+						return true;
+				}
+				HumanReadableCodecName = Printable4CC(GetCodec4CC());
+				return true;
+			}
+			case ECodec::AAC:
+			{
+				HumanReadableCodecName = TEXT("MPEG AAC");
+				return true;
+			}
+			case ECodec::EAC3:
+			{
+				HumanReadableCodecName = TEXT("Dolby Digital");
+				return true;
+			}
+			case ECodec::Audio4CC:
+			{
+				switch(GetCodec4CC())
+				{
+					case Make4CC('O','p','u','s'):
+						HumanReadableCodecName = TEXT("Opus");
+						return true;
+				}
+				HumanReadableCodecName = Printable4CC(GetCodec4CC());
+				return true;
+			}
+			case ECodec::WebVTT:
+			{
+				HumanReadableCodecName = TEXT("WebVTT");
+				return true;
+			}
+			case ECodec::TTML:
+			{
+				HumanReadableCodecName = TEXT("TTML");
+				return true;
+			}
+			case ECodec::TX3G:
+			{
+				HumanReadableCodecName = TEXT("SRT/TX3G");
+				return true;
+			}
+		}
+		return false;
+	}
 
 
 
@@ -508,7 +656,7 @@ namespace Electra
 		}
 		return true;
 	}
-	
+
 	int32 FCodecSelectionPriorities::GetClassPriority(const FString& CodecSpecifierRFC6381) const
 	{
 		for(auto &CodecClass : ClassPriorities)
@@ -520,7 +668,7 @@ namespace Electra
 		}
 		return -1;
 	}
-	
+
 	int32 FCodecSelectionPriorities::GetStreamPriority(const FString& CodecSpecifierRFC6381) const
 	{
 		for(auto &CodecClass : ClassPriorities)
diff --git a/Engine/Plugins/Media/ElectraPlayer/Source/ElectraPlayerRuntime/Private/Runtime/StreamTypes.h b/Engine/Plugins/Media/ElectraPlayer/Source/ElectraPlayerRuntime/Private/Runtime/StreamTypes.h
--- a/Engine/Plugins/Media/ElectraPlayer/Source/ElectraPlayerRuntime/Private/Runtime/StreamTypes.h
+++ b/Engine/Plugins/Media/ElectraPlayer/Source/ElectraPlayerRuntime/Private/Runtime/StreamTypes.h
@@ -81,7 +81,7 @@ namespace Electra
 		{
 			StreamType = InStreamType;
 		}
-		
+
 		uint32 GetCodec4CC() const
 		{
 			return Codec4CC;
@@ -90,6 +90,7 @@ namespace Electra
 		void SetCodec4CC(uint32 In4CC)
 		{
 			Codec4CC = In4CC;
+			ResetHumanReadableCodecName();
 		}
 
 		ECodec GetCodec() const
@@ -100,6 +101,7 @@ namespace Electra
 		void SetCodec(ECodec InCodec)
 		{
 			Codec = InCodec;
+			ResetHumanReadableCodecName();
 		}
 
 		FString GetCodecName() const;
@@ -152,6 +154,14 @@ namespace Electra
 		void SetCodecSpecifierRFC6381(const FString& InCodecSpecifier)
 		{
 			CodecSpecifier = InCodecSpecifier;
+			ResetHumanReadableCodecName();
+		}
+
+		const FString& GetHumanReadableCodecName() const;
+
+		void SetHumanReadableCodecName(const FString& InHumanReadableCodecName)
+		{
+			HumanReadableCodecName = InHumanReadableCodecName;
 		}
 
 		struct FResolution
@@ -208,6 +218,7 @@ namespace Electra
 		void SetResolution(const FResolution& InResolution)
 		{
 			Resolution = InResolution;
+			ResetHumanReadableCodecName();
 		}
 
 
@@ -343,6 +354,7 @@ namespace Electra
 		void SetProfileSpace(int32 InProfileSpace)
 		{
 			ProfileLevel.ProfileSpace = InProfileSpace;
+			ResetHumanReadableCodecName();
 		}
 
 		int32 GetProfileSpace() const
@@ -353,6 +365,7 @@ namespace Electra
 		void SetProfile(int32 InProfile)
 		{
 			ProfileLevel.Profile = InProfile;
+			ResetHumanReadableCodecName();
 		}
 
 		int32 GetProfile() const
@@ -363,6 +376,7 @@ namespace Electra
 		void SetProfileLevel(int32 InLevel)
 		{
 			ProfileLevel.Level = InLevel;
+			ResetHumanReadableCodecName();
 		}
 
 		int32 GetProfileLevel() const
@@ -373,6 +387,7 @@ namespace Electra
 		void SetProfileCompatibilityFlags(uint32 InCompatibilityFlags)
 		{
 			ProfileLevel.CompatibilityFlags = InCompatibilityFlags;
+			ResetHumanReadableCodecName();
 		}
 
 		uint32 GetProfileCompatibilityFlags() const
@@ -383,6 +398,7 @@ namespace Electra
 		void SetProfileConstraints(uint64 InConstraints)
 		{
 			ProfileLevel.Constraints = InConstraints;
+			ResetHumanReadableCodecName();
 		}
 
 		uint64 GetProfileConstraints() const
@@ -393,6 +409,7 @@ namespace Electra
 		void SetProfileTier(int32 InTier)
 		{
 			ProfileLevel.Tier = InTier;
+			ResetHumanReadableCodecName();
 		}
 
 		int32 GetProfileTier() const
@@ -403,6 +420,7 @@ namespace Electra
 		void SetSamplingRate(int32 InSamplingRate)
 		{
 			SampleRate = InSamplingRate;
+			ResetHumanReadableCodecName();
 		}
 
 		int32 GetSamplingRate() const
@@ -413,6 +431,7 @@ namespace Electra
 		void SetNumberOfChannels(int32 InNumberOfChannels)
 		{
 			NumChannels = InNumberOfChannels;
+			ResetHumanReadableCodecName();
 		}
 
 		int32 GetNumberOfChannels() const
@@ -423,6 +442,7 @@ namespace Electra
 		void SetChannelConfiguration(uint32 InChannelConfiguration)
 		{
 			ChannelConfiguration = InChannelConfiguration;
+			ResetHumanReadableCodecName();
 		}
 
 		uint32 GetChannelConfiguration() const
@@ -433,6 +453,7 @@ namespace Electra
 		void SetAudioDecodingComplexity(int InAudioDecodingComplexity)
 		{
 			AudioDecodingComplexity = InAudioDecodingComplexity;
+			ResetHumanReadableCodecName();
 		}
 
 		int32 GetAudioDecodingComplexity() const
@@ -443,6 +464,7 @@ namespace Electra
 		void SetAudioAccessibility(int32 InAudioAccessibility)
 		{
 			AudioAccessibility = InAudioAccessibility;
+			ResetHumanReadableCodecName();
 		}
 
 		int32 GetAudioAccessibility() const
@@ -453,6 +475,7 @@ namespace Electra
 		void SetNumberOfAudioObjects(int32 InNumberOfAudioObjects)
 		{
 			NumberOfAudioObjects = InNumberOfAudioObjects;
+			ResetHumanReadableCodecName();
 		}
 
 		int32 GetNumberOfAudioObjects() const
@@ -494,7 +517,7 @@ namespace Electra
 		{
 			Bitrate = InBitrate;
 		}
-		
+
 		int32 GetBitrate() const
 		{
 			return Bitrate;
@@ -535,12 +558,7 @@ namespace Electra
 			CSD.Empty();
 			DCR.Empty();
 			CodecVideoColorInfo.Reset();
-		}
-
-		bool IsDifferentFromOtherVideo(const FStreamCodecInformation& Other) const
-		{
-			// We do not compare frame rate here as we are interested in values that may require a decoder reconfiguration.
-			return Codec != Other.Codec || Resolution != Other.Resolution || ProfileLevel != Other.ProfileLevel || AspectRatio != Other.AspectRatio;
+			ResetHumanReadableCodecName();
 		}
 
 		bool Equals(const FStreamCodecInformation& Other) const
@@ -599,6 +617,12 @@ namespace Electra
 
 
 	private:
+		void ResetHumanReadableCodecName()
+		{
+			HumanReadableCodecName.Empty();
+		}
+		bool TryConstructHumanReadableCodecName() const;
+
 		struct FProfileLevel
 		{
 			FProfileLevel()
@@ -655,6 +679,7 @@ namespace Electra
 		TArray<uint8>	CSD;						//!< Codec specific data, if available.
 		TArray<uint8>	DCR;						//!< Decoder configuration record, if available.
 		FCodecVideoColorInfo CodecVideoColorInfo;
+		mutable FString	HumanReadableCodecName;
 	};
 
 
@@ -819,7 +844,7 @@ namespace Electra
 	};
 
 
-	
+
 	class FCodecSelectionPriorities
 	{
 	public:
@@ -835,16 +860,16 @@ namespace Electra
 		 *   CLASSWITHPRIO = CODECPRIO 0*[ { CODECPRIO 0*[ COMMA CODECPRIO ] } ]
 		 *   CLASSWITHOUTPRIO = CLASS 1*[ { CODECPRIO 0*[ COMMA CODECPRIO ] } ]
 		 *   CLASSPRIO = CLASSWITHPRIO / CLASSWITHOUTPRIO
-		 * 
+		 *
 		 * Examples: hvc=2,hev=2,avc=1
 		 *           mp4a{mp4a.40.5=0,mp4a.40.2=1}
-		 * 
+		 *
 		 * First codec priorities are given for an entire codec class (eg. "hvc").
 		 * Within each class, where it makes sense, individual streams can be prioritized.
 		 * Say within a class "mp4a" there are two AAC streams. One LC and one HE.
 		 * To use the LC over the HE stream the "mp4a" class gives more detailed codec
 		 * prefixes and their priorities like the above example.
-		 * 
+		 *
 		 * If used with DASH streams the class priority can be thought of the priority
 		 * of an AdaptationSet and the stream priority of that of a Representation.
 		 * User defined priorities override the @selectionPriority attribute of a
