ElectraPlayer: added 'rap ' sample group handling to base mp4 parsing; removed 'elst' offset from the DTS timestamps as they affect only the PTS; removed the length value from some Pascal strings as this is not a printable character

#rnx

diff --git a/Engine/Plugins/Media/ElectraPlayer/Source/ElectraPlayerRuntime/Private/Runtime/Demuxer/ParserISO14496-12.cpp b/Engine/Plugins/Media/ElectraPlayer/Source/ElectraPlayerRuntime/Private/Runtime/Demuxer/ParserISO14496-12.cpp
--- a/Engine/Plugins/Media/ElectraPlayer/Source/ElectraPlayerRuntime/Private/Runtime/Demuxer/ParserISO14496-12.cpp
+++ b/Engine/Plugins/Media/ElectraPlayer/Source/ElectraPlayerRuntime/Private/Runtime/Demuxer/ParserISO14496-12.cpp
@@ -173,7 +173,7 @@
 			// Do a special test here to check if the string is perhaps a Pascal string (first byte is length, not NUL terminated)
 			if (FirstByte + 1 == MaxBytes)
 			{
-				// NOTE: We do _not_ remove the length from the string!
+				Buf.RemoveAt(0);
 				OutString = ConvFromBuf(Buf, bIsUTF8);
 				return UEMEDIA_ERROR_OK;
 			}
@@ -6349,6 +6349,47 @@
 	{
 		// NOTE: All boxes that are referenced here have been checked to exist so accessing them is safe.
 
+		auto SetupRAPGroup = [&]() -> void
+		{
+			// Sample groups
+			for(int32 i=0; i<Track->SBGPBoxes.Num(); ++i)
+			{
+				switch(Track->SBGPBoxes[i]->GetGroupingType())
+				{
+					// Random access
+					case EGroupType::Grouping_rap:
+					{
+						// ISO/IEC 14496-15 disallows `rap ` boxes greater than version 0.
+						if (Track->SBGPBoxes[i]->GetBoxVersion() == 0)
+						{
+							// There needs to be a matching entry
+							const FMP4BoxSGPD* const* SGPDBox = Track->SGPDBoxes.FindByPredicate([](const FMP4BoxSGPD* In) { return In->GetGroupingType() == EGroupType::Grouping_rap; });
+							if (SGPDBox && (*SGPDBox)->GetNumberOfEntries() > 0)
+							{
+								// We do not care about the actual `VisualRandomAccessEntry` for now, only the fact that there is one.
+								FSampleGroupIterator& sgit = GroupIterators.Add(EGroupType::Grouping_rap);
+								sgit.GroupingType = EGroupType::Grouping_rap;
+								sgit.SBGPBox = Track->SBGPBoxes[i];
+								sgit.DefaultDescriptionIndex = (*SGPDBox)->GetDefaultSampleDescriptionIndex();
+								sgit.SetFirst();
+								// If the sample group says this sample is a RAP then clear the `sample_is_non_sync_sample` flag.
+								if (sgit.GetCurrentGroupDescriptionIndex() != 0)
+								{
+									SampleFlag &= ~0x10000U;
+								}
+								bHaveRAPGroup = true;
+							}
+						}
+						break;
+					}
+					default:
+					{
+						break;
+					}
+				}
+			}
+		};
+
 		check(Track);
 		if (!Track)
 		{
@@ -6432,43 +6473,8 @@
 			SampleDescriptionIndex = Track->TFHDBox->HasSampleDescriptionIndex() ? Track->TFHDBox->GetSampleDescriptionIndex() :
 				Track->TREXBox->GetDefaultSampleDescriptionIndex();
 
-			// Sample groups
-			for(int32 i=0; i<Track->SBGPBoxes.Num(); ++i)
-			{
-				switch(Track->SBGPBoxes[i]->GetGroupingType())
-				{
-					// Random access
-					case EGroupType::Grouping_rap:
-					{
-						// ISO/IEC 14496-15 disallows `rap ` boxes greater than version 0.
-						if (Track->SBGPBoxes[i]->GetBoxVersion() == 0)
-						{
-							// There needs to be a matching entry
-							const FMP4BoxSGPD* const* SGPDBox = Track->SGPDBoxes.FindByPredicate([](const FMP4BoxSGPD* In) { return In->GetGroupingType() == EGroupType::Grouping_rap; });
-							if (SGPDBox && (*SGPDBox)->GetNumberOfEntries() > 0)
-							{
-								// We do not care about the actual `VisualRandomAccessEntry` for now, only the fact that there is one.
-								FSampleGroupIterator& sgit = GroupIterators.Add(EGroupType::Grouping_rap);
-								sgit.GroupingType = EGroupType::Grouping_rap;
-								sgit.SBGPBox = Track->SBGPBoxes[i];
-								sgit.DefaultDescriptionIndex = (*SGPDBox)->GetDefaultSampleDescriptionIndex();
-								sgit.SetFirst();
-								// If the sample group says this sample is a RAP then clear the `sample_is_non_sync_sample` flag.
-								if (sgit.GetCurrentGroupDescriptionIndex() != 0)
-								{
-									SampleFlag &= ~0x10000U;
-								}
-								bHaveRAPGroup = true;
-							}
-						}
-						break;
-					}
-					default:
-					{
-						break;
-					}
-				}
-			}
+			// Set up the `rap` group if it exists and update the first SampleFlag accordingly.
+			SetupRAPGroup();
 
 			RemainingSamplesInTRUN = TRUNBox->GetNumberOfSamples();
 			check(RemainingSamplesInTRUN != 0);
@@ -6609,6 +6615,9 @@
 				SampleFlag = 0;			// If 0 then the sample *IS* a sync sample!
 			}
 
+			// Set up the `rap` group if it exists and update the first SampleFlag accordingly.
+			SetupRAPGroup();
+
 			return UEMEDIA_ERROR_OK;
 		}
 	}
@@ -6827,6 +6836,21 @@
 				}
 			}
 
+			// Advance the grouping iterators
+			for(auto &grpit : GroupIterators)
+			{
+				grpit.Value.Advance();
+			}
+			// If we have a RAP grouping we need to look at it now
+			if (bHaveRAPGroup)
+			{
+				// If the sample group says this sample is a RAP then clear the `sample_is_non_sync_sample` flag.
+				if (GroupIterators[EGroupType::Grouping_rap].GetCurrentGroupDescriptionIndex() != 0)
+				{
+					SampleFlag &= ~0x10000U;
+				}
+			}
+
 			return UEMEDIA_ERROR_OK;
 		}
 	}
@@ -6939,7 +6963,7 @@
 
 	int64 FParserISO14496_12::FTrackIterator::GetDTS() const
 	{
-		return SampleDTS + EmptyEditDurationInMediaTimeUnits - CompositionTimeEditOffset;
+		return SampleDTS + EmptyEditDurationInMediaTimeUnits;
 	}
 
 	int64 FParserISO14496_12::FTrackIterator::GetPTS() const
@@ -8202,6 +8226,20 @@
 							}
 						}
 
+						// Locate all sample grouping boxes.
+						TArray<const FMP4Box*> ListOfBoxes;
+						Track->STBLBox->GetAllBoxInstances(ListOfBoxes, FMP4Box::kBox_sgpd);
+						for(int32 nBox=0; nBox<ListOfBoxes.Num(); ++nBox)
+						{
+							Track->SGPDBoxes.Add(static_cast<const FMP4BoxSGPD*>(ListOfBoxes[nBox]));
+						}
+						ListOfBoxes.Empty();
+						Track->STBLBox->GetAllBoxInstances(ListOfBoxes, FMP4Box::kBox_sbgp);
+						for(int32 nBox=0; nBox<ListOfBoxes.Num(); ++nBox)
+						{
+							Track->SBGPBoxes.Add(static_cast<const FMP4BoxSBGP*>(ListOfBoxes[nBox]));
+						}
+
 						// So far all went well. Add the track to the map.
 						if (bIsSupported)
 						{
@@ -8252,7 +8290,7 @@
 									{
 										Track->TRUNBoxes.Add(static_cast<const FMP4BoxTRUN*>(ListOfBoxes[nBox]));
 									}
-									// Locate all sample grouping boxes. At present we do this only for fragmented streams.
+									// Locate all sample grouping boxes.
 									ListOfBoxes.Empty();
 									Box->GetAllBoxInstances(ListOfBoxes, FMP4Box::kBox_sgpd);
 									for(int32 nBox=0; nBox<ListOfBoxes.Num(); ++nBox)
