ElectraPlayer: added buffer query to display utilization in editor playback; exporting log symbols

#rnx

diff --git a/Engine/Plugins/Media/ElectraPlayer/Source/ElectraPlayerPlugin/Private/ElectraPlayerPlugin.cpp b/Engine/Plugins/Media/ElectraPlayer/Source/ElectraPlayerPlugin/Private/ElectraPlayerPlugin.cpp
--- a/Engine/Plugins/Media/ElectraPlayer/Source/ElectraPlayerPlugin/Private/ElectraPlayerPlugin.cpp
+++ b/Engine/Plugins/Media/ElectraPlayer/Source/ElectraPlayerPlugin/Private/ElectraPlayerPlugin.cpp
@@ -1003,6 +1003,45 @@ bool FElectraPlayerPlugin::SetPlaybackTimeRange(const TRange<FTimespan>& InTimeR
 	return true;
 }
 
+bool FElectraPlayerPlugin::QueryCacheState(EMediaCacheState State, TRangeSet<FTimespan>& OutTimeRanges) const
+{
+	// Note: The data of time ranges returned here will not actually get "cached" as
+	//       it is always only transient. We thus report the ranges only for `Loaded` and `Loading`,
+	//       but never for `Cached`!
+	switch(State)
+	{
+		case EMediaCacheState::Loaded:
+		case EMediaCacheState::Loading:
+		{
+			// When asked to provide what's already loaded we look at what we have in the sample queue
+			// and add that to the result. These samples have already left the player but are ready
+			// for use.
+			if (State == EMediaCacheState::Loaded)
+			{
+				TRange<FMediaTimeStamp> QueuedRange;
+				if (MediaSamples.IsValid() && MediaSamples->PeekVideoSampleTimeRange(QueuedRange))
+				{
+					OutTimeRanges.Add(TRange<FTimespan>(QueuedRange.GetLowerBoundValue().Time, QueuedRange.GetUpperBoundValue().Time));
+				}
+			}
+
+			// Get the data time range from the player. It returns both current and future data in one call, so we
+			// separate the result here based on what is being asked for.
+			IElectraPlayerInterface::FStreamBufferInfo vidBuf, audBuf;
+			bool bHaveVid = Player->GetStreamBufferInformation(vidBuf, IElectraPlayerInterface::EPlayerTrackType::Video);
+			bool bHaveAud = !bHaveVid ? Player->GetStreamBufferInformation(audBuf, IElectraPlayerInterface::EPlayerTrackType::Audio) : false;
+			const IElectraPlayerInterface::FStreamBufferInfo* Buffer = bHaveVid ? &vidBuf : bHaveAud ? &audBuf : nullptr;
+			const TArray<IElectraPlayerInterface::FStreamBufferInfo::FTimeRange>* tr = Buffer ? (State == EMediaCacheState::Loaded ? &Buffer->TimeAvailable : &Buffer->TimeRequested) : nullptr;
+			for(int32 i=0,iMax=tr?tr->Num():0; i<iMax; ++i)
+			{
+				OutTimeRanges.Add(TRange<FTimespan>(tr->operator[](i).Start.Time, tr->operator[](i).End.Time));
+			}
+			return true;
+		}
+	}
+	return false;
+}
+
 
 bool FElectraPlayerPlugin::GetAudioTrackFormat(int32 TrackIndex, int32 FormatIndex, FMediaAudioTrackFormat& OutFormat) const
 {
diff --git a/Engine/Plugins/Media/ElectraPlayer/Source/ElectraPlayerPlugin/Private/ElectraPlayerPlugin.h b/Engine/Plugins/Media/ElectraPlayer/Source/ElectraPlayerPlugin/Private/ElectraPlayerPlugin.h
--- a/Engine/Plugins/Media/ElectraPlayer/Source/ElectraPlayerPlugin/Private/ElectraPlayerPlugin.h
+++ b/Engine/Plugins/Media/ElectraPlayer/Source/ElectraPlayerPlugin/Private/ElectraPlayerPlugin.h
@@ -25,7 +25,7 @@ class IElectraSafeMediaOptionInterface;
 class FElectraPlayerResourceDelegate;
 
 
-DECLARE_LOG_CATEGORY_EXTERN(LogElectraPlayerPlugin, Log, All);
+ELECTRAPLAYERPLUGIN_API DECLARE_LOG_CATEGORY_EXTERN(LogElectraPlayerPlugin, Log, All);
 
 //-----------------------------------------------------------------------------
 
@@ -130,6 +130,10 @@ private:
 	TRange<FTimespan> GetPlaybackTimeRange(EMediaTimeRangeType InRangeToGet) const override;
 	bool SetPlaybackTimeRange(const TRange<FTimespan>& InTimeRange) override;
 
+	// From IMediaCache
+	bool QueryCacheState(EMediaCacheState State, TRangeSet<FTimespan>& OutTimeRanges) const override;
+
+
 	// From IMediaTracks
 	bool GetAudioTrackFormat(int32 TrackIndex, int32 FormatIndex, FMediaAudioTrackFormat& OutFormat) const override;
 	int32 GetNumTracks(EMediaTrackType TrackType) const override;
diff --git a/Engine/Plugins/Media/ElectraPlayer/Source/ElectraPlayerRuntime/Private/ElectraPlayer.cpp b/Engine/Plugins/Media/ElectraPlayer/Source/ElectraPlayerRuntime/Private/ElectraPlayer.cpp
--- a/Engine/Plugins/Media/ElectraPlayer/Source/ElectraPlayerRuntime/Private/ElectraPlayer.cpp
+++ b/Engine/Plugins/Media/ElectraPlayer/Source/ElectraPlayerRuntime/Private/ElectraPlayer.cpp
@@ -583,6 +583,48 @@ void FElectraPlayer::NotifyOfOptionChange()
 	PlatformNotifyOfOptionChange();
 }
 
+//-----------------------------------------------------------------------------
+/**
+ * Provides information about the time ranges that are currently available to the
+ * player and those that are being loaded.
+ */
+bool FElectraPlayer::GetStreamBufferInformation(IElectraPlayerInterface::FStreamBufferInfo& OutBufferInformation, EPlayerTrackType InTrackType) const
+{
+	TSharedPtr<FInternalPlayerImpl, ESPMode::ThreadSafe> Player = CurrentPlayer;
+	if (Player.IsValid())
+	{
+		IAdaptiveStreamingPlayer::FStreamBufferInfo bi;
+		if (InTrackType == EPlayerTrackType::Video)
+		{
+			Player->AdaptivePlayer->QueryStreamBufferInfo(bi, Electra::EStreamType::Video);
+		}
+		else if (InTrackType == EPlayerTrackType::Audio)
+		{
+			Player->AdaptivePlayer->QueryStreamBufferInfo(bi, Electra::EStreamType::Audio);
+		}
+		if (bi.bIsBufferActive)
+		{
+			auto AddRanges = [](TArray<IElectraPlayerInterface::FStreamBufferInfo::FTimeRange>& OutRanges, const TArray<Electra::FTimeRange>& InRanges) -> void
+			{
+				for(int32 i=0; i<InRanges.Num(); ++i)
+				{
+					IElectraPlayerInterface::FStreamBufferInfo::FTimeRange& tv = OutRanges.Emplace_GetRef();
+					tv.Start.Time = InRanges[i].Start.GetAsTimespan();
+					tv.Start.SequenceIndex = InRanges[i].Start.GetSequenceIndex();
+					tv.End.Time = InRanges[i].End.GetAsTimespan();
+					tv.End.SequenceIndex = InRanges[i].End.GetSequenceIndex();
+				}
+			};
+			AddRanges(OutBufferInformation.TimeAvailable, bi.TimeEnqueued);
+			AddRanges(OutBufferInformation.TimeAvailable, bi.TimeAvailable);
+			AddRanges(OutBufferInformation.TimeRequested, bi.TimeRequested);
+			return true;
+		}
+	}
+	return false;
+}
+
+
 //-----------------------------------------------------------------------------
 /**
 */
@@ -1405,7 +1447,7 @@ bool FElectraPlayer::GetAudioTrackFormat(int32 TrackIndex, int32 FormatIndex, FA
 			OutFormat.BitsPerSample = 16;
 			OutFormat.NumChannels = (uint32)ci.GetNumberOfChannels();
 			OutFormat.SampleRate = (uint32)ci.GetSamplingRate();
-			OutFormat.TypeName = ci.GetCodecSpecifierRFC6381();
+			OutFormat.TypeName = ci.GetHumanReadableCodecName();
 			return true;
 		}
 	}
@@ -1426,7 +1468,7 @@ bool FElectraPlayer::GetVideoTrackFormat(int32 TrackIndex, int32 FormatIndex, FV
 			OutFormat.Dim.Y = ci.GetResolution().Height;
 			OutFormat.FrameRate = (float)ci.GetFrameRate().GetAsDouble();
 			OutFormat.FrameRates = TRange<float>{ OutFormat.FrameRate };
-			OutFormat.TypeName = ci.GetCodecSpecifierRFC6381();
+			OutFormat.TypeName = ci.GetHumanReadableCodecName();
 			return true;
 		}
 	}
diff --git a/Engine/Plugins/Media/ElectraPlayer/Source/ElectraPlayerRuntime/Private/ElectraPlayer.h b/Engine/Plugins/Media/ElectraPlayer/Source/ElectraPlayerRuntime/Private/ElectraPlayer.h
--- a/Engine/Plugins/Media/ElectraPlayer/Source/ElectraPlayerRuntime/Private/ElectraPlayer.h
+++ b/Engine/Plugins/Media/ElectraPlayer/Source/ElectraPlayerRuntime/Private/ElectraPlayer.h
@@ -140,6 +140,7 @@ public:
 	TSharedPtr<TMap<FString, TArray<TSharedPtr<Electra::IMediaStreamMetadata::IItem, ESPMode::ThreadSafe>>>, ESPMode::ThreadSafe> GetMediaMetadata() const override;
 
 	void NotifyOfOptionChange() override;
+	bool GetStreamBufferInformation(FStreamBufferInfo& OutBufferInformation, EPlayerTrackType InTrackType) const override;
 
 	void SuspendOrResumeDecoders(bool bSuspend, const Electra::FParamDict& InOptions) override;
 
diff --git a/Engine/Plugins/Media/ElectraPlayer/Source/ElectraPlayerRuntime/Private/Runtime/Decoder/AudioDecoder.cpp b/Engine/Plugins/Media/ElectraPlayer/Source/ElectraPlayerRuntime/Private/Runtime/Decoder/AudioDecoder.cpp
--- a/Engine/Plugins/Media/ElectraPlayer/Source/ElectraPlayerRuntime/Private/Runtime/Decoder/AudioDecoder.cpp
+++ b/Engine/Plugins/Media/ElectraPlayer/Source/ElectraPlayerRuntime/Private/Runtime/Decoder/AudioDecoder.cpp
@@ -69,13 +69,8 @@ private:
 	struct FDecoderInput
 	{
 		~FDecoderInput()
-		{
-			ReleasePayload();
-		}
-		void ReleasePayload()
 		{
 			FAccessUnit::Release(AccessUnit);
-			AccessUnit = nullptr;
 		}
 
 		FAccessUnit*	AccessUnit = nullptr;
@@ -397,10 +392,20 @@ void FAudioDecoderImpl::NotifyReadyBufferListener(bool bHaveOutput)
 	if (ReadyBufferListener)
 	{
 		IDecoderOutputBufferListener::FDecodeReadyStats stats;
-		stats.MaxDecodedElementsReady = MaxDecodeBufferSize;
-		stats.NumElementsInDecoder = CurrentOutputBuffer ? 1 : 0;
+		stats.OutputBufferPoolSize = MaxDecodeBufferSize;
+		if ((stats.NumElementsInDecoder = InDecoderInput.Num()) != 0)
+		{
+			stats.InDecoderTimeRangePTS.Start = InDecoderInput[0]->AccessUnit->PTS;
+			stats.InDecoderTimeRangePTS.End = InDecoderInput.Last()->AccessUnit->PTS + InDecoderInput.Last()->AccessUnit->Duration;
+		}
+		if (CurrentAccessUnit)
+		{
+			stats.InDecoderTimeRangePTS.Start = Utils::Min(CurrentAccessUnit->AccessUnit->PTS, stats.InDecoderTimeRangePTS.Start.IsValid() ? stats.InDecoderTimeRangePTS.Start : FTimeValue::GetPositiveInfinity());
+			stats.InDecoderTimeRangePTS.End = Utils::Max(CurrentAccessUnit->AccessUnit->PTS + CurrentAccessUnit->AccessUnit->Duration, stats.InDecoderTimeRangePTS.End.IsValid() ? stats.InDecoderTimeRangePTS.End : FTimeValue::GetNegativeInfinity());
+		}
+		stats.InDecoderTimeRangePTS.End.SetSequenceIndex(stats.InDecoderTimeRangePTS.Start.GetSequenceIndex());
 		stats.bOutputStalled = !bHaveOutput;
-		stats.bEODreached = NextAccessUnits.ReachedEOD() && stats.NumDecodedElementsReady == 0 && stats.NumElementsInDecoder == 0;
+		stats.bEODreached = NextAccessUnits.ReachedEOD() && CurrentOutputBuffer == nullptr;
 		ListenerMutex.Lock();
 		if (ReadyBufferListener)
 		{
@@ -511,7 +516,7 @@ bool FAudioDecoderImpl::InternalDecoderCreate()
 		}
 	}
 
-	
+
 	return true;
 }
 
@@ -621,7 +626,7 @@ FAudioDecoderImpl::EAUChangeFlags FAudioDecoderImpl::GetAndPrepareInputAU()
 			}
 			ListenerMutex.Unlock();
 		}
-		
+
 		// Get the AU to be decoded if one is there.
 		if (NextAccessUnits.Wait(1000 * 10))
 		{
@@ -630,7 +635,7 @@ FAudioDecoderImpl::EAUChangeFlags FAudioDecoderImpl::GetAndPrepareInputAU()
 			{
 				PrepareAU(CurrentAccessUnit);
 				// Is there a discontinuity/break in sequence of sorts?
-				if (CurrentAccessUnit->AccessUnit->bTrackChangeDiscontinuity || 
+				if (CurrentAccessUnit->AccessUnit->bTrackChangeDiscontinuity ||
 					(!bInDummyDecodeMode && CurrentAccessUnit->AccessUnit->bIsDummyData) ||
 					(CurrentSequenceIndex.IsSet() && CurrentSequenceIndex.GetValue() != CurrentAccessUnit->AccessUnit->PTS.GetSequenceIndex()))
 				{
@@ -913,8 +918,8 @@ bool FAudioDecoderImpl::HandleDecoding()
 			{
 				InDecoderInput.Emplace(CurrentAccessUnit);
 				InDecoderInput.Sort([](const TSharedPtr<FDecoderInput, ESPMode::ThreadSafe>& a, const TSharedPtr<FDecoderInput, ESPMode::ThreadSafe>& b)
-				{ 
-					return a->PTS < b->PTS; 
+				{
+					return a->PTS < b->PTS;
 				});
 
 				// If this was the last access unit in a period we need to drain the decoder _after_ having sent it
@@ -943,7 +948,7 @@ bool FAudioDecoderImpl::SendSilenceOrEOS(ESendMode InSendMode)
 {
 	check(CurrentAccessUnit.IsValid() || InSendMode == ESendMode::SendEOS);
 	check(bIsDecoderClean);
-	
+
 	// Get output unless flushing or terminating
 	while(!TerminateThreadSignal.IsSignaled() && !FlushDecoderSignal.IsSignaled())
 	{
diff --git a/Engine/Plugins/Media/ElectraPlayer/Source/ElectraPlayerRuntime/Private/Runtime/Decoder/VideoDecoder.cpp b/Engine/Plugins/Media/ElectraPlayer/Source/ElectraPlayerRuntime/Private/Runtime/Decoder/VideoDecoder.cpp
--- a/Engine/Plugins/Media/ElectraPlayer/Source/ElectraPlayerRuntime/Private/Runtime/Decoder/VideoDecoder.cpp
+++ b/Engine/Plugins/Media/ElectraPlayer/Source/ElectraPlayerRuntime/Private/Runtime/Decoder/VideoDecoder.cpp
@@ -72,13 +72,8 @@ private:
 	struct FDecoderInput
 	{
 		~FDecoderInput()
-		{
-			ReleasePayload();
-		}
-		void ReleasePayload()
 		{
 			FAccessUnit::Release(AccessUnit);
-			AccessUnit = nullptr;
 		}
 
 		FTimeValue		AdjustedPTS;
@@ -469,10 +464,20 @@ void FVideoDecoderImpl::NotifyReadyBufferListener(bool bHaveOutput)
 	if (ReadyBufferListener)
 	{
 		IDecoderOutputBufferListener::FDecodeReadyStats stats;
-		stats.MaxDecodedElementsReady = MaxOutputBuffers;
-		stats.NumElementsInDecoder = CurrentOutputBuffer ? 1 : 0;
+		stats.OutputBufferPoolSize = MaxOutputBuffers;
+		if ((stats.NumElementsInDecoder = InDecoderInput.Num()) != 0)
+		{
+			stats.InDecoderTimeRangePTS.Start = InDecoderInput[0]->AccessUnit->PTS;
+			stats.InDecoderTimeRangePTS.End = InDecoderInput.Last()->AccessUnit->PTS + InDecoderInput.Last()->AccessUnit->Duration;
+		}
+		if (CurrentAccessUnit)
+		{
+			stats.InDecoderTimeRangePTS.Start = Utils::Min(CurrentAccessUnit->AccessUnit->PTS, stats.InDecoderTimeRangePTS.Start.IsValid() ? stats.InDecoderTimeRangePTS.Start : FTimeValue::GetPositiveInfinity());
+			stats.InDecoderTimeRangePTS.End = Utils::Max(CurrentAccessUnit->AccessUnit->PTS + CurrentAccessUnit->AccessUnit->Duration, stats.InDecoderTimeRangePTS.End.IsValid() ? stats.InDecoderTimeRangePTS.End : FTimeValue::GetNegativeInfinity());
+		}
+		stats.InDecoderTimeRangePTS.End.SetSequenceIndex(stats.InDecoderTimeRangePTS.Start.GetSequenceIndex());
 		stats.bOutputStalled = !bHaveOutput;
-		stats.bEODreached = NextAccessUnits.ReachedEOD() && stats.NumDecodedElementsReady == 0 && stats.NumElementsInDecoder == 0;
+		stats.bEODreached = NextAccessUnits.ReachedEOD() && CurrentOutputBuffer == nullptr;
 		ListenerMutex.Lock();
 		if (ReadyBufferListener)
 		{
diff --git a/Engine/Plugins/Media/ElectraPlayer/Source/ElectraPlayerRuntime/Private/Runtime/Player/AdaptiveStreamingPlayer.cpp b/Engine/Plugins/Media/ElectraPlayer/Source/ElectraPlayerRuntime/Private/Runtime/Player/AdaptiveStreamingPlayer.cpp
--- a/Engine/Plugins/Media/ElectraPlayer/Source/ElectraPlayerRuntime/Private/Runtime/Player/AdaptiveStreamingPlayer.cpp
+++ b/Engine/Plugins/Media/ElectraPlayer/Source/ElectraPlayerRuntime/Private/Runtime/Player/AdaptiveStreamingPlayer.cpp
@@ -20,6 +20,7 @@ DECLARE_CYCLE_STAT(TEXT("FAdaptiveStreamingPlayer::WorkerThread"), STAT_ElectraP
 DECLARE_CYCLE_STAT(TEXT("FAdaptiveStreamingPlayer::EventThread"), STAT_ElectraPlayer_EventWorker, STATGROUP_ElectraPlayer);
 
 
+
 namespace Electra
 {
 FAdaptiveStreamingPlayer *FAdaptiveStreamingPlayer::PointerToLatestPlayer;
@@ -39,24 +40,6 @@ void IAdaptiveStreamingPlayer::DebugHandle(void* pPlayer, void (*debugDrawPrintf
 //---------------------------------------------------------------------------------------------------------------------------------
 
 
-namespace
-{
-	inline int32 StreamTypeToArrayIndex(EStreamType StreamType)
-	{
-		switch(StreamType)
-		{
-			case EStreamType::Video:
-				return 0;
-			case EStreamType::Audio:
-				return 1;
-			case EStreamType::Subtitle:
-				return 2;
-			default:
-				return 3;
-		}
-	}
-}
-
 //-----------------------------------------------------------------------------
 /**
  * CTOR
@@ -175,7 +158,7 @@ FAdaptiveStreamingPlayer::~FAdaptiveStreamingPlayer()
  */
 FErrorDetail FAdaptiveStreamingPlayer::GetError() const
 {
-	FMediaCriticalSection::ScopedLock lock(DiagnosticsCriticalSection);
+	FScopeLock lock(&DiagnosticsCriticalSection);
 	return LastErrorDetail;
 }
 
@@ -305,7 +288,7 @@ void FAdaptiveStreamingPlayer::AddMetricsReceiver(IAdaptiveStreamingPlayerMetric
 {
 	if (InMetricsReceiver)
 	{
-		FMediaCriticalSection::ScopedLock lock(MetricListenerCriticalSection);
+		FScopeLock lock(&MetricListenerCriticalSection);
 		if (MetricListeners.Find(InMetricsReceiver) == INDEX_NONE)
 		{
 			MetricListeners.Push(InMetricsReceiver);
@@ -317,7 +300,7 @@ void FAdaptiveStreamingPlayer::RemoveMetricsReceiver(IAdaptiveStreamingPlayerMet
 {
 	if (InMetricsReceiver)
 	{
-		FMediaCriticalSection::ScopedLock lock(MetricListenerCriticalSection);
+		FScopeLock lock(&MetricListenerCriticalSection);
 		/*bool bRemoved =*/ MetricListeners.Remove(InMetricsReceiver);
 	}
 }
@@ -976,6 +959,94 @@ bool FAdaptiveStreamingPlayer::IsTrackDeselected(EStreamType StreamType)
 	}
 }
 
+
+void FAdaptiveStreamingPlayer::QueryStreamBufferInfo(FStreamBufferInfo& OutStreamBufferInfo, EStreamType InStreamType)
+{
+	OutStreamBufferInfo.TimeAvailable.Empty();
+	OutStreamBufferInfo.TimeRequested.Empty();
+	OutStreamBufferInfo.TimeEnqueued.Empty();
+	OutStreamBufferInfo.bIsBufferActive = IsExpectedToStreamNow(InStreamType);
+	FTimeRange dl = PlaybackState.GetCurrentDownloadRequestTimeRange(InStreamType);
+	if (dl.IsValid())
+	{
+		OutStreamBufferInfo.TimeRequested.Emplace(MoveTemp(dl));
+	}
+	if (OutStreamBufferInfo.bIsBufferActive)
+	{
+		switch(InStreamType)
+		{
+			case EStreamType::Video:
+			case EStreamType::Audio:
+			{
+				auto AddBufferInfo = [](TArray<FTimeRange>& OutRanges, const TSharedPtrTS<FMultiTrackAccessUnitBuffer>& InBuffer) -> void
+				{
+					if (InBuffer.IsValid())
+					{
+						FAccessUnitBufferInfo bi;
+						InBuffer->GetStats(bi);
+						if (bi.SmallestPTS.IsValid())
+						{
+							FTimeRange tr;
+							tr.Start = bi.SmallestPTS;
+							tr.End = bi.LargestPTSPlusDur;
+							OutRanges.Emplace(MoveTemp(tr));
+						}
+					}
+				};
+
+#if 0
+				// Ask the renderers what they know they have released into the media sample queue.
+				TArray<IAdaptiveStreamingWrappedRenderer::FEnqueuedSampleInfo> Enqueued;
+				if (InStreamType == EStreamType::Video && VideoRender.Renderer.IsValid())
+				{
+					VideoRender.Renderer->GetNumEnqueuedSamples(&Enqueued);
+				}
+				else if (InStreamType == EStreamType::Audio && AudioRender.Renderer.IsValid())
+				{
+					AudioRender.Renderer->GetNumEnqueuedSamples(&Enqueued);
+				}
+				if (Enqueued.Num())
+				{
+					// TBD: Do we need to break them down into separate ranges based on the sequence index?
+					FTimeRange tr;
+					tr.Start = Enqueued[0].PTS;
+					tr.End = Enqueued.Last().PTS + Enqueued.Last().Duration;
+					OutStreamBufferInfo.TimeEnqueued.Emplace(MoveTemp(tr));
+				}
+#endif
+
+				// Get the time range of the access units the respective decoder has in its internal pipeline
+				// These AUs are not in the player buffer any more.
+				{
+					FScopeLock lock(&DiagnosticsCriticalSection);
+					IDecoderOutputBufferListener::FDecodeReadyStats* Ready = InStreamType == EStreamType::Video ? &VideoBufferStats.DecoderOutputBuffer :
+																			 InStreamType == EStreamType::Audio ? &AudioBufferStats.DecoderOutputBuffer :
+																			 nullptr;
+					if (Ready && Ready->InDecoderTimeRangePTS.IsValid())
+					{
+						OutStreamBufferInfo.TimeEnqueued.Emplace(Ready->InDecoderTimeRangePTS);
+					}
+				}
+
+				// Get the time range of AUs in the current active and future buffers (loops).
+				{
+					FScopeLock lock(&DataBuffersCriticalSection);
+					if (ActiveDataOutputBuffers.IsValid())
+					{
+						AddBufferInfo(OutStreamBufferInfo.TimeAvailable, ActiveDataOutputBuffers->GetBuffer(InStreamType));
+					}
+					for(int32 i=0,iMax=NextDataBuffers.Num(); i<iMax; ++i)
+					{
+						AddBufferInfo(OutStreamBufferInfo.TimeAvailable, NextDataBuffers[i]->GetBuffer(InStreamType));
+					}
+				}
+				break;
+			}
+		}
+	}
+}
+
+
 void FAdaptiveStreamingPlayer::SuspendOrResumeDecoders(bool bSuspend, const FParamDict& InOptions)
 {
 	VideoDecoder.SuspendOrResume(bSuspend, InOptions);
@@ -1557,6 +1628,8 @@ bool FAdaptiveStreamingPlayer::InternalHandleThreadMessages()
 				// Check that the request is for this current playback sequence and not an outdated one.
 				if (pRequest.IsValid() && pRequest->GetPlaybackSequenceID() == CurrentPlaybackSequenceID[StreamTypeToArrayIndex(reqType)])
 				{
+					PlaybackState.SetCurrentDownloadRequestTimeRange(reqType, pRequest->GetTimeRange());
+
 					DispatchBufferUtilizationEvent(ev.Request->GetType());
 
 					// Video bitrate change?
@@ -1597,6 +1670,8 @@ bool FAdaptiveStreamingPlayer::InternalHandleThreadMessages()
 				// Check that the request is for this current playback sequence and not an outdated one.
 				if (pRequest.IsValid() && pRequest->GetPlaybackSequenceID() == CurrentPlaybackSequenceID[StreamTypeToArrayIndex(reqType)])
 				{
+					PlaybackState.ClearCurrentDownloadRequestTimeRange(reqType);
+
 					// Dispatch download event
 					DispatchSegmentDownloadedEvent(pRequest);
 
@@ -4100,10 +4175,12 @@ void FAdaptiveStreamingPlayer::CheckForStreamEnd()
 				// which tends to happen at the end of the stream or when there is no video any more when audio is longer.
 				bEndVid = (vidStats.StreamBuffer.bEndOfData && vidStats.DecoderInputBuffer.bEODSignaled && vidStats.DecoderInputBuffer.bEODReached);
 				VidStalled = vidStats.StreamBuffer.bEndOfData ? vidStats.GetStalledDurationMillisec() : 0;
+#if NOTIFY_DATA_AVAILABILITY_AT_START_AND_END
 				if (bEndVid)
 				{
 					UpdateDataAvailabilityState(DataAvailabilityStateVid, Metrics::FDataAvailabilityChange::EAvailability::DataNotAvailable);
 				}
+#endif
 			}
 
 			// Check for end of audio stream
@@ -4114,10 +4191,12 @@ void FAdaptiveStreamingPlayer::CheckForStreamEnd()
 				// which tends to happen at the end of the stream or when there is no audio any more when video is longer.
 				bEndAud = (audStats.StreamBuffer.bEndOfData && audStats.DecoderInputBuffer.bEODSignaled && audStats.DecoderInputBuffer.bEODReached);
 				AudStalled = audStats.StreamBuffer.bEndOfData ? audStats.GetStalledDurationMillisec() : 0;
+#if NOTIFY_DATA_AVAILABILITY_AT_START_AND_END
 				if (bEndAud)
 				{
 					UpdateDataAvailabilityState(DataAvailabilityStateAud, Metrics::FDataAvailabilityChange::EAvailability::DataNotAvailable);
 				}
+#endif
 			}
 
 			// Text stream
@@ -4229,10 +4308,12 @@ void FAdaptiveStreamingPlayer::InternalStartAt(const FSeekParam& NewPosition, co
 	AudioBufferStats.Clear();
 	TextBufferStats.Clear();
 
+#if NOTIFY_DATA_AVAILABILITY_AT_START_AND_END
 	// Update data availability states in case this wasn't done yet.
 	UpdateDataAvailabilityState(DataAvailabilityStateVid, Metrics::FDataAvailabilityChange::EAvailability::DataNotAvailable);
 	UpdateDataAvailabilityState(DataAvailabilityStateAud, Metrics::FDataAvailabilityChange::EAvailability::DataNotAvailable);
 	UpdateDataAvailabilityState(DataAvailabilityStateTxt, Metrics::FDataAvailabilityChange::EAvailability::DataNotAvailable);
+#endif
 
 	// (Re-)configure the stream selector
 	StreamSelector->SetBandwidthCeiling(BitrateCeiling);
@@ -4901,14 +4982,14 @@ void FAdaptiveStreamingPlayer::DebugPrint(void* pPlayer, void (*pPrintFN)(void*
 	{
 		pD = &VideoBufferStats.StreamBuffer;
 		pPrintFN(pPlayer, "Video buffer  : EOT %d; EOS %d; %3u AUs; %8lld bytes in; %#7.4fs", pD->bEndOfTrack, pD->bEndOfData, (uint32)pD->NumCurrentAccessUnits, (long long int)pD->CurrentMemInUse, pD->PlayableDuration.GetAsSeconds());
-		pPrintFN(pPlayer, "Video decoder : %2u in decoder, %zu total; EOD in %d; EOD out %d; %s", (uint32)VideoBufferStats.DecoderOutputBuffer.NumElementsInDecoder, (uint32)VideoBufferStats.DecoderOutputBuffer.MaxDecodedElementsReady, VideoBufferStats.DecoderInputBuffer.bEODReached, VideoBufferStats.DecoderOutputBuffer.bEODreached, VideoBufferStats.DecoderOutputBuffer.bOutputStalled?"stalled":"not stalled");
+		pPrintFN(pPlayer, "Video decoder : %2u in decoder, %zu total; EOD in %d; EOD out %d; %s", (uint32)VideoBufferStats.DecoderOutputBuffer.NumElementsInDecoder, (uint32)VideoBufferStats.DecoderOutputBuffer.OutputBufferPoolSize, VideoBufferStats.DecoderInputBuffer.bEODReached, VideoBufferStats.DecoderOutputBuffer.bEODreached, VideoBufferStats.DecoderOutputBuffer.bOutputStalled?"stalled":"not stalled");
 		pPrintFN(pPlayer, "Video renderer: %#7.4fs enqueued", VideoRender.Renderer.IsValid() ? VideoRender.Renderer->GetEnqueuedSampleDuration().GetAsSeconds() : 0.0);
 	}
 	if(1)
 	{
 		pD = &AudioBufferStats.StreamBuffer;
 		pPrintFN(pPlayer, "Audio buffer  : EOT %d; EOS %d; %3u AUs; %8lld bytes in; %#7.4fs", pD->bEndOfTrack, pD->bEndOfData, (uint32)pD->NumCurrentAccessUnits, (long long int)pD->CurrentMemInUse, pD->PlayableDuration.GetAsSeconds());
-		pPrintFN(pPlayer, "Audio decoder : %2u in decoder, %2u total; EOD in %d; EOD out %d; %s", (uint32)AudioBufferStats.DecoderOutputBuffer.NumElementsInDecoder, (uint32)AudioBufferStats.DecoderOutputBuffer.MaxDecodedElementsReady, AudioBufferStats.DecoderInputBuffer.bEODReached, AudioBufferStats.DecoderOutputBuffer.bEODreached, AudioBufferStats.DecoderOutputBuffer.bOutputStalled?"stalled":"not stalled");
+		pPrintFN(pPlayer, "Audio decoder : %2u in decoder, %2u total; EOD in %d; EOD out %d; %s", (uint32)AudioBufferStats.DecoderOutputBuffer.NumElementsInDecoder, (uint32)AudioBufferStats.DecoderOutputBuffer.OutputBufferPoolSize, AudioBufferStats.DecoderInputBuffer.bEODReached, AudioBufferStats.DecoderOutputBuffer.bEODreached, AudioBufferStats.DecoderOutputBuffer.bOutputStalled?"stalled":"not stalled");
 		pPrintFN(pPlayer, "Audio renderer: %#7.4fs enqueued", AudioRender.Renderer.IsValid() ? AudioRender.Renderer->GetEnqueuedSampleDuration().GetAsSeconds() : 0.0);
 	}
 	if(1)
diff --git a/Engine/Plugins/Media/ElectraPlayer/Source/ElectraPlayerRuntime/Private/Runtime/Player/AdaptiveStreamingPlayer.h b/Engine/Plugins/Media/ElectraPlayer/Source/ElectraPlayerRuntime/Private/Runtime/Player/AdaptiveStreamingPlayer.h
--- a/Engine/Plugins/Media/ElectraPlayer/Source/ElectraPlayerRuntime/Private/Runtime/Player/AdaptiveStreamingPlayer.h
+++ b/Engine/Plugins/Media/ElectraPlayer/Source/ElectraPlayerRuntime/Private/Runtime/Player/AdaptiveStreamingPlayer.h
@@ -100,7 +100,7 @@ public:
 	 * This also affects looping which implicitly seeks back to the loop point when the end is reached.
 	 * Frame accurate seeking is enabled by default.
 	 * This method is intended mostly to disable frame accurate seeking on this player instance.
-	 * 
+	 *
 	 * This should be called prior to SeekTo() and should only be called once on the player instance to
 	 * disable or re-enable frame accurate seeking.
 	 * Calling this during playback may have undesired results.
@@ -144,15 +144,15 @@ public:
 		// Ignore the seek for sequence index updates
 		TOptional<bool> bIgnoreForSequenceIndex;
 	};
-	
+
 	/**
 	 * Seek to a new position and play from there. This includes first playstart.
 	 * Playback is initially paused on first player use and must be resumed to begin.
 	 * Query the seekable range (GetSeekableRange()) to get the valid time range.
-	 * 
+	 *
 	 * If the seek-to time is not set the seek will start at the beginning for
 	 * on-demand presentations and on the Live edge for Live presentations.
-	 * 
+	 *
 	 * Seeks can be issued while a seek is already executing. The seek parameters
 	 * control behaviour. If seeking is performed for scrubbing any new seek will
 	 * be performed only when the previous seek has completed, otherwise the current
@@ -187,20 +187,20 @@ public:
 	 * Constrains playback to the specified time range, which should be a subset of GetTimelineRange().
 	 * The playback range can be specified via URL fragment parameters on the URL given to LoadManifest()
 	 * if the mime type allows for it.
-	 * 
+	 *
 	 * If you set the playback range before calling LoadManifest() the URL parameter will not be used.
 	 * Otherwise the URL parameters set the playback range if they are specified.
 	 * You can query the playback range set by URL parameters as soon as HaveMetadata() returns true.
 	 * Setting a playback range through this method overrides URL parameters.
-	 * 
+	 *
 	 * To set or change only the start or end of the playback range, set only the corresponding TOptional<>
 	 * and leave the other unset.
 	 * To disable either start or end set the value to FTimeValue::GetInvalid().
-	 * 
+	 *
 	 * If you only set start or end, the other value may be set by the respective URL parameter.
 	 * To fully disable any playback range that may be present on the URL you should set the range to
 	 * invalid values once HaveMetadata() returns true.
-	 * 
+	 *
 	 * Setting a playback range during playback will result in an immediate seek to the current
 	 * playback position. Frame accurate seeking is recommended.
 	 */
@@ -239,18 +239,18 @@ public:
 	 * Returns ranges the playback rate can be set to.
 	 * There are two types, one `Unthinned`, where (ideally) no frames will be dropped and
 	 * `Thinned`, where frames will be dropped to maintain the playback rate.
-	 * 
+	 *
 	 * Either way, the player will not resample audio or generate interpolated frames of video.
 	 * Decoded samples will be delivered to the renderer as they are. It is also up to the
 	 * renderers to consume the data at a faster or slower rate to actually realize the desired
 	 * playback rate.
-	 * 
+	 *
 	 * The range of supported rates depends on the type of media and the decoder capability to
 	 * decode faster than realtime. If the media allows for adaptive bitrate selection the player
 	 * will choose the stream to play back accordingly.
 	 * In thinned mode it will need to drop and not decode samples, so the renderer will not
 	 * receive all the possible data.
-	 * 
+	 *
 	 * Live streams will only allow for rates of 0.0 (pause) and 1.0 (real time play forward).
 	 */
 	virtual TRangeSet<double> GetSupportedRates(EPlaybackRateType InForPlayRateType) = 0;
@@ -333,6 +333,19 @@ public:
 	virtual bool IsTrackDeselected(EStreamType StreamType) = 0;
 
 
+	//-------------------------------------------------------------------------
+	// Buffer related functions
+	//
+	struct FStreamBufferInfo
+	{
+		TArray<FTimeRange> TimeAvailable;		//!< Time range currently available in the buffer
+		TArray<FTimeRange> TimeRequested;		//!< Time range requested for download.
+		TArray<FTimeRange> TimeEnqueued;		//!< Time range already enqueued with the renderer.
+		bool bIsBufferActive = false;			//!< true if the buffer is active, false if not (eg. track not selected)
+	};
+	virtual void QueryStreamBufferInfo(FStreamBufferInfo& OutStreamBufferInfo, EStreamType InStreamType) = 0;
+
+
 	//-------------------------------------------------------------------------
 	// Special functions needed for application suspend & resume.
 	//
diff --git a/Engine/Plugins/Media/ElectraPlayer/Source/ElectraPlayerRuntime/Private/Runtime/Player/AdaptiveStreamingPlayerDecoders.cpp b/Engine/Plugins/Media/ElectraPlayer/Source/ElectraPlayerRuntime/Private/Runtime/Player/AdaptiveStreamingPlayerDecoders.cpp
--- a/Engine/Plugins/Media/ElectraPlayer/Source/ElectraPlayerRuntime/Private/Runtime/Player/AdaptiveStreamingPlayerDecoders.cpp
+++ b/Engine/Plugins/Media/ElectraPlayer/Source/ElectraPlayerRuntime/Private/Runtime/Player/AdaptiveStreamingPlayerDecoders.cpp
@@ -632,15 +632,14 @@ void FAdaptiveStreamingPlayer::FeedDecoder(EStreamType Type, IAccessUnitBufferIn
 			FTimeValue PushedDuration = FromMultistreamBuffer->GetPlayableDurationPushedSinceEOT();
 			if (!bEODSet && !bEOTSet && FromMultistreamBuffer->Num() == 0 && PushedDuration >= MinPushedDuration)
 			{
-				FTimeValue EnqueuedDuration(FTimeValue::GetZero());
 				int32 NumEnqueuedSamples = 0;
 				if (Type == EStreamType::Video && VideoRender.Renderer.IsValid())
 				{
-					NumEnqueuedSamples = VideoRender.Renderer->GetNumEnqueuedSamples(&EnqueuedDuration);
+					NumEnqueuedSamples = VideoRender.Renderer->GetNumEnqueuedSamples(nullptr);
 				}
 				else if (Type == EStreamType::Audio && AudioRender.Renderer.IsValid())
 				{
-					NumEnqueuedSamples = AudioRender.Renderer->GetNumEnqueuedSamples(&EnqueuedDuration);
+					NumEnqueuedSamples = AudioRender.Renderer->GetNumEnqueuedSamples(nullptr);
 				}
 				if (NumEnqueuedSamples <= 1)
 				{
diff --git a/Engine/Plugins/Media/ElectraPlayer/Source/ElectraPlayerRuntime/Private/Runtime/Player/AdaptiveStreamingPlayerInternal.h b/Engine/Plugins/Media/ElectraPlayer/Source/ElectraPlayerRuntime/Private/Runtime/Player/AdaptiveStreamingPlayerInternal.h
--- a/Engine/Plugins/Media/ElectraPlayer/Source/ElectraPlayerRuntime/Private/Runtime/Player/AdaptiveStreamingPlayerInternal.h
+++ b/Engine/Plugins/Media/ElectraPlayer/Source/ElectraPlayerRuntime/Private/Runtime/Player/AdaptiveStreamingPlayerInternal.h
@@ -46,18 +46,44 @@
 #define INTERR_REBUFFER_SHALL_THROW_ERROR		0x200
 
 
+// Set this to `1` to signal no data availability right at the start of a seek (and initial playback) and at the end.
+// Set to `0` to only notify about mid-playback state changes.
+#define NOTIFY_DATA_AVAILABILITY_AT_START_AND_END 0
+
 
 namespace Electra
 {
 class FAdaptiveStreamingPlayer;
 
+inline int32 StreamTypeToArrayIndex(EStreamType StreamType)
+{
+	switch(StreamType)
+	{
+		case EStreamType::Video:
+			return 0;
+		case EStreamType::Audio:
+			return 1;
+		case EStreamType::Subtitle:
+			return 2;
+		default:
+			return 3;
+	}
+}
+
+
 
 class IAdaptiveStreamingWrappedRenderer : public IMediaRenderer
 {
 public:
+	struct FEnqueuedSampleInfo
+	{
+		FTimeValue PTS;
+		FTimeValue Duration;
+	};
+
 	virtual ~IAdaptiveStreamingWrappedRenderer() = default;
 	virtual FTimeValue GetEnqueuedSampleDuration() = 0;
-	virtual int32 GetNumEnqueuedSamples(FTimeValue* OutOptionalDuration) = 0;
+	virtual int32 GetNumEnqueuedSamples(TArray<FEnqueuedSampleInfo>* OutOptionalSampleInfos) = 0;
 
 	virtual void DisableHoldbackOfFirstRenderableVideoFrame(bool bDisableHoldback) = 0;
 
@@ -85,7 +111,7 @@ public:
 		if (Clk)
 		{
 			FTimeValue Now(MEDIAutcTime::Current());
-			FMediaCriticalSection::ScopedLock lock(Lock);
+			FScopeLock lock(&Lock);
 			Clk->RenderTime = CurrentRenderTime;
 			Clk->LastSystemBaseTime = Now;
 			Clk->RunningTimeOffset.SetToZero();
@@ -98,7 +124,7 @@ public:
 		if (Clk)
 		{
 			FTimeValue Now(MEDIAutcTime::Current());
-			FMediaCriticalSection::ScopedLock lock(Lock);
+			FScopeLock lock(&Lock);
 			FTimeValue diff = !bIsPaused ? Now - Clk->LastSystemBaseTime : FTimeValue::GetZero();
 			FTimeValue r(Clk->RenderTime);
 			r += Clk->RunningTimeOffset + diff;
@@ -110,7 +136,7 @@ public:
 	void Start()
 	{
 		FTimeValue Now(MEDIAutcTime::Current());
-		FMediaCriticalSection::ScopedLock lock(Lock);
+		FScopeLock lock(&Lock);
 		if (bIsPaused)
 		{
 			for(int32 i = 0; i < FMEDIA_STATIC_ARRAY_COUNT(Clock); ++i)
@@ -124,7 +150,7 @@ public:
 	void Stop()
 	{
 		FTimeValue Now(MEDIAutcTime::Current());
-		FMediaCriticalSection::ScopedLock lock(Lock);
+		FScopeLock lock(&Lock);
 		if (!bIsPaused)
 		{
 			bIsPaused = true;
@@ -175,9 +201,9 @@ private:
 		return Clk;
 	}
 
-	FMediaCriticalSection		Lock;
-	FClock			Clock[3];
-	bool			bIsPaused;
+	FCriticalSection Lock;
+	FClock Clock[3];
+	bool bIsPaused;
 };
 
 
@@ -217,6 +243,10 @@ struct FPlaybackState
 		bPlayrangeHasChanged = false;
 		bLoopStateHasChanged = false;
 		bShouldPlayOnLiveEdge = false;
+		for(int32 i=0; i<UE_ARRAY_COUNT(CurrentSegmentDownloadTimeRange); ++i)
+		{
+			CurrentSegmentDownloadTimeRange[i].Reset();
+		}
 	}
 
 	mutable FCriticalSection				Lock;
@@ -247,6 +277,7 @@ struct FPlaybackState
 	TArray<FTrackMetadata>					AudioTracks;
 	TArray<FTrackMetadata>					SubtitleTracks;
 	TArray<FTimespan>						SeekablePositions;
+	FTimeRange								CurrentSegmentDownloadTimeRange[4];	// 0=video, 1=audio, 2=subtitiles, 3=UNSUPPORTED
 
 
 	void SetSeekableRange(const FTimeRange& TimeRange)
@@ -593,6 +624,22 @@ struct FPlaybackState
 		return EndPlaybackAtTime;
 	}
 
+	FTimeRange GetCurrentDownloadRequestTimeRange(EStreamType InStreamType)
+	{
+		FScopeLock lock(&Lock);
+		return CurrentSegmentDownloadTimeRange[StreamTypeToArrayIndex(InStreamType)];
+	}
+	void SetCurrentDownloadRequestTimeRange(EStreamType InStreamType, const FTimeRange& InRange)
+	{
+		FScopeLock lock(&Lock);
+		CurrentSegmentDownloadTimeRange[StreamTypeToArrayIndex(InStreamType)] = InRange;
+	}
+	void ClearCurrentDownloadRequestTimeRange(EStreamType InStreamType)
+	{
+		FScopeLock lock(&Lock);
+		CurrentSegmentDownloadTimeRange[StreamTypeToArrayIndex(InStreamType)].Reset();
+	}
+
 };
 
 
@@ -1066,6 +1113,8 @@ public:
 	void DeselectTrack(EStreamType StreamType) override;
 	bool IsTrackDeselected(EStreamType StreamType) override;
 
+	void QueryStreamBufferInfo(FStreamBufferInfo& OutStreamBufferInfo, EStreamType InStreamType) override;
+
 	void SuspendOrResumeDecoders(bool bSuspend, const FParamDict& InOptions) override;
 
 	void DebugPrint(void* pPlayer, void (*debugDrawPrintf)(void* pPlayer, const char *pFmt, ...));
@@ -2288,6 +2337,7 @@ private:
 	FPlayerSequenceState												CurrentPlaybackSequenceState;
 
 	uint32																CurrentPlaybackSequenceID[4]; // 0=video, 1=audio, 2=subtitiles, 3=UNSUPPORTED
+	FTimeRange															CurrentSegmentDownloadTimeRange[4];
 
 	FVideoRenderer														VideoRender;
 	FAudioRenderer														AudioRender;
@@ -2298,10 +2348,10 @@ private:
 	TArray<TWeakPtrTS<IAdaptiveStreamingPlayerSubtitleReceiver>>		SubtitleReceivers;
 
 
-	FMediaCriticalSection												MetricListenerCriticalSection;
+	FCriticalSection													MetricListenerCriticalSection;
 	TArray<IAdaptiveStreamingPlayerMetrics*, TInlineAllocator<4>>		MetricListeners;
 
-	FMediaCriticalSection												DiagnosticsCriticalSection;
+	mutable FCriticalSection											DiagnosticsCriticalSection;
 	FBufferStats														VideoBufferStats;
 	FBufferStats														AudioBufferStats;
 	FBufferStats														TextBufferStats;
diff --git a/Engine/Plugins/Media/ElectraPlayer/Source/ElectraPlayerRuntime/Private/Runtime/Player/AdaptiveStreamingPlayerRenderers.cpp b/Engine/Plugins/Media/ElectraPlayer/Source/ElectraPlayerRuntime/Private/Runtime/Player/AdaptiveStreamingPlayerRenderers.cpp
--- a/Engine/Plugins/Media/ElectraPlayer/Source/ElectraPlayerRuntime/Private/Runtime/Player/AdaptiveStreamingPlayerRenderers.cpp
+++ b/Engine/Plugins/Media/ElectraPlayer/Source/ElectraPlayerRuntime/Private/Runtime/Player/AdaptiveStreamingPlayerRenderers.cpp
@@ -73,8 +73,7 @@ namespace Electra
 			// If there is no renderer to wrap there should not be a wrapper created in the first place!
 			check(WrappedRenderer.IsValid());
 			NumPendingReturnBuffers = 0;
-			NumEnqueuedSamples = 0;
-			EnqueuedDuration.SetToZero();
+			EnqueuedSamples.Empty();
 		}
 
 	private:
@@ -94,7 +93,7 @@ namespace Electra
 
 
 		FTimeValue GetEnqueuedSampleDuration() override;
-		int32 GetNumEnqueuedSamples(FTimeValue* OutOptionalDuration) override;
+		int32 GetNumEnqueuedSamples(TArray<FEnqueuedSampleInfo>* OutOptionalSampleInfos) override;
 
 		void DisableHoldbackOfFirstRenderableVideoFrame(bool bDisableHoldback) override;
 
@@ -150,7 +149,7 @@ namespace Electra
 				bool DiffersFrom(int32 InSampleRate, int32 InNumChannels)
 				{ return InSampleRate != SampleRate || InNumChannels != NumChannels; }
 				void Update(int32 InSampleRate, int32 InNumChannels)
-				{ 
+				{
 					SampleRate = InSampleRate;
 					NumChannels = InNumChannels;
 				}
@@ -231,8 +230,7 @@ namespace Electra
 		FAudioVars AudioVars;
 
 		// Stats
-		int32 NumEnqueuedSamples = 0;
-		FTimeValue EnqueuedDuration;
+		TArray<FEnqueuedSampleInfo> EnqueuedSamples;
 	};
 
 
@@ -256,7 +254,7 @@ void FAdaptiveStreamingWrappedRenderer::SampleReleasedToPool(IDecoderOutput* InD
 	check(InDecoderOutput);
 	if (InDecoderOutput && RenderClock.IsValid())
 	{
-		int64 ValidityValue = InDecoderOutput->GetMutablePropertyDictionary().GetValue(RenderOptionKeys::ValidityValue).SafeGetInt64(0);
+		int64 ValidityValue = InDecoderOutput->GetMutablePropertyDictionary().GetValue(RenderOptionKeys::ValidityValue).SafeGetInt64(-1);
 		if (ValidityValue == CurrentValidityValue)
 		{
 			FTimeValue RenderTime = InDecoderOutput->GetMutablePropertyDictionary().GetValue(RenderOptionKeys::PTS).SafeGetTimeValue(FTimeValue::GetInvalid());
@@ -275,16 +273,16 @@ void FAdaptiveStreamingWrappedRenderer::SampleReleasedToPool(IDecoderOutput* InD
 
 			{
 				FScopeLock lock(&Lock);
-
-				if (--NumEnqueuedSamples < 0)
-				{
-					NumEnqueuedSamples = 0;
-				}
-
-				EnqueuedDuration -= Duration;
-				if (!EnqueuedDuration.IsValid() || EnqueuedDuration < FTimeValue::GetZero())
+				check(EnqueuedSamples.Num());
+				bool bFound = false;
+				for(int32 i=0; i<EnqueuedSamples.Num(); ++i)
 				{
-					EnqueuedDuration = FTimeValue::GetZero();
+					if (EnqueuedSamples[i].PTS == RenderTime)
+					{
+						EnqueuedSamples.RemoveAt(i);
+						bFound = true;
+						break;
+					}
 				}
 			}
 		}
@@ -303,10 +301,6 @@ UEMediaError FAdaptiveStreamingWrappedRenderer::CreateBufferPool(const FParamDic
 
 	FParamDict Parameters(InParameters);
 
-	NumBuffersInCirculation = 0;
-	NumEnqueuedSamples = 0;
-	EnqueuedDuration.SetToZero();
-
 	// Ask for larger buffers in case of audio. For playback speed changes we may need to create artificial
 	// samples to slow down audio playback and need larger buffers for that.
 	if (Type == EStreamType::Audio)
@@ -336,7 +330,12 @@ UEMediaError FAdaptiveStreamingWrappedRenderer::CreateBufferPool(const FParamDic
 			AudioVars.TempoChanger->SetMaxOutputSamples(AudioVars.MaxOutputSampleBlockSize);
 		}
 	}
-	return WrappedRenderer->CreateBufferPool(Parameters);
+	UEMediaError Error = WrappedRenderer->CreateBufferPool(Parameters);
+	// Clear the buffer bookkeeping values as creating a buffer _may_ call `SampleReleasedToPool()` to populate
+	// its internal structures without us having requested a buffer yet.
+	NumBuffersInCirculation = 0;
+	EnqueuedSamples.Empty();
+	return Error;
 }
 
 UEMediaError FAdaptiveStreamingWrappedRenderer::AcquireBuffer(IBuffer*& OutBuffer, int32 TimeoutInMicroseconds, const FParamDict& InParameters)
@@ -452,8 +451,9 @@ UEMediaError FAdaptiveStreamingWrappedRenderer::ReturnBufferCommon(IBuffer* Buff
 	bool bIsUnusedReturnBuffer = bRender == false && InSampleProperties.GetValue(RenderOptionKeys::EOSFlag).SafeGetBool(false) == false;
 
 	FScopeLock lock(&Lock);
-	EnqueuedDuration += Duration;
-	++NumEnqueuedSamples;
+	FEnqueuedSampleInfo& enqInf = EnqueuedSamples.Emplace_GetRef();
+	enqInf.Duration = Duration;
+	enqInf.PTS = InSampleProperties.GetValue(RenderOptionKeys::PTS).SafeGetTimeValue(FTimeValue::GetInvalid());;
 
 	if (!bIsUnusedReturnBuffer)
 	{
@@ -493,8 +493,7 @@ UEMediaError FAdaptiveStreamingWrappedRenderer::ReleaseBufferPool()
 	LLM_SCOPE(ELLMTag::ElectraPlayer);
 
 	ReturnAllPendingBuffers(false);
-	NumEnqueuedSamples = 0;
-	EnqueuedDuration.SetToZero();
+	EnqueuedSamples.Empty();
 	AudioVars.Reset();
 	Lock.Unlock();
 
@@ -534,9 +533,8 @@ UEMediaError FAdaptiveStreamingWrappedRenderer::Flush(const FParamDict& InOption
 
 	ReturnAllPendingBuffers(true);
 	++CurrentValidityValue;
-	NumEnqueuedSamples = 0;
 	NumBuffersNotHeldBack = 0;
-	EnqueuedDuration.SetToZero();
+	EnqueuedSamples.Empty();
 	AudioVars.Reset();
 	Lock.Unlock();
 
@@ -584,10 +582,18 @@ void FAdaptiveStreamingWrappedRenderer::ReturnAllPendingBuffers(bool bForFlush)
 FTimeValue FAdaptiveStreamingWrappedRenderer::GetEnqueuedSampleDuration()
 {
 	FScopeLock lock(&Lock);
-	return EnqueuedDuration;
+	FTimeValue dur(FTimeValue::GetZero());
+	for(auto &it : EnqueuedSamples)
+	{
+		if (it.PTS.IsValid())
+		{
+			dur += it.Duration;
+		}
+	}
+	return dur;
 }
 
-int32 FAdaptiveStreamingWrappedRenderer::GetNumEnqueuedSamples(FTimeValue* OutOptionalDuration)
+int32 FAdaptiveStreamingWrappedRenderer::GetNumEnqueuedSamples(TArray<IAdaptiveStreamingWrappedRenderer::FEnqueuedSampleInfo>* OutOptionalSampleInfos)
 {
 	FScopeLock lock(&Lock);
 
@@ -599,11 +605,17 @@ int32 FAdaptiveStreamingWrappedRenderer::GetNumEnqueuedSamples(FTimeValue* OutOp
 		DurAvail = FTimeValue::GetZero();
 	}
 
-	if (OutOptionalDuration)
+	if (OutOptionalSampleInfos)
 	{
-		*OutOptionalDuration = EnqueuedDuration + DurAvail;
+		for(auto &it : EnqueuedSamples)
+		{
+			if (it.PTS.IsValid())
+			{
+				(*OutOptionalSampleInfos).Emplace(it);
+			}
+		}
 	}
-	return NumEnqueuedSamples + NumAvail;
+	return EnqueuedSamples.Num() + NumAvail;
 }
 
 void FAdaptiveStreamingWrappedRenderer::DisableHoldbackOfFirstRenderableVideoFrame(bool bInDisableHoldback)
diff --git a/Engine/Plugins/Media/ElectraPlayer/Source/ElectraPlayerRuntime/Private/Runtime/Player/DASH/StreamReaderDASH.cpp b/Engine/Plugins/Media/ElectraPlayer/Source/ElectraPlayerRuntime/Private/Runtime/Player/DASH/StreamReaderDASH.cpp
--- a/Engine/Plugins/Media/ElectraPlayer/Source/ElectraPlayerRuntime/Private/Runtime/Player/DASH/StreamReaderDASH.cpp
+++ b/Engine/Plugins/Media/ElectraPlayer/Source/ElectraPlayerRuntime/Private/Runtime/Player/DASH/StreamReaderDASH.cpp
@@ -122,6 +122,16 @@ FTimeValue FStreamSegmentRequestDASH::GetFirstPTS() const
 	return AST + AdditionalAdjustmentTime + PeriodStart + FTimeValue((Segment.bFrameAccuracyRequired ? Segment.MediaLocalFirstPTS : Segment.Time) - Segment.PTO, Segment.Timescale);
 }
 
+FTimeRange FStreamSegmentRequestDASH::GetTimeRange() const
+{
+	FTimeRange tr;
+	tr.Start = AST + AdditionalAdjustmentTime + PeriodStart + FTimeValue(Segment.Time - Segment.PTO, Segment.Timescale);
+	tr.End = AST + AdditionalAdjustmentTime + PeriodStart + FTimeValue(Segment.Time + Segment.Duration - Segment.PTO, Segment.Timescale);
+	tr.Start.SetSequenceIndex(TimestampSequenceIndex);
+	tr.End.SetSequenceIndex(TimestampSequenceIndex);
+	return tr;
+}
+
 int32 FStreamSegmentRequestDASH::GetQualityIndex() const
 {
 	return Representation->GetQualityIndex();
diff --git a/Engine/Plugins/Media/ElectraPlayer/Source/ElectraPlayerRuntime/Private/Runtime/Player/DASH/StreamReaderDASH.h b/Engine/Plugins/Media/ElectraPlayer/Source/ElectraPlayerRuntime/Private/Runtime/Player/DASH/StreamReaderDASH.h
--- a/Engine/Plugins/Media/ElectraPlayer/Source/ElectraPlayerRuntime/Private/Runtime/Player/DASH/StreamReaderDASH.h
+++ b/Engine/Plugins/Media/ElectraPlayer/Source/ElectraPlayerRuntime/Private/Runtime/Player/DASH/StreamReaderDASH.h
@@ -32,8 +32,8 @@ public:
 	void GetRequestedStreams(TArray<TSharedPtrTS<IStreamSegment>>& OutRequestedStreams) override;
 	void GetEndedStreams(TArray<TSharedPtrTS<IStreamSegment>>& OutAlreadyEndedStreams) override;
 
-	//! Returns the first PTS value as indicated by the media timeline. This should correspond to the actual absolute PTS of the sample.
 	FTimeValue GetFirstPTS() const override;
+	FTimeRange GetTimeRange() const override;
 
 	int32 GetQualityIndex() const override;
 	int32 GetBitrate() const override;
diff --git a/Engine/Plugins/Media/ElectraPlayer/Source/ElectraPlayerRuntime/Private/Runtime/Player/HLS/StreamReaderHLSfmp4.cpp b/Engine/Plugins/Media/ElectraPlayer/Source/ElectraPlayerRuntime/Private/Runtime/Player/HLS/StreamReaderHLSfmp4.cpp
--- a/Engine/Plugins/Media/ElectraPlayer/Source/ElectraPlayerRuntime/Private/Runtime/Player/HLS/StreamReaderHLSfmp4.cpp
+++ b/Engine/Plugins/Media/ElectraPlayer/Source/ElectraPlayerRuntime/Private/Runtime/Player/HLS/StreamReaderHLSfmp4.cpp
@@ -116,6 +116,16 @@ FTimeValue FStreamSegmentRequestHLSfmp4::GetFirstPTS() const
 	return EarliestPTS.IsValid() && !EarliestPTS.IsPositiveInfinity() && EarliestPTS > AbsoluteDateTime ? EarliestPTS : AbsoluteDateTime;
 }
 
+FTimeRange FStreamSegmentRequestHLSfmp4::GetTimeRange() const
+{
+	FTimeRange tr;
+	tr.Start = AbsoluteDateTime;
+	tr.End = AbsoluteDateTime + SegmentDuration;
+	tr.Start.SetSequenceIndex(TimestampSequenceIndex);
+	tr.End.SetSequenceIndex(TimestampSequenceIndex);
+	return tr;
+}
+
 int32 FStreamSegmentRequestHLSfmp4::GetQualityIndex() const
 {
 	return QualityLevel;
diff --git a/Engine/Plugins/Media/ElectraPlayer/Source/ElectraPlayerRuntime/Private/Runtime/Player/HLS/StreamReaderHLSfmp4.h b/Engine/Plugins/Media/ElectraPlayer/Source/ElectraPlayerRuntime/Private/Runtime/Player/HLS/StreamReaderHLSfmp4.h
--- a/Engine/Plugins/Media/ElectraPlayer/Source/ElectraPlayerRuntime/Private/Runtime/Player/HLS/StreamReaderHLSfmp4.h
+++ b/Engine/Plugins/Media/ElectraPlayer/Source/ElectraPlayerRuntime/Private/Runtime/Player/HLS/StreamReaderHLSfmp4.h
@@ -31,8 +31,8 @@ public:
 	void GetRequestedStreams(TArray<TSharedPtrTS<IStreamSegment>>& OutRequestedStreams) override;
 	void GetEndedStreams(TArray<TSharedPtrTS<IStreamSegment>>& OutAlreadyEndedStreams) override;
 
-	//! Returns the first PTS value as indicated by the media timeline. This should correspond to the actual absolute PTS of the sample.
 	FTimeValue GetFirstPTS() const override;
+	FTimeRange GetTimeRange() const override;
 
 	int32 GetQualityIndex() const override;
 	int32 GetBitrate() const override;
diff --git a/Engine/Plugins/Media/ElectraPlayer/Source/ElectraPlayerRuntime/Private/Runtime/Player/Manifest.h b/Engine/Plugins/Media/ElectraPlayer/Source/ElectraPlayerRuntime/Private/Runtime/Player/Manifest.h
--- a/Engine/Plugins/Media/ElectraPlayer/Source/ElectraPlayerRuntime/Private/Runtime/Player/Manifest.h
+++ b/Engine/Plugins/Media/ElectraPlayer/Source/ElectraPlayerRuntime/Private/Runtime/Player/Manifest.h
@@ -126,9 +126,16 @@ namespace Electra
 		 */
 		virtual void GetEndedStreams(TArray<TSharedPtrTS<IStreamSegment>>& OutAlreadyEndedStreams) = 0;
 
-		//! Returns the first PTS value as indicated by the media timeline. This should correspond to the actual absolute PTS of the sample.
+		/**
+		 * Returns the first PTS value as indicated by the media timeline. This should correspond to the actual absolute PTS of the sample.
+		 */
 		virtual FTimeValue GetFirstPTS() const = 0;
 
+		/**
+		 * Returns the time range this request covers.
+		 */
+		virtual FTimeRange GetTimeRange() const = 0;
+
 		virtual int32 GetQualityIndex() const = 0;
 		virtual int32 GetBitrate() const = 0;
 
diff --git a/Engine/Plugins/Media/ElectraPlayer/Source/ElectraPlayerRuntime/Private/Runtime/Player/mkv/StreamReaderMKV.cpp b/Engine/Plugins/Media/ElectraPlayer/Source/ElectraPlayerRuntime/Private/Runtime/Player/mkv/StreamReaderMKV.cpp
--- a/Engine/Plugins/Media/ElectraPlayer/Source/ElectraPlayerRuntime/Private/Runtime/Player/mkv/StreamReaderMKV.cpp
+++ b/Engine/Plugins/Media/ElectraPlayer/Source/ElectraPlayerRuntime/Private/Runtime/Player/mkv/StreamReaderMKV.cpp
@@ -103,6 +103,16 @@ FTimeValue FStreamSegmentRequestMKV::GetFirstPTS() const
 	return EarliestPTS > FirstPTS ? EarliestPTS : FirstPTS;
 }
 
+FTimeRange FStreamSegmentRequestMKV::GetTimeRange() const
+{
+	FTimeRange tr;
+	tr.Start = FirstPTS;
+	tr.End = FirstPTS + SegmentDuration;
+	tr.Start.SetSequenceIndex(TimestampSequenceIndex);
+	tr.End.SetSequenceIndex(TimestampSequenceIndex);
+	return tr;
+}
+
 int32 FStreamSegmentRequestMKV::GetQualityIndex() const
 {
 	// No quality choice here.
@@ -831,7 +841,7 @@ FStreamReaderMKV::EEmitResult FStreamReaderMKV::EmitSamples(EEmitType InEmitType
 				{
 					break;
 				}
-				
+
 				// If emitting only one sample we leave this loop.
 				if (InEmitType == EEmitType::One)
 				{
diff --git a/Engine/Plugins/Media/ElectraPlayer/Source/ElectraPlayerRuntime/Private/Runtime/Player/mkv/StreamReaderMKV.h b/Engine/Plugins/Media/ElectraPlayer/Source/ElectraPlayerRuntime/Private/Runtime/Player/mkv/StreamReaderMKV.h
--- a/Engine/Plugins/Media/ElectraPlayer/Source/ElectraPlayerRuntime/Private/Runtime/Player/mkv/StreamReaderMKV.h
+++ b/Engine/Plugins/Media/ElectraPlayer/Source/ElectraPlayerRuntime/Private/Runtime/Player/mkv/StreamReaderMKV.h
@@ -35,8 +35,8 @@ public:
 	void GetRequestedStreams(TArray<TSharedPtrTS<IStreamSegment>>& OutRequestedStreams) override;
 	void GetEndedStreams(TArray<TSharedPtrTS<IStreamSegment>>& OutAlreadyEndedStreams) override;
 
-	//! Returns the first PTS value as indicated by the media timeline. This should correspond to the actual absolute PTS of the sample.
 	FTimeValue GetFirstPTS() const override;
+	FTimeRange GetTimeRange() const override;
 
 	int32 GetQualityIndex() const override;
 	int32 GetBitrate() const override;
diff --git a/Engine/Plugins/Media/ElectraPlayer/Source/ElectraPlayerRuntime/Private/Runtime/Player/mp4/StreamReaderMP4.cpp b/Engine/Plugins/Media/ElectraPlayer/Source/ElectraPlayerRuntime/Private/Runtime/Player/mp4/StreamReaderMP4.cpp
--- a/Engine/Plugins/Media/ElectraPlayer/Source/ElectraPlayerRuntime/Private/Runtime/Player/mp4/StreamReaderMP4.cpp
+++ b/Engine/Plugins/Media/ElectraPlayer/Source/ElectraPlayerRuntime/Private/Runtime/Player/mp4/StreamReaderMP4.cpp
@@ -116,6 +116,16 @@ FTimeValue FStreamSegmentRequestMP4::GetFirstPTS() const
 	return EarliestPTS > FirstPTS ? EarliestPTS : FirstPTS;
 }
 
+FTimeRange FStreamSegmentRequestMP4::GetTimeRange() const
+{
+	FTimeRange tr;
+	tr.Start = FirstPTS;
+	tr.End = FirstPTS + SegmentDuration;
+	tr.Start.SetSequenceIndex(TimestampSequenceIndex);
+	tr.End.SetSequenceIndex(TimestampSequenceIndex);
+	return tr;
+}
+
 int32 FStreamSegmentRequestMP4::GetQualityIndex() const
 {
 	// No quality choice here.
@@ -593,7 +603,7 @@ void FStreamReaderMP4::HandleRequest()
 						// If we need to decrypt we have to wait for the decrypter to become ready.
 						if (bIsSampleEncrypted && Decrypter.IsValid())
 						{
-							while(!bTerminate && !HasBeenAborted() && 
+							while(!bTerminate && !HasBeenAborted() &&
 								(Decrypter->GetState() == ElectraCDM::ECDMState::WaitingForKey || Decrypter->GetState() == ElectraCDM::ECDMState::Idle))
 							{
 								FMediaRunnable::SleepMilliseconds(100);
diff --git a/Engine/Plugins/Media/ElectraPlayer/Source/ElectraPlayerRuntime/Private/Runtime/Player/mp4/StreamReaderMP4.h b/Engine/Plugins/Media/ElectraPlayer/Source/ElectraPlayerRuntime/Private/Runtime/Player/mp4/StreamReaderMP4.h
--- a/Engine/Plugins/Media/ElectraPlayer/Source/ElectraPlayerRuntime/Private/Runtime/Player/mp4/StreamReaderMP4.h
+++ b/Engine/Plugins/Media/ElectraPlayer/Source/ElectraPlayerRuntime/Private/Runtime/Player/mp4/StreamReaderMP4.h
@@ -41,8 +41,8 @@ public:
 	void GetRequestedStreams(TArray<TSharedPtrTS<IStreamSegment>>& OutRequestedStreams) override;
 	void GetEndedStreams(TArray<TSharedPtrTS<IStreamSegment>>& OutAlreadyEndedStreams) override;
 
-	//! Returns the first PTS value as indicated by the media timeline. This should correspond to the actual absolute PTS of the sample.
 	FTimeValue GetFirstPTS() const override;
+	FTimeRange GetTimeRange() const override;
 
 	int32 GetQualityIndex() const override;
 	int32 GetBitrate() const override;
diff --git a/Engine/Plugins/Media/ElectraPlayer/Source/ElectraPlayerRuntime/Private/Runtime/StreamAccessUnitBuffer.h b/Engine/Plugins/Media/ElectraPlayer/Source/ElectraPlayerRuntime/Private/Runtime/StreamAccessUnitBuffer.h
--- a/Engine/Plugins/Media/ElectraPlayer/Source/ElectraPlayerRuntime/Private/Runtime/StreamAccessUnitBuffer.h
+++ b/Engine/Plugins/Media/ElectraPlayer/Source/ElectraPlayerRuntime/Private/Runtime/StreamAccessUnitBuffer.h
@@ -205,6 +205,8 @@ namespace Electra
 		void Clear()
 		{
 			FrontDTS.SetToInvalid();
+			SmallestPTS.SetToInvalid();
+			LargestPTSPlusDur.SetToInvalid();
 			PushedDuration.SetToZero();
 			PlayableDuration.SetToZero();
 			CurrentMemInUse = 0;
@@ -215,6 +217,8 @@ namespace Electra
 		}
 
 		FTimeValue			FrontDTS;
+		FTimeValue			SmallestPTS;
+		FTimeValue			LargestPTSPlusDur;
 		FTimeValue			PushedDuration;
 		FTimeValue			PlayableDuration;
 		int64				CurrentMemInUse;
@@ -247,6 +251,8 @@ namespace Electra
 
 		FAccessUnitBuffer()
 			: FrontDTS(FTimeValue::GetInvalid())
+			, SmallestPTS(FTimeValue::GetInvalid())
+			, LargestPTSPlusDur(FTimeValue::GetInvalid())
 			, PushedDuration(FTimeValue::GetZero())
 			, PlayableDuration(FTimeValue::GetZero())
 			, CurrentMemInUse(0)
@@ -289,6 +295,8 @@ namespace Electra
 		{
 			FScopeLock Lock(&AccessLock);
 			OutStats.FrontDTS = FrontDTS;
+			OutStats.SmallestPTS = SmallestPTS;
+			OutStats.LargestPTSPlusDur = LargestPTSPlusDur;
 			OutStats.PushedDuration = PushedDuration;
 			OutStats.PlayableDuration = PlayableDuration;
 			OutStats.CurrentMemInUse = CurrentMemInUse;
@@ -320,6 +328,16 @@ namespace Electra
 					{
 						FrontDTS = AU->DTS;
 					}
+					if (!SmallestPTS.IsValid() || AU->PTS < SmallestPTS)
+					{
+						SmallestPTS = AU->PTS;
+					}
+					FTimeValue End = AU->PTS + AU->Duration;
+					End.SetSequenceIndex(AU->PTS.GetSequenceIndex());
+					if (!LargestPTSPlusDur.IsValid() || End > LargestPTSPlusDur)
+					{
+						LargestPTSPlusDur = End;
+					}
 					PlayableDuration += AU->Duration;
 				}
 
@@ -382,6 +400,23 @@ namespace Electra
 							}
 						}
 					}
+					SmallestPTS.SetToPositiveInfinity();
+					for(int32 i=0,iMax=AccessUnits.Num(),j=0; i<iMax; ++i)
+					{
+						if (AccessUnits[i]->DropState == FAccessUnit::EDropState::None)
+						{
+							if (AccessUnits[i]->PTS < SmallestPTS)
+							{
+								SmallestPTS = AccessUnits[i]->PTS;
+							}
+							// Look only at the first couple of AU's. The smallest one is going to be among them
+							// unless there is a huge amount of reordered samples.
+							if (++j >= 10)
+							{
+								break;
+							}
+						}
+					}
 					if (OutAU->DropState == FAccessUnit::EDropState::None)
 					{
 						PlayableDuration -= OutAU->Duration;
@@ -391,6 +426,8 @@ namespace Electra
 				else
 				{
 					FrontDTS.SetToInvalid();
+					SmallestPTS.SetToInvalid();
+					LargestPTSPlusDur.SetToInvalid();
 					PlayableDuration.SetToZero();
 					PushedDuration.SetToZero();
 				}
@@ -638,6 +675,8 @@ namespace Electra
 		TMediaQueueDynamicNoLock<FAccessUnit*>		AccessUnits;
 		FMediaSemaphore								NumInSemaphore;
 		FTimeValue									FrontDTS;					//!< DTS of first AU in buffer
+		FTimeValue									SmallestPTS;				//!< Smallest PTS in buffer
+		FTimeValue									LargestPTSPlusDur;			//!< Largest PTS plus the AU duration in buffer
 		FTimeValue									PushedDuration;
 		FTimeValue									PlayableDuration;
 		int64										CurrentMemInUse = 0;
@@ -828,19 +867,17 @@ namespace Electra
 			}
 			void Clear()
 			{
-				ReadyDuration.SetToInvalid();
-				NumDecodedElementsReady = 0;
-				MaxDecodedElementsReady = 0;
+				InDecoderTimeRangePTS.Reset();
+				OutputBufferPoolSize = 0;
 				NumElementsInDecoder = 0;
 				bOutputStalled = false;
 				bEODreached = false;
 			}
-			FTimeValue				ReadyDuration;				//!< Duration of the ready material
-			int64					NumDecodedElementsReady;	//!< Number of decoded elements ready for rendering
-			int64					MaxDecodedElementsReady;	//!< Maximum number of decoded elements.
-			int64					NumElementsInDecoder;		//!< Number of elements currently in the decoder pipeline
+			FTimeRange				InDecoderTimeRangePTS;		//!< Time range of elements in the decoder pipeline by PTS.
+			int64					OutputBufferPoolSize;		//!< Maximum number of decoded elements the output pool can hold.
+			int64					NumElementsInDecoder;		//!< Number of elements currently in the decoder pipeline.
 			bool					bOutputStalled;				//!< true if the output is full and decoding is delayed until there's room again.
-			bool					bEODreached;				//!< true when the final decoded element has been passed on.
+			bool					bEODreached;				//!< true when the final decoded element has been passed on (but may still be in the queue).
 		};
 
 		virtual void DecoderOutputReady(const FDecodeReadyStats& CurrentReadyStats) = 0;
diff --git a/Engine/Plugins/Media/ElectraPlayer/Source/ElectraPlayerRuntime/Public/ElectraPlayerMisc.h b/Engine/Plugins/Media/ElectraPlayer/Source/ElectraPlayerRuntime/Public/ElectraPlayerMisc.h
--- a/Engine/Plugins/Media/ElectraPlayer/Source/ElectraPlayerRuntime/Public/ElectraPlayerMisc.h
+++ b/Engine/Plugins/Media/ElectraPlayer/Source/ElectraPlayerRuntime/Public/ElectraPlayerMisc.h
@@ -7,8 +7,8 @@
 #include "Stats/Stats.h"
 
 /** Log category for the ElectraPlayer module. */
-DECLARE_LOG_CATEGORY_EXTERN(LogElectraPlayer, Log, All);
-DECLARE_LOG_CATEGORY_EXTERN(LogElectraPlayerStats, Log, All);
+ELECTRAPLAYERRUNTIME_API DECLARE_LOG_CATEGORY_EXTERN(LogElectraPlayer, Log, All);
+ELECTRAPLAYERRUNTIME_API DECLARE_LOG_CATEGORY_EXTERN(LogElectraPlayerStats, Log, All);
 
 /** CSV Profiler */
 CSV_DECLARE_CATEGORY_MODULE_EXTERN(ELECTRAPLAYERRUNTIME_API, MediaStreaming);
diff --git a/Engine/Plugins/Media/ElectraPlayer/Source/ElectraPlayerRuntime/Public/IElectraPlayerInterface.h b/Engine/Plugins/Media/ElectraPlayer/Source/ElectraPlayerRuntime/Public/IElectraPlayerInterface.h
--- a/Engine/Plugins/Media/ElectraPlayer/Source/ElectraPlayerRuntime/Public/IElectraPlayerInterface.h
+++ b/Engine/Plugins/Media/ElectraPlayer/Source/ElectraPlayerRuntime/Public/IElectraPlayerInterface.h
@@ -324,6 +324,23 @@ public:
 
 	virtual void NotifyOfOptionChange() = 0;
 
+	struct FStreamBufferInfo
+	{
+		struct FTimeValue
+		{
+			FTimespan Time;
+			int64 SequenceIndex = -1;
+		};
+		struct FTimeRange
+		{
+			FTimeValue Start;
+			FTimeValue End;
+		};
+		TArray<FTimeRange> TimeAvailable;
+		TArray<FTimeRange> TimeRequested;
+	};
+	virtual bool GetStreamBufferInformation(FStreamBufferInfo& OutBufferInformation, EPlayerTrackType InTrackType) const = 0;
+
 	// Suspends or resumes decoder instances. Not supported on all platforms.
 	virtual void SuspendOrResumeDecoders(bool bSuspend, const Electra::FParamDict& InOptions) = 0;
 
