ElectraPlayer: fixed a static code analysis warning

#rnx

diff --git a/Engine/Plugins/Media/ElectraPlayer/Source/ElectraPlayerRuntime/Private/Runtime/Player/AdaptiveStreamingPlayer.cpp b/Engine/Plugins/Media/ElectraPlayer/Source/ElectraPlayerRuntime/Private/Runtime/Player/AdaptiveStreamingPlayer.cpp
--- a/Engine/Plugins/Media/ElectraPlayer/Source/ElectraPlayerRuntime/Private/Runtime/Player/AdaptiveStreamingPlayer.cpp
+++ b/Engine/Plugins/Media/ElectraPlayer/Source/ElectraPlayerRuntime/Private/Runtime/Player/AdaptiveStreamingPlayer.cpp
@@ -1019,9 +1019,8 @@ void FAdaptiveStreamingPlayer::QueryStreamBufferInfo(FStreamBufferInfo& OutStrea
 				// These AUs are not in the player buffer any more.
 				{
 					FScopeLock lock(&DiagnosticsCriticalSection);
-					IDecoderOutputBufferListener::FDecodeReadyStats* Ready = InStreamType == EStreamType::Video ? &VideoBufferStats.DecoderOutputBuffer :
-																			 InStreamType == EStreamType::Audio ? &AudioBufferStats.DecoderOutputBuffer :
-																			 nullptr;
+					// Video or audio buffer? (see switch-case we are in above)
+					IDecoderOutputBufferListener::FDecodeReadyStats* Ready = InStreamType == EStreamType::Video ? &VideoBufferStats.DecoderOutputBuffer : &AudioBufferStats.DecoderOutputBuffer;
 					if (Ready && Ready->InDecoderTimeRangePTS.IsValid())
 					{
 						OutStreamBufferInfo.TimeEnqueued.Emplace(Ready->InDecoderTimeRangePTS);
