ElectraPlayer: Changed the http response cache init method parameters

#rnx

diff --git a/Engine/Plugins/Media/ElectraPlayer/Source/ElectraPlayerRuntime/Private/Runtime/HTTP/HTTPResponseCache.cpp b/Engine/Plugins/Media/ElectraPlayer/Source/ElectraPlayerRuntime/Private/Runtime/HTTP/HTTPResponseCache.cpp
--- a/Engine/Plugins/Media/ElectraPlayer/Source/ElectraPlayerRuntime/Private/Runtime/HTTP/HTTPResponseCache.cpp
+++ b/Engine/Plugins/Media/ElectraPlayer/Source/ElectraPlayerRuntime/Private/Runtime/HTTP/HTTPResponseCache.cpp
@@ -2,19 +2,15 @@
 
 #include "HTTPResponseCache.h"
 #include "ElectraHTTPStreamBuffer.h"
-#include "Player/PlayerSessionServices.h"
 
 
 namespace Electra
 {
-	static const FName OptionKeyResponseCacheMaxEntries(TEXT("httpcache_max_entries"));
-	static const FName OptionKeyResponseCacheMaxByteSize(TEXT("httpcache_max_bytesize"));
-
 
 class FHTTPResponseCache : public IHTTPResponseCache
 {
 public:
-	FHTTPResponseCache(IPlayerSessionServices* SessionServices, TSharedPtr<IElectraPlayerDataCache, ESPMode::ThreadSafe> ExternalCache);
+	FHTTPResponseCache(int64 InMaxByteCapacity, int32 InMaxEntries, TSharedPtr<IElectraPlayerDataCache, ESPMode::ThreadSafe> ExternalCache);
 	virtual ~FHTTPResponseCache();
 
 	void Disable() override;
@@ -304,11 +300,11 @@ private:
 			}
 		}
 		void AddHeader(FString InHeader, FString InValue)
-		{ 
+		{
 			FElectraHTTPStreamHeader Hdr;
 			Hdr.Header = MoveTemp(InHeader);
 			Hdr.Value = MoveTemp(InValue);
-			Headers.Emplace(MoveTemp(Hdr)); 
+			Headers.Emplace(MoveTemp(Hdr));
 		}
 		void SetURL(const FString& InUrl)
 		{ URL = InUrl; }
@@ -345,7 +341,6 @@ private:
 	void EvictToAddSize(int64 ResponseSize);
 
 	TSharedPtr<IElectraPlayerDataCache, ESPMode::ThreadSafe> ExternalCache;
-	IPlayerSessionServices* SessionServices = nullptr;
 	int64 MaxElementSize = 0;
 	int32 MaxNumElements = 0;
 
@@ -359,34 +354,23 @@ private:
 /***************************************************************************************************************************************************/
 /***************************************************************************************************************************************************/
 
-TSharedPtrTS<IHTTPResponseCache> IHTTPResponseCache::Create(IPlayerSessionServices* SessionServices, TSharedPtr<IElectraPlayerDataCache, ESPMode::ThreadSafe> ExternalCache)
+TSharedPtrTS<IHTTPResponseCache> IHTTPResponseCache::Create(int64 InMaxByteCapacity, int32 InMaxEntries, TSharedPtr<IElectraPlayerDataCache, ESPMode::ThreadSafe> ExternalCache)
 {
-	return MakeSharedTS<FHTTPResponseCache>(SessionServices, ExternalCache);
+	return MakeSharedTS<FHTTPResponseCache>(InMaxByteCapacity, InMaxEntries, ExternalCache);
 }
 
 /***************************************************************************************************************************************************/
 /***************************************************************************************************************************************************/
 /***************************************************************************************************************************************************/
 
-FHTTPResponseCache::FHTTPResponseCache(IPlayerSessionServices* InSessionServices, TSharedPtr<IElectraPlayerDataCache, ESPMode::ThreadSafe> InExternalCache)
-	: ExternalCache(InExternalCache), SessionServices(InSessionServices)
+FHTTPResponseCache::FHTTPResponseCache(int64 InMaxByteCapacity, int32 InMaxEntries, TSharedPtr<IElectraPlayerDataCache, ESPMode::ThreadSafe> InExternalCache)
+	: ExternalCache(InExternalCache)
 {
 	// Only configure our cache when there is no external one!
 	if (!ExternalCache.IsValid())
 	{
-		int64 v;
-		v = InSessionServices->GetOptionValue(OptionKeyResponseCacheMaxByteSize).SafeGetInt64(0);
-		if (v >= 0)
-		{
-			MaxElementSize = (int64)v;
-		}
-		// Max number of elements is probably not used a lot, if at all.
-		// If not specified we allow for some reasonable number.
-		v = InSessionServices->GetOptionValue(OptionKeyResponseCacheMaxEntries).SafeGetInt64(8192);
-		if (v >= 0)
-		{
-			MaxNumElements = (int32) v;
-		}
+		MaxElementSize = InMaxByteCapacity >= 0 ? InMaxByteCapacity : 0;
+		MaxNumElements = InMaxEntries >= 0 ? InMaxEntries : 0;
 	}
 }
 
diff --git a/Engine/Plugins/Media/ElectraPlayer/Source/ElectraPlayerRuntime/Private/Runtime/HTTP/HTTPResponseCache.h b/Engine/Plugins/Media/ElectraPlayer/Source/ElectraPlayerRuntime/Private/Runtime/HTTP/HTTPResponseCache.h
--- a/Engine/Plugins/Media/ElectraPlayer/Source/ElectraPlayerRuntime/Private/Runtime/HTTP/HTTPResponseCache.h
+++ b/Engine/Plugins/Media/ElectraPlayer/Source/ElectraPlayerRuntime/Private/Runtime/HTTP/HTTPResponseCache.h
@@ -12,7 +12,6 @@
 
 namespace Electra
 {
-class IPlayerSessionServices;
 
 class IHTTPResponseCache
 {
@@ -33,7 +32,7 @@ public:
 		FQualityInfo Quality;
 	};
 
-	static TSharedPtrTS<IHTTPResponseCache> Create(IPlayerSessionServices* SessionServices, TSharedPtr<IElectraPlayerDataCache, ESPMode::ThreadSafe> ExternalCache);
+	static TSharedPtrTS<IHTTPResponseCache> Create(int64 InMaxByteCapacity, int32 InMaxEntries, TSharedPtr<IElectraPlayerDataCache, ESPMode::ThreadSafe> ExternalCache);
 
 	virtual ~IHTTPResponseCache() = default;
 
diff --git a/Engine/Plugins/Media/ElectraPlayer/Source/ElectraPlayerRuntime/Private/Runtime/Player/AdaptivePlayerOptionKeynames.h b/Engine/Plugins/Media/ElectraPlayer/Source/ElectraPlayerRuntime/Private/Runtime/Player/AdaptivePlayerOptionKeynames.h
--- a/Engine/Plugins/Media/ElectraPlayer/Source/ElectraPlayerRuntime/Private/Runtime/Player/AdaptivePlayerOptionKeynames.h
+++ b/Engine/Plugins/Media/ElectraPlayer/Source/ElectraPlayerRuntime/Private/Runtime/Player/AdaptivePlayerOptionKeynames.h
@@ -60,6 +60,8 @@ namespace Electra
 	const FName OptionKeyPreferredCodecsAudio(TEXT("preferred_codecs_audio"));
 	const FName OptionKeyPreferredCodecsSubtitles(TEXT("preferred_codecs_subtitles"));
 
+	const FName OptionKeyResponseCacheMaxEntries(TEXT("httpcache_max_entries"));
+	const FName OptionKeyResponseCacheMaxByteSize(TEXT("httpcache_max_bytesize"));
 } // namespace Electra
 
 
diff --git a/Engine/Plugins/Media/ElectraPlayer/Source/ElectraPlayerRuntime/Private/Runtime/Player/AdaptiveStreamingPlayer.cpp b/Engine/Plugins/Media/ElectraPlayer/Source/ElectraPlayerRuntime/Private/Runtime/Player/AdaptiveStreamingPlayer.cpp
--- a/Engine/Plugins/Media/ElectraPlayer/Source/ElectraPlayerRuntime/Private/Runtime/Player/AdaptiveStreamingPlayer.cpp
+++ b/Engine/Plugins/Media/ElectraPlayer/Source/ElectraPlayerRuntime/Private/Runtime/Player/AdaptiveStreamingPlayer.cpp
@@ -1468,7 +1468,7 @@ bool FAdaptiveStreamingPlayer::InternalHandleThreadMessages()
 				if (ExternalCache.IsValid())
 				{
 					HttpResponseCache.Reset();
-					HttpResponseCache = IHTTPResponseCache::Create(this, MoveTemp(ExternalCache));
+					HttpResponseCache = IHTTPResponseCache::Create(GetOptionValue(OptionKeyResponseCacheMaxByteSize).SafeGetInt64(0), (int32)GetOptionValue(OptionKeyResponseCacheMaxEntries).SafeGetInt64(8192), MoveTemp(ExternalCache));
 				}
 
 				FWorkerThreadMessages::FMessage::FLoadManifest& ev = msg.Data.Get<FWorkerThreadMessages::FMessage::FLoadManifest>();
@@ -4581,7 +4581,7 @@ void FAdaptiveStreamingPlayer::InternalInitialize()
 	EntityCache = IPlayerEntityCache::Create(this);
 
 	// Create an HTTP response cache. Hand over any externally set cache. We do not need it in here any further.
-	HttpResponseCache = IHTTPResponseCache::Create(this, MoveTemp(ExternalCache));
+	HttpResponseCache = IHTTPResponseCache::Create(GetOptionValue(OptionKeyResponseCacheMaxByteSize).SafeGetInt64(0), (int32)GetOptionValue(OptionKeyResponseCacheMaxEntries).SafeGetInt64(8192), MoveTemp(ExternalCache));
 
 	// If all read requests for this player are to be routed to an external reader we need to create a wrapper for it.
 	if (PlayerOptions.GetValue(OptionKeyUseExternalDataReader).SafeGetBool(false))
