ElectraPlayer: Added support for DASH self-initializing media segments

#rnx

diff --git a/Engine/Plugins/Media/ElectraPlayer/Source/ElectraPlayerRuntime/Private/Runtime/Player/DASH/StreamReaderDASH.cpp b/Engine/Plugins/Media/ElectraPlayer/Source/ElectraPlayerRuntime/Private/Runtime/Player/DASH/StreamReaderDASH.cpp
--- a/Engine/Plugins/Media/ElectraPlayer/Source/ElectraPlayerRuntime/Private/Runtime/Player/DASH/StreamReaderDASH.cpp
+++ b/Engine/Plugins/Media/ElectraPlayer/Source/ElectraPlayerRuntime/Private/Runtime/Player/DASH/StreamReaderDASH.cpp
@@ -467,6 +467,29 @@ void FStreamReaderDASH::FStreamHandler::HTTPUpdateStats(const FTimeValue& Curren
 
 
 
+void FStreamReaderDASH::FStreamHandler::SetupInitSegmentDownloadStatsFromRequestAndConnectionInfo(Metrics::FSegmentDownloadStats& ds, const TSharedPtrTS<FStreamSegmentRequestDASH>& Request, const HTTP::FConnectionInfo* ci, bool bWasSuccessful)
+{
+	ds.StatsID = FMediaInterlockedIncrement(UniqueDownloadID);
+	ds.StreamType = Request->GetType();
+	ds.SegmentType = Metrics::ESegmentType::Init;
+	ds.URL = Request->Segment.InitializationURL.URL;
+	ds.Range = Request->Segment.InitializationURL.Range;
+	ds.CDN = Request->Segment.InitializationURL.CDN;
+	ds.bWasSuccessful = bWasSuccessful;
+	ds.HTTPStatusCode = ci->StatusInfo.HTTPStatus;
+	ds.TimeToFirstByte = ci->TimeUntilFirstByte;
+	ds.TimeToDownload = (ci->RequestEndTime - ci->RequestStartTime).GetAsSeconds();
+	ds.ByteSize = ci->ContentLength;
+	ds.NumBytesDownloaded = ci->BytesReadSoFar;
+	ds.MediaAssetID = Request->Period.IsValid() ? Request->Period->GetUniqueIdentifier() : "";
+	ds.AdaptationSetID = Request->AdaptationSet.IsValid() ? Request->AdaptationSet->GetUniqueIdentifier() : "";
+	ds.RepresentationID = Request->Representation.IsValid() ? Request->Representation->GetUniqueIdentifier() : "";
+	ds.Bitrate = Request->GetBitrate();
+	ds.QualityIndex = Request->QualityIndex;
+	ds.HighestQualityIndex = Request->MaxQualityIndex;
+}
+
+
 
 FErrorDetail FStreamReaderDASH::FStreamHandler::LoadInitSegment(TSharedPtrTS<FMPDLoadRequestDASH>& OutLoadRequest, Metrics::FSegmentDownloadStats& ds, const TSharedPtrTS<FStreamSegmentRequestDASH>& Request)
 {
@@ -523,24 +546,7 @@ FErrorDetail FStreamReaderDASH::FStreamHandler::LoadInitSegment(TSharedPtrTS<FMP
 
 	// Set up download stats to be sent to the stream selector.
 	const HTTP::FConnectionInfo* ci = OutLoadRequest->GetConnectionInfo();
-	ds.StatsID = FMediaInterlockedIncrement(UniqueDownloadID);
-	ds.StreamType = Request->GetType();
-	ds.SegmentType = Metrics::ESegmentType::Init;
-	ds.URL = Request->Segment.InitializationURL.URL;
-	ds.Range = Request->Segment.InitializationURL.Range;
-	ds.CDN = Request->Segment.InitializationURL.CDN;
-	ds.bWasSuccessful = LoadResult->bSuccess;
-	ds.HTTPStatusCode = ci->StatusInfo.HTTPStatus;
-	ds.TimeToFirstByte = ci->TimeUntilFirstByte;
-	ds.TimeToDownload = (ci->RequestEndTime - ci->RequestStartTime).GetAsSeconds();
-	ds.ByteSize = ci->ContentLength;
-	ds.NumBytesDownloaded = ci->BytesReadSoFar;
-	ds.MediaAssetID = Request->Period.IsValid() ? Request->Period->GetUniqueIdentifier() : "";
-	ds.AdaptationSetID = Request->AdaptationSet.IsValid() ? Request->AdaptationSet->GetUniqueIdentifier() : "";
-	ds.RepresentationID = Request->Representation.IsValid() ? Request->Representation->GetUniqueIdentifier() : "";
-	ds.Bitrate = Request->GetBitrate();
-	ds.QualityIndex = Request->QualityIndex;
-	ds.HighestQualityIndex = Request->MaxQualityIndex;
+	SetupInitSegmentDownloadStatsFromRequestAndConnectionInfo(ds, Request, ci, LoadResult->bSuccess);
 
 	if (!LoadResult->bSuccess)
 	{
@@ -553,12 +559,12 @@ FErrorDetail FStreamReaderDASH::FStreamHandler::LoadInitSegment(TSharedPtrTS<FMP
 
 FErrorDetail FStreamReaderDASH::FStreamHandler::GetInitSegment(TSharedPtrTS<const IParserISO14496_12>& OutMP4InitSegment, const TSharedPtrTS<FStreamSegmentRequestDASH>& Request)
 {
-	// Is an init segment required?
-	if (Request->Segment.InitializationURL.URL.IsEmpty())
-	{
-		// No init segment required. We're done.
-		return FErrorDetail();
-	}
+	FString InitURL(Request->Segment.InitializationURL.URL);
+	FString InitRange(Request->Segment.InitializationURL.Range);
+	FString MediaURL(Request->Segment.MediaURL.URL);
+	HTTP::FConnectionInfo ConnectionInfo;
+	Metrics::FSegmentDownloadStats ds;
+
 	// Get the entity cache. If it's not there we return.
 	TSharedPtrTS<IPlayerEntityCache> EntityCache = PlayerSessionService ? PlayerSessionService->GetEntityCache() : nullptr;
 	if (!EntityCache.IsValid())
@@ -567,18 +573,70 @@ FErrorDetail FStreamReaderDASH::FStreamHandler::GetInitSegment(TSharedPtrTS<cons
 	}
 	// Check if we already have this cached.
 	IPlayerEntityCache::FCacheItem CachedItem;
-	if (EntityCache->GetCachedEntity(CachedItem, Request->Segment.InitializationURL.URL, Request->Segment.InitializationURL.Range))
+	if (EntityCache->GetCachedEntity(CachedItem, InitURL.IsEmpty() ? MediaURL : InitURL, InitRange))
 	{
 		// Already cached. Use it.
 		OutMP4InitSegment = CachedItem.Parsed14496_12Data;
 		return FErrorDetail();
 	}
-	TSharedPtrTS<FMPDLoadRequestDASH> LoadReq;
-	Metrics::FSegmentDownloadStats ds;
-	FErrorDetail LoadError = LoadInitSegment(LoadReq, ds, Request);
-	if (!LoadReq.IsValid() || LoadError.IsSet())
+
+	FMP4StaticDataReader StaticDataReader;
+
+	// Is an explicit init segment specified?
+	if (InitURL.IsEmpty())
 	{
-		return LoadError;
+		// No, this is a self-initializing segment. We need to locate the moov box manually.
+		InitURL = Request->Segment.MediaURL.URL;
+		UtilsMP4::FMP4RootBoxLocator BoxLocator;
+		TArray<UtilsMP4::FMP4RootBoxLocator::FBoxInfo> Boxes;
+		const TArray<uint32> FirstBox { UtilsMP4::Make4CC('f','t','y','p') };
+		const TArray<uint32> MoovBox { UtilsMP4::Make4CC('m','o','o','v') };
+		bool bSuccess = BoxLocator.LocateRootBoxes(Boxes, PlayerSessionService->GetHTTPManager(), InitURL, FirstBox, MoovBox, MoovBox, UtilsMP4::FMP4RootBoxLocator::FCancellationCheckDelegate::CreateLambda([&]()
+		{
+			return HasReadBeenAborted();
+		}));
+		if (HasReadBeenAborted())
+		{
+			return FErrorDetail();
+		}
+		ConnectionInfo = BoxLocator.GetConnectionInfo();
+		SetupInitSegmentDownloadStatsFromRequestAndConnectionInfo(ds, Request, &ConnectionInfo, bSuccess);
+		if (!bSuccess)
+		{
+			Request->ConnectionInfo = ConnectionInfo;
+			FString ErrMsg = BoxLocator.GetErrorMessage();
+			if (ErrMsg.IsEmpty())
+			{
+				ErrMsg = BoxLocator.GetConnectionInfo().StatusInfo.ErrorDetail.GetMessage();
+			}
+			return CreateError(FString::Printf(TEXT("Init segment download error: %s"),	*ErrMsg), INTERNAL_ERROR_INIT_SEGMENT_DOWNLOAD_ERROR);
+		}
+		const UtilsMP4::FMP4RootBoxLocator::FBoxInfo* Moov = Boxes.FindByPredicate([InType=UtilsMP4::Make4CC('m','o','o','v')](const UtilsMP4::FMP4RootBoxLocator::FBoxInfo& InBox){return InBox.Type == InType;});
+		if (Moov)
+		{
+			StaticDataReader.SetParseData(Moov->DataBuffer);
+		}
+		else
+		{
+			// Pass an empty buffer so parsing will fail below.
+			TSharedPtrTS<IElectraHttpManager::FReceiveBuffer> EmptyBuffer = MakeSharedTS<IElectraHttpManager::FReceiveBuffer>();
+			EmptyBuffer->Buffer.SetEOD();
+			StaticDataReader.SetParseData(EmptyBuffer);
+		}
+	}
+	else
+	{
+		TSharedPtrTS<FMPDLoadRequestDASH> LoadReq;
+		FErrorDetail LoadError = LoadInitSegment(LoadReq, ds, Request);
+		if (!LoadReq.IsValid() || LoadError.IsSet())
+		{
+			return LoadError;
+		}
+		if (LoadReq->Request.IsValid() && LoadReq->Request->GetResponseBuffer().IsValid())
+		{
+			StaticDataReader.SetParseData(LoadReq->Request->GetResponseBuffer());
+		}
+		ConnectionInfo = *LoadReq->GetConnectionInfo();
 	}
 
 	SCOPE_CYCLE_COUNTER(STAT_ElectraPlayer_DASH_StreamReader);
@@ -586,12 +644,7 @@ FErrorDetail FStreamReaderDASH::FStreamHandler::GetInitSegment(TSharedPtrTS<cons
 
 	TSharedPtrTS<IParserISO14496_12> Init = IParserISO14496_12::CreateParser();
 	UEMediaError parseError = UEMEDIA_ERROR_FORMAT_ERROR;
-	if (LoadReq->Request.IsValid() && LoadReq->Request->GetResponseBuffer().IsValid())
-	{
-		FMP4StaticDataReader StaticDataReader;
-		StaticDataReader.SetParseData(LoadReq->Request->GetResponseBuffer());
-		parseError = Init->ParseHeader(&StaticDataReader, this, PlayerSessionService, nullptr);
-	}
+	parseError = Init->ParseHeader(&StaticDataReader, this, PlayerSessionService, nullptr);
 	if (parseError == UEMEDIA_ERROR_OK || parseError == UEMEDIA_ERROR_END_OF_STREAM)
 	{
 		// Parse the tracks of the init segment. We do this mainly to get to the CSD we might need should we have to insert filler data later.
@@ -600,8 +653,8 @@ FErrorDetail FStreamReaderDASH::FStreamHandler::GetInitSegment(TSharedPtrTS<cons
 		{
 			// Add this to the entity cache in case it needs to be retrieved again.
 			IPlayerEntityCache::FCacheItem CacheItem;
-			CacheItem.URL = LoadReq->URL;
-			CacheItem.Range = LoadReq->Range;
+			CacheItem.URL = InitURL;
+			CacheItem.Range = InitRange;
 			CacheItem.Parsed14496_12Data = Init;
 			EntityCache->CacheEntity(CacheItem);
 			OutMP4InitSegment = Init;
@@ -611,17 +664,17 @@ FErrorDetail FStreamReaderDASH::FStreamHandler::GetInitSegment(TSharedPtrTS<cons
 		}
 		else
 		{
+			Request->ConnectionInfo = ConnectionInfo;
 			ds.bParseFailure = true;
-			Request->ConnectionInfo = *LoadReq->GetConnectionInfo();
 			StreamSelector->ReportDownloadEnd(ds);
-			return CreateError(FString::Printf(TEXT("Track preparation of init segment \"%s\" failed"), *LoadReq->URL), INTERNAL_ERROR_INIT_SEGMENT_PARSE_ERROR);
+			return CreateError(FString::Printf(TEXT("Track preparation of init segment \"%s\" failed"), *InitURL), INTERNAL_ERROR_INIT_SEGMENT_PARSE_ERROR);
 		}
 	}
 	else
 	{
 		ds.bParseFailure = true;
 		StreamSelector->ReportDownloadEnd(ds);
-		return CreateError(FString::Printf(TEXT("Parse error of init segment \"%s\""), *LoadReq->URL), INTERNAL_ERROR_INIT_SEGMENT_PARSE_ERROR);
+		return CreateError(FString::Printf(TEXT("Parse error of init segment \"%s\""), *InitURL), INTERNAL_ERROR_INIT_SEGMENT_PARSE_ERROR);
 	}
 }
 
diff --git a/Engine/Plugins/Media/ElectraPlayer/Source/ElectraPlayerRuntime/Private/Runtime/Player/DASH/StreamReaderDASH.h b/Engine/Plugins/Media/ElectraPlayer/Source/ElectraPlayerRuntime/Private/Runtime/Player/DASH/StreamReaderDASH.h
--- a/Engine/Plugins/Media/ElectraPlayer/Source/ElectraPlayerRuntime/Private/Runtime/Player/DASH/StreamReaderDASH.h
+++ b/Engine/Plugins/Media/ElectraPlayer/Source/ElectraPlayerRuntime/Private/Runtime/Player/DASH/StreamReaderDASH.h
@@ -271,6 +271,7 @@ private:
 		void HandleRequestMP4();
 		void HandleRequestMKV();
 
+		void SetupInitSegmentDownloadStatsFromRequestAndConnectionInfo(Metrics::FSegmentDownloadStats& ds, const TSharedPtrTS<FStreamSegmentRequestDASH>& Request, const HTTP::FConnectionInfo* ci, bool bWasSuccessful);
 		FErrorDetail LoadInitSegment(TSharedPtrTS<FMPDLoadRequestDASH>& OutLoadRequest, Metrics::FSegmentDownloadStats& OutStats, const TSharedPtrTS<FStreamSegmentRequestDASH>& Request);
 		FErrorDetail GetInitSegment(TSharedPtrTS<const IParserISO14496_12>& OutMP4InitSegment, const TSharedPtrTS<FStreamSegmentRequestDASH>& InRequest);
 		FErrorDetail GetInitSegment(TSharedPtrTS<const IParserMKV>& OutMKVInitSegment, const TSharedPtrTS<FStreamSegmentRequestDASH>& InRequest);
