ElectraPlayer: Changed the mp4 moov box reader to locate the box and read it in fewer requests and pass out the box data in a more direct way, significantly speeding up the parsing process

#rnx

diff --git a/Engine/Plugins/Media/ElectraPlayer/Source/ElectraPlayerRuntime/Private/Runtime/Player/mp4/PlaylistReaderMP4.cpp b/Engine/Plugins/Media/ElectraPlayer/Source/ElectraPlayerRuntime/Private/Runtime/Player/mp4/PlaylistReaderMP4.cpp
--- a/Engine/Plugins/Media/ElectraPlayer/Source/ElectraPlayerRuntime/Private/Runtime/Player/mp4/PlaylistReaderMP4.cpp
+++ b/Engine/Plugins/Media/ElectraPlayer/Source/ElectraPlayerRuntime/Private/Runtime/Player/mp4/PlaylistReaderMP4.cpp
@@ -3,12 +3,14 @@
 #include "PlayerCore.h"
 #include "Player/Manifest.h"
 #include "HTTP/HTTPManager.h"
+#include "HTTP/HTTPResponseCache.h"
 #include "Demuxer/ParserISO14496-12.h"
 #include "Stats/Stats.h"
 #include "SynchronizedClock.h"
 #include "PlaylistReaderMP4.h"
 #include "Utilities/HashFunctions.h"
 #include "Utilities/TimeUtilities.h"
+#include "Utilities/UtilsMP4.h"
 #include "HAL/LowLevelMemTracker.h"
 #include "ElectraPlayerPrivate.h"
 #include "Player/AdaptiveStreamingPlayerResourceRequest.h"
@@ -90,52 +92,23 @@ private:
 	void PostError(const FString& Message, uint16 Code, UEMediaError Error = UEMEDIA_ERROR_OK);
 	void LogMessage(IInfoLog::ELevel Level, const FString& Message);
 
-	FErrorDetail CreateErrorAndLog(const FString& Message, uint16 Code, UEMediaError Error = UEMEDIA_ERROR_OK);
+	IPlayerSessionServices* PlayerSessionServices = nullptr;
+	FString MasterPlaylistURL;
+	FString URLFragment;
+	FMediaEvent WorkerThreadQuitSignal;
+	bool bIsWorkerThreadStarted = false;
 
-	int32 HTTPProgressCallback(const IElectraHttpManager::FRequest* Request);
-	void HTTPCompletionCallback(const IElectraHttpManager::FRequest* Request);
+	HTTP::FConnectionInfo ConnectionInfo;
 
-	void ClearRequest();
+	TSharedPtrTS<IParserISO14496_12> MP4Parser;
+	TSharedPtrTS<IElectraHttpManager::FReceiveBuffer> ParseBuffer;
+	int64 ParsePos = 0;
+	int64 ParseBufferSize = 0;
 
-	void ReadNextChunk(int64 InFromOffset, int64 ChunkSize);
+	volatile bool bAbort = false;
 
-	bool HasErrored() const
-	{
-		return bHasErrored;
-	}
-
-	const int32 kChunkReadSize = 65536;
-
-	IPlayerSessionServices*									PlayerSessionServices = nullptr;
-	FString													MasterPlaylistURL;
-	FString													URLFragment;
-	FMediaEvent												WorkerThreadQuitSignal;
-	bool													bIsWorkerThreadStarted = false;
-
-	FCriticalSection										Lock;
-	TSharedPtrTS<IElectraHttpManager::FRequest> 			Request;
-	TSharedPtrTS<IElectraHttpManager::FReceiveBuffer> 		ReceiveBuffer;
-	TSharedPtrTS<IElectraHttpManager::FProgressListener>	ProgressListener;
-	HTTP::FConnectionInfo									ConnectionInfo;
-	FWaitableBuffer 										Buffer;
-	int64													ParsePos = 0;
-	int64													ChunkReadOffset = 0;
-	int64													FileSize = -1;
-	bool													bChunkReadInProgress = false;
-
-	bool													bAbort = false;
-	bool													bHasErrored = false;
-
-	TSharedPtrTS<IParserISO14496_12>						MP4Parser;
-	bool													bFoundBoxFTYP = false;
-	bool													bFoundBoxMOOV = false;
-	bool													bFoundBoxSIDX = false;
-	bool													bFoundBoxMOOF = false;
-	bool													bFoundBoxMDAT = false;
-	bool													bIsFastStartable = false;
-
-	TSharedPtrTS<FManifestMP4Internal>						Manifest;
-	FErrorDetail											LastErrorDetail;
+	TSharedPtrTS<FManifestMP4Internal> Manifest;
+	FErrorDetail LastErrorDetail;
 };
 
 
@@ -181,14 +154,11 @@ TSharedPtrTS<IManifest> FPlaylistReaderMP4::GetManifest()
 void FPlaylistReaderMP4::Initialize(IPlayerSessionServices* InPlayerSessionServices)
 {
 	PlayerSessionServices = InPlayerSessionServices;
-	Buffer.Reserve(kChunkReadSize);
 }
 
 void FPlaylistReaderMP4::Close()
 {
 	bAbort = true;
-	ClearRequest();
-
 	StopWorkerThread();
 }
 
@@ -230,21 +200,6 @@ void FPlaylistReaderMP4::PostError(const FString& InMessage, uint16 InCode, UEMe
 	}
 }
 
-FErrorDetail FPlaylistReaderMP4::CreateErrorAndLog(const FString& InMessage, uint16 InCode, UEMediaError InError)
-{
-	FErrorDetail err;
-	err.SetError(InError != UEMEDIA_ERROR_OK ? InError : UEMEDIA_ERROR_DETAIL);
-	err.SetFacility(Facility::EFacility::MP4PlaylistReader);
-	err.SetCode(InCode);
-	err.SetMessage(InMessage);
-	check(PlayerSessionServices);
-	if (PlayerSessionServices)
-	{
-		PlayerSessionServices->PostLog(Facility::EFacility::MP4PlaylistReader, IInfoLog::ELevel::Error, err.GetPrintable());
-	}
-	return err;
-}
-
 void FPlaylistReaderMP4::LogMessage(IInfoLog::ELevel Level, const FString& Message)
 {
 	if (PlayerSessionServices)
@@ -263,97 +218,6 @@ void FPlaylistReaderMP4::LoadAndParse(const FString& URL)
 	StartWorkerThread();
 }
 
-int32 FPlaylistReaderMP4::HTTPProgressCallback(const IElectraHttpManager::FRequest* InRequest)
-{
-	// Aborted?
-	return bAbort ? 1 : 0;
-}
-
-void FPlaylistReaderMP4::HTTPCompletionCallback(const IElectraHttpManager::FRequest* InRequest)
-{
-	bool bFailed = InRequest->ConnectionInfo.StatusInfo.ErrorDetail.IsError();
-	ConnectionInfo = InRequest->ConnectionInfo;
-	if (!bFailed)
-	{
-		// Set the size of the resource if we don't have it yet.
-		if (FileSize < 0)
-		{
-			ElectraHTTPStream::FHttpRange crh;
-			if (crh.ParseFromContentRangeResponse(InRequest->ConnectionInfo.ContentRangeHeader))
-			{
-				FileSize = crh.GetDocumentSize();
-			}
-		}
-
-		if (ConnectionInfo.EffectiveURL.Len())
-		{
-			MasterPlaylistURL = ConnectionInfo.EffectiveURL;
-		}
-
-		// Copy the read data across.
-		int32 NumRead = ReceiveBuffer->Buffer.Num();
-		Buffer.PushData(ReceiveBuffer->Buffer.GetLinearReadData(), NumRead);
-		if (FileSize >= 0 && ChunkReadOffset + NumRead >= FileSize)
-		{
-			Buffer.SetEOD();
-		}
-	}
-	bHasErrored = bFailed;
-	bChunkReadInProgress = false;
-}
-
-void FPlaylistReaderMP4::ClearRequest()
-{
-	FScopeLock lock(&Lock);
-	if (Request.IsValid())
-	{
-		PlayerSessionServices->GetHTTPManager()->RemoveRequest(Request, false);
-		Request.Reset();
-	}
-	ProgressListener.Reset();
-	ReceiveBuffer.Reset();
-}
-
-void FPlaylistReaderMP4::ReadNextChunk(int64 InFromOffset, int64 ChunkSize)
-{
-	ClearRequest();
-
-	FScopeLock lock(&Lock);
-	// Asked to go beyond the size of the file?
-	ChunkReadOffset = InFromOffset;
-	if (FileSize >= 0 && InFromOffset >= FileSize)
-	{
-		Buffer.SetEOD();
-		return;
-	}
-	Buffer.Reset();
-	ProgressListener = MakeSharedTS<IElectraHttpManager::FProgressListener>();
-	ProgressListener->CompletionDelegate = IElectraHttpManager::FProgressListener::FCompletionDelegate::CreateRaw(this, &FPlaylistReaderMP4::HTTPCompletionCallback);
-	ProgressListener->ProgressDelegate   = IElectraHttpManager::FProgressListener::FProgressDelegate::CreateRaw(this, &FPlaylistReaderMP4::HTTPProgressCallback);
-
-	ReceiveBuffer = MakeSharedTS<IElectraHttpManager::FReceiveBuffer>();
-	ReceiveBuffer->Buffer.Reserve(ChunkSize);
-
-	Request = MakeSharedTS<IElectraHttpManager::FRequest>();
-	Request->Parameters.URL = MasterPlaylistURL;
-	Request->Parameters.Range.SetStart(InFromOffset);
-	int64 LastByte = InFromOffset + ChunkSize - 1;
-	if (FileSize >= 0 && LastByte > FileSize-1)
-	{
-		LastByte = FileSize - 1;
-	}
-	Request->Parameters.Range.SetEndIncluding(LastByte);
-	Request->Parameters.ConnectTimeout = PlayerSessionServices->GetOptionValue(MP4::OptionKeyMP4LoadConnectTimeout).SafeGetTimeValue(FTimeValue().SetFromMilliseconds(1000 * 8));
-	Request->Parameters.NoDataTimeout = PlayerSessionServices->GetOptionValue(MP4::OptionKeyMP4LoadNoDataTimeout).SafeGetTimeValue(FTimeValue().SetFromMilliseconds(1000 * 6));
-	Request->ReceiveBuffer = ReceiveBuffer;
-	Request->ProgressListener = ProgressListener;
-	Request->ResponseCache = PlayerSessionServices->GetHTTPResponseCache();
-	Request->ExternalDataReader = PlayerSessionServices->GetExternalDataReader();
-	PlayerSessionServices->GetHTTPManager()->AddRequest(Request, false);
-}
-
-
-
 
 void FPlaylistReaderMP4::WorkerThread()
 {
@@ -362,27 +226,54 @@ void FPlaylistReaderMP4::WorkerThread()
 	SCOPE_CYCLE_COUNTER(STAT_ElectraPlayer_MP4_PlaylistWorker);
 	CSV_SCOPED_TIMING_STAT(ElectraPlayer, MP4_PlaylistWorker);
 
-	MP4Parser = IParserISO14496_12::CreateParser();
-	UEMediaError parseError = MP4Parser->ParseHeader(this, this, PlayerSessionServices, nullptr);
-	ClearRequest();
-
-	if (parseError != UEMEDIA_ERROR_ABORTED)
+	const TArray<uint32> StopAfterBoxes { UtilsMP4::Make4CC('m','o','o','v') };
+	const TArray<uint32> ReadBoxes { UtilsMP4::Make4CC('m','o','o','v') };
+	const TArray<uint32> FirstBoxes { UtilsMP4::Make4CC('f','t','y','p'), UtilsMP4::Make4CC('s','t','y','p'), UtilsMP4::Make4CC('s','i','d','x'), UtilsMP4::Make4CC('f','r','e','e'), UtilsMP4::Make4CC('s','k','i','p') };
+	UtilsMP4::FMP4RootBoxLocator BoxLocator;
+	TArray<UtilsMP4::FMP4RootBoxLocator::FBoxInfo> BoxInfos;
+	bool bGotBoxes = BoxLocator.LocateRootBoxes(BoxInfos, PlayerSessionServices->GetHTTPManager(), MasterPlaylistURL, FirstBoxes, StopAfterBoxes, ReadBoxes, UtilsMP4::FMP4RootBoxLocator::FCancellationCheckDelegate::CreateLambda([&]()
+	{
+		return bAbort;
+	}));
+	bool bHasErrored = BoxLocator.DidDownloadFail();
+	ConnectionInfo = BoxLocator.GetConnectionInfo();
+	if (!BoxLocator.GetErrorMessage().IsEmpty())
 	{
-		// Notify the download of the "master playlist". This indicates the download only, not the parsing thereof.
-		PlayerSessionServices->SendMessageToPlayer(IPlaylistReader::PlaylistDownloadMessage::Create(&ConnectionInfo, Playlist::EListType::Master, Playlist::ELoadType::Initial));
-		// Notify that the "master playlist" has been parsed, successfully or not.
-		PlayerSessionServices->SendMessageToPlayer(IPlaylistReader::PlaylistLoadedMessage::Create(LastErrorDetail, &ConnectionInfo, Playlist::EListType::Master, Playlist::ELoadType::Initial));
+		PostError(BoxLocator.GetErrorMessage(), ERRCODE_MP4_INVALID_FILE, UEMEDIA_ERROR_FORMAT_ERROR);
+	}
 
-		if (parseError == UEMEDIA_ERROR_OK || parseError == UEMEDIA_ERROR_END_OF_STREAM)
+	// Notify the download of the "master playlist". This indicates the download only, not the parsing thereof.
+	PlayerSessionServices->SendMessageToPlayer(IPlaylistReader::PlaylistDownloadMessage::Create(&ConnectionInfo, Playlist::EListType::Master, Playlist::ELoadType::Initial));
+	// Notify that the "master playlist" has been parsed, successfully or not.
+	PlayerSessionServices->SendMessageToPlayer(IPlaylistReader::PlaylistLoadedMessage::Create(LastErrorDetail, &ConnectionInfo, Playlist::EListType::Master, Playlist::ELoadType::Initial));
+	// Failed to get the boxes but was not aborted?
+	if (!bAbort && (!bGotBoxes || bHasErrored))
+	{
+		// See if there was a download error
+		if (ConnectionInfo.StatusInfo.ErrorDetail.IsError())
 		{
-			// See that we have parsed all the boxes we need.
-			if (bFoundBoxFTYP && bFoundBoxMOOV)
+			PostError(FString::Printf(TEXT("%s while downloading \"%s\""), *ConnectionInfo.StatusInfo.ErrorDetail.GetMessage(), *ConnectionInfo.EffectiveURL), ERRCODE_MP4_DOWNLOAD_ERROR, UEMEDIA_ERROR_READ_ERROR);
+		}
+	}
+	else if (bGotBoxes && !bAbort && !bHasErrored)
+	{
+		// Do we have the `ftyp` and `moov` boxes?
+		if (BoxInfos.ContainsByPredicate([InType=UtilsMP4::Make4CC('f','t','y','p')](const UtilsMP4::FMP4RootBoxLocator::FBoxInfo& InBox){return InBox.Type == InType;}) &&
+			BoxInfos.ContainsByPredicate([InType=UtilsMP4::Make4CC('m','o','o','v')](const UtilsMP4::FMP4RootBoxLocator::FBoxInfo& InBox){return InBox.Type == InType;}))
+		{
+			if (BoxInfos.ContainsByPredicate([InType=UtilsMP4::Make4CC('m','d','a','t')](const UtilsMP4::FMP4RootBoxLocator::FBoxInfo& InBox){return InBox.Type == InType;}))
 			{
-				if (!bIsFastStartable)
-				{
-					LogMessage(IInfoLog::ELevel::Info, FString::Printf(TEXT("The mp4 at \"%s\" is not fast-startable. Consider moving the 'moov' box in front of the 'mdat' for faster startup times."), *ConnectionInfo.EffectiveURL));
-				}
+				LogMessage(IInfoLog::ELevel::Info, FString::Printf(TEXT("The mp4 at \"%s\" is not fast-startable. Consider moving the 'moov' box in front of the 'mdat' for faster startup times."), *ConnectionInfo.EffectiveURL));
+			}
 
+			ParseBuffer = BoxInfos.FindByPredicate([InType=UtilsMP4::Make4CC('m','o','o','v')](const UtilsMP4::FMP4RootBoxLocator::FBoxInfo& InBox){return InBox.Type == InType;})->DataBuffer;
+			ParsePos = 0;
+			ParseBufferSize = ParseBuffer->Buffer.Num();
+
+			MP4Parser = IParserISO14496_12::CreateParser();
+			UEMediaError parseError = MP4Parser->ParseHeader(this, this, PlayerSessionServices, nullptr);
+			if (parseError == UEMEDIA_ERROR_OK || parseError == UEMEDIA_ERROR_END_OF_STREAM)
+			{
 				// Prepare the tracks in the stream that are of a supported codec.
 				parseError = MP4Parser->PrepareTracks(PlayerSessionServices, TSharedPtrTS<const IParserISO14496_12>());
 				if (parseError == UEMEDIA_ERROR_OK)
@@ -405,21 +296,13 @@ void FPlaylistReaderMP4::WorkerThread()
 			}
 			else
 			{
-				// No moov box usually means this is not a fast-start file.
-				PostError(FString::Printf(TEXT("No moov box found in \"%s\". This is not a valid file."), *ConnectionInfo.EffectiveURL), ERRCODE_MP4_INVALID_FILE, UEMEDIA_ERROR_FORMAT_ERROR);
+				PostError(FString::Printf(TEXT("Failed to parse mp4 \"%s\" with error %u"), *ConnectionInfo.EffectiveURL, parseError), ERRCODE_MP4_INVALID_FILE, UEMEDIA_ERROR_FORMAT_ERROR);
 			}
 		}
 		else
 		{
-			// See if there was a download error
-			if (ConnectionInfo.StatusInfo.ErrorDetail.IsError())
-			{
-				PostError(FString::Printf(TEXT("%s while downloading \"%s\""), *ConnectionInfo.StatusInfo.ErrorDetail.GetMessage(), *ConnectionInfo.EffectiveURL), ERRCODE_MP4_DOWNLOAD_ERROR, UEMEDIA_ERROR_READ_ERROR);
-			}
-			else
-			{
-				PostError(FString::Printf(TEXT("Failed to parse mp4 \"%s\" with error %u"), *ConnectionInfo.EffectiveURL, parseError), ERRCODE_MP4_INVALID_FILE, UEMEDIA_ERROR_FORMAT_ERROR);
-			}
+			// No moov box usually means this is not a fast-start file.
+			PostError(FString::Printf(TEXT("No ftyp or moov box found in \"%s\". This is not a valid file."), *ConnectionInfo.EffectiveURL), ERRCODE_MP4_INVALID_FILE, UEMEDIA_ERROR_FORMAT_ERROR);
 		}
 	}
 
@@ -444,66 +327,20 @@ void FPlaylistReaderMP4::WorkerThread()
  */
 int64 FPlaylistReaderMP4::ReadData(void* IntoBuffer, int64 NumBytesToRead)
 {
-	uint8* OutputBuffer = (uint8*)IntoBuffer;
-	// Do we have enough data in the buffer to satisfy the read?
-	if (Buffer.Num() >= NumBytesToRead)
+	// We have all the data available.
+	check(ParseBuffer.IsValid());
+
+	if (ParsePos >= ParseBufferSize)
 	{
-		// Yes. Get the data and return.
-		int32 NumGot = Buffer.PopData(OutputBuffer, NumBytesToRead);
-		check(NumGot == NumBytesToRead);
-		ParsePos += NumBytesToRead;
-		return NumBytesToRead;
+		return 0;
 	}
-	else
+	const uint8* Src = ParseBuffer->Buffer.GetLinearReadData() + ParsePos;
+	if (IntoBuffer)
 	{
-		// Do not have enough data yet or we want to read more than the buffer can hold?
-		int32 NumBytesToGo = NumBytesToRead;
-		int64 NextChunkReadOffset = ParsePos;
-		while(NumBytesToGo > 0)
-		{
-			if (bHasErrored || bAbort)
-			{
-				return -1;
-			}
-			// EOD?
-			if (Buffer.IsEndOfData())
-			{
-				return 0;
-			}
-
-			// Get whatever amount of data is currently available to free up the buffer for receiving more data.
-			int32 NumGot = Buffer.PopData(OutputBuffer, NumBytesToGo);
-			if ((NumBytesToGo -= NumGot) > 0)
-			{
-				if (OutputBuffer)
-				{
-					OutputBuffer += NumGot;
-				}
-				// Trigger read of next chunk of data.
-				if (!bChunkReadInProgress)
-				{
-					check(Buffer.IsEmpty());
-					// Is the data to read actually used or is it skipped over?
-					if (OutputBuffer)
-					{
-						NextChunkReadOffset += NumGot;
-						ReadNextChunk(NextChunkReadOffset, kChunkReadSize);
-						bChunkReadInProgress = true;
-					}
-					else
-					{
-						// Data is not used, so do not request it.
-						break;
-					}
-				}
-				// Wait for data to arrive in the ringbuffer.
-				int32 WaitForBytes = NumBytesToGo > Buffer.Capacity() ? Buffer.Capacity() : NumBytesToGo;
-				Buffer.WaitUntilSizeAvailable(WaitForBytes, 1000 * 100);
-			}
-		}
-		ParsePos += NumBytesToRead;
-		return NumBytesToRead;
+		FMemory::Memcpy(IntoBuffer, Src, NumBytesToRead);
 	}
+	ParsePos += NumBytesToRead;
+	return NumBytesToRead;
 }
 
 /**
@@ -513,7 +350,7 @@ int64 FPlaylistReaderMP4::ReadData(void* IntoBuffer, int64 NumBytesToRead)
  */
 bool FPlaylistReaderMP4::HasReachedEOF() const
 {
-	return Buffer.IsEndOfData();
+	return ParsePos >= ParseBufferSize;
 }
 
 /**
@@ -538,51 +375,16 @@ int64 FPlaylistReaderMP4::GetCurrentOffset() const
 	return ParsePos;
 }
 
-
 IParserISO14496_12::IBoxCallback::EParseContinuation FPlaylistReaderMP4::OnFoundBox(IParserISO14496_12::FBoxType Box, int64 BoxSizeInBytes, int64 FileDataOffset, int64 BoxDataOffset)
 {
-	// We require the very first box to be an 'ftyp' box.
-	if (FileDataOffset == 0 && Box != IParserISO14496_12::BoxType_ftyp)
-	{
-		PostError("Invalid mp4 file: first box is not 'ftyp'", ERRCODE_MP4_INVALID_FILE, UEMEDIA_ERROR_FORMAT_ERROR);
-		return IParserISO14496_12::IBoxCallback::EParseContinuation::Stop;
-	}
-
-	// Check which box is being parsed next.
-	switch(Box)
-	{
-		case IParserISO14496_12::BoxType_ftyp:
-			bFoundBoxFTYP = true;
-			break;
-		case IParserISO14496_12::BoxType_moov:
-			bFoundBoxMOOV = true;
-			bIsFastStartable = !bFoundBoxMDAT;
-			break;
-		case IParserISO14496_12::BoxType_sidx:
-			bFoundBoxSIDX = true;
-			break;
-		case IParserISO14496_12::BoxType_moof:
-			bFoundBoxMOOF = true;
-			return IParserISO14496_12::IBoxCallback::EParseContinuation::Stop;
-		case IParserISO14496_12::BoxType_mdat:
-			bFoundBoxMDAT = true;
-			break;
-		default:
-			break;
-	}
 	return IParserISO14496_12::IBoxCallback::EParseContinuation::Continue;
 }
 
 IParserISO14496_12::IBoxCallback::EParseContinuation FPlaylistReaderMP4::OnEndOfBox(IParserISO14496_12::FBoxType Box, int64 BoxSizeInBytes, int64 FileDataOffset, int64 BoxDataOffset)
 {
-	if (bFoundBoxMOOV)
-	{
-		return IParserISO14496_12::IBoxCallback::EParseContinuation::Stop;
-	}
-	return IParserISO14496_12::IBoxCallback::EParseContinuation::Continue;
+	return IParserISO14496_12::IBoxCallback::EParseContinuation::Stop;
 }
 
-
 } // namespace Electra
 
 
