Fix for crash during async code execution inside Electra output pipeline

diff --git a/Engine/Plugins/Media/ElectraPlayer/Source/ElectraPlayerRuntime/Private/Runtime/Decoder/Windows/ElectraVideoDecoder_PC.cpp b/Engine/Plugins/Media/ElectraPlayer/Source/ElectraPlayerRuntime/Private/Runtime/Decoder/Windows/ElectraVideoDecoder_PC.cpp
--- a/Engine/Plugins/Media/ElectraPlayer/Source/ElectraPlayerRuntime/Private/Runtime/Decoder/Windows/ElectraVideoDecoder_PC.cpp
+++ b/Engine/Plugins/Media/ElectraPlayer/Source/ElectraPlayerRuntime/Private/Runtime/Decoder/Windows/ElectraVideoDecoder_PC.cpp
@@ -359,24 +359,24 @@ void FElectraPlayerVideoDecoderOutputPC::InitializeWithResource(const TRefCountP
 		auto ElectraDecoderResourceDelegate = Electra::FElectraDecoderResourceManagerWindows::GetDelegate();
 
 		// Note: we capture "this" as we ensure that this instance only dies once the copy triggered here is actually done, hence ensuring any reference to "this" is done
-		bTriggerOk = ElectraDecoderResourceDelegate->RunCodeAsync([D3DCmdList=this->D3DCmdList, D3DFence=this->D3DFence, FenceValue=this->FenceValue, OutputSync, ResourceDelegate = InResourceDelegate.Pin()]()
+		bTriggerOk = ElectraDecoderResourceDelegate->RunCodeAsync([D3DCmdList=this->D3DCmdList, D3DCmdAllocator=this->D3DCmdAllocator, D3DFence=this->D3DFence, FenceValue=this->FenceValue, OutputSync, ResourceDelegate = InResourceDelegate.Pin()]()
 			{
-				TriggerDataCopy(D3DCmdList, D3DFence, FenceValue, OutputSync, ResourceDelegate.Get());
+				TriggerDataCopy(D3DCmdList, D3DCmdAllocator, D3DFence, FenceValue, OutputSync, ResourceDelegate.Get());
 			}, OutputSync.TaskSync.Get());
 	}
 
 	if (!bTriggerOk)
 	{
 		// We could not run the trigger async. Schedule the copy right away. Any needed synchronization will be done in the copy-queue by the GPU
-		TriggerDataCopy(D3DCmdList, D3DFence, FenceValue, OutputSync, InResourceDelegate.Pin().Get());
+		TriggerDataCopy(D3DCmdList, D3DCmdAllocator, D3DFence, FenceValue, OutputSync, InResourceDelegate.Pin().Get());
 	}
 }
 
-void FElectraPlayerVideoDecoderOutputPC::TriggerDataCopy(TRefCountPtr<ID3D12GraphicsCommandList> D3DCmdList, TRefCountPtr<ID3D12Fence> D3DFence, uint64 FenceValue, const FElectraDecoderOutputSync& OutputSync, Electra::IVideoDecoderResourceDelegate* InResourceDelegate)
+void FElectraPlayerVideoDecoderOutputPC::TriggerDataCopy(TRefCountPtr<ID3D12GraphicsCommandList> D3DCmdList, TRefCountPtr<ID3D12CommandAllocator> D3DCmdAllocator, TRefCountPtr<ID3D12Fence> D3DFence, uint64 FenceValue, const FElectraDecoderOutputSync& OutputSync, Electra::IVideoDecoderResourceDelegate* InResourceDelegate)
 {
 	// Trigger copy (this will eventually execute on the submission thread of RHI if running in UE)
 	// (note: we pass in all of FElectraDecoderOutputSync to guarantee any references needed to make the decoder output sync work are passed along, too!)
-	InResourceDelegate->ExecuteCodeWithCopyCommandQueueUsage([CmdList = D3DCmdList, DestFence = D3DFence, DestFenceValue = FenceValue, OutputSync](ID3D12CommandQueue* D3DCmdQueue)
+	InResourceDelegate->ExecuteCodeWithCopyCommandQueueUsage([CmdList = D3DCmdList, CmdAllocator = D3DCmdAllocator, DestFence = D3DFence, DestFenceValue = FenceValue, OutputSync](ID3D12CommandQueue* D3DCmdQueue)
 		{
 			TRefCountPtr<ID3D12Fence> ResourceFence;
 	#if ALLOW_MFSAMPLE_WITH_DX12
diff --git a/Engine/Plugins/Media/ElectraPlayer/Source/ElectraPlayerRuntime/Private/Runtime/Decoder/Windows/ElectraVideoDecoder_PC.h b/Engine/Plugins/Media/ElectraPlayer/Source/ElectraPlayerRuntime/Private/Runtime/Decoder/Windows/ElectraVideoDecoder_PC.h
--- a/Engine/Plugins/Media/ElectraPlayer/Source/ElectraPlayerRuntime/Private/Runtime/Decoder/Windows/ElectraVideoDecoder_PC.h
+++ b/Engine/Plugins/Media/ElectraPlayer/Source/ElectraPlayerRuntime/Private/Runtime/Decoder/Windows/ElectraVideoDecoder_PC.h
@@ -72,7 +72,7 @@ public:
 	TRefCountPtr<IUnknown> GetSync(uint64& SyncValue) const override;
 
 private:
-	static void TriggerDataCopy(TRefCountPtr<ID3D12GraphicsCommandList> D3DCmdList, TRefCountPtr<ID3D12Fence> D3DFence, uint64 FenceValue, const FElectraDecoderOutputSync& OutputSync, Electra::IVideoDecoderResourceDelegate* InResourceDelegate);
+	static void TriggerDataCopy(TRefCountPtr<ID3D12GraphicsCommandList> D3DCmdList, TRefCountPtr<ID3D12CommandAllocator> D3DCmdAllocator, TRefCountPtr<ID3D12Fence> D3DFence, uint64 FenceValue, const FElectraDecoderOutputSync& OutputSync, Electra::IVideoDecoderResourceDelegate* InResourceDelegate);
 
 	// Decoder output type
 	EOutputType OutputType = EOutputType::Unknown;
