ElectraPlayer: Adjusting the Live stream default latency value if the stream provides a larger value without using a latency element.

#rnx

diff --git a/Engine/Plugins/Media/ElectraPlayer/Source/ElectraPlayerRuntime/Private/Runtime/Player/ABRRules/ABRLiveStream.cpp b/Engine/Plugins/Media/ElectraPlayer/Source/ElectraPlayerRuntime/Private/Runtime/Player/ABRRules/ABRLiveStream.cpp
--- a/Engine/Plugins/Media/ElectraPlayer/Source/ElectraPlayerRuntime/Private/Runtime/Player/ABRRules/ABRLiveStream.cpp
+++ b/Engine/Plugins/Media/ElectraPlayer/Source/ElectraPlayerRuntime/Private/Runtime/Player/ABRRules/ABRLiveStream.cpp
@@ -103,9 +103,9 @@ private:
 	struct FQualityMetrics
 	{
 		FQualityMetrics()
-		{ 
+		{
 			AverageKbps.Resize(5);
-			Reset(); 
+			Reset();
 		}
 		void Reset()
 		{
@@ -159,15 +159,15 @@ private:
 		};
 
 		FStreamWorkVars()
-		{ 
+		{
 			const int32 HistorySize = 3;
 			AverageBandwidth.Resize(HistorySize);
 			AverageThroughput.Resize(HistorySize);
 			const int32 LatencyHistorySize = 5;
 			AverageLatency.Resize(LatencyHistorySize);
-			Reset(); 
+			Reset();
 		}
-		
+
 		void ClearForNextDownload()
 		{
 			BufferContentDurationAtSegmentStart = -1.0;
@@ -274,7 +274,7 @@ private:
 		const TSharedPtrTS<FABRStreamInformation>* Stream = StreamInfos.FindByPredicate([InQualityIndex](const TSharedPtrTS<FABRStreamInformation>& InInfo) { return InQualityIndex == InInfo->QualityIndex;} );
 		return Stream ? (*Stream) : nullptr;
 	}
-	
+
 	double GetPlayablePlayerDuration(bool& bEOS, EStreamType InStreamType)
 	{
 		IAdaptiveStreamSelector::IPlayerLiveControl::FABRBufferStats bs;
@@ -359,7 +359,7 @@ private:
 	const double AbortDownloadCheckTimeMaxOverSegmentDuration = 2.0;
 	const double SlowdownAfterDownloadTimeBySegmentDurationScale = 1.075;
 	const double SlowdownAfterDownloadCheckTimeMaxOverSegmentDuration = 0.2;
-	
+
 	const double BufferingNonLLCompleteStableScaleBySegmentDuration = 0.3;
 	const double EmitPartialDataNonLLBufferDurationBelow = 1.0;
 
@@ -369,7 +369,7 @@ private:
 	const double BandwidthScaleToReconsiderDeselectedStream = 1.5;
 	const double ClampBandwidthToMaxStreamBitrateScaleFactor = 2.0;
 	const double DownloadOvertimePenaltyScale = 0.5;
-	
+
 	const int32 StableBufferDramaticDropPercentage = -40;
 	const double StableBufferNextHitBadDropSegmentDurationScale = 0.6;
 	const double StableBufferByNetworkLatencyScale = 2.0;
@@ -419,10 +419,28 @@ void FABRLiveStream::PrepareLatencyConfiguration()
 {
 	TSharedPtrTS<const FLowLatencyDescriptor> lld = Info->ABRGetLowLatencyDescriptor();
 	double TargetLatency = Info->ABRGetDesiredLiveEdgeLatency().GetAsSeconds();
+	FTimeRange TimelineRange = Info->ABRGetTimeline();
 
 	FScopeLock lock(&Lock);
 	LatencyConfig.Reset();
 	LatencyConfig.TargetLatency = TargetLatency;
+	// If the target latency is greater than our assumed default value we need to adjust our value.
+	if (TargetLatency > LatencyConfig.MaxLatency)
+	{
+		// If the media timeline is valid we set the latency to the duration of the timeline.
+		// While this may be far too large it will at least prevent falling behind what the timeline indicates.
+		// If there is a latency element present this value will be replaced by the latency element, so this
+		// is a fallback for streams without a latency element (and therefore not low-latency).
+		if (TimelineRange.IsValid())
+		{
+			LatencyConfig.MaxLatency = (TimelineRange.End - TimelineRange.Start).GetAsSeconds();
+		}
+		else
+		{
+			// Arbitrary fallback value.
+			LatencyConfig.MaxLatency = 60.0;
+		}
+	}
 	if (lld.IsValid())
 	{
 		LatencyConfig.MinLatency = lld->GetLatencyMin().IsValid() ? lld->GetLatencyMin().GetAsSeconds() : LatencyConfig.TargetLatency;
@@ -431,7 +449,7 @@ void FABRLiveStream::PrepareLatencyConfiguration()
 		LatencyConfig.MaxPlayRate = lld->GetPlayrateMax().IsValid() ? lld->GetPlayrateMax().GetAsSeconds() : LatencyConfig.MaxPlayRate;
 	}
 	#ifdef ENABLE_LATENCY_OVERRIDE_CVAR
-	AsyncTask(ENamedThreads::GameThread, [=]() {(*CVarElectraTL).Set(*LexToString(TargetLatency), EConsoleVariableFlags::ECVF_SetByCode);});	
+	AsyncTask(ENamedThreads::GameThread, [=]() {(*CVarElectraTL).Set(*LexToString(TargetLatency), EConsoleVariableFlags::ECVF_SetByCode);});
 	#endif
 }
 
@@ -1148,7 +1166,7 @@ IAdaptiveStreamSelector::ESegmentAction FABRLiveStream::PerformSelection(const T
 
 			bIsBufferStable = HasStableBuffer(BufferGain, BufferTrend, StreamType, 0.0);
 			CurrentStreamQualityIndex = WorkVars->SegmentDownloadHistory.Num() ? WorkVars->SegmentDownloadHistory.BackRef().QualityIndex : -1;
-	
+
 			double OvertimePenalty = WorkVars->OverDownloadTimeTotal > 0.0 ? WorkVars->OverDownloadTimeTotal * DownloadOvertimePenaltyScale : 0.0;
 			WorkVars->OverDownloadTimeTotal = 0.0;
 
@@ -1469,7 +1487,7 @@ bool FABRLiveStream::HasStableBuffer(int32& OutGain, int32& OutTrend, EStreamTyp
 		Get a reference segment duration.
 
 		It is possible for the segment duration to be larger than the desired latency, espcially for low-latency
-		Live streams where segments are consumed as they are produced. Large segment durations can't really be 
+		Live streams where segments are consumed as they are produced. Large segment durations can't really be
 		used to gauge buffer stability since it will not really be possible to gather that much data ahead of time.
 	*/
 	const double SegmentDuration = Utils::Min(WorkVars->AverageSegmentDuration.GetAsSeconds(DefaultAssumedSegmentDuration), DesiredLatency);
@@ -1479,13 +1497,13 @@ bool FABRLiveStream::HasStableBuffer(int32& OutGain, int32& OutTrend, EStreamTyp
 	{
 		const FStreamWorkVars::FSegmentInfo& Last = WorkVars->SegmentDownloadHistory.BackRef();
 		const double NetworkLatency = WorkVars->AverageLatency.GetSMA();
-		
+
 		// No data at all?
 		if (Last.BufferDurationAtEnd == 0.0 || Last.BufferDurationAtStart == 0.0)
 		{
 			return false;
 		}
-		
+
 
 		// Buffer level
 		double r = Last.BufferDurationAtEnd / Last.BufferDurationAtStart;
@@ -1577,7 +1595,7 @@ bool FABRLiveStream::HasStableBuffer(int32& OutGain, int32& OutTrend, EStreamTyp
 
 
 FTimeValue FABRLiveStream::GetMinBufferTimeForPlayback(IAdaptiveStreamSelector::EMinBufferType InBufferingType, FTimeValue InDefaultMBT)
-{ 
+{
 	if (FormatType == EMediaFormatType::DASH)
 	{
 		const FStreamWorkVars* WorkVars = GetWorkVars(GetPrimaryStreamType());
@@ -1587,13 +1605,13 @@ FTimeValue FABRLiveStream::GetMinBufferTimeForPlayback(IAdaptiveStreamSelector::
 			// but require a tiny amount of data to be buffered to minimize the chance of an immediate rebuffer.
 			if (WorkVars->bIsLowLatencyEnabled)
 			{
-				return FTimeValue(0.3); 
+				return FTimeValue(0.3);
 			}
 			// When it has been determined that buffering will soon be done return a value small enough
 			// to just start now.
 			else if (WorkVars->bBufferingWillCompleteOnTime)
 			{
-				return FTimeValue(0.1); 
+				return FTimeValue(0.1);
 			}
 		}
 	}
