ElectraPlayer: removed unnecessary check and flag

#rnx

diff --git a/Engine/Plugins/Media/ElectraPlayer/Source/ElectraPlayerRuntime/Private/Runtime/Player/AdaptiveStreamingPlayerRenderers.cpp b/Engine/Plugins/Media/ElectraPlayer/Source/ElectraPlayerRuntime/Private/Runtime/Player/AdaptiveStreamingPlayerRenderers.cpp
--- a/Engine/Plugins/Media/ElectraPlayer/Source/ElectraPlayerRuntime/Private/Runtime/Player/AdaptiveStreamingPlayerRenderers.cpp
+++ b/Engine/Plugins/Media/ElectraPlayer/Source/ElectraPlayerRuntime/Private/Runtime/Player/AdaptiveStreamingPlayerRenderers.cpp
@@ -273,14 +273,11 @@ void FAdaptiveStreamingWrappedRenderer::SampleReleasedToPool(IDecoderOutput* InD
 
 			{
 				FScopeLock lock(&Lock);
-				check(EnqueuedSamples.Num());
-				bool bFound = false;
 				for(int32 i=0; i<EnqueuedSamples.Num(); ++i)
 				{
 					if (EnqueuedSamples[i].PTS == RenderTime)
 					{
 						EnqueuedSamples.RemoveAt(i);
-						bFound = true;
 						break;
 					}
 				}
@@ -453,7 +450,7 @@ UEMediaError FAdaptiveStreamingWrappedRenderer::ReturnBufferCommon(IBuffer* Buff
 	FScopeLock lock(&Lock);
 	FEnqueuedSampleInfo& enqInf = EnqueuedSamples.Emplace_GetRef();
 	enqInf.Duration = Duration;
-	enqInf.PTS = InSampleProperties.GetValue(RenderOptionKeys::PTS).SafeGetTimeValue(FTimeValue::GetInvalid());;
+	enqInf.PTS = InSampleProperties.GetValue(RenderOptionKeys::PTS).SafeGetTimeValue(FTimeValue::GetInvalid());
 
 	if (!bIsUnusedReturnBuffer)
 	{
