ElectraPlayer: fixed potential nullptr access on a canceled request

#rnx

diff --git a/Engine/Plugins/Media/ElectraPlayer/Source/ElectraPlayerRuntime/Private/Runtime/Player/DASH/StreamReaderDASH.cpp b/Engine/Plugins/Media/ElectraPlayer/Source/ElectraPlayerRuntime/Private/Runtime/Player/DASH/StreamReaderDASH.cpp
--- a/Engine/Plugins/Media/ElectraPlayer/Source/ElectraPlayerRuntime/Private/Runtime/Player/DASH/StreamReaderDASH.cpp
+++ b/Engine/Plugins/Media/ElectraPlayer/Source/ElectraPlayerRuntime/Private/Runtime/Player/DASH/StreamReaderDASH.cpp
@@ -476,17 +476,20 @@ void FStreamReaderDASH::FStreamHandler::SetupInitSegmentDownloadStatsFromRequest
 	ds.Range = Request->Segment.InitializationURL.Range;
 	ds.CDN = Request->Segment.InitializationURL.CDN;
 	ds.bWasSuccessful = bWasSuccessful;
-	ds.HTTPStatusCode = ci->StatusInfo.HTTPStatus;
-	ds.TimeToFirstByte = ci->TimeUntilFirstByte;
-	ds.TimeToDownload = (ci->RequestEndTime - ci->RequestStartTime).GetAsSeconds();
-	ds.ByteSize = ci->ContentLength;
-	ds.NumBytesDownloaded = ci->BytesReadSoFar;
 	ds.MediaAssetID = Request->Period.IsValid() ? Request->Period->GetUniqueIdentifier() : "";
 	ds.AdaptationSetID = Request->AdaptationSet.IsValid() ? Request->AdaptationSet->GetUniqueIdentifier() : "";
 	ds.RepresentationID = Request->Representation.IsValid() ? Request->Representation->GetUniqueIdentifier() : "";
 	ds.Bitrate = Request->GetBitrate();
 	ds.QualityIndex = Request->QualityIndex;
 	ds.HighestQualityIndex = Request->MaxQualityIndex;
+	if (ci)
+	{
+		ds.HTTPStatusCode = ci->StatusInfo.HTTPStatus;
+		ds.TimeToFirstByte = ci->TimeUntilFirstByte;
+		ds.TimeToDownload = (ci->RequestEndTime - ci->RequestStartTime).GetAsSeconds();
+		ds.ByteSize = ci->ContentLength;
+		ds.NumBytesDownloaded = ci->BytesReadSoFar;
+	}
 }
 
 
@@ -550,9 +553,12 @@ FErrorDetail FStreamReaderDASH::FStreamHandler::LoadInitSegment(TSharedPtrTS<FMP
 
 	if (!LoadResult->bSuccess)
 	{
-		Request->ConnectionInfo = *ci;
+		if (ci)
+		{
+			Request->ConnectionInfo = *ci;
+		}
 		StreamSelector->ReportDownloadEnd(ds);
-		return CreateError(FString::Printf(TEXT("Init segment download error: %s"),	*ci->StatusInfo.ErrorDetail.GetMessage()), INTERNAL_ERROR_INIT_SEGMENT_DOWNLOAD_ERROR);
+		return CreateError(FString::Printf(TEXT("Init segment download error: %s"),	*Request->ConnectionInfo.StatusInfo.ErrorDetail.GetMessage()), INTERNAL_ERROR_INIT_SEGMENT_DOWNLOAD_ERROR);
 	}
 	return FErrorDetail();
 }
@@ -636,7 +642,11 @@ FErrorDetail FStreamReaderDASH::FStreamHandler::GetInitSegment(TSharedPtrTS<cons
 		{
 			StaticDataReader.SetParseData(LoadReq->Request->GetResponseBuffer());
 		}
-		ConnectionInfo = *LoadReq->GetConnectionInfo();
+		const HTTP::FConnectionInfo* ci = LoadReq->GetConnectionInfo();
+		if (ci)
+		{
+			ConnectionInfo = *ci;
+		}
 	}
 
 	SCOPE_CYCLE_COUNTER(STAT_ElectraPlayer_DASH_StreamReader);
@@ -736,7 +746,11 @@ FErrorDetail FStreamReaderDASH::FStreamHandler::GetInitSegment(TSharedPtrTS<cons
 			else
 			{
 				ds.bParseFailure = true;
-				Request->ConnectionInfo = *LoadReq->GetConnectionInfo();
+				const HTTP::FConnectionInfo* ci = LoadReq->GetConnectionInfo();
+				if (ci)
+				{
+					Request->ConnectionInfo = *ci;
+				}
 				StreamSelector->ReportDownloadEnd(ds);
 				return CreateError(FString::Printf(TEXT("Track preparation of init segment \"%s\" failed. %d"), *LoadReq->URL, *parseError.GetMessage()), INTERNAL_ERROR_INIT_SEGMENT_PARSE_ERROR);
 			}
@@ -836,21 +850,27 @@ FErrorDetail FStreamReaderDASH::FStreamHandler::RetrieveSideloadedFile(TSharedPt
 	ds.Range = Request->Segment.MediaURL.Range;
 	ds.CDN = Request->Segment.MediaURL.CDN;
 	ds.bWasSuccessful = LoadResult->bSuccess;
-	ds.HTTPStatusCode = ci->StatusInfo.HTTPStatus;
-	ds.TimeToFirstByte = ci->TimeUntilFirstByte;
-	ds.TimeToDownload = (ci->RequestEndTime - ci->RequestStartTime).GetAsSeconds();
-	ds.ByteSize = ci->ContentLength;
-	ds.NumBytesDownloaded = ci->BytesReadSoFar;
 	ds.MediaAssetID = Request->Period.IsValid() ? Request->Period->GetUniqueIdentifier() : "";
 	ds.AdaptationSetID = Request->AdaptationSet.IsValid() ? Request->AdaptationSet->GetUniqueIdentifier() : "";
 	ds.RepresentationID = Request->Representation.IsValid() ? Request->Representation->GetUniqueIdentifier() : "";
 	ds.Bitrate = Request->GetBitrate();
+	if (ci)
+	{
+		ds.HTTPStatusCode = ci->StatusInfo.HTTPStatus;
+		ds.TimeToFirstByte = ci->TimeUntilFirstByte;
+		ds.TimeToDownload = (ci->RequestEndTime - ci->RequestStartTime).GetAsSeconds();
+		ds.ByteSize = ci->ContentLength;
+		ds.NumBytesDownloaded = ci->BytesReadSoFar;
+	}
 
 	if (!LoadResult->bSuccess)
 	{
-		Request->ConnectionInfo = *ci;
+		if (ci)
+		{
+			Request->ConnectionInfo = *ci;
+		}
 		StreamSelector->ReportDownloadEnd(ds);
-		return CreateError(FString::Printf(TEXT("Sideloaded media download error: %s"), *ci->StatusInfo.ErrorDetail.GetMessage()), INTERNAL_ERROR_INIT_SEGMENT_DOWNLOAD_ERROR);
+		return CreateError(FString::Printf(TEXT("Sideloaded media download error: %s"), *Request->ConnectionInfo.StatusInfo.ErrorDetail.GetMessage()), INTERNAL_ERROR_INIT_SEGMENT_DOWNLOAD_ERROR);
 	}
 
 	OutData = MakeSharedTS<const TArray<uint8>>(TArrayView<uint8>(LoadReq->Request->GetResponseBuffer()->Buffer.GetLinearReadData(), LoadReq->Request->GetResponseBuffer()->Buffer.GetLinearReadSize()));
