Media: Optimize for small forward seeks

#rnx

diff --git a/Engine/Plugins/Media/ElectraPlayer/Source/ElectraPlayerPlugin/Private/ElectraPlayerAudioSample.h b/Engine/Plugins/Media/ElectraPlayer/Source/ElectraPlayerPlugin/Private/ElectraPlayerAudioSample.h
--- a/Engine/Plugins/Media/ElectraPlayer/Source/ElectraPlayerPlugin/Private/ElectraPlayerAudioSample.h
+++ b/Engine/Plugins/Media/ElectraPlayer/Source/ElectraPlayerPlugin/Private/ElectraPlayerAudioSample.h
@@ -8,7 +8,7 @@
 
 #include "MediaAudioDecoderOutput.h"
 
- 
+
 // IMediaAudioSample impl
 // is used as audio decoder output
 // contains decoded audio samples ready for playback along with corresponding metadata
@@ -27,7 +27,7 @@ public:
 	{
 	}
 
-	void Initialize(const IAudioDecoderOutputPtr & InDecoderOutput)
+	void Initialize(const IAudioDecoderOutputPtr& InDecoderOutput)
 	{
 		DecoderOutput = InDecoderOutput;
 	}
@@ -42,7 +42,7 @@ public:
 	{ return DecoderOutput ? DecoderOutput->GetDuration() : FTimespan::Zero(); }
 
 	EMediaAudioSampleFormat GetFormat() const override
-	{ 
+	{
 		static_assert((int32)EMediaAudioSampleFormat::Undefined == (int32)IAudioDecoderOutput::ESampleFormat::Undefined, "check enums are identical");
 		static_assert((int32)EMediaAudioSampleFormat::Double == (int32)IAudioDecoderOutput::ESampleFormat::Double, "check enums are identical");
 		static_assert((int32)EMediaAudioSampleFormat::Float == (int32)IAudioDecoderOutput::ESampleFormat::Float, "check enums are identical");
@@ -71,7 +71,15 @@ public:
 
 	uint32 GetMaxBufferBytes() const
 	{ return DecoderOutput ? DecoderOutput->GetReservedBufferBytes() : 0; }
-	
+
+	virtual void SetTime(const FMediaTimeStamp& InTime)
+	{
+		if (DecoderOutput)
+		{
+			DecoderOutput->SetTime(FDecoderTimeStamp(InTime.Time, InTime.SequenceIndex));
+		}
+	}
+
 	virtual void InitializePoolable() override;
 	virtual void ShutdownPoolable() override;
 
diff --git a/Engine/Plugins/Media/ElectraPlayer/Source/ElectraPlayerPlugin/Private/ElectraPlayerPlugin.cpp b/Engine/Plugins/Media/ElectraPlayer/Source/ElectraPlayerPlugin/Private/ElectraPlayerPlugin.cpp
--- a/Engine/Plugins/Media/ElectraPlayer/Source/ElectraPlayerPlugin/Private/ElectraPlayerPlugin.cpp
+++ b/Engine/Plugins/Media/ElectraPlayer/Source/ElectraPlayerPlugin/Private/ElectraPlayerPlugin.cpp
@@ -86,8 +86,7 @@ bool FElectraPlayerPlugin::Initialize(IMediaEventSink& InEventSink,
 
 	PlayerResourceDelegate = MakeShareable(PlatformCreatePlayerResourceDelegate());
 
-	PlayerDelegate = MakeShareable(new FPlayerAdapterDelegate(AsShared()));
-	Player = MakeShareable(FElectraPlayerRuntimeFactory::CreatePlayer(PlayerDelegate, InSendAnalyticMetricsDelegate, InSendAnalyticMetricsPerMinuteDelegate, InReportVideoStreamingErrorDelegate, InReportSubtitlesFileMetricsDelegate));
+	Player = MakeShareable(FElectraPlayerRuntimeFactory::CreatePlayer(SharedThis(this), InSendAnalyticMetricsDelegate, InSendAnalyticMetricsPerMinuteDelegate, InReportVideoStreamingErrorDelegate, InReportSubtitlesFileMetricsDelegate));
 
 	bMetadataChanged = false;
 	CurrentMetadata.Reset();
@@ -105,7 +104,6 @@ FElectraPlayerPlugin::~FElectraPlayerPlugin()
 		Player->CloseInternal(true);
 		Player.Reset();
 	}
-	PlayerDelegate.Reset();
 	PlayerResourceDelegate.Reset();
 	MediaSamples.Reset();
 }
@@ -253,111 +251,107 @@ private:
 
 //-----------------------------------------------------------------------------
 
-void FElectraPlayerPlugin::FPlayerAdapterDelegate::BlobReceived(const TSharedPtr<TArray<uint8>, ESPMode::ThreadSafe>& InBlobData, IElectraPlayerAdapterDelegate::EBlobResultType InResultType, int32 InResultCode, const Electra::FParamDict* InExtraInfo)
+void FElectraPlayerPlugin::BlobReceived(const TSharedPtr<TArray<uint8>, ESPMode::ThreadSafe>& InBlobData, IElectraPlayerAdapterDelegate::EBlobResultType InResultType, int32 InResultCode, const Electra::FParamDict* InExtraInfo)
 {
 }
 
-Electra::FVariantValue FElectraPlayerPlugin::FPlayerAdapterDelegate::QueryOptions(EOptionType Type, const Electra::FVariantValue& Param)
+Electra::FVariantValue FElectraPlayerPlugin::QueryOptions(EOptionType Type, const Electra::FVariantValue& Param)
 {
-	TSharedPtr<FElectraPlayerPlugin, ESPMode::ThreadSafe> PinnedHost = Host.Pin();
-	if (PinnedHost.IsValid())
+	CallbackPointerLock.Lock();
+	TSharedPtr<IElectraSafeMediaOptionInterface, ESPMode::ThreadSafe> SafeOptionInterface = OptionInterface.Pin();
+	CallbackPointerLock.Unlock();
+	if (SafeOptionInterface.IsValid())
 	{
-		PinnedHost->CallbackPointerLock.Lock();
-		TSharedPtr<IElectraSafeMediaOptionInterface, ESPMode::ThreadSafe> SafeOptionInterface = PinnedHost->OptionInterface.Pin();
-		PinnedHost->CallbackPointerLock.Unlock();
-		if (SafeOptionInterface.IsValid())
+		IElectraSafeMediaOptionInterface::FScopedLock SafeLock(SafeOptionInterface);
+		IMediaOptions *SafeOptions = SafeOptionInterface->GetMediaOptionInterface();
+		if (SafeOptions)
 		{
-			IElectraSafeMediaOptionInterface::FScopedLock SafeLock(SafeOptionInterface);
-			IMediaOptions *SafeOptions = SafeOptionInterface->GetMediaOptionInterface();
-			if (SafeOptions)
+			switch(Type)
 			{
-				switch(Type)
+				case EOptionType::MaxVerticalStreamResolution:
 				{
-					case EOptionType::MaxVerticalStreamResolution:
-					{
-						static const FName MaxResolutionOptionKey = TEXT("MaxResolutionForMediaStreaming");
-						return FVariantValue((int64)SafeOptions->GetMediaOption(MaxResolutionOptionKey, (int64)0));
-					}
+					static const FName MaxResolutionOptionKey = TEXT("MaxResolutionForMediaStreaming");
+					return FVariantValue((int64)SafeOptions->GetMediaOption(MaxResolutionOptionKey, (int64)0));
+				}
 
-					case EOptionType::MaxBandwidthForStreaming:
-					{
-						static const FName MaxBandwidthOptionKey = TEXT("ElectraMaxStreamingBandwidth");
-						return FVariantValue((int64)SafeOptions->GetMediaOption(MaxBandwidthOptionKey, (int64)0));
-					}
+				case EOptionType::MaxBandwidthForStreaming:
+				{
+					static const FName MaxBandwidthOptionKey = TEXT("ElectraMaxStreamingBandwidth");
+					return FVariantValue((int64)SafeOptions->GetMediaOption(MaxBandwidthOptionKey, (int64)0));
+				}
 
-					case EOptionType::PlayListData:
+				case EOptionType::PlayListData:
+				{
+					static const FName PlaylistOptionKey = TEXT("ElectraGetPlaylistData");
+					if (SafeOptions->HasMediaOption(PlaylistOptionKey))
 					{
-						static const FName PlaylistOptionKey = TEXT("ElectraGetPlaylistData");
-						if (SafeOptions->HasMediaOption(PlaylistOptionKey))
-						{
-							check(Param.IsType(FVariantValue::EDataType::TypeFString));
-							return FVariantValue(SafeOptions->GetMediaOption(PlaylistOptionKey, Param.GetFString()));
-						}
-						break;
+						check(Param.IsType(FVariantValue::EDataType::TypeFString));
+						return FVariantValue(SafeOptions->GetMediaOption(PlaylistOptionKey, Param.GetFString()));
 					}
+					break;
+				}
 
-					case EOptionType::LicenseKeyData:
+				case EOptionType::LicenseKeyData:
+				{
+					static const FName LicenseKeyDataOptionKey = TEXT("ElectraGetLicenseKeyData");
+					if (SafeOptions->HasMediaOption(LicenseKeyDataOptionKey))
 					{
-						static const FName LicenseKeyDataOptionKey = TEXT("ElectraGetLicenseKeyData");
-						if (SafeOptions->HasMediaOption(LicenseKeyDataOptionKey))
-						{
-							check(Param.IsType(FVariantValue::EDataType::TypeFString));
-							return FVariantValue(SafeOptions->GetMediaOption(LicenseKeyDataOptionKey, Param.GetFString()));
-						}
-						break;
+						check(Param.IsType(FVariantValue::EDataType::TypeFString));
+						return FVariantValue(SafeOptions->GetMediaOption(LicenseKeyDataOptionKey, Param.GetFString()));
 					}
+					break;
+				}
 
-					case EOptionType::MediaMetadataUpdate:
+				case EOptionType::MediaMetadataUpdate:
+				{
+					static const FName MetadataUpdateOptionKey = TEXT("ElectraMetaDataUpdate");
+					if (SafeOptions->HasMediaOption(MetadataUpdateOptionKey))
 					{
-						static const FName MetadataUpdateOptionKey = TEXT("ElectraMetaDataUpdate");
-						if (SafeOptions->HasMediaOption(MetadataUpdateOptionKey))
-						{
-							check(Param.IsType(FVariantValue::EDataType::TypeFString));
-							// This only provides metadata, the return value of the Get is of no consequence.
-							SafeOptions->GetMediaOption(MetadataUpdateOptionKey, Param.GetFString());
-						}
-						break;
+						check(Param.IsType(FVariantValue::EDataType::TypeFString));
+						// This only provides metadata, the return value of the Get is of no consequence.
+						SafeOptions->GetMediaOption(MetadataUpdateOptionKey, Param.GetFString());
 					}
+					break;
+				}
 
-					case EOptionType::CustomAnalyticsMetric:
+				case EOptionType::CustomAnalyticsMetric:
+				{
+					check(Param.IsType(FVariantValue::EDataType::TypeFString));
+					if (Param.IsType(FVariantValue::EDataType::TypeFString))
 					{
-						check(Param.IsType(FVariantValue::EDataType::TypeFString));
-						if (Param.IsType(FVariantValue::EDataType::TypeFString))
+						FName OptionKey(*Param.GetFString());
+						if (SafeOptions->HasMediaOption(OptionKey))
 						{
-							FName OptionKey(*Param.GetFString());
-							if (SafeOptions->HasMediaOption(OptionKey))
-							{
-								return FVariantValue(SafeOptions->GetMediaOption(OptionKey, FString()));
-							}
+							return FVariantValue(SafeOptions->GetMediaOption(OptionKey, FString()));
 						}
-						break;
 					}
+					break;
+				}
 
-					case EOptionType::PlaystartPosFromSeekPositions:
+				case EOptionType::PlaystartPosFromSeekPositions:
+				{
+					static const FName PlaystartOptionKey = TEXT("ElectraGetPlaystartPosFromSeekPositions");
+					if (SafeOptions->HasMediaOption(PlaystartOptionKey))
 					{
-						static const FName PlaystartOptionKey = TEXT("ElectraGetPlaystartPosFromSeekPositions");
-						if (SafeOptions->HasMediaOption(PlaystartOptionKey))
-						{
-							check(Param.IsType(FVariantValue::EDataType::TypeSharedPointer));
+						check(Param.IsType(FVariantValue::EDataType::TypeSharedPointer));
 
-							TSharedPtr<TArray<FTimespan>, ESPMode::ThreadSafe> PosArray = Param.GetSharedPointer<TArray<FTimespan>>();
-							if (PosArray.IsValid())
+						TSharedPtr<TArray<FTimespan>, ESPMode::ThreadSafe> PosArray = Param.GetSharedPointer<TArray<FTimespan>>();
+						if (PosArray.IsValid())
+						{
+							TSharedPtr<FElectraSeekablePositions, ESPMode::ThreadSafe> Res = StaticCastSharedPtr<FElectraSeekablePositions, IMediaOptions::FDataContainer, ESPMode::ThreadSafe>(SafeOptions->GetMediaOption(PlaystartOptionKey, MakeShared<FElectraSeekablePositions, ESPMode::ThreadSafe>(*PosArray)));
+							if (Res.IsValid() && Res->Data.Num())
 							{
-								TSharedPtr<FElectraSeekablePositions, ESPMode::ThreadSafe> Res = StaticCastSharedPtr<FElectraSeekablePositions, IMediaOptions::FDataContainer, ESPMode::ThreadSafe>(SafeOptions->GetMediaOption(PlaystartOptionKey, MakeShared<FElectraSeekablePositions, ESPMode::ThreadSafe>(*PosArray)));
-								if (Res.IsValid() && Res->Data.Num())
-								{
-									return FVariantValue(int64(Res->Data[0].GetTicks())); // return HNS
-								}
+								return FVariantValue(int64(Res->Data[0].GetTicks())); // return HNS
 							}
-							return FVariantValue();
 						}
-						break;
+						return FVariantValue();
 					}
+					break;
+				}
 
-					default:
-					{
-						break;
-					}
+				default:
+				{
+					break;
 				}
 			}
 		}
@@ -366,174 +360,138 @@ Electra::FVariantValue FElectraPlayerPlugin::FPlayerAdapterDelegate::QueryOption
 }
 
 
-void FElectraPlayerPlugin::FPlayerAdapterDelegate::SendMediaEvent(EPlayerEvent Event)
+void FElectraPlayerPlugin::SendMediaEvent(EPlayerEvent Event)
 {
-	TSharedPtr<FElectraPlayerPlugin, ESPMode::ThreadSafe> PinnedHost = Host.Pin();
-	if (PinnedHost.IsValid())
+	if (Event == EPlayerEvent::MetadataChanged)
 	{
-		if (Event == EPlayerEvent::MetadataChanged)
-		{
-			PinnedHost->SetMetadataChanged();
-		}
-		FScopeLock lock(&PinnedHost->CallbackPointerLock);
-		if (PinnedHost->EventSink)
-		{
-			PinnedHost->EventSink->ReceiveMediaEvent((EMediaEvent)Event);
-		}
+		SetMetadataChanged();
+	}
+	FScopeLock lock(&CallbackPointerLock);
+	if (EventSink)
+	{
+		EventSink->ReceiveMediaEvent((EMediaEvent)Event);
 	}
 }
 
 
-void FElectraPlayerPlugin::FPlayerAdapterDelegate::OnVideoFlush()
+void FElectraPlayerPlugin::OnVideoFlush()
 {
-	TSharedPtr<FElectraPlayerPlugin, ESPMode::ThreadSafe> PinnedHost = Host.Pin();
-	if (PinnedHost.IsValid())
-	{
-		TRange<FTimespan> AllTime(FTimespan::MinValue(), FTimespan::MaxValue());
-		TSharedPtr<IMediaTextureSample, ESPMode::ThreadSafe> FlushSample;
-		while (PinnedHost->GetSamples().FetchVideo(AllTime, FlushSample))
-		{ }
-	}
+	TRange<FTimespan> AllTime(FTimespan::MinValue(), FTimespan::MaxValue());
+	TSharedPtr<IMediaTextureSample, ESPMode::ThreadSafe> FlushSample;
+	while (GetSamples().FetchVideo(AllTime, FlushSample))
+	{ }
 }
 
 
-void FElectraPlayerPlugin::FPlayerAdapterDelegate::OnAudioFlush()
+void FElectraPlayerPlugin::OnAudioFlush()
 {
-	TSharedPtr<FElectraPlayerPlugin, ESPMode::ThreadSafe> PinnedHost = Host.Pin();
-	if (PinnedHost.IsValid())
-	{
-		TRange<FTimespan> AllTime(FTimespan::MinValue(), FTimespan::MaxValue());
-		TSharedPtr<IMediaAudioSample, ESPMode::ThreadSafe> FlushSample;
-		while (PinnedHost->GetSamples().FetchAudio(AllTime, FlushSample))
-		{ }
-	}
+	TRange<FTimespan> AllTime(FTimespan::MinValue(), FTimespan::MaxValue());
+	TSharedPtr<IMediaAudioSample, ESPMode::ThreadSafe> FlushSample;
+	while (GetSamples().FetchAudio(AllTime, FlushSample))
+	{ }
 }
 
 
-void FElectraPlayerPlugin::FPlayerAdapterDelegate::OnSubtitleFlush()
+void FElectraPlayerPlugin::OnSubtitleFlush()
 {
-	TSharedPtr<FElectraPlayerPlugin, ESPMode::ThreadSafe> PinnedHost = Host.Pin();
-	if (PinnedHost.IsValid())
-	{
-		TRange<FTimespan> AllTime(FTimespan::MinValue(), FTimespan::MaxValue());
-		TSharedPtr<IMediaOverlaySample, ESPMode::ThreadSafe> FlushSample;
-		while (PinnedHost->GetSamples().FetchSubtitle(AllTime, FlushSample))
-		{ }
-	}
+	TRange<FTimespan> AllTime(FTimespan::MinValue(), FTimespan::MaxValue());
+	TSharedPtr<IMediaOverlaySample, ESPMode::ThreadSafe> FlushSample;
+	while (GetSamples().FetchSubtitle(AllTime, FlushSample))
+	{ }
 }
 
 
-void FElectraPlayerPlugin::FPlayerAdapterDelegate::PresentVideoFrame(const FVideoDecoderOutputPtr& InVideoFrame)
+
+
+void FElectraPlayerPlugin::PresentVideoFrame(const FVideoDecoderOutputPtr& InVideoFrame)
 {
-	TSharedPtr<FElectraPlayerPlugin, ESPMode::ThreadSafe> PinnedHost = Host.Pin();
-	if (PinnedHost.IsValid())
+	FScopeLock SampleLock(&MediaSamplesLock);
+
+	FVideoDecoderOutputPtr VideoFrame = InVideoFrame;
+	TSharedPtr<FElectraTextureSamplePool, ESPMode::ThreadSafe> TexturePool = OutputTexturePool;
+	if (VideoFrame.IsValid() && TexturePool.IsValid())
 	{
-		FVideoDecoderOutputPtr VideoFrame = InVideoFrame;
-		TSharedPtr<FElectraTextureSamplePool, ESPMode::ThreadSafe> TexturePool = PinnedHost->OutputTexturePool;
-		if (VideoFrame.IsValid() && TexturePool.IsValid())
-		{
-			FElectraTextureSampleRef TextureSample = TexturePool->AcquireShared();
-			TextureSample->Initialize(VideoFrame.Get());
-			PinnedHost->MediaSamples->AddVideo(TextureSample);
-		}
+		SequenceIndexMapperVideo.Remap(VideoFrame);
+		FElectraTextureSampleRef TextureSample = TexturePool->AcquireShared();
+		TextureSample->Initialize(VideoFrame.Get());
+		MediaSamples->AddVideo(TextureSample);
 	}
 }
 
 
-void FElectraPlayerPlugin::FPlayerAdapterDelegate::PresentAudioFrame(const IAudioDecoderOutputPtr& InAudioFrame)
+void FElectraPlayerPlugin::PresentAudioFrame(const IAudioDecoderOutputPtr& InAudioFrame)
 {
-	TSharedPtr<FElectraPlayerPlugin, ESPMode::ThreadSafe> PinnedHost = Host.Pin();
-	if (PinnedHost.IsValid())
+	FScopeLock SampleLock(&MediaSamplesLock);
+
+	IAudioDecoderOutputPtr AudioFrame = InAudioFrame;
+	if (AudioFrame.IsValid())
 	{
-		IAudioDecoderOutputPtr AudioFrame = InAudioFrame;
-		if (AudioFrame.IsValid())
-		{
-			TSharedRef<FElectraPlayerAudioSample, ESPMode::ThreadSafe> AudioSample = PinnedHost->OutputAudioPool.AcquireShared();
-			AudioSample->Initialize(InAudioFrame);
-			PinnedHost->MediaSamples->AddAudio(AudioSample);
-		}
+		SequenceIndexMapperAudio.Remap(AudioFrame);
+		TSharedRef<FElectraPlayerAudioSample, ESPMode::ThreadSafe> AudioSample = OutputAudioPool.AcquireShared();
+		AudioSample->Initialize(InAudioFrame);
+		MediaSamples->AddAudio(AudioSample);
 	}
 }
 
 
-void FElectraPlayerPlugin::FPlayerAdapterDelegate::PresentSubtitleSample(const ISubtitleDecoderOutputPtr& InSubtitleSample)
+void FElectraPlayerPlugin::PresentSubtitleSample(const ISubtitleDecoderOutputPtr& InSubtitleSample)
 {
-	TSharedPtr<FElectraPlayerPlugin, ESPMode::ThreadSafe> PinnedHost = Host.Pin();
-	if (PinnedHost.IsValid())
+	FScopeLock SampleLock(&MediaSamplesLock);
+
+	ISubtitleDecoderOutputPtr Subtitle = InSubtitleSample;
+	if (Subtitle.IsValid())
 	{
-		ISubtitleDecoderOutputPtr Subtitle = InSubtitleSample;
-		if (Subtitle.IsValid())
-		{
-			TSharedRef<FElectraSubtitleSample, ESPMode::ThreadSafe> SubtitleSample = MakeShared<FElectraSubtitleSample, ESPMode::ThreadSafe>();
-			SubtitleSample->Subtitle = InSubtitleSample;
-			PinnedHost->MediaSamples->AddSubtitle(SubtitleSample);
-		}
+		SequenceIndexMapperSubtitle.Remap(Subtitle);
+		TSharedRef<FElectraSubtitleSample, ESPMode::ThreadSafe> SubtitleSample = MakeShared<FElectraSubtitleSample, ESPMode::ThreadSafe>();
+		SubtitleSample->Subtitle = InSubtitleSample;
+		MediaSamples->AddSubtitle(SubtitleSample);
 	}
 }
 
 
-void FElectraPlayerPlugin::FPlayerAdapterDelegate::PresentMetadataSample(const IMetaDataDecoderOutputPtr& InMetadataFrame)
+void FElectraPlayerPlugin::PresentMetadataSample(const IMetaDataDecoderOutputPtr& InMetadataFrame)
 {
-	TSharedPtr<FElectraPlayerPlugin, ESPMode::ThreadSafe> PinnedHost = Host.Pin();
-	if (PinnedHost.IsValid())
+	FScopeLock SampleLock(&MediaSamplesLock);
+
+	IMetaDataDecoderOutputPtr MetadataFrame = InMetadataFrame;
+	if (MetadataFrame.IsValid())
 	{
-		IMetaDataDecoderOutputPtr MetadataFrame = InMetadataFrame;
-		if (MetadataFrame.IsValid())
-		{
-			// Create a binary media sample of our extended format and pass it up.
-			TSharedRef<FElectraBinarySample, ESPMode::ThreadSafe> MetaDataSample = MakeShared<FElectraBinarySample, ESPMode::ThreadSafe>();
-			MetaDataSample->Metadata = InMetadataFrame;
-			PinnedHost->MediaSamples->AddMetadata(MetaDataSample);
-		}
+		SequenceIndexMapperMetadata.Remap(MetadataFrame);
+		TSharedRef<FElectraBinarySample, ESPMode::ThreadSafe> MetaDataSample = MakeShared<FElectraBinarySample, ESPMode::ThreadSafe>();
+		MetaDataSample->Metadata = InMetadataFrame;
+		MediaSamples->AddMetadata(MetaDataSample);
 	}
 }
 
 
-bool FElectraPlayerPlugin::FPlayerAdapterDelegate::CanReceiveVideoSamples(int32 NumFrames)
+bool FElectraPlayerPlugin::CanReceiveVideoSamples(int32 NumFrames)
 {
-	TSharedPtr<FElectraPlayerPlugin, ESPMode::ThreadSafe> PinnedHost = Host.Pin();
-	if (PinnedHost.IsValid())
-	{
-		return PinnedHost->MediaSamples->CanReceiveVideoSamples(NumFrames);
-	}
-	return false;
+	FScopeLock SampleLock(&MediaSamplesLock);
+	return MediaSamples->CanReceiveVideoSamples(NumFrames);
 }
 
 
-bool FElectraPlayerPlugin::FPlayerAdapterDelegate::CanReceiveAudioSamples(int32 NumFrames)
+bool FElectraPlayerPlugin::CanReceiveAudioSamples(int32 NumFrames)
 {
-	TSharedPtr<FElectraPlayerPlugin, ESPMode::ThreadSafe> PinnedHost = Host.Pin();
-	if (PinnedHost.IsValid())
-	{
-		return PinnedHost->MediaSamples->CanReceiveAudioSamples(NumFrames);
-	}
-	return false;
+	FScopeLock SampleLock(&MediaSamplesLock);
+	return MediaSamples->CanReceiveAudioSamples(NumFrames);
 }
 
-void FElectraPlayerPlugin::FPlayerAdapterDelegate::PrepareForDecoderShutdown()
+void FElectraPlayerPlugin::PrepareForDecoderShutdown()
 {
-	TSharedPtr<FElectraPlayerPlugin, ESPMode::ThreadSafe> PinnedHost = Host.Pin();
-	if (PinnedHost.IsValid())
-	{
-		PinnedHost->OutputTexturePool->PrepareForDecoderShutdown();
-	}
+	OutputTexturePool->PrepareForDecoderShutdown();
 }
 
 
-FString FElectraPlayerPlugin::FPlayerAdapterDelegate::GetVideoAdapterName() const
+FString FElectraPlayerPlugin::GetVideoAdapterName() const
 {
 	return GRHIAdapterName;
 }
 
 
-TSharedPtr<IElectraPlayerResourceDelegate, ESPMode::ThreadSafe> FElectraPlayerPlugin::FPlayerAdapterDelegate::GetResourceDelegate() const
+TSharedPtr<IElectraPlayerResourceDelegate, ESPMode::ThreadSafe> FElectraPlayerPlugin::GetResourceDelegate() const
 {
-	TSharedPtr<FElectraPlayerPlugin, ESPMode::ThreadSafe> PinnedHost = Host.Pin();
-	if (PinnedHost.IsValid())
-	{
-		return PinnedHost->PlayerResourceDelegate;
-	}
-	return nullptr;
+	return PlayerResourceDelegate;
 }
 
 //-----------------------------------------------------------------------------
@@ -553,6 +511,7 @@ FString FElectraPlayerPlugin::GetInfo() const
 
 IMediaSamples& FElectraPlayerPlugin::GetSamples()
 {
+	FScopeLock SampleLock(&MediaSamplesLock);
 	return *MediaSamples;
 }
 
@@ -601,19 +560,19 @@ bool FElectraPlayerPlugin::Open(const FString& Url, const IMediaOptions* Options
 	if (InitialAudioLanguage.Len())
 	{
 		LocalPlaystartOptions.InitialAudioTrackAttributes.Language_ISO639 = InitialAudioLanguage;
-		UE_LOG(LogElectraPlayerPlugin, Log, TEXT("[%p] IMediaPlayer::Open: Asking for initial audio language \"%s\""), this, *InitialAudioLanguage);
+		UE_LOG(LogElectraPlayerPlugin, Verbose, TEXT("[%p] IMediaPlayer::Open: Asking for initial audio language \"%s\""), this, *InitialAudioLanguage);
 	}
 	FString InitialSubtitleLanguage = Options->GetMediaOption(TEXT("InitialSubtitleLanguage"), FString());
 	if (InitialSubtitleLanguage.Len())
 	{
 		LocalPlaystartOptions.InitialSubtitleTrackAttributes.Language_ISO639 = InitialSubtitleLanguage;
-		UE_LOG(LogElectraPlayerPlugin, Log, TEXT("[%p] IMediaPlayer::Open: Asking for initial subtitle language \"%s\""), this, *InitialSubtitleLanguage);
+		UE_LOG(LogElectraPlayerPlugin, Verbose, TEXT("[%p] IMediaPlayer::Open: Asking for initial subtitle language \"%s\""), this, *InitialSubtitleLanguage);
 	}
 	bool bNoPreloading = Options->GetMediaOption(TEXT("ElectraNoPreloading"), (bool)false);
 	if (bNoPreloading)
 	{
 		LocalPlaystartOptions.bDoNotPreload = true;
-		UE_LOG(LogElectraPlayerPlugin, Log, TEXT("[%p] IMediaPlayer::Open: No preloading after opening media"), this);
+		UE_LOG(LogElectraPlayerPlugin, Verbose, TEXT("[%p] IMediaPlayer::Open: No preloading after opening media"), this);
 	}
 
 
@@ -634,49 +593,49 @@ bool FElectraPlayerPlugin::Open(const FString& Url, const IMediaOptions* Options
 	if (InitialStreamBitrate > 0)
 	{
 		PlayerOptions.Set(TEXT("initial_bitrate"), Electra::FVariantValue(InitialStreamBitrate));
-		UE_LOG(LogElectraPlayerPlugin, Log, TEXT("[%p] IMediaPlayer::Open: Using initial bitrate of %d bits/second"), this, (int32)InitialStreamBitrate);
+		UE_LOG(LogElectraPlayerPlugin, Verbose, TEXT("[%p] IMediaPlayer::Open: Using initial bitrate of %d bits/second"), this, (int32)InitialStreamBitrate);
 	}
 	FString MediaMimeType = Options->GetMediaOption(TEXT("mimetype"), FString());
 	if (MediaMimeType.Len())
 	{
 		PlayerOptions.Set(TEXT("mime_type"), Electra::FVariantValue(MediaMimeType));
-		UE_LOG(LogElectraPlayerPlugin, Log, TEXT("[%p] IMediaPlayer::Open: Setting media mime type to \"%s\""), this, *MediaMimeType);
+		UE_LOG(LogElectraPlayerPlugin, Verbose, TEXT("[%p] IMediaPlayer::Open: Setting media mime type to \"%s\""), this, *MediaMimeType);
 	}
 	int64 MaxVerticalHeight = Options->GetMediaOption(TEXT("MaxElectraVerticalResolution"), (int64)-1);
 	if (MaxVerticalHeight > 0)
 	{
 		PlayerOptions.Set(TEXT("max_resoY"), Electra::FVariantValue(MaxVerticalHeight));
-		UE_LOG(LogElectraPlayerPlugin, Log, TEXT("[%p] IMediaPlayer::Open: Limiting vertical resolution to %d for all streams"), this, (int32)MaxVerticalHeight);
+		UE_LOG(LogElectraPlayerPlugin, Verbose, TEXT("[%p] IMediaPlayer::Open: Limiting vertical resolution to %d for all streams"), this, (int32)MaxVerticalHeight);
 	}
 	int64 MaxVerticalHeightAt60 = Options->GetMediaOption(TEXT("MaxElectraVerticalResolutionOf60fpsVideos"), (int64)-1);
 	if (MaxVerticalHeightAt60 > 0)
 	{
 		PlayerOptions.Set(TEXT("max_resoY_above_30fps"), Electra::FVariantValue(MaxVerticalHeightAt60));
-		UE_LOG(LogElectraPlayerPlugin, Log, TEXT("[%p] IMediaPlayer::Open: Limiting vertical resolution to %d for streams >30fps"), this, (int32)MaxVerticalHeightAt60);
+		UE_LOG(LogElectraPlayerPlugin, Verbose, TEXT("[%p] IMediaPlayer::Open: Limiting vertical resolution to %d for streams >30fps"), this, (int32)MaxVerticalHeightAt60);
 	}
 	double LiveEdgeDistanceForNormalPresentation = Options->GetMediaOption(TEXT("ElectraLivePresentationOffset"), (double)-1.0);
 	if (LiveEdgeDistanceForNormalPresentation > 0.0)
 	{
 		PlayerOptions.Set(TEXT("seekable_range_live_end_offset"), Electra::FVariantValue(Electra::FTimeValue().SetFromSeconds(LiveEdgeDistanceForNormalPresentation)));
-		UE_LOG(LogElectraPlayerPlugin, Log, TEXT("[%p] IMediaPlayer::Open: Setting distance to live edge for normal presentations to %.3f seconds"), this, LiveEdgeDistanceForNormalPresentation);
+		UE_LOG(LogElectraPlayerPlugin, Verbose, TEXT("[%p] IMediaPlayer::Open: Setting distance to live edge for normal presentations to %.3f seconds"), this, LiveEdgeDistanceForNormalPresentation);
 	}
 	double LiveEdgeDistanceForAudioOnlyPresentation = Options->GetMediaOption(TEXT("ElectraLiveAudioPresentationOffset"), (double)-1.0);
 	if (LiveEdgeDistanceForAudioOnlyPresentation > 0.0)
 	{
 		PlayerOptions.Set(TEXT("seekable_range_live_end_offset_audioonly"), Electra::FVariantValue(Electra::FTimeValue().SetFromSeconds(LiveEdgeDistanceForAudioOnlyPresentation)));
-		UE_LOG(LogElectraPlayerPlugin, Log, TEXT("[%p] IMediaPlayer::Open: Setting distance to live edge for audio-only presentation to %.3f seconds"), this, LiveEdgeDistanceForAudioOnlyPresentation);
+		UE_LOG(LogElectraPlayerPlugin, Verbose, TEXT("[%p] IMediaPlayer::Open: Setting distance to live edge for audio-only presentation to %.3f seconds"), this, LiveEdgeDistanceForAudioOnlyPresentation);
 	}
 	bool bUseConservativeLiveEdgeDistance = Options->GetMediaOption(TEXT("ElectraLiveUseConservativePresentationOffset"), (bool)false);
 	if (bUseConservativeLiveEdgeDistance)
 	{
 		PlayerOptions.Set(TEXT("seekable_range_live_end_offset_conservative"), Electra::FVariantValue(bUseConservativeLiveEdgeDistance));
-		UE_LOG(LogElectraPlayerPlugin, Log, TEXT("[%p] IMediaPlayer::Open: Using conservative live edge for distance calculation"), this);
+		UE_LOG(LogElectraPlayerPlugin, Verbose, TEXT("[%p] IMediaPlayer::Open: Using conservative live edge for distance calculation"), this);
 	}
 	bool bThrowErrorWhenRebuffering = Options->GetMediaOption(TEXT("ElectraThrowErrorWhenRebuffering"), (bool)false);
 	if (bThrowErrorWhenRebuffering)
 	{
 		PlayerOptions.Set(TEXT("throw_error_when_rebuffering"), Electra::FVariantValue(bThrowErrorWhenRebuffering));
-		UE_LOG(LogElectraPlayerPlugin, Log, TEXT("[%p] IMediaPlayer::Open: Throw playback error when rebuffering"), this);
+		UE_LOG(LogElectraPlayerPlugin, Verbose, TEXT("[%p] IMediaPlayer::Open: Throw playback error when rebuffering"), this);
 	}
 	FString CDNHTTPStatusDenyStream = Options->GetMediaOption(TEXT("ElectraGetDenyStreamCode"), FString());
 	if (CDNHTTPStatusDenyStream.Len())
@@ -686,7 +645,7 @@ bool FElectraPlayerPlugin::Open(const FString& Url, const IMediaOptions* Options
 		if (HTTPStatus > 0 && HTTPStatus < 1000)
 		{
 			PlayerOptions.Set(TEXT("abr:cdn_deny_httpstatus"), Electra::FVariantValue((int64)HTTPStatus));
-			UE_LOG(LogElectraPlayerPlugin, Log, TEXT("[%p] IMediaPlayer::Open: CDN HTTP status %d will deny a stream permanently"), this, HTTPStatus);
+			UE_LOG(LogElectraPlayerPlugin, Verbose, TEXT("[%p] IMediaPlayer::Open: CDN HTTP status %d will deny a stream permanently"), this, HTTPStatus);
 		}
 	}
 
@@ -696,7 +655,8 @@ bool FElectraPlayerPlugin::Open(const FString& Url, const IMediaOptions* Options
 	{
 		PlayerOptions.Set(TEXT("optimize_seek_for_scrubbing"), Electra::FVariantValue(true));
 		PlayerOptions.Set(TEXT("new_scrubbing_seek_cancels_current"), Electra::FVariantValue(true));
-		PlayerOptions.Set(TEXT("do_not_hold_back_first_frame"), Electra::FVariantValue(true));
+		//PlayerOptions.Set(TEXT("do_not_hold_back_first_frame"), Electra::FVariantValue(true));
+		PlayerOptions.Set(TEXT("always_emit_samples_when_paused"), Electra::FVariantValue(true));
 		PlayerOptions.Set(TEXT("worker_threads"), Electra::FVariantValue(FString(TEXT("worker"))));
 	}
 	else
@@ -741,6 +701,11 @@ bool FElectraPlayerPlugin::Open(const FString& Url, const IMediaOptions* Options
 	bMetadataChanged = false;
 	CurrentMetadata.Reset();
 
+	SequenceIndexMapperVideo.Reset();
+	SequenceIndexMapperAudio.Reset();
+	SequenceIndexMapperSubtitle.Reset();
+	SequenceIndexMapperMetadata.Reset();
+
 	// Check if we can get a segment cache interface for this playback request...
 	TSharedPtr<FElectraPlayerDataCacheContainer, ESPMode::ThreadSafe> ElectraPlayerDataCacheContainer;
 	TSharedPtr<IMediaOptions::FDataContainer, ESPMode::ThreadSafe> ElectraPlayerDataCacheDefaultValue;
@@ -986,7 +951,109 @@ bool FElectraPlayerPlugin::Seek(const FTimespan& Time)
 	UE_LOG(LogElectraPlayerPlugin, Log, TEXT("[%p] IMediaControls::Seek() to %s"), this, *Time.ToString(TEXT("%h:%m:%s.%f")));
 	CSV_EVENT(ElectraPlayer, TEXT("Seeking"));
 
-	return Player->Seek(Time);
+	const int32 NextSequenceIndex = CurrentSequenceIndex + 1;
+
+	// Check if the target time is already available.
+	FScopeLock SampleLock(&MediaSamplesLock);
+	if (!MediaSamples.Get())
+	{
+		return false;
+	}
+	FMediaSamples& Samples = *MediaSamples;
+	const float Rate = Player->GetRate();
+	if (Rate >= 0.0f)
+	{
+		TArray<TRange<FMediaTimeStamp>> QueuedRange;
+		// Note: When a Seek() is performed, the loop counter contractually restarts at 0 !!
+		FMediaTimeStamp TargetTime = FMediaTimeStamp(Time, FMediaTimeStamp::MakeSequenceIndex(CurrentSequenceIndex, 0));
+		if (Samples.PeekVideoSampleTimeRanges(QueuedRange) && QueuedRange.Num())
+		{
+			FMediaTimeStamp MinTime(FTimespan::MinValue(), FMediaTimeStamp::MakeSequenceIndex(-1, 0));
+			FMediaTimeStamp DiscardTo = MinTime;
+			// Loop over all entries in case there are several candidates in a (more or less)
+			// degenerate case of a single looping frame.
+			for(int32 i=0; i<QueuedRange.Num(); ++i)
+			{
+				// For comparison reasons we need to assume the loop counter of the sample range we're looking at.
+				TargetTime.SetSecondaryIndex(QueuedRange[i].GetLowerBoundValue().GetSecondaryIndex());
+				if (QueuedRange[i].Contains(TargetTime))
+				{
+					DiscardTo = QueuedRange[i].GetLowerBoundValue() - FTimespan(1);
+				}
+			}
+			if (DiscardTo > MinTime)
+			{
+				UE_LOG(LogElectraPlayerPlugin, Verbose, TEXT("Seek time %lld (%d,%d) already available, rolling forward"), (long long int)TargetTime.Time.GetTicks(), TargetTime.GetPrimaryIndex(), TargetTime.GetSecondaryIndex());
+				TRange<FMediaTimeStamp> DiscardRange(MinTime, DiscardTo);
+
+				Samples.DiscardVideoSamples(DiscardRange, false);
+				Samples.DiscardAudioSamples(DiscardRange, false);
+				Samples.DiscardCaptionSamples(DiscardRange, false);
+				Samples.DiscardSubtitleSamples(DiscardRange, false);
+				Samples.DiscardMetadataSamples(DiscardRange, false);
+
+				TOptional<int32> LastVideoLoopIndex;
+				TOptional<int32> LastAudioLoopIndex;
+
+				// We now need to "renumber" the (primary) sequence index of what is left in the sample queues.
+				// The secondary index (the loop index) restarts at 0. See above.
+				FMediaSamples TempSamples;
+				TRange<FTimespan> AllTime(FTimespan::MinValue(), FTimespan::MaxValue());
+				// Video
+				{
+					TSharedPtr<IMediaTextureSample, ESPMode::ThreadSafe> TempSample;
+					while(Samples.FetchVideo(AllTime, TempSample))
+					{
+						TSharedPtr<FElectraTextureSample, ESPMode::ThreadSafe> vs = StaticCastSharedPtr<FElectraTextureSample>(TempSample);
+						FMediaTimeStamp ts = vs->GetTime();
+						LastVideoLoopIndex = ts.GetSecondaryIndex();
+						vs->SetTime(ts.SetPrimaryIndex(NextSequenceIndex).SetSecondaryIndex(0));
+						TempSamples.AddVideo(TempSample.ToSharedRef());
+					}
+					while(TempSamples.FetchVideo(AllTime, TempSample))
+					{
+						Samples.AddVideo(TempSample.ToSharedRef());
+					}
+				}
+				// Audio
+				{
+					TSharedPtr<IMediaAudioSample, ESPMode::ThreadSafe> TempSample;
+					while(Samples.FetchAudio(AllTime, TempSample))
+					{
+						TSharedPtr<FElectraPlayerAudioSample, ESPMode::ThreadSafe> as = StaticCastSharedPtr<FElectraPlayerAudioSample>(TempSample);
+						FMediaTimeStamp ts = as->GetTime();
+						LastAudioLoopIndex = ts.GetSecondaryIndex();
+						as->SetTime(ts.SetPrimaryIndex(NextSequenceIndex).SetSecondaryIndex(0));
+						TempSamples.AddAudio(TempSample.ToSharedRef());
+					}
+					while(TempSamples.FetchAudio(AllTime, TempSample))
+					{
+						Samples.AddAudio(TempSample.ToSharedRef());
+					}
+				}
+				// Set the index remappers. The primary index changes only with a Seek() call, and since we are not
+				// actually seeking whatever the player still has available needs to be remapped to the new index
+				// the facade is expecting.
+				SequenceIndexMapperVideo.SetRemapPrimaryIndex(NextSequenceIndex, LastVideoLoopIndex);
+				SequenceIndexMapperAudio.SetRemapPrimaryIndex(NextSequenceIndex, LastAudioLoopIndex);
+				SequenceIndexMapperSubtitle.SetRemapPrimaryIndex(NextSequenceIndex, TOptional<int32>());
+				SequenceIndexMapperMetadata.SetRemapPrimaryIndex(NextSequenceIndex, TOptional<int32>());
+				CurrentSequenceIndex = NextSequenceIndex;
+				SendMediaEvent(EPlayerEvent::SeekCompleted);
+				return true;
+			}
+		}
+	}
+
+	SequenceIndexMapperVideo.SetExpectedPrimaryIndex(NextSequenceIndex);
+	SequenceIndexMapperAudio.SetExpectedPrimaryIndex(NextSequenceIndex);
+	SequenceIndexMapperSubtitle.SetExpectedPrimaryIndex(NextSequenceIndex);
+	SequenceIndexMapperMetadata.SetExpectedPrimaryIndex(NextSequenceIndex);
+	CurrentSequenceIndex = NextSequenceIndex;
+
+	IElectraPlayerInterface::FSeekParam sp;
+	sp.SequenceIndex = CurrentSequenceIndex;
+	return Player->Seek(Time, sp);
 }
 
 TRange<FTimespan> FElectraPlayerPlugin::GetPlaybackTimeRange(EMediaTimeRangeType InRangeToGet) const
@@ -1012,16 +1079,18 @@ bool FElectraPlayerPlugin::QueryCacheState(EMediaCacheState State, TRangeSet<FTi
 	{
 		case EMediaCacheState::Loaded:
 		case EMediaCacheState::Loading:
+		case EMediaCacheState::Pending:
 		{
 			// When asked to provide what's already loaded we look at what we have in the sample queue
 			// and add that to the result. These samples have already left the player but are ready
 			// for use.
 			if (State == EMediaCacheState::Loaded)
 			{
-				TRange<FMediaTimeStamp> QueuedRange;
-				if (MediaSamples.IsValid() && MediaSamples->PeekVideoSampleTimeRange(QueuedRange))
+				TArray<TRange<FMediaTimeStamp>> QueuedRange;
+				FScopeLock SampleLock(&MediaSamplesLock);
+				if (MediaSamples.IsValid() && MediaSamples->PeekVideoSampleTimeRanges(QueuedRange) && QueuedRange.Num())
 				{
-					OutTimeRanges.Add(TRange<FTimespan>(QueuedRange.GetLowerBoundValue().Time, QueuedRange.GetUpperBoundValue().Time));
+					OutTimeRanges.Add(TRange<FTimespan>(QueuedRange[0].GetLowerBoundValue().Time, QueuedRange.Last().GetUpperBoundValue().Time));
 				}
 			}
 
@@ -1031,7 +1100,22 @@ bool FElectraPlayerPlugin::QueryCacheState(EMediaCacheState State, TRangeSet<FTi
 			bool bHaveVid = Player->GetStreamBufferInformation(vidBuf, IElectraPlayerInterface::EPlayerTrackType::Video);
 			bool bHaveAud = !bHaveVid ? Player->GetStreamBufferInformation(audBuf, IElectraPlayerInterface::EPlayerTrackType::Audio) : false;
 			const IElectraPlayerInterface::FStreamBufferInfo* Buffer = bHaveVid ? &vidBuf : bHaveAud ? &audBuf : nullptr;
-			const TArray<IElectraPlayerInterface::FStreamBufferInfo::FTimeRange>* tr = Buffer ? (State == EMediaCacheState::Loaded ? &Buffer->TimeAvailable : &Buffer->TimeRequested) : nullptr;
+			const TArray<IElectraPlayerInterface::FStreamBufferInfo::FTimeRange>* tr = nullptr;
+			if (Buffer)
+			{
+				switch(State)
+				{
+					case EMediaCacheState::Loaded:
+						tr = &Buffer->TimeEnqueued;
+						break;
+					case EMediaCacheState::Loading:
+						tr = &Buffer->TimeAvailable;
+						break;
+					case EMediaCacheState::Pending:
+						tr = &Buffer->TimeRequested;
+						break;
+				}
+			}
 			for(int32 i=0,iMax=tr?tr->Num():0; i<iMax; ++i)
 			{
 				OutTimeRanges.Add(TRange<FTimespan>(tr->operator[](i).Start.Time, tr->operator[](i).End.Time));
diff --git a/Engine/Plugins/Media/ElectraPlayer/Source/ElectraPlayerPlugin/Private/ElectraPlayerPlugin.h b/Engine/Plugins/Media/ElectraPlayer/Source/ElectraPlayerPlugin/Private/ElectraPlayerPlugin.h
--- a/Engine/Plugins/Media/ElectraPlayer/Source/ElectraPlayerPlugin/Private/ElectraPlayerPlugin.h
+++ b/Engine/Plugins/Media/ElectraPlayer/Source/ElectraPlayerPlugin/Private/ElectraPlayerPlugin.h
@@ -34,15 +34,16 @@ class FMediaSamples;
 class FArchive;
 
 /**
-	* Implements a media player
-	* Supports multiple platforms.
-	*/
+ * Implements a media player
+ * Supports multiple platforms.
+ */
 class FElectraPlayerPlugin
 	: public IMediaPlayer
 	, protected IMediaCache
 	, protected IMediaView
 	, protected IMediaControls
 	, public IMediaTracks
+	, public IElectraPlayerAdapterDelegate
 	, public TSharedFromThis<FElectraPlayerPlugin, ESPMode::ThreadSafe>
 {
 public:
@@ -151,37 +152,85 @@ private:
 	IElectraPlayerResourceDelegate* PlatformCreatePlayerResourceDelegate();
 
 	// IElectraPlayerAdapterDelegate impl
-	class FPlayerAdapterDelegate : public IElectraPlayerAdapterDelegate
+	Electra::FVariantValue QueryOptions(EOptionType Type, const Electra::FVariantValue & Param) override;
+	void BlobReceived(const TSharedPtr<TArray<uint8>, ESPMode::ThreadSafe>& InBlobData, IElectraPlayerAdapterDelegate::EBlobResultType InResultType, int32 InResultCode, const Electra::FParamDict* InExtraInfo) override;
+	void SendMediaEvent(EPlayerEvent Event) override;
+	void OnVideoFlush() override;
+	void OnAudioFlush() override;
+	void OnSubtitleFlush() override;
+	void PresentVideoFrame(const FVideoDecoderOutputPtr& InVideoFrame) override;
+	void PresentAudioFrame(const IAudioDecoderOutputPtr& InAudioFrame) override;
+	void PresentSubtitleSample(const ISubtitleDecoderOutputPtr& InSubtitleSample) override;
+	void PresentMetadataSample(const IMetaDataDecoderOutputPtr& InMetadataSample) override;
+	bool CanReceiveVideoSamples(int32 NumFrames) override;
+	bool CanReceiveAudioSamples(int32 NumFrames) override;
+	void PrepareForDecoderShutdown() override;
+	FString GetVideoAdapterName() const override;
+	TSharedPtr<IElectraPlayerResourceDelegate, ESPMode::ThreadSafe> GetResourceDelegate() const override;
+
+
+private:
+	template <typename T>
+	struct FSequenceIndexMapper
 	{
-	public:
-		FPlayerAdapterDelegate(const TSharedPtr<FElectraPlayerPlugin, ESPMode::ThreadSafe> & InHost) : Host(InHost) {}
-
-		Electra::FVariantValue QueryOptions(EOptionType Type, const Electra::FVariantValue & Param) override;
-		void BlobReceived(const TSharedPtr<TArray<uint8>, ESPMode::ThreadSafe>& InBlobData, IElectraPlayerAdapterDelegate::EBlobResultType InResultType, int32 InResultCode, const Electra::FParamDict* InExtraInfo) override;
-		void SendMediaEvent(EPlayerEvent Event) override;
-		void OnVideoFlush() override;
-		void OnAudioFlush() override;
-		void OnSubtitleFlush() override;
-		void PresentVideoFrame(const FVideoDecoderOutputPtr& InVideoFrame) override;
-		void PresentAudioFrame(const IAudioDecoderOutputPtr& InAudioFrame) override;
-		void PresentSubtitleSample(const ISubtitleDecoderOutputPtr& InSubtitleSample) override;
-		void PresentMetadataSample(const IMetaDataDecoderOutputPtr& InMetadataSample) override;
-		bool CanReceiveVideoSamples(int32 NumFrames) override;
-		bool CanReceiveAudioSamples(int32 NumFrames) override;
-		void PrepareForDecoderShutdown() override;
-		FString GetVideoAdapterName() const override;
-		TSharedPtr<IElectraPlayerResourceDelegate, ESPMode::ThreadSafe> GetResourceDelegate() const override;
-
-	private:
-		TWeakPtr<FElectraPlayerPlugin, ESPMode::ThreadSafe> Host;
+		void Reset()
+		{
+			SetExpectedPrimaryIndex(0);
+		}
+		void SetExpectedPrimaryIndex(int32 InIndex)
+		{
+			RemapPrimaryIndex.Reset();
+			RemapSecondaryIndex.Reset();
+			PrimaryIndex = InIndex;
+			SecondaryIndex = 0;
+		}
+		void SetRemapPrimaryIndex(int32 InIndex, const TOptional<int32>& InLoop)
+		{
+			RemapPrimaryIndex = InIndex;
+			RemapSecondaryIndex = InLoop;
+		}
+		bool Remap(T InOutSample)
+		{
+			FDecoderTimeStamp ts = InOutSample->GetTime();
+			// Remap?
+			if (RemapPrimaryIndex.IsSet())
+			{
+				ts = RemapTime(ts);
+				InOutSample->SetTime(ts);
+			}
+			// Always return true for now.
+			// In the future we could check if the sequence index is
+			// what is expected and if not return false to indicate this.
+			return true;
+		}
+		FDecoderTimeStamp RemapTime(FDecoderTimeStamp InTime)
+		{
+			int32 InLoopCount = (int32)(InTime.SequenceIndex & 0xffffffff);
+			// Do we have an established loop count yet?
+			if (!RemapSecondaryIndex.IsSet())
+			{
+				RemapSecondaryIndex = InLoopCount;
+			}
+			uint32 NewLoopCount = (uint32)InLoopCount - (uint32)RemapSecondaryIndex.GetValue();
+
+			InTime.SequenceIndex = ((int64)RemapPrimaryIndex.Get(0) << 32) | NewLoopCount;
+			return InTime;
+		}
+		int32 PrimaryIndex = 0;
+		int32 SecondaryIndex = 0;
+		TOptional<int32> RemapPrimaryIndex;
+		TOptional<int32> RemapSecondaryIndex;
 	};
-	friend class FPlayerAdapterDelegate;
-	TSharedPtr<FPlayerAdapterDelegate, ESPMode::ThreadSafe> PlayerDelegate;
+	FSequenceIndexMapper<FVideoDecoderOutputPtr> SequenceIndexMapperVideo;
+	FSequenceIndexMapper<IAudioDecoderOutputPtr> SequenceIndexMapperAudio;
+	FSequenceIndexMapper<ISubtitleDecoderOutputPtr> SequenceIndexMapperSubtitle;
+	FSequenceIndexMapper<IMetaDataDecoderOutputPtr> SequenceIndexMapperMetadata;
+	int32 CurrentSequenceIndex = 0;
 
 
-private:
 	/** Output queues as needed by MediaFramework */
 	TUniquePtr<FMediaSamples> MediaSamples;
+	mutable FCriticalSection MediaSamplesLock;
 
 	/** Lock to guard the POD callback pointers from being changed while being used. */
 	FCriticalSection CallbackPointerLock;
diff --git a/Engine/Plugins/Media/ElectraPlayer/Source/ElectraPlayerPlugin/Private/ElectraPlayerPluginModule.cpp b/Engine/Plugins/Media/ElectraPlayer/Source/ElectraPlayerPlugin/Private/ElectraPlayerPluginModule.cpp
--- a/Engine/Plugins/Media/ElectraPlayer/Source/ElectraPlayerPlugin/Private/ElectraPlayerPluginModule.cpp
+++ b/Engine/Plugins/Media/ElectraPlayer/Source/ElectraPlayerPlugin/Private/ElectraPlayerPluginModule.cpp
@@ -35,14 +35,11 @@ public:
 	{
 		if (bInitialized)
 		{
-			FElectraPlayerPlugin* NewRawPlayer = new FElectraPlayerPlugin();
-			if (NewRawPlayer)
+			TSharedPtr<FElectraPlayerPlugin, ESPMode::ThreadSafe> NewPlayer = MakeShared<FElectraPlayerPlugin, ESPMode::ThreadSafe>();
+			check(NewPlayer.IsValid());
+			if (NewPlayer->Initialize(EventSink, SendAnalyticMetricsDelegate, SendAnalyticMetricsPerMinuteDelegate, ReportVideoStreamingErrorDelegate, ReportSubtitlesMetricsDelegate))
 			{
-				TSharedPtr<IMediaPlayer, ESPMode::ThreadSafe> NewPlayer = MakeShareable(NewRawPlayer);
-				if (NewRawPlayer->Initialize(EventSink, SendAnalyticMetricsDelegate, SendAnalyticMetricsPerMinuteDelegate, ReportVideoStreamingErrorDelegate, ReportSubtitlesMetricsDelegate))
-				{
-					return NewPlayer;
-				}
+				return NewPlayer;
 			}
 		}
 		return nullptr;
@@ -112,7 +109,7 @@ public:
 		}
 		Event = FFunctionGraphTask::CreateAndDispatchWhenReady(MoveTemp(CodeToRun), GET_STATID(STAT_ElectraAsyncJob), &Events);
 	}
-	
+
 	void StartupModule() override
 	{
 		// Check that we have the player module and that it has initialized successfully.
diff --git a/Engine/Plugins/Media/ElectraPlayer/Source/ElectraPlayerRuntime/Private/ElectraPlayer.cpp b/Engine/Plugins/Media/ElectraPlayer/Source/ElectraPlayerRuntime/Private/ElectraPlayer.cpp
--- a/Engine/Plugins/Media/ElectraPlayer/Source/ElectraPlayerRuntime/Private/ElectraPlayer.cpp
+++ b/Engine/Plugins/Media/ElectraPlayer/Source/ElectraPlayerRuntime/Private/ElectraPlayer.cpp
@@ -90,16 +90,17 @@ class FMetaDataDecoderOutput : public IMetaDataDecoderOutput
 {
 public:
 	virtual ~FMetaDataDecoderOutput() = default;
-	virtual const void* GetData() override									{ return Data.GetData(); }
-	virtual FTimespan GetDuration() const override							{ return Duration; }
-	virtual uint32 GetSize() const override									{ return (uint32) Data.Num(); }
-	virtual FDecoderTimeStamp GetTime() const override						{ return PresentationTime; }
-	virtual EOrigin GetOrigin() const override								{ return Origin; }
-	virtual EDispatchedMode GetDispatchedMode() const override				{ return DispatchedMode; }
-	virtual const FString& GetSchemeIdUri() const override					{ return SchemeIdUri; }
-	virtual const FString& GetValue() const override						{ return Value; }
-	virtual const FString& GetID() const override							{ return ID; }
-	virtual TOptional<FDecoderTimeStamp> GetTrackBaseTime() const override	{ return TrackBaseTime; }
+	const void* GetData() override									{ return Data.GetData(); }
+	FTimespan GetDuration() const override							{ return Duration; }
+	uint32 GetSize() const override									{ return (uint32) Data.Num(); }
+	FDecoderTimeStamp GetTime() const override						{ return PresentationTime; }
+	EOrigin GetOrigin() const override								{ return Origin; }
+	EDispatchedMode GetDispatchedMode() const override				{ return DispatchedMode; }
+	const FString& GetSchemeIdUri() const override					{ return SchemeIdUri; }
+	const FString& GetValue() const override						{ return Value; }
+	const FString& GetID() const override							{ return ID; }
+	TOptional<FDecoderTimeStamp> GetTrackBaseTime() const override	{ return TrackBaseTime; }
+	void SetTime(FDecoderTimeStamp& InTime) override				{ PresentationTime = InTime; }
 
 	TArray<uint8> Data;
 	FDecoderTimeStamp PresentationTime;
@@ -164,7 +165,7 @@ FElectraPlayer::FElectraPlayer(const TSharedPtr<IElectraPlayerAdapterDelegate, E
 
 	ClearToDefaultState();
 
-	UE_LOG(LogElectraPlayer, Log, TEXT("[%p] FElectraPlayer() created."), this);
+	UE_LOG(LogElectraPlayer, Verbose, TEXT("[%p] FElectraPlayer() created."), this);
 }
 
 //-----------------------------------------------------------------------------
@@ -185,7 +186,7 @@ FElectraPlayer::~FElectraPlayer()
 	SendAnalyticMetricsPerMinuteDelegate.RemoveAll(this);
 	ReportVideoStreamingErrorDelegate.RemoveAll(this);
 	ReportSubtitlesMetricsDelegate.RemoveAll(this);
-	UE_LOG(LogElectraPlayer, Log, TEXT("[%p] ~FElectraPlayer() finished."), this);
+	UE_LOG(LogElectraPlayer, Verbose, TEXT("[%p] ~FElectraPlayer() finished."), this);
 
 	if (AsyncResourceReleaseNotification.IsValid())
 	{
@@ -373,7 +374,7 @@ bool FElectraPlayer::OpenInternal(const FString& Url, const FParamDict& InPlayer
 			return false;
 		}
 		CurrentPlayer = MoveTemp(NewPlayer);
-		UE_LOG(LogElectraPlayer, Log, TEXT("[%p][%p] IMediaPlayer::OpenBlob(%s)"), this, CurrentPlayer.Get(), *SanitizeMessage(MediaUrl));
+		UE_LOG(LogElectraPlayer, Verbose, TEXT("[%p][%p] IMediaPlayer::OpenBlob(%s)"), this, CurrentPlayer.Get(), *SanitizeMessage(MediaUrl));
 		PendingBlobRequest->Request->URL(MediaUrl).Callback().BindThreadSafeSP(PendingBlobRequest.ToSharedRef(), &FBlobRequest::OnBlobRequestComplete);
 		CurrentPlayer->AdaptivePlayer->LoadBlob(PendingBlobRequest->Request);
 	}
@@ -509,7 +510,7 @@ void FElectraPlayer::FInternalPlayerImpl::DoCloseAsync(TSharedPtr<FInternalPlaye
 		Player->RendererAudio.Reset();
 		*bClosedSig = true;
 		double TimeCloseEnded = FPlatformTime::Seconds();
-		UE_LOG(LogElectraPlayer, Log, TEXT("[%p] DoCloseAsync() finished after %.3f msec!"), Player.Get(), (TimeCloseEnded - TimeCloseBegan) * 1000.0);
+		UE_LOG(LogElectraPlayer, Verbose, TEXT("[%p] DoCloseAsync() finished after %.3f msec!"), Player.Get(), (TimeCloseEnded - TimeCloseBegan) * 1000.0);
 
 		if (AsyncResourceReleaseNotification.IsValid())
 		{
@@ -615,7 +616,7 @@ bool FElectraPlayer::GetStreamBufferInformation(IElectraPlayerInterface::FStream
 					tv.End.SequenceIndex = InRanges[i].End.GetSequenceIndex();
 				}
 			};
-			AddRanges(OutBufferInformation.TimeAvailable, bi.TimeEnqueued);
+			AddRanges(OutBufferInformation.TimeEnqueued, bi.TimeEnqueued);
 			AddRanges(OutBufferInformation.TimeAvailable, bi.TimeAvailable);
 			AddRanges(OutBufferInformation.TimeRequested, bi.TimeRequested);
 			return true;
@@ -1120,9 +1121,9 @@ bool FElectraPlayer::SetRate(float Rate)
 				TriggerFirstSeekIfNecessary();
 				CurrentPlayer->AdaptivePlayer->Resume();
 			}
-			IAdaptiveStreamingPlayer::FTrickplayParams Params;
-			CurrentPlayer->AdaptivePlayer->SetPlayRate((double) Rate, Params);
 		}
+		IAdaptiveStreamingPlayer::FTrickplayParams Params;
+		CurrentPlayer->AdaptivePlayer->SetPlayRate((double) Rate, Params);
 		return true;
 	}
 	return false;
@@ -1153,9 +1154,6 @@ void FElectraPlayer::TriggerFirstSeekIfNecessary()
 		// Set up the initial playback position
 		IAdaptiveStreamingPlayer::FSeekParam playParam;
 
-		// Whatever the initial seek target: we will ignore it for sequence index updates
-		playParam.bIgnoreForSequenceIndex = true;
-
 		// First we look at any potential time offset specified in the playstart options.
 		if (PlaystartOptions.TimeOffset.IsSet())
 		{
@@ -1271,6 +1269,7 @@ bool FElectraPlayer::Seek(const FTimespan& Time, const FSeekParam& Param)
 		{
 			seek.Time.SetFromTimespan(Target);
 		}
+		seek.NewSequenceIndex = Param.SequenceIndex;
 		seek.StartingBitrate = Param.StartingBitrate;
 		seek.bOptimizeForScrubbing = Param.bOptimizeForScrubbing;
 		seek.DistanceThreshold = Param.DistanceThreshold;
@@ -2086,7 +2085,7 @@ void FElectraPlayer::HandlePlayerEventOpenSource(const FString& URL)
 	PlayerState.State = EPlayerState::Preparing;
 	DeferredEvents.Enqueue(IElectraPlayerAdapterDelegate::EPlayerEvent::MediaConnecting);
 
-	UE_LOG(LogElectraPlayer, Log, TEXT("[%p][%p] Opening stream at \"%s\""), this, CurrentPlayer.Get(), *SanitizeMessage(URL));
+	UE_LOG(LogElectraPlayer, Verbose, TEXT("[%p][%p] Opening stream at \"%s\""), this, CurrentPlayer.Get(), *SanitizeMessage(URL));
 
 	// Update statistics
 	FScopeLock Lock(&StatisticsLock);
@@ -2107,7 +2106,7 @@ void FElectraPlayer::HandlePlayerEventOpenSource(const FString& URL)
 
 void FElectraPlayer::HandlePlayerEventReceivedMasterPlaylist(const FString& EffectiveURL)
 {
-	UE_LOG(LogElectraPlayer, Log, TEXT("[%p][%p] Received master playlist from \"%s\""), this, CurrentPlayer.Get(), *SanitizeMessage(EffectiveURL));
+	UE_LOG(LogElectraPlayer, Verbose, TEXT("[%p][%p] Received master playlist from \"%s\""), this, CurrentPlayer.Get(), *SanitizeMessage(EffectiveURL));
 
 	// Update statistics
 	FScopeLock Lock(&StatisticsLock);
@@ -2138,7 +2137,7 @@ void FElectraPlayer::HandlePlayerEventReceivedPlaylists()
 
 	MediaStateOnPreparingFinished();
 
-	UE_LOG(LogElectraPlayer, Log, TEXT("[%p][%p] Received initial stream playlists"), this, CurrentPlayer.Get());
+	UE_LOG(LogElectraPlayer, Verbose, TEXT("[%p][%p] Received initial stream playlists"), this, CurrentPlayer.Get());
 
 	Electra::FTimeRange MediaTimeline;
 	Electra::FTimeValue MediaDuration;
@@ -2164,7 +2163,7 @@ void FElectraPlayer::HandlePlayerEventReceivedPlaylists()
 	{
 		for(int32 i=0; i<VideoStreamMetaData[0].StreamDetails.Num(); ++i)
 		{
-			UE_LOG(LogElectraPlayer, Log, TEXT("[%p][%p] Found %d * %d video stream at bitrate %d"), this, CurrentPlayer.Get(),
+			UE_LOG(LogElectraPlayer, Verbose, TEXT("[%p][%p] Found %d * %d video stream at bitrate %d"), this, CurrentPlayer.Get(),
 												VideoStreamMetaData[0].StreamDetails[i].CodecInformation.GetResolution().Width,
 												VideoStreamMetaData[0].StreamDetails[i].CodecInformation.GetResolution().Height,
 												VideoStreamMetaData[0].StreamDetails[i].Bandwidth);
@@ -2189,7 +2188,7 @@ void FElectraPlayer::HandlePlayerEventReceivedPlaylists()
 	{
 		for(int32 i=0; i<AudioStreamMetaData[0].StreamDetails.Num(); ++i)
 		{
-			UE_LOG(LogElectraPlayer, Log, TEXT("[%p][%p] Found audio stream at bitrate %d"), this, CurrentPlayer.Get(),
+			UE_LOG(LogElectraPlayer, Verbose, TEXT("[%p][%p] Found audio stream at bitrate %d"), this, CurrentPlayer.Get(),
 												AudioStreamMetaData[0].StreamDetails[i].Bandwidth);
 		}
 	}
@@ -2228,7 +2227,7 @@ void FElectraPlayer::HandlePlayerEventTracksChanged()
 	{
 		for(int32 i=0; i<VideoStreamMetaData[0].StreamDetails.Num(); ++i)
 		{
-			UE_LOG(LogElectraPlayer, Log, TEXT("[%p][%p] Found %d * %d video stream at bitrate %d"), this, CurrentPlayer.Get(),
+			UE_LOG(LogElectraPlayer, Verbose, TEXT("[%p][%p] Found %d * %d video stream at bitrate %d"), this, CurrentPlayer.Get(),
 												VideoStreamMetaData[0].StreamDetails[i].CodecInformation.GetResolution().Width,
 												VideoStreamMetaData[0].StreamDetails[i].CodecInformation.GetResolution().Height,
 												VideoStreamMetaData[0].StreamDetails[i].Bandwidth);
@@ -2242,7 +2241,7 @@ void FElectraPlayer::HandlePlayerEventTracksChanged()
 	{
 		for(int32 i=0; i<AudioStreamMetaData[0].StreamDetails.Num(); ++i)
 		{
-			UE_LOG(LogElectraPlayer, Log, TEXT("[%p][%p] Found audio stream at bitrate %d"), this, CurrentPlayer.Get(),
+			UE_LOG(LogElectraPlayer, Verbose, TEXT("[%p][%p] Found audio stream at bitrate %d"), this, CurrentPlayer.Get(),
 												AudioStreamMetaData[0].StreamDetails[i].Bandwidth);
 		}
 	}
@@ -2291,7 +2290,7 @@ void FElectraPlayer::HandlePlayerEventLicenseKey(const Electra::Metrics::FLicens
 	// TBD
 	if (LicenseKeyStats.bWasSuccessful)
 	{
-		UE_LOG(LogElectraPlayer, Log, TEXT("[%p][%p] License key obtained"), this, CurrentPlayer.Get());
+		UE_LOG(LogElectraPlayer, Verbose, TEXT("[%p][%p] License key obtained"), this, CurrentPlayer.Get());
 		FScopeLock Lock(&StatisticsLock);
 		Statistics.AddMessageToHistory(TEXT("Obtained license key"));
 	}
@@ -2381,7 +2380,7 @@ void FElectraPlayer::HandlePlayerEventBufferingStart(Electra::Metrics::EBufferin
 	FString Msg = FString::Printf(TEXT("%s buffering starts"), Electra::Metrics::GetBufferingReasonString(BufferingReason));
 	Statistics.AddMessageToHistory(Msg);
 
-	UE_LOG(LogElectraPlayer, Log, TEXT("[%p][%p] %s"), this, CurrentPlayer.Get(), *Msg);
+	UE_LOG(LogElectraPlayer, Verbose, TEXT("[%p][%p] %s"), this, CurrentPlayer.Get(), *Msg);
 	CSV_EVENT(ElectraPlayer, TEXT("Buffering starts"));
 }
 
@@ -2427,7 +2426,7 @@ void FElectraPlayer::HandlePlayerEventBufferingEnd(Electra::Metrics::EBufferingR
 		AnalyticEvent->ParamArray.Add(FAnalyticsEventAttribute(TEXT("Type"), Electra::Metrics::GetBufferingReasonString(BufferingReason)));
 		EnqueueAnalyticsEvent(AnalyticEvent);
 	}
-	UE_LOG(LogElectraPlayer, Log, TEXT("[%p][%p] %s buffering ended after %.3fs"), this, CurrentPlayer.Get(), Electra::Metrics::GetBufferingReasonString(BufferingReason), BufferingDuration);
+	UE_LOG(LogElectraPlayer, Verbose, TEXT("[%p][%p] %s buffering ended after %.3fs"), this, CurrentPlayer.Get(), Electra::Metrics::GetBufferingReasonString(BufferingReason), BufferingDuration);
 	Statistics.AddMessageToHistory(TEXT("Buffering ended"));
 // Should we set the state (back?) to something or wait for the following play/pause event to set a new one?
 	Statistics.LastState = "Ready";
@@ -2487,7 +2486,7 @@ void FElectraPlayer::HandlePlayerEventSegmentDownload(const Electra::Metrics::FS
 	}
 	if (!SegmentDownloadStats.bWasSuccessful || SegmentDownloadStats.RetryNumber)
 	{
-		UE_LOG(LogElectraPlayer, Log, TEXT("[%p][%p] %s segment download issue (%s): retry:%d, success:%d, aborted:%d, filler:%d"), this, CurrentPlayer.Get(), Electra::Metrics::GetSegmentTypeString(SegmentDownloadStats.SegmentType), *SegmentDownloadStats.FailureReason, SegmentDownloadStats.RetryNumber, SegmentDownloadStats.bWasSuccessful, SegmentDownloadStats.bWasAborted, SegmentDownloadStats.bInsertedFillerData);
+		UE_LOG(LogElectraPlayer, Verbose, TEXT("[%p][%p] %s segment download issue (%s): retry:%d, success:%d, aborted:%d, filler:%d"), this, CurrentPlayer.Get(), Electra::Metrics::GetSegmentTypeString(SegmentDownloadStats.SegmentType), *SegmentDownloadStats.FailureReason, SegmentDownloadStats.RetryNumber, SegmentDownloadStats.bWasSuccessful, SegmentDownloadStats.bWasAborted, SegmentDownloadStats.bInsertedFillerData);
 
 		if (SegmentDownloadStats.FailureReason.Len())
 		{
@@ -2669,7 +2668,7 @@ void FElectraPlayer::HandlePlayerEventPrerollStart()
 	// Update statistics
 	FScopeLock Lock(&StatisticsLock);
 	Statistics.TimeAtPrerollBegin = FPlatformTime::Seconds();
-	UE_LOG(LogElectraPlayer, Log, TEXT("[%p][%p] Player starts prerolling to warm decoders and renderers"), this, CurrentPlayer.Get());
+	UE_LOG(LogElectraPlayer, Verbose, TEXT("[%p][%p] Player starts prerolling to warm decoders and renderers"), this, CurrentPlayer.Get());
 //	Statistics.LastState = "Ready";
 //	FString					CurrentState;	// "Empty", "Opening", "Preparing", "Buffering", "Idle", "Ready", "Playing", "Paused", "Seeking", "Rebuffering", "Ended", "Closed"
 
@@ -2694,7 +2693,7 @@ void FElectraPlayer::HandlePlayerEventPrerollEnd()
 	{
 		Statistics.TimeForInitialPreroll = FPlatformTime::Seconds() - Statistics.TimeAtPrerollBegin;
 	}
-	UE_LOG(LogElectraPlayer, Log, TEXT("[%p][%p] Player prerolling complete"), this, CurrentPlayer.Get());
+	UE_LOG(LogElectraPlayer, Verbose, TEXT("[%p][%p] Player prerolling complete"), this, CurrentPlayer.Get());
 	Statistics.LastState = "Ready";
 
 	// Enqueue a "PrerollEnd" event.
@@ -2862,7 +2861,7 @@ void FElectraPlayer::HandlePlayerEventPlaybackStopped()
 
 void FElectraPlayer::HandlePlayerEventSeekCompleted()
 {
-	UE_LOG(LogElectraPlayer, Log, TEXT("[%p][%p] Seek completed"), this, CurrentPlayer.Get());
+	UE_LOG(LogElectraPlayer, Verbose, TEXT("[%p][%p] Seek completed"), this, CurrentPlayer.Get());
 	bDiscardOutputUntilCleanStart = false;
 	MediaStateOnSeekFinished();
 }
@@ -3037,7 +3036,7 @@ void FElectraPlayer::LogStatistics()
 	TSharedPtr<IElectraPlayerAdapterDelegate, ESPMode::ThreadSafe> PinnedAdapterDelegate = AdapterDelegate.Pin();
 	if (PinnedAdapterDelegate.IsValid())
 	{
-		UE_LOG(LogElectraPlayer, Log, TEXT(
+		UE_LOG(LogElectraPlayer, Verbose, TEXT(
 			"[%p][%p] Electra player statistics:\n"\
 			"OS: %s\n"\
 			"GPU Adapter: %s\n"
@@ -3126,7 +3125,7 @@ void FElectraPlayer::LogStatistics()
 				MessageHistory.Append(FString::Printf(TEXT("%8.3f: %s"), msg.TimeSinceStart, *msg.Message));
 				MessageHistory.Append(TEXT("\n"));
 			}
-			UE_LOG(LogElectraPlayer, Log, TEXT("Most recent log messages:\n%s"), *MessageHistory);
+			UE_LOG(LogElectraPlayer, Verbose, TEXT("Most recent log messages:\n%s"), *MessageHistory);
 		}
 	}
 }
@@ -3145,12 +3144,12 @@ void FElectraPlayer::SendAnalyticMetrics(const TSharedPtr<IAnalyticsProviderET>&
 
 	if (!Statistics.bDidPlaybackEnd)
 	{
-		UE_LOG(LogElectraPlayer, Warning, TEXT("[%p][%p] Submitting analytics during playback, some data may be incomplete"), this, CurrentPlayer.Get());
+		UE_LOG(LogElectraPlayer, Verbose, TEXT("[%p][%p] Submitting analytics during playback, some data may be incomplete"), this, CurrentPlayer.Get());
 		// Try to fill in some of the blanks.
 		UpdatePlayEndStatistics();
 	}
 
-	UE_LOG(LogElectraPlayer, Log, TEXT("[%p][%p] Submitting analytics"), this, CurrentPlayer.Get());
+	UE_LOG(LogElectraPlayer, Verbose, TEXT("[%p][%p] Submitting analytics"), this, CurrentPlayer.Get());
 
 
 	// First emit all enqueued events before sending the final one.
diff --git a/Engine/Plugins/Media/ElectraPlayer/Source/ElectraPlayerRuntime/Private/ElectraPlayerAudioDecoderOutput.cpp b/Engine/Plugins/Media/ElectraPlayer/Source/ElectraPlayerRuntime/Private/ElectraPlayerAudioDecoderOutput.cpp
--- a/Engine/Plugins/Media/ElectraPlayer/Source/ElectraPlayerRuntime/Private/ElectraPlayerAudioDecoderOutput.cpp
+++ b/Engine/Plugins/Media/ElectraPlayer/Source/ElectraPlayerRuntime/Private/ElectraPlayerAudioDecoderOutput.cpp
@@ -36,7 +36,7 @@ public:
 		OwningRenderer = Renderer;
 	}
 
-	void Initialize(ESampleFormat InFormat, uint32 InNumChannels, uint32 InSampleRate, FTimespan InDuration, const FDecoderTimeStamp & InPts, uint32 InBufferSizeBytes) override
+	void Initialize(ESampleFormat InFormat, uint32 InNumChannels, uint32 InSampleRate, FTimespan InDuration, const FDecoderTimeStamp& InPts, uint32 InBufferSizeBytes) override
 	{
 		check(BufferMaxSizeBytes >= InBufferSizeBytes);
 
@@ -110,6 +110,11 @@ public:
 		return SampleRate;
 	}
 
+	void SetTime(const FDecoderTimeStamp& InTime) override
+	{
+		Time = InTime;
+	}
+
 	void ShutdownPoolable() override
 	{
 		TSharedPtr<IDecoderOutputOwner, ESPMode::ThreadSafe> lockedAudioRenderer = OwningRenderer.Pin();
diff --git a/Engine/Plugins/Media/ElectraPlayer/Source/ElectraPlayerRuntime/Private/Runtime/Player/AdaptivePlayerOptionKeynames.h b/Engine/Plugins/Media/ElectraPlayer/Source/ElectraPlayerRuntime/Private/Runtime/Player/AdaptivePlayerOptionKeynames.h
--- a/Engine/Plugins/Media/ElectraPlayer/Source/ElectraPlayerRuntime/Private/Runtime/Player/AdaptivePlayerOptionKeynames.h
+++ b/Engine/Plugins/Media/ElectraPlayer/Source/ElectraPlayerRuntime/Private/Runtime/Player/AdaptivePlayerOptionKeynames.h
@@ -42,6 +42,9 @@ namespace Electra
 	//! (bool) true to emit the first decoded video frame while prerolling so it can be displayed while scrubbing.
 	const FName OptionKeyDoNotHoldBackFirstVideoFrame(TEXT("do_not_hold_back_first_frame"));
 
+	//! (bool) true to always emit decoded samples when the player is paused. Supersedes `OptionKeyDoNotHoldBackFirstVideoFrame`.
+	const FName OptionKeyAlwaysEmitSamplesWhenPaused(TEXT("always_emit_samples_when_paused"));
+
 	//! (bool) true to not truncate the media segment access units at the end of the presentation. Must only be used without a set playback range end!
 	const FName OptionKeyDoNotTruncateAtPresentationEnd(TEXT("do_not_truncate_at_presentation_end"));
 
diff --git a/Engine/Plugins/Media/ElectraPlayer/Source/ElectraPlayerRuntime/Private/Runtime/Player/AdaptiveStreamingPlayer.cpp b/Engine/Plugins/Media/ElectraPlayer/Source/ElectraPlayerRuntime/Private/Runtime/Player/AdaptiveStreamingPlayer.cpp
--- a/Engine/Plugins/Media/ElectraPlayer/Source/ElectraPlayerRuntime/Private/Runtime/Player/AdaptiveStreamingPlayer.cpp
+++ b/Engine/Plugins/Media/ElectraPlayer/Source/ElectraPlayerRuntime/Private/Runtime/Player/AdaptiveStreamingPlayer.cpp
@@ -615,10 +615,6 @@ void FAdaptiveStreamingPlayer::SeekTo(const FSeekParam& NewPosition)
 	FScopeLock lock(&SeekVars.Lock);
 	SeekVars.PendingRequest = NewPosition;
 	SeekVars.PlayrangeOnRequest = NewPlayRange;
-	if (!NewPosition.bIgnoreForSequenceIndex.Get(false))
-	{
-		++SeekVars.NumSeekToCallsSinceLastSeen;
-	}
 	WorkerThread.TriggerSharedWorkerThread();
 }
 
@@ -1103,20 +1099,29 @@ void FAdaptiveStreamingPlayer::StopRendering()
  */
 int32 FAdaptiveStreamingPlayer::CreateRenderers()
 {
+	bool bAlwaysEmitSamplesWhenPaused = PlayerOptions.GetValue(OptionKeyAlwaysEmitSamplesWhenPaused).SafeGetBool(false);
+	double CurrentRate = PlaybackState.GetCurrentPlayRate();
+	double IntendedRate = PlaybackState.GetDesiredPlayRate();
+	const bool bCurrentlyPaused = PlaybackState.GetIsPaused();
+
 	// Set the render clock with the renderes.
 	if (VideoRender.Renderer)
 	{
 		VideoRender.Renderer->SetRenderClock(RenderClock);
+		VideoRender.Renderer->SetPlaybackRate(CurrentRate, IntendedRate, bCurrentlyPaused);
+		VideoRender.Renderer->AlwaysEmitSamplesWhenPaused(bAlwaysEmitSamplesWhenPaused);
+		// Hold back all video frames during preroll or emit the first frame for scrubbing?
+		if (PlayerOptions.HaveKey(OptionKeyDoNotHoldBackFirstVideoFrame))
+		{
+			bool bDoNotHoldBackFirstVideoFrame = PlayerOptions.GetValue(OptionKeyDoNotHoldBackFirstVideoFrame).SafeGetBool(false);
+			VideoRender.Renderer->DisableHoldbackOfFirstRenderableVideoFrame(bDoNotHoldBackFirstVideoFrame);
+		}
 	}
 	if (AudioRender.Renderer)
 	{
 		AudioRender.Renderer->SetRenderClock(RenderClock);
-	}
-
-	// Hold back all frames during preroll or emit the first frame for scrubbing?
-	if (VideoRender.Renderer && PlayerOptions.HaveKey(OptionKeyDoNotHoldBackFirstVideoFrame))
-	{
-		VideoRender.Renderer->DisableHoldbackOfFirstRenderableVideoFrame(PlayerOptions.GetValue(OptionKeyDoNotHoldBackFirstVideoFrame).SafeGetBool(false));
+		AudioRender.Renderer->SetPlaybackRate(CurrentRate, IntendedRate, bCurrentlyPaused);
+		AudioRender.Renderer->AlwaysEmitSamplesWhenPaused(bAlwaysEmitSamplesWhenPaused);
 	}
 
 	return 0;
@@ -1787,6 +1792,19 @@ void FAdaptiveStreamingPlayer::HandlePlayStateChanges()
 
 	// Update the current live latency.
 	PlaybackState.SetCurrentLiveLatency(CalculateCurrentLiveLatency(false));
+
+	// Inform the renderers about the current playback rate.
+	const double CurrentRate = PlaybackState.GetCurrentPlayRate();
+	const double IntendedRate = PlaybackState.GetDesiredPlayRate();
+	const bool bCurrentlyPaused = PlaybackState.GetIsPaused();
+	if (VideoRender.Renderer)
+	{
+		VideoRender.Renderer->SetPlaybackRate(CurrentRate, IntendedRate, bCurrentlyPaused);
+	}
+	if (AudioRender.Renderer)
+	{
+		AudioRender.Renderer->SetPlaybackRate(CurrentRate, IntendedRate, bCurrentlyPaused);
+	}
 }
 
 
@@ -1825,9 +1843,11 @@ void FAdaptiveStreamingPlayer::HandleSeeking()
 
 		// Adjust seek index as indicated
 		// (we even do this for seeks marked as not affecting the seek index as we must still add any priovious accumulated calls that had not been flagged in such a way)
-		CurrentPlaybackSequenceState.PrimaryIndex += SeekVars.NumSeekToCallsSinceLastSeen;
+		if (SeekVars.PendingRequest.GetValue().NewSequenceIndex.IsSet())
+		{
+			CurrentPlaybackSequenceState.PrimaryIndex = SeekVars.PendingRequest.GetValue().NewSequenceIndex.GetValue();
+		}
 		CurrentPlaybackSequenceState.SecondaryIndex = 0;
-		SeekVars.NumSeekToCallsSinceLastSeen = 0;
 
 		// And since it is a seek on purpose the loop counter is reset as well.
 		FInternalLoopState LoopStateNow;
@@ -4107,31 +4127,38 @@ void FAdaptiveStreamingPlayer::UpdateStreamResolutionLimit()
  */
 void FAdaptiveStreamingPlayer::CheckForStreamEnd()
 {
-	if (CurrentState == EPlayerState::eState_Playing)
+	if (CurrentState == EPlayerState::eState_Playing || CurrentState == EPlayerState::eState_Paused)
 	{
 		if (StreamState == EStreamState::eStream_Running)
 		{
 			// First check if there is an end time set at which we need to stop.
-			FTimeValue EndAtTime = PlaybackState.GetPlaybackEndAtTime();
-			FTimeValue InitialEndTime = Manifest.IsValid() ? Manifest->GetDefaultEndTime() : FTimeValue();
-			if (EndAtTime.IsValid() && InitialEndTime.IsValid())
-			{
-				EndAtTime = EndAtTime < InitialEndTime ? EndAtTime : InitialEndTime;
-			}
-			else if (InitialEndTime.IsValid())
-			{
-				EndAtTime = InitialEndTime;
-			}
-			if (EndAtTime.IsValid())
+			// Do this only when playing, not in paused state.
+			// We run through this code here when paused to ensure that if looping is
+			// enabled we switch over to the new buffer set to keep the decoders busy.
+			// This is intended for scrubbing across the loop point.
+			if (CurrentState == EPlayerState::eState_Playing)
 			{
-				if (PlaybackState.GetPlayPosition() >= EndAtTime)
+				FTimeValue EndAtTime = PlaybackState.GetPlaybackEndAtTime();
+				FTimeValue InitialEndTime = Manifest.IsValid() ? Manifest->GetDefaultEndTime() : FTimeValue();
+				if (EndAtTime.IsValid() && InitialEndTime.IsValid())
 				{
-					// A forced end time is intended to stop playback. We do NOT look at the loop state here.
-					InternalSetPlaybackEnded();
-					DataBuffersCriticalSection.Lock();
-					NextDataBuffers.Empty();
-					DataBuffersCriticalSection.Unlock();
-					return;
+					EndAtTime = EndAtTime < InitialEndTime ? EndAtTime : InitialEndTime;
+				}
+				else if (InitialEndTime.IsValid())
+				{
+					EndAtTime = InitialEndTime;
+				}
+				if (EndAtTime.IsValid())
+				{
+					if (PlaybackState.GetPlayPosition() >= EndAtTime)
+					{
+						// A forced end time is intended to stop playback. We do NOT look at the loop state here.
+						InternalSetPlaybackEnded();
+						DataBuffersCriticalSection.Lock();
+						NextDataBuffers.Empty();
+						DataBuffersCriticalSection.Unlock();
+						return;
+					}
 				}
 			}
 
diff --git a/Engine/Plugins/Media/ElectraPlayer/Source/ElectraPlayerRuntime/Private/Runtime/Player/AdaptiveStreamingPlayer.h b/Engine/Plugins/Media/ElectraPlayer/Source/ElectraPlayerRuntime/Private/Runtime/Player/AdaptiveStreamingPlayer.h
--- a/Engine/Plugins/Media/ElectraPlayer/Source/ElectraPlayerRuntime/Private/Runtime/Player/AdaptiveStreamingPlayer.h
+++ b/Engine/Plugins/Media/ElectraPlayer/Source/ElectraPlayerRuntime/Private/Runtime/Player/AdaptiveStreamingPlayer.h
@@ -135,14 +135,14 @@ public:
 		}
 		// Time to seek to.
 		FTimeValue	Time;
+		// New sequence index to associate with the newly decoded samples.
+		TOptional<int32> NewSequenceIndex;
 		// Maximum stream bitrate to use when seeking.
 		TOptional<int32> StartingBitrate;
 		// Optimize for frame scrubbing (faster display of frame at target time)?
 		TOptional<bool> bOptimizeForScrubbing;
 		// Allowed distance to last performed seek to save a redundant new seek.
 		TOptional<double> DistanceThreshold;
-		// Ignore the seek for sequence index updates
-		TOptional<bool> bIgnoreForSequenceIndex;
 	};
 
 	/**
diff --git a/Engine/Plugins/Media/ElectraPlayer/Source/ElectraPlayerRuntime/Private/Runtime/Player/AdaptiveStreamingPlayerInternal.h b/Engine/Plugins/Media/ElectraPlayer/Source/ElectraPlayerRuntime/Private/Runtime/Player/AdaptiveStreamingPlayerInternal.h
--- a/Engine/Plugins/Media/ElectraPlayer/Source/ElectraPlayerRuntime/Private/Runtime/Player/AdaptiveStreamingPlayerInternal.h
+++ b/Engine/Plugins/Media/ElectraPlayer/Source/ElectraPlayerRuntime/Private/Runtime/Player/AdaptiveStreamingPlayerInternal.h
@@ -85,6 +85,8 @@ public:
 	virtual FTimeValue GetEnqueuedSampleDuration() = 0;
 	virtual int32 GetNumEnqueuedSamples(TArray<FEnqueuedSampleInfo>* OutOptionalSampleInfos) = 0;
 
+	virtual void AlwaysEmitSamplesWhenPaused(bool bEmitAlways) = 0;
+	virtual void SetPlaybackRate(double InCurrentPlaybackRate, double InIntendedPlaybackRate, bool bInCurrentlyPaused) = 0;
 	virtual void DisableHoldbackOfFirstRenderableVideoFrame(bool bDisableHoldback) = 0;
 
 	virtual FTimeRange GetSupportedRenderRateScale() = 0;
@@ -1926,7 +1928,6 @@ private:
 		bool	bForScrubbing = false;
 		bool	bScrubPrerollDone = false;
 		bool	bIsPlayStart = true;
-		int32	NumSeekToCallsSinceLastSeen = 0;
 
 		TOptional<FSeekParam> ActiveRequest;
 		TOptional<FSeekParam> LastFinishedRequest;
diff --git a/Engine/Plugins/Media/ElectraPlayer/Source/ElectraPlayerRuntime/Private/Runtime/Player/AdaptiveStreamingPlayerRenderers.cpp b/Engine/Plugins/Media/ElectraPlayer/Source/ElectraPlayerRuntime/Private/Runtime/Player/AdaptiveStreamingPlayerRenderers.cpp
--- a/Engine/Plugins/Media/ElectraPlayer/Source/ElectraPlayerRuntime/Private/Runtime/Player/AdaptiveStreamingPlayerRenderers.cpp
+++ b/Engine/Plugins/Media/ElectraPlayer/Source/ElectraPlayerRuntime/Private/Runtime/Player/AdaptiveStreamingPlayerRenderers.cpp
@@ -95,6 +95,8 @@ namespace Electra
 		FTimeValue GetEnqueuedSampleDuration() override;
 		int32 GetNumEnqueuedSamples(TArray<FEnqueuedSampleInfo>* OutOptionalSampleInfos) override;
 
+		void AlwaysEmitSamplesWhenPaused(bool bEmitAlways) override;
+		void SetPlaybackRate(double InCurrentPlaybackRate, double InIntendedPlaybackRate, bool bInCurrentlyPaused) override;
 		void DisableHoldbackOfFirstRenderableVideoFrame(bool bDisableHoldback) override;
 
 		FTimeRange GetSupportedRenderRateScale() override;
@@ -131,6 +133,9 @@ namespace Electra
 		int64 CurrentValidityValue = 0;
 		EStreamType Type = EStreamType::Unsupported;
 		bool bIsRunning = false;
+		double CurrentPlaybackRate = 0.0;
+		double IntendedPlaybackRate = 0.0;
+		bool bAlwaysEmitSamplesWhenPaused = false;
 		bool bDoNotHoldBackFirstVideoFrame = true;
 		uint32 NumBuffersNotHeldBack = 0;
 
@@ -455,9 +460,17 @@ UEMediaError FAdaptiveStreamingWrappedRenderer::ReturnBufferCommon(IBuffer* Buff
 	if (!bIsUnusedReturnBuffer)
 	{
 		bool bHoldback = !bIsRunning;
+
+		// Never hold back when the player is paused?
+		// The main player state as set by the user through API calls, not the current actual rate which
+		// may be different during prerolling and buffering!
+		if (bAlwaysEmitSamplesWhenPaused && IntendedPlaybackRate == 0.0)
+		{
+			bHoldback = false;
+		}
 		// If the video renderer shall not hold back the first frame (used for scrubbing video)
 		// then we pass it out. The count is reset in Flush().
-		if (Type == EStreamType::Video && bDoNotHoldBackFirstVideoFrame)
+		else if (Type == EStreamType::Video && bDoNotHoldBackFirstVideoFrame)
 		{
 			if (NumBuffersNotHeldBack == 0)
 			{
@@ -615,6 +628,19 @@ int32 FAdaptiveStreamingWrappedRenderer::GetNumEnqueuedSamples(TArray<IAdaptiveS
 	return EnqueuedSamples.Num() + NumAvail;
 }
 
+void FAdaptiveStreamingWrappedRenderer::AlwaysEmitSamplesWhenPaused(bool bEmitAlways)
+{
+	FScopeLock lock(&Lock);
+	bAlwaysEmitSamplesWhenPaused = bEmitAlways;
+}
+
+void FAdaptiveStreamingWrappedRenderer::SetPlaybackRate(double InCurrentPlaybackRate, double InIntendedPlaybackRate, bool bInCurrentlyPaused)
+{
+	FScopeLock lock(&Lock);
+	CurrentPlaybackRate = bInCurrentlyPaused ? 0.0 : InCurrentPlaybackRate;
+	IntendedPlaybackRate = bInCurrentlyPaused ? 0.0 : InIntendedPlaybackRate;
+}
+
 void FAdaptiveStreamingWrappedRenderer::DisableHoldbackOfFirstRenderableVideoFrame(bool bInDisableHoldback)
 {
 	FScopeLock lock(&Lock);
diff --git a/Engine/Plugins/Media/ElectraPlayer/Source/ElectraPlayerRuntime/Public/IElectraPlayerInterface.h b/Engine/Plugins/Media/ElectraPlayer/Source/ElectraPlayerRuntime/Public/IElectraPlayerInterface.h
--- a/Engine/Plugins/Media/ElectraPlayer/Source/ElectraPlayerRuntime/Public/IElectraPlayerInterface.h
+++ b/Engine/Plugins/Media/ElectraPlayer/Source/ElectraPlayerRuntime/Public/IElectraPlayerInterface.h
@@ -30,7 +30,7 @@ DECLARE_MULTICAST_DELEGATE_FourParams(FElectraPlayerReportSubtitlesMetricsDelega
 
 // ---------------------------------------------------------------------------------------------
 
-class IElectraPlayerAdapterDelegate : public TSharedFromThis<IElectraPlayerAdapterDelegate, ESPMode::ThreadSafe>
+class IElectraPlayerAdapterDelegate
 {
 public:
 	virtual ~IElectraPlayerAdapterDelegate() {}
@@ -263,6 +263,7 @@ public:
 
 	struct FSeekParam
 	{
+		TOptional<int32> SequenceIndex;
 		TOptional<int32> StartingBitrate;
 		TOptional<bool> bOptimizeForScrubbing;
 		TOptional<double> DistanceThreshold;
@@ -336,6 +337,7 @@ public:
 			FTimeValue Start;
 			FTimeValue End;
 		};
+		TArray<FTimeRange> TimeEnqueued;
 		TArray<FTimeRange> TimeAvailable;
 		TArray<FTimeRange> TimeRequested;
 	};
diff --git a/Engine/Plugins/Media/ElectraSubtitles/Source/ElectraSubtitles/Private/ttml/ElectraSubtitleDecoder_TTML.cpp b/Engine/Plugins/Media/ElectraSubtitles/Source/ElectraSubtitles/Private/ttml/ElectraSubtitleDecoder_TTML.cpp
--- a/Engine/Plugins/Media/ElectraSubtitles/Source/ElectraSubtitles/Private/ttml/ElectraSubtitleDecoder_TTML.cpp
+++ b/Engine/Plugins/Media/ElectraSubtitles/Source/ElectraSubtitles/Private/ttml/ElectraSubtitleDecoder_TTML.cpp
@@ -18,14 +18,14 @@ public:
 void FElectraSubtitleDecoderTTML::RegisterCodecs(IElectraSubtitleDecoderFactoryRegistry& InRegistry)
 {
 	static FElectraSubtitleDecoderFactoryTTML Factory;
-	TArray<IElectraSubtitleDecoderFactoryRegistry::FCodecInfo> CodecInfos 
-	{ 
-		{ FString(TEXT("stpp")), 0}, 
+	TArray<IElectraSubtitleDecoderFactoryRegistry::FCodecInfo> CodecInfos
+	{
+		{ FString(TEXT("stpp")), 0},
 		{ FString(TEXT("application/ttml+xml")), 0},
 	// See: https://www.w3.org/TR/ttml-profile-registry/#registry-profile-designator-specifications
-		{ FString(TEXT("stpp.ttml.im1t")), 0}, 
-		{ FString(TEXT("stpp.ttml.im2t")), 0}, 
-		{ FString(TEXT("stpp.ttml.im3t")), 0} 
+		{ FString(TEXT("stpp.ttml.im1t")), 0},
+		{ FString(TEXT("stpp.ttml.im2t")), 0},
+		{ FString(TEXT("stpp.ttml.im3t")), 0}
 	};
 	InRegistry.AddDecoderFactory(CodecInfos, &Factory);
 }
@@ -52,45 +52,50 @@ public:
 		TextAsArray.Empty();
 		TextAsArray.Append(reinterpret_cast<const uint8*>(Converted.Get()), Converted.Length());
 	}
-	
+
 	void SetDuration(const Electra::FTimeValue& InDuration)
 	{
 		Duration = InDuration.GetAsTimespan();
 	}
-	
+
 	void SetTimestamp(const Electra::FTimeValue& InTimestamp)
 	{
 		Timestamp.Time = InTimestamp.GetAsTimespan();
 		Timestamp.SequenceIndex = 0;
 	}
-	
+
 	void SetID(const FString& InID)
 	{
 		ID = InID;
 	}
 
 
-	virtual const TArray<uint8>& GetData() override
+	const TArray<uint8>& GetData() override
 	{
 		return TextAsArray;
 	}
-	
-	virtual FDecoderTimeStamp GetTime() const override
+
+	FDecoderTimeStamp GetTime() const override
 	{
 		return Timestamp;
 	}
 
-	virtual FTimespan GetDuration() const override
+	void SetTime(FDecoderTimeStamp& InTime) override
+	{
+		Timestamp = InTime;
+	}
+
+	FTimespan GetDuration() const override
 	{
 		return Duration;
 	}
 
-	virtual const FString& GetFormat() const override
+	const FString& GetFormat() const override
 	{
 		static FString Format(TEXT("stpp"));
 		return Format;
 	}
-	virtual const FString& GetID() const override
+	const FString& GetID() const override
 	{
 		return ID;
 	}
@@ -171,7 +176,7 @@ void FElectraSubtitleDecoderTTML::AddStreamedSubtitleData(const TArray<uint8>& I
 		AccessLock.Lock();
 		TimelineHandlers.Empty();
 		AccessLock.Unlock();
-		
+
 		// For sideloaded data we try to get a unique ID we can use to check if we are still operating on the
 		// same data. If we are we do not need to parse the entire document again.
 		SideloadedID = InAdditionalInfo.GetValue(ElectraSubtitleDecoderTTMLOptions::SideloadedID).SafeGetFString();
diff --git a/Engine/Plugins/Media/ElectraSubtitles/Source/ElectraSubtitles/Private/tx3g/ElectraSubtitleDecoder_TX3G.cpp b/Engine/Plugins/Media/ElectraSubtitles/Source/ElectraSubtitles/Private/tx3g/ElectraSubtitleDecoder_TX3G.cpp
--- a/Engine/Plugins/Media/ElectraSubtitles/Source/ElectraSubtitles/Private/tx3g/ElectraSubtitleDecoder_TX3G.cpp
+++ b/Engine/Plugins/Media/ElectraSubtitles/Source/ElectraSubtitles/Private/tx3g/ElectraSubtitleDecoder_TX3G.cpp
@@ -180,45 +180,50 @@ public:
 		TextAsArray.Empty();
 		TextAsArray.Append(reinterpret_cast<const uint8*>(Converted.Get()), Converted.Length());
 	}
-	
+
 	void SetDuration(const Electra::FTimeValue& InDuration)
 	{
 		Duration = InDuration.GetAsTimespan();
 	}
-	
+
 	void SetTimestamp(const Electra::FTimeValue& InTimestamp)
 	{
 		Timestamp.Time = InTimestamp.GetAsTimespan();
 		Timestamp.SequenceIndex = 0;
 	}
-	
+
 	void SetID(const FString& InID)
 	{
 		ID = InID;
 	}
 
 
-	virtual const TArray<uint8>& GetData() override
+	const TArray<uint8>& GetData() override
 	{
 		return TextAsArray;
 	}
-	
-	virtual FDecoderTimeStamp GetTime() const override
+
+	FDecoderTimeStamp GetTime() const override
 	{
 		return Timestamp;
 	}
 
-	virtual FTimespan GetDuration() const override
+	void SetTime(FDecoderTimeStamp& InTime) override
+	{
+		Timestamp = InTime;
+	}
+
+	FTimespan GetDuration() const override
 	{
 		return Duration;
 	}
 
-	virtual const FString& GetFormat() const override
+	const FString& GetFormat() const override
 	{
 		static FString Format(TEXT("tx3g"));
 		return Format;
 	}
-	virtual const FString& GetID() const override
+	const FString& GetID() const override
 	{
 		return ID;
 	}
@@ -335,7 +340,7 @@ void FElectraSubtitleDecoderTX3G::AddStreamedSubtitleData(const TArray<uint8>& I
 			UE_LOG(LogElectraSubtitles, Error, TEXT("Bad TX3G text sample, ignoring."));	\
 			return;																			\
 		}																					\
-	
+
 	ElectraSubtitleDecoderTX3GUtils::FDataReaderMP4 r(InData);
 
 	// Start a subtitle entry.
diff --git a/Engine/Plugins/Media/ElectraSubtitles/Source/ElectraSubtitles/Private/wvtt/ElectraSubtitleDecoder_WVTT.cpp b/Engine/Plugins/Media/ElectraSubtitles/Source/ElectraSubtitles/Private/wvtt/ElectraSubtitleDecoder_WVTT.cpp
--- a/Engine/Plugins/Media/ElectraSubtitles/Source/ElectraSubtitles/Private/wvtt/ElectraSubtitleDecoder_WVTT.cpp
+++ b/Engine/Plugins/Media/ElectraSubtitles/Source/ElectraSubtitles/Private/wvtt/ElectraSubtitleDecoder_WVTT.cpp
@@ -181,45 +181,50 @@ public:
 		TextAsArray.Empty();
 		TextAsArray.Append(reinterpret_cast<const uint8*>(Converted.Get()), Converted.Length());
 	}
-	
+
 	void SetDuration(const Electra::FTimeValue& InDuration)
 	{
 		Duration = InDuration.GetAsTimespan();
 	}
-	
+
 	void SetTimestamp(const Electra::FTimeValue& InTimestamp)
 	{
 		Timestamp.Time = InTimestamp.GetAsTimespan();
 		Timestamp.SequenceIndex = 0;
 	}
-	
+
 	void SetID(const FString& InID)
 	{
 		ID = InID;
 	}
 
 
-	virtual const TArray<uint8>& GetData() override
+	const TArray<uint8>& GetData() override
 	{
 		return TextAsArray;
 	}
-	
-	virtual FDecoderTimeStamp GetTime() const override
+
+	FDecoderTimeStamp GetTime() const override
 	{
 		return Timestamp;
 	}
 
-	virtual FTimespan GetDuration() const override
+	void SetTime(FDecoderTimeStamp& InTime) override
+	{
+		Timestamp = InTime;
+	}
+
+	FTimespan GetDuration() const override
 	{
 		return Duration;
 	}
 
-	virtual const FString& GetFormat() const override
+	const FString& GetFormat() const override
 	{
 		static FString Format(TEXT("wvtt"));
 		return Format;
 	}
-	virtual const FString& GetID() const override
+	const FString& GetID() const override
 	{
 		return ID;
 	}
@@ -280,7 +285,7 @@ bool FElectraSubtitleDecoderWVTT::InitializeStreamWithCSD(const TArray<uint8>& I
 			UE_LOG(LogElectraSubtitles, Error, TEXT("Bad WVTT box in CSD, ignoring."));
 			break;
 		}
-		
+
 		if (BoxType == ElectraSubtitleDecoderWVTTUtils::FDataReaderMP4::BoxType_vttC)
 		{
 			RETURN_IF_ERROR(r.ReadString(Configuration, BoxLen - 8));
@@ -328,7 +333,7 @@ void FElectraSubtitleDecoderWVTT::AddStreamedSubtitleData(const TArray<uint8>& I
 			UE_LOG(LogElectraSubtitles, Error, TEXT("Bad WVTT text sample box, ignoring."));	\
 			return;																				\
 		}																						\
-	
+
 	ElectraSubtitleDecoderWVTTUtils::FDataReaderMP4 r(InData);
 
 	// List of collected subtitles.
diff --git a/Engine/Plugins/Media/ElectraUtil/Source/ElectraBase/Public/MediaAudioDecoderOutput.h b/Engine/Plugins/Media/ElectraUtil/Source/ElectraBase/Public/MediaAudioDecoderOutput.h
--- a/Engine/Plugins/Media/ElectraUtil/Source/ElectraBase/Public/MediaAudioDecoderOutput.h
+++ b/Engine/Plugins/Media/ElectraUtil/Source/ElectraBase/Public/MediaAudioDecoderOutput.h
@@ -34,6 +34,8 @@ public:
 	virtual FTimespan GetDuration() const = 0;
 
 	virtual uint32 GetSampleRate() const = 0;
+
+	virtual void SetTime(const FDecoderTimeStamp& InTime) = 0;
 };
 
 using IAudioDecoderOutputPtr = TSharedPtr<IAudioDecoderOutput, ESPMode::ThreadSafe>;
diff --git a/Engine/Plugins/Media/ElectraUtil/Source/ElectraBase/Public/MediaMetaDataDecoderOutput.h b/Engine/Plugins/Media/ElectraUtil/Source/ElectraBase/Public/MediaMetaDataDecoderOutput.h
--- a/Engine/Plugins/Media/ElectraUtil/Source/ElectraBase/Public/MediaMetaDataDecoderOutput.h
+++ b/Engine/Plugins/Media/ElectraUtil/Source/ElectraBase/Public/MediaMetaDataDecoderOutput.h
@@ -48,6 +48,8 @@ public:
 	// Returns the base time (the "zero point") of the track, if available, to calculate the relative time
 	// of this metadata from GetTime() if this is based on some non-zero anchor.
 	virtual TOptional<FDecoderTimeStamp> GetTrackBaseTime() const = 0;
+
+	virtual void SetTime(FDecoderTimeStamp& InTime) = 0;
 };
 
 using IMetaDataDecoderOutputPtr = TSharedPtr<IMetaDataDecoderOutput, ESPMode::ThreadSafe>;
diff --git a/Engine/Plugins/Media/ElectraUtil/Source/ElectraBase/Public/MediaSubtitleDecoderOutput.h b/Engine/Plugins/Media/ElectraUtil/Source/ElectraBase/Public/MediaSubtitleDecoderOutput.h
--- a/Engine/Plugins/Media/ElectraUtil/Source/ElectraBase/Public/MediaSubtitleDecoderOutput.h
+++ b/Engine/Plugins/Media/ElectraUtil/Source/ElectraBase/Public/MediaSubtitleDecoderOutput.h
@@ -18,6 +18,8 @@ public:
 
 	virtual const FString& GetFormat() const = 0;
 	virtual const FString& GetID() const = 0;
+
+	virtual void SetTime(FDecoderTimeStamp& InTime) = 0;
 };
 
 using ISubtitleDecoderOutputPtr = TSharedPtr<ISubtitleDecoderOutput, ESPMode::ThreadSafe>;
diff --git a/Engine/Plugins/Media/ElectraUtil/Source/ElectraBase/Public/MediaVideoDecoderOutput.h b/Engine/Plugins/Media/ElectraUtil/Source/ElectraBase/Public/MediaVideoDecoderOutput.h
--- a/Engine/Plugins/Media/ElectraUtil/Source/ElectraBase/Public/MediaVideoDecoderOutput.h
+++ b/Engine/Plugins/Media/ElectraUtil/Source/ElectraBase/Public/MediaVideoDecoderOutput.h
@@ -126,9 +126,9 @@ public:
 	virtual ~FVideoDecoderOutput() = default;
 
 	virtual const Electra::FParamDict& GetDict() const
-	{ 
+	{
 		check(ParamDict.IsValid());
-		return *ParamDict; 
+		return *ParamDict;
 	}
 
 	virtual FDecoderTimeStamp GetTime() const
@@ -137,8 +137,21 @@ public:
 		{
 			return FDecoderTimeStamp(FTimespan::Zero(), 0);
 		}
-		Electra::FTimeValue pts(ParamDict->GetValue(IDecoderOutputOptionNames::PTS).GetTimeValue());
-		return FDecoderTimeStamp(pts.GetAsTimespan(), pts.GetSequenceIndex());
+		if (!(Cached.Flags & FCached::Valid_PTS))
+		{
+			Electra::FTimeValue pts(ParamDict->GetValue(IDecoderOutputOptionNames::PTS).GetTimeValue());
+			Cached.PTS = FDecoderTimeStamp(pts.GetAsTimespan(), pts.GetSequenceIndex());
+			FPlatformMisc::MemoryBarrier();
+			Cached.Flags |= FCached::Valid_PTS;
+		}
+		return Cached.PTS;
+	}
+
+	virtual void SetTime(const FDecoderTimeStamp& InTime)
+	{
+		Cached.PTS = FDecoderTimeStamp(InTime);
+		FPlatformMisc::MemoryBarrier();
+		Cached.Flags |= FCached::Valid_PTS;
 	}
 
 	virtual FTimespan GetDuration() const
@@ -147,7 +160,13 @@ public:
 		{
 			return FTimespan(0);
 		}
-		return FTimespan(ParamDict->GetValue(IDecoderOutputOptionNames::Duration).GetTimeValue().GetAsHNS());
+		if (!(Cached.Flags & FCached::Valid_Duration))
+		{
+			Cached.Duration = ParamDict->GetValue(IDecoderOutputOptionNames::Duration).GetTimeValue().GetAsTimespan();
+			FPlatformMisc::MemoryBarrier();
+			Cached.Flags |= FCached::Valid_Duration;
+		}
+		return Cached.Duration;
 	}
 
 	virtual FIntPoint GetOutputDim() const
@@ -183,7 +202,7 @@ public:
 			return FVideoDecoderCropInfo();
 		}
 		if (!(Cached.Flags & FCached::Valid_CropInfo))
-		{ 
+		{
 			Cached.CropInfo.CropLeft = (int32)ParamDict->GetValue(IDecoderOutputOptionNames::CropLeft).SafeGetInt64(0);
 			Cached.CropInfo.CropTop = (int32)ParamDict->GetValue(IDecoderOutputOptionNames::CropTop).SafeGetInt64(0);
 			Cached.CropInfo.CropRight = (int32)ParamDict->GetValue(IDecoderOutputOptionNames::CropRight).SafeGetInt64(0);
@@ -277,6 +296,11 @@ protected:
 	{
 		Cached.Flags = 0;
 		ParamDict = MoveTemp(InParamDict);
+#if !UE_BUILD_SHIPPING
+		// For debugging purposes, get the PTS and duration so they're easily accessible.
+		(void)GetTime();
+		(void)GetDuration();
+#endif
 	}
 
 private:
@@ -291,9 +315,13 @@ private:
 			Valid_CropInfo = 1 << 0,
 			Valid_OutputDim = 1 << 1,
 			Valid_Orientation = 1 << 2,
+			Valid_PTS = 1 << 3,
+			Valid_Duration = 1 << 4
 		};
 
 		FVideoDecoderCropInfo CropInfo;
+		FDecoderTimeStamp PTS;
+		FTimespan Duration;
 		FIntPoint OutputDim;
 		EVideoOrientation Orientation = EVideoOrientation::Original;;
 		uint32 Flags = 0;
diff --git a/Engine/Plugins/Media/ElectraUtil/Source/ElectraSamples/Private/ElectraTextureSampleBase.cpp b/Engine/Plugins/Media/ElectraUtil/Source/ElectraSamples/Private/ElectraTextureSampleBase.cpp
--- a/Engine/Plugins/Media/ElectraUtil/Source/ElectraSamples/Private/ElectraTextureSampleBase.cpp
+++ b/Engine/Plugins/Media/ElectraUtil/Source/ElectraSamples/Private/ElectraTextureSampleBase.cpp
@@ -142,7 +142,7 @@ void IElectraTextureSampleBase::Initialize(FVideoDecoderOutput* InVideoDecoderOu
 
 	// Compute scale to make correct towards the max value (P010 will max out at 0xffc0 not 0xffff - so if it is present we need to adjust the scale a bit)
 	float NormScale = (VideoDecoderOutput->GetFormat() == PF_P010) ? (65535.0f / 65472.0f) : 1.0f;
- 
+
 	// Matrix to transform sample data to standard YUV values
 	FMatrix PreMtx = FMatrix::Identity;
 	PreMtx.M[0][0] = DataScale * NormScale;
@@ -198,6 +198,14 @@ FMediaTimeStamp IElectraTextureSampleBase::GetTime() const
 	return FMediaTimeStamp();
 }
 
+void IElectraTextureSampleBase::SetTime(const FMediaTimeStamp& InTime)
+{
+	if (VideoDecoderOutput)
+	{
+		VideoDecoderOutput->SetTime(FDecoderTimeStamp(InTime.Time, InTime.SequenceIndex));
+	}
+}
+
 
 FTimespan IElectraTextureSampleBase::GetDuration() const
 {
diff --git a/Engine/Plugins/Media/ElectraUtil/Source/ElectraSamples/Public/IElectraTextureSample.h b/Engine/Plugins/Media/ElectraUtil/Source/ElectraSamples/Public/IElectraTextureSample.h
--- a/Engine/Plugins/Media/ElectraUtil/Source/ElectraSamples/Public/IElectraTextureSample.h
+++ b/Engine/Plugins/Media/ElectraUtil/Source/ElectraSamples/Public/IElectraTextureSample.h
@@ -57,6 +57,8 @@ public:
 	virtual bool GetDisplayMasteringLuminance(float& OutMin, float& OutMax) const override;
 	virtual bool GetMaxLuminanceLevels(uint16& OutCLL, uint16& OutFALL) const override;
 
+	virtual void SetTime(const FMediaTimeStamp& InTime);
+
 protected:
 	virtual float GetSampleDataScale(bool b10Bit) const { return 1.0f; }
 
diff --git a/Engine/Plugins/Media/ElectraUtil/Source/ElectraSamples/Public/Linux/ElectraTextureSample.h b/Engine/Plugins/Media/ElectraUtil/Source/ElectraSamples/Public/Linux/ElectraTextureSample.h
--- a/Engine/Plugins/Media/ElectraUtil/Source/ElectraSamples/Public/Linux/ElectraTextureSample.h
+++ b/Engine/Plugins/Media/ElectraUtil/Source/ElectraSamples/Public/Linux/ElectraTextureSample.h
@@ -56,6 +56,7 @@ private:
 
 using FElectraTextureSamplePtr = TSharedPtr<FElectraTextureSampleLinux, ESPMode::ThreadSafe>;
 using FElectraTextureSampleRef = TSharedRef<FElectraTextureSampleLinux, ESPMode::ThreadSafe>;
+using FElectraTextureSample = FElectraTextureSampleLinux;
 
 class FElectraTextureSamplePool : public TMediaObjectPool<FElectraTextureSampleLinux>
 {
diff --git a/Engine/Source/Runtime/Media/Public/IMediaSamples.h b/Engine/Source/Runtime/Media/Public/IMediaSamples.h
--- a/Engine/Source/Runtime/Media/Public/IMediaSamples.h
+++ b/Engine/Source/Runtime/Media/Public/IMediaSamples.h
@@ -128,7 +128,9 @@ public:
 	}
 
 	virtual bool PeekVideoSampleTime(FMediaTimeStamp& TimeStamp) = 0;
-	virtual bool PeekVideoSampleTimeRange(TRange<FMediaTimeStamp>& TimeRange) { return false; }
+
+	virtual bool PeekVideoSampleTimeRanges(TArray<TRange<FMediaTimeStamp>>& TimeRange) { return false; }
+	virtual bool PeekAudioSampleTimeRanges(TArray<TRange<FMediaTimeStamp>>& TimeRange) { return false; }
 
 	virtual bool DiscardVideoSamples(const TRange<FMediaTimeStamp>& TimeRange, bool bReverse) { return false; }
 	virtual bool DiscardAudioSamples(const TRange<FMediaTimeStamp>& TimeRange, bool bReverse) { return false; }
diff --git a/Engine/Source/Runtime/Media/Public/IMediaTimeSource.h b/Engine/Source/Runtime/Media/Public/IMediaTimeSource.h
--- a/Engine/Source/Runtime/Media/Public/IMediaTimeSource.h
+++ b/Engine/Source/Runtime/Media/Public/IMediaTimeSource.h
@@ -42,6 +42,41 @@ public:
 	bool operator > (const FMediaTimeStamp & Other) const { return (SequenceIndex > Other.SequenceIndex) || (SequenceIndex == Other.SequenceIndex && Time > Other.Time); }
 	bool operator >= (const FMediaTimeStamp & Other) const { return (SequenceIndex > Other.SequenceIndex) || (SequenceIndex == Other.SequenceIndex && Time >= Other.Time); }
 
+	FMediaTimeStamp& SetPrimaryIndex(int32 InSetTo)
+	{
+		SequenceIndex = MakeSequenceIndex(InSetTo, GetSecondaryIndex());
+		return *this;
+	}
+
+	FMediaTimeStamp& SetSecondaryIndex(int32 InSetTo)
+	{
+		SequenceIndex = MakeSequenceIndex(GetPrimaryIndex(), InSetTo);
+		return *this;
+	}
+
+	FMediaTimeStamp& AdjustPrimaryIndex(int32 Add)
+	{
+		SequenceIndex += (static_cast<int64>(Add) << 32);
+		return *this;
+	}
+
+	FMediaTimeStamp& AdjustSecondaryIndex(int32 Add)
+	{
+		SequenceIndex += Add;
+		return *this;
+	}
+
+	int32 GetPrimaryIndex() const
+	{
+		return static_cast<int32>(SequenceIndex >> 32);
+	}
+
+	int32 GetSecondaryIndex() const
+	{
+		return static_cast<int32>(((uint64)SequenceIndex) & 0xffffffff);
+	}
+
+
 	static int64 MakeSequenceIndex(int32 PrimaryIndex, int32 SecondaryIndex)
 	{
 		return (static_cast<int64>(PrimaryIndex) << 32) + int64(SecondaryIndex);
@@ -88,7 +123,7 @@ public:
 	FMediaTimeStampSample() : SampledAtTime(-1.0) {}
 	FMediaTimeStampSample(const FMediaTimeStamp & InTimeStamp, double InSampledAtTime) : TimeStamp(InTimeStamp), SampledAtTime(InSampledAtTime) {}
 
-	void Invalidate() { TimeStamp.Invalidate(); }
+	void Invalidate() { TimeStamp.Invalidate(); SampledAtTime = -1.0; }
 	bool IsValid() const { return TimeStamp.IsValid(); }
 
 	FMediaTimeStamp TimeStamp;
diff --git a/Engine/Source/Runtime/MediaAssets/Private/Assets/MediaSoundComponent.cpp b/Engine/Source/Runtime/MediaAssets/Private/Assets/MediaSoundComponent.cpp
--- a/Engine/Source/Runtime/MediaAssets/Private/Assets/MediaSoundComponent.cpp
+++ b/Engine/Source/Runtime/MediaAssets/Private/Assets/MediaSoundComponent.cpp
@@ -62,7 +62,7 @@ FMediaSoundGenerator::FMediaSoundGenerator(FSoundGeneratorParams& InParams)
 
 	Audio::FEnvelopeFollowerInitParams EnvelopeInitParams;
 	EnvelopeInitParams.SampleRate = (float)InParams.SampleRate;
-	EnvelopeInitParams.NumChannels = 1; //EnvelopeFollower uses mixed down mono buffer 
+	EnvelopeInitParams.NumChannels = 1; //EnvelopeFollower uses mixed down mono buffer
 	EnvelopeFollower.Init(EnvelopeInitParams);
 
 	CachedRate = Params.CachedRate;
@@ -127,6 +127,7 @@ void FMediaSoundGenerator::SetSampleQueue(TSharedPtr<FMediaAudioSampleQueue, ESP
 {
 	FScopeLock Lock(&SampleQueueCritSect);
 	Params.SampleQueue = InSampleQueue;
+	Params.PreviousSampleQueueFlushCount = Params.SampleQueue.IsValid() ? Params.SampleQueue->GetFlushCount() : 0;
 
 	UE_LOG(LogMediaAssets, Verbose, TEXT("MediaSoundComponent: SetSampleQueue called with new sample queue."));
 }
@@ -275,9 +276,9 @@ void UMediaSoundComponent::UpdatePlayer()
 	// create a new sample queue if the player changed
 	TSharedRef<FMediaPlayerFacade, ESPMode::ThreadSafe> PlayerFacade = CurrentPlayerPtr->GetPlayerFacade();
 
-	// We have some audio decoders which are running with a limited amount of pre-allocated audio sample packets. 
-	// When the audio packets are not consumed in the FMediaSoundGenerator::OnGenerateAudio method below, these packets are not 
-	// returned to the decoder which then cannot produce more audio samples. 
+	// We have some audio decoders which are running with a limited amount of pre-allocated audio sample packets.
+	// When the audio packets are not consumed in the FMediaSoundGenerator::OnGenerateAudio method below, these packets are not
+	// returned to the decoder which then cannot produce more audio samples.
 	//
 	// The FMediaSoundGenerator::OnGenerateAudio is only called when our parent USynthComponent it active and
 	// this is controlled by USynthComponent::Start() and USynthComponent::Stop(). We are tracking a state change here.
@@ -485,6 +486,7 @@ ISoundGeneratorPtr UMediaSoundComponent::CreateSoundGenerator(const FSoundGenera
 	Params.SampleRate = (int32)InParams.SampleRate;
 	Params.NumChannels = InParams.NumChannels;
 	Params.SampleQueue = SampleQueue;
+	Params.PreviousSampleQueueFlushCount = Params.SampleQueue.IsValid() ? Params.SampleQueue->GetFlushCount() : 0;
 
 	Params.bSpectralAnalysisEnabled = bSpectralAnalysisEnabled;
 	Params.bEnvelopeFollowingEnabled = bEnvelopeFollowingEnabled;
@@ -492,14 +494,14 @@ ISoundGeneratorPtr UMediaSoundComponent::CreateSoundGenerator(const FSoundGenera
 	Params.EnvelopeFollowerReleaseTime = EnvelopeFollowerReleaseTime;
 	Params.SpectrumAnalyzerSettings = SpectrumAnalyzerSettings;
 	Params.FrequenciesToAnalyze = FrequenciesToAnalyze;
-	
+
 	Params.CachedRate = CachedRate;
 	Params.CachedTime = CachedTime;
 	Params.LastPlaySampleTime = LastPlaySampleTime;
 
 	return MediaSoundGenerator = ISoundGeneratorPtr(new FMediaSoundGenerator(Params));
 }
- 
+
 int32 FMediaSoundGenerator::OnGenerateAudio(float* OutAudio, int32 NumSamples)
 {
 	CSV_SCOPED_TIMING_STAT(MediaStreaming, FMediaSoundGenerator_OnGenerateAudio);
@@ -515,8 +517,11 @@ int32 FMediaSoundGenerator::OnGenerateAudio(float* OutAudio, int32 NumSamples)
 
 	const float Rate = CachedRate.Load();
 
-	// 	// We have an input queue and are actively playing?
-	if (PinnedSampleQueue.IsValid() && (Rate != 0.0f))
+	// We have an input queue and are actively playing?
+	uint32 CurrentSampleQueueFlushCount = PinnedSampleQueue.IsValid() ? PinnedSampleQueue->GetFlushCount() : 0;
+	bool bStillGood = Rate != 0.0f && PinnedSampleQueue.IsValid() && CurrentSampleQueueFlushCount == Params.PreviousSampleQueueFlushCount;
+	Params.PreviousSampleQueueFlushCount = CurrentSampleQueueFlushCount;
+	if (bStillGood)
 	{
 		const FTimespan Time = CachedTime.Load();
 
@@ -539,7 +544,7 @@ int32 FMediaSoundGenerator::OnGenerateAudio(float* OutAudio, int32 NumSamples)
 
 			// Update audio time
 			LastPlaySampleTime = OutTime.Time;
-			PinnedSampleQueue->SetAudioTime(FMediaTimeStampSample(OutTime, FPlatformTime::Seconds()));
+			PinnedSampleQueue->SetAudioTimeIfEqualFlushCount(FMediaTimeStampSample(OutTime, FPlatformTime::Seconds()), CurrentSampleQueueFlushCount);
 
 			SET_FLOAT_STAT(STAT_MediaUtils_MediaSoundComponentSampleTime, OutTime.Time.GetTotalSeconds());
 			SET_DWORD_STAT(STAT_Media_SoundCompQueued, PinnedSampleQueue->Num());
@@ -551,7 +556,7 @@ int32 FMediaSoundGenerator::OnGenerateAudio(float* OutAudio, int32 NumSamples)
 		{
 			float* BufferToUseForAnalysis = nullptr;
 			int32 NumFrames = NumSamples;
-			
+
 			if (Params.NumChannels == 2)
 			{
 				NumFrames = NumSamples / 2;
@@ -631,20 +636,20 @@ void UMediaSoundComponent::SetSpectralAnalysisSettings(TArray<float> InFrequenci
 
 	switch (InFFTSize)
 	{
-		case EMediaSoundComponentFFTSize::Min_64: 
+		case EMediaSoundComponentFFTSize::Min_64:
 			SpectrumAnalyzerSize = Audio::FSpectrumAnalyzerSettings::EFFTSize::Min_64;
 			break;
-		
-		case EMediaSoundComponentFFTSize::Small_256: 
+
+		case EMediaSoundComponentFFTSize::Small_256:
 			SpectrumAnalyzerSize = Audio::FSpectrumAnalyzerSettings::EFFTSize::Small_256;
 			break;
-		
+
 		default:
 		case EMediaSoundComponentFFTSize::Medium_512:
 			SpectrumAnalyzerSize = Audio::FSpectrumAnalyzerSettings::EFFTSize::Medium_512;
 			break;
 
-		case EMediaSoundComponentFFTSize::Large_1024: 
+		case EMediaSoundComponentFFTSize::Large_1024:
 			SpectrumAnalyzerSize = Audio::FSpectrumAnalyzerSettings::EFFTSize::Large_1024;
 			break;
 	}
@@ -759,7 +764,7 @@ const FSoundAttenuationSettings* UMediaSoundComponent::GetSelectedAttenuationSet
 	{
 		return &AttenuationOverrides;
 	}
-	
+
 	if (AttenuationSettings != nullptr)
 	{
 		return &AttenuationSettings->Attenuation;
diff --git a/Engine/Source/Runtime/MediaAssets/Public/MediaSoundComponent.h b/Engine/Source/Runtime/MediaAssets/Public/MediaSoundComponent.h
--- a/Engine/Source/Runtime/MediaAssets/Public/MediaSoundComponent.h
+++ b/Engine/Source/Runtime/MediaAssets/Public/MediaSoundComponent.h
@@ -78,6 +78,7 @@ public:
 		int32 NumChannels = 0;
 
 		TSharedPtr<FMediaAudioSampleQueue, ESPMode::ThreadSafe> SampleQueue;
+		uint32 PreviousSampleQueueFlushCount = 0;
 
 		bool bSpectralAnalysisEnabled = false;
 		bool bEnvelopeFollowingEnabled = false;
@@ -249,7 +250,7 @@ public:
 	/** Turns on spectral analysis of the audio generated in the media sound component. */
 	UFUNCTION(BlueprintCallable, Category = "Media|MediaSoundComponent")
 	MEDIAASSETS_API void SetEnableSpectralAnalysis(bool bInSpectralAnalysisEnabled);
-	
+
 	/** Sets the settings to use for spectral analysis. */
 	UFUNCTION(BlueprintCallable, Category = "Media|MediaSoundComponent")
 	MEDIAASSETS_API void SetSpectralAnalysisSettings(TArray<float> InFrequenciesToAnalyze, EMediaSoundComponentFFTSize InFFTSize = EMediaSoundComponentFFTSize::Medium_512);
diff --git a/Engine/Source/Runtime/MediaUtils/Private/MediaPlayerFacade.cpp b/Engine/Source/Runtime/MediaUtils/Private/MediaPlayerFacade.cpp
--- a/Engine/Source/Runtime/MediaUtils/Private/MediaPlayerFacade.cpp
+++ b/Engine/Source/Runtime/MediaUtils/Private/MediaPlayerFacade.cpp
@@ -123,7 +123,6 @@ FMediaPlayerFacade::FMediaPlayerFacade(TWeakObjectPtr<UMediaPlayer> InMediaPlaye
 	, bHaveActiveAudio(false)
 	, VideoSampleAvailability(-1)
 	, AudioSampleAvailability(-1)
-	, bIsSinkFlushPending(false)
 	, bAreEventsSafeForAnyThread(false)
 	, MediaPlayer(InMediaPlayer)
 {
@@ -286,20 +285,10 @@ void FMediaPlayerFacade::Close()
 		NotifyLifetimeManagerDelegate_PlayerClosed();
 	}
 
+	Flush();
+	ReInit();
 	BlockOnRange.Reset();
-
-	Cache->Empty();
-	CurrentUrl.Empty();
-	LastRate = 0.0f;
-	CurrentRate = 0.0f;
-
-	bHaveActiveAudio = false;
-	VideoSampleAvailability = -1;
-	AudioSampleAvailability = -1;
-	bIsSinkFlushPending = false;
 	bDidRecentPlayerHaveError = false;
-
-	Flush();
 }
 
 
@@ -1518,7 +1507,7 @@ bool FMediaPlayerFacade::SetRate(float Rate)
 		if (Rate == 0.0f)
 		{
 			// Invalidate audio time on entering pause mode...
-			if (TSharedPtr< FMediaAudioSampleSink, ESPMode::ThreadSafe> AudioSink = PrimaryAudioSink.Pin())
+			if (TSharedPtr<FMediaAudioSampleSink, ESPMode::ThreadSafe> AudioSink = PrimaryAudioSink.Pin())
 			{
 				AudioSink->InvalidateAudioTime();
 			}
@@ -2179,12 +2168,6 @@ void FMediaPlayerFacade::TickInput(FTimespan DeltaTime, FTimespan Timecode)
 				return;
 			}
 
-			if (bIsSinkFlushPending)
-			{
-				bIsSinkFlushPending = false;
-				Flush();
-			}
-
 			//
 			// Setup timing for sample processing
 			//
@@ -3524,11 +3507,6 @@ void FMediaPlayerFacade::ReceiveMediaEvent(EMediaEvent Event)
 
 		case	EMediaEvent::Internal_ResetForDiscontinuity:
 		{
-			// Disabled for now to prevent a flush on the next handling iteration that might discard the samples a blocking range is waiting for.
-			#if 0
-				UE_LOG(LogMediaUtils, VeryVerbose, TEXT("PlayerFacade %p: Reset for discontinuity"), this);
-				bIsSinkFlushPending = true;
-			#endif
 			break;
 		}
 		case	EMediaEvent::Internal_RenderClockStart:
@@ -3580,3 +3558,34 @@ void FMediaPlayerFacade::ReceiveMediaEvent(EMediaEvent Event)
 		QueuedEvents.Enqueue(Event);
 	}
 }
+
+void FMediaPlayerFacade::ReInit()
+{
+	// We leave the registered sinks and delegates alone
+	{
+		FScopeLock lock(&CriticalSection);
+		BlockOnRange.Reset();
+		BlockOnRangeDisabled = false;
+		CurrentUrl.Empty();
+		LastRate = 0.0f;
+		CurrentRate = 0.0f;
+		bHaveActiveAudio = false;
+		VideoSampleAvailability = -1;
+		AudioSampleAvailability = -1;
+		NextVideoSampleTime = FTimespan::Zero();
+	}
+
+	{
+		FScopeLock lock(&LastTimeValuesCS);
+		LastAudioRenderedSampleTime.Invalidate();
+		LastAudioSampleProcessedTime.Invalidate();
+		LastVideoSampleProcessedTimeRange = TRange<FMediaTimeStamp>::Empty();
+		CurrentFrameAudioTimeStamp.Invalidate();
+		CurrentFrameVideoTimeStamp.Invalidate();
+		CurrentFrameVideoDisplayTimeStamp.Invalidate();
+		NextEstVideoTimeAtFrameStart.Invalidate();
+		SeekTargetTime.Invalidate();
+		SeekIndex = 0;
+	}
+}
+
diff --git a/Engine/Source/Runtime/MediaUtils/Public/MediaPlayerFacade.h b/Engine/Source/Runtime/MediaUtils/Public/MediaPlayerFacade.h
--- a/Engine/Source/Runtime/MediaUtils/Public/MediaPlayerFacade.h
+++ b/Engine/Source/Runtime/MediaUtils/Public/MediaPlayerFacade.h
@@ -938,6 +938,8 @@ private:
 private:
 	MEDIAUTILS_API TRange<FMediaTimeStamp> GetAdjustedBlockOnRange() const;
 
+	void ReInit();
+
 	class FBlockOnRange
 	{
 	public:
@@ -1066,9 +1068,6 @@ private:
 	/** Current seek index */
 	int32 SeekIndex;
 
-	/** Set if sinks are to be flushed at the request of the player. */
-	TAtomic<bool>	bIsSinkFlushPending;
-
 	/** Latch for error state of the most recently used player to be queried after it may have been closed. **/
 	bool bDidRecentPlayerHaveError;
 
diff --git a/Engine/Source/Runtime/MediaUtils/Public/MediaSampleQueue.h b/Engine/Source/Runtime/MediaUtils/Public/MediaSampleQueue.h
--- a/Engine/Source/Runtime/MediaUtils/Public/MediaSampleQueue.h
+++ b/Engine/Source/Runtime/MediaUtils/Public/MediaSampleQueue.h
@@ -126,6 +126,18 @@ public:
 		return true;
 	}
 
+	virtual void GetSampleTimes(TArray<TRange<FMediaTimeStamp>>& OutSampleTimeRanges) override
+	{
+		FScopeLock Lock(&CriticalSection);
+		for(int32 i=0,iMax=Samples.Num(); i<iMax; ++i)
+		{
+			if (Samples[i].IsValid())
+			{
+				OutSampleTimeRanges.Emplace(TRange<FMediaTimeStamp>(Samples[i]->GetTime(), Samples[i]->GetTime() + Samples[i]->GetDuration()));
+			}
+		}
+	}
+
 	virtual bool Pop() override
 	{
 		FScopeLock Lock(&CriticalSection);
@@ -343,6 +355,13 @@ public:
 	{
 		FScopeLock Lock(&CriticalSection);
 		Samples.Empty();
+		++FlushCount;
+	}
+
+	virtual uint32 GetFlushCount() const
+	{
+		FScopeLock Lock(&CriticalSection);
+		return FlushCount;
 	}
 
 	virtual bool CanAcceptSamples(int32 NumSamples) const override
@@ -405,6 +424,7 @@ protected:
 	mutable FCriticalSection CriticalSection;
 	TArray<TSharedPtr<SampleType, ESPMode::ThreadSafe>> Samples;
 	int32 MaxSamplesInQueue;
+	uint32 FlushCount = 0;
 };
 
 
@@ -416,12 +436,21 @@ public:
 		: TMediaSampleQueue<class IMediaAudioSample, class FMediaAudioSampleSink>(MaxSamplesInQueue)
 	{ }
 
-	void SetAudioTime(const FMediaTimeStampSample & InAudioTime)
+	void SetAudioTime(const FMediaTimeStampSample& InAudioTime)
 	{
 		FScopeLock Lock(&CriticalSection);
 		AudioTime = InAudioTime;
 	}
 
+	void SetAudioTimeIfEqualFlushCount(const FMediaTimeStampSample& InAudioTime, uint32 InFlushCount)
+	{
+		FScopeLock Lock(&CriticalSection);
+		if (InFlushCount == FlushCount)
+		{
+			AudioTime = InAudioTime;
+		}
+	}
+
 	FMediaTimeStampSample GetAudioTime() const override
 	{
 		FScopeLock Lock(&CriticalSection);
diff --git a/Engine/Source/Runtime/MediaUtils/Public/MediaSampleSink.h b/Engine/Source/Runtime/MediaUtils/Public/MediaSampleSink.h
--- a/Engine/Source/Runtime/MediaUtils/Public/MediaSampleSink.h
+++ b/Engine/Source/Runtime/MediaUtils/Public/MediaSampleSink.h
@@ -14,7 +14,7 @@ class UMediaPlayer;
 /** Events send from player to sink */
 enum class EMediaSampleSinkEvent
 {
-	Attached,				//!< Attached to a UMediaPlayer 
+	Attached,				//!< Attached to a UMediaPlayer
 	Detached,				//!< Detached from a UMediaPlayer
 	PlayerPluginChange,		//!< Player plugin used changed
 	SampleDataUpdate,		//!< Sample data was updated
@@ -93,7 +93,7 @@ public:
 
 	/**
 	 * Check if sink can accept new samples
-	 * 
+	 *
 	 * @param NumSamples How many samples we would like the sink to accept
 	 * @return True if samples could be accepted, false otherwise
 	 * @note Override in implementation as needed
@@ -111,6 +111,12 @@ public:
 	 */
 	virtual void RequestFlush() = 0;
 
+	/**
+	 * Returns the number of times the queue has been requested to flush already.
+	 * @return Number of times RequestFlush() was called.
+	 */
+	virtual uint32 GetFlushCount() const = 0;
+
 	/**
 	 * Receive event
 	 */
diff --git a/Engine/Source/Runtime/MediaUtils/Public/MediaSampleSource.h b/Engine/Source/Runtime/MediaUtils/Public/MediaSampleSource.h
--- a/Engine/Source/Runtime/MediaUtils/Public/MediaSampleSource.h
+++ b/Engine/Source/Runtime/MediaUtils/Public/MediaSampleSource.h
@@ -3,9 +3,10 @@
 #pragma once
 
 #include "CoreTypes.h"
-#include "Misc/Timespan.h"
+#include "IMediaTimeSource.h"
 #include "Templates/SharedPointer.h"
-
+#include "Containers/Array.h"
+#include "Math/Range.h"
 
 /**
  * Interface for media sample sources.
@@ -48,6 +49,12 @@ public:
 	 */
 	virtual bool PeekFrontAndBack(TSharedPtr<SampleType, ESPMode::ThreadSafe>& OutFirstSample, TSharedPtr<SampleType, ESPMode::ThreadSafe>& OutLastSample) = 0;
 
+	/**
+	 * Returns the sample start and end times of all samples currently in the queue.
+	 * @param OutSampleTimeRanges Will contain the sample time ranges.
+	 */
+	virtual void GetSampleTimes(TArray<TRange<FMediaTimeStamp>>& OutSampleTimeRanges) = 0;
+
 	/**
 	 * Remove the next sample from the queue.
 	 *
diff --git a/Engine/Source/Runtime/MediaUtils/Public/MediaSamples.h b/Engine/Source/Runtime/MediaUtils/Public/MediaSamples.h
--- a/Engine/Source/Runtime/MediaUtils/Public/MediaSamples.h
+++ b/Engine/Source/Runtime/MediaUtils/Public/MediaSamples.h
@@ -180,15 +180,15 @@ public:
 		return true;
 	}
 
-	virtual bool PeekVideoSampleTimeRange(TRange<FMediaTimeStamp>& TimeRange) override
+	virtual bool PeekVideoSampleTimeRanges(TArray<TRange<FMediaTimeStamp>>& TimeRange) override
 	{
-		TSharedPtr<IMediaTextureSample, ESPMode::ThreadSafe> First, Last;
-		if (!VideoSampleQueue.PeekFrontAndBack(First, Last))
-		{
-			return false;
-		}
-		TimeRange = TRange<FMediaTimeStamp>(First->GetTime(), Last->GetTime() + Last->GetDuration());
-		return true;
+		VideoSampleQueue.GetSampleTimes(TimeRange);
+		return TimeRange.Num() > 0;
+	}
+	virtual bool PeekAudioSampleTimeRanges(TArray<TRange<FMediaTimeStamp>>& TimeRange) override
+	{
+		AudioSampleQueue.GetSampleTimes(TimeRange);
+		return TimeRange.Num() > 0;
 	}
 
 	virtual bool DiscardVideoSamples(const TRange<FMediaTimeStamp>& TimeRange, bool bReverse) override
