Electra: Fixed a nullptr crash when download of an init segment has been canceled

#rnx

diff --git a/Engine/Plugins/Media/ElectraPlayer/Source/ElectraPlayerRuntime/Private/Runtime/Demuxer/ParserISO14496-12_Utils.h b/Engine/Plugins/Media/ElectraPlayer/Source/ElectraPlayerRuntime/Private/Runtime/Demuxer/ParserISO14496-12_Utils.h
--- a/Engine/Plugins/Media/ElectraPlayer/Source/ElectraPlayerRuntime/Private/Runtime/Demuxer/ParserISO14496-12_Utils.h
+++ b/Engine/Plugins/Media/ElectraPlayer/Source/ElectraPlayerRuntime/Private/Runtime/Demuxer/ParserISO14496-12_Utils.h
@@ -22,13 +22,17 @@ namespace Electra
 			Data = (const uint8*)ResponseBuffer->Buffer.GetLinearReadData();
 			CurrentOffset = 0;
 		}
+		virtual bool HaveParseData() const
+		{
+			return ResponseBuffer.IsValid();
+		}
 	private:
 		//----------------------------------------------------------------------
 		// Methods from IParserISO14496_12::IReader
 		//
 		virtual int64 ReadData(void* IntoBuffer, int64 NumBytesToRead) override
 		{
-			if (NumBytesToRead <= DataSize - CurrentOffset)
+			if (ResponseBuffer.IsValid() && NumBytesToRead <= DataSize - CurrentOffset)
 			{
 				if (IntoBuffer)
 				{
@@ -41,11 +45,11 @@ namespace Electra
 		}
 		virtual bool HasReachedEOF() const override
 		{
-			return ResponseBuffer->Buffer.GetEOD() && CurrentOffset >= DataSize;
+			return ResponseBuffer.IsValid() ? ResponseBuffer->Buffer.GetEOD() && CurrentOffset >= DataSize : true;
 		}
 		virtual bool HasReadBeenAborted() const override
 		{
-			return ResponseBuffer->Buffer.WasAborted();
+			return ResponseBuffer.IsValid() ? ResponseBuffer->Buffer.WasAborted() : true;
 		}
 		virtual int64 GetCurrentOffset() const override
 		{
