Media: Changed CanPlayUrl() to return a value indicating the confidence that the media can be played in case several plugins of different capabilities claim the ability.

#rnx

diff --git a/Engine/Plugins/Experimental/TextureMediaPlayer/Source/TextureMediaPlayerFactory/Private/TextureMediaPlayerFactoryModule.cpp b/Engine/Plugins/Experimental/TextureMediaPlayer/Source/TextureMediaPlayerFactory/Private/TextureMediaPlayerFactoryModule.cpp
--- a/Engine/Plugins/Experimental/TextureMediaPlayer/Source/TextureMediaPlayerFactory/Private/TextureMediaPlayerFactoryModule.cpp
+++ b/Engine/Plugins/Experimental/TextureMediaPlayer/Source/TextureMediaPlayerFactory/Private/TextureMediaPlayerFactoryModule.cpp
@@ -26,6 +26,11 @@ public:
 	// IMediaPlayerFactory interface
 
 	bool CanPlayUrl(const FString& Url, const IMediaOptions* Options, TArray<FText>* OutWarnings, TArray<FText>* OutErrors) const override
+	{
+		return GetPlayabilityConfidenceScore(Url, Options, OutWarnings, OutErrors) > 0 ? true : false;
+	}
+
+	int32 GetPlayabilityConfidenceScore(const FString& Url, const IMediaOptions* Options, TArray<FText>* OutWarnings, TArray<FText>* OutErrors) const override
 	{
 		FString Scheme;
 		FString Location;
@@ -38,7 +43,7 @@ public:
 				OutErrors->Add(LOCTEXT("NoSchemeFound", "No URI scheme found"));
 			}
 
-			return false;
+			return 0;
 		}
 
 		if (!SupportedUriSchemes.Contains(Scheme))
@@ -48,10 +53,10 @@ public:
 				OutErrors->Add(FText::Format(LOCTEXT("SchemeNotSupported", "The URI scheme '{0}' is not supported"), FText::FromString(Scheme)));
 			}
 
-			return false;
+			return 0;
 		}
 
-		return true;
+		return 100;
 	}
 
 	TSharedPtr<IMediaPlayer, ESPMode::ThreadSafe> CreatePlayer(IMediaEventSink& EventSink) override
@@ -99,7 +104,7 @@ public:
 		// supported platforms
 		auto MediaModule = FModuleManager::GetModulePtr<IMediaModule>("Media");
 		if (MediaModule != nullptr)
-		{		
+		{
 			AddSupportedPlatform(FGuid(0xd1d5f296, 0xff834a87, 0xb20faaa9, 0xd6b8e9a6));
 			AddSupportedPlatform(FGuid(0xb80decd6, 0x997a4b3f, 0x92063970, 0xe572c0db));
 			AddSupportedPlatform(FGuid(0x30ebce04, 0x2c8247bd, 0xaf873017, 0x5a27ed45));
diff --git a/Engine/Plugins/Media/AjaMedia/Source/AjaMediaFactory/Private/AjaMediaFactoryModule.cpp b/Engine/Plugins/Media/AjaMedia/Source/AjaMediaFactory/Private/AjaMediaFactoryModule.cpp
--- a/Engine/Plugins/Media/AjaMedia/Source/AjaMediaFactory/Private/AjaMediaFactoryModule.cpp
+++ b/Engine/Plugins/Media/AjaMedia/Source/AjaMediaFactory/Private/AjaMediaFactoryModule.cpp
@@ -29,7 +29,12 @@ public:
 
 	//~ IMediaPlayerFactory interface
 
-	virtual bool CanPlayUrl(const FString& Url, const IMediaOptions* /*Options*/, TArray<FText>* /*OutWarnings*/, TArray<FText>* OutErrors) const override
+	virtual bool CanPlayUrl(const FString& Url, const IMediaOptions* Options, TArray<FText>* OutWarnings, TArray<FText>* OutErrors) const override
+	{
+		return GetPlayabilityConfidenceScore(Url, Options, OutWarnings, OutErrors) > 0 ? true : false;
+	}
+
+	virtual int32 GetPlayabilityConfidenceScore(const FString& Url, const IMediaOptions* Options, TArray<FText>* OutWarnings, TArray<FText>* OutErrors) const override
 	{
 		FString Scheme;
 		FString Location;
@@ -42,7 +47,7 @@ public:
 				OutErrors->Add(LOCTEXT("NoSchemeFound", "No URI scheme found"));
 			}
 
-			return false;
+			return 0;
 		}
 
 		if (!SupportedUriSchemes.Contains(Scheme))
@@ -52,10 +57,10 @@ public:
 				OutErrors->Add(FText::Format(LOCTEXT("SchemeNotSupported", "The URI scheme '{0}' is not supported"), FText::FromString(Scheme)));
 			}
 
-			return false;
+			return 0;
 		}
 
-		return true;
+		return 100;
 	}
 
 	virtual TSharedPtr<IMediaPlayer, ESPMode::ThreadSafe> CreatePlayer(IMediaEventSink& EventSink) override
diff --git a/Engine/Plugins/Media/AndroidCamera/Source/AndroidCameraFactory/Private/AndroidCameraFactoryModule.cpp b/Engine/Plugins/Media/AndroidCamera/Source/AndroidCameraFactory/Private/AndroidCameraFactoryModule.cpp
--- a/Engine/Plugins/Media/AndroidCamera/Source/AndroidCameraFactory/Private/AndroidCameraFactoryModule.cpp
+++ b/Engine/Plugins/Media/AndroidCamera/Source/AndroidCameraFactory/Private/AndroidCameraFactoryModule.cpp
@@ -28,7 +28,12 @@ public:
 
 	//~ IMediaPlayerFactory interface
 
-	virtual bool CanPlayUrl(const FString& Url, const IMediaOptions* /*Options*/, TArray<FText>* /*OutWarnings*/, TArray<FText>* OutErrors) const override
+	virtual bool CanPlayUrl(const FString& Url, const IMediaOptions* Options, TArray<FText>* OutWarnings, TArray<FText>* OutErrors) const override
+	{
+		return GetPlayabilityConfidenceScore(Url, Options, OutWarnings, OutErrors) > 0 ? true : false;
+	}
+
+	virtual int32 GetPlayabilityConfidenceScore(const FString& Url, const IMediaOptions* Options, TArray<FText>* OutWarnings, TArray<FText>* OutErrors) const override
 	{
 		FString Scheme;
 		FString Location;
@@ -41,7 +46,7 @@ public:
 				OutErrors->Add(LOCTEXT("NoSchemeFound", "No URI scheme found"));
 			}
 
-			return false;
+			return 0;
 		}
 
 		if (!SupportedUriSchemes.Contains(Scheme))
@@ -51,7 +56,7 @@ public:
 				OutErrors->Add(FText::Format(LOCTEXT("SchemeNotSupported", "The URI scheme '{0}' is not supported"), FText::FromString(Scheme)));
 			}
 
-			return false;
+			return 0;
 		}
 
 		// check file extension
@@ -66,11 +71,11 @@ public:
 					OutErrors->Add(FText::Format(LOCTEXT("ExtensionNotSupported", "The file extension '{0}' is not supported"), FText::FromString(Extension)));
 				}
 
-				return false;
+				return 0;
 			}
 		}
 
-		return true;
+		return 100;
 	}
 
 	virtual TSharedPtr<IMediaPlayer, ESPMode::ThreadSafe> CreatePlayer(IMediaEventSink& EventSink) override
diff --git a/Engine/Plugins/Media/AndroidMedia/Source/AndroidMediaFactory/Private/AndroidMediaFactoryModule.cpp b/Engine/Plugins/Media/AndroidMedia/Source/AndroidMediaFactory/Private/AndroidMediaFactoryModule.cpp
--- a/Engine/Plugins/Media/AndroidMedia/Source/AndroidMediaFactory/Private/AndroidMediaFactoryModule.cpp
+++ b/Engine/Plugins/Media/AndroidMedia/Source/AndroidMediaFactory/Private/AndroidMediaFactoryModule.cpp
@@ -28,7 +28,12 @@ public:
 
 	//~ IMediaPlayerFactory interface
 
-	virtual bool CanPlayUrl(const FString& Url, const IMediaOptions* /*Options*/, TArray<FText>* /*OutWarnings*/, TArray<FText>* OutErrors) const override
+	virtual bool CanPlayUrl(const FString& Url, const IMediaOptions* Options, TArray<FText>* OutWarnings, TArray<FText>* OutErrors) const override
+	{
+		return GetPlayabilityConfidenceScore(Url, Options, OutWarnings, OutErrors) > 0 ? true : false;
+	}
+
+	virtual int32 GetPlayabilityConfidenceScore(const FString& Url, const IMediaOptions* Options, TArray<FText>* OutWarnings, TArray<FText>* OutErrors) const override
 	{
 		FString Scheme;
 		FString Location;
@@ -41,7 +46,7 @@ public:
 				OutErrors->Add(LOCTEXT("NoSchemeFound", "No URI scheme found"));
 			}
 
-			return false;
+			return 0;
 		}
 
 		if (!SupportedUriSchemes.Contains(Scheme))
@@ -51,7 +56,7 @@ public:
 				OutErrors->Add(FText::Format(LOCTEXT("SchemeNotSupported", "The URI scheme '{0}' is not supported"), FText::FromString(Scheme)));
 			}
 
-			return false;
+			return 0;
 		}
 
 		// check file extension
@@ -66,11 +71,11 @@ public:
 					OutErrors->Add(FText::Format(LOCTEXT("ExtensionNotSupported", "The file extension '{0}' is not supported"), FText::FromString(Extension)));
 				}
 
-				return false;
+				return 0;
 			}
 		}
 
-		return true;
+		return 80;
 	}
 
 	virtual TSharedPtr<IMediaPlayer, ESPMode::ThreadSafe> CreatePlayer(IMediaEventSink& EventSink) override
diff --git a/Engine/Plugins/Media/AvfMedia/Source/AvfMediaCapture/Private/AvfMediaCaptureModule.cpp b/Engine/Plugins/Media/AvfMedia/Source/AvfMediaCapture/Private/AvfMediaCaptureModule.cpp
--- a/Engine/Plugins/Media/AvfMedia/Source/AvfMediaCapture/Private/AvfMediaCaptureModule.cpp
+++ b/Engine/Plugins/Media/AvfMedia/Source/AvfMediaCapture/Private/AvfMediaCaptureModule.cpp
@@ -90,7 +90,7 @@ public:
 			#elif PLATFORM_MAC
 				// AVCaptureDeviceTypeExternalUnknown is only available on macOS 10.15 - 14.0
 				// https://developer.apple.com/documentation/avfoundation/avcapturedevicetypeexternalunknown?language=objc
-				[DeviceTypes addObject: AVCaptureDeviceTypeExternalUnknown]; 
+				[DeviceTypes addObject: AVCaptureDeviceTypeExternalUnknown];
 			#endif
 		}
 
@@ -117,24 +117,24 @@ public:
 				}
 
 				FMediaCaptureDeviceInfo DeviceInfo;
-					
+
 				DeviceInfo.Type = TargetDeviceType;
 				DeviceInfo.DisplayName = FText::FromString(FString(AvailableDevice.localizedName));
 				DeviceInfo.Url = Scheme + FString(AvailableDevice.uniqueID);
 				DeviceInfo.Info = FString(AvailableDevice.manufacturer);
-					
+
 				OutDeviceInfos.Add(MoveTemp(DeviceInfo));
 			}
 		}
 	}
 
 	//~ IMediaCaptureSupport interface
-	
+
 	virtual void EnumerateAudioCaptureDevices(TArray<FMediaCaptureDeviceInfo>& OutDeviceInfos)
 	{
 		EnumerateCaptureDevices(OutDeviceInfos, EMediaCaptureDeviceType::Audio);
 	}
-	
+
 	virtual void EnumerateVideoCaptureDevices(TArray<FMediaCaptureDeviceInfo>& OutDeviceInfos)
 	{
 		EnumerateCaptureDevices(OutDeviceInfos, EMediaCaptureDeviceType::Video);
@@ -142,7 +142,12 @@ public:
 
 	//~ IMediaPlayerFactory interface
 
-	virtual bool CanPlayUrl(const FString& Url, const IMediaOptions* /*Options*/, TArray<FText>* /*OutWarnings*/, TArray<FText>* OutErrors) const override
+	virtual bool CanPlayUrl(const FString& Url, const IMediaOptions* Options, TArray<FText>* OutWarnings, TArray<FText>* OutErrors) const override
+	{
+		return GetPlayabilityConfidenceScore(Url, Options, OutWarnings, OutErrors) > 0 ? true : false;
+	}
+
+	virtual int32 GetPlayabilityConfidenceScore(const FString& Url, const IMediaOptions* Options, TArray<FText>* OutWarnings, TArray<FText>* OutErrors) const override
 	{
 		FString Scheme;
 		FString Location;
@@ -155,7 +160,7 @@ public:
 				OutErrors->Add(LOCTEXT("NoSchemeFound", "No URI scheme found"));
 			}
 
-			return false;
+			return 0;
 		}
 
 		if (!SupportedUriSchemes.Contains(Scheme))
@@ -165,10 +170,10 @@ public:
 				OutErrors->Add(FText::Format(LOCTEXT("SchemeNotSupported", "The URI scheme '{0}' is not supported"), FText::FromString(Scheme)));
 			}
 
-			return false;
+			return 0;
 		}
 
-		return true;
+		return 100;
 	}
 
 	virtual TSharedPtr<IMediaPlayer, ESPMode::ThreadSafe> CreatePlayer(IMediaEventSink& EventSink) override
diff --git a/Engine/Plugins/Media/AvfMedia/Source/AvfMediaFactory/Private/AvfMediaFactoryModule.cpp b/Engine/Plugins/Media/AvfMedia/Source/AvfMediaFactory/Private/AvfMediaFactoryModule.cpp
--- a/Engine/Plugins/Media/AvfMedia/Source/AvfMediaFactory/Private/AvfMediaFactoryModule.cpp
+++ b/Engine/Plugins/Media/AvfMedia/Source/AvfMediaFactory/Private/AvfMediaFactoryModule.cpp
@@ -40,7 +40,12 @@ public:
 
 	//~ IMediaPlayerFactory interface
 
-	virtual bool CanPlayUrl(const FString& Url, const IMediaOptions* /*Options*/, TArray<FText>* /*OutWarnings*/, TArray<FText>* OutErrors) const override
+	virtual bool CanPlayUrl(const FString& Url, const IMediaOptions* Options, TArray<FText>* OutWarnings, TArray<FText>* OutErrors) const override
+	{
+		return GetPlayabilityConfidenceScore(Url, Options, OutWarnings, OutErrors) > 0 ? true : false;
+	}
+
+	virtual int32 GetPlayabilityConfidenceScore(const FString& Url, const IMediaOptions* Options, TArray<FText>* OutWarnings, TArray<FText>* OutErrors) const override
 	{
 		FString Scheme;
 		FString Location;
@@ -53,7 +58,7 @@ public:
 				OutErrors->Add(LOCTEXT("NoSchemeFound", "No URI scheme found"));
 			}
 
-			return false;
+			return 0;
 		}
 
 		if (!SupportedUriSchemes.Contains(Scheme))
@@ -63,7 +68,7 @@ public:
 				OutErrors->Add(FText::Format(LOCTEXT("SchemeNotSupported", "The URI scheme '{0}' is not supported"), FText::FromString(Scheme)));
 			}
 
-			return false;
+			return 0;
 		}
 
 		// check file extension
@@ -78,11 +83,11 @@ public:
 					OutErrors->Add(FText::Format(LOCTEXT("ExtensionNotSupported", "The file extension '{0}' is not supported"), FText::FromString(Extension)));
 				}
 
-				return false;
+				return 0;
 			}
 		}
 
-		return true;
+		return 80;
 	}
 
 	virtual TSharedPtr<IMediaPlayer, ESPMode::ThreadSafe> CreatePlayer(IMediaEventSink& EventSink) override
diff --git a/Engine/Plugins/Media/BlackmagicMedia/Source/BlackmagicMediaFactory/Private/BlackmagicMediaFactoryModule.cpp b/Engine/Plugins/Media/BlackmagicMedia/Source/BlackmagicMediaFactory/Private/BlackmagicMediaFactoryModule.cpp
--- a/Engine/Plugins/Media/BlackmagicMedia/Source/BlackmagicMediaFactory/Private/BlackmagicMediaFactoryModule.cpp
+++ b/Engine/Plugins/Media/BlackmagicMedia/Source/BlackmagicMediaFactory/Private/BlackmagicMediaFactoryModule.cpp
@@ -24,7 +24,12 @@ public:
 
 	//~ IMediaPlayerFactory interface
 
-	virtual bool CanPlayUrl(const FString& Url, const IMediaOptions* /*Options*/, TArray<FText>* /*OutWarnings*/, TArray<FText>* OutErrors) const override
+	virtual bool CanPlayUrl(const FString& Url, const IMediaOptions* Options, TArray<FText>* OutWarnings, TArray<FText>* OutErrors) const override
+	{
+		return GetPlayabilityConfidenceScore(Url, Options, OutWarnings, OutErrors) > 0 ? true : false;
+	}
+
+	virtual int32 GetPlayabilityConfidenceScore(const FString& Url, const IMediaOptions* Options, TArray<FText>* OutWarnings, TArray<FText>* OutErrors) const override
 	{
 		FString Scheme;
 		FString Location;
@@ -37,7 +42,7 @@ public:
 				OutErrors->Add(LOCTEXT("NoSchemeFound", "No URI scheme found"));
 			}
 
-			return false;
+			return 0;
 		}
 
 		if (!SupportedUriSchemes.Contains(Scheme))
@@ -47,10 +52,10 @@ public:
 				OutErrors->Add(FText::Format(LOCTEXT("SchemeNotSupported", "The URI scheme '{0}' is not supported"), FText::FromString(Scheme)));
 			}
 
-			return false;
+			return 0;
 		}
 
-		return true;
+		return 100;
 	}
 
 	virtual TSharedPtr<IMediaPlayer, ESPMode::ThreadSafe> CreatePlayer(IMediaEventSink& EventSink) override
diff --git a/Engine/Plugins/Media/ElectraPlayer/Source/ElectraPlayerFactory/ElectraPlayerFactory.Build.cs b/Engine/Plugins/Media/ElectraPlayer/Source/ElectraPlayerFactory/ElectraPlayerFactory.Build.cs
--- a/Engine/Plugins/Media/ElectraPlayer/Source/ElectraPlayerFactory/ElectraPlayerFactory.Build.cs
+++ b/Engine/Plugins/Media/ElectraPlayer/Source/ElectraPlayerFactory/ElectraPlayerFactory.Build.cs
@@ -24,6 +24,7 @@ namespace UnrealBuildTool.Rules
 
             if (bSupportedPlatform)
             {
+                PrivateDependencyModuleNames.Add("ElectraBase");
                 PrivateIncludePathModuleNames.Add("ElectraPlayerPlugin");
                 PublicDefinitions.Add("UE_PLATFORM_ELECTRAPLAYER=1");
             }
diff --git a/Engine/Plugins/Media/ElectraPlayer/Source/ElectraPlayerFactory/Private/ElectraPlayerFactoryModule.cpp b/Engine/Plugins/Media/ElectraPlayer/Source/ElectraPlayerFactory/Private/ElectraPlayerFactoryModule.cpp
--- a/Engine/Plugins/Media/ElectraPlayer/Source/ElectraPlayerFactory/Private/ElectraPlayerFactoryModule.cpp
+++ b/Engine/Plugins/Media/ElectraPlayer/Source/ElectraPlayerFactory/Private/ElectraPlayerFactoryModule.cpp
@@ -13,8 +13,12 @@
 #include "Misc/Paths.h"
 #include "Misc/Guid.h"
 
-#ifdef UE_PLATFORM_ELECTRAPLAYER
+#ifndef UE_PLATFORM_ELECTRAPLAYER
+#define UE_PLATFORM_ELECTRAPLAYER 0
+#endif
+#if UE_PLATFORM_ELECTRAPLAYER
 #include "IElectraPlayerPluginModule.h"
+#include "Utilities/URLParser.h"
 #endif
 
 DEFINE_LOG_CATEGORY(LogElectraPlayerFactory);
@@ -36,36 +40,64 @@ public:
 
 	bool CanPlayUrl(const FString& Url, const IMediaOptions* Options, TArray<FText>* OutWarnings, TArray<FText>* OutErrors) const override
 	{
-		FString Scheme;
-		FString Location;
+		return GetPlayabilityConfidenceScore(Url, Options, OutWarnings, OutErrors) > 0 ? true : false;
+	}
 
-		// check scheme
-		if (!Url.Split(TEXT("://"), &Scheme, &Location, ESearchCase::CaseSensitive))
+	int32 GetPlayabilityConfidenceScore(const FString& Url, const IMediaOptions* Options, TArray<FText>* OutWarnings, TArray<FText>* OutErrors) const override
+	{
+#if UE_PLATFORM_ELECTRAPLAYER
+		// Split the URL apart.
+		Electra::FURL_RFC3986 UrlParser;
+		if (!UrlParser.Parse(Url))
+		{
+			if (OutErrors != nullptr)
+			{
+				OutErrors->Add(FText::Format(LOCTEXT("SchemeNotSupported", "The URI '{0}' could not be parsed"), FText::FromString(Url)));
+			}
+		}
+		// Check scheme
+		FString Scheme = UrlParser.GetScheme();
+		if (Scheme.IsEmpty())
 		{
 			if (OutErrors != nullptr)
 			{
 				OutErrors->Add(LOCTEXT("NoSchemeFound", "No URI scheme found"));
 			}
-
-			return false;
+			return 0;
 		}
-
 		if (!SupportedUriSchemes.Contains(Scheme))
 		{
 			if (OutErrors != nullptr)
 			{
 				OutErrors->Add(FText::Format(LOCTEXT("SchemeNotSupported", "The URI scheme '{0}' is not supported"), FText::FromString(Scheme)));
 			}
-
-			return false;
+			return 0;
 		}
 
-		return true;
+		// Check for known extensions
+		TArray<FString> PathComponents;
+		UrlParser.GetPathComponents(PathComponents);
+		FString LowerCaseExtension = FPaths::GetExtension(PathComponents.Last().ToLower());
+		// If the extension is known, we are confident that we can play this.
+		// At this point there is no information provided on the codecs used in the media,
+		// so we cannot check for this.
+		if (SupportedFileExtensions.Contains(LowerCaseExtension))
+		{
+			return 100;
+		}
+		// For http URLs, if there is no extension then we can't be sure. Return a lower confidence.
+		// If the scheme is file:// then we can actually demand there to be an extension, so if it is missing, too bad.
+		else if (LowerCaseExtension.IsEmpty() && (Scheme.Equals(TEXT("https")) || Scheme.Equals(TEXT("http"))))
+		{
+			return 20;
+		}
+#endif
+		return 0;
 	}
 
 	TSharedPtr<IMediaPlayer, ESPMode::ThreadSafe> CreatePlayer(IMediaEventSink& EventSink) override
 	{
-#ifdef UE_PLATFORM_ELECTRAPLAYER
+#if UE_PLATFORM_ELECTRAPLAYER
 		auto PlayerModule = FModuleManager::LoadModulePtr<IElectraPlayerPluginModule>("ElectraPlayerPlugin");
 		return (PlayerModule != nullptr) ? PlayerModule->CreatePlayer(EventSink) : nullptr;
 #else
@@ -114,7 +146,7 @@ public:
 	void StartupModule() override
 	{
 		// supported platforms
-#ifdef UE_PLATFORM_ELECTRAPLAYER
+#if UE_PLATFORM_ELECTRAPLAYER
 		MediaModule = FModuleManager::GetModulePtr<IMediaModule>("Media");
 		check(MediaModule);
 
@@ -134,13 +166,24 @@ public:
 		AddSupportedPlatform(FGuid(0xc0b45a33, 0x9de340c7, 0xbce24c47, 0x15c3babf));
         AddSupportedPlatform(FGuid(0xa478294f, 0xbd0d4ec0, 0x8830b6d4, 0xd219c1a4));
 		AddSupportedPlatform(FGuid(0xae496f22, 0x95534328, 0xbd035b4c, 0x919dc51a));
-#endif
 
 		// supported schemes
 		SupportedUriSchemes.Add(TEXT("http"));
 		SupportedUriSchemes.Add(TEXT("https"));
 		SupportedUriSchemes.Add(TEXT("file"));
 
+		// supported file extensions
+		SupportedFileExtensions.Add(TEXT("mp4"));
+		SupportedFileExtensions.Add(TEXT("m4v"));
+		SupportedFileExtensions.Add(TEXT("m4a"));
+		SupportedFileExtensions.Add(TEXT("mov"));
+		SupportedFileExtensions.Add(TEXT("mpd"));
+		SupportedFileExtensions.Add(TEXT("m3u8"));
+		SupportedFileExtensions.Add(TEXT("mkv"));
+		SupportedFileExtensions.Add(TEXT("mka"));
+		SupportedFileExtensions.Add(TEXT("webm"));
+#endif
+
 		// register player factory
 		MediaModule->RegisterPlayerFactory(*this);
 	}
diff --git a/Engine/Plugins/Media/ImgMedia/Source/ImgMediaFactory/Private/ImgMediaFactoryModule.cpp b/Engine/Plugins/Media/ImgMedia/Source/ImgMediaFactory/Private/ImgMediaFactoryModule.cpp
--- a/Engine/Plugins/Media/ImgMedia/Source/ImgMediaFactory/Private/ImgMediaFactoryModule.cpp
+++ b/Engine/Plugins/Media/ImgMedia/Source/ImgMediaFactory/Private/ImgMediaFactoryModule.cpp
@@ -47,7 +47,12 @@ public:
 
 	//~ IMediaPlayerInfo interface
 
-	virtual bool CanPlayUrl(const FString& Url, const IMediaOptions* /*Options*/, TArray<FText>* /*OutWarnings*/, TArray<FText>* OutErrors) const override
+	virtual bool CanPlayUrl(const FString& Url, const IMediaOptions* Options, TArray<FText>* OutWarnings, TArray<FText>* OutErrors) const override
+	{
+		return GetPlayabilityConfidenceScore(Url, Options, OutWarnings, OutErrors) > 0 ? true : false;
+	}
+
+	virtual int32 GetPlayabilityConfidenceScore(const FString& Url, const IMediaOptions* Options, TArray<FText>* OutWarnings, TArray<FText>* OutErrors) const override
 	{
 		FString Scheme;
 		FString Location;
@@ -60,7 +65,7 @@ public:
 				OutErrors->Add(LOCTEXT("NoSchemeFound", "No URI scheme found"));
 			}
 
-			return false;
+			return 0;
 		}
 
 		if (!SupportedUriSchemes.Contains(Scheme))
@@ -70,10 +75,10 @@ public:
 				OutErrors->Add(FText::Format(LOCTEXT("SchemeNotSupported", "The URI scheme '{0}' is not supported"), FText::FromString(Scheme)));
 			}
 
-			return false;
+			return 0;
 		}
 
-		return true;
+		return 100;
 	}
 
 	virtual TSharedPtr<IMediaPlayer, ESPMode::ThreadSafe> CreatePlayer(IMediaEventSink& EventSink) override
diff --git a/Engine/Plugins/Media/MfMedia/Source/MfMediaFactory/Private/MfMediaFactoryModule.cpp b/Engine/Plugins/Media/MfMedia/Source/MfMediaFactory/Private/MfMediaFactoryModule.cpp
--- a/Engine/Plugins/Media/MfMedia/Source/MfMediaFactory/Private/MfMediaFactoryModule.cpp
+++ b/Engine/Plugins/Media/MfMedia/Source/MfMediaFactory/Private/MfMediaFactoryModule.cpp
@@ -38,6 +38,11 @@ public:
 	//~ IMediaPlayerInfo interface
 
 	virtual bool CanPlayUrl(const FString& Url, const IMediaOptions* Options, TArray<FText>* OutWarnings, TArray<FText>* OutErrors) const override
+	{
+		return GetPlayabilityConfidenceScore(Url, Options, OutWarnings, OutErrors) > 0 ? true : false;
+	}
+
+	virtual int32 GetPlayabilityConfidenceScore(const FString& Url, const IMediaOptions* Options, TArray<FText>* OutWarnings, TArray<FText>* OutErrors) const override
 	{
 		FString Scheme;
 		FString Location;
@@ -50,7 +55,7 @@ public:
 				OutErrors->Add(LOCTEXT("NoSchemeFound", "No URI scheme found"));
 			}
 
-			return false;
+			return 0;
 		}
 
 		if (!SupportedUriSchemes.Contains(Scheme))
@@ -60,7 +65,7 @@ public:
 				OutErrors->Add(FText::Format(LOCTEXT("SchemeNotSupported", "The URI scheme '{0}' is not supported"), FText::FromString(Scheme)));
 			}
 
-			return false;
+			return 0;
 		}
 
 		// check file extension
@@ -75,7 +80,7 @@ public:
 					OutErrors->Add(FText::Format(LOCTEXT("ExtensionNotSupported", "The file extension '{0}' is not supported"), FText::FromString(Extension)));
 				}
 
-				return false;
+				return 0;
 			}
 		}
 
@@ -88,7 +93,7 @@ public:
 			}
 		}
 
-		return true;
+		return 80;
 	}
 
 	virtual TSharedPtr<IMediaPlayer, ESPMode::ThreadSafe> CreatePlayer(IMediaEventSink& EventSink) override
diff --git a/Engine/Plugins/Media/WebMMedia/Source/WebMMediaFactory/Private/WebMMediaFactoryModule.cpp b/Engine/Plugins/Media/WebMMedia/Source/WebMMediaFactory/Private/WebMMediaFactoryModule.cpp
--- a/Engine/Plugins/Media/WebMMedia/Source/WebMMediaFactory/Private/WebMMediaFactoryModule.cpp
+++ b/Engine/Plugins/Media/WebMMedia/Source/WebMMediaFactory/Private/WebMMediaFactoryModule.cpp
@@ -34,7 +34,12 @@ public:
 
 	//~ IMediaPlayerFactory interface
 
-	virtual bool CanPlayUrl(const FString& Url, const IMediaOptions* /*Options*/, TArray<FText>* /*OutWarnings*/, TArray<FText>* OutErrors) const override
+	virtual bool CanPlayUrl(const FString& Url, const IMediaOptions* Options, TArray<FText>* OutWarnings, TArray<FText>* OutErrors) const override
+	{
+		return GetPlayabilityConfidenceScore(Url, Options, OutWarnings, OutErrors) > 0 ? true : false;
+	}
+
+	virtual int32 GetPlayabilityConfidenceScore(const FString& Url, const IMediaOptions* Options, TArray<FText>* OutWarnings, TArray<FText>* OutErrors) const override
 	{
 		FString Scheme;
 		FString Location;
@@ -47,7 +52,7 @@ public:
 				OutErrors->Add(LOCTEXT("NoSchemeFound", "No URI scheme found"));
 			}
 
-			return false;
+			return 0;
 		}
 
 		if (!SupportedUriSchemes.Contains(Scheme))
@@ -57,7 +62,7 @@ public:
 				OutErrors->Add(FText::Format(LOCTEXT("SchemeNotSupported", "The URI scheme '{0}' is not supported"), FText::FromString(Scheme)));
 			}
 
-			return false;
+			return 0;
 		}
 
 		// check file extension
@@ -72,11 +77,11 @@ public:
 					OutErrors->Add(FText::Format(LOCTEXT("ExtensionNotSupported", "The file extension '{0}' is not supported"), FText::FromString(Extension)));
 				}
 
-				return false;
+				return 0;
 			}
 		}
 
-		return true;
+		return 80;
 	}
 
 	virtual TSharedPtr<IMediaPlayer, ESPMode::ThreadSafe> CreatePlayer(IMediaEventSink& EventSink) override
diff --git a/Engine/Plugins/Media/WmfMedia/Source/WmfMediaFactory/Private/WmfMediaFactoryModule.cpp b/Engine/Plugins/Media/WmfMedia/Source/WmfMediaFactory/Private/WmfMediaFactoryModule.cpp
--- a/Engine/Plugins/Media/WmfMedia/Source/WmfMediaFactory/Private/WmfMediaFactoryModule.cpp
+++ b/Engine/Plugins/Media/WmfMedia/Source/WmfMediaFactory/Private/WmfMediaFactoryModule.cpp
@@ -45,6 +45,11 @@ public:
 	//~ IMediaPlayerFactory interface
 
 	virtual bool CanPlayUrl(const FString& Url, const IMediaOptions* Options, TArray<FText>* OutWarnings, TArray<FText>* OutErrors) const override
+	{
+		return GetPlayabilityConfidenceScore(Url, Options, OutWarnings, OutErrors) > 0 ? true : false;
+	}
+
+	virtual int32 GetPlayabilityConfidenceScore(const FString& Url, const IMediaOptions* Options, TArray<FText>* OutWarnings, TArray<FText>* OutErrors) const override
 	{
 		FString Scheme;
 		FString Location;
@@ -57,7 +62,7 @@ public:
 				OutErrors->Add(LOCTEXT("NoSchemeFound", "No URI scheme found"));
 			}
 
-			return false;
+			return 0;
 		}
 
 		if (!SupportedUriSchemes.Contains(Scheme))
@@ -67,7 +72,7 @@ public:
 				OutErrors->Add(FText::Format(LOCTEXT("SchemeNotSupported", "The URI scheme '{0}' is not supported"), FText::FromString(Scheme)));
 			}
 
-			return false;
+			return 0;
 		}
 
 		// check file extension
@@ -82,7 +87,7 @@ public:
 					OutErrors->Add(FText::Format(LOCTEXT("ExtensionNotSupported", "The file extension '{0}' is not supported"), FText::FromString(Extension)));
 				}
 
-				return false;
+				return 0;
 			}
 		}
 
@@ -95,7 +100,7 @@ public:
 			}
 		}
 
-		return true;
+		return 80;
 	}
 
 	virtual TSharedPtr<IMediaPlayer, ESPMode::ThreadSafe> CreatePlayer(IMediaEventSink& EventSink) override
diff --git a/Engine/Plugins/Runtime/nDisplay/Source/SharedMemoryMedia/Private/SharedMemoryMediaPlayerFactory.cpp b/Engine/Plugins/Runtime/nDisplay/Source/SharedMemoryMedia/Private/SharedMemoryMediaPlayerFactory.cpp
--- a/Engine/Plugins/Runtime/nDisplay/Source/SharedMemoryMedia/Private/SharedMemoryMediaPlayerFactory.cpp
+++ b/Engine/Plugins/Runtime/nDisplay/Source/SharedMemoryMedia/Private/SharedMemoryMediaPlayerFactory.cpp
@@ -18,6 +18,11 @@ FSharedMemoryMediaPlayerFactory::FSharedMemoryMediaPlayerFactory()
 }
 
 bool FSharedMemoryMediaPlayerFactory::CanPlayUrl(const FString& Url, const IMediaOptions* Options, TArray<FText>* OutWarnings, TArray<FText>* OutErrors) const
+{
+	return GetPlayabilityConfidenceScore(Url, Options, OutWarnings, OutErrors) > 0 ? true : false;
+}
+
+int32 FSharedMemoryMediaPlayerFactory::GetPlayabilityConfidenceScore(const FString& Url, const IMediaOptions* Options, TArray<FText>* OutWarnings, TArray<FText>* OutErrors) const
 {
 	FString Scheme;
 	FString Location;
@@ -30,7 +35,7 @@ bool FSharedMemoryMediaPlayerFactory::CanPlayUrl(const FString& Url, const IMedi
 			OutErrors->Add(LOCTEXT("NoSchemeFound", "No URI scheme found"));
 		}
 
-		return false;
+		return 0;
 	}
 
 	// see if the scheme is supported
@@ -41,10 +46,10 @@ bool FSharedMemoryMediaPlayerFactory::CanPlayUrl(const FString& Url, const IMedi
 			OutErrors->Add(FText::Format(LOCTEXT("SchemeNotSupported", "The URI scheme '{0}' is not supported"), FText::FromString(Scheme)));
 		}
 
-		return false;
+		return 0;
 	}
 
-	return true;
+	return 100;
 }
 
 TSharedPtr<IMediaPlayer, ESPMode::ThreadSafe> FSharedMemoryMediaPlayerFactory::CreatePlayer(IMediaEventSink& EventSink)
diff --git a/Engine/Plugins/Runtime/nDisplay/Source/SharedMemoryMedia/Private/SharedMemoryMediaPlayerFactory.h b/Engine/Plugins/Runtime/nDisplay/Source/SharedMemoryMedia/Private/SharedMemoryMediaPlayerFactory.h
--- a/Engine/Plugins/Runtime/nDisplay/Source/SharedMemoryMedia/Private/SharedMemoryMediaPlayerFactory.h
+++ b/Engine/Plugins/Runtime/nDisplay/Source/SharedMemoryMedia/Private/SharedMemoryMediaPlayerFactory.h
@@ -27,6 +27,7 @@ public:
 	//~ Begin IMediaPlayerFactory interface
 
 	virtual bool CanPlayUrl(const FString& Url, const IMediaOptions* Options, TArray<FText>* OutWarnings, TArray<FText>* OutErrors) const override;
+	virtual int32 GetPlayabilityConfidenceScore(const FString& Url, const IMediaOptions* Options, TArray<FText>* OutWarnings, TArray<FText>* OutErrors) const override;
 	virtual TSharedPtr<IMediaPlayer, ESPMode::ThreadSafe> CreatePlayer(IMediaEventSink& EventSink) override;
 	virtual FText GetDisplayName() const override;
 	virtual FName GetPlayerName() const override;
diff --git a/Engine/Plugins/VirtualProduction/Rivermax/RivermaxMedia/Source/RivermaxMediaFactory/Private/RivermaxMediaFactoryModule.cpp b/Engine/Plugins/VirtualProduction/Rivermax/RivermaxMedia/Source/RivermaxMediaFactory/Private/RivermaxMediaFactoryModule.cpp
--- a/Engine/Plugins/VirtualProduction/Rivermax/RivermaxMedia/Source/RivermaxMediaFactory/Private/RivermaxMediaFactoryModule.cpp
+++ b/Engine/Plugins/VirtualProduction/Rivermax/RivermaxMedia/Source/RivermaxMediaFactory/Private/RivermaxMediaFactoryModule.cpp
@@ -29,7 +29,12 @@ public:
 
 	//~ IMediaPlayerFactory interface
 
-	virtual bool CanPlayUrl(const FString& Url, const IMediaOptions* /*Options*/, TArray<FText>* /*OutWarnings*/, TArray<FText>* OutErrors) const override
+	bool CanPlayUrl(const FString& Url, const IMediaOptions* Options, TArray<FText>* OutWarnings, TArray<FText>* OutErrors) const override
+	{
+		return GetPlayabilityConfidenceScore(Url, Options, OutWarnings, OutErrors) > 0 ? true : false;
+	}
+
+	int32 GetPlayabilityConfidenceScore(const FString& Url, const IMediaOptions* /*Options*/, TArray<FText>* /*OutWarnings*/, TArray<FText>* OutErrors) const override
 	{
 		FString Scheme;
 		FString Location;
@@ -42,7 +47,7 @@ public:
 				OutErrors->Add(LOCTEXT("NoSchemeFound", "No URI scheme found"));
 			}
 
-			return false;
+			return 0;
 		}
 
 		if (!SupportedUriSchemes.Contains(Scheme))
@@ -52,10 +57,10 @@ public:
 				OutErrors->Add(FText::Format(LOCTEXT("SchemeNotSupported", "The URI scheme '{0}' is not supported"), FText::FromString(Scheme)));
 			}
 
-			return false;
+			return 0;
 		}
 
-		return true;
+		return 100;
 	}
 
 	virtual TSharedPtr<IMediaPlayer, ESPMode::ThreadSafe> CreatePlayer(IMediaEventSink& EventSink) override
diff --git a/Engine/Source/Runtime/Media/Public/IMediaPlayerFactory.h b/Engine/Source/Runtime/Media/Public/IMediaPlayerFactory.h
--- a/Engine/Source/Runtime/Media/Public/IMediaPlayerFactory.h
+++ b/Engine/Source/Runtime/Media/Public/IMediaPlayerFactory.h
@@ -143,6 +143,20 @@ public:
 		return CanPlayUrl(Url, Options, nullptr, nullptr);
 	}
 
+	/**
+	 * Returns a confidence score how sure the player thinks it is to be able to play the specified source URL.
+	 *
+	 * @param Url The media source URL to check.
+	 * @param Options Optional media player parameters.
+	 * @param OutWarnings Will contain warning messages (optional).
+	 * @param OutErrors will contain error messages (optional).
+	 * @return A score giving a probability (0-100 (or higher if necessary)) whether or not the media can be played.
+	 */
+	virtual int32 GetPlayabilityConfidenceScore(const FString& Url, const IMediaOptions* Options, TArray<FText>* OutWarnings, TArray<FText>* OutErrors) const
+	{
+		return CanPlayUrl(Url, Options, OutWarnings, OutErrors) ? 100 : 0;
+	}
+
 	/**
 	 * Whether the player works on the given platform.
 	 *
diff --git a/Engine/Source/Runtime/MediaUtils/Private/MediaPlayerFacade.cpp b/Engine/Source/Runtime/MediaUtils/Private/MediaPlayerFacade.cpp
--- a/Engine/Source/Runtime/MediaUtils/Private/MediaPlayerFacade.cpp
+++ b/Engine/Source/Runtime/MediaUtils/Private/MediaPlayerFacade.cpp
@@ -1815,63 +1815,66 @@ IMediaPlayerFactory* FMediaPlayerFacade::GetPlayerFactoryForUrl(const FString& U
 		return Factory;
 	}
 
-
-
 	//
 	// Try to find a fitting player with no explicit name given
 	//
-
-
-	// Can any existing player play the URL?
-	if (Player.IsValid())
+	const TArray<IMediaPlayerFactory*>& PlayerFactories = MediaModule->GetPlayerFactories();
+	if (PlayerFactories.Num() == 0)
 	{
-		IMediaPlayerFactory* Factory = MediaModule->GetPlayerFactory(Player->GetPlayerPluginGUID());
-
-		if ((Factory != nullptr) && Factory->CanPlayUrl(Url, Options))
+		UE_LOG(LogMediaUtils, Error, TEXT("Cannot play %s: no media player plug-ins are installed and enabled in this project"), *Url);
+		return nullptr;
+	}
+	struct FCandidate
+	{
+		FName Name;
+		IMediaPlayerFactory* Factory = nullptr;
+		int32 ConfidenceScore = 0;
+	};
+	TArray<FCandidate> Candidates;
+	const FString RunningPlatformName(FPlatformProperties::IniPlatformName());
+	for(IMediaPlayerFactory* Factory : PlayerFactories)
+	{
+		if (Factory->SupportsPlatform(RunningPlatformName))
 		{
-			// Yes...
-			return Factory;
+			int32 ConfidenceScore = Factory->GetPlayabilityConfidenceScore(Url, Options, nullptr, nullptr);
+			if (ConfidenceScore > 0)
+			{
+				FCandidate& Candidate = Candidates.Emplace_GetRef();
+				Candidate.Name = Factory->GetPlayerName();
+				Candidate.Factory = Factory;
+				Candidate.ConfidenceScore = ConfidenceScore;
+			}
 		}
 	}
-
-	// Try to auto-select new player...
-	const FString RunningPlatformName(FPlatformProperties::IniPlatformName());
-	const TArray<IMediaPlayerFactory*>& PlayerFactories = MediaModule->GetPlayerFactories();
-
-	for (IMediaPlayerFactory* Factory : PlayerFactories)
+	Candidates.Sort([](const FCandidate& c1, const FCandidate& c2)
 	{
-		if (!Factory->SupportsPlatform(RunningPlatformName) || !Factory->CanPlayUrl(Url, Options))
+		// If both factories are equally confident, sort alphabetically by name.
+		if (c1.ConfidenceScore == c2.ConfidenceScore)
 		{
-			continue;
+			return c1.Name.ToString() < c2.Name.ToString();
 		}
-
-		return Factory;
+		// Sort by descending confidence score.
+		return c1.ConfidenceScore > c2.ConfidenceScore;
+	});
+	if (Candidates.Num())
+	{
+		return Candidates[0].Factory;
 	}
-
 	//
 	// No suitable player found!
 	//
-	if (PlayerFactories.Num() > 0)
+	UE_LOG(LogMediaUtils, Error, TEXT("Cannot play %s, because none of the enabled media player plug-ins support it:"), *Url);
+	for (IMediaPlayerFactory* Factory : PlayerFactories)
 	{
-		UE_LOG(LogMediaUtils, Error, TEXT("Cannot play %s, because none of the enabled media player plug-ins support it:"), *Url);
-
-		for (IMediaPlayerFactory* Factory : PlayerFactories)
+		if (Factory->SupportsPlatform(RunningPlatformName))
 		{
-			if (Factory->SupportsPlatform(RunningPlatformName))
-			{
-				UE_LOG(LogMediaUtils, Log, TEXT("| %s (URI scheme or file extension not supported)"), *Factory->GetPlayerName().ToString());
-			}
-			else
-			{
-				UE_LOG(LogMediaUtils, Log, TEXT("| %s (only available on %s, but not on %s)"), *Factory->GetPlayerName().ToString(), *FString::Join(Factory->GetSupportedPlatforms(), TEXT(", ")), *RunningPlatformName);
-			}
+			UE_LOG(LogMediaUtils, Log, TEXT("| %s (URI scheme or file extension not supported)"), *Factory->GetPlayerName().ToString());
+		}
+		else
+		{
+			UE_LOG(LogMediaUtils, Log, TEXT("| %s (only available on %s, but not on %s)"), *Factory->GetPlayerName().ToString(), *FString::Join(Factory->GetSupportedPlatforms(), TEXT(", ")), *RunningPlatformName);
 		}
 	}
-	else
-	{
-		UE_LOG(LogMediaUtils, Error, TEXT("Cannot play %s: no media player plug-ins are installed and enabled in this project"), *Url);
-	}
-
 	return nullptr;
 }
 
