Electra: Fixed occasional nullptr crash when closing down a player on PC

#rnx

diff --git a/Engine/Plugins/Media/ElectraPlayer/Source/ElectraPlayerRuntime/Private/Runtime/Decoder/Windows/ElectraVideoDecoder_PC.cpp b/Engine/Plugins/Media/ElectraPlayer/Source/ElectraPlayerRuntime/Private/Runtime/Decoder/Windows/ElectraVideoDecoder_PC.cpp
--- a/Engine/Plugins/Media/ElectraPlayer/Source/ElectraPlayerRuntime/Private/Runtime/Decoder/Windows/ElectraVideoDecoder_PC.cpp
+++ b/Engine/Plugins/Media/ElectraPlayer/Source/ElectraPlayerRuntime/Private/Runtime/Decoder/Windows/ElectraVideoDecoder_PC.cpp
@@ -231,7 +231,7 @@ void FElectraPlayerVideoDecoderOutputPC::InitializeWithResource(const TRefCountP
 			}
 			InOutD3D12ResourcePool = D3D12ResourcePool;
 		}
-		
+
 		FElectraMediaDecoderOutputBufferPool_DX12::FOutputData OutputData;
 		if (!D3D12ResourcePool->AllocateOutputDataAsTexture(OutputData, SampleDim.X, SampleDim.Y, DXGIFmt))
 		{
@@ -374,6 +374,10 @@ void FElectraPlayerVideoDecoderOutputPC::InitializeWithResource(const TRefCountP
 
 void FElectraPlayerVideoDecoderOutputPC::TriggerDataCopy(TRefCountPtr<ID3D12GraphicsCommandList> D3DCmdList, TRefCountPtr<ID3D12CommandAllocator> D3DCmdAllocator, TRefCountPtr<ID3D12Fence> D3DFence, uint64 FenceValue, const FElectraDecoderOutputSync& OutputSync, Electra::IVideoDecoderResourceDelegate* InResourceDelegate)
 {
+	if (!InResourceDelegate)
+	{
+		return;
+	}
 	// Trigger copy (this will eventually execute on the submission thread of RHI if running in UE)
 	// (note: we pass in all of FElectraDecoderOutputSync to guarantee any references needed to make the decoder output sync work are passed along, too!)
 	InResourceDelegate->ExecuteCodeWithCopyCommandQueueUsage([CmdList = D3DCmdList, CmdAllocator = D3DCmdAllocator, DestFence = D3DFence, DestFenceValue = FenceValue, OutputSync](ID3D12CommandQueue* D3DCmdQueue)
diff --git a/Engine/Plugins/Media/ElectraPlayer/Source/ElectraPlayerRuntime/Private/Runtime/Decoder/Windows/WindowsElectraDecoderResourceManager.cpp b/Engine/Plugins/Media/ElectraPlayer/Source/ElectraPlayerRuntime/Private/Runtime/Decoder/Windows/WindowsElectraDecoderResourceManager.cpp
--- a/Engine/Plugins/Media/ElectraPlayer/Source/ElectraPlayerRuntime/Private/Runtime/Decoder/Windows/WindowsElectraDecoderResourceManager.cpp
+++ b/Engine/Plugins/Media/ElectraPlayer/Source/ElectraPlayerRuntime/Private/Runtime/Decoder/Windows/WindowsElectraDecoderResourceManager.cpp
@@ -152,11 +152,8 @@ bool FElectraDecoderResourceManagerWindows::RunCodeAsync(TFunction<void()>&& Cod
 bool FElectraDecoderResourceManagerWindows::SetupRenderBufferFromDecoderOutputFromMFSample(IMediaRenderer::IBuffer* InOutBufferToSetup, TSharedPtr<FParamDict, ESPMode::ThreadSafe> InOutBufferPropertes, TSharedPtr<IElectraDecoderVideoOutput, ESPMode::ThreadSafe> InDecoderOutput, FElectraDecoderResourceManagerWindows::IDecoderPlatformResource* InPlatformSpecificResource)
 {
 	TSharedPtr<FElectraPlayerVideoDecoderOutputPC, ESPMode::ThreadSafe> DecoderOutput = InOutBufferToSetup->GetBufferProperties().GetValue(RenderOptionKeys::Texture).GetSharedPointer<FElectraPlayerVideoDecoderOutputPC>();
-	FInstanceVars* Vars = static_cast<FInstanceVars*>(InPlatformSpecificResource);
-
-	if (DecoderOutput.IsValid() && Vars != nullptr)
+	if (DecoderOutput.IsValid())
 	{
-		TSharedPtr<IVideoDecoderResourceDelegate, ESPMode::ThreadSafe> PinnedResourceDelegate = Vars->VideoDecoderResourceDelegate.Pin();
 
 		FElectraVideoDecoderOutputCropValues Crop = InDecoderOutput->GetCropValues();
 		InOutBufferPropertes->Set(IDecoderOutputOptionNames::Width, FVariantValue((int64)InDecoderOutput->GetWidth()));
@@ -235,7 +232,7 @@ bool FElectraDecoderResourceManagerWindows::SetupRenderBufferFromDecoderOutputFr
 			{
 				//
 				// DX12 (with DX11 decode device) & non-DX
-				// 
+				//
 				// (access buffer for CPU use)
 				//
 				TRefCountPtr<IMF2DBuffer> Buffer2D;
@@ -407,7 +404,8 @@ bool FElectraDecoderResourceManagerWindows::SetupRenderBufferFromDecoderOutput(I
 
 	TSharedPtr<FElectraPlayerVideoDecoderOutputPC, ESPMode::ThreadSafe> DecoderOutput = InOutBufferToSetup->GetBufferProperties().GetValue(RenderOptionKeys::Texture).GetSharedPointer<FElectraPlayerVideoDecoderOutputPC>();
 	FInstanceVars* Vars = static_cast<FInstanceVars*>(InPlatformSpecificResource);
-	if (DecoderOutput.IsValid())
+	TSharedPtr<IVideoDecoderResourceDelegate, ESPMode::ThreadSafe> PinnedResourceDelegate = Vars ? Vars->VideoDecoderResourceDelegate.Pin() : nullptr;
+	if (DecoderOutput.IsValid() && Vars != nullptr && PinnedResourceDelegate.IsValid())
 	{
 		//
 		// Image buffers?
@@ -474,7 +472,7 @@ bool FElectraDecoderResourceManagerWindows::SetupRenderBufferFromDecoderOutput(I
 				case EElectraDecoderPlatformPixelEncoding::ARGB_BigEndian:	DecPixEnc = EVideoDecoderPixelEncoding::ARGB_BigEndian; break;
 				default: DecPixEnc = EVideoDecoderPixelEncoding::Native; break;
 			}
-			
+
 			int32 Pitch = ImageBuffers->GetBufferPitchByIndex(0);
 
 			InOutBufferPropertes->Set(IDecoderOutputOptionNames::PixelFormat, FVariantValue((int64)RHIPixFmt));
@@ -531,7 +529,7 @@ bool FElectraDecoderResourceManagerWindows::SetupRenderBufferFromDecoderOutput(I
 							FElectraDecoderOutputSync OutputSync;
 							ImageBuffers->GetBufferTextureSyncByIndex(0, OutputSync);
 
-							DecoderOutput->InitializeWithResource(D3D12Device, Resource, Pitch, OutputSync, FIntPoint(InDecoderOutput->GetDecodedWidth(), InDecoderOutput->GetDecodedHeight()), InOutBufferPropertes, Vars->VideoDecoderResourceDelegate,
+							DecoderOutput->InitializeWithResource(D3D12Device, Resource, Pitch, OutputSync, FIntPoint(InDecoderOutput->GetDecodedWidth(), InDecoderOutput->GetDecodedHeight()), InOutBufferPropertes, PinnedResourceDelegate,
 																  Vars->D3D12ResourcePool, Vars->MaxWidth, Vars->MaxHeight, Vars->MaxOutputBuffers);
 							return true;
 						}
