Electra: Fix for duplicate LOCTEXT

#rnx

diff --git a/Engine/Plugins/Media/ElectraPlayer/Source/ElectraPlayerFactory/Private/ElectraPlayerFactoryModule.cpp b/Engine/Plugins/Media/ElectraPlayer/Source/ElectraPlayerFactory/Private/ElectraPlayerFactoryModule.cpp
--- a/Engine/Plugins/Media/ElectraPlayer/Source/ElectraPlayerFactory/Private/ElectraPlayerFactoryModule.cpp
+++ b/Engine/Plugins/Media/ElectraPlayer/Source/ElectraPlayerFactory/Private/ElectraPlayerFactoryModule.cpp
@@ -52,7 +52,7 @@ public:
 		{
 			if (OutErrors != nullptr)
 			{
-				OutErrors->Add(FText::Format(LOCTEXT("SchemeNotSupported", "The URI '{0}' could not be parsed"), FText::FromString(Url)));
+				OutErrors->Add(FText::Format(LOCTEXT("MalformedURI", "The URI '{0}' could not be parsed"), FText::FromString(Url)));
 			}
 		}
 		// Check scheme
