ElectraHTTPManager: Added getter for default User-Agent

#rnx

diff --git a/Engine/Plugins/Media/ElectraPlayer/Source/ElectraPlayerRuntime/Private/Runtime/HTTP/HTTPManager.cpp b/Engine/Plugins/Media/ElectraPlayer/Source/ElectraPlayerRuntime/Private/Runtime/HTTP/HTTPManager.cpp
--- a/Engine/Plugins/Media/ElectraPlayer/Source/ElectraPlayerRuntime/Private/Runtime/HTTP/HTTPManager.cpp
+++ b/Engine/Plugins/Media/ElectraPlayer/Source/ElectraPlayerRuntime/Private/Runtime/HTTP/HTTPManager.cpp
@@ -437,6 +437,12 @@ namespace Electra
 	/***************************************************************************************************************************************************/
 	/***************************************************************************************************************************************************/
 
+	FString IElectraHttpManager::GetDefaultUserAgent()
+	{
+		return FString(ELECTRA_HTTPMANAGER_USER_AGENT);
+	}
+
+
 	TSharedPtrTS<IElectraHttpManager> IElectraHttpManager::Create()
 	{
 		return FElectraHttpManager::Create();
@@ -717,7 +723,7 @@ namespace Electra
 			}
 			else
 			{
-				Handle->HttpRequest->SetUserAgent(ELECTRA_HTTPMANAGER_USER_AGENT);
+				Handle->HttpRequest->SetUserAgent(GetDefaultUserAgent());
 			}
 			Handle->HttpRequest->AllowCompression(!Request->Parameters.AcceptEncoding.GetWithDefault(TEXT("")).Equals(TEXT("identity")));
 			#if defined(ELECTRA_HTTPMANAGER_ALLOW_UNSAFE_CONNECTIONS) && ELECTRA_HTTPMANAGER_ALLOW_UNSAFE_CONNECTIONS
diff --git a/Engine/Plugins/Media/ElectraPlayer/Source/ElectraPlayerRuntime/Private/Runtime/HTTP/HTTPManager.h b/Engine/Plugins/Media/ElectraPlayer/Source/ElectraPlayerRuntime/Private/Runtime/HTTP/HTTPManager.h
--- a/Engine/Plugins/Media/ElectraPlayer/Source/ElectraPlayerRuntime/Private/Runtime/HTTP/HTTPManager.h
+++ b/Engine/Plugins/Media/ElectraPlayer/Source/ElectraPlayerRuntime/Private/Runtime/HTTP/HTTPManager.h
@@ -179,6 +179,8 @@ namespace Electra
 	public:
 		static TSharedPtrTS<IElectraHttpManager> Create();
 
+		static FString GetDefaultUserAgent();
+
 		virtual ~IElectraHttpManager() = default;
 
 		struct FRequest;
