ElectraPlayer: Updating metadata related handling

#rnx

diff --git a/Engine/Plugins/Media/ElectraPlayer/Source/ElectraPlayerRuntime/Private/Runtime/Player/AdaptiveStreamingPlayer.cpp b/Engine/Plugins/Media/ElectraPlayer/Source/ElectraPlayerRuntime/Private/Runtime/Player/AdaptiveStreamingPlayer.cpp
--- a/Engine/Plugins/Media/ElectraPlayer/Source/ElectraPlayerRuntime/Private/Runtime/Player/AdaptiveStreamingPlayer.cpp
+++ b/Engine/Plugins/Media/ElectraPlayer/Source/ElectraPlayerRuntime/Private/Runtime/Player/AdaptiveStreamingPlayer.cpp
@@ -2003,8 +2003,26 @@ void FAdaptiveStreamingPlayer::HandleMetadataChanges()
 	}
 
 	// Handle media metadata updates.
-	if (MediaMetadataUpdates.Handle(CurrentTime))
+	FMediaMetadataUpdate::EResult MetaResult = MediaMetadataUpdates.Handle(CurrentTime);
+	if (MetaResult != FMediaMetadataUpdate::EResult::NoChange)
 	{
+		if (MetaResult == FMediaMetadataUpdate::EResult::ChangedAndUpdate && Manifest.IsValid())
+		{
+			Manifest->UpdateRunningMetaData(MediaMetadataUpdates.GetActive());
+			if (ActivePeriods.Num())
+			{
+				TArray<FTrackMetadata> MetadataVideo;
+				TArray<FTrackMetadata> MetadataAudio;
+				TArray<FTrackMetadata> MetadataSubtitle;
+				// FIXME: We are currently assuming that this happens only for audio casts where
+				//        we won't have more than just one period.
+				ActivePeriods[0].Period->GetMetaData(MetadataVideo, EStreamType::Video);
+				ActivePeriods[0].Period->GetMetaData(MetadataAudio, EStreamType::Audio);
+				ActivePeriods[0].Period->GetMetaData(MetadataSubtitle, EStreamType::Subtitle);
+				PlaybackState.SetTrackMetadata(MetadataVideo, MetadataAudio, MetadataSubtitle);
+				PlaybackState.SetHaveMetadata(true);
+			}
+		}
 		DispatchEvent(FMetricEvent::ReportMediaMetadataChanged(MediaMetadataUpdates.GetActive()));
 	}
 }
@@ -2085,7 +2103,7 @@ void FAdaptiveStreamingPlayer::HandleSessionMessage(TSharedPtrTS<IPlayerMessage>
 	else if (SessionMessage->GetType() == FPlaylistMetadataUpdateMessage::Type())
 	{
 		FPlaylistMetadataUpdateMessage* pMsg = static_cast<FPlaylistMetadataUpdateMessage*>(SessionMessage.Get());
-		MediaMetadataUpdates.AddEntry(pMsg->GetValidFrom(), pMsg->GetMetadata());
+		MediaMetadataUpdates.AddEntry(pMsg->GetValidFrom(), pMsg->GetMetadata(), pMsg->GetTriggerInternalRefresh());
 	}
 	// License key?
 	else if (SessionMessage->GetType() == FLicenseKeyMessage::Type())
@@ -2623,10 +2641,18 @@ void FAdaptiveStreamingPlayer::InternalHandlePendingStartRequest(const FTimeValu
 
 						check(Seekable.End.IsValid());
 						StartAt.Time = Seekable.End;
-						// DASH Live streams provide the means to join at an exact wallclock time, so enable this.
-						if (ManifestType == EMediaFormatType::DASH)
+						switch(Manifest->GetLiveEdgePlayMode())
 						{
-							PlaybackState.SetShouldPlayOnLiveEdge(true);
+							case IManifest::ELiveEdgePlayMode::Default:
+							case IManifest::ELiveEdgePlayMode::Always:
+							{
+								PlaybackState.SetShouldPlayOnLiveEdge(true);
+								break;
+							}
+							default:
+							{
+								break;
+							}
 						}
 					}
 				}
diff --git a/Engine/Plugins/Media/ElectraPlayer/Source/ElectraPlayerRuntime/Private/Runtime/Player/AdaptiveStreamingPlayerInternal.h b/Engine/Plugins/Media/ElectraPlayer/Source/ElectraPlayerRuntime/Private/Runtime/Player/AdaptiveStreamingPlayerInternal.h
--- a/Engine/Plugins/Media/ElectraPlayer/Source/ElectraPlayerRuntime/Private/Runtime/Player/AdaptiveStreamingPlayerInternal.h
+++ b/Engine/Plugins/Media/ElectraPlayer/Source/ElectraPlayerRuntime/Private/Runtime/Player/AdaptiveStreamingPlayerInternal.h
@@ -2030,11 +2030,12 @@ private:
 			// Do NOT reset the active metadata, but the time it became valid!
 			ActiveSince.SetToInvalid();
 		}
-		void AddEntry(const FTimeValue& InValidFrom, const TSharedPtrTS<UtilsMP4::FMetadataParser>& InMetadata)
+		void AddEntry(const FTimeValue& InValidFrom, const TSharedPtrTS<UtilsMP4::FMetadataParser>& InMetadata, bool bInTriggerInternalRefresh)
 		{
 			FEntry& e = NextEntries.Emplace_GetRef();
 			e.ValidFrom = InValidFrom.IsValid() ? InValidFrom : FTimeValue::GetZero();
 			e.Metadata = InMetadata;
+			e.bTriggerInternalRefresh = bInTriggerInternalRefresh;
 			NextEntries.StableSort([](const FEntry& a, const FEntry& b)
 			{
 				const FTimeValue& t1 = a.ValidFrom;
@@ -2044,7 +2045,13 @@ private:
 				return (s1 == s2 && t1 < t2) || (s1 < s2);
 			});
 		}
-		bool Handle(const FTimeValue& InAtTime);
+		enum class EResult
+		{
+			NoChange,
+			Changed,
+			ChangedAndUpdate
+		};
+		EResult Handle(const FTimeValue& InAtTime);
 		TSharedPtrTS<UtilsMP4::FMetadataParser> GetActive() const
 		{
 			return ActiveMetadata;
@@ -2054,6 +2061,7 @@ private:
 		{
 			FTimeValue ValidFrom;
 			TSharedPtrTS<UtilsMP4::FMetadataParser> Metadata;
+			bool bTriggerInternalRefresh = false;
 		};
 		TArray<FEntry> NextEntries;
 		TSharedPtrTS<UtilsMP4::FMetadataParser> ActiveMetadata;
diff --git a/Engine/Plugins/Media/ElectraPlayer/Source/ElectraPlayerRuntime/Private/Runtime/Player/AdaptiveStreamingPlayerManifest.cpp b/Engine/Plugins/Media/ElectraPlayer/Source/ElectraPlayerRuntime/Private/Runtime/Player/AdaptiveStreamingPlayerManifest.cpp
--- a/Engine/Plugins/Media/ElectraPlayer/Source/ElectraPlayerRuntime/Private/Runtime/Player/AdaptiveStreamingPlayerManifest.cpp
+++ b/Engine/Plugins/Media/ElectraPlayer/Source/ElectraPlayerRuntime/Private/Runtime/Player/AdaptiveStreamingPlayerManifest.cpp
@@ -331,8 +331,8 @@ bool FAdaptiveStreamingPlayer::SelectManifest()
 
 			double minBufTimeMPD = Manifest->GetMinBufferTime().GetAsSeconds();
 			PlayerConfig.InitialBufferMinTimeAvailBeforePlayback = Utils::Min(minBufTimeMPD, PlayerConfig.InitialBufferMinTimeAvailBeforePlayback);
-			PlayerConfig.SeekBufferMinTimeAvailBeforePlayback    = Utils::Min(minBufTimeMPD, PlayerConfig.SeekBufferMinTimeAvailBeforePlayback);
-			PlayerConfig.RebufferMinTimeAvailBeforePlayback 	 = Utils::Min(minBufTimeMPD, PlayerConfig.RebufferMinTimeAvailBeforePlayback);
+			PlayerConfig.SeekBufferMinTimeAvailBeforePlayback = Utils::Min(minBufTimeMPD, PlayerConfig.SeekBufferMinTimeAvailBeforePlayback);
+			PlayerConfig.RebufferMinTimeAvailBeforePlayback = Utils::Min(minBufTimeMPD, PlayerConfig.RebufferMinTimeAvailBeforePlayback);
 
 			// For an mp4 or mkv stream we can now get rid of the manifest reader. It is no longer needed and we don't need to have it linger.
 			if (ManifestType == EMediaFormatType::ISOBMFF || ManifestType == EMediaFormatType::MKV)
@@ -378,11 +378,11 @@ void FAdaptiveStreamingPlayer::UpdateManifest()
 	}
 }
 
-
-bool FAdaptiveStreamingPlayer::FMediaMetadataUpdate::Handle(const FTimeValue& InAtTime)
+FAdaptiveStreamingPlayer::FMediaMetadataUpdate::EResult FAdaptiveStreamingPlayer::FMediaMetadataUpdate::Handle(const FTimeValue& InAtTime)
 {
 	TSharedPtrTS<UtilsMP4::FMetadataParser> NextMetadata;
 	FTimeValue NextActiveTime;
+	bool bTriggerInternalRefresh = false;
 	while(NextEntries.Num())
 	{
 		// Make the first metadata available right away if there is none yet and the time is not valid either.
@@ -390,6 +390,7 @@ bool FAdaptiveStreamingPlayer::FMediaMetadataUpdate::Handle(const FTimeValue& In
 		{
 			NextMetadata = NextEntries[0].Metadata;
 			NextActiveTime = NextEntries[0].ValidFrom;
+			bTriggerInternalRefresh = NextEntries[0].bTriggerInternalRefresh;
 			break;
 		}
 		else
@@ -409,6 +410,7 @@ bool FAdaptiveStreamingPlayer::FMediaMetadataUpdate::Handle(const FTimeValue& In
 					{
 						NextMetadata = NextEntries[0].Metadata;
 						NextActiveTime = NextEntries[0].ValidFrom;
+						bTriggerInternalRefresh = NextEntries[0].bTriggerInternalRefresh;
 					}
 					NextEntries.RemoveAt(0);
 				}
@@ -428,12 +430,11 @@ bool FAdaptiveStreamingPlayer::FMediaMetadataUpdate::Handle(const FTimeValue& In
 		bool bChanged = !ActiveMetadata.IsValid() || ActiveMetadata->IsDifferentFrom(*NextMetadata);
 		ActiveMetadata = MoveTemp(NextMetadata);
 		ActiveSince = NextActiveTime;
-		return bChanged;
+		return bChanged ? bTriggerInternalRefresh ? EResult::ChangedAndUpdate : EResult::Changed : EResult::NoChange;
 	}
-	return false;
+
+	return EResult::NoChange;
 }
 
 
 } // namespace Electra
-
-
diff --git a/Engine/Plugins/Media/ElectraPlayer/Source/ElectraPlayerRuntime/Private/Runtime/Player/DASH/ManifestDASH.cpp b/Engine/Plugins/Media/ElectraPlayer/Source/ElectraPlayerRuntime/Private/Runtime/Player/DASH/ManifestDASH.cpp
--- a/Engine/Plugins/Media/ElectraPlayer/Source/ElectraPlayerRuntime/Private/Runtime/Player/DASH/ManifestDASH.cpp
+++ b/Engine/Plugins/Media/ElectraPlayer/Source/ElectraPlayerRuntime/Private/Runtime/Player/DASH/ManifestDASH.cpp
@@ -435,6 +435,11 @@ FTimeValue FManifestDASH::GetDesiredLiveLatency() const
 	return Manifest.IsValid() ? Manifest->GetDesiredLiveLatency() : FTimeValue();
 }
 
+IManifest::ELiveEdgePlayMode FManifestDASH::GetLiveEdgePlayMode() const
+{
+	return IManifest::ELiveEdgePlayMode::Default;
+}
+
 
 TSharedPtrTS<IProducerReferenceTimeInfo> FManifestDASH::GetProducerReferenceTimeInfo(int64 ID) const
 {
@@ -469,6 +474,10 @@ void FManifestDASH::GetTrackMetadata(TArray<FTrackMetadata>& OutMetadata, EStrea
 		Manifest->GetPeriods()[0]->GetMetaData(OutMetadata, StreamType);
 	}
 }
+void FManifestDASH::UpdateRunningMetaData(TSharedPtrTS<UtilsMP4::FMetadataParser> InUpdatedMetaData)
+{
+	// No-op.
+}
 
 void FManifestDASH::UpdateDynamicRefetchCounter()
 {
diff --git a/Engine/Plugins/Media/ElectraPlayer/Source/ElectraPlayerRuntime/Private/Runtime/Player/DASH/ManifestDASH.h b/Engine/Plugins/Media/ElectraPlayer/Source/ElectraPlayerRuntime/Private/Runtime/Player/DASH/ManifestDASH.h
--- a/Engine/Plugins/Media/ElectraPlayer/Source/ElectraPlayerRuntime/Private/Runtime/Player/DASH/ManifestDASH.h
+++ b/Engine/Plugins/Media/ElectraPlayer/Source/ElectraPlayerRuntime/Private/Runtime/Player/DASH/ManifestDASH.h
@@ -32,9 +32,11 @@ public:
 	void ClearDefaultEndTime() override;
 	FTimeValue GetMinBufferTime() const override;
 	FTimeValue GetDesiredLiveLatency() const override;
+	ELiveEdgePlayMode GetLiveEdgePlayMode() const override;
 	TRangeSet<double> GetPossiblePlaybackRates(EPlayRateType InForType) const override;
 	TSharedPtrTS<IProducerReferenceTimeInfo> GetProducerReferenceTimeInfo(int64 ID) const override;
 	void GetTrackMetadata(TArray<FTrackMetadata>& OutMetadata, EStreamType StreamType) const override;
+	void UpdateRunningMetaData(TSharedPtrTS<UtilsMP4::FMetadataParser> InUpdatedMetaData) override;
 	void UpdateDynamicRefetchCounter() override;
 	void TriggerClockSync(IManifest::EClockSyncType InClockSyncType) override;
 	void TriggerPlaylistRefresh() override;
diff --git a/Engine/Plugins/Media/ElectraPlayer/Source/ElectraPlayerRuntime/Private/Runtime/Player/DASH/PlaylistReaderDASH_Internal.h b/Engine/Plugins/Media/ElectraPlayer/Source/ElectraPlayerRuntime/Private/Runtime/Player/DASH/PlaylistReaderDASH_Internal.h
--- a/Engine/Plugins/Media/ElectraPlayer/Source/ElectraPlayerRuntime/Private/Runtime/Player/DASH/PlaylistReaderDASH_Internal.h
+++ b/Engine/Plugins/Media/ElectraPlayer/Source/ElectraPlayerRuntime/Private/Runtime/Player/DASH/PlaylistReaderDASH_Internal.h
@@ -764,6 +764,9 @@ public:
 			}
 		}
 
+		void UpdateRunningMetaData(const FString& InKindOfValue, const FVariant& InNewValue) override
+		{ }
+
 	private:
 		friend class FManifestDASHInternal;
 		TArray<TSharedPtrTS<FAdaptationSet>> AdaptationSets;
diff --git a/Engine/Plugins/Media/ElectraPlayer/Source/ElectraPlayerRuntime/Private/Runtime/Player/DASH/StreamReaderDASH.cpp b/Engine/Plugins/Media/ElectraPlayer/Source/ElectraPlayerRuntime/Private/Runtime/Player/DASH/StreamReaderDASH.cpp
--- a/Engine/Plugins/Media/ElectraPlayer/Source/ElectraPlayerRuntime/Private/Runtime/Player/DASH/StreamReaderDASH.cpp
+++ b/Engine/Plugins/Media/ElectraPlayer/Source/ElectraPlayerRuntime/Private/Runtime/Player/DASH/StreamReaderDASH.cpp
@@ -1260,7 +1260,7 @@ void FStreamReaderDASH::FStreamHandler::HandleRequestMP4()
 							{
 								FTimeValue StartTime = md == MoofMetadata ? Request->GetFirstPTS() : TimeOffset;
 								StartTime.SetSequenceIndex(Request->TimestampSequenceIndex);
-  								PlayerSessionService->SendMessageToPlayer(FPlaylistMetadataUpdateMessage::Create(StartTime, MediaMetadata));
+  								PlayerSessionService->SendMessageToPlayer(FPlaylistMetadataUpdateMessage::Create(StartTime, MediaMetadata, false));
 							}
 						}
 
diff --git a/Engine/Plugins/Media/ElectraPlayer/Source/ElectraPlayerRuntime/Private/Runtime/Player/HLS/ManifestBuilderHLS.h b/Engine/Plugins/Media/ElectraPlayer/Source/ElectraPlayerRuntime/Private/Runtime/Player/HLS/ManifestBuilderHLS.h
--- a/Engine/Plugins/Media/ElectraPlayer/Source/ElectraPlayerRuntime/Private/Runtime/Player/HLS/ManifestBuilderHLS.h
+++ b/Engine/Plugins/Media/ElectraPlayer/Source/ElectraPlayerRuntime/Private/Runtime/Player/HLS/ManifestBuilderHLS.h
@@ -489,7 +489,7 @@ public:
 		return Bitrate;
 	}
 	int32 GetQualityIndex() const override
-	{ 
+	{
 		return QualityIndex;
 	}
 	bool CanBePlayed() const override
@@ -639,6 +639,9 @@ public:
 		}
 	}
 
+	void UpdateRunningMetaData(const FString& InKindOfValue, const FVariant& InNewValue) override
+	{ }
+
 	TSharedPtrTS<IPlaybackAssetAdaptationSet> GetAdaptationSetByTypeAndUniqueIdentifier(EStreamType OfStreamType, const FString& InUniqueIdentifier) const
 	{
 		switch(OfStreamType)
diff --git a/Engine/Plugins/Media/ElectraPlayer/Source/ElectraPlayerRuntime/Private/Runtime/Player/HLS/ManifestHLS.cpp b/Engine/Plugins/Media/ElectraPlayer/Source/ElectraPlayerRuntime/Private/Runtime/Player/HLS/ManifestHLS.cpp
--- a/Engine/Plugins/Media/ElectraPlayer/Source/ElectraPlayerRuntime/Private/Runtime/Player/HLS/ManifestHLS.cpp
+++ b/Engine/Plugins/Media/ElectraPlayer/Source/ElectraPlayerRuntime/Private/Runtime/Player/HLS/ManifestHLS.cpp
@@ -143,7 +143,7 @@ IManifest::EType FManifestHLS::GetPresentationType() const
 
 TSharedPtrTS<const FLowLatencyDescriptor> FManifestHLS::GetLowLatencyDescriptor() const
 {
-	return nullptr; 
+	return nullptr;
 }
 
 FTimeValue FManifestHLS::GetAnchorTime() const
@@ -229,6 +229,11 @@ void FManifestHLS::GetTrackMetadata(TArray<FTrackMetadata>& OutMetadata, EStream
 	}
 }
 
+void FManifestHLS::UpdateRunningMetaData(TSharedPtrTS<UtilsMP4::FMetadataParser> InUpdatedMetaData)
+{
+	// No-op.
+}
+
 
 void FManifestHLS::UpdateDynamicRefetchCounter()
 {
@@ -262,6 +267,12 @@ FTimeValue FManifestHLS::GetDesiredLiveLatency() const
 	return ll;
 }
 
+IManifest::ELiveEdgePlayMode FManifestHLS::GetLiveEdgePlayMode() const
+{
+	return IManifest::ELiveEdgePlayMode::Never;
+}
+
+
 TRangeSet<double> FManifestHLS::GetPossiblePlaybackRates(EPlayRateType InForType) const
 {
 	TRangeSet<double> Ranges;
diff --git a/Engine/Plugins/Media/ElectraPlayer/Source/ElectraPlayerRuntime/Private/Runtime/Player/HLS/ManifestHLS.h b/Engine/Plugins/Media/ElectraPlayer/Source/ElectraPlayerRuntime/Private/Runtime/Player/HLS/ManifestHLS.h
--- a/Engine/Plugins/Media/ElectraPlayer/Source/ElectraPlayerRuntime/Private/Runtime/Player/HLS/ManifestHLS.h
+++ b/Engine/Plugins/Media/ElectraPlayer/Source/ElectraPlayerRuntime/Private/Runtime/Player/HLS/ManifestHLS.h
@@ -30,9 +30,11 @@ public:
 	void ClearDefaultEndTime() override;
 	FTimeValue GetMinBufferTime() const override;
 	FTimeValue GetDesiredLiveLatency() const override;
+	ELiveEdgePlayMode GetLiveEdgePlayMode() const override;
 	TRangeSet<double> GetPossiblePlaybackRates(EPlayRateType InForType) const override;
 	TSharedPtrTS<IProducerReferenceTimeInfo> GetProducerReferenceTimeInfo(int64 ID) const override;
 	void GetTrackMetadata(TArray<FTrackMetadata>& OutMetadata, EStreamType StreamType) const override;
+	void UpdateRunningMetaData(TSharedPtrTS<UtilsMP4::FMetadataParser> InUpdatedMetaData) override;
 	void UpdateDynamicRefetchCounter() override;
 	void TriggerClockSync(EClockSyncType InClockSyncType) override;
 	void TriggerPlaylistRefresh() override;
diff --git a/Engine/Plugins/Media/ElectraPlayer/Source/ElectraPlayerRuntime/Private/Runtime/Player/Manifest.h b/Engine/Plugins/Media/ElectraPlayer/Source/ElectraPlayerRuntime/Private/Runtime/Player/Manifest.h
--- a/Engine/Plugins/Media/ElectraPlayer/Source/ElectraPlayerRuntime/Private/Runtime/Player/Manifest.h
+++ b/Engine/Plugins/Media/ElectraPlayer/Source/ElectraPlayerRuntime/Private/Runtime/Player/Manifest.h
@@ -343,6 +343,9 @@ namespace Electra
 		//! Returns track metadata. For period based presentations the streams can be different per period in which case the metadata of the first period is returned.
 		virtual void GetTrackMetadata(TArray<FTrackMetadata>& OutMetadata, EStreamType StreamType) const = 0;
 
+		//! Updates long term playback metadata (mostly for audio casts)
+		virtual void UpdateRunningMetaData(TSharedPtrTS<UtilsMP4::FMetadataParser> InUpdatedMetaData) = 0;
+
 		//
 		virtual FTimeValue GetMinBufferTime() const = 0;
 
@@ -351,6 +354,18 @@ namespace Electra
 
 		virtual FTimeValue GetDesiredLiveLatency() const = 0;
 
+		enum class ELiveEdgePlayMode
+		{
+			// Never play on the Live edge.
+			Never,
+			// Play on Live edge on start, disable when a Pause() or Seek() is issued unless it's a seek to Live edge.
+			Default,
+			// Always play on the Live edge (stream can't or should not be paused)
+			Always,
+		};
+		virtual ELiveEdgePlayMode GetLiveEdgePlayMode() const = 0;
+
+
 		enum class EPlayRateType
 		{
 			UnthinnedRate,
@@ -372,6 +387,7 @@ namespace Electra
 
 		virtual void TriggerPlaylistRefresh() = 0;
 
+
 		//-------------------------------------------------------------------------
 		// Stream fragment reader
 		//
@@ -556,9 +572,9 @@ namespace Electra
 	class FPlaylistMetadataUpdateMessage : public IPlayerMessage
 	{
 	public:
-		static TSharedPtrTS<IPlayerMessage> Create(FTimeValue InValidFrom, TSharedPtrTS<UtilsMP4::FMetadataParser> InMetadata)
+		static TSharedPtrTS<IPlayerMessage> Create(FTimeValue InValidFrom, TSharedPtrTS<UtilsMP4::FMetadataParser> InMetadata, bool bInTriggerInternalRefresh)
 		{
-			TSharedPtrTS<FPlaylistMetadataUpdateMessage> p(new FPlaylistMetadataUpdateMessage(InValidFrom, InMetadata));
+			TSharedPtrTS<FPlaylistMetadataUpdateMessage> p(new FPlaylistMetadataUpdateMessage(InValidFrom, InMetadata, bInTriggerInternalRefresh));
 			return p;
 		}
 
@@ -583,13 +599,20 @@ namespace Electra
 			return Metadata;
 		}
 
+		bool GetTriggerInternalRefresh() const
+		{
+			return bTriggerInternalRefresh;
+		}
+
 	private:
-		FPlaylistMetadataUpdateMessage(FTimeValue InValidFrom, TSharedPtrTS<UtilsMP4::FMetadataParser> InMetadata)
+		FPlaylistMetadataUpdateMessage(FTimeValue InValidFrom, TSharedPtrTS<UtilsMP4::FMetadataParser> InMetadata, bool bInTriggerInternalRefresh)
 			: ValidFrom(InValidFrom)
 			, Metadata(MoveTemp(InMetadata))
+			, bTriggerInternalRefresh(bInTriggerInternalRefresh)
 		{ }
 		FTimeValue ValidFrom;
 		TSharedPtrTS<UtilsMP4::FMetadataParser> Metadata;
+		bool bTriggerInternalRefresh;
 	};
 
 
diff --git a/Engine/Plugins/Media/ElectraPlayer/Source/ElectraPlayerRuntime/Private/Runtime/Player/PlaybackTimeline.h b/Engine/Plugins/Media/ElectraPlayer/Source/ElectraPlayerRuntime/Private/Runtime/Player/PlaybackTimeline.h
--- a/Engine/Plugins/Media/ElectraPlayer/Source/ElectraPlayerRuntime/Private/Runtime/Player/PlaybackTimeline.h
+++ b/Engine/Plugins/Media/ElectraPlayer/Source/ElectraPlayerRuntime/Private/Runtime/Player/PlaybackTimeline.h
@@ -123,6 +123,11 @@ namespace Electra
 		 * individual representations.
 		 */
 		virtual void GetMetaData(TArray<FTrackMetadata>& OutMetadata, EStreamType OfStreamType) const = 0;
+
+		/**
+		 * Updates asset metadata on long running Livestreams.
+		 */
+		virtual void UpdateRunningMetaData(const FString& InKindOfValue, const FVariant& InNewValue) = 0;
 	};
 
 
diff --git a/Engine/Plugins/Media/ElectraPlayer/Source/ElectraPlayerRuntime/Private/Runtime/Player/mkv/ManifestMKV.cpp b/Engine/Plugins/Media/ElectraPlayer/Source/ElectraPlayerRuntime/Private/Runtime/Player/mkv/ManifestMKV.cpp
--- a/Engine/Plugins/Media/ElectraPlayer/Source/ElectraPlayerRuntime/Private/Runtime/Player/mkv/ManifestMKV.cpp
+++ b/Engine/Plugins/Media/ElectraPlayer/Source/ElectraPlayerRuntime/Private/Runtime/Player/mkv/ManifestMKV.cpp
@@ -96,6 +96,11 @@ void FManifestMKVInternal::GetTrackMetadata(TArray<FTrackMetadata>& OutMetadata,
 	}
 }
 
+void FManifestMKVInternal::UpdateRunningMetaData(TSharedPtrTS<UtilsMP4::FMetadataParser> InUpdatedMetaData)
+{
+	// No-op.
+}
+
 
 //-----------------------------------------------------------------------------
 /**
diff --git a/Engine/Plugins/Media/ElectraPlayer/Source/ElectraPlayerRuntime/Private/Runtime/Player/mkv/ManifestMKV.h b/Engine/Plugins/Media/ElectraPlayer/Source/ElectraPlayerRuntime/Private/Runtime/Player/mkv/ManifestMKV.h
--- a/Engine/Plugins/Media/ElectraPlayer/Source/ElectraPlayerRuntime/Private/Runtime/Player/mkv/ManifestMKV.h
+++ b/Engine/Plugins/Media/ElectraPlayer/Source/ElectraPlayerRuntime/Private/Runtime/Player/mkv/ManifestMKV.h
@@ -57,9 +57,12 @@ public:
 	void ClearDefaultEndTime() override
 	{ DefaultEndTime.SetToInvalid(); }
 	void GetTrackMetadata(TArray<FTrackMetadata>& OutMetadata, EStreamType StreamType) const override;
+	void UpdateRunningMetaData(TSharedPtrTS<UtilsMP4::FMetadataParser> InUpdatedMetaData) override;
 	FTimeValue GetMinBufferTime() const override;
 	FTimeValue GetDesiredLiveLatency() const override
 	{ return FTimeValue(); }
+	ELiveEdgePlayMode GetLiveEdgePlayMode() const override
+	{ return IManifest::ELiveEdgePlayMode::Never; }
 	TRangeSet<double> GetPossiblePlaybackRates(EPlayRateType InForType) const override;
 	TSharedPtrTS<IProducerReferenceTimeInfo> GetProducerReferenceTimeInfo(int64 ID) const override;
 	void UpdateDynamicRefetchCounter() override;
@@ -239,6 +242,8 @@ public:
 			}
 		}
 
+		void UpdateRunningMetaData(const FString& InKindOfValue, const FVariant& InNewValue) override
+		{ }
 
 		FResult GetStartingSegment(TSharedPtrTS<IStreamSegment>& OutSegment, const FPlayerSequenceState& InSequenceState, FString InSelectedStreamID, const FPlayStartPosition& StartPosition, ESearchType SearchType, int64 AtAbsoluteFilePos);
 		FResult GetNextSegment(TSharedPtrTS<IStreamSegment>& OutSegment, TSharedPtrTS<const IStreamSegment> CurrentSegment, const FPlayStartOptions& Options);
diff --git a/Engine/Plugins/Media/ElectraPlayer/Source/ElectraPlayerRuntime/Private/Runtime/Player/mp4/ManifestMP4.cpp b/Engine/Plugins/Media/ElectraPlayer/Source/ElectraPlayerRuntime/Private/Runtime/Player/mp4/ManifestMP4.cpp
--- a/Engine/Plugins/Media/ElectraPlayer/Source/ElectraPlayerRuntime/Private/Runtime/Player/mp4/ManifestMP4.cpp
+++ b/Engine/Plugins/Media/ElectraPlayer/Source/ElectraPlayerRuntime/Private/Runtime/Player/mp4/ManifestMP4.cpp
@@ -116,6 +116,11 @@ void FManifestMP4Internal::GetTrackMetadata(TArray<FTrackMetadata>& OutMetadata,
 	}
 }
 
+void FManifestMP4Internal::UpdateRunningMetaData(TSharedPtrTS<UtilsMP4::FMetadataParser> InUpdatedMetaData)
+{
+	// No-op.
+}
+
 
 //-----------------------------------------------------------------------------
 /**
@@ -801,7 +806,7 @@ FErrorDetail FManifestMP4Internal::FTimelineAssetMP4::Build(IPlayerSessionServic
 		MediaMetadata = MakeSharedTS<UtilsMP4::FMetadataParser>();
 		if (MediaMetadata->Parse(hdlr, res0, Boxes) == UtilsMP4::FMetadataParser::EResult::Success)
 		{
-			PlayerSessionServices->SendMessageToPlayer(FPlaylistMetadataUpdateMessage::Create(FTimeValue(), MediaMetadata));
+			PlayerSessionServices->SendMessageToPlayer(FPlaylistMetadataUpdateMessage::Create(FTimeValue(), MediaMetadata, false));
 		}
 		else
 		{
diff --git a/Engine/Plugins/Media/ElectraPlayer/Source/ElectraPlayerRuntime/Private/Runtime/Player/mp4/ManifestMP4.h b/Engine/Plugins/Media/ElectraPlayer/Source/ElectraPlayerRuntime/Private/Runtime/Player/mp4/ManifestMP4.h
--- a/Engine/Plugins/Media/ElectraPlayer/Source/ElectraPlayerRuntime/Private/Runtime/Player/mp4/ManifestMP4.h
+++ b/Engine/Plugins/Media/ElectraPlayer/Source/ElectraPlayerRuntime/Private/Runtime/Player/mp4/ManifestMP4.h
@@ -58,9 +58,12 @@ public:
 	void ClearDefaultEndTime() override
 	{ DefaultEndTime.SetToInvalid(); }
 	void GetTrackMetadata(TArray<FTrackMetadata>& OutMetadata, EStreamType StreamType) const override;
+	void UpdateRunningMetaData(TSharedPtrTS<UtilsMP4::FMetadataParser> InUpdatedMetaData) override;
 	FTimeValue GetMinBufferTime() const override;
 	FTimeValue GetDesiredLiveLatency() const override
 	{ return FTimeValue(); }
+	ELiveEdgePlayMode GetLiveEdgePlayMode() const override
+	{ return IManifest::ELiveEdgePlayMode::Never; }
 	TRangeSet<double> GetPossiblePlaybackRates(EPlayRateType InForType) const override;
 	TSharedPtrTS<IProducerReferenceTimeInfo> GetProducerReferenceTimeInfo(int64 ID) const override;
 	void UpdateDynamicRefetchCounter() override;
@@ -224,6 +227,8 @@ public:
 			}
 		}
 
+		void UpdateRunningMetaData(const FString& InKindOfValue, const FVariant& InNewValue) override
+		{ }
 
 		FResult GetStartingSegment(TSharedPtrTS<IStreamSegment>& OutSegment, const FPlayerSequenceState& InSequenceState, const FPlayStartPosition& StartPosition, ESearchType SearchType, int64 AtAbsoluteFilePos);
 		FResult GetNextSegment(TSharedPtrTS<IStreamSegment>& OutSegment, TSharedPtrTS<const IStreamSegment> CurrentSegment, const FPlayStartOptions& Options);
diff --git a/Engine/Plugins/Media/ElectraPlayer/Source/ElectraPlayerRuntime/Private/Runtime/StreamDataBuffer.h b/Engine/Plugins/Media/ElectraPlayer/Source/ElectraPlayerRuntime/Private/Runtime/StreamDataBuffer.h
--- a/Engine/Plugins/Media/ElectraPlayer/Source/ElectraPlayerRuntime/Private/Runtime/StreamDataBuffer.h
+++ b/Engine/Plugins/Media/ElectraPlayer/Source/ElectraPlayerRuntime/Private/Runtime/StreamDataBuffer.h
@@ -200,6 +200,7 @@ namespace Electra
 			SizeAvailableSignal.Signal();
 		}
 
+#if 0
 		// "Pops" data from the buffer to a destination. At most the specified number of bytes are popped, or fewer if not as many are available.
 		int64 PopData(uint8* OutData, int64 MaxElementsWanted)
 		{
@@ -218,6 +219,24 @@ namespace Electra
 			ReadPos += MaxElementsWanted;
 			return MaxElementsWanted;
 		}
+#endif
+
+		void RemoveFromBeginning(int64 InNumBytesToRemove)
+		{
+			FScopeLock Lock(&AccessLock);
+			check(ExternalBuffer == nullptr);
+			check(InNumBytesToRemove >= 0);
+			check(ReadPos == 0);
+			check(InNumBytesToRemove <= (int64)WritePos);
+			int64 InNow = WritePos - ReadPos;
+			check(InNow - InNumBytesToRemove >= 0);
+			if (InNumBytesToRemove > 0)
+			{
+				uint8 *Base = GetBufferBase();
+				FMemory::Memmove(Base, Base + InNumBytesToRemove, InNow - InNumBytesToRemove);
+				WritePos -= (uint64)InNumBytesToRemove;
+			}
+		}
 
 		void Lock()
 		{
@@ -229,6 +248,12 @@ namespace Electra
 			AccessLock.Unlock();
 		}
 
+		// For use with an external FScopeLock
+		FCriticalSection* GetLock()
+		{
+			return &AccessLock;
+		}
+
 		int64 GetLinearReadSize() const
 		{
 			FScopeLock Lock(&AccessLock);
@@ -365,7 +390,7 @@ namespace Electra
 		// If set a buffer is provided externally to read into directly.
 		uint8* ExternalBuffer = nullptr;
 		// Flag indicating that no additional data will be added to the buffer.
-		volatile bool bEOD = false;	
+		volatile bool bEOD = false;
 		// Flag indicating that reading into the buffer has been aborted.
 		volatile bool bWasAborted = false;
 	};
diff --git a/Engine/Plugins/Media/ElectraPlayer/Source/ElectraPlayerRuntime/Private/Runtime/Utilities/UtilsMP4.cpp b/Engine/Plugins/Media/ElectraPlayer/Source/ElectraPlayerRuntime/Private/Runtime/Utilities/UtilsMP4.cpp
--- a/Engine/Plugins/Media/ElectraPlayer/Source/ElectraPlayerRuntime/Private/Runtime/Utilities/UtilsMP4.cpp
+++ b/Engine/Plugins/Media/ElectraPlayer/Source/ElectraPlayerRuntime/Private/Runtime/Utilities/UtilsMP4.cpp
@@ -54,6 +54,39 @@ FString UtilsMP4::FMetadataParser::PrintableBoxAtom(const uint32 InAtom)
 	return Out;
 }
 
+void UtilsMP4::FMetadataParser::AddItem(const FString& InType, const FString& InValue)
+{
+	// Add a string item.
+	TSharedPtr<FItem, ESPMode::ThreadSafe> Item = MakeShared<FItem, ESPMode::ThreadSafe>();
+	Item->Type = 1;
+	Item->Value = FVariant(InValue);
+	++NumTotalItems;
+	Items.FindOrAdd(InType).Emplace(MoveTemp(Item));
+}
+
+void UtilsMP4::FMetadataParser::AddItem(const FString& InType, const FString& InMimeType, const TArray<uint8>& InValue)
+{
+	if (InMimeType.Equals(TEXT("image/jpeg")))
+	{
+		TSharedPtr<FItem, ESPMode::ThreadSafe> Item = MakeShared<FItem, ESPMode::ThreadSafe>();
+		Item->MimeType = InMimeType;
+		Item->Type = 13;
+		Item->Value = FVariant(InValue);
+		++NumTotalItems;
+		Items.FindOrAdd(InType).Emplace(MoveTemp(Item));
+	}
+	else if (InMimeType.Equals(TEXT("image/png")))
+	{
+		TSharedPtr<FItem, ESPMode::ThreadSafe> Item = MakeShared<FItem, ESPMode::ThreadSafe>();
+		Item->MimeType = InMimeType;
+		Item->Type = 14;
+		Item->Value = FVariant(InValue);
+		++NumTotalItems;
+		Items.FindOrAdd(InType).Emplace(MoveTemp(Item));
+	}
+}
+
+
 UtilsMP4::FMetadataParser::EResult UtilsMP4::FMetadataParser::Parse(uint32 InHandler, uint32 InHandlerReserved0, const TArray<UtilsMP4::FMetadataParser::FBoxInfo>& InBoxes)
 {
 	// We only support the Apple iTunes metadata at the moment.
diff --git a/Engine/Plugins/Media/ElectraPlayer/Source/ElectraPlayerRuntime/Private/Runtime/Utilities/UtilsMP4.h b/Engine/Plugins/Media/ElectraPlayer/Source/ElectraPlayerRuntime/Private/Runtime/Utilities/UtilsMP4.h
--- a/Engine/Plugins/Media/ElectraPlayer/Source/ElectraPlayerRuntime/Private/Runtime/Utilities/UtilsMP4.h
+++ b/Engine/Plugins/Media/ElectraPlayer/Source/ElectraPlayerRuntime/Private/Runtime/Utilities/UtilsMP4.h
@@ -46,6 +46,9 @@ public:
 		bool IsDifferentFrom(const FMetadataParser& Other);
 		FString GetAsJSON() const;
 		TSharedPtr<TMap<FString, TArray<TSharedPtr<IMediaStreamMetadata::IItem, ESPMode::ThreadSafe>>>, ESPMode::ThreadSafe> GetMediaStreamMetadata() const;
+
+		void AddItem(const FString& InType, const FString& InValue);
+		void AddItem(const FString& InType, const FString& InMimeType, const TArray<uint8>& InValue);
 	private:
 		FString PrintableBoxAtom(const uint32 InAtom);
 		void Parse(const FBoxInfo& InBox);
