ElectraPlayer: Added playback of .mp3 files and Icycast streams

diff --git a/Engine/Plugins/Media/ElectraPlayer/Source/ElectraPlayerFactory/Private/ElectraPlayerFactoryModule.cpp b/Engine/Plugins/Media/ElectraPlayer/Source/ElectraPlayerFactory/Private/ElectraPlayerFactoryModule.cpp
--- a/Engine/Plugins/Media/ElectraPlayer/Source/ElectraPlayerFactory/Private/ElectraPlayerFactoryModule.cpp
+++ b/Engine/Plugins/Media/ElectraPlayer/Source/ElectraPlayerFactory/Private/ElectraPlayerFactoryModule.cpp
@@ -182,6 +182,8 @@ public:
 		SupportedFileExtensions.Add(TEXT("mkv"));
 		SupportedFileExtensions.Add(TEXT("mka"));
 		SupportedFileExtensions.Add(TEXT("webm"));
+		SupportedFileExtensions.Add(TEXT("mp3"));
+		SupportedFileExtensions.Add(TEXT("mpa"));
 #endif
 
 		// register player factory
diff --git a/Engine/Plugins/Media/ElectraPlayer/Source/ElectraPlayerRuntime/ElectraPlayerRuntime.Build.cs b/Engine/Plugins/Media/ElectraPlayer/Source/ElectraPlayerRuntime/ElectraPlayerRuntime.Build.cs
--- a/Engine/Plugins/Media/ElectraPlayer/Source/ElectraPlayerRuntime/ElectraPlayerRuntime.Build.cs
+++ b/Engine/Plugins/Media/ElectraPlayer/Source/ElectraPlayerRuntime/ElectraPlayerRuntime.Build.cs
@@ -26,7 +26,8 @@ namespace UnrealBuildTool.Rules
 					"ElectraCDM",
 					"ElectraSubtitles",
 					"XmlParser",
-					"SoundTouchZ"
+					"SoundTouchZ",
+                    "HTTP"
 				});
 			if (Target.bCompileAgainstEngine)
 			{
@@ -48,8 +49,6 @@ namespace UnrealBuildTool.Rules
 
 				PrivateDefinitions.Add("_CRT_SECURE_NO_WARNINGS=1");
 
-				AddEngineThirdPartyPrivateStaticDependencies(Target, "WinHttp");
-
 				PrivateDependencyModuleNames.AddAll("D3D11RHI", "D3D12RHI");
 				AddEngineThirdPartyPrivateStaticDependencies(Target, "DX11", "DX12");
 
@@ -95,26 +94,12 @@ namespace UnrealBuildTool.Rules
 			}
 			else if (Target.IsInPlatformGroup(UnrealPlatformGroup.Android) )
 			{
-				PublicDefinitions.Add("CURL_ENABLE_DEBUG_CALLBACK=1");
-
-				if (Target.Configuration != UnrealTargetConfiguration.Shipping)
-				{
-					PublicDefinitions.Add("CURL_ENABLE_NO_TIMEOUTS_OPTION=1");
-				}
-
 				PublicIncludePaths.Add("$(ModuleDir)/Public/Android");
 
 				PrivateIncludePaths.Add("ElectraPlayerRuntime/Private/Runtime/Decoder/Android");
 			}
 			else if (Target.IsInPlatformGroup(UnrealPlatformGroup.Unix) )
 			{
-				PublicDefinitions.Add("CURL_ENABLE_DEBUG_CALLBACK=1");
-
-				if (Target.Configuration != UnrealTargetConfiguration.Shipping)
-				{
-					PublicDefinitions.Add("CURL_ENABLE_NO_TIMEOUTS_OPTION=1");
-				}
-
 				PublicIncludePaths.Add("$(ModuleDir)/Public/Linux");
 				PrivateIncludePaths.Add("ElectraPlayerRuntime/Private/Runtime/Decoder/Linux");
 
diff --git a/Engine/Plugins/Media/ElectraPlayer/Source/ElectraPlayerRuntime/Private/Runtime/Player/AdaptiveStreamingPlayerManifest.cpp b/Engine/Plugins/Media/ElectraPlayer/Source/ElectraPlayerRuntime/Private/Runtime/Player/AdaptiveStreamingPlayerManifest.cpp
--- a/Engine/Plugins/Media/ElectraPlayer/Source/ElectraPlayerRuntime/Private/Runtime/Player/AdaptiveStreamingPlayerManifest.cpp
+++ b/Engine/Plugins/Media/ElectraPlayer/Source/ElectraPlayerRuntime/Private/Runtime/Player/AdaptiveStreamingPlayerManifest.cpp
@@ -8,6 +8,7 @@
 #include "Player/mp4/PlaylistReaderMP4.h"
 #include "Player/DASH/PlaylistReaderDASH.h"
 #include "Player/mkv/PlaylistReaderMKV.h"
+#include "Player/mpegaudio/PlaylistReaderMPEGAudio.h"
 #include "Utilities/Utilities.h"
 #include "Utilities/StringHelpers.h"
 #include "Utilities/URLParser.h"
@@ -26,6 +27,7 @@ static const FString MIMETypeHLS(TEXT("application/vnd.apple.mpegURL"));
 static const FString MIMETypeDASH(TEXT("application/dash+xml"));
 static const FString MIMETypeMKV(TEXT("video/x-matroska"));
 static const FString MIMETypeMKA(TEXT("audio/x-matroska"));
+static const FString MIMETypeMPEGAudio(TEXT("audio/mpeg"));
 
 
 
@@ -79,6 +81,8 @@ FString GetMIMETypeForURL(const FString& URL)
 		static const FString kTextMKV(TEXT("mkv"));
 		static const FString kTextMKA(TEXT("mka"));
 		static const FString kTextWEBM(TEXT("webm"));
+		static const FString kTextMPA(TEXT("mpa"));
+		static const FString kTextMP3(TEXT("mp3"));
 		if (LowerCaseExtension == kTextMP4 || LowerCaseExtension == kTextMP4V)
 		{
 			MimeType = MIMETypeMP4;
@@ -107,6 +111,10 @@ FString GetMIMETypeForURL(const FString& URL)
 		{
 			MimeType = MIMETypeMKA;
 		}
+		else if (LowerCaseExtension == kTextMPA || LowerCaseExtension == kTextMP3)
+		{
+			MimeType = MIMETypeMPEGAudio;
+		}
 	}
 
 	return MimeType;
@@ -230,6 +238,11 @@ void FAdaptiveStreamingPlayer::InternalLoadManifest(const FString& InURL, const
 				ManifestReader = IPlaylistReaderMKV::Create(this);
 				ManifestType = EMediaFormatType::MKV;
 			}
+			else if (mimeType == Playlist::MIMETypeMPEGAudio)
+			{
+				ManifestReader = IPlaylistReaderMPEGAudio::Create(this);
+				ManifestType = EMediaFormatType::MPEGAudio;
+			}
 			else
 			{
 				FErrorDetail err;
diff --git a/Engine/Plugins/Media/ElectraPlayer/Source/ElectraPlayerRuntime/Private/Runtime/Player/Manifest.h b/Engine/Plugins/Media/ElectraPlayer/Source/ElectraPlayerRuntime/Private/Runtime/Player/Manifest.h
--- a/Engine/Plugins/Media/ElectraPlayer/Source/ElectraPlayerRuntime/Private/Runtime/Player/Manifest.h
+++ b/Engine/Plugins/Media/ElectraPlayer/Source/ElectraPlayerRuntime/Private/Runtime/Player/Manifest.h
@@ -23,7 +23,8 @@ namespace Electra
 		ISOBMFF,					// mp4
 		HLS,						// Apple HLS (HTTP Live Streaming)
 		DASH,						// MPEG DASH
-		MKV							// Matroska / WebM
+		MKV,						// Matroska / WebM
+		MPEGAudio					// MPEG audio (eg .mp3)
 	};
 
 
diff --git a/Engine/Plugins/Media/ElectraPlayer/Source/ElectraPlayerRuntime/Private/Runtime/Player/mpegaudio/ManifestMPEGAudio.cpp b/Engine/Plugins/Media/ElectraPlayer/Source/ElectraPlayerRuntime/Private/Runtime/Player/mpegaudio/ManifestMPEGAudio.cpp
new file mode 100644
--- /dev/null
+++ b/Engine/Plugins/Media/ElectraPlayer/Source/ElectraPlayerRuntime/Private/Runtime/Player/mpegaudio/ManifestMPEGAudio.cpp
@@ -0,0 +1,1101 @@
+// Copyright Epic Games, Inc. All Rights Reserved.
+
+#include "PlayerCore.h"
+#include "Player/Manifest.h"
+#include "HTTP/HTTPManager.h"
+#include "Stats/Stats.h"
+#include "SynchronizedClock.h"
+#include "PlaylistReaderMPEGAudio.h"
+#include "Utilities/Utilities.h"
+#include "Utilities/HashFunctions.h"
+#include "Utilities/TimeUtilities.h"
+#include "HAL/LowLevelMemTracker.h"
+#include "ElectraPlayerPrivate.h"
+#include "Player/AdaptiveStreamingPlayerResourceRequest.h"
+#include "Player/AdaptivePlayerOptionKeynames.h"
+#include "Player/PlayerStreamFilter.h"
+#include "Player/mpegaudio/StreamReaderMPEGAudio.h"
+#include "Player/mpegaudio/ManifestMPEGAudio.h"
+
+
+#define ERRCODE_MANIFEST_MPEGAUDIO_NOT_DECODABLE			1
+#define ERRCODE_MANIFEST_MPEGAUDIO_STARTSEGMENT_NOT_FOUND	2
+
+
+#define INTENDED_LIVE_LATENCY_IN_SECONDS	1.0
+
+#define USE_LOW_LATENCY_DESCRIPTOR 0
+
+DECLARE_CYCLE_STAT(TEXT("FPlayPeriodMPEGAudio::FindSegment"), STAT_ElectraPlayer_MPEGAudio_FindSegment, STATGROUP_ElectraPlayer);
+DECLARE_CYCLE_STAT(TEXT("FPlayPeriodMPEGAudio::GetSegmentInformation"), STAT_ElectraPlayer_MPEGAudio_GetSegmentInformation, STATGROUP_ElectraPlayer);
+
+namespace Electra
+{
+
+//-----------------------------------------------------------------------------
+/**
+ * CTOR
+ */
+FManifestMPEGAudioInternal::FManifestMPEGAudioInternal(IPlayerSessionServices* InPlayerSessionServices)
+	: PlayerSessionServices(InPlayerSessionServices)
+{
+}
+
+
+//-----------------------------------------------------------------------------
+/**
+ * DTOR
+ */
+FManifestMPEGAudioInternal::~FManifestMPEGAudioInternal()
+{
+}
+
+
+//-----------------------------------------------------------------------------
+/**
+ * Builds the internal manifest.
+ *
+ * @param InHeader
+ * @param URL
+ *
+ * @return
+ */
+FErrorDetail FManifestMPEGAudioInternal::Build(const FMPEGAudioInfoHeader& InHeader, const FString& URL)
+{
+	MPEGInfoHeader = InHeader;
+	MediaAsset = MakeSharedTS<FTimelineAssetMPEGAudio>();
+	FErrorDetail Result = MediaAsset->Build(PlayerSessionServices, InHeader, URL);
+	FTimeRange PlaybackRange = GetPlaybackRange(IManifest::EPlaybackRangeType::TemporaryPlaystartRange);
+	DefaultStartTime = PlaybackRange.Start;
+	DefaultEndTime = PlaybackRange.End;
+
+#if USE_LOW_LATENCY_DESCRIPTOR
+	if (MPEGInfoHeader.bIsLive)
+	{
+		LatencyDescriptor = MakeSharedTS<FLowLatencyDescriptor>();
+		LatencyDescriptor->Latency.Target.SetFromSeconds(INTENDED_LIVE_LATENCY_IN_SECONDS);
+		LatencyDescriptor->Latency.Min.SetFromSeconds(0.5);	// needs to be less than INTENDED_LIVE_LATENCY_IN_SECONDS
+		LatencyDescriptor->Latency.Max.SetFromSeconds(4.0);	// needs to be more than INTENDED_LIVE_LATENCY_IN_SECONDS
+		LatencyDescriptor->PlayRate.Min.SetFromSeconds(0.9);
+		LatencyDescriptor->PlayRate.Max.SetFromSeconds(1.05);
+	}
+#endif
+	return Result;
+}
+
+
+
+//-----------------------------------------------------------------------------
+/**
+ * Logs a message.
+ *
+ * @param Level
+ * @param Message
+ */
+void FManifestMPEGAudioInternal::LogMessage(IInfoLog::ELevel Level, const FString& Message)
+{
+	if (PlayerSessionServices)
+	{
+		PlayerSessionServices->PostLog(Facility::EFacility::MPEGAudioPlaylist, Level, Message);
+	}
+}
+
+
+//-----------------------------------------------------------------------------
+/**
+ * Returns the type of presentation.
+ *
+ * @return
+ */
+IManifest::EType FManifestMPEGAudioInternal::GetPresentationType() const
+{
+	return MediaAsset.IsValid() ? MediaAsset->GetDuration().IsPositiveInfinity() ? IManifest::EType::Live : IManifest::EType::OnDemand : IManifest::EType::OnDemand;
+}
+
+
+//-----------------------------------------------------------------------------
+/**
+ * Returns track metadata.
+ *
+ * @param OutMetadata
+ * @param StreamType
+ */
+void FManifestMPEGAudioInternal::GetTrackMetadata(TArray<FTrackMetadata>& OutMetadata, EStreamType StreamType) const
+{
+	if (MediaAsset.IsValid())
+	{
+		MediaAsset->GetMetaData(OutMetadata, StreamType);
+	}
+}
+
+
+void FManifestMPEGAudioInternal::UpdateRunningMetaData(TSharedPtrTS<UtilsMP4::FMetadataParser> InUpdatedMetaData)
+{
+	TSharedPtr<TMap<FString, TArray<TSharedPtr<IMediaStreamMetadata::IItem, ESPMode::ThreadSafe>>>, ESPMode::ThreadSafe> md = InUpdatedMetaData->GetMediaStreamMetadata();
+	if (md.IsValid() && MediaAsset.IsValid())
+	{
+		for(auto& it : *md)
+		{
+			if (it.Key == TEXT("Title") && it.Value.Num())
+			{
+				MediaAsset->UpdateRunningMetaData(it.Key, it.Value[0]->GetValue());
+			}
+		}
+	}
+}
+
+
+//-----------------------------------------------------------------------------
+/**
+ * Returns the playback range on the timeline, which is a subset of the total
+ * time range. This may be set through manifest internal means or by URL fragment
+ * parameters where permissable (eg. example.mpa#t=22,50).
+ * If start or end are not specified they will be set to invalid.
+ *
+ * @return Optionally set time range to which playback is restricted.
+ */
+FTimeRange FManifestMPEGAudioInternal::GetPlaybackRange(EPlaybackRangeType InRangeType) const
+{
+	FTimeRange FromTo;
+
+	// We are interested in the 't' or 'r' fragment value here.
+	FString Time;
+	for(auto& Fragment : URLFragmentComponents)
+	{
+		if ((InRangeType == IManifest::EPlaybackRangeType::TemporaryPlaystartRange && Fragment.Name.Equals(TEXT("t"))) ||
+			(InRangeType == IManifest::EPlaybackRangeType::LockedPlaybackRange && Fragment.Name.Equals(TEXT("r"))))
+		{
+			Time = Fragment.Value;
+		}
+	}
+	if (!Time.IsEmpty())
+	{
+		FTimeRange TotalRange = GetTotalTimeRange();
+		TArray<FString> TimeRange;
+		const TCHAR* const TimeDelimiter = TEXT(",");
+		Time.ParseIntoArray(TimeRange, TimeDelimiter, false);
+		if (TimeRange.Num() && !TimeRange[0].IsEmpty())
+		{
+			RFC2326::ParseNPTTime(FromTo.Start, TimeRange[0]);
+		}
+		if (TimeRange.Num() > 1 && !TimeRange[1].IsEmpty())
+		{
+			RFC2326::ParseNPTTime(FromTo.End, TimeRange[1]);
+		}
+		// Need to clamp this into the total time range to prevent any issues.
+		if (FromTo.Start.IsValid() && TotalRange.Start.IsValid() && FromTo.Start < TotalRange.Start)
+		{
+			FromTo.Start = TotalRange.Start;
+		}
+		if (FromTo.End.IsValid() && TotalRange.End.IsValid() && FromTo.End > TotalRange.End)
+		{
+			FromTo.End = TotalRange.End;
+		}
+	}
+	return FromTo;
+}
+
+
+//-----------------------------------------------------------------------------
+/**
+ * Returns the minimum duration of content that must be buffered up before playback
+ * will begin.
+ *
+ * @return
+ */
+FTimeValue FManifestMPEGAudioInternal::GetMinBufferTime() const
+{
+	if (MPEGInfoHeader.bIsLive)
+	{
+		return FTimeValue().SetFromSeconds(1.0);
+	}
+	// NOTE: This is an arbitrary choice.
+	return FTimeValue().SetFromSeconds(1.0);
+}
+
+FTimeValue FManifestMPEGAudioInternal::GetDesiredLiveLatency() const
+{
+	return LatencyDescriptor.IsValid() ? LatencyDescriptor->Latency.Target : FTimeValue().SetFromSeconds(INTENDED_LIVE_LATENCY_IN_SECONDS);
+}
+
+IManifest::ELiveEdgePlayMode FManifestMPEGAudioInternal::GetLiveEdgePlayMode() const
+{
+	return IManifest::ELiveEdgePlayMode::Never;
+}
+
+
+TSharedPtrTS<const FLowLatencyDescriptor> FManifestMPEGAudioInternal::GetLowLatencyDescriptor() const
+{
+	return LatencyDescriptor;
+}
+
+
+TSharedPtrTS<IProducerReferenceTimeInfo> FManifestMPEGAudioInternal::GetProducerReferenceTimeInfo(int64 ID) const
+{
+	return nullptr;
+}
+
+TRangeSet<double> FManifestMPEGAudioInternal::GetPossiblePlaybackRates(EPlayRateType InForType) const
+{
+	TRangeSet<double> Ranges;
+//	Ranges.Add(TRange<double>{1.0}); // normal (real-time) playback rate
+	Ranges.Add(TRange<double>::Inclusive(0.5, 4.0));
+	Ranges.Add(TRange<double>{0.0}); // and pause
+	return Ranges;
+}
+
+
+void FManifestMPEGAudioInternal::UpdateDynamicRefetchCounter()
+{
+	// No-op.
+}
+
+void FManifestMPEGAudioInternal::TriggerClockSync(IManifest::EClockSyncType InClockSyncType)
+{
+	// No-op.
+}
+
+void FManifestMPEGAudioInternal::TriggerPlaylistRefresh()
+{
+	// No-op.
+}
+
+
+//-----------------------------------------------------------------------------
+/**
+ * Creates an instance of a stream reader.
+ *
+ * @return
+ */
+IStreamReader* FManifestMPEGAudioInternal::CreateStreamReaderHandler()
+{
+	return new FStreamReaderMPEGAudio;
+}
+
+
+//-----------------------------------------------------------------------------
+/**
+ * Returns the playback period for the given time.
+ *
+ * @param OutPlayPeriod
+ * @param StartPosition
+ * @param SearchType
+ *
+ * @return
+ */
+IManifest::FResult FManifestMPEGAudioInternal::FindPlayPeriod(TSharedPtrTS<IPlayPeriod>& OutPlayPeriod, const FPlayStartPosition& StartPosition, ESearchType SearchType)
+{
+	if (MediaAsset.IsValid() && StartPosition.Time.IsValid() && (StartPosition.Time < MediaAsset->GetDuration() || MediaAsset->GetDuration().IsPositiveInfinity()))
+	{
+		OutPlayPeriod = MakeSharedTS<FPlayPeriodMPEGAudio>(MediaAsset);
+		return IManifest::FResult(IManifest::FResult::EType::Found);
+	}
+	else
+	{
+		return IManifest::FResult(IManifest::FResult::EType::PastEOS);
+	}
+}
+
+IManifest::FResult FManifestMPEGAudioInternal::FindNextPlayPeriod(TSharedPtrTS<IPlayPeriod>& OutPlayPeriod, TSharedPtrTS<const IStreamSegment> CurrentSegment)
+{
+	// There is no following period.
+	return IManifest::FResult(IManifest::FResult::EType::PastEOS);
+}
+
+
+
+
+
+//-----------------------------------------------------------------------------
+/**
+ * Constructs a playback period.
+ *
+ * @param InMediaAsset
+ */
+FManifestMPEGAudioInternal::FPlayPeriodMPEGAudio::FPlayPeriodMPEGAudio(TSharedPtrTS<FManifestMPEGAudioInternal::FTimelineAssetMPEGAudio> InMediaAsset)
+	: MediaAsset(InMediaAsset)
+	, CurrentReadyState(IManifest::IPlayPeriod::EReadyState::NotLoaded)
+{
+}
+
+
+//-----------------------------------------------------------------------------
+/**
+ * Destroys a playback period.
+ */
+FManifestMPEGAudioInternal::FPlayPeriodMPEGAudio::~FPlayPeriodMPEGAudio()
+{
+}
+
+
+//-----------------------------------------------------------------------------
+/**
+ * Sets stream playback preferences for this playback period.
+ *
+ * @param ForStreamType
+ * @param StreamAttributes
+ */
+void FManifestMPEGAudioInternal::FPlayPeriodMPEGAudio::SetStreamPreferences(EStreamType ForStreamType, const FStreamSelectionAttributes& StreamAttributes)
+{
+	if (ForStreamType == EStreamType::Audio)
+	{
+		AudioPreferences = StreamAttributes;
+	}
+}
+
+
+//-----------------------------------------------------------------------------
+/**
+ * Returns the starting bitrate.
+ *
+ * This is merely informational and not strictly required.
+ *
+ * @return
+ */
+int64 FManifestMPEGAudioInternal::FPlayPeriodMPEGAudio::GetDefaultStartingBitrate() const
+{
+	TSharedPtrTS<FTimelineAssetMPEGAudio> ma = MediaAsset.Pin();
+	if (ma.IsValid())
+	{
+		TArray<FTrackMetadata> Metadata;
+		ma->GetMetaData(Metadata, EStreamType::Audio);
+		if (Metadata.Num() && Metadata[0].StreamDetails.Num())
+		{
+			return Metadata[0].StreamDetails[0].Bandwidth;
+		}
+	}
+	return 128000;
+}
+
+//-----------------------------------------------------------------------------
+/**
+ * Returns the ready state of this playback period.
+ *
+ * @return
+ */
+IManifest::IPlayPeriod::EReadyState FManifestMPEGAudioInternal::FPlayPeriodMPEGAudio::GetReadyState()
+{
+	return CurrentReadyState;
+}
+
+
+void FManifestMPEGAudioInternal::FPlayPeriodMPEGAudio::Load()
+{
+	CurrentReadyState = IManifest::IPlayPeriod::EReadyState::Loaded;
+}
+
+//-----------------------------------------------------------------------------
+/**
+ * Prepares the playback period for playback.
+ * We are actually always ready for playback, but we say we're not
+ * one time to get here with any possible options.
+ */
+void FManifestMPEGAudioInternal::FPlayPeriodMPEGAudio::PrepareForPlay()
+{
+	SelectedAudioMetadata.Reset();
+	AudioBufferSourceInfo.Reset();
+	SelectInitialStream(EStreamType::Audio);
+	CurrentReadyState = IManifest::IPlayPeriod::EReadyState::IsReady;
+}
+
+
+TSharedPtrTS<FBufferSourceInfo> FManifestMPEGAudioInternal::FPlayPeriodMPEGAudio::GetSelectedStreamBufferSourceInfo(EStreamType StreamType)
+{
+	return StreamType == EStreamType::Audio ? AudioBufferSourceInfo : nullptr;
+}
+
+FString FManifestMPEGAudioInternal::FPlayPeriodMPEGAudio::GetSelectedAdaptationSetID(EStreamType StreamType)
+{
+	if (StreamType == EStreamType::Audio)
+	{
+		return SelectedAudioMetadata.IsValid() ? SelectedAudioMetadata->ID : FString();
+	}
+	return FString();
+}
+
+
+IManifest::IPlayPeriod::ETrackChangeResult FManifestMPEGAudioInternal::FPlayPeriodMPEGAudio::ChangeTrackStreamPreference(EStreamType StreamType, const FStreamSelectionAttributes& StreamAttributes)
+{
+	TSharedPtrTS<FTrackMetadata> Metadata = SelectMetadataForAttributes(StreamType, StreamAttributes);
+	if (Metadata.IsValid())
+	{
+		if (StreamType == EStreamType::Audio)
+		{
+			if (!(SelectedAudioMetadata.IsValid() && Metadata->Equals(*SelectedAudioMetadata)))
+			{
+				SelectedAudioMetadata = Metadata;
+				MakeBufferSourceInfoFromMetadata(StreamType, AudioBufferSourceInfo, SelectedAudioMetadata);
+				return IManifest::IPlayPeriod::ETrackChangeResult::Changed;
+			}
+		}
+	}
+	return IManifest::IPlayPeriod::ETrackChangeResult::NotChanged;
+}
+
+void FManifestMPEGAudioInternal::FPlayPeriodMPEGAudio::SelectInitialStream(EStreamType StreamType)
+{
+	if (StreamType == EStreamType::Audio)
+	{
+		SelectedAudioMetadata = SelectMetadataForAttributes(StreamType, AudioPreferences);
+		MakeBufferSourceInfoFromMetadata(StreamType, AudioBufferSourceInfo, SelectedAudioMetadata);
+	}
+}
+
+TSharedPtrTS<FTrackMetadata> FManifestMPEGAudioInternal::FPlayPeriodMPEGAudio::SelectMetadataForAttributes(EStreamType StreamType, const FStreamSelectionAttributes& InAttributes)
+{
+	TSharedPtrTS<FTimelineAssetMPEGAudio> Asset = MediaAsset.Pin();
+	if (Asset.IsValid())
+	{
+		TArray<FTrackMetadata> Metadata;
+		Asset->GetMetaData(Metadata, StreamType);
+		// Is there a fixed index to be used?
+		if (InAttributes.OverrideIndex.IsSet() && InAttributes.OverrideIndex.GetValue() >= 0 && InAttributes.OverrideIndex.GetValue() < Metadata.Num())
+		{
+			// Use this.
+			return MakeSharedTS<FTrackMetadata>(Metadata[InAttributes.OverrideIndex.GetValue()]);
+		}
+		if (Metadata.Num())
+		{
+			// We do not look at the 'kind' here, only the language.
+			// Set the first track as default in case we do not find the one we're looking for.
+			if (InAttributes.Language_ISO639.IsSet())
+			{
+				for(auto &Meta : Metadata)
+				{
+					if (Meta.Language.Equals(InAttributes.Language_ISO639.GetValue()))
+					{
+						return MakeSharedTS<FTrackMetadata>(Meta);
+					}
+				}
+			}
+			return MakeSharedTS<FTrackMetadata>(Metadata[0]);
+		}
+	}
+	return nullptr;
+}
+
+void FManifestMPEGAudioInternal::FPlayPeriodMPEGAudio::MakeBufferSourceInfoFromMetadata(EStreamType StreamType, TSharedPtrTS<FBufferSourceInfo>& OutBufferSourceInfo, TSharedPtrTS<FTrackMetadata> InMetadata)
+{
+	if (InMetadata.IsValid())
+	{
+		OutBufferSourceInfo = MakeSharedTS<FBufferSourceInfo>();
+		OutBufferSourceInfo->Kind = InMetadata->Kind;
+		OutBufferSourceInfo->Language = InMetadata->Language;
+		OutBufferSourceInfo->Codec = InMetadata->HighestBandwidthCodec.GetCodecName();
+		TSharedPtrTS<FTimelineAssetMPEGAudio> Asset = MediaAsset.Pin();
+		OutBufferSourceInfo->PeriodID = Asset->GetUniqueIdentifier();
+		OutBufferSourceInfo->PeriodAdaptationSetID = Asset->GetUniqueIdentifier() + TEXT(".") + InMetadata->ID;
+		TArray<FTrackMetadata> Metadata;
+		Asset->GetMetaData(Metadata, StreamType);
+		for(int32 i=0; i<Metadata.Num(); ++i)
+		{
+			if (Metadata[i].Equals(*InMetadata))
+			{
+				OutBufferSourceInfo->HardIndex = i;
+				break;
+			}
+		}
+	}
+}
+
+
+
+//-----------------------------------------------------------------------------
+/**
+ * Returns the timeline media asset. We have a weak pointer to it only to
+ * prevent any cyclic locks, so we need to lock it first.
+ *
+ * @return
+ */
+TSharedPtrTS<ITimelineMediaAsset> FManifestMPEGAudioInternal::FPlayPeriodMPEGAudio::GetMediaAsset() const
+{
+	TSharedPtrTS<ITimelineMediaAsset> ma = MediaAsset.Pin();
+	return ma;
+}
+
+
+//-----------------------------------------------------------------------------
+/**
+ * Selects a particular stream (== internal track ID) for playback.
+ *
+ * @param AdaptationSetID
+ * @param RepresentationID
+ */
+void FManifestMPEGAudioInternal::FPlayPeriodMPEGAudio::SelectStream(const FString& AdaptationSetID, const FString& RepresentationID, int32 QualityIndex, int32 MaxQualityIndex)
+{
+	// Presently this method is only called by the ABR to switch between quality levels, of which there are none.
+}
+
+void FManifestMPEGAudioInternal::FPlayPeriodMPEGAudio::TriggerInitSegmentPreload(const TArray<FInitSegmentPreload>& InitSegmentsToPreload)
+{
+	// No-op.
+}
+
+//-----------------------------------------------------------------------------
+/**
+ * Creates the starting segment request to start playback with.
+ *
+ * @param OutSegment
+ * @param InSequenceState
+ * @param StartPosition
+ * @param SearchType
+ *
+ * @return
+ */
+IManifest::FResult FManifestMPEGAudioInternal::FPlayPeriodMPEGAudio::GetStartingSegment(TSharedPtrTS<IStreamSegment>& OutSegment, const FPlayerSequenceState& InSequenceState, const FPlayStartPosition& StartPosition, ESearchType SearchType)
+{
+	TSharedPtrTS<FTimelineAssetMPEGAudio> ma = MediaAsset.Pin();
+	return ma.IsValid() ? ma->GetStartingSegment(OutSegment, InSequenceState, StartPosition, SearchType, -1) : IManifest::FResult(IManifest::FResult::EType::NotFound);
+}
+
+//-----------------------------------------------------------------------------
+/**
+ * Same as GetStartingSegment() except this is for a specific stream (video, audio, ...) only.
+ * To be used when a track (language) change is made and a new segment is needed at the current playback position.
+ */
+IManifest::FResult FManifestMPEGAudioInternal::FPlayPeriodMPEGAudio::GetContinuationSegment(TSharedPtrTS<IStreamSegment>& OutSegment, EStreamType StreamType, const FPlayerSequenceState& SequenceState, const FPlayStartPosition& StartPosition, ESearchType SearchType)
+{
+	// Not supported
+	return IManifest::FResult(IManifest::FResult::EType::NotFound);
+}
+
+
+//-----------------------------------------------------------------------------
+/**
+ * Sets up a starting segment request to loop playback to.
+ * The streams selected through SelectStream() will be used.
+ *
+ * @param OutSegment
+ * @param SequenceState
+ * @param StartPosition
+ * @param SearchType
+ *
+ * @return
+ */
+IManifest::FResult FManifestMPEGAudioInternal::FPlayPeriodMPEGAudio::GetLoopingSegment(TSharedPtrTS<IStreamSegment>& OutSegment, const FPlayerSequenceState& SequenceState, const FPlayStartPosition& StartPosition, ESearchType SearchType)
+{
+	TSharedPtrTS<FTimelineAssetMPEGAudio> ma = MediaAsset.Pin();
+	return ma.IsValid() ? ma->GetLoopingSegment(OutSegment, SequenceState, StartPosition, SearchType) : IManifest::FResult(IManifest::FResult::EType::NotFound);
+}
+
+
+//-----------------------------------------------------------------------------
+/**
+ * Called by the ABR to increase the delay in fetching the next segment in case the segment returned a 404 when fetched at
+ * the announced availability time. This may reduce 404's on the next segment fetches.
+ *
+ * @param IncreaseAmount
+ */
+void FManifestMPEGAudioInternal::FPlayPeriodMPEGAudio::IncreaseSegmentFetchDelay(const FTimeValue& IncreaseAmount)
+{
+	// No-op.
+}
+
+
+//-----------------------------------------------------------------------------
+/**
+ * Creates the next segment request.
+ *
+ * @param OutSegment
+ * @param CurrentSegment
+ * @param Options
+ *
+ * @return
+ */
+IManifest::FResult FManifestMPEGAudioInternal::FPlayPeriodMPEGAudio::GetNextSegment(TSharedPtrTS<IStreamSegment>& OutSegment, TSharedPtrTS<const IStreamSegment> CurrentSegment, const FPlayStartOptions& Options)
+{
+	TSharedPtrTS<FTimelineAssetMPEGAudio> ma = MediaAsset.Pin();
+	return ma.IsValid() ? ma->GetNextSegment(OutSegment, CurrentSegment, Options) : IManifest::FResult(IManifest::FResult::EType::NotFound);
+}
+
+
+//-----------------------------------------------------------------------------
+/**
+ * Creates a segment retry request.
+ *
+ * @param OutSegment
+ * @param CurrentSegment
+ * @param Options
+ * @param bReplaceWithFillerData
+ *
+ * @return
+ */
+IManifest::FResult FManifestMPEGAudioInternal::FPlayPeriodMPEGAudio::GetRetrySegment(TSharedPtrTS<IStreamSegment>& OutSegment, TSharedPtrTS<const IStreamSegment> CurrentSegment, const FPlayStartOptions& Options, bool bReplaceWithFillerData)
+{
+	TSharedPtrTS<FTimelineAssetMPEGAudio> ma = MediaAsset.Pin();
+	return ma.IsValid() ? ma->GetRetrySegment(OutSegment, CurrentSegment, Options, bReplaceWithFillerData) : IManifest::FResult(IManifest::FResult::EType::NotFound);
+}
+
+
+//-----------------------------------------------------------------------------
+/**
+ * Returns segment information for the next n segments.
+ *
+ * @param OutSegmentInformation
+ * @param OutAverageSegmentDuration
+ * @param CurrentSegment
+ * @param LookAheadTime
+ * @param AdaptationSet
+ * @param Representation
+ */
+void FManifestMPEGAudioInternal::FPlayPeriodMPEGAudio::GetSegmentInformation(TArray<FSegmentInformation>& OutSegmentInformation, FTimeValue& OutAverageSegmentDuration, TSharedPtrTS<const IStreamSegment> CurrentSegment, const FTimeValue& LookAheadTime, const FString& AdaptationSetID, const FString& RepresentationID)
+{
+	TSharedPtrTS<FTimelineAssetMPEGAudio> ma = MediaAsset.Pin();
+	if (ma.IsValid())
+	{
+		ma->GetSegmentInformation(OutSegmentInformation, OutAverageSegmentDuration, CurrentSegment, LookAheadTime, AdaptationSetID, RepresentationID);
+	}
+}
+
+
+
+
+
+//-----------------------------------------------------------------------------
+/**
+ * Builds the timeline asset.
+ *
+ * @param InPlayerSessionServices
+ * @param InHeader
+ * @param URL
+ *
+ * @return
+ */
+FErrorDetail FManifestMPEGAudioInternal::FTimelineAssetMPEGAudio::Build(IPlayerSessionServices* InPlayerSessionServices, const FMPEGAudioInfoHeader& InHeader, const FString& URL)
+{
+	PlayerSessionServices = InPlayerSessionServices;
+	MediaURL = URL;
+	MPEGInfoHeader = InHeader;
+
+	// Can we decode this track?
+	bool bIsUsable = false;
+	IPlayerStreamFilter* StreamFilter = InPlayerSessionServices->GetStreamFilter();
+	if (StreamFilter && StreamFilter->CanDecodeStream(MPEGInfoHeader.CodecInfo))
+	{
+		FErrorDetail err;
+		TSharedPtrTS<FAdaptationSetMPEGAudio> AdaptationSet = MakeSharedTS<FAdaptationSetMPEGAudio>();
+		err = AdaptationSet->CreateFrom(MPEGInfoHeader, URL);
+		if (err.IsOK())
+		{
+			AudioAdaptationSets.Add(AdaptationSet);
+		}
+		else
+		{
+			return err;
+		}
+	}
+	else
+	{
+		FErrorDetail err;
+		err.SetFacility(Facility::EFacility::MPEGAudioPlaylist);
+		err.SetMessage("This stream cannot be decoded");
+		err.SetCode(ERRCODE_MANIFEST_MPEGAUDIO_NOT_DECODABLE);
+		return err;
+	}
+
+	// Convert metadata, if present.
+	if (MPEGInfoHeader.ID3v2.IsValid())
+	{
+		TSharedPtr<UtilsMP4::FMetadataParser, ESPMode::ThreadSafe> mp = MakeShared<UtilsMP4::FMetadataParser, ESPMode::ThreadSafe>();
+
+		for(auto& it : MPEGInfoHeader.ID3v2->GetTags())
+		{
+			if (it.Key == Utils::Make4CC('T','I','T','2'))
+			{
+				mp->AddItem(TEXT("Title"), it.Value.Value.GetValue<FString>());
+			}
+			else if (it.Key == Utils::Make4CC('T','A','L','B'))
+			{
+				mp->AddItem(TEXT("Album"), it.Value.Value.GetValue<FString>());
+			}
+			else if (it.Key == Utils::Make4CC('T','P','E','1'))
+			{
+				mp->AddItem(TEXT("Artist"), it.Value.Value.GetValue<FString>());
+			}
+			else if (it.Key == Utils::Make4CC('T','E','N','C'))
+			{
+				mp->AddItem(TEXT("Encoder"), it.Value.Value.GetValue<FString>());
+			}
+			else if (it.Key == Utils::Make4CC('T','C','O','N'))
+			{
+				mp->AddItem(TEXT("Genre"), it.Value.Value.GetValue<FString>());
+			}
+			else if (it.Key == Utils::Make4CC('T','Y','E','R'))
+			{
+				mp->AddItem(TEXT("Date"), it.Value.Value.GetValue<FString>());
+			}
+			else if (it.Key == Utils::Make4CC('A','P','I','C'))
+			{
+				mp->AddItem(TEXT("covr"), it.Value.MimeType, it.Value.Value.GetValue<TArray<uint8>>());
+			}
+			else if (it.Key == Utils::Make4CC('T','L','E','N') ||
+					 it.Key == Utils::Make4CC('M','L','L','T') ||
+					 0)
+			{
+				// ignore
+			}
+		}
+		PlayerSessionServices->SendMessageToPlayer(FPlaylistMetadataUpdateMessage::Create(FTimeValue(), mp, false));
+	}
+	else if (MPEGInfoHeader.HTTPResponseHeaders.Contains(TEXT("icy-name")))
+	{
+		TSharedPtr<UtilsMP4::FMetadataParser, ESPMode::ThreadSafe> mp = MakeShared<UtilsMP4::FMetadataParser, ESPMode::ThreadSafe>();
+		TArray<FString> HeaderValues;
+		MPEGInfoHeader.HTTPResponseHeaders.MultiFind(TEXT("icy-name"), HeaderValues);
+		mp->AddItem(TEXT("Album"), HeaderValues[0]);
+		if (MPEGInfoHeader.HTTPResponseHeaders.Contains(TEXT("icy-genre")))
+		{
+			HeaderValues.Empty();
+			MPEGInfoHeader.HTTPResponseHeaders.MultiFind(TEXT("icy-genre"), HeaderValues);
+			mp->AddItem(TEXT("Genre"), HeaderValues[0]);
+		}
+		PlayerSessionServices->SendMessageToPlayer(FPlaylistMetadataUpdateMessage::Create(FTimeValue(), mp, false));
+	}
+	return FErrorDetail();
+}
+
+
+
+void FManifestMPEGAudioInternal::FTimelineAssetMPEGAudio::LogMessage(IInfoLog::ELevel Level, const FString& Message)
+{
+	if (PlayerSessionServices)
+	{
+		PlayerSessionServices->PostLog(Facility::EFacility::MPEGAudioPlaylist, Level, Message);
+	}
+}
+
+
+
+
+IManifest::FResult FManifestMPEGAudioInternal::FTimelineAssetMPEGAudio::GetStartingSegment(TSharedPtrTS<IStreamSegment>& OutSegment, const FPlayerSequenceState& InSequenceState, const FPlayStartPosition& StartPosition, ESearchType SearchType, int64 AtAbsoluteFilePos)
+{
+	SCOPE_CYCLE_COUNTER(STAT_ElectraPlayer_MPEGAudio_FindSegment);
+	CSV_SCOPED_TIMING_STAT(ElectraPlayer, MPEGAudio_FindSegment);
+
+	if (AudioAdaptationSets.Num())
+	{
+		// Live stream?
+		if (MPEGInfoHeader.bIsLive)
+		{
+			TSharedPtrTS<FStreamSegmentRequestMPEGAudio> req(new FStreamSegmentRequestMPEGAudio);
+			OutSegment = req;
+			req->MediaAsset = SharedThis(this);
+			FTimeValue StartTime = PlayerSessionServices && PlayerSessionServices->GetSynchronizedUTCTime() ? PlayerSessionServices->GetSynchronizedUTCTime()->GetTime() : MEDIAutcTime::Current();
+			// Is this a known audio cast type?
+			if (MPEGInfoHeader.HTTPResponseHeaders.Contains(TEXT("icy-name")))
+			{
+				req->CastType = FStreamSegmentRequestMPEGAudio::ECastType::IcyCast;
+			}
+			req->FirstPTS = StartTime;
+			req->EarliestPTS = StartTime;
+			req->LastPTS.SetToPositiveInfinity();
+			req->Bitrate = MPEGInfoHeader.AverageBitrate ? MPEGInfoHeader.AverageBitrate : MPEGInfoHeader.Bitrate;
+			req->bIsContinuationSegment = false;
+			req->bIsFirstSegment = true;
+			req->TimestampSequenceIndex = InSequenceState.GetSequenceIndex();
+			req->bIsLastSegment = false;
+			req->MPEGHeaderMask = MPEGInfoHeader.MPEGHeaderMask;
+			req->MPEGHeaderExpectedValue = MPEGInfoHeader.MPEGHeaderExpectedValue;
+			req->CBRFrameSize = MPEGInfoHeader.CBRFrameSize;
+			req->bIsVBR = MPEGInfoHeader.bIsVBR;
+			req->bIsLive = true;
+			req->CodecInfo = MPEGInfoHeader.CodecInfo;
+			return IManifest::FResult(IManifest::FResult::EType::Found);
+		}
+		else
+		{
+			FTimeValue StartTime = StartPosition.Time;
+			FTimeValue PlayRangeEnd = StartPosition.Options.PlaybackRange.End;
+			check(PlayRangeEnd.IsValid());
+			if (PlayRangeEnd > MPEGInfoHeader.EstimatedDuration)
+			{
+				PlayRangeEnd = MPEGInfoHeader.EstimatedDuration;
+			}
+
+			int64 ApproxByteOffset = -1;
+			if (AtAbsoluteFilePos < 0)
+			{
+				check(MPEGInfoHeader.EstimatedDuration.IsValid());
+				check(MPEGInfoHeader.SampleRate);
+				check(MPEGInfoHeader.SamplesPerFrame);
+				check(MPEGInfoHeader.FirstDataByte >= 0);
+				check(MPEGInfoHeader.LastDataByte > 0);
+
+				double Percentage = FMath::Clamp(StartTime.GetAsSeconds() / MPEGInfoHeader.EstimatedDuration.GetAsSeconds(), 0.0, 1.0);
+				// Seek accuracy varies greatly depending on the (optional!) information provided.
+				if (!MPEGInfoHeader.bIsVBR)
+				{
+					// CBR. We can calculate an approximate position based on frame size.
+					check(MPEGInfoHeader.Bitrate);
+					check(MPEGInfoHeader.CBRFrameSize);
+					ApproxByteOffset = MPEGInfoHeader.FirstDataByte + (MPEGInfoHeader.LastDataByte - MPEGInfoHeader.FirstDataByte) * Percentage;
+					ApproxByteOffset = FMath::Clamp(ApproxByteOffset, MPEGInfoHeader.FirstDataByte, MPEGInfoHeader.LastDataByte);
+				}
+				else
+				{
+					check(MPEGInfoHeader.AverageBitrate);
+
+					// Is there an MLLT entry from the ID3v2 header?
+					if (MPEGInfoHeader.MLLT.IsValid())
+					{
+						int64 StartTimeMillis = StartTime.GetAsMilliseconds();
+						// Locate the entry
+						int32 EntryIdx, LastEntryIdx=MPEGInfoHeader.MLLT->TimeAndOffsets.Num();
+						for(EntryIdx=0; EntryIdx<LastEntryIdx-1; ++EntryIdx)
+						{
+							if ((int64)MPEGInfoHeader.MLLT->TimeAndOffsets[EntryIdx+1].Milliseconds > StartTimeMillis)
+							{
+								break;
+							}
+						}
+						int64 NextEntryOffset = EntryIdx+1 < LastEntryIdx ? (int64)MPEGInfoHeader.MLLT->TimeAndOffsets[EntryIdx+1].Offset : MPEGInfoHeader.LastDataByte - MPEGInfoHeader.FirstDataByte;
+						int64 NextEntryMillis = EntryIdx+1 < LastEntryIdx ? (int64)MPEGInfoHeader.MLLT->TimeAndOffsets[EntryIdx+1].Milliseconds : MPEGInfoHeader.EstimatedDuration.GetAsMilliseconds();
+						double Frac = (double)(StartTimeMillis - MPEGInfoHeader.MLLT->TimeAndOffsets[EntryIdx].Milliseconds) / (double)(NextEntryMillis - MPEGInfoHeader.MLLT->TimeAndOffsets[EntryIdx].Milliseconds);
+						ApproxByteOffset = MPEGInfoHeader.FirstDataByte + MPEGInfoHeader.MLLT->TimeAndOffsets[EntryIdx].Offset + (int64)(Frac * (NextEntryOffset - MPEGInfoHeader.MLLT->TimeAndOffsets[EntryIdx].Offset));
+						ApproxByteOffset = FMath::Clamp(ApproxByteOffset, MPEGInfoHeader.FirstDataByte, MPEGInfoHeader.LastDataByte);
+					}
+					// Is there a VBRI seek table?
+					else if (MPEGInfoHeader.SeekTable.IsValid() && MPEGInfoHeader.FramesPerSeekTableEntry)
+					{
+						FTimeValue DurPerEntry = MPEGInfoHeader.EstimatedDuration / MPEGInfoHeader.SeekTable->Num();
+						FTimeValue DurSoFar(FTimeValue::GetZero());
+						int32 SeekTablePos;
+						ApproxByteOffset = MPEGInfoHeader.FirstDataByte;
+						for(SeekTablePos=0; SeekTablePos<MPEGInfoHeader.SeekTable->Num() && DurSoFar <= StartTime; ++SeekTablePos)
+						{
+							DurSoFar += DurPerEntry;
+							ApproxByteOffset += MPEGInfoHeader.SeekTable->operator[](SeekTablePos);
+						}
+						int32 Frac = FMath::Floor(((DurSoFar - StartTime).GetAsSeconds() / DurPerEntry.GetAsSeconds() + 1.0 / (MPEGInfoHeader.FramesPerSeekTableEntry * 2)) * MPEGInfoHeader.FramesPerSeekTableEntry);
+						int64 Delta = FMath::Floor(MPEGInfoHeader.SeekTable->operator[](SeekTablePos-1) * (double)Frac / MPEGInfoHeader.FramesPerSeekTableEntry);
+						ApproxByteOffset -= Delta;
+						ApproxByteOffset = FMath::Clamp(ApproxByteOffset, MPEGInfoHeader.FirstDataByte, MPEGInfoHeader.LastDataByte);
+					}
+					// Is there a TOC?
+					else if (MPEGInfoHeader.bHaveTOC)
+ 					{
+						// The TOC has 100 entries, mapping each full 1% of seek position to a percentage (scaled by 2.56 to map into 0-255 range) within the file.
+						// This makes seeking less precise the longer the file gets.
+						Percentage *= 100.0;
+						int32 Idx = FMath::Floor(Percentage);
+						double a = MPEGInfoHeader.TOC[FMath::Clamp(Idx, 0, 99)];
+						double b = Idx < 99 ? MPEGInfoHeader.TOC[Idx + 1] : 256.0;
+						double x = a + (b - a) * (Percentage - Idx);
+						Percentage = x / 256.0;
+						ApproxByteOffset = MPEGInfoHeader.FirstDataByte + (MPEGInfoHeader.LastDataByte - MPEGInfoHeader.FirstDataByte) * Percentage;
+						ApproxByteOffset = FMath::Clamp(ApproxByteOffset, MPEGInfoHeader.FirstDataByte, MPEGInfoHeader.LastDataByte);
+					}
+					else
+					{
+						// Short of any seek table we can only jump to somewhere in the file by percentage as we do with CBR.
+						ApproxByteOffset = MPEGInfoHeader.FirstDataByte + (MPEGInfoHeader.LastDataByte - MPEGInfoHeader.FirstDataByte) * Percentage;
+						ApproxByteOffset = FMath::Clamp(ApproxByteOffset, MPEGInfoHeader.FirstDataByte, MPEGInfoHeader.LastDataByte);
+					}
+				}
+			}
+			else
+			{
+				ApproxByteOffset = AtAbsoluteFilePos;
+				// Going past the end?
+				check(MPEGInfoHeader.LastDataByte > 0);
+				if (ApproxByteOffset >= MPEGInfoHeader.LastDataByte)
+				{
+					ApproxByteOffset = -1;
+				}
+			}
+
+			if (ApproxByteOffset >= 0)
+			{
+				TSharedPtrTS<FStreamSegmentRequestMPEGAudio> req(new FStreamSegmentRequestMPEGAudio);
+				OutSegment = req;
+				req->MediaAsset = SharedThis(this);
+				req->FirstPTS = StartTime;
+				req->EarliestPTS = StartTime;
+				req->LastPTS = PlayRangeEnd;
+
+				req->FileStartOffset = ApproxByteOffset;
+				req->Bitrate = MPEGInfoHeader.AverageBitrate ? MPEGInfoHeader.AverageBitrate : MPEGInfoHeader.Bitrate;
+				req->bIsContinuationSegment = false;
+				req->bIsFirstSegment = true;
+				req->TimestampSequenceIndex = InSequenceState.GetSequenceIndex();
+
+				// Approximate how many bytes will equal about n seconds
+				const double kSegmentDuration = 3.0;
+				double FrameDuration = (double)MPEGInfoHeader.SamplesPerFrame / (double)MPEGInfoHeader.SampleRate;
+				int32 ApproxSegmentSize = 0;
+				if (!MPEGInfoHeader.bIsVBR)
+				{
+					ApproxSegmentSize = MPEGInfoHeader.CBRFrameSize * (kSegmentDuration / FrameDuration);
+				}
+				else
+				{
+					int32 br = MPEGInfoHeader.AverageBitrate ? MPEGInfoHeader.AverageBitrate : MPEGInfoHeader.Bitrate;
+					ApproxSegmentSize = br * kSegmentDuration / 8.0;
+				}
+
+				req->FileEndOffset = FMath::Clamp(ApproxByteOffset + ApproxSegmentSize, 0, MPEGInfoHeader.LastDataByte);
+				req->bIsLastSegment = req->FileEndOffset >= MPEGInfoHeader.LastDataByte;
+				req->MPEGHeaderMask = MPEGInfoHeader.MPEGHeaderMask;
+				req->MPEGHeaderExpectedValue = MPEGInfoHeader.MPEGHeaderExpectedValue;
+				req->CBRFrameSize = MPEGInfoHeader.CBRFrameSize;
+				req->bIsVBR = MPEGInfoHeader.bIsVBR;
+				req->bIsLive = false;
+				req->Duration = kSegmentDuration;
+				req->CodecInfo = MPEGInfoHeader.CodecInfo;
+				// How many bytes will we be reading?
+				int64 SegmentInternalSize = req->FileEndOffset - req->FileStartOffset;
+				if (StartTime < PlayRangeEnd && SegmentInternalSize >= 12)
+				{
+					return IManifest::FResult(IManifest::FResult::EType::Found);
+				}
+				else
+				{
+					return IManifest::FResult(IManifest::FResult::EType::PastEOS);
+				}
+			}
+			else
+			{
+				return IManifest::FResult(IManifest::FResult::EType::PastEOS);
+			}
+		}
+	}
+	return IManifest::FResult(IManifest::FResult::EType::NotFound).SetErrorDetail(FErrorDetail().SetError(UEMEDIA_ERROR_INSUFFICIENT_DATA)
+					   .SetFacility(Facility::EFacility::MPEGAudioPlaylist)
+					   .SetCode(ERRCODE_MANIFEST_MPEGAUDIO_STARTSEGMENT_NOT_FOUND)
+					   .SetMessage(FString::Printf(TEXT("Could not find start segment for time %lld, no valid tracks"), (long long int)StartPosition.Time.GetAsHNS())));
+}
+
+IManifest::FResult FManifestMPEGAudioInternal::FTimelineAssetMPEGAudio::GetNextSegment(TSharedPtrTS<IStreamSegment>& OutSegment, TSharedPtrTS<const IStreamSegment> CurrentSegment, const FPlayStartOptions& Options)
+{
+	const FStreamSegmentRequestMPEGAudio* Request = static_cast<const FStreamSegmentRequestMPEGAudio*>(CurrentSegment.Get());
+	if (Request)
+	{
+		// Check if the current request did not already go up to the end of the stream. If so there is no next segment.
+		if (MPEGInfoHeader.bIsLive || (Request->FileEndOffset >= 0 && Request->FileEndOffset < MPEGInfoHeader.LastDataByte))
+		{
+			FPlayStartPosition dummyPos;
+			FPlayerSequenceState seqState;
+			dummyPos.Options = Options;
+			dummyPos.Time = Request->LastSuccessfullyUsedPTS;
+			seqState.SetSequenceIndex(Request->TimestampSequenceIndex);
+			IManifest::FResult res = GetStartingSegment(OutSegment, seqState, dummyPos, ESearchType::Same, Request->LastSuccessfullyUsedBytePos);
+			if (res.GetType() == IManifest::FResult::EType::Found)
+			{
+				FStreamSegmentRequestMPEGAudio* NextRequest = static_cast<FStreamSegmentRequestMPEGAudio*>(OutSegment.Get());
+				NextRequest->bIsContinuationSegment = true;
+				if (!MPEGInfoHeader.bIsLive)
+				{
+					NextRequest->bIsFirstSegment = false;
+					NextRequest->EarliestPTS = Request->EarliestPTS;
+					NextRequest->LastPTS = Request->LastPTS;
+				}
+				else
+				{
+					NextRequest->NumOverallRetries = Request->NumOverallRetries + 1;
+				}
+				return res;
+			}
+		}
+	}
+	return IManifest::FResult(IManifest::FResult::EType::PastEOS);
+}
+
+IManifest::FResult FManifestMPEGAudioInternal::FTimelineAssetMPEGAudio::GetRetrySegment(TSharedPtrTS<IStreamSegment>& OutSegment, TSharedPtrTS<const IStreamSegment> CurrentSegment, const FPlayStartOptions& Options, bool bReplaceWithFillerData)
+{
+	const FStreamSegmentRequestMPEGAudio* Request = static_cast<const FStreamSegmentRequestMPEGAudio*>(CurrentSegment.Get());
+	if (Request)
+	{
+		FPlayStartPosition dummyPos;
+		FPlayerSequenceState seqState;
+		dummyPos.Options = Options;
+		dummyPos.Time = Request->LastSuccessfullyUsedPTS.IsValid() ? Request->LastSuccessfullyUsedPTS : Request->FirstPTS;
+		seqState.SetSequenceIndex(Request->TimestampSequenceIndex);
+		IManifest::FResult res = GetStartingSegment(OutSegment, seqState, dummyPos, ESearchType::Same, Request->LastSuccessfullyUsedBytePos);
+		if (res.GetType() == IManifest::FResult::EType::Found)
+		{
+			FStreamSegmentRequestMPEGAudio* RetryRequest = static_cast<FStreamSegmentRequestMPEGAudio*>(OutSegment.Get());
+			RetryRequest->bIsContinuationSegment = true;
+			RetryRequest->NumOverallRetries = Request->NumOverallRetries + 1;
+			RetryRequest->EarliestPTS = Request->EarliestPTS;
+			RetryRequest->LastPTS = Request->LastPTS;
+			return res;
+		}
+	}
+	return IManifest::FResult(IManifest::FResult::EType::NotFound);
+}
+
+
+IManifest::FResult FManifestMPEGAudioInternal::FTimelineAssetMPEGAudio::GetLoopingSegment(TSharedPtrTS<IStreamSegment>& OutSegment, const FPlayerSequenceState& SequenceState, const FPlayStartPosition& StartPosition, ESearchType SearchType)
+{
+	return GetStartingSegment(OutSegment, SequenceState, StartPosition, SearchType, -1);
+}
+
+
+void FManifestMPEGAudioInternal::FTimelineAssetMPEGAudio::GetSegmentInformation(TArray<IManifest::IPlayPeriod::FSegmentInformation>& OutSegmentInformation, FTimeValue& OutAverageSegmentDuration, TSharedPtrTS<const IStreamSegment> CurrentSegment, const FTimeValue& LookAheadTime, const FString& AdaptationSetID, const FString& RepresentationID)
+{
+	SCOPE_CYCLE_COUNTER(STAT_ElectraPlayer_MPEGAudio_GetSegmentInformation);
+	CSV_SCOPED_TIMING_STAT(ElectraPlayer, MPEGAudio_GetSegmentInformation);
+
+	// This is not expected to be called. And if it does we return a dummy entry.
+	OutAverageSegmentDuration.SetFromSeconds(60.0);
+	OutSegmentInformation.Empty();
+	IManifest::IPlayPeriod::FSegmentInformation& si = OutSegmentInformation.AddDefaulted_GetRef();
+	si.ByteSize = 16384;
+	si.Duration.SetFromSeconds(1.0);
+}
+
+FErrorDetail FManifestMPEGAudioInternal::FAdaptationSetMPEGAudio::CreateFrom(const FMPEGAudioInfoHeader& InHeader, const FString& URL)
+{
+	Representation = MakeSharedTS<FRepresentationMPEGAudio>();
+	FErrorDetail err = Representation->CreateFrom(InHeader, URL);
+	if (err.IsOK())
+	{
+		CodecRFC6381 = Representation->GetCodecInformation().GetCodecSpecifierRFC6381();
+		UniqueIdentifier = Representation->GetUniqueIdentifier();
+		Language = TEXT("und");
+		bIsLivecast = InHeader.bIsLive;
+	}
+	return err;
+}
+
+bool FManifestMPEGAudioInternal::FAdaptationSetMPEGAudio::IsLowLatencyEnabled() const
+{
+#if USE_LOW_LATENCY_DESCRIPTOR
+	return bIsLivecast;
+#else
+	return false;
+#endif
+}
+
+
+FErrorDetail FManifestMPEGAudioInternal::FRepresentationMPEGAudio::CreateFrom(const FMPEGAudioInfoHeader& InHeader, const FString& URL)
+{
+	CodecInformation = InHeader.CodecInfo;
+
+	// NOTE: This *MUST* be just a number since it gets parsed back out from a string into a number later! Do *NOT* prepend/append any string literals!!
+	UniqueIdentifier = TEXT("1");
+
+	// Get name from metadata if it exists
+	if (InHeader.ID3v2.IsValid() && InHeader.ID3v2->HaveTag(Utils::Make4CC('T','I','T','2')))
+	{
+		MPEG::FID3V2Metadata::FItem Tit2;
+		InHeader.ID3v2->GetTag(Tit2, Utils::Make4CC('T','I','T','2'));
+		Name = Tit2.Value.GetValue<FString>();
+	}
+	else
+	{
+		Name = TEXT("Unknown");
+	}
+
+	Bitrate = InHeader.AverageBitrate ? InHeader.AverageBitrate : InHeader.Bitrate;
+	Bitrate = Bitrate <= 0 ? 64000 : Bitrate;
+
+	if (!CodecInformation.GetBitrate())
+	{
+		CodecInformation.SetBitrate(Bitrate);
+	}
+	return FErrorDetail();
+}
+
+
+} // namespace Electra
diff --git a/Engine/Plugins/Media/ElectraPlayer/Source/ElectraPlayerRuntime/Private/Runtime/Player/mpegaudio/ManifestMPEGAudio.h b/Engine/Plugins/Media/ElectraPlayer/Source/ElectraPlayerRuntime/Private/Runtime/Player/mpegaudio/ManifestMPEGAudio.h
new file mode 100644
--- /dev/null
+++ b/Engine/Plugins/Media/ElectraPlayer/Source/ElectraPlayerRuntime/Private/Runtime/Player/mpegaudio/ManifestMPEGAudio.h
@@ -0,0 +1,345 @@
+// Copyright Epic Games, Inc. All Rights Reserved.
+
+#pragma once
+
+#include "PlayerCore.h"
+#include "Player/Manifest.h"
+#include "Player/PlaybackTimeline.h"
+#include "HTTP/HTTPManager.h"
+#include "SynchronizedClock.h"
+#include "HAL/LowLevelMemTracker.h"
+#include "ElectraPlayerPrivate.h"
+#include "Player/PlayerStreamReader.h"
+#include "Utilities/Utilities.h"
+#include "Utilities/URLParser.h"
+#include "Utilities/UtilsMPEG.h"
+
+namespace Electra
+{
+
+struct FMPEGAudioInfoHeader
+{
+	TSharedPtrTS<MPEG::FID3V2Metadata> ID3v2;
+	uint32 MPEGHeaderMask = 0;
+	uint32 MPEGHeaderExpectedValue = 0;
+	int32 MPEGVersion = 0;
+	int32 MPEGLayer = 0;
+	int32 SampleRate = 0;
+	int32 NumChannels = 0;
+	bool bIsVBR = false;
+	int32 Bitrate = 0;
+	int32 AverageBitrate = 0;
+	int32 CBRFrameSize = 0;
+	int32 SamplesPerFrame = 0;
+	uint32 NumFrames = 0;
+	int32 EncoderDelayStart = 0;
+	int32 EncoderDelayEnd = 0;
+	int64 FirstDataByte = -1;
+	int64 LastDataByte = -1;
+	bool bIsLive = false;
+	bool bHaveTOC = false;
+	TArray<uint8> TOC;
+	int32 FramesPerSeekTableEntry = 0;
+	TSharedPtrTS<TArray<uint32>> SeekTable;
+	struct FMLLT
+	{
+		struct FTimeAndOffset
+		{
+			uint32 Offset = 0;
+			uint32 Milliseconds = 0;
+		};
+		uint32 FramesBetweenReferences = 0;
+		TArray<FTimeAndOffset> TimeAndOffsets;
+	};
+	TSharedPtrTS<FMLLT> MLLT;
+	// Generated codec info from above values.
+	FStreamCodecInformation CodecInfo;
+	// Estimated duration
+	FTimeValue EstimatedDuration;
+	// HTTP response headers
+	TMultiMap<FString, FString> HTTPResponseHeaders;
+};
+
+/**
+ * This class represents the internal "manifest" or "playlist" of an MPEG audio file.
+ */
+class FManifestMPEGAudioInternal : public IManifest, public TSharedFromThis<FManifestMPEGAudioInternal, ESPMode::ThreadSafe>
+{
+public:
+	FManifestMPEGAudioInternal(IPlayerSessionServices* InPlayerSessionServices);
+	virtual ~FManifestMPEGAudioInternal();
+
+	FErrorDetail Build(const FMPEGAudioInfoHeader& InHeader, const FString& URL);
+
+	EType GetPresentationType() const override;
+	TSharedPtrTS<const FLowLatencyDescriptor> GetLowLatencyDescriptor() const override;
+	FTimeValue GetAnchorTime() const override
+	{ return FTimeValue::GetZero(); }
+	FTimeRange GetTotalTimeRange() const override
+	{ return MediaAsset.IsValid() ? MediaAsset->GetTimeRange() : FTimeRange(); }
+	FTimeRange GetSeekableTimeRange() const override
+	{ return GetTotalTimeRange(); }
+	FTimeRange GetPlaybackRange(EPlaybackRangeType InRangeType) const override;
+	void GetSeekablePositions(TArray<FTimespan>& OutPositions) const override
+	{ }
+	FTimeValue GetDuration() const override
+	{ return MediaAsset.IsValid() ? MediaAsset->GetDuration() : FTimeValue(); }
+	FTimeValue GetDefaultStartTime() const override
+	{ return DefaultStartTime; }
+	void ClearDefaultStartTime() override
+	{ DefaultStartTime.SetToInvalid(); }
+	FTimeValue GetDefaultEndTime() const override
+	{ return DefaultEndTime; }
+	void ClearDefaultEndTime() override
+	{ DefaultEndTime.SetToInvalid(); }
+	void GetTrackMetadata(TArray<FTrackMetadata>& OutMetadata, EStreamType StreamType) const override;
+	void UpdateRunningMetaData(TSharedPtrTS<UtilsMP4::FMetadataParser> InUpdatedMetaData) override;
+	FTimeValue GetMinBufferTime() const override;
+	FTimeValue GetDesiredLiveLatency() const override;
+	ELiveEdgePlayMode GetLiveEdgePlayMode() const override;
+
+	TRangeSet<double> GetPossiblePlaybackRates(EPlayRateType InForType) const override;
+	TSharedPtrTS<IProducerReferenceTimeInfo> GetProducerReferenceTimeInfo(int64 ID) const override;
+	void UpdateDynamicRefetchCounter() override;
+	void TriggerClockSync(EClockSyncType InClockSyncType) override;
+	void TriggerPlaylistRefresh() override;
+
+	IStreamReader* CreateStreamReaderHandler() override;
+	FResult FindPlayPeriod(TSharedPtrTS<IPlayPeriod>& OutPlayPeriod, const FPlayStartPosition& StartPosition, ESearchType SearchType) override;
+	FResult FindNextPlayPeriod(TSharedPtrTS<IPlayPeriod>& OutPlayPeriod, TSharedPtrTS<const IStreamSegment> CurrentSegment) override;
+
+	class FRepresentationMPEGAudio : public IPlaybackAssetRepresentation
+	{
+	public:
+		virtual ~FRepresentationMPEGAudio()
+		{ }
+
+		FErrorDetail CreateFrom(const FMPEGAudioInfoHeader& InHeader, const FString& URL);
+
+		FString GetUniqueIdentifier() const override
+		{ return UniqueIdentifier; }
+		const FStreamCodecInformation& GetCodecInformation() const override
+		{ return CodecInformation; }
+		int32 GetBitrate() const override
+		{ return Bitrate; }
+		int32 GetQualityIndex() const override
+		{ return 0; }
+		bool CanBePlayed() const override
+		{ return true; }
+		const FString& GetName() const
+		{ return Name; }
+		void SetName(const FString& InNewName)
+		{ Name = InNewName; }
+	private:
+		FStreamCodecInformation CodecInformation;
+		FString UniqueIdentifier;
+		FString Name;
+		int32 Bitrate = 0;
+	};
+
+	class FAdaptationSetMPEGAudio : public IPlaybackAssetAdaptationSet
+	{
+	public:
+		virtual ~FAdaptationSetMPEGAudio()
+		{ }
+
+		FErrorDetail CreateFrom(const FMPEGAudioInfoHeader& InHeader, const FString& URL);
+
+		FString GetUniqueIdentifier() const override
+		{ return UniqueIdentifier; }
+		FString GetListOfCodecs() const override
+		{ return CodecRFC6381; }
+		FString GetLanguage() const override
+		{ return Language; }
+		int32 GetNumberOfRepresentations() const override
+		{ return Representation.IsValid() ? 1 : 0; }
+		bool IsLowLatencyEnabled() const override;
+		TSharedPtrTS<IPlaybackAssetRepresentation> GetRepresentationByIndex(int32 RepresentationIndex) const override
+		{ return RepresentationIndex == 0 ? Representation : TSharedPtrTS<IPlaybackAssetRepresentation>(); }
+		TSharedPtrTS<IPlaybackAssetRepresentation> GetRepresentationByUniqueIdentifier(const FString& InUniqueIdentifier) const override
+		{ return Representation.IsValid() && Representation->GetUniqueIdentifier() == InUniqueIdentifier ? Representation : TSharedPtrTS<IPlaybackAssetRepresentation>(); }
+	private:
+		TSharedPtrTS<FRepresentationMPEGAudio> Representation;
+		FString Language;
+		FString CodecRFC6381;
+		FString UniqueIdentifier;
+		bool bIsLivecast = false;
+	};
+
+	class FTimelineAssetMPEGAudio : public ITimelineMediaAsset, public TSharedFromThis<FTimelineAssetMPEGAudio, ESPMode::ThreadSafe>
+	{
+	public:
+		FTimelineAssetMPEGAudio()
+			: PlayerSessionServices(nullptr)
+		{ }
+
+		virtual ~FTimelineAssetMPEGAudio()
+		{ }
+
+		FErrorDetail Build(IPlayerSessionServices* InPlayerSessionServices, const FMPEGAudioInfoHeader& InHeader, const FString& URL);
+
+		FTimeRange GetTimeRange() const override
+		{
+			FTimeRange tr;
+			tr.Start.SetToZero();
+			tr.End = GetDuration();
+			return tr;
+		}
+
+		FTimeValue GetDuration() const override
+		{
+			return MPEGInfoHeader.EstimatedDuration;
+		}
+
+		FString GetAssetIdentifier() const override
+		{ return FString("mpegaudio-asset.0"); }
+		FString GetUniqueIdentifier() const override
+		{ return FString("mpegaudio-media.0"); }
+		int32 GetNumberOfAdaptationSets(EStreamType OfStreamType) const override
+		{
+			return OfStreamType == EStreamType::Audio ? AudioAdaptationSets.Num() : 0;
+		}
+		TSharedPtrTS<IPlaybackAssetAdaptationSet> GetAdaptationSetByTypeAndIndex(EStreamType OfStreamType, int32 AdaptationSetIndex) const override
+		{
+			return OfStreamType == EStreamType::Audio ? AdaptationSetIndex < AudioAdaptationSets.Num() ? AudioAdaptationSets[AdaptationSetIndex] : TSharedPtrTS<IPlaybackAssetAdaptationSet>() : TSharedPtrTS<IPlaybackAssetAdaptationSet>();
+		}
+
+		void GetMetaData(TArray<FTrackMetadata>& OutMetadata, EStreamType StreamType) const override
+		{
+			for(int32 i=0, iMax=GetNumberOfAdaptationSets(StreamType); i<iMax; ++i)
+			{
+				TSharedPtrTS<FAdaptationSetMPEGAudio> AdaptSet = StaticCastSharedPtr<FAdaptationSetMPEGAudio>(GetAdaptationSetByTypeAndIndex(StreamType, i));
+				if (AdaptSet.IsValid())
+				{
+					FTrackMetadata tm;
+					tm.ID = AdaptSet->GetUniqueIdentifier();
+					tm.Language = AdaptSet->GetLanguage();
+					tm.Kind = i==0 ? TEXT("main") : TEXT("translation");
+
+					for(int32 j=0, jMax=AdaptSet->GetNumberOfRepresentations(); j<jMax; ++j)
+					{
+						TSharedPtrTS<FRepresentationMPEGAudio> Repr = StaticCastSharedPtr<FRepresentationMPEGAudio>(AdaptSet->GetRepresentationByIndex(j));
+						if (Repr.IsValid())
+						{
+							FStreamMetadata sd;
+							sd.Bandwidth = Repr->GetBitrate();
+							sd.CodecInformation = Repr->GetCodecInformation();
+							sd.ID = Repr->GetUniqueIdentifier();
+							// There is only 1 "stream" per "track" so we can set the highest bitrate and codec info the same as the track.
+							tm.HighestBandwidth = sd.Bandwidth;
+							tm.HighestBandwidthCodec = sd.CodecInformation;
+							tm.Label = Repr->GetName();
+							tm.StreamDetails.Emplace(MoveTemp(sd));
+						}
+					}
+					OutMetadata.Emplace(MoveTemp(tm));
+				}
+			}
+		}
+
+		void UpdateRunningMetaData(const FString& InKindOfValue, const FVariant& InNewValue) override
+		{
+			if (InKindOfValue.Equals(TEXT("Title")))
+			{
+				if (GetNumberOfAdaptationSets(EStreamType::Audio))
+				{
+					TSharedPtrTS<FAdaptationSetMPEGAudio> AdaptSet = StaticCastSharedPtr<FAdaptationSetMPEGAudio>(GetAdaptationSetByTypeAndIndex(EStreamType::Audio, 0));
+					if (AdaptSet.IsValid() && AdaptSet->GetNumberOfRepresentations())
+					{
+						TSharedPtrTS<FRepresentationMPEGAudio> Repr = StaticCastSharedPtr<FRepresentationMPEGAudio>(AdaptSet->GetRepresentationByIndex(0));
+						if (Repr.IsValid())
+						{
+							Repr->SetName(InNewValue.GetValue<FString>());
+						}
+					}
+				}
+				if (!MPEGInfoHeader.ID3v2.IsValid())
+				{
+					MPEGInfoHeader.ID3v2 = MakeSharedTS<MPEG::FID3V2Metadata>();
+				}
+				MPEGInfoHeader.ID3v2->GetTags().Add(Utils::Make4CC('T','I','T','2'), MPEG::FID3V2Metadata::FItem({.Value = InNewValue }));
+			}
+		}
+
+		FResult GetStartingSegment(TSharedPtrTS<IStreamSegment>& OutSegment, const FPlayerSequenceState& InSequenceState, const FPlayStartPosition& StartPosition, ESearchType SearchType, int64 AtAbsoluteFilePos);
+		FResult GetNextSegment(TSharedPtrTS<IStreamSegment>& OutSegment, TSharedPtrTS<const IStreamSegment> CurrentSegment, const FPlayStartOptions& Options);
+		FResult GetRetrySegment(TSharedPtrTS<IStreamSegment>& OutSegment, TSharedPtrTS<const IStreamSegment> CurrentSegment, const FPlayStartOptions& Options, bool bReplaceWithFillerData);
+		FResult GetLoopingSegment(TSharedPtrTS<IStreamSegment>& OutSegment, const FPlayerSequenceState& InSequenceState, const FPlayStartPosition& StartPosition, ESearchType SearchType);
+
+		void GetSegmentInformation(TArray<IManifest::IPlayPeriod::FSegmentInformation>& OutSegmentInformation, FTimeValue& OutAverageSegmentDuration, TSharedPtrTS<const IStreamSegment> CurrentSegment, const FTimeValue& LookAheadTime, const FString& AdaptationSetID, const FString& RepresentationID);
+		const FString& GetMediaURL() const
+		{ return MediaURL; }
+
+	private:
+		void LogMessage(IInfoLog::ELevel Level, const FString& Message);
+		IPlayerSessionServices* PlayerSessionServices = nullptr;
+		FString MediaURL;
+		TArray<TSharedPtrTS<FAdaptationSetMPEGAudio>> AudioAdaptationSets;
+		FMPEGAudioInfoHeader MPEGInfoHeader;
+	};
+
+
+	class FPlayPeriodMPEGAudio : public IManifest::IPlayPeriod
+	{
+	public:
+		FPlayPeriodMPEGAudio(TSharedPtrTS<FTimelineAssetMPEGAudio> InMediaAsset);
+		virtual ~FPlayPeriodMPEGAudio();
+		void SetStreamPreferences(EStreamType ForStreamType, const FStreamSelectionAttributes& StreamAttributes) override;
+		EReadyState GetReadyState() override;
+		void Load() override;
+		void PrepareForPlay() override;
+		int64 GetDefaultStartingBitrate() const override;
+		TSharedPtrTS<FBufferSourceInfo> GetSelectedStreamBufferSourceInfo(EStreamType StreamType) override;
+		FString GetSelectedAdaptationSetID(EStreamType StreamType) override;
+		ETrackChangeResult ChangeTrackStreamPreference(EStreamType ForStreamType, const FStreamSelectionAttributes& StreamAttributes) override;
+		TSharedPtrTS<ITimelineMediaAsset> GetMediaAsset() const override;
+		void SelectStream(const FString& AdaptationSetID, const FString& RepresentationID, int32 QualityIndex, int32 MaxQualityIndex) override;
+		void TriggerInitSegmentPreload(const TArray<FInitSegmentPreload>& InitSegmentsToPreload) override;
+		FResult GetStartingSegment(TSharedPtrTS<IStreamSegment>& OutSegment, const FPlayerSequenceState& InSequenceState, const FPlayStartPosition& StartPosition, ESearchType SearchType) override;
+		FResult GetContinuationSegment(TSharedPtrTS<IStreamSegment>& OutSegment, EStreamType StreamType, const FPlayerSequenceState& LoopState, const FPlayStartPosition& StartPosition, ESearchType SearchType) override;
+		FResult GetNextSegment(TSharedPtrTS<IStreamSegment>& OutSegment, TSharedPtrTS<const IStreamSegment> CurrentSegment, const FPlayStartOptions& Options) override;
+		FResult GetRetrySegment(TSharedPtrTS<IStreamSegment>& OutSegment, TSharedPtrTS<const IStreamSegment> CurrentSegment, const FPlayStartOptions& Options, bool bReplaceWithFillerData) override;
+		FResult GetLoopingSegment(TSharedPtrTS<IStreamSegment>& OutSegment, const FPlayerSequenceState& InSequenceState, const FPlayStartPosition& StartPosition, ESearchType SearchType) override;
+		void IncreaseSegmentFetchDelay(const FTimeValue& IncreaseAmount) override;
+		void GetSegmentInformation(TArray<FSegmentInformation>& OutSegmentInformation, FTimeValue& OutAverageSegmentDuration, TSharedPtrTS<const IStreamSegment> CurrentSegment, const FTimeValue& LookAheadTime, const FString& AdaptationSetID, const FString& RepresentationID) override;
+
+	private:
+		void SelectInitialStream(EStreamType StreamType);
+		TSharedPtrTS<FTrackMetadata> SelectMetadataForAttributes(EStreamType StreamType, const FStreamSelectionAttributes& InAttributes);
+		void MakeBufferSourceInfoFromMetadata(EStreamType StreamType, TSharedPtrTS<FBufferSourceInfo>& OutBufferSourceInfo, TSharedPtrTS<FTrackMetadata> InMetadata);
+
+		TWeakPtrTS<FTimelineAssetMPEGAudio> MediaAsset;
+		FStreamSelectionAttributes AudioPreferences;
+		TSharedPtrTS<FTrackMetadata> SelectedAudioMetadata;
+		TSharedPtrTS<FBufferSourceInfo> AudioBufferSourceInfo;
+		EReadyState CurrentReadyState;
+	};
+
+	void LogMessage(IInfoLog::ELevel Level, const FString& Message);
+
+	const TArray<FURL_RFC3986::FQueryParam>& GetURLFragmentComponents() const
+	{
+		return URLFragmentComponents;
+	}
+
+	void SetURLFragmentComponents(TArray<FURL_RFC3986::FQueryParam> InURLFragmentComponents)
+	{
+		URLFragmentComponents = MoveTemp(InURLFragmentComponents);
+	}
+
+	const TArray<FURL_RFC3986::FQueryParam>& GetURLFragmentComponents()
+	{
+		return URLFragmentComponents;
+	}
+
+	IPlayerSessionServices* PlayerSessionServices = nullptr;
+	TSharedPtrTS<FTimelineAssetMPEGAudio> MediaAsset;
+	TArray<FURL_RFC3986::FQueryParam> URLFragmentComponents;
+	FTimeValue DefaultStartTime;
+	FTimeValue DefaultEndTime;
+	FMPEGAudioInfoHeader MPEGInfoHeader;
+	TSharedPtrTS<FLowLatencyDescriptor> LatencyDescriptor;
+};
+
+
+} // namespace Electra
diff --git a/Engine/Plugins/Media/ElectraPlayer/Source/ElectraPlayerRuntime/Private/Runtime/Player/mpegaudio/OptionKeynamesMPEGAudio.h b/Engine/Plugins/Media/ElectraPlayer/Source/ElectraPlayerRuntime/Private/Runtime/Player/mpegaudio/OptionKeynamesMPEGAudio.h
new file mode 100644
--- /dev/null
+++ b/Engine/Plugins/Media/ElectraPlayer/Source/ElectraPlayerRuntime/Private/Runtime/Player/mpegaudio/OptionKeynamesMPEGAudio.h
@@ -0,0 +1,18 @@
+// Copyright Epic Games, Inc. All Rights Reserved.
+
+#pragma once
+
+#include "PlayerCore.h"
+
+namespace Electra
+{
+
+namespace MPEGAudio
+{
+const FName OptionKeyMPEGAudioLoadConnectTimeout(TEXT("mpegaudio_connection_timeout"));
+const FName OptionKeyMPEGAudioLoadNoDataTimeout(TEXT("mpegaudio_nodata_timeout"));
+const FName OptionKeyMPEGAudioLoadTimeout(TEXT("mpegaudio_load_timeout"));
+const FName OptionKeyMPEGAudioMaxPreloadBufferDuration(TEXT("mpegaudio_max_preload_buffer_duration"));
+}
+
+} // namespace Electra
diff --git a/Engine/Plugins/Media/ElectraPlayer/Source/ElectraPlayerRuntime/Private/Runtime/Player/mpegaudio/PlaylistReaderMPEGAudio.cpp b/Engine/Plugins/Media/ElectraPlayer/Source/ElectraPlayerRuntime/Private/Runtime/Player/mpegaudio/PlaylistReaderMPEGAudio.cpp
new file mode 100644
--- /dev/null
+++ b/Engine/Plugins/Media/ElectraPlayer/Source/ElectraPlayerRuntime/Private/Runtime/Player/mpegaudio/PlaylistReaderMPEGAudio.cpp
@@ -0,0 +1,965 @@
+// Copyright Epic Games, Inc. All Rights Reserved.
+
+#include "PlayerCore.h"
+#include "Player/Manifest.h"
+#include "HTTP/HTTPManager.h"
+#include "HTTP/HTTPResponseCache.h"
+#include "Http.h"
+#include "HAL/FileManager.h"
+#include "Misc/FileHelper.h"
+#include "Serialization/Archive.h"
+#include "Serialization/ArrayReader.h"
+#include "Stats/Stats.h"
+#include "SynchronizedClock.h"
+#include "PlaylistReaderMPEGAudio.h"
+#include "Utilities/Utilities.h"
+#include "Utilities/TimeUtilities.h"
+#include "Utilities/URLParser.h"
+#include "HAL/LowLevelMemTracker.h"
+#include "ElectraPlayerPrivate.h"
+#include "Player/mpegaudio/ManifestMPEGAudio.h"
+#include "Player/mpegaudio/OptionKeynamesMPEGAudio.h"
+#include "Utils/MPEG/ElectraUtilsMPEGAudio.h"
+#include "Utils/ElectraBitstreamReader.h"
+
+#define ERRCODE_MPEGAUDIO_INVALID_FILE		1
+#define ERRCODE_MPEGAUDIO_DOWNLOAD_ERROR	2
+
+
+DECLARE_CYCLE_STAT(TEXT("FPlaylistReaderMPEGAudio_WorkerThread"), STAT_ElectraPlayer_MPEGAudio_PlaylistWorker, STATGROUP_ElectraPlayer);
+
+
+namespace Electra
+{
+
+class FPlaylistReaderMPEGAudio : public IPlaylistReaderMPEGAudio, public FMediaThread
+{
+public:
+	FPlaylistReaderMPEGAudio();
+	void Initialize(IPlayerSessionServices* PlayerSessionServices);
+
+	virtual ~FPlaylistReaderMPEGAudio();
+
+	void Close() override;
+	void HandleOnce() override;
+
+	/**
+	 * Returns the type of playlist format.
+	 *
+	 * @return "mpegaudio"
+	 */
+	const FString& GetPlaylistType() const override
+	{
+		static FString Type("mpegaudio");
+		return Type;
+	}
+
+	/**
+	 * Loads and parses the playlist.
+	 *
+	 * @param URL URL of the playlist to load
+	 */
+	void LoadAndParse(const FString& URL) override;
+
+	/**
+	 * Returns the URL from which the playlist was loaded (or supposed to be loaded).
+	 *
+	 * @return The playlist URL
+	 */
+	FString GetURL() const override;
+
+	/**
+	 * Returns an interface to the manifest.
+	 *
+	 * @return A shared manifest interface pointer.
+	 */
+	TSharedPtrTS<IManifest> GetManifest() override;
+
+private:
+	void StartWorkerThread();
+	void StopWorkerThread();
+	void WorkerThread();
+
+	void PostError(const FString& Message, uint16 Code, UEMediaError Error = UEMEDIA_ERROR_OK);
+	void LogMessage(IInfoLog::ELevel Level, const FString& Message);
+
+	struct FRequest : public TSharedFromThis<FRequest, ESPMode::ThreadSafe>
+	{
+		static constexpr int32 GetMinProbeSize()
+		{
+			return 32 << 10;
+		}
+		static constexpr int32 GetMaxProbeSize()
+		{
+			return 1 << 20;
+		}
+		static constexpr uint32 Make4CC(const uint8 A, const uint8 B, const uint8 C, const uint8 D)
+		{
+			return (static_cast<uint32>(A) << 24) | (static_cast<uint32>(B) << 16) | (static_cast<uint32>(C) << 8) | static_cast<uint32>(D);
+		}
+		static uint32 GetUINT32BE(const uint8* Data)
+		{
+			return (static_cast<uint32>(Data[0]) << 24) | (static_cast<uint32>(Data[1]) << 16) | (static_cast<uint32>(Data[2]) << 8) | static_cast<uint32>(Data[3]);
+		}
+
+		FRequest()
+		{
+			Events.Resize(4);
+		}
+		void OnProcessRequestComplete(FHttpRequestPtr InSourceHttpRequest, FHttpResponsePtr InHttpResponse, bool bInSucceeded);
+		void OnHeaderReceived(FHttpRequestPtr InSourceHttpRequest, const FString& InHeaderName, const FString& InHeaderValue);
+		void OnStatusCodeReceived(FHttpRequestPtr InSourceHttpRequest, int32 InHttpStatusCode);
+		bool OnProcessRequestStream(void* InDataPtr, int64 InLength);
+		void FindSyncMarkers();
+		void Cancel();
+		bool Validate(FMPEGAudioInfoHeader& InOutInfoHeader, FString& OutError, TArray<FString>& OutWarnings);
+
+		enum class EEvent
+		{
+			None,
+			Abort,
+			HaveProbeData,
+			Finished,
+		};
+		TMediaMessageQueueWithTimeout<EEvent> Events;
+		enum class EResult
+		{
+			Running,
+			Succeeded,
+			Aborted,
+			Failed,
+		};
+		EResult Result = EResult::Running;
+
+		TSharedPtr<IHttpRequest, ESPMode::ThreadSafe> Handle;
+		TMultiMap<FString, FString> Headers;
+		TArray<uint8> Buffer;
+		TArray<int32> SyncMarkerOffsets;
+		int64 ContentLength = -1;
+		int32 StatusCode = -1;
+		int32 ProbeSize = -1;
+		int32 CurrentSyncMarkerCheckPos = 0;
+		bool bSentHaveProbeDataMsg = false;
+	};
+
+	IPlayerSessionServices* PlayerSessionServices = nullptr;
+	FString PlaylistURL;
+	FString URLFragment;
+	FMediaEvent WorkerThreadQuitSignal;
+	bool bIsWorkerThreadStarted = false;
+	HTTP::FConnectionInfo ConnectionInfo;
+	volatile bool bAbort = false;
+	FErrorDetail LastErrorDetail;
+
+	FMPEGAudioInfoHeader InfoHeader;
+	TSharedPtrTS<FManifestMPEGAudioInternal> Manifest;
+};
+
+
+/***************************************************************************************************************************************************/
+/***************************************************************************************************************************************************/
+/***************************************************************************************************************************************************/
+
+TSharedPtrTS<IPlaylistReader> IPlaylistReaderMPEGAudio::Create(IPlayerSessionServices* PlayerSessionServices)
+{
+	TSharedPtrTS<FPlaylistReaderMPEGAudio> PlaylistReader = MakeSharedTS<FPlaylistReaderMPEGAudio>();
+	if (PlaylistReader)
+	{
+		PlaylistReader->Initialize(PlayerSessionServices);
+	}
+	return PlaylistReader;
+}
+
+/***************************************************************************************************************************************************/
+/***************************************************************************************************************************************************/
+/***************************************************************************************************************************************************/
+
+
+FPlaylistReaderMPEGAudio::FPlaylistReaderMPEGAudio()
+	: FMediaThread("ElectraPlayer::MPEGAudio Playlist")
+{
+}
+
+FPlaylistReaderMPEGAudio::~FPlaylistReaderMPEGAudio()
+{
+	Close();
+}
+
+FString FPlaylistReaderMPEGAudio::GetURL() const
+{
+	return PlaylistURL;
+}
+
+TSharedPtrTS<IManifest> FPlaylistReaderMPEGAudio::GetManifest()
+{
+	return Manifest;
+}
+
+void FPlaylistReaderMPEGAudio::Initialize(IPlayerSessionServices* InPlayerSessionServices)
+{
+	PlayerSessionServices = InPlayerSessionServices;
+}
+
+void FPlaylistReaderMPEGAudio::Close()
+{
+	bAbort = true;
+	StopWorkerThread();
+}
+
+void FPlaylistReaderMPEGAudio::HandleOnce()
+{
+	// No-op. This class is using a dedicated thread to read data from the stream
+	// which can stall at any moment and thus not lend itself to a tickable instance.
+}
+
+void FPlaylistReaderMPEGAudio::StartWorkerThread()
+{
+	check(!bIsWorkerThreadStarted);
+	ThreadStart(FMediaRunnable::FStartDelegate::CreateRaw(this, &FPlaylistReaderMPEGAudio::WorkerThread));
+	bIsWorkerThreadStarted = true;
+}
+
+void FPlaylistReaderMPEGAudio::StopWorkerThread()
+{
+	if (bIsWorkerThreadStarted)
+	{
+		WorkerThreadQuitSignal.Signal();
+		ThreadWaitDone();
+		ThreadReset();
+		bIsWorkerThreadStarted = false;
+	}
+}
+
+void FPlaylistReaderMPEGAudio::PostError(const FString& InMessage, uint16 InCode, UEMediaError InError)
+{
+	LastErrorDetail.Clear();
+	LastErrorDetail.SetError(InError != UEMEDIA_ERROR_OK ? InError : UEMEDIA_ERROR_DETAIL);
+	LastErrorDetail.SetFacility(Facility::EFacility::MPEGAudioPlaylistReader);
+	LastErrorDetail.SetCode(InCode);
+	LastErrorDetail.SetMessage(InMessage);
+	check(PlayerSessionServices);
+	if (PlayerSessionServices)
+	{
+		PlayerSessionServices->PostError(LastErrorDetail);
+	}
+}
+
+void FPlaylistReaderMPEGAudio::LogMessage(IInfoLog::ELevel Level, const FString& Message)
+{
+	if (PlayerSessionServices)
+	{
+		PlayerSessionServices->PostLog(Facility::EFacility::MPEGAudioPlaylistReader, Level, Message);
+	}
+}
+
+void FPlaylistReaderMPEGAudio::LoadAndParse(const FString& URL)
+{
+	PlaylistURL = URL;
+	StartWorkerThread();
+}
+
+
+void FPlaylistReaderMPEGAudio::FRequest::OnProcessRequestComplete(FHttpRequestPtr InSourceHttpRequest, FHttpResponsePtr InHttpResponse, bool bInSucceeded)
+{
+	Result = bInSucceeded ? EResult::Succeeded : Result != EResult::Aborted ? EResult::Failed : Result;
+	Events.SendMessage(EEvent::Finished);
+}
+
+void FPlaylistReaderMPEGAudio::FRequest::OnHeaderReceived(FHttpRequestPtr InSourceHttpRequest, const FString& InHeaderName, const FString& InHeaderValue)
+{
+	if (InHeaderName.Len())
+	{
+		// Headers are treated as case insensitive, so for simpliciy in later comparisons convert it to all lowercase here.
+		Headers.Add(InHeaderName.ToLower(), InHeaderValue);
+		// Content length?
+		if (InHeaderName.Equals(TEXT("content-length"), ESearchCase::IgnoreCase))
+		{
+			LexFromString(ContentLength, *InHeaderValue);
+		}
+	}
+	else
+	{
+		Headers.Add(TEXT("_"), InHeaderValue);
+	}
+}
+
+void FPlaylistReaderMPEGAudio::FRequest::OnStatusCodeReceived(FHttpRequestPtr InSourceHttpRequest, int32 InHttpStatusCode)
+{
+	if (InHttpStatusCode > 0 && InHttpStatusCode < 600)
+	{
+		StatusCode = InHttpStatusCode;
+	}
+}
+
+void FPlaylistReaderMPEGAudio::FRequest::FindSyncMarkers()
+{
+	if (Buffer.Num() && CurrentSyncMarkerCheckPos < ProbeSize)
+	{
+		const uint8* Base = Buffer.GetData() + CurrentSyncMarkerCheckPos;
+		const uint8* End = Buffer.GetData() + Buffer.Num() - 2;
+		for(; Base < End; ++Base)
+		{
+			if (*Base != 0xff)
+			{
+				continue;
+			}
+			// Check for validity
+			if (Base+2 < End && (Base[1] & 0xe0) == 0xe0		// sync marker (11 1-bits)
+				&& ((Base[1] >> 3) & 3) >= 2					// audio version 1 or 2 (2.5 not supported)
+				&& ((Base[1] >> 1) & 3) != 0					// layer index 1, 2 or 3
+				&& (Base[2] >> 4) != 15							// bitrate index not 15
+				&& (Base[2] & 0x0c) != 0x0c)					// sample rate index not 3
+			{
+				SyncMarkerOffsets.Add(Base - Buffer.GetData());
+				Base += 2;
+			}
+		}
+		CurrentSyncMarkerCheckPos = Buffer.Num() - 1;
+	}
+}
+
+bool FPlaylistReaderMPEGAudio::FRequest::OnProcessRequestStream(void *InDataPtr, int64 InLength)
+{
+	if (StatusCode < 200 || StatusCode >= 300)
+	{
+		return true;
+	}
+
+	// Add new data unconditionally. This won't be overly much so not to worry.
+	if (InDataPtr && InLength > 0)
+	{
+		int32 MaxToCopy = ProbeSize < 0 ? static_cast<int32>(InLength) : Utils::Min(ProbeSize - Buffer.Num(), static_cast<int32>(InLength));
+		Buffer.Append(reinterpret_cast<const uint8*>(InDataPtr), MaxToCopy);
+	}
+
+	// When we receive the first chunk of data we check if it starts with an ID3v2 tag.
+	if (ProbeSize < 0 && Buffer.Num() >= 10)
+	{
+		int32 ID3HeaderSize = 0;
+		const uint8* const HeaderData = Buffer.GetData();
+		if (HeaderData[0] == 'I' && HeaderData[1] == 'D' && HeaderData[2] == '3' &&
+			HeaderData[3] != 0xff && HeaderData[4] != 0xff && HeaderData[6] < 0x80 && HeaderData[7] < 0x80 && HeaderData[8] < 0x80 && HeaderData[9] < 0x80)
+		{
+			ID3HeaderSize = (int32) (10U + (HeaderData[6] << 21U) + (HeaderData[7] << 14U) + (HeaderData[8] << 7U) + HeaderData[9]);
+			CurrentSyncMarkerCheckPos = ID3HeaderSize;
+		}
+
+		ProbeSize = GetMinProbeSize() + ID3HeaderSize;
+		if (ContentLength >= 0)
+		{
+			ProbeSize = Utils::Min(ProbeSize, (int32)ContentLength);
+		}
+	}
+
+	if (!bSentHaveProbeDataMsg && Buffer.Num() >= ProbeSize)
+	{
+		bSentHaveProbeDataMsg = true;
+		Events.SendMessage(EEvent::HaveProbeData);
+	}
+	return true;
+}
+
+void FPlaylistReaderMPEGAudio::FRequest::Cancel()
+{
+	Result = Result == EResult::Running ? EResult::Aborted : Result;
+	if (Handle.IsValid())
+	{
+		Handle->CancelRequest();
+	}
+}
+
+bool FPlaylistReaderMPEGAudio::FRequest::Validate(FMPEGAudioInfoHeader& InOutInfoHeader, FString& OutError, TArray<FString>& OutWarnings)
+{
+	check(Buffer.Num() >= 8192);
+	if (Buffer.Num() < 8192)
+	{
+		OutError = FString::Printf(TEXT("Insufficient buffer data to validate file contents."));
+		return false;
+	}
+
+	InOutInfoHeader.MPEGHeaderMask = 0xfffe0c00;
+
+	const uint8* const HeaderData = Buffer.GetData();
+
+	// Check if there is an ID3 header.
+	if (SyncMarkerOffsets.Num() && SyncMarkerOffsets[0] && HeaderData[0] == 'I' && HeaderData[1] == 'D' && HeaderData[2] == '3' &&
+		HeaderData[3] != 0xff && HeaderData[4] != 0xff && HeaderData[6] < 0x80 && HeaderData[7] < 0x80 && HeaderData[8] < 0x80 && HeaderData[9] < 0x80)
+	{
+		int32 ID3HeaderSize = (int32) (10U + (HeaderData[6] << 21U) + (HeaderData[7] << 14U) + (HeaderData[8] << 7U) + HeaderData[9]);
+		if (ID3HeaderSize >= ProbeSize)
+		{
+			// ID3v2 header is larger than we have probe data for, so this data cannot be validated.
+			OutError = FString::Printf(TEXT("ID3v2 header larger than buffer data. Cannot validate file contents."));
+			return false;
+		}
+		// Parse the header.
+		InOutInfoHeader.ID3v2 = MakeSharedTS<MPEG::FID3V2Metadata>();
+		if (!InOutInfoHeader.ID3v2->Parse(HeaderData, ID3HeaderSize))
+		{
+			InOutInfoHeader.ID3v2.Reset();
+			OutWarnings.Emplace(FString::Printf(TEXT("Could not parse the ID3v2 header, ignoring.")));
+		}
+		while(SyncMarkerOffsets.Num() && SyncMarkerOffsets[0] < ID3HeaderSize)
+		{
+			SyncMarkerOffsets.RemoveAt(0);
+		}
+		// Is there an MLLT entry?
+		if (InOutInfoHeader.ID3v2.IsValid() && InOutInfoHeader.ID3v2->HaveTag(Utils::Make4CC('M','L','L','T')))
+		{
+			MPEG::FID3V2Metadata::FItem MLLTBlob;
+			InOutInfoHeader.ID3v2->GetTag(MLLTBlob, Utils::Make4CC('M','L','L','T'));
+			TArray<uint8> MLLT = MLLTBlob.Value.GetValue<TArray<uint8>>();
+			ElectraDecodersUtil::FElectraBitstreamReader br(MLLT.GetData(), MLLT.Num());
+			InOutInfoHeader.MLLT = MakeSharedTS<FMPEGAudioInfoHeader::FMLLT>();
+			InOutInfoHeader.MLLT->FramesBetweenReferences = br.GetBits(16);
+			uint32 BytesBetweenReference = br.GetBits(24);
+			uint32 MillisecondsBetweenReference = br.GetBits(24);
+			uint32 nbDevBytes = br.GetBits(8);
+			uint32 nbDevMillis = br.GetBits(8);
+			check(nbDevBytes <= 32);
+			check(nbDevMillis <= 32);
+			check((nbDevBytes + nbDevMillis) % 4 == 0);
+			int32 NumEntries = br.GetRemainingBits() / (nbDevBytes + nbDevMillis);
+			InOutInfoHeader.MLLT->TimeAndOffsets.AddDefaulted(NumEntries + 1);
+			uint32 Offset = 0;
+			uint32 Milliseconds = 0;
+			for(int32 i=1; i<=NumEntries; ++i)
+			{
+				Offset += BytesBetweenReference + br.GetBits(nbDevBytes);
+				Milliseconds += MillisecondsBetweenReference + br.GetBits(nbDevMillis);
+				InOutInfoHeader.MLLT->TimeAndOffsets[i].Offset = Offset;
+				InOutInfoHeader.MLLT->TimeAndOffsets[i].Milliseconds = Milliseconds;
+			}
+			check(br.GetRemainingBits() == 0);
+		}
+	}
+
+	// Try to locate one of the special info headers
+	auto IsXingHeader = [](const uint32 InHeaderValue, const uint8* InData, int32& InOutOffset, uint32& OutHeader) -> bool
+	{
+		const int32 Version = ElectraDecodersUtil::MPEG::UtilsMPEG123::GetVersion(InHeaderValue);
+		const int32 NumChannels = ElectraDecodersUtil::MPEG::UtilsMPEG123::GetChannelCount(InHeaderValue);
+		int32 Off = InOutOffset + 4 + (Version == 1 ? NumChannels == 1 ? 17 : 32 : NumChannels == 1 ? 9 : 17);
+		const uint32 Hdr = GetUINT32BE(InData + Off);
+		if (Hdr == Make4CC('X','i','n','g') || Hdr == Make4CC('I','n','f','o'))
+		{
+			OutHeader = Hdr;
+			InOutOffset = Off;
+			return true;
+		}
+		return false;
+	};
+
+	auto IsVBRIHeader = [](const uint32 InHeaderValue, const uint8* InData, int32& InOutOffset, uint32& OutHeader) -> bool
+	{
+		int32 Off = 4+ InOutOffset + 32;
+		const uint32 Hdr = GetUINT32BE(InData + Off);
+		if (Hdr == Make4CC('V','B','R','I'))
+		{
+			OutHeader = Hdr;
+			InOutOffset = Off;
+			return true;
+		}
+		return false;
+	};
+
+	int32 InfoHeaderOffset = -1;
+	for(int32 nSyncMarker=0; InfoHeaderOffset < 0 && nSyncMarker<SyncMarkerOffsets.Num()-2; ++nSyncMarker)
+	{
+		int32 Off = SyncMarkerOffsets[nSyncMarker];
+		const uint32 HeaderValue = GetUINT32BE(HeaderData + Off);
+
+		// These exist only in Layer III
+		const int32 Layer = ElectraDecodersUtil::MPEG::UtilsMPEG123::GetLayer(HeaderValue);
+		if (Layer == 3)
+		{
+			uint32 HeaderName = 0;
+			// Test Xing header
+			if (IsXingHeader(HeaderValue, HeaderData, Off, HeaderName))
+			{
+				const int32 Version = ElectraDecodersUtil::MPEG::UtilsMPEG123::GetVersion(HeaderValue);
+				const int32 NumChannels = ElectraDecodersUtil::MPEG::UtilsMPEG123::GetChannelCount(HeaderValue);
+				InfoHeaderOffset = SyncMarkerOffsets[nSyncMarker];
+				InOutInfoHeader.MPEGHeaderExpectedValue = HeaderValue & InOutInfoHeader.MPEGHeaderMask;
+				InOutInfoHeader.MPEGVersion = Version;
+				InOutInfoHeader.MPEGLayer = Layer;
+				InOutInfoHeader.SampleRate = ElectraDecodersUtil::MPEG::UtilsMPEG123::GetSamplingRate(HeaderValue);
+				InOutInfoHeader.NumChannels = NumChannels;
+				InOutInfoHeader.Bitrate = ElectraDecodersUtil::MPEG::UtilsMPEG123::GetBitrate(HeaderValue);
+
+				uint32 NumFrames = 0;
+				uint32 NumBytes = 0;
+				int32 VBRScale = -1;
+				InOutInfoHeader.bIsVBR = HeaderName == Make4CC('X','i','n','g');
+				if (!InOutInfoHeader.bIsVBR)
+				{
+					InOutInfoHeader.CBRFrameSize = ElectraDecodersUtil::MPEG::UtilsMPEG123::GetFrameSize(HeaderValue, 1);
+				}
+				InOutInfoHeader.SamplesPerFrame = ElectraDecodersUtil::MPEG::UtilsMPEG123::GetSamplesPerFrame(HeaderValue);
+
+				uint32 Flags = GetUINT32BE(HeaderData + Off + 4);
+				Off += 8;
+				if ((Flags & 1) != 0)
+				{
+					NumFrames = GetUINT32BE(HeaderData + Off);
+					Off += 4;
+					InOutInfoHeader.NumFrames = NumFrames;
+				}
+				if ((Flags & 2) != 0)
+				{
+					NumBytes = GetUINT32BE(HeaderData + Off);
+					Off += 4;
+				}
+				InOutInfoHeader.bHaveTOC = (Flags & 4) != 0;
+				if (InOutInfoHeader.bHaveTOC)
+				{
+					InOutInfoHeader.TOC.AddUninitialized(100);
+					FMemory::Memcpy(InOutInfoHeader.TOC.GetData(), HeaderData + Off, 100);
+					Off += 100;
+				}
+				if ((Flags & 8) != 0)
+				{
+					VBRScale = (int32) GetUINT32BE(HeaderData + Off);
+					Off += 4;
+				}
+				// Check if there is a `LAME` header following.
+				if (GetUINT32BE(HeaderData + Off) == Make4CC('L','A','M','E'))
+				{
+					uint32 VersionMajor = HeaderData[Off + 4] - 0x30;
+					uint32 VersionMinor = (HeaderData[Off + 5] == '.' ? 0 : HeaderData[Off + 5] - 0x30) * 100 + (HeaderData[Off + 6] - 0x30) * 10 + (HeaderData[Off + 7] - 0x30);
+					// Additional info added with version 3.90
+					if (VersionMajor > 3 || (VersionMajor == 3 && VersionMinor >= 90))
+					{
+						ElectraDecodersUtil::FElectraBitstreamReader br(HeaderData + Off, 208, 9, 0);
+						uint32 InfoTagRevision = br.GetBits(4);			// 15=reserved
+						uint32 VBRMethod = br.GetBits(4);				// 0=unknown, 1=CBR, 2=ABR, 3-6=VBR method 1-4, 8=CBR 2 pass, 9=ABR 2 pass, 15=reserved
+						uint32 LowPassFreq = br.GetBits(8) * 100;		// 0=unknown
+						union UReplayGain { uint32 u; float f; } rg;
+						rg.u = br.GetBits(32);
+						float PeakSignalAmplitude = rg.f;				// 0.0 = unknown
+						uint32 RadioReplayGain = br.GetBits(16);		// AAAAAAAAASOOONNN, A=Absolute gain adjustment, S=sign bit, O=originator (0=not set, 1=artist, 2=user, 3=model, 4=RMS), N=name (0=not set, 1=radio, 2=audiophile)
+						uint32 AudiophileReplayGain = br.GetBits(16);	// AAAAAAAAASOOONNN
+						uint32 ATHType = br.GetBits(4);
+						uint32 NoGapContPrevTrack = br.GetBits(1);
+						uint32 NoGapContNextTrack = br.GetBits(1);
+						uint32 NsSafeJoint = br.GetBits(1);
+						uint32 NsPsyTune = br.GetBits(1);
+						uint32 ABR = br.GetBits(8);						// Average bitrate if VBRMethod==2|9, CBR rate for 1|8, min VBR rate otherwise, 0=unknown, 255= >=255 kbps
+						uint32 EncoderDelayS = br.GetBits(12);			// Number of prepended silent samples
+						uint32 EncoderDelayE = br.GetBits(12);			// Number of appended silent samples
+						uint32 SourceSampleRate = br.GetBits(2);		// 0= <=32kHz, 1=44.1kHz, 2=48kHz, 3= >48kHz
+						uint32 UnwiseSettingsUsed = br.GetBits(1);
+						uint32 StereoMode = br.GetBits(3);				// 0=mono, 1=stereo, 2=dual, 3=joint, 4=force, 5=auto, 6=intensity, 7=other
+						uint32 NoiseShaping = br.GetBits(2);			// 0-3 = noise shaping mode 0-3
+						uint32 MP3GainSign = br.GetBits(1);				// Sign of gain (value is not 2's complement!)
+						uint32 MP3Gain = br.GetBits(7);					// Gain that was applied (without clipping/clamping). Factor was 2^(x/4.0). Used to undo scaling if necessary.
+						br.SkipBits(2);	// unused
+						uint32 SurroundInfo = br.GetBits(3);			// 0=no info, 1=DPL, 2=DPL2, 3=Ambisonic
+						uint32 Preset = br.GetBits(11);					// 0=unknown, LAME internal preset number otherwise
+						uint32 MusicLength = br.GetBits(32);			// number of bytes in file including this block, excluding any earlier blocks and ID3v2 tag and ID3v tag at the end. If 0 file >4GiB or unknown.
+						uint32 MusicCRC = br.GetBits(16);				// CRC of the music data, starting with the next block after this one and up to the last block excluding IDv3 tag at the end
+						uint32 InfoTagCRC = br.GetBits(16);				// CRC of this header
+
+						// VBR mismatch?
+						if ((InOutInfoHeader.bIsVBR && (VBRMethod == 1 || VBRMethod == 8)) ||						// VBR via Xing, but CBR here?
+							(!InOutInfoHeader.bIsVBR && ((VBRMethod >= 2 && VBRMethod <= 6) || VBRMethod == 9)))	// CBR via Info, but VBR here?
+						{
+							OutWarnings.Emplace(FString::Printf(TEXT("VBR/CBR mismatch")));
+						}
+						InOutInfoHeader.AverageBitrate = VBRMethod == 1 || VBRMethod == 8 ? 0 : ABR * 1000;
+						if (MusicLength)
+						{
+							// Mismatch?
+							if (NumBytes && NumBytes != MusicLength)
+							{
+								OutWarnings.Emplace(FString::Printf(TEXT("Mismatching music length in LAME tag")));
+							}
+							// Trust this info more than the one above.
+							NumBytes = MusicLength;
+						}
+						InOutInfoHeader.EncoderDelayStart = EncoderDelayS;
+						InOutInfoHeader.EncoderDelayEnd = EncoderDelayE;
+					}
+				}
+
+				InOutInfoHeader.FirstDataByte = SyncMarkerOffsets[nSyncMarker];
+				if (NumBytes)
+				{
+					InOutInfoHeader.LastDataByte = InOutInfoHeader.FirstDataByte + (int64)NumBytes;
+				}
+			}
+			// Test VBRI header
+			else if (IsVBRIHeader(HeaderValue, HeaderData, Off, HeaderName))
+			{
+				ElectraDecodersUtil::FElectraBitstreamReader br(HeaderData + Off + 4, Utils::Max(0, ProbeSize-Off));
+				int32 VersionId = (int32)br.GetBits(16);
+				int32 Delay = (int32)br.GetBits(16);
+				int32 QualityIndicator = (int32)br.GetBits(16);
+				uint32 NumBytes = br.GetBits(32);
+				uint32 NumFrames = br.GetBits(32);
+				int32 NumTOCEntries = (int32)br.GetBits(16);
+				uint16 TOCTableScale = (uint16)br.GetBits(16);
+				int32 TableEntrySize = (int32)br.GetBits(16);
+				int32 FramesPerTableEntry = (int32)br.GetBits(16);
+				if (br.GetRemainingByteLength() >= TableEntrySize * NumTOCEntries)
+				{
+					if (VersionId == 1 && TOCTableScale && TableEntrySize <= 4 && FramesPerTableEntry)
+					{
+						InOutInfoHeader.SeekTable = MakeSharedTS<TArray<uint32>>();
+						InOutInfoHeader.FramesPerSeekTableEntry = FramesPerTableEntry;
+						InOutInfoHeader.SeekTable->AddUninitialized(NumTOCEntries);
+						const int32 nb = 8 * TableEntrySize;
+						for(int32 j=0; j<NumTOCEntries; ++j)
+						{
+							InOutInfoHeader.SeekTable->operator[](j) = br.GetBits(nb) * TOCTableScale;
+						}
+					}
+					else
+					{
+						OutError = FString::Printf(TEXT("Unsupported VBRI header. Check values!"));
+						return false;
+					}
+				}
+				else
+				{
+					OutError = FString::Printf(TEXT("VBRI seek table larger than available probe data"));
+					return false;
+				}
+				InOutInfoHeader.NumFrames = NumFrames;
+				InOutInfoHeader.bIsVBR = true;
+
+				int32 VBRIHeaderEndOffset = reinterpret_cast<const uint8*>(br.GetRemainingData()) - HeaderData;
+				while(SyncMarkerOffsets.Num() && SyncMarkerOffsets[0] < VBRIHeaderEndOffset)
+				{
+					SyncMarkerOffsets.RemoveAt(0);
+					--nSyncMarker;
+				}
+				if (SyncMarkerOffsets.Num())
+				{
+					InOutInfoHeader.FirstDataByte = SyncMarkerOffsets[0];
+					if (NumBytes)
+					{
+						InOutInfoHeader.LastDataByte = SyncMarkerOffsets[0] + (int64)NumBytes;
+					}
+				}
+				else
+				{
+					OutError = FString::Printf(TEXT("No further sync marker found after VBRI header. Probe data too small?"));
+					return false;
+				}
+			}
+		}
+	}
+
+	// Iterate all the headers that resemble an mpeg audio header and count those that are identical
+	// except for slight differences (like padding, private, copyright bit etc.)
+	TMap<uint32, int32> HeaderMap;
+	for(int32 i=0; i<SyncMarkerOffsets.Num()-2; ++i)
+	{
+		int32 Off = SyncMarkerOffsets[i];
+		uint32 HeaderValue = GetUINT32BE(HeaderData + Off);
+		HeaderValue &= InOutInfoHeader.MPEGHeaderMask;
+		++HeaderMap.FindOrAdd(HeaderValue, 0);
+	}
+	HeaderMap.ValueSort([](const int32& e1, const int32& e2) { return e1 > e2; });
+
+	uint32 ExpectedHeaderValue = HeaderMap.Num() ? HeaderMap.CreateConstIterator().Key() : 0xfffffffU;
+	// Possible mismatch between the header value found in the info packet above and the majority vote?
+	if (InOutInfoHeader.MPEGHeaderExpectedValue && InOutInfoHeader.MPEGHeaderExpectedValue != ExpectedHeaderValue)
+	{
+		OutWarnings.Emplace(FString::Printf(TEXT("Mismatching MPEG header bitmask")));
+	}
+	InOutInfoHeader.MPEGHeaderExpectedValue = ExpectedHeaderValue;
+
+	// Locate the first matching header.
+	int32 FirstHeaderOffset = InfoHeaderOffset;
+	for(int32 i=0; FirstHeaderOffset<0 && i<SyncMarkerOffsets.Num()-2; ++i)
+	{
+		const int32 Off = SyncMarkerOffsets[i];
+		if ((GetUINT32BE(HeaderData + Off) & InOutInfoHeader.MPEGHeaderMask) == InOutInfoHeader.MPEGHeaderExpectedValue)
+		{
+			FirstHeaderOffset = Off;
+			break;
+		}
+	}
+	// Perform a quick sanity check that the blocks are really where they should be.
+	bool bProbeOk = true;
+	if (FirstHeaderOffset >= 0)
+	{
+		const int32 LastHeaderOffset = SyncMarkerOffsets.Num() >= 2 ? SyncMarkerOffsets[SyncMarkerOffsets.Num()-2] : 0;
+		// Probing 10 headers should be fine
+		int32 Off = FirstHeaderOffset;
+		for(int32 i=0; i<10; ++i)
+		{
+			const uint32 HeaderValue = GetUINT32BE(HeaderData + Off);
+			int32 FrameSize = ElectraDecodersUtil::MPEG::UtilsMPEG123::GetFrameSize(HeaderValue);
+			if (Off + FrameSize + 4 <= LastHeaderOffset)
+			{
+				const uint32 HeaderValue2 = GetUINT32BE(HeaderData + Off + FrameSize);
+				if ((HeaderValue2 & InOutInfoHeader.MPEGHeaderMask) != InOutInfoHeader.MPEGHeaderExpectedValue)
+				{
+					bProbeOk = false;
+					break;
+				}
+				Off += FrameSize;
+			}
+			else
+			{
+				break;
+			}
+		}
+	}
+	// Fill in values that haven't been set up by any of the special headers yet.
+	if (bProbeOk && FirstHeaderOffset >= 0)
+	{
+		const uint32 HeaderValue = GetUINT32BE(HeaderData + FirstHeaderOffset);
+		InOutInfoHeader.MPEGVersion = InOutInfoHeader.MPEGVersion ? InOutInfoHeader.MPEGVersion : ElectraDecodersUtil::MPEG::UtilsMPEG123::GetVersion(HeaderValue);
+		InOutInfoHeader.MPEGLayer = InOutInfoHeader.MPEGLayer ? InOutInfoHeader.MPEGLayer : ElectraDecodersUtil::MPEG::UtilsMPEG123::GetLayer(HeaderValue);
+		InOutInfoHeader.SampleRate = InOutInfoHeader.SampleRate ? InOutInfoHeader.SampleRate : ElectraDecodersUtil::MPEG::UtilsMPEG123::GetSamplingRate(HeaderValue);
+		InOutInfoHeader.NumChannels = InOutInfoHeader.NumChannels ? InOutInfoHeader.NumChannels : ElectraDecodersUtil::MPEG::UtilsMPEG123::GetChannelCount(HeaderValue);
+		InOutInfoHeader.Bitrate = InOutInfoHeader.Bitrate ? InOutInfoHeader.Bitrate : ElectraDecodersUtil::MPEG::UtilsMPEG123::GetBitrate(HeaderValue);
+		InOutInfoHeader.CBRFrameSize = InOutInfoHeader.bIsVBR ? 0 : InOutInfoHeader.CBRFrameSize ? InOutInfoHeader.CBRFrameSize : ElectraDecodersUtil::MPEG::UtilsMPEG123::GetFrameSize(HeaderValue, 1);
+		InOutInfoHeader.SamplesPerFrame = InOutInfoHeader.SamplesPerFrame ? InOutInfoHeader.SamplesPerFrame : ElectraDecodersUtil::MPEG::UtilsMPEG123::GetSamplesPerFrame(HeaderValue);
+		InOutInfoHeader.FirstDataByte = InOutInfoHeader.FirstDataByte >= 0 ? InOutInfoHeader.FirstDataByte : FirstHeaderOffset;
+		InOutInfoHeader.LastDataByte = InOutInfoHeader.LastDataByte >= 0 ? InOutInfoHeader.LastDataByte : -1;
+	}
+
+	if (bProbeOk)
+	{
+		InOutInfoHeader.CodecInfo.SetStreamType(EStreamType::Audio);
+		InOutInfoHeader.CodecInfo.SetMimeType(TEXT("audio/mpeg"));
+		InOutInfoHeader.CodecInfo.SetCodec(FStreamCodecInformation::ECodec::Audio4CC);
+		InOutInfoHeader.CodecInfo.SetCodec4CC(Utils::Make4CC('m','p','g','a'));
+		InOutInfoHeader.CodecInfo.SetProfile(InOutInfoHeader.MPEGVersion);
+		InOutInfoHeader.CodecInfo.SetProfileLevel(InOutInfoHeader.MPEGLayer);
+		InOutInfoHeader.CodecInfo.SetCodecSpecifierRFC6381(TEXT("mp4a.6b"));
+		InOutInfoHeader.CodecInfo.SetSamplingRate(InOutInfoHeader.SampleRate);
+		InOutInfoHeader.CodecInfo.SetNumberOfChannels(InOutInfoHeader.NumChannels);
+		InOutInfoHeader.CodecInfo.SetBitrate(InOutInfoHeader.AverageBitrate ? InOutInfoHeader.AverageBitrate : InOutInfoHeader.Bitrate ? InOutInfoHeader.Bitrate : 0);
+
+		if (InOutInfoHeader.bIsLive)
+		{
+			InOutInfoHeader.EstimatedDuration.SetToPositiveInfinity();
+		}
+		else
+		{
+			auto CalculateAverageBitrate = [&InOutInfoHeader]() -> void
+			{
+				// If VBR we can recalculate the average bitrate if we know the file size.
+				if (InOutInfoHeader.bIsVBR && InOutInfoHeader.FirstDataByte >= 0 && InOutInfoHeader.LastDataByte > InOutInfoHeader.FirstDataByte)
+				{
+					InOutInfoHeader.AverageBitrate = (int32)((InOutInfoHeader.LastDataByte - InOutInfoHeader.FirstDataByte) * 8 / InOutInfoHeader.EstimatedDuration.GetAsSeconds());
+				}
+			};
+			InOutInfoHeader.EstimatedDuration.SetToInvalid();
+			// Is there an ID3v2 tag giving the duration?
+			if (InOutInfoHeader.ID3v2.IsValid())
+			{
+				MPEG::FID3V2Metadata::FItem v;
+				if (InOutInfoHeader.ID3v2->GetTag(v, Utils::Make4CC('T','L','E','N')))
+				{
+					InOutInfoHeader.EstimatedDuration.SetFromTimespan(v.Value.GetValue<FTimespan>());
+					CalculateAverageBitrate();
+				}
+			}
+			// Duration not valid, try to calculate an estimate from the pieces of information we have.
+			// Is there a total frame number given?
+			if (!InOutInfoHeader.EstimatedDuration.IsValid() && InOutInfoHeader.NumFrames && InOutInfoHeader.SamplesPerFrame && InOutInfoHeader.SampleRate)
+			{
+				InOutInfoHeader.EstimatedDuration.SetFromTimeFraction(FTimeFraction((int64)(InOutInfoHeader.NumFrames * (uint32)InOutInfoHeader.SamplesPerFrame - (uint32)(InOutInfoHeader.EncoderDelayStart + InOutInfoHeader.EncoderDelayEnd)), (uint32)InOutInfoHeader.SampleRate));
+				CalculateAverageBitrate();
+			}
+			// For CBR we can guess based on filesize, if that is known.
+			if (!InOutInfoHeader.EstimatedDuration.IsValid() && !InOutInfoHeader.bIsVBR && InOutInfoHeader.Bitrate && InOutInfoHeader.FirstDataByte >= 0 && InOutInfoHeader.LastDataByte > InOutInfoHeader.FirstDataByte)
+			{
+				if (InOutInfoHeader.CBRFrameSize && InOutInfoHeader.SamplesPerFrame && InOutInfoHeader.SampleRate)
+				{
+					int64 NumFrames = (InOutInfoHeader.LastDataByte - InOutInfoHeader.FirstDataByte) / InOutInfoHeader.CBRFrameSize;
+					InOutInfoHeader.EstimatedDuration.SetFromTimeFraction(FTimeFraction(NumFrames * InOutInfoHeader.SamplesPerFrame - (InOutInfoHeader.EncoderDelayStart + InOutInfoHeader.EncoderDelayEnd), (uint32)InOutInfoHeader.SampleRate));
+					InOutInfoHeader.NumFrames = (uint32)NumFrames;
+				}
+				else
+				{
+					// This is very likely to overshoot the actual duration.
+					InOutInfoHeader.EstimatedDuration.SetFromTimeFraction(FTimeFraction((InOutInfoHeader.LastDataByte - InOutInfoHeader.FirstDataByte) * 8, (uint32)InOutInfoHeader.Bitrate));
+				}
+			}
+			// For VBR without a dedicated info header the estimate will be extremely rough.
+			if (!InOutInfoHeader.EstimatedDuration.IsValid() && InOutInfoHeader.bIsVBR && (InOutInfoHeader.Bitrate || InOutInfoHeader.AverageBitrate) && InOutInfoHeader.FirstDataByte >= 0 && InOutInfoHeader.LastDataByte > InOutInfoHeader.FirstDataByte)
+			{
+				InOutInfoHeader.EstimatedDuration.SetFromTimeFraction(FTimeFraction((InOutInfoHeader.LastDataByte - InOutInfoHeader.FirstDataByte) * 8, (uint32)(InOutInfoHeader.AverageBitrate ? InOutInfoHeader.AverageBitrate : InOutInfoHeader.Bitrate)));
+				CalculateAverageBitrate();
+			}
+		}
+		// Copy the HTTP response headers across. This might be useful for Casting (ie. Icy Cast)
+		InOutInfoHeader.HTTPResponseHeaders = Headers;
+	}
+
+	return bProbeOk;
+}
+
+
+void FPlaylistReaderMPEGAudio::WorkerThread()
+{
+	LLM_SCOPE(ELLMTag::ElectraPlayer);
+
+	SCOPE_CYCLE_COUNTER(STAT_ElectraPlayer_MPEGAudio_PlaylistWorker);
+	CSV_SCOPED_TIMING_STAT(ElectraPlayer, MPEGAudio_PlaylistWorker);
+
+	FURL_RFC3986 UrlParser;
+	UrlParser.Parse(PlaylistURL);
+	PlaylistURL = UrlParser.Get(true, false);
+	URLFragment = UrlParser.GetFragment();
+
+	TSharedPtr<FRequest, ESPMode::ThreadSafe> Req = MakeShared<FRequest, ESPMode::ThreadSafe>();
+
+	// We have to distinguish between local file playback and internet here.
+	ConnectionInfo.RequestStartTime = MEDIAutcTime::Current();
+	if (UrlParser.GetScheme().Equals(TEXT("file")))
+	{
+		FString Filename;
+		FURL_RFC3986::UrlDecode(Filename, UrlParser.Get(false, false));
+		Filename.MidInline(7);	// remove "file://"
+		TSharedPtr<FArchive, ESPMode::ThreadSafe> Archive = MakeShareable(IFileManager::Get().CreateFileReader(*Filename));
+		ConnectionInfo.EffectiveURL = Filename;
+		ConnectionInfo.ContentType = TEXT("audio/mpeg");
+		if (Archive.IsValid())
+		{
+			Req->StatusCode = 200;
+			Req->ContentLength = Archive->TotalSize();
+			int32 AlreadyRead = 0;
+			if (Req->ContentLength > 10)
+			{
+				AlreadyRead = 10;
+				Req->Buffer.AddUninitialized(10);
+				Archive->Serialize(Req->Buffer.GetData(), 10);
+				Req->OnProcessRequestStream(nullptr, 0);
+				check(Req->ProbeSize > 0);
+			}
+			Req->ProbeSize = Req->ProbeSize < 0 ? FRequest::GetMinProbeSize() : Req->ProbeSize;
+			Req->Buffer.AddUninitialized(Req->ProbeSize - AlreadyRead);
+			Archive->Serialize(Req->Buffer.GetData() + AlreadyRead, Req->ProbeSize - AlreadyRead);
+			Req->FindSyncMarkers();
+			Req->Result = FRequest::EResult::Succeeded;
+		}
+		else
+		{
+			Req->StatusCode = 404;
+			Req->Result = FRequest::EResult::Failed;
+		}
+	}
+	else
+	{
+		FHttpRequestStreamDelegate StreamDelegate;
+		StreamDelegate.BindThreadSafeSP(Req.ToSharedRef(), &FPlaylistReaderMPEGAudio::FRequest::OnProcessRequestStream);
+		Req->Handle = FHttpModule::Get().CreateRequest();
+		Req->Handle->SetVerb(TEXT("GET"));
+		Req->Handle->SetURL(PlaylistURL);
+		Req->Handle->SetDelegateThreadPolicy(EHttpRequestDelegateThreadPolicy::CompleteOnHttpThread);
+		Req->Handle->OnProcessRequestComplete().BindThreadSafeSP(Req.ToSharedRef(), &FPlaylistReaderMPEGAudio::FRequest::OnProcessRequestComplete);
+		Req->Handle->OnHeaderReceived().BindThreadSafeSP(Req.ToSharedRef(), &FPlaylistReaderMPEGAudio::FRequest::OnHeaderReceived);
+		Req->Handle->OnStatusCodeReceived().BindThreadSafeSP(Req.ToSharedRef(), &FPlaylistReaderMPEGAudio::FRequest::OnStatusCodeReceived);
+		Req->Handle->SetResponseBodyReceiveStreamDelegate(StreamDelegate);
+		Req->Handle->SetHeader(TEXT("User-Agent"), IElectraHttpManager::GetDefaultUserAgent());
+		Req->Handle->SetHeader(TEXT("Accept-Encoding"), TEXT("identity"));
+		Req->Handle->SetTimeout(PlayerSessionServices->GetOptionValue(MPEGAudio::OptionKeyMPEGAudioLoadTimeout).SafeGetTimeValue(FTimeValue().SetFromMilliseconds(1000 * 60)).GetAsSeconds());
+		Req->Handle->SetActivityTimeout(PlayerSessionServices->GetOptionValue(MPEGAudio::OptionKeyMPEGAudioLoadNoDataTimeout).SafeGetTimeValue(FTimeValue().SetFromMilliseconds(1000 * 8)).GetAsSeconds());
+		Req->Handle->ProcessRequest();
+
+		bool bDone = false;
+		bool bDidAbort = false;
+		while(!bDone)
+		{
+			FRequest::EEvent evt = FRequest::EEvent::None;
+			if (Req->Events.ReceiveMessage(evt, 20*1000))
+			{
+				switch(evt)
+				{
+					case FRequest::EEvent::Abort:
+					{
+						Req->Cancel();
+						break;
+					}
+					case FRequest::EEvent::HaveProbeData:
+					{
+						if (!bAbort)
+						{
+							Req->Cancel();
+							Req->FindSyncMarkers();
+						}
+						break;
+					}
+					case FRequest::EEvent::Finished:
+					{
+						bDone = true;
+						break;
+					}
+				}
+			}
+			if (bAbort)
+			{
+				if (!bDidAbort)
+				{
+					bDidAbort = true;
+					Req->Events.SendMessage(FRequest::EEvent::Abort);
+				}
+			}
+		}
+		Req->Handle.Reset();
+	}
+
+	ConnectionInfo.RequestEndTime = MEDIAutcTime::Current();
+	ConnectionInfo.bHasFinished = !bAbort;
+	ConnectionInfo.bWasAborted = bAbort;
+	ConnectionInfo.StatusInfo.HTTPStatus = Req->StatusCode;
+	ConnectionInfo.StatusInfo.bReadError = Req->Result == FRequest::EResult::Failed;
+
+	// Notify the download of the "master playlist". This indicates the download only, not the parsing thereof.
+	PlayerSessionServices->SendMessageToPlayer(IPlaylistReader::PlaylistDownloadMessage::Create(&ConnectionInfo, Playlist::EListType::Master, Playlist::ELoadType::Initial));
+	// Notify that the "master playlist" has been parsed, successfully or not.
+	PlayerSessionServices->SendMessageToPlayer(IPlaylistReader::PlaylistLoadedMessage::Create(LastErrorDetail, &ConnectionInfo, Playlist::EListType::Master, Playlist::ELoadType::Initial));
+	if (!bAbort)
+	{
+		// Error?
+		if (Req->Result == FRequest::EResult::Failed || Req->StatusCode != 200)
+		{
+			PostError(FString::Printf(TEXT("Error while downloading \"%s\""), *PlaylistURL), ERRCODE_MPEGAUDIO_DOWNLOAD_ERROR, UEMEDIA_ERROR_READ_ERROR);
+		}
+		else
+		{
+			FString OutError;
+			TArray<FString> OutWarnings;
+			InfoHeader.bIsLive = Req->ContentLength < 0;
+			InfoHeader.LastDataByte = Req->ContentLength;
+			bool bIsValid = Req->Validate(InfoHeader, OutError, OutWarnings);
+			if (bIsValid)
+			{
+				for(auto& Warn : OutWarnings)
+				{
+					LogMessage(IInfoLog::ELevel::Info, FString::Printf(TEXT("%s"), *Warn));
+				}
+
+				Manifest = MakeSharedTS<FManifestMPEGAudioInternal>(PlayerSessionServices);
+
+				TArray<FURL_RFC3986::FQueryParam> URLFragmentComponents;
+				FURL_RFC3986::GetQueryParams(URLFragmentComponents, URLFragment, false);	// The fragment is already URL escaped, so no need to do it again.
+				Manifest->SetURLFragmentComponents(MoveTemp(URLFragmentComponents));
+
+				LastErrorDetail = Manifest->Build(InfoHeader, PlaylistURL);
+
+				// Notify that the "variant playlists" are ready. There are no variants in an mp4, but this is the trigger that the playlists are all set up and are good to go now.
+				PlayerSessionServices->SendMessageToPlayer(IPlaylistReader::PlaylistLoadedMessage::Create(LastErrorDetail, &ConnectionInfo, Playlist::EListType::Variant, Playlist::ELoadType::Initial));
+			}
+			else
+			{
+				// Not valid. File might have an ID3v2 tag larger than the probe amount or it is actually not a valid file.
+				PostError(FString::Printf(TEXT("File \"%s\" does not appear to be a valid MPEG audio file."), *PlaylistURL), ERRCODE_MPEGAUDIO_INVALID_FILE, UEMEDIA_ERROR_FORMAT_ERROR);
+			}
+		}
+	}
+
+	// This thread's work is done. We only wait for termination now.
+	WorkerThreadQuitSignal.Wait();
+}
+
+} // namespace Electra
diff --git a/Engine/Plugins/Media/ElectraPlayer/Source/ElectraPlayerRuntime/Private/Runtime/Player/mpegaudio/PlaylistReaderMPEGAudio.h b/Engine/Plugins/Media/ElectraPlayer/Source/ElectraPlayerRuntime/Private/Runtime/Player/mpegaudio/PlaylistReaderMPEGAudio.h
new file mode 100644
--- /dev/null
+++ b/Engine/Plugins/Media/ElectraPlayer/Source/ElectraPlayerRuntime/Private/Runtime/Player/mpegaudio/PlaylistReaderMPEGAudio.h
@@ -0,0 +1,21 @@
+// Copyright Epic Games, Inc. All Rights Reserved.
+
+#pragma once
+
+#include "PlayerCore.h"
+#include "Player/PlaylistReader.h"
+#include "Player/PlayerSessionServices.h"
+
+namespace Electra
+{
+
+class IPlaylistReaderMPEGAudio : public IPlaylistReader
+{
+public:
+	static TSharedPtrTS<IPlaylistReader> Create(IPlayerSessionServices* PlayerSessionServices);
+
+	virtual ~IPlaylistReaderMPEGAudio() = default;
+};
+
+
+} // namespace Electra
diff --git a/Engine/Plugins/Media/ElectraPlayer/Source/ElectraPlayerRuntime/Private/Runtime/Player/mpegaudio/StreamReaderMPEGAudio.cpp b/Engine/Plugins/Media/ElectraPlayer/Source/ElectraPlayerRuntime/Private/Runtime/Player/mpegaudio/StreamReaderMPEGAudio.cpp
new file mode 100644
--- /dev/null
+++ b/Engine/Plugins/Media/ElectraPlayer/Source/ElectraPlayerRuntime/Private/Runtime/Player/mpegaudio/StreamReaderMPEGAudio.cpp
@@ -0,0 +1,858 @@
+// Copyright Epic Games, Inc. All Rights Reserved.
+
+#include "PlayerCore.h"
+#include "Player/Manifest.h"
+#include "HTTP/HTTPManager.h"
+#include "Http.h"
+#include "Stats/Stats.h"
+#include "SynchronizedClock.h"
+#include "PlaylistReaderMPEGAudio.h"
+#include "Utilities/Utilities.h"
+#include "Utilities/TimeUtilities.h"
+#include "Utilities/UtilsMP4.h"
+#include "Utils/MPEG/ElectraUtilsMPEGAudio.h"
+#include "HAL/LowLevelMemTracker.h"
+#include "ElectraPlayerPrivate.h"
+#include "Player/Manifest.h"
+#include "Player/PlayerStreamReader.h"
+#include "Player/AdaptiveStreamingPlayerABR.h"
+#include "Player/mpegaudio/ManifestMPEGAudio.h"
+#include "Player/mpegaudio/StreamReaderMPEGAudio.h"
+#include "Player/mpegaudio/OptionKeynamesMPEGAudio.h"
+#include "Player/PlayerSessionServices.h"
+#include "Player/PlayerStreamFilter.h"
+
+
+DECLARE_CYCLE_STAT(TEXT("FStreamReaderMPEGAudio_HandleRequest"), STAT_ElectraPlayer_MPEGAudio_StreamReader, STATGROUP_ElectraPlayer);
+
+
+namespace Electra
+{
+
+void FStreamSegmentRequestMPEGAudio::GetRequestedStreams(TArray<TSharedPtrTS<IStreamSegment>>& OutRequestedStreams)
+{
+	OutRequestedStreams.Empty();
+	OutRequestedStreams.Emplace(SharedThis(this));
+}
+
+void FStreamSegmentRequestMPEGAudio::GetEndedStreams(TArray<TSharedPtrTS<IStreamSegment>>& OutAlreadyEndedStreams)
+{
+	OutAlreadyEndedStreams.Empty();
+}
+
+FTimeRange FStreamSegmentRequestMPEGAudio::GetTimeRange() const
+{
+	FTimeRange tr;
+	tr.Start = FirstPTS;
+	tr.End = FirstPTS + FTimeValue(Duration);
+	tr.Start.SetSequenceIndex(TimestampSequenceIndex);
+	tr.End.SetSequenceIndex(TimestampSequenceIndex);
+	return tr;
+}
+
+
+struct FStreamReaderMPEGAudio::FLiveRequest : public TSharedFromThis<FStreamReaderMPEGAudio::FLiveRequest, ESPMode::ThreadSafe>
+{
+	FLiveRequest()
+	{
+		Events.Resize(4);
+	}
+
+	void OnProcessRequestComplete(FHttpRequestPtr InSourceHttpRequest, FHttpResponsePtr InHttpResponse, bool bInSucceeded);
+	void OnHeaderReceived(FHttpRequestPtr InSourceHttpRequest, const FString& InHeaderName, const FString& InHeaderValue);
+	void OnStatusCodeReceived(FHttpRequestPtr InSourceHttpRequest, int32 InHttpStatusCode);
+	bool OnProcessRequestStream(void* InDataPtr, int64 InLength);
+	void Cancel();
+	void WaitUntilFinished();
+	enum class EEvent
+	{
+		None,
+		Failed,
+		Finished
+	};
+	TSharedPtr<IHttpRequest, ESPMode::ThreadSafe> Handle;
+	FHttpRequestStreamDelegate StreamDelegate;
+	int32 StatusCode = 0;
+	TMediaMessageQueueWithTimeout<EEvent> Events;
+	TMultiMap<FString, FString> Headers;
+	TWeakPtrTS<IElectraHttpManager::FReceiveBuffer> ReceiveBuffer;
+	bool bCanceled = false;
+	int32 MetaDataEveryNBytes = 0;
+	int32 MetaDataBytesToGo = 0;
+	int32 BytesUntilNextMetadata = 0;
+	int64 TotalDataBytePos = 0;
+	int32 MaxDataBytes = 0;
+	bool bIsReceivingMetadata = false;
+	volatile bool bHasFailed = false;
+	TArray<uint8> MetadataBuffer;
+	TMap<int64, TArray<uint8>> MetadataBufferMap;
+};
+
+
+
+uint32	FStreamReaderMPEGAudio::UniqueDownloadID = 1;
+
+FStreamReaderMPEGAudio::FStreamReaderMPEGAudio()
+{
+}
+
+FStreamReaderMPEGAudio::~FStreamReaderMPEGAudio()
+{
+	Close();
+}
+
+UEMediaError FStreamReaderMPEGAudio::Create(IPlayerSessionServices* InPlayerSessionService, const CreateParam &InCreateParam)
+{
+	if (!InCreateParam.MemoryProvider || !InCreateParam.EventListener)
+	{
+		return UEMEDIA_ERROR_BAD_ARGUMENTS;
+	}
+
+	PlayerSessionServices = InPlayerSessionService;
+	Parameters = InCreateParam;
+	bTerminate = false;
+	bIsStarted = true;
+
+	ThreadSetName("ElectraPlayer::MPEGAudio streamer");
+	ThreadStart(FMediaRunnable::FStartDelegate::CreateRaw(this, &FStreamReaderMPEGAudio::WorkerThread));
+
+	return UEMEDIA_ERROR_OK;
+}
+
+void FStreamReaderMPEGAudio::Close()
+{
+	if (bIsStarted)
+	{
+		bIsStarted = false;
+		CancelRequests();
+		bTerminate = true;
+		WorkSignal.Signal();
+		ThreadWaitDone();
+		ThreadReset();
+		CurrentRequest.Reset();
+	}
+}
+
+void FStreamReaderMPEGAudio::LogMessage(IInfoLog::ELevel Level, const FString& Message)
+{
+	PlayerSessionServices->PostLog(Facility::EFacility::MPEGAudioStreamReader, Level, Message);
+}
+
+IStreamReader::EAddResult FStreamReaderMPEGAudio::AddRequest(uint32 CurrentPlaybackSequenceID, TSharedPtrTS<IStreamSegment> InRequest)
+{
+	TSharedPtrTS<FStreamSegmentRequestMPEGAudio> Request = CurrentRequest;
+	if (Request.IsValid())
+	{
+		check(!"why is the handler busy??");
+		return IStreamReader::EAddResult::TryAgainLater;
+	}
+	Request = StaticCastSharedPtr<FStreamSegmentRequestMPEGAudio>(InRequest);
+	Request->SetPlaybackSequenceID(CurrentPlaybackSequenceID);
+	bRequestCanceled = false;
+	bHasErrored = false;
+	// Only add the request if it is not an all-EOS one!
+	if (!Request->bIsEOSRequest)
+	{
+		CurrentRequest = Request;
+		WorkSignal.Signal();
+	}
+	return EAddResult::Added;
+}
+
+void FStreamReaderMPEGAudio::CancelRequest(EStreamType StreamType, bool bSilent)
+{
+}
+
+void FStreamReaderMPEGAudio::CancelRequests()
+{
+	bRequestCanceled = true;
+	TSharedPtrTS<IElectraHttpManager::FReceiveBuffer> RcvBuf = ReceiveBuffer;
+	if (RcvBuf.IsValid())
+	{
+		RcvBuf->Buffer.Abort();
+	}
+	TSharedPtrTS<FLiveRequest> lr(LiveRequest);
+	if (lr.IsValid())
+	{
+		lr->Cancel();
+	}
+}
+
+bool FStreamReaderMPEGAudio::HasBeenAborted() const
+{
+	TSharedPtrTS<IElectraHttpManager::FReceiveBuffer> RcvBuf = ReceiveBuffer;
+	return bRequestCanceled || (RcvBuf.IsValid() && RcvBuf->Buffer.WasAborted());
+}
+
+bool FStreamReaderMPEGAudio::HasErrored() const
+{
+	return bHasErrored;
+}
+
+int32 FStreamReaderMPEGAudio::HTTPProgressCallback(const IElectraHttpManager::FRequest* InRequest)
+{
+	// Aborted?
+	return HasBeenAborted() ? 1 : 0;
+}
+
+void FStreamReaderMPEGAudio::HTTPCompletionCallback(const IElectraHttpManager::FRequest* InRequest)
+{
+	bHasErrored = bHasErrored || InRequest->ConnectionInfo.StatusInfo.ErrorDetail.IsError();
+}
+
+void FStreamReaderMPEGAudio::FLiveRequest::Cancel()
+{
+	if (!bCanceled)
+	{
+		bCanceled = true;
+		if (Handle.IsValid())
+		{
+			Handle->CancelRequest();
+		}
+	}
+}
+
+void FStreamReaderMPEGAudio::FLiveRequest::WaitUntilFinished()
+{
+	EEvent evt = EEvent::None;
+	while(1)
+	{
+		if (Events.ReceiveMessage(evt, 20*1000))
+		{
+			break;
+		}
+	}
+}
+
+void FStreamReaderMPEGAudio::FLiveRequest::OnProcessRequestComplete(FHttpRequestPtr InSourceHttpRequest, FHttpResponsePtr InHttpResponse, bool bInSucceeded)
+{
+	bHasFailed = true;
+	Events.SendMessage(bInSucceeded ? EEvent::Finished : EEvent::Failed);
+}
+
+void FStreamReaderMPEGAudio::FLiveRequest::OnHeaderReceived(FHttpRequestPtr InSourceHttpRequest, const FString& InHeaderName, const FString& InHeaderValue)
+{
+	if (InHeaderName.Len())
+	{
+		Headers.Add(InHeaderName.ToLower(), InHeaderValue);
+
+		// Icecast metadata interval?
+		if (InHeaderName.Equals(TEXT("icy-metaint"), ESearchCase::IgnoreCase))
+		{
+			LexFromString(MetaDataEveryNBytes, *InHeaderValue);
+			BytesUntilNextMetadata = MetaDataEveryNBytes > 0 ? MetaDataEveryNBytes : -1;
+		}
+	}
+}
+
+void FStreamReaderMPEGAudio::FLiveRequest::OnStatusCodeReceived(FHttpRequestPtr InSourceHttpRequest, int32 InHttpStatusCode)
+{
+	if (InHttpStatusCode > 0 && InHttpStatusCode < 600)
+	{
+		StatusCode = InHttpStatusCode;
+	}
+}
+
+bool FStreamReaderMPEGAudio::FLiveRequest::OnProcessRequestStream(void* InDataPtr, int64 InLength)
+{
+	if (StatusCode == 200)
+	{
+		while(InLength > 0)
+		{
+			// Are we receiving metadata right now?
+			if (bIsReceivingMetadata)
+			{
+				// How many more bytes of metadata to receive?
+				if (MetaDataBytesToGo <= 0)
+				{
+					MetaDataBytesToGo = 16 * *reinterpret_cast<const uint8*>(InDataPtr);
+					InLength -= 1;
+					InDataPtr = (void*)((uint8*)InDataPtr + 1);
+					MetadataBuffer.Empty();
+				}
+				int32 MetaBytesNow = Utils::Min((int32) InLength, MetaDataBytesToGo);
+				MetadataBuffer.Append(reinterpret_cast<const uint8*>(InDataPtr), MetaBytesNow);
+				InLength -= MetaBytesNow;
+				InDataPtr = (void*)((uint8*)InDataPtr + MetaBytesNow);
+				bIsReceivingMetadata = (MetaDataBytesToGo -= MetaBytesNow) > 0;
+				if (!bIsReceivingMetadata)
+				{
+					BytesUntilNextMetadata = MetaDataEveryNBytes;
+					while(MetadataBuffer.Num() && MetadataBuffer[MetadataBuffer.Num() - 1] == 0)
+					{
+						MetadataBuffer.RemoveAt(MetadataBuffer.Num() - 1);
+					}
+					MetadataBufferMap.Emplace(TotalDataBytePos, MoveTemp(MetadataBuffer));
+				}
+			}
+			else
+			{
+				int32 DataBytesNow = BytesUntilNextMetadata > 0 ? Utils::Min((int32) InLength, BytesUntilNextMetadata) : InLength;
+				TSharedPtrTS<IElectraHttpManager::FReceiveBuffer> rb = ReceiveBuffer.Pin();
+				if (rb.IsValid())
+				{
+					FScopeLock Lock(rb->Buffer.GetLock());
+					int64 BufSizeRequired = rb->Buffer.Num() + DataBytesNow;
+					if (!rb->Buffer.EnlargeTo(BufSizeRequired))
+					{
+						return false;
+					}
+					if (!rb->Buffer.PushData(reinterpret_cast<const uint8*>(InDataPtr), DataBytesNow))
+					{
+						return false;
+					}
+
+					// Sanity check that we are not reading excessive data, which is the case when the
+					// player has been paused for instance.
+					if (MaxDataBytes && rb->Buffer.Num() > MaxDataBytes)
+					{
+						bHasFailed = true;
+						return false;
+					}
+				}
+				TotalDataBytePos += DataBytesNow;
+				InLength -= DataBytesNow;
+				InDataPtr = (void*)((uint8*)InDataPtr + DataBytesNow);
+				bIsReceivingMetadata = MetaDataEveryNBytes && (BytesUntilNextMetadata -= DataBytesNow) <= 0;
+				if (bIsReceivingMetadata)
+				{
+					BytesUntilNextMetadata = MetaDataEveryNBytes;
+					MetaDataBytesToGo = 0;
+				}
+			}
+		}
+	}
+	return true;
+}
+
+
+void FStreamReaderMPEGAudio::HandleRequest()
+{
+	TSharedPtrTS<FStreamSegmentRequestMPEGAudio> Request = CurrentRequest;
+	FString ParsingErrorMessage;
+
+	Metrics::FSegmentDownloadStats& ds = Request->DownloadStats;
+	ds.StatsID = FMediaInterlockedIncrement(UniqueDownloadID);
+	ds.MediaAssetID = TEXT("1");
+	ds.AdaptationSetID = TEXT("1");
+	ds.RepresentationID = TEXT("1");
+	ds.Bitrate = Request->GetBitrate();
+	FManifestMPEGAudioInternal::FTimelineAssetMPEGAudio* TimelineAsset = static_cast<FManifestMPEGAudioInternal::FTimelineAssetMPEGAudio*>(Request->MediaAsset.Get());
+	if (TimelineAsset->GetNumberOfAdaptationSets(EStreamType::Audio))
+	{
+		TSharedPtrTS<IPlaybackAssetAdaptationSet> AdaptationSet = Request->MediaAsset->GetAdaptationSetByTypeAndIndex(EStreamType::Audio, 0);
+		FString AdaptID = AdaptationSet->GetUniqueIdentifier();
+		if (AdaptationSet->GetNumberOfRepresentations())
+		{
+			TSharedPtrTS<IPlaybackAssetRepresentation> Representation = AdaptationSet->GetRepresentationByIndex(0);
+			ds.MediaAssetID = TimelineAsset->GetUniqueIdentifier();
+			ds.AdaptationSetID = AdaptID;
+			ds.RepresentationID = Representation->GetUniqueIdentifier();
+			ds.Bitrate = Representation->GetBitrate();
+		}
+	}
+	ds.FailureReason.Empty();
+	ds.bWasSuccessful = true;
+	ds.bWasAborted = false;
+	ds.bDidTimeout = false;
+	ds.HTTPStatusCode = 0;
+	ds.StreamType = Request->GetType();
+	ds.SegmentType = Metrics::ESegmentType::Media;
+	ds.PresentationTime = Request->FirstPTS.GetAsSeconds();
+	ds.Duration = Request->Duration;
+	ds.DurationDownloaded = 0.0;
+	ds.DurationDelivered = 0.0;
+	ds.TimeToFirstByte = 0.0;
+	ds.TimeToDownload = 0.0;
+	ds.ByteSize = -1;
+	ds.NumBytesDownloaded  = 0;
+	ds.bInsertedFillerData = false;
+	ds.URL = TimelineAsset->GetMediaURL();
+	ds.bIsMissingSegment = false;
+	ds.bParseFailure = false;
+	ds.RetryNumber = Request->NumOverallRetries;
+
+	Parameters.EventListener->OnFragmentOpen(Request);
+
+	// We need to handle Live playback (ie an Icecast) differently.
+	const bool bIsLivePlayback = Request->bIsLive;
+
+	TSharedPtrTS<IElectraHttpManager::FRequest> HTTP(new IElectraHttpManager::FRequest);
+	TSharedPtrTS<IElectraHttpManager::FProgressListener> ProgressListener;
+	ReceiveBuffer = MakeSharedTS<IElectraHttpManager::FReceiveBuffer>();
+	if (!bIsLivePlayback)
+	{
+		ProgressListener = MakeSharedTS<IElectraHttpManager::FProgressListener>();
+		ProgressListener->CompletionDelegate = IElectraHttpManager::FProgressListener::FCompletionDelegate::CreateRaw(this, &FStreamReaderMPEGAudio::HTTPCompletionCallback);
+		ProgressListener->ProgressDelegate   = IElectraHttpManager::FProgressListener::FProgressDelegate::CreateRaw(this, &FStreamReaderMPEGAudio::HTTPProgressCallback);
+
+		HTTP->Parameters.URL = TimelineAsset->GetMediaURL();
+		HTTP->Parameters.Range.Start = Request->FileStartOffset;
+		HTTP->Parameters.Range.EndIncluding = Request->FileEndOffset;
+		// No compression as this would not yield much with already compressed data.
+		HTTP->Parameters.AcceptEncoding.Set(TEXT("identity"));
+		// Timeouts
+		HTTP->Parameters.ConnectTimeout = PlayerSessionServices->GetOptionValue(MPEGAudio::OptionKeyMPEGAudioLoadConnectTimeout).SafeGetTimeValue(FTimeValue().SetFromMilliseconds(1000 * 8));
+		HTTP->Parameters.NoDataTimeout = PlayerSessionServices->GetOptionValue(MPEGAudio::OptionKeyMPEGAudioLoadNoDataTimeout).SafeGetTimeValue(FTimeValue().SetFromMilliseconds(1000 * 6));
+
+		// Explicit range?
+		int64 NumRequestedBytes = HTTP->Parameters.Range.GetNumberOfBytes();
+
+		HTTP->ReceiveBuffer = ReceiveBuffer;
+		HTTP->ProgressListener = ProgressListener;
+		HTTP->ResponseCache = PlayerSessionServices->GetHTTPResponseCache();
+		HTTP->ExternalDataReader = PlayerSessionServices->GetExternalDataReader();
+		PlayerSessionServices->GetHTTPManager()->AddRequest(HTTP, false);
+	}
+	else
+	{
+		LiveRequest = MakeSharedTS<FLiveRequest>();
+		LiveRequest->ReceiveBuffer = ReceiveBuffer;
+		LiveRequest->StreamDelegate.BindThreadSafeSP(LiveRequest.ToSharedRef(), &FStreamReaderMPEGAudio::FLiveRequest::OnProcessRequestStream);
+		LiveRequest->Handle = FHttpModule::Get().CreateRequest();
+		LiveRequest->Handle->SetVerb(TEXT("GET"));
+		LiveRequest->Handle->SetURL(TimelineAsset->GetMediaURL());
+		LiveRequest->Handle->SetDelegateThreadPolicy(EHttpRequestDelegateThreadPolicy::CompleteOnHttpThread);
+		LiveRequest->Handle->OnProcessRequestComplete().BindThreadSafeSP(LiveRequest.ToSharedRef(), &FStreamReaderMPEGAudio::FLiveRequest::OnProcessRequestComplete);
+		LiveRequest->Handle->OnHeaderReceived().BindThreadSafeSP(LiveRequest.ToSharedRef(), &FStreamReaderMPEGAudio::FLiveRequest::OnHeaderReceived);
+		LiveRequest->Handle->OnStatusCodeReceived().BindThreadSafeSP(LiveRequest.ToSharedRef(), &FStreamReaderMPEGAudio::FLiveRequest::OnStatusCodeReceived);
+		LiveRequest->Handle->SetResponseBodyReceiveStreamDelegate(LiveRequest->StreamDelegate);
+		LiveRequest->Handle->SetHeader(TEXT("User-Agent"), IElectraHttpManager::GetDefaultUserAgent());
+		LiveRequest->Handle->SetHeader(TEXT("Accept-Encoding"), TEXT("identity"));
+		// If this is an Icycast, we ask for period metadata.
+		if (Request->CastType == FStreamSegmentRequestMPEGAudio::ECastType::IcyCast)
+		{
+			LiveRequest->Handle->SetHeader(TEXT("Icy-Metadata"), TEXT("1"));
+		}
+		LiveRequest->MaxDataBytes = (int32) (PlayerSessionServices->GetOptionValue(MPEGAudio::OptionKeyMPEGAudioMaxPreloadBufferDuration).SafeGetTimeValue(FTimeValue().SetFromMilliseconds(1000 * 90)).GetAsSeconds() * Request->GetBitrate() / 8);
+		LiveRequest->Handle->SetActivityTimeout(PlayerSessionServices->GetOptionValue(MPEGAudio::OptionKeyMPEGAudioLoadNoDataTimeout).SafeGetTimeValue(FTimeValue().SetFromMilliseconds(1000 * 5)).GetAsSeconds());
+		LiveRequest->Handle->ProcessRequest();
+	}
+
+	FTimeValue DurationSuccessfullyDelivered(FTimeValue::GetZero());
+	FTimeValue DurationSuccessfullyRead(FTimeValue::GetZero());
+	bool bDone = false;
+
+	uint32 PlaybackSequenceID = Request->GetPlaybackSequenceID();
+
+	// We have to probe the data for sync markers as the file has no framing whatsoever.
+	// Other than perhaps at the beginning of the file it would be purely coincidental to start on a frame.
+	auto GetUINT32BE = [](const uint8* InData) -> uint32
+	{
+		return (static_cast<uint32>(InData[0]) << 24) | (static_cast<uint32>(InData[1]) << 16) | (static_cast<uint32>(InData[2]) << 8) | static_cast<uint32>(InData[3]);
+	};
+	int32 SyncMarkerCheckPos = 0;
+	const int32 FrameCheckSize = 4096;
+	const int32 NumFramesToProbe = 10;
+	const uint32 FrameSyncMask = Request->MPEGHeaderMask;
+	const uint32 FrameSyncValue = Request->MPEGHeaderExpectedValue;
+	TArray<int32> SyncMarkerOffsets;
+	while(!bDone && !HasErrored() && !HasBeenAborted() && !bTerminate)
+	{
+		if (ReceiveBuffer->Buffer.WaitUntilSizeAvailable(SyncMarkerCheckPos + FrameCheckSize, 1000 * 20))
+		{
+			FScopeLock Lock(ReceiveBuffer->Buffer.GetLock());
+			const uint8* BufferBaseData = ReceiveBuffer->Buffer.GetLinearReadData();
+			int64 BufferDataSize = ReceiveBuffer->Buffer.GetLinearReadSize();
+			if (!BufferBaseData)
+			{
+				bDone = true;
+				break;
+			}
+
+			const uint8* Base = BufferBaseData + SyncMarkerCheckPos;
+			const uint8* End = BufferBaseData + BufferDataSize - 4;
+			for(; Base < End; ++Base)
+			{
+				if (*Base != 0xff)
+				{
+					continue;
+				}
+				// Check for validity
+				uint32 hdr = GetUINT32BE(Base);
+				if ((hdr & FrameSyncMask) == FrameSyncValue)
+				{
+					SyncMarkerOffsets.Add(Base - BufferBaseData);
+					Base += 3;
+					if (SyncMarkerOffsets.Num() >= NumFramesToProbe)
+					{
+						bDone = true;
+					}
+				}
+			}
+			SyncMarkerCheckPos = BufferDataSize - 1;
+			if (ReceiveBuffer->Buffer.GetEOD())
+			{
+				bDone = true;
+			}
+		}
+	}
+	int32 FirstAUBufferOffset = -1;
+	int32 FirstAUFrameSize = 0;
+	// Probe that we are properly locked on to the frames.
+	if (!HasErrored() && !HasBeenAborted() && !bTerminate)
+	{
+		FScopeLock Lock(ReceiveBuffer->Buffer.GetLock());
+		const uint8* BufferBaseData = ReceiveBuffer->Buffer.GetLinearReadData();
+		if (BufferBaseData)
+		{
+			const int32 MaxCheckFrames = FMath::Min(10, SyncMarkerOffsets.Num());
+			const uint8* MaxCheckAddr = BufferBaseData + (SyncMarkerOffsets.Num() ? SyncMarkerOffsets[MaxCheckFrames-1] : 0);
+			for(int32 i=0; i<MaxCheckFrames; ++i)
+			{
+				const uint8* CheckBaseAddr = BufferBaseData + SyncMarkerOffsets[i];
+				bool bRunOk = true;
+				uint32 HeaderValue = GetUINT32BE(CheckBaseAddr);
+				int32 StartFrameSize = ElectraDecodersUtil::MPEG::UtilsMPEG123::GetFrameSize(HeaderValue);
+				for(int32 j=i+1; j<MaxCheckFrames; ++j)
+				{
+					int32 FrameSize = ElectraDecodersUtil::MPEG::UtilsMPEG123::GetFrameSize(HeaderValue);
+					CheckBaseAddr += FrameSize;
+					if (CheckBaseAddr >= MaxCheckAddr)
+					{
+						break;
+					}
+					HeaderValue = GetUINT32BE(CheckBaseAddr);
+					if ((HeaderValue & FrameSyncMask) != FrameSyncValue)
+					{
+						bRunOk = false;
+						break;
+					}
+				}
+				if (bRunOk)
+				{
+					FirstAUBufferOffset = SyncMarkerOffsets[i];
+					FirstAUFrameSize = StartFrameSize;
+					break;
+				}
+			}
+		}
+	}
+
+	bDone = FirstAUBufferOffset < 0;
+	// If we did not find any sync marker to start on we are done.
+	if (bDone)
+	{
+		bDone = true;
+		Request->LastSuccessfullyUsedBytePos = Request->FileEndOffset;
+		Request->LastSuccessfullyUsedPTS = Request->FirstPTS;
+	}
+	FTimeValue NextAUPTS = Request->FirstPTS;
+	NextAUPTS.SetSequenceIndex(Request->TimestampSequenceIndex);
+	int32 NextAUBufferOffset = FirstAUBufferOffset;
+	int32 NextAUFrameSize = FirstAUFrameSize;
+	bool bIsFirstInSequence = true;
+	bool bReadPastLastPTS = false;
+	FTimeValue AUDuration;
+	TSharedPtrTS<FAccessUnit::CodecData> AUCodecData = MakeSharedTS<FAccessUnit::CodecData>();
+	AUCodecData->ParsedInfo = Request->CodecInfo;
+
+	TSharedPtrTS<FBufferSourceInfo> BufferSourceInfo = MakeSharedTS<FBufferSourceInfo>();
+	BufferSourceInfo->PeriodID = ds.MediaAssetID;
+	BufferSourceInfo->PeriodAdaptationSetID = ds.MediaAssetID + TEXT(".") + ds.AdaptationSetID;
+	BufferSourceInfo->Language = TEXT("und");
+	BufferSourceInfo->Codec = TEXT("mp4a.6b");
+	BufferSourceInfo->HardIndex = 0;
+	BufferSourceInfo->PlaybackSequenceID = Request->GetPlaybackSequenceID();
+
+	int64 NumTotalBytesRead = 0;
+	// Remove data we already processed every 4 seconds. Do this in blocks to reduce the overall
+	// amount of memory being moved around.
+	int32 NumLiveStreamBytesToRemove = Request->GetBitrate() * 4 / 8;
+
+	while(!bDone && !HasErrored() && !HasBeenAborted() && !bTerminate)
+	{
+		check(NextAUFrameSize);
+		if (!NextAUFrameSize)
+		{
+			break;
+		}
+
+		// With Live playback we have to remove the data bytes we already passed along from the start of the buffer.
+		if (LiveRequest.IsValid() && NumLiveStreamBytesToRemove && NextAUBufferOffset >= NumLiveStreamBytesToRemove)
+		{
+			ReceiveBuffer->Buffer.RemoveFromBeginning(NextAUBufferOffset);
+			NextAUBufferOffset = 0;
+		}
+
+		// Wait until we get the next AU's data plus the following 4 bytes that are the header of the following frame.
+		const int64 TotalNumNeeded = NextAUBufferOffset + NextAUFrameSize + 4;
+		if (ReceiveBuffer->Buffer.WaitUntilSizeAvailable(TotalNumNeeded, 1000 * 20))
+		{
+			FScopeLock Lock(ReceiveBuffer->Buffer.GetLock());
+
+			const uint8* BufferBaseData = ReceiveBuffer->Buffer.GetLinearReadData();
+			if (!BufferBaseData)
+			{
+				break;
+			}
+			int64 BufferDataSize = ReceiveBuffer->Buffer.GetLinearReadSize();
+
+			int64 NumGot = BufferDataSize - NextAUBufferOffset;
+			// Did we get the next frame's worth?
+			if (NumGot >= NextAUFrameSize)
+			{
+				BufferBaseData += NextAUBufferOffset;
+				uint32 HeaderValue = GetUINT32BE(BufferBaseData);
+				// Safety check.
+				check((HeaderValue & FrameSyncMask) == FrameSyncValue);
+				if ((HeaderValue & FrameSyncMask) != FrameSyncValue)
+				{
+					ParsingErrorMessage = FString::Printf(TEXT("Frame sync marker not found. Corrupt file?"));
+					LogMessage(IInfoLog::ELevel::Error, ParsingErrorMessage);
+					bDone = true;
+					bHasErrored = true;
+				}
+
+				if (!AUDuration.IsValid())
+				{
+					int32 FrameSize = ElectraDecodersUtil::MPEG::UtilsMPEG123::GetSamplesPerFrame(HeaderValue);
+					int32 SampleRate = ElectraDecodersUtil::MPEG::UtilsMPEG123::GetSamplingRate(HeaderValue);
+					AUDuration.SetFromND(FrameSize, (uint32)SampleRate);
+				}
+
+				// Keep track of how many bytes we have consumed so far. This correlates the Live stream metadata to this time.
+				NumTotalBytesRead += NextAUFrameSize;
+				if (LiveRequest.IsValid())
+				{
+					if (LiveRequest->bHasFailed)
+					{
+						bHasErrored = true;
+					}
+
+					TArray<int64> MetaPositions;
+					LiveRequest->MetadataBufferMap.GenerateKeyArray(MetaPositions);
+					for(int32 nm=0; nm<MetaPositions.Num(); ++nm)
+					{
+						if (MetaPositions[nm] <= NumTotalBytesRead)
+						{
+							TArray<uint8> mda = LiveRequest->MetadataBufferMap[MetaPositions[nm]];
+							LiveRequest->MetadataBufferMap.Remove(MetaPositions[nm]);
+							if (mda.Num())
+							{
+								const TCHAR* const Delimiter = TEXT(";");
+								TArray<FString> md;
+								StringHelpers::ArrayToString(mda).ParseIntoArray(md, Delimiter, true);
+								for(int32 k=0; k<md.Num(); ++k)
+								{
+									if (md[k].StartsWith(TEXT("StreamTitle=")))
+									{
+										FString Title = md[k].Mid(12);
+										if (Title.Len()>1 && Title[0]==TCHAR('\'') && Title[Title.Len()-1]==TCHAR('\''))
+										{
+											Title.MidInline(1, Title.Len()-2);
+										}
+
+										TSharedPtr<UtilsMP4::FMetadataParser, ESPMode::ThreadSafe> mp = MakeShared<UtilsMP4::FMetadataParser, ESPMode::ThreadSafe>();
+										mp->AddItem(TEXT("Title"), Title);
+										PlayerSessionServices->SendMessageToPlayer(FPlaylistMetadataUpdateMessage::Create(NextAUPTS, mp, true));
+									}
+								}
+							}
+						}
+					}
+				}
+
+				// Create an access unit.
+				FAccessUnit *AccessUnit = FAccessUnit::Create(Parameters.MemoryProvider);
+				if (AccessUnit)
+				{
+					AccessUnit->ESType = EStreamType::Audio;
+					AccessUnit->PTS = NextAUPTS;
+					AccessUnit->DTS = NextAUPTS;
+					AccessUnit->Duration = AUDuration;
+
+					AccessUnit->EarliestPTS = Request->EarliestPTS;
+					AccessUnit->LatestPTS = Request->LastPTS;
+
+					AccessUnit->AUSize = (uint32) NextAUFrameSize;
+					AccessUnit->AUCodecData = AUCodecData;
+					AccessUnit->DropState = FAccessUnit::EDropState::None;
+					if (AccessUnit->PTS >= AccessUnit->LatestPTS)
+					{
+						AccessUnit->DropState |= FAccessUnit::EDropState::TooLate;
+					}
+
+					// Set the sequence index member and update all timestamps with it as well.
+					AccessUnit->SequenceIndex = Request->TimestampSequenceIndex;
+					AccessUnit->DTS.SetSequenceIndex(Request->TimestampSequenceIndex);
+					AccessUnit->PTS.SetSequenceIndex(Request->TimestampSequenceIndex);
+					AccessUnit->EarliestPTS.SetSequenceIndex(Request->TimestampSequenceIndex);
+					AccessUnit->LatestPTS.SetSequenceIndex(Request->TimestampSequenceIndex);
+
+					AccessUnit->bIsFirstInSequence = bIsFirstInSequence;
+					AccessUnit->bIsSyncSample = true;
+					AccessUnit->bIsDummyData = false;
+					AccessUnit->AUData = AccessUnit->AllocatePayloadBuffer(AccessUnit->AUSize);
+
+					// Set the associated stream metadata
+					AccessUnit->BufferSourceInfo = BufferSourceInfo;
+
+					bIsFirstInSequence = false;
+
+					FMemory::Memcpy(AccessUnit->AUData, BufferBaseData, NextAUFrameSize);
+
+					// Unlock the receive buffer now so the reader won't be blocked.
+					Lock.Unlock();
+
+					DurationSuccessfullyRead += AccessUnit->Duration;
+
+					bool bSentOff = false;
+
+					// Check if the AU is outside the time range we are allowed to read.
+					// The last one (the one that is already outside the range, actually) is tagged as such and sent into the buffer.
+					// The respective decoder has to handle this flag if necessary and/or drop the AU.
+					// We need to send at least one AU down so the FMultiTrackAccessUnitBuffer does not stay empty for this period!
+					// Already sent the last one?
+					if (bReadPastLastPTS)
+					{
+						// Yes. Release this AU and be done.
+						bSentOff = true;
+						bDone = true;
+					}
+					else if (AccessUnit->DTS >= AccessUnit->LatestPTS && AccessUnit->PTS >= AccessUnit->LatestPTS)
+					{
+						// Tag the last one and send it off.
+						AccessUnit->bIsLastInPeriod = true;
+						bReadPastLastPTS = true;
+					}
+
+					while(!bSentOff && !HasBeenAborted() && !bTerminate)
+					{
+						if (Parameters.EventListener->OnFragmentAccessUnitReceived(AccessUnit))
+						{
+							DurationSuccessfullyDelivered += AccessUnit->Duration;
+							bSentOff = true;
+							AccessUnit = nullptr;
+
+							// Since we have delivered this access unit, if we are detecting an error now we need to then
+							// retry on the _next_ AU and not this one again!
+							Request->LastSuccessfullyUsedBytePos = Request->FileStartOffset + NextAUBufferOffset + NextAUFrameSize;
+							Request->LastSuccessfullyUsedPTS = NextAUPTS + AUDuration;
+						}
+						else
+						{
+							FMediaRunnable::SleepMicroseconds(1000 * 10);
+						}
+					}
+
+					// Release the AU if we still have it.
+					FAccessUnit::Release(AccessUnit);
+					AccessUnit = nullptr;
+
+					// For error handling, if we managed to get additional data we reset the retry count.
+					if (ds.RetryNumber && DurationSuccessfullyRead.GetAsSeconds() > 1.0)
+					{
+						ds.RetryNumber = 0;
+						Request->NumOverallRetries = 0;
+					}
+				}
+
+				// Advance to next
+				NextAUPTS += AUDuration;
+				NextAUBufferOffset += NextAUFrameSize;
+				NumGot -= NextAUFrameSize;
+				NextAUFrameSize = 0;
+				// Did we also get the next 4 bytes?
+				if (NumGot >= 4)
+				{
+					FScopeLock Lock2(ReceiveBuffer->Buffer.GetLock());
+					HeaderValue = GetUINT32BE(ReceiveBuffer->Buffer.GetLinearReadData() + NextAUBufferOffset + NextAUFrameSize);
+					NextAUFrameSize = ElectraDecodersUtil::MPEG::UtilsMPEG123::GetFrameSize(HeaderValue);
+				}
+				else
+				{
+					bDone = true;
+				}
+			}
+			else
+			{
+				bDone = true;
+			}
+		}
+	}
+
+	// Remove the download request.
+	if (!bIsLivePlayback)
+	{
+		ProgressListener.Reset();
+		PlayerSessionServices->GetHTTPManager()->RemoveRequest(HTTP, false);
+		Request->ConnectionInfo = HTTP->ConnectionInfo;
+		HTTP.Reset();
+	}
+	else
+	{
+		LiveRequest->Cancel();
+		LiveRequest->WaitUntilFinished();
+
+		Request->ConnectionInfo.ContentType = TEXT("audio/mpeg");
+		Request->ConnectionInfo.BytesReadSoFar = LiveRequest->TotalDataBytePos;
+		Request->ConnectionInfo.HTTPVersionReceived = 11;
+		Request->ConnectionInfo.bWasAborted = bRequestCanceled;
+		Request->ConnectionInfo.bHasFinished = true;
+		Request->ConnectionInfo.StatusInfo.HTTPStatus = LiveRequest->StatusCode;
+		Request->ConnectionInfo.StatusInfo.bReadError = true;
+		// Unless the server responded with a bad status code, set the retry count to a negative
+		// value to prevent the ABR from checking against it. A live cast has no end, so getting
+		// here with a good status code (unless aborted) means there was a connection error,
+		// which can be retried indefinitely.
+		if (!bRequestCanceled && LiveRequest->StatusCode < 300)
+		{
+			ds.RetryNumber = -1;
+		}
+
+		LiveRequest.Reset();
+	}
+
+	// Set up download stat fields.
+	if (Request->ConnectionInfo.EffectiveURL.Len())
+	{
+		ds.URL = Request->ConnectionInfo.EffectiveURL;
+	}
+	ds.HTTPStatusCode = Request->ConnectionInfo.StatusInfo.HTTPStatus;
+	ds.TimeToFirstByte = Request->ConnectionInfo.TimeUntilFirstByte;
+	ds.TimeToDownload = (Request->ConnectionInfo.RequestEndTime - Request->ConnectionInfo.RequestStartTime).GetAsSeconds();
+	ds.ByteSize = Request->ConnectionInfo.ContentLength;
+	ds.NumBytesDownloaded = Request->ConnectionInfo.BytesReadSoFar;
+
+	ds.FailureReason = Request->ConnectionInfo.StatusInfo.ErrorDetail.GetMessage();
+	if (ParsingErrorMessage.Len())
+	{
+		ds.FailureReason = ParsingErrorMessage;
+	}
+	ds.bWasSuccessful = !bHasErrored;
+	ds.URL = Request->ConnectionInfo.EffectiveURL;
+	ds.HTTPStatusCode = Request->ConnectionInfo.StatusInfo.HTTPStatus;
+	ds.DurationDownloaded = DurationSuccessfullyRead.GetAsSeconds();
+	ds.DurationDelivered = DurationSuccessfullyDelivered.GetAsSeconds();
+	ds.TimeToFirstByte = Request->ConnectionInfo.TimeUntilFirstByte;
+	ds.TimeToDownload = (Request->ConnectionInfo.RequestEndTime - Request->ConnectionInfo.RequestStartTime).GetAsSeconds();
+	ds.ByteSize = Request->ConnectionInfo.ContentLength;
+	ds.NumBytesDownloaded = Request->ConnectionInfo.BytesReadSoFar;
+	ds.bIsCachedResponse = Request->ConnectionInfo.bIsCachedResponse;
+
+	// Reset the current request so another one can be added immediately when we call OnFragmentClose()
+	CurrentRequest.Reset();
+	PlayerSessionServices->GetStreamSelector()->ReportDownloadEnd(ds);
+	Parameters.EventListener->OnFragmentClose(Request);
+}
+
+void FStreamReaderMPEGAudio::WorkerThread()
+{
+	LLM_SCOPE(ELLMTag::ElectraPlayer);
+	while(!bTerminate)
+	{
+		WorkSignal.WaitAndReset();
+		if (bTerminate)
+		{
+			break;
+		}
+		TSharedPtrTS<FStreamSegmentRequestMPEGAudio> Request = CurrentRequest;
+		if (Request.IsValid())
+		{
+			HandleRequest();
+		}
+	}
+}
+
+} // namespace Electra
diff --git a/Engine/Plugins/Media/ElectraPlayer/Source/ElectraPlayerRuntime/Private/Runtime/Player/mpegaudio/StreamReaderMPEGAudio.h b/Engine/Plugins/Media/ElectraPlayer/Source/ElectraPlayerRuntime/Private/Runtime/Player/mpegaudio/StreamReaderMPEGAudio.h
new file mode 100644
--- /dev/null
+++ b/Engine/Plugins/Media/ElectraPlayer/Source/ElectraPlayerRuntime/Private/Runtime/Player/mpegaudio/StreamReaderMPEGAudio.h
@@ -0,0 +1,134 @@
+// Copyright Epic Games, Inc. All Rights Reserved.
+
+#pragma once
+
+#include "PlayerCore.h"
+#include "Player/Manifest.h"
+#include "HTTP/HTTPManager.h"
+#include "SynchronizedClock.h"
+#include "PlaylistReaderMPEGAudio.h"
+#include "Utilities/TimeUtilities.h"
+#include "HAL/LowLevelMemTracker.h"
+#include "ElectraPlayerPrivate.h"
+#include "Player/PlayerStreamReader.h"
+#include "StreamAccessUnitBuffer.h"
+
+
+
+namespace Electra
+{
+
+
+class FStreamSegmentRequestMPEGAudio : public IStreamSegment
+{
+public:
+	virtual ~FStreamSegmentRequestMPEGAudio()
+	{}
+
+	void SetPlaybackSequenceID(uint32 InPlaybackSequenceID) override
+	{ PlaybackSequenceID = InPlaybackSequenceID; }
+	uint32 GetPlaybackSequenceID() const override
+	{ return PlaybackSequenceID; }
+	void SetExecutionDelay(const FTimeValue& UTCNow, const FTimeValue& ExecutionDelay) override
+	{ }
+	FTimeValue GetExecuteAtUTCTime() const override
+	{ return FTimeValue::GetInvalid(); }
+	EStreamType GetType() const override
+	{ return EStreamType::Audio; }
+	void GetDependentStreams(TArray<TSharedPtrTS<IStreamSegment>>& OutDependentStreams) const override
+	{ }
+	void GetRequestedStreams(TArray<TSharedPtrTS<IStreamSegment>>& OutRequestedStreams) override;
+	void GetEndedStreams(TArray<TSharedPtrTS<IStreamSegment>>& OutAlreadyEndedStreams) override;
+	FTimeValue GetFirstPTS() const override
+	{ return FirstPTS; }
+	FTimeRange GetTimeRange() const override;
+
+	int32 GetQualityIndex() const override
+	{ return 0; }
+	int32 GetBitrate() const override
+	{ return Bitrate; }
+	void GetDownloadStats(Metrics::FSegmentDownloadStats& OutStats) const override
+	{ OutStats = DownloadStats; }
+	bool GetStartupDelay(FTimeValue& OutStartTime, FTimeValue& OutTimeIntoSegment, FTimeValue& OutSegmentDuration) const override
+	{ return false; }
+
+
+	TSharedPtrTS<ITimelineMediaAsset> MediaAsset;
+
+	FTimeValue FirstPTS;
+	FTimeValue EarliestPTS;					//!< PTS of the first sample to be presented.
+	FTimeValue LastPTS;						//!< PTS at which no further samples are to be presented.
+	int64 FileStartOffset = -1;				//!< Where to start in the file
+	int64 FileEndOffset = -1;				//!< Where to end in the file (for HTTP range GET requests)
+	enum class ECastType
+	{
+		None,
+		IcyCast
+	};
+	ECastType CastType = ECastType::None;
+
+	uint32 MPEGHeaderMask = 0;
+	uint32 MPEGHeaderExpectedValue = 0;
+	int32 CBRFrameSize = 0;
+	bool bIsVBR = false;
+	bool bIsLive = false;
+	bool bIsEOSRequest = false;
+	double Duration = 0.0;
+	FStreamCodecInformation CodecInfo;
+
+	uint32 PlaybackSequenceID = ~0U;
+	int32 Bitrate = 0;
+	bool bIsContinuationSegment = false;		//!< true if this segment continues where the previous left off and no sync samples should be expected.
+	bool bIsFirstSegment = false;				//!< true if this segment is the first to start with or the first after a seek.
+	bool bIsLastSegment = false;				//!< true if this segment is the last.
+
+	int64 TimestampSequenceIndex = 0;		//!< Sequence index to set in all timestamp values of the decoded access unit.
+	int32 NumOverallRetries = 0;			//!< Number of retries
+
+	int64 LastSuccessfullyUsedBytePos = -1;
+	FTimeValue LastSuccessfullyUsedPTS;
+
+	Metrics::FSegmentDownloadStats DownloadStats;
+	HTTP::FConnectionInfo ConnectionInfo;
+};
+
+
+class FStreamReaderMPEGAudio : public IStreamReader, public FMediaThread
+{
+public:
+	FStreamReaderMPEGAudio();
+	virtual ~FStreamReaderMPEGAudio();
+	virtual UEMediaError Create(IPlayerSessionServices* PlayerSessionService, const CreateParam &createParam) override;
+	virtual void Close() override;
+	virtual EAddResult AddRequest(uint32 CurrentPlaybackSequenceID, TSharedPtrTS<IStreamSegment> Request) override;
+	virtual void CancelRequest(EStreamType StreamType, bool bSilent) override;
+	virtual void CancelRequests() override;
+
+private:
+	void WorkerThread();
+	void LogMessage(IInfoLog::ELevel Level, const FString& Message);
+	int32 HTTPProgressCallback(const IElectraHttpManager::FRequest* InRequest);
+	void HTTPCompletionCallback(const IElectraHttpManager::FRequest* InRequest);
+	void HandleRequest();
+
+	bool HasBeenAborted() const;
+	bool HasErrored() const;
+
+	CreateParam Parameters;
+	IPlayerSessionServices* PlayerSessionServices = nullptr;
+	bool bIsStarted = false;
+	bool bTerminate = false;
+	bool bRequestCanceled = false;
+	bool bHasErrored = false;
+	FErrorDetail ErrorDetail;
+
+	struct FLiveRequest;
+	TSharedPtrTS<FLiveRequest> LiveRequest;
+	TSharedPtrTS<FStreamSegmentRequestMPEGAudio> CurrentRequest;
+	TSharedPtrTS<IElectraHttpManager::FReceiveBuffer> ReceiveBuffer;
+	FMediaEvent WorkSignal;
+
+	static uint32 UniqueDownloadID;
+};
+
+} // namespace Electra
diff --git a/Engine/Plugins/Media/ElectraPlayer/Source/ElectraPlayerRuntime/Private/Runtime/PlayerFacility.h b/Engine/Plugins/Media/ElectraPlayer/Source/ElectraPlayerRuntime/Private/Runtime/PlayerFacility.h
--- a/Engine/Plugins/Media/ElectraPlayer/Source/ElectraPlayerRuntime/Private/Runtime/PlayerFacility.h
+++ b/Engine/Plugins/Media/ElectraPlayer/Source/ElectraPlayerRuntime/Private/Runtime/PlayerFacility.h
@@ -40,6 +40,9 @@ namespace Electra
 			DASHMPDBuilder,
 			DASHManifest,
 			DASHStreamReader,
+			MPEGAudioPlaylistReader,
+			MPEGAudioPlaylist,
+			MPEGAudioStreamReader,
 			EntityCache,
 			DRM,
 			SubtitleDecoder,
@@ -100,6 +103,12 @@ namespace Electra
 					return TEXT("DASH manifest");
 				case EFacility::DASHStreamReader:
 					return TEXT("DASH stream reader");
+				case EFacility::MPEGAudioPlaylistReader:
+					return TEXT("MPEG audio playlist reader");
+				case EFacility::MPEGAudioPlaylist:
+					return TEXT("MPEG audio playlist");
+				case EFacility::MPEGAudioStreamReader:
+					return TEXT("MPEG audio stream reader");
 				case EFacility::EntityCache:
 					return TEXT("Entity cache");
 				case EFacility::DRM:
diff --git a/Engine/Plugins/Media/ElectraPlayer/Source/ElectraPlayerRuntime/Private/Runtime/Utilities/UtilsMPEG.cpp b/Engine/Plugins/Media/ElectraPlayer/Source/ElectraPlayerRuntime/Private/Runtime/Utilities/UtilsMPEG.cpp
--- a/Engine/Plugins/Media/ElectraPlayer/Source/ElectraPlayerRuntime/Private/Runtime/Utilities/UtilsMPEG.cpp
+++ b/Engine/Plugins/Media/ElectraPlayer/Source/ElectraPlayerRuntime/Private/Runtime/Utilities/UtilsMPEG.cpp
@@ -1,6 +1,7 @@
 // Copyright Epic Games, Inc. All Rights Reserved.
 
 #include "Utilities/UtilsMPEG.h"
+#include "Utilities/Utilities.h"
 #include "Utils/ElectraBitstreamReader.h"
 
 namespace Electra
@@ -148,6 +149,258 @@ namespace Electra
 			return true;
 		}
 
+
+
+
+
+		bool FID3V2Metadata::Parse(const uint8* InData, int64 InDataSize)
+		{
+			if (!(InData[0] == 'I' && InData[1] == 'D' && InData[2] == '3' &&
+				InData[3] != 0xff && InData[4] != 0xff && InData[6] < 0x80 && InData[7] < 0x80 && InData[8] < 0x80 && InData[9] < 0x80))
+			{
+				return false;
+			}
+			auto GetID3Size = [](const uint8* InAddr) -> int32
+			{
+				return (int32) (10U + (static_cast<uint32>(InAddr[0]) << 21) + (static_cast<uint32>(InAddr[1]) << 14) + (static_cast<uint32>(InAddr[2]) << 7) + InAddr[3]);
+			};
+			auto GetTagSize = [](const uint8* InAddr) -> int32
+			{
+				return (int32) (10U + (static_cast<uint32>(InAddr[0]) << 24) + (static_cast<uint32>(InAddr[1]) << 16) + (static_cast<uint32>(InAddr[2]) << 8) + InAddr[3]);
+			};
+			if (GetID3Size(InData + 6) > InDataSize)
+			{
+				return false;
+			}
+			// "Unsynchronization" is not currently supported.
+			check((InData[5] & 0x80) == 0);
+			if ((InData[5] & 0x80) != 0)
+			{
+				return false;
+			}
+			// Experimental headers are not supported
+			if ((InData[5] & 0x20) != 0)
+			{
+				return false;
+			}
+			int32 HeaderSize = 10;
+			if ((InData[5] & 0x40) != 0)
+			{
+				check(!"extended header not currently supported");
+				return false;
+				//HeaderSize += GetTagSize(InData + 10) + xxxx;
+			}
+			const uint8* InDataEnd = InData + InDataSize;
+			InData += HeaderSize;
+			auto GetString = [](const uint8* InStr, int32 NumBytes, const uint8** OutPtr=nullptr, int32* InOutEncoding=nullptr) -> FString
+			{
+				if (InStr && NumBytes > 0)
+				{
+					int32 Encoding = (!InOutEncoding || (InOutEncoding && *InOutEncoding<0)) ? *InStr++ : *InOutEncoding;
+					if (InOutEncoding)
+					{
+						*InOutEncoding = Encoding;
+					}
+					if (Encoding == 0)
+					{
+						// ISO 8859-1 encoding
+						TArray<uint8> ConvBuf;
+						ConvBuf.Reserve(NumBytes*2);
+						const uint8* SrcPtr = reinterpret_cast<const uint8*>(InStr);
+						int32 NumSrcChars = NumBytes - 1;
+						for(int32 i=0; i<NumSrcChars; ++i, ++SrcPtr)
+						{
+							if (*SrcPtr == 0x00)
+							{
+								++SrcPtr;
+								break;
+							}
+							// Permitted whitespaces
+							if (*SrcPtr == 0x0a || *SrcPtr == 0x09)
+							{
+								ConvBuf.Add(*SrcPtr);
+							}
+							else if (*SrcPtr >= 0x20 && *SrcPtr < 0x80)
+							{
+								ConvBuf.Add(*SrcPtr);
+							}
+							else  if (*SrcPtr >= 0xa0)
+							{
+								// We can convert straight from ISO 8859-1 to UTF8 by doing this:
+								ConvBuf.Add(0xc0 | (*SrcPtr >> 6));
+								ConvBuf.Add(0x80 | (*SrcPtr & 0x3f));
+							}
+						}
+						if (OutPtr)
+						{
+							*OutPtr = SrcPtr;
+						}
+						auto Cnv = StringCast<TCHAR>(reinterpret_cast<const UTF8CHAR*>(ConvBuf.GetData()), ConvBuf.Num());
+						return FString(Cnv.Length(), Cnv.Get());
+					}
+					else if (Encoding == 1 && NumBytes > 1)
+					{
+						const uint8* SrcPtr = reinterpret_cast<const uint8*>(InStr);
+						if (OutPtr)
+						{
+							*OutPtr = SrcPtr + NumBytes;
+						}
+						int32 NumSrcChars = NumBytes - 1;
+						uint16 BOM = ((uint16)SrcPtr[0] << 8) | SrcPtr[1];
+						// Have BOM?
+						if (BOM == 0xfffe)
+						{
+							auto Cnv = StringCast<TCHAR>(reinterpret_cast<const UCS2CHAR*>(SrcPtr+2), NumSrcChars/2-1);
+							return FString(Cnv.Length(), Cnv.Get());
+						}
+						else if (BOM == 0xfeff)
+						{
+							TArray<uint8> be;
+							be.AddUninitialized(NumSrcChars);
+							for(int32 k=0; k<NumSrcChars; k+=2)
+							{
+								be[k+0] = SrcPtr[k+1];
+								be[k+1] = SrcPtr[k+0];
+							}
+							auto Cnv = StringCast<TCHAR>(reinterpret_cast<const UCS2CHAR*>(be.GetData()+2), NumSrcChars/2-1);
+							return FString(Cnv.Length(), Cnv.Get());
+						}
+					}
+				}
+				return FString();
+			};
+			while(InData+10 < InDataEnd)
+			{
+				if (((InData[0] >= 'A' && InData[0] <= 'Z') || (InData[0] >= '0' && InData[0] <= '9')) &&
+					((InData[1] >= 'A' && InData[1] <= 'Z') || (InData[1] >= '0' && InData[1] <= '9')) &&
+					((InData[2] >= 'A' && InData[2] <= 'Z') || (InData[2] >= '0' && InData[2] <= '9')) &&
+					((InData[3] >= 'A' && InData[3] <= 'Z') || (InData[3] >= '0' && InData[3] <= '9')))
+				{
+					uint32 FrameID = (static_cast<uint32>(InData[0]) << 24) | (static_cast<uint32>(InData[1]) << 16) | (static_cast<uint32>(InData[2]) << 8) | static_cast<uint32>(InData[3]);
+					int32 FrameSize = (int32)GetTagSize(InData + 4);
+					check(FrameSize > 10);
+					if (FrameSize <= 10)
+					{
+						return false;
+					}
+					const uint8* const FrameEnd = InData + FrameSize;
+					uint16 Flags = (InData[8] << 8) | InData[9];
+					check((Flags & 0xc0) == 0);
+					// Ignore compressed, encrypted or grouped tags.
+					if ((Flags & 0xc0) != 0)
+					{
+						InData += FrameSize;
+						continue;
+					}
+					switch(FrameID)
+					{
+						default:
+						{
+							break;
+						}
+						// Track length
+						case Utils::Make4CC('T','L','E','N'):
+						{
+							FString ms(GetString(InData + 10, FrameSize - 10));
+							int64 dur = -1;
+							LexFromString(dur, *ms);
+							if (dur > 0)
+							{
+								FTimespan Duration = FTimespan::FromMilliseconds((double)dur);
+								FItem Item {.Value = FVariant(Duration) };
+								Tags.Emplace(FrameID, MoveTemp(Item));
+							}
+							break;
+						}
+						// MPEG location lookup table
+						case Utils::Make4CC('M','L','L','T'):
+						{
+							TArray<uint8> Blob(InData + 10, FrameSize - 10);
+							FItem Item {.Value = FVariant(MoveTemp(Blob)) };
+							Tags.Emplace(FrameID, MoveTemp(Item));
+							break;
+						}
+						// Cover image
+						case Utils::Make4CC('A','P','I','C'):
+						{
+							const uint8* NextData = InData + 10;
+							int32 TextEncoding = -1;
+							FString MimeType(GetString(NextData, FrameEnd - NextData, &NextData, &TextEncoding));
+							if (NextData < FrameEnd)
+							{
+								uint8 PictureType = *NextData++;
+								FString Description(GetString(NextData, FrameEnd - NextData, &NextData, &TextEncoding));
+								int32 ImageSize = FrameEnd - NextData;
+								TArray<uint8> PicData(NextData, ImageSize);
+								FItem Item {.MimeType = MoveTemp(MimeType), .Value = FVariant(MoveTemp(PicData)), .ItemType = PictureType };
+								Tags.Emplace(FrameID, MoveTemp(Item));
+							}
+							break;
+						}
+						// Recognized text fields
+						case Utils::Make4CC('T','A','L','B'):
+						case Utils::Make4CC('T','C','O','M'):
+						case Utils::Make4CC('T','C','O','N'):
+						case Utils::Make4CC('T','C','O','P'):
+						case Utils::Make4CC('T','D','A','T'):
+						case Utils::Make4CC('T','E','N','C'):
+						case Utils::Make4CC('T','E','X','T'):
+						case Utils::Make4CC('T','I','M','E'):
+						case Utils::Make4CC('T','I','T','1'):
+						case Utils::Make4CC('T','I','T','2'):
+						case Utils::Make4CC('T','I','T','3'):
+						case Utils::Make4CC('T','L','A','N'):
+						case Utils::Make4CC('T','P','E','1'):
+						case Utils::Make4CC('T','P','E','2'):
+						case Utils::Make4CC('T','P','E','3'):
+						case Utils::Make4CC('T','P','E','4'):
+						case Utils::Make4CC('T','P','O','S'):
+						case Utils::Make4CC('T','P','U','B'):
+						case Utils::Make4CC('T','R','C','K'):
+						case Utils::Make4CC('T','Y','E','R'):
+						{
+							FString s(GetString(InData + 10, FrameSize - 10));
+							if (s.Len())
+							{
+								FItem Item {.Value = FVariant(MoveTemp(s)) };
+								Tags.Emplace(FrameID, MoveTemp(Item));
+							}
+							break;
+						}
+					}
+					InData += FrameSize;
+				}
+				else
+				{
+					return InData[0] == 0;
+				}
+			}
+			return true;
+		}
+		bool FID3V2Metadata::HaveTag(uint32 InTag)
+		{
+			const FItem* v = Tags.Find(InTag);
+			return !!v;
+		}
+		bool FID3V2Metadata::GetTag(FItem& OutValue, uint32 InTag)
+		{
+			const FItem* v = Tags.Find(InTag);
+			if (v)
+			{
+				OutValue = *v;
+			}
+			return !!v;
+		}
+		const TMap<uint32, FID3V2Metadata::FItem>& FID3V2Metadata::GetTags() const
+		{
+			return Tags;
+		}
+
+		TMap<uint32, FID3V2Metadata::FItem>& FID3V2Metadata::GetTags()
+		{
+			return Tags;
+		}
+
 	} // namespace MPEG
 } // namespace Electra
 
diff --git a/Engine/Plugins/Media/ElectraPlayer/Source/ElectraPlayerRuntime/Private/Runtime/Utilities/UtilsMPEG.h b/Engine/Plugins/Media/ElectraPlayer/Source/ElectraPlayerRuntime/Private/Runtime/Utilities/UtilsMPEG.h
--- a/Engine/Plugins/Media/ElectraPlayer/Source/ElectraPlayerRuntime/Private/Runtime/Utilities/UtilsMPEG.h
+++ b/Engine/Plugins/Media/ElectraPlayer/Source/ElectraPlayerRuntime/Private/Runtime/Utilities/UtilsMPEG.h
@@ -3,6 +3,7 @@
 #pragma once
 
 #include <CoreMinimal.h>
+#include "Misc/Variant.h"
 
 namespace Electra
 {
@@ -79,5 +80,26 @@ namespace Electra
 		};
 
 
+
+		class FID3V2Metadata
+		{
+		public:
+			struct FItem
+			{
+				FString Language;				// ISO 639-2; if not set (all zero) the default entry for all languages
+				FString MimeType;
+				FVariant Value;
+				int32 ItemType = -1;
+			};
+
+			bool Parse(const uint8* InData, int64 InDataSize);
+			bool HaveTag(uint32 InTag);
+			bool GetTag(FItem& OutValue, uint32 InTag);
+			const TMap<uint32, FItem>& GetTags() const;
+			TMap<uint32, FItem>& GetTags();
+		private:
+			TMap<uint32, FItem> Tags;
+		};
+
 	} // namespace MPEG
 } // namespace Electra
