ElectraPlayer: Releasing the media sample lock before sending out media events since this may otherwise result in a race condition deadlock.

#rnx
#tests Modified source code to provoke issue as an easy repro case to see that fix works.

diff --git a/Engine/Plugins/Media/ElectraPlayer/Source/ElectraPlayerPlugin/Private/ElectraPlayerPlugin.cpp b/Engine/Plugins/Media/ElectraPlayer/Source/ElectraPlayerPlugin/Private/ElectraPlayerPlugin.cpp
--- a/Engine/Plugins/Media/ElectraPlayer/Source/ElectraPlayerPlugin/Private/ElectraPlayerPlugin.cpp
+++ b/Engine/Plugins/Media/ElectraPlayer/Source/ElectraPlayerPlugin/Private/ElectraPlayerPlugin.cpp
@@ -1039,6 +1039,7 @@ bool FElectraPlayerPlugin::Seek(const FTimespan& Time)
 				SequenceIndexMapperSubtitle.SetRemapPrimaryIndex(NextSequenceIndex, TOptional<int32>());
 				SequenceIndexMapperMetadata.SetRemapPrimaryIndex(NextSequenceIndex, TOptional<int32>());
 				CurrentSequenceIndex = NextSequenceIndex;
+				SampleLock.Unlock();
 				SendMediaEvent(EPlayerEvent::SeekCompleted);
 				return true;
 			}
@@ -1050,6 +1051,7 @@ bool FElectraPlayerPlugin::Seek(const FTimespan& Time)
 	SequenceIndexMapperSubtitle.SetExpectedPrimaryIndex(NextSequenceIndex);
 	SequenceIndexMapperMetadata.SetExpectedPrimaryIndex(NextSequenceIndex);
 	CurrentSequenceIndex = NextSequenceIndex;
+	SampleLock.Unlock();
 
 	IElectraPlayerInterface::FSeekParam sp;
 	sp.SequenceIndex = CurrentSequenceIndex;
