ElectraPlayer: Allowing for more than one audio representation with the exact same settings as this can be found in some MPDs for unknown reasons.

#rnx

diff --git a/Engine/Plugins/Media/ElectraPlayer/Source/ElectraPlayerRuntime/Private/Runtime/Player/ABRRules/ABRLiveStream.cpp b/Engine/Plugins/Media/ElectraPlayer/Source/ElectraPlayerRuntime/Private/Runtime/Player/ABRRules/ABRLiveStream.cpp
--- a/Engine/Plugins/Media/ElectraPlayer/Source/ElectraPlayerRuntime/Private/Runtime/Player/ABRRules/ABRLiveStream.cpp
+++ b/Engine/Plugins/Media/ElectraPlayer/Source/ElectraPlayerRuntime/Private/Runtime/Player/ABRRules/ABRLiveStream.cpp
@@ -1321,7 +1321,33 @@ IAdaptiveStreamSelector::ESegmentAction FABRLiveStream::PerformSelection(const T
 		}
 		else
 		{
-			CurrentStreamQualityIndex = NewQualityIndex = InCandidates.Num() - 1;
+			if (WorkVars)
+			{
+				FScopeLock lock(&WorkVars->Lock);
+
+				TArray<FDecisionAttributes> QualityDecision;
+
+				CurrentStreamQualityIndex = WorkVars->SegmentDownloadHistory.Num() ? WorkVars->SegmentDownloadHistory.BackRef().QualityIndex : -1;
+
+				for(auto &Can : InCandidates)
+				{
+					FDecisionAttributes da;
+					da.QualityIndex = Can->QualityIndex;
+					da.Bitrate = Can->Bitrate;
+					da.bIsCandidate = true;
+					QualityDecision.Emplace(MoveTemp(da));
+				}
+
+				for(int32 i=QualityDecision.Num()-1; i>=0; --i)
+				{
+					if (QualityDecision[i].bIsCandidate || i==0)
+					{
+						BandwidthScore = QualityDecision[i].BandwidthScore;
+						NewQualityIndex = QualityDecision[i].QualityIndex;
+						break;
+					}
+				}
+			}
 		}
 
 		check(NewQualityIndex >= 0);
diff --git a/Engine/Plugins/Media/ElectraPlayer/Source/ElectraPlayerRuntime/Private/Runtime/Player/ABRRules/ABRVoDPlus.cpp b/Engine/Plugins/Media/ElectraPlayer/Source/ElectraPlayerRuntime/Private/Runtime/Player/ABRRules/ABRVoDPlus.cpp
--- a/Engine/Plugins/Media/ElectraPlayer/Source/ElectraPlayerRuntime/Private/Runtime/Player/ABRRules/ABRVoDPlus.cpp
+++ b/Engine/Plugins/Media/ElectraPlayer/Source/ElectraPlayerRuntime/Private/Runtime/Player/ABRRules/ABRVoDPlus.cpp
@@ -947,7 +947,21 @@ IAdaptiveStreamSelector::ESegmentAction FABROnDemandPlus::PerformSelection(const
 		}
 		else
 		{
-			NewQualityIndex = InCandidates[InCandidates.Num() - 1]->QualityIndex;
+//			FScopeLock lock(&WorkVars->Lock);
+			struct FDecisionAttributes
+			{
+				int32 Bitrate = 0;
+				int32 QualityIndex = 0;
+			};
+			TArray<FDecisionAttributes> QualityLevels;
+			for(auto &Can : InCandidates)
+			{
+				FDecisionAttributes& da = QualityLevels.Emplace_GetRef();
+				da.QualityIndex = Can->QualityIndex;
+				da.Bitrate = Can->Bitrate;
+			}
+			QualityLevels.StableSort([](const FDecisionAttributes& a, const FDecisionAttributes& b){ return a.Bitrate > b.Bitrate; });
+			NewQualityIndex = QualityLevels[0].QualityIndex;
 		}
 
 		check(NewQualityIndex >= 0);
diff --git a/Engine/Plugins/Media/ElectraPlayer/Source/ElectraPlayerRuntime/Private/Runtime/Player/AdaptiveStreamingPlayerABR.cpp b/Engine/Plugins/Media/ElectraPlayer/Source/ElectraPlayerRuntime/Private/Runtime/Player/AdaptiveStreamingPlayerABR.cpp
--- a/Engine/Plugins/Media/ElectraPlayer/Source/ElectraPlayerRuntime/Private/Runtime/Player/AdaptiveStreamingPlayerABR.cpp
+++ b/Engine/Plugins/Media/ElectraPlayer/Source/ElectraPlayerRuntime/Private/Runtime/Player/AdaptiveStreamingPlayerABR.cpp
@@ -187,7 +187,7 @@ namespace Electra
 	//-----------------------------------------------------------------------------
 	/**
 	 * Sets the presentation format (HLS, DASH, mp4, etc...)
-	 * 
+	 *
 	 * This must be the first call that will determine the type of ABR algorithm to use.
 	 *
 	 * @param InFormatType
@@ -262,11 +262,11 @@ namespace Electra
 									StreamInformationVideo.Push(si);
 								}
 								// Sort the representations by ascending bitrate
-								StreamInformationVideo.Sort([](const TSharedPtrTS<FABRStreamInformation>& a, const TSharedPtrTS<FABRStreamInformation>& b)
+								StreamInformationVideo.StableSort([](const TSharedPtrTS<FABRStreamInformation>& a, const TSharedPtrTS<FABRStreamInformation>& b)
 								{
 									return a->Bitrate < b->Bitrate;
 								});
-									
+
 								break;
 							}
 						}
@@ -310,7 +310,7 @@ namespace Electra
 									StreamInformationAudio.Push(si);
 								}
 								// Sort the representations by ascending bitrate
-								StreamInformationAudio.Sort([](const TSharedPtrTS<FABRStreamInformation>& a, const TSharedPtrTS<FABRStreamInformation>& b)
+								StreamInformationAudio.StableSort([](const TSharedPtrTS<FABRStreamInformation>& a, const TSharedPtrTS<FABRStreamInformation>& b)
 								{
 									return a->Bitrate < b->Bitrate;
 								});
@@ -486,7 +486,7 @@ namespace Electra
 			ABRMethod->ReportPlaybackResumed();
 		}
 	}
-	
+
 	void FAdaptiveStreamSelector::ReportPlaybackEnded()
 	{
 		if (ABRMethod.IsValid())
@@ -564,10 +564,10 @@ namespace Electra
 	IAdaptiveStreamSelector::ESegmentAction FAdaptiveStreamSelector::SelectSuitableStreams(FTimeValue& OutDelay, TSharedPtrTS<const IStreamSegment> CurrentSegment)
 	{
 		FTimeValue TimeNow = PlayerSessionServices->GetSynchronizedUTCTime()->GetTime();
-		
+
 		// No delay to fetch this segment with.
 		OutDelay.SetToZero();
-		
+
 		// Check if streams that were temporarily marked as bad can now be used again.
 		AccessMutex.Lock();
 		for(int32 i=0, iMax=StreamInformationVideo.Num(); i<iMax; ++i)
diff --git a/Engine/Plugins/Media/ElectraPlayer/Source/ElectraPlayerRuntime/Private/Runtime/Player/DASH/ManifestBuilderDASH.cpp b/Engine/Plugins/Media/ElectraPlayer/Source/ElectraPlayerRuntime/Private/Runtime/Player/DASH/ManifestBuilderDASH.cpp
--- a/Engine/Plugins/Media/ElectraPlayer/Source/ElectraPlayerRuntime/Private/Runtime/Player/DASH/ManifestBuilderDASH.cpp
+++ b/Engine/Plugins/Media/ElectraPlayer/Source/ElectraPlayerRuntime/Private/Runtime/Player/DASH/ManifestBuilderDASH.cpp
@@ -1462,7 +1462,13 @@ void FManifestDASHInternal::PreparePeriodAdaptationSets(TSharedPtrTS<FPeriod> Pe
 			// We do not rely on the AdaptationSet@contentType, AdaptationSet@mimeType and/or AdaptationSet@codecs as these are often not set.
 			// First we go over the representations in the set and hope they have the @codecs attribute set.
 			const TArray<TSharedPtrTS<FDashMPD_RepresentationType>>& MPDRepresentations = MPDAdaptationSet->GetRepresentations();
-			TMultiMap<int32, TSharedPtrTS<FRepresentation>> RepresentationQualityIndexMap;
+			struct FQualityEntry
+			{
+				int32 Bitrate = 0;
+				int32 ListedIndex = 0;
+				TSharedPtrTS<FRepresentation> Rep;
+			};
+			TArray<FQualityEntry> RepresentationQualityIndexList;
 			for(int32 nRepr=0, nReprMax=MPDRepresentations.Num(); nRepr<nReprMax; ++nRepr)
 			{
 				const TSharedPtrTS<FDashMPD_RepresentationType>& MPDRepresentation = MPDRepresentations[nRepr];
@@ -1881,8 +1887,11 @@ void FManifestDASHInternal::PreparePeriodAdaptationSets(TSharedPtrTS<FPeriod> Pe
 				// For all intents and purposes we consider this Representation as usable now.
 				Representation->bIsUsable = true;
 				AdaptationSet->Representations.Emplace(Representation);
-				// Add this representation to the bandwidth-to-index map.
-				RepresentationQualityIndexMap.Add(MPDRepresentation->GetBandwidth(), Representation);
+				// Add this representation to the bandwidth-to-index list.
+				FQualityEntry& qe = RepresentationQualityIndexList.Emplace_GetRef();
+				qe.Bitrate = MPDRepresentation->GetBandwidth();
+				qe.ListedIndex = nRepr;
+				qe.Rep = Representation;
 			}
 
 
@@ -1899,18 +1908,23 @@ void FManifestDASHInternal::PreparePeriodAdaptationSets(TSharedPtrTS<FPeriod> Pe
 					}
 				}
 
-				RepresentationQualityIndexMap.KeySort([](int32 A, int32 B){return A<B;});
+				// Sort by ascending bitrate, keeping the given order for same rates.
+				RepresentationQualityIndexList.StableSort([](const FQualityEntry& a, const FQualityEntry& b)
+				{
+					return a.Bitrate == b.Bitrate ? a.ListedIndex < b.ListedIndex : a.Bitrate < b.Bitrate;
+				});
 				int32 CurrentQualityIndex = -1;
 				int32 CurrentQualityBitrate = -1;
-				for(auto& E : RepresentationQualityIndexMap)
+				for(auto& qlIt : RepresentationQualityIndexList)
 				{
-					if (E.Key != CurrentQualityBitrate)
+					if (CurrentQualityBitrate != qlIt.Bitrate)
 					{
-						CurrentQualityBitrate = E.Key;
+						CurrentQualityBitrate = qlIt.Bitrate;
 						++CurrentQualityIndex;
 					}
-					E.Value->QualityIndex = CurrentQualityIndex;
+					qlIt.Rep->QualityIndex = CurrentQualityIndex;
 				}
+
 				AdaptationSet->bIsUsable = true;
 
 				TMediaOptionalValue<bool> lowLatencyUsable;
