ElectraHTTP: removed receive buffer indirection

#rnx

diff --git a/Engine/Plugins/Media/ElectraPlayer/Source/ElectraPlayerRuntime/Private/ElectraPlayer.cpp b/Engine/Plugins/Media/ElectraPlayer/Source/ElectraPlayerRuntime/Private/ElectraPlayer.cpp
--- a/Engine/Plugins/Media/ElectraPlayer/Source/ElectraPlayerRuntime/Private/ElectraPlayer.cpp
+++ b/Engine/Plugins/Media/ElectraPlayer/Source/ElectraPlayerRuntime/Private/ElectraPlayer.cpp
@@ -736,10 +736,10 @@ void FElectraPlayer::HandleBlobDownload()
 																	ErrCode > 0 && ErrCode < 100 ? IElectraPlayerAdapterDelegate::EBlobResultType::TimedOut :
 																	IElectraPlayerAdapterDelegate::EBlobResultType::HttpFailure;
 			TSharedPtr<TArray<uint8>, ESPMode::ThreadSafe> BlobData = MakeShared<TArray<uint8>, ESPMode::ThreadSafe>();
-			TSharedPtrTS<IElectraHttpManager::FReceiveBuffer> ResponseBuffer = PendingBlobRequest->Request->GetResponseBuffer();
+			TSharedPtrTS<FWaitableBuffer> ResponseBuffer = PendingBlobRequest->Request->GetResponseBuffer();
 			if (ResponseBuffer.IsValid())
 			{
-				BlobData->Append((const uint8*)ResponseBuffer->Buffer.GetLinearReadData(), ResponseBuffer->Buffer.Num());
+				BlobData->Append((const uint8*)ResponseBuffer->GetLinearReadData(), ResponseBuffer->Num());
 			}
 			TSharedPtr<IElectraPlayerAdapterDelegate, ESPMode::ThreadSafe> PinnedAdapterDelegate = AdapterDelegate.Pin();
 			if (PinnedAdapterDelegate.IsValid())
diff --git a/Engine/Plugins/Media/ElectraPlayer/Source/ElectraPlayerRuntime/Private/Runtime/Demuxer/ParserISO14496-12_Utils.h b/Engine/Plugins/Media/ElectraPlayer/Source/ElectraPlayerRuntime/Private/Runtime/Demuxer/ParserISO14496-12_Utils.h
--- a/Engine/Plugins/Media/ElectraPlayer/Source/ElectraPlayerRuntime/Private/Runtime/Demuxer/ParserISO14496-12_Utils.h
+++ b/Engine/Plugins/Media/ElectraPlayer/Source/ElectraPlayerRuntime/Private/Runtime/Demuxer/ParserISO14496-12_Utils.h
@@ -15,11 +15,11 @@ namespace Electra
 	public:
 		FMP4StaticDataReader() = default;
 		virtual ~FMP4StaticDataReader() = default;
-		virtual void SetParseData(TSharedPtrTS<IElectraHttpManager::FReceiveBuffer> InResponseBuffer)
+		virtual void SetParseData(TSharedPtrTS<FWaitableBuffer> InResponseBuffer)
 		{
 			ResponseBuffer = InResponseBuffer;
-			DataSize = ResponseBuffer->Buffer.Num();
-			Data = (const uint8*)ResponseBuffer->Buffer.GetLinearReadData();
+			DataSize = ResponseBuffer->Num();
+			Data = (const uint8*)ResponseBuffer->GetLinearReadData();
 			CurrentOffset = 0;
 		}
 		virtual bool HaveParseData() const
@@ -45,18 +45,18 @@ namespace Electra
 		}
 		virtual bool HasReachedEOF() const override
 		{
-			return ResponseBuffer.IsValid() ? ResponseBuffer->Buffer.GetEOD() && CurrentOffset >= DataSize : true;
+			return ResponseBuffer.IsValid() ? ResponseBuffer->GetEOD() && CurrentOffset >= DataSize : true;
 		}
 		virtual bool HasReadBeenAborted() const override
 		{
-			return ResponseBuffer.IsValid() ? ResponseBuffer->Buffer.WasAborted() : true;
+			return ResponseBuffer.IsValid() ? ResponseBuffer->WasAborted() : true;
 		}
 		virtual int64 GetCurrentOffset() const override
 		{
 			return CurrentOffset;
 		}
 
-		TSharedPtrTS<IElectraHttpManager::FReceiveBuffer> ResponseBuffer;
+		TSharedPtrTS<FWaitableBuffer> ResponseBuffer;
 		const uint8* Data = nullptr;
 		int64 DataSize = 0;
 		int64 CurrentOffset = 0;
diff --git a/Engine/Plugins/Media/ElectraPlayer/Source/ElectraPlayerRuntime/Private/Runtime/Demuxer/ParserMKV_Utils.h b/Engine/Plugins/Media/ElectraPlayer/Source/ElectraPlayerRuntime/Private/Runtime/Demuxer/ParserMKV_Utils.h
--- a/Engine/Plugins/Media/ElectraPlayer/Source/ElectraPlayerRuntime/Private/Runtime/Demuxer/ParserMKV_Utils.h
+++ b/Engine/Plugins/Media/ElectraPlayer/Source/ElectraPlayerRuntime/Private/Runtime/Demuxer/ParserMKV_Utils.h
@@ -16,11 +16,11 @@ namespace Electra
 	public:
 		FMKVStaticDataReader() = default;
 		virtual ~FMKVStaticDataReader() = default;
-		virtual void SetParseData(TSharedPtrTS<IElectraHttpManager::FReceiveBuffer> InResponseBuffer)
+		virtual void SetParseData(TSharedPtrTS<FWaitableBuffer> InResponseBuffer)
 		{
 			ResponseBuffer = InResponseBuffer;
-			DataSize = ResponseBuffer->Buffer.Num();
-			Data = (const uint8*)ResponseBuffer->Buffer.GetLinearReadData();
+			DataSize = ResponseBuffer->Num();
+			Data = (const uint8*)ResponseBuffer->GetLinearReadData();
 			CurrentOffset = 0;
 		}
 	private:
@@ -49,7 +49,7 @@ namespace Electra
 		bool MKVHasReadBeenAborted() const override
 		{ return false; }
 
-		TSharedPtrTS<IElectraHttpManager::FReceiveBuffer> ResponseBuffer;
+		TSharedPtrTS<FWaitableBuffer> ResponseBuffer;
 		const uint8* Data = nullptr;
 		int64 DataSize = 0;
 		int64 CurrentOffset = 0;
diff --git a/Engine/Plugins/Media/ElectraPlayer/Source/ElectraPlayerRuntime/Private/Runtime/HTTP/HTTPManager.cpp b/Engine/Plugins/Media/ElectraPlayer/Source/ElectraPlayerRuntime/Private/Runtime/HTTP/HTTPManager.cpp
--- a/Engine/Plugins/Media/ElectraPlayer/Source/ElectraPlayerRuntime/Private/Runtime/HTTP/HTTPManager.cpp
+++ b/Engine/Plugins/Media/ElectraPlayer/Source/ElectraPlayerRuntime/Private/Runtime/HTTP/HTTPManager.cpp
@@ -76,7 +76,7 @@ namespace Electra
 		{
 			virtual ~FLocalByteStream() = default;
 			virtual void SetConnected(TSharedPtrTS<FRequest> Request) = 0;
-			virtual int32 Read(TSharedPtrTS<FReceiveBuffer> RcvBuffer, TSharedPtrTS<FRequest> Request) = 0;
+			virtual int32 Read(TSharedPtrTS<FWaitableBuffer> RcvBuffer, TSharedPtrTS<FRequest> Request) = 0;
 
 			bool										bIsConnected = false;
 			int64										FileStartOffset = 0;		//!< The base offset into the file data is requested at.
@@ -88,7 +88,7 @@ namespace Electra
 		{
 			virtual ~FFileStream() = default;
 			void SetConnected(TSharedPtrTS<FRequest> Request) override;
-			int32 Read(TSharedPtrTS<FReceiveBuffer> RcvBuffer, TSharedPtrTS<FRequest> Request) override;
+			int32 Read(TSharedPtrTS<FWaitableBuffer> RcvBuffer, TSharedPtrTS<FRequest> Request) override;
 			TSharedPtr<FArchive, ESPMode::ThreadSafe>	Archive;
 			FString										Filename;
 		};
@@ -98,7 +98,7 @@ namespace Electra
 			virtual ~FDataUrl() = default;
 			bool SetData(const FString& InUrl);
 			void SetConnected(TSharedPtrTS<FRequest> Request) override;
-			int32 Read(TSharedPtrTS<FReceiveBuffer> RcvBuffer, TSharedPtrTS<FRequest> Request) override;
+			int32 Read(TSharedPtrTS<FWaitableBuffer> RcvBuffer, TSharedPtrTS<FRequest> Request) override;
 			TArray<uint8>								Data;
 			FString										MimeType;
 		};
@@ -107,7 +107,7 @@ namespace Electra
 		{
 			virtual ~FExternalReader() = default;
 			void SetConnected(TSharedPtrTS<FRequest> Request) override;
-			int32 Read(TSharedPtrTS<FReceiveBuffer> RcvBuffer, TSharedPtrTS<FRequest> Request) override;
+			int32 Read(TSharedPtrTS<FWaitableBuffer> RcvBuffer, TSharedPtrTS<FRequest> Request) override;
 			void OnReadComplete(IExternalDataReader::FResponseDataPtr InResponseData, int64 InTotalFileSize, const IExternalDataReader::FReadParams& InFromRequestParams)
 			{
 				ResponseData = MoveTemp(InResponseData);
@@ -499,10 +499,10 @@ namespace Electra
 				SetMessage(FString::Printf(TEXT("FElectraHttpManager did not initialize")));
 			Request->ConnectionInfo.RequestEndTime = MEDIAutcTime::Current();
 			Request->ConnectionInfo.bHasFinished = true;
-			TSharedPtrTS<FReceiveBuffer> ReceiveBuffer = Request->ReceiveBuffer.Pin();
+			TSharedPtrTS<FWaitableBuffer> ReceiveBuffer = Request->ReceiveBuffer.Pin();
 			if (ReceiveBuffer.IsValid())
 			{
-				ReceiveBuffer->Buffer.SetEOD();
+				ReceiveBuffer->SetEOD();
 			}
 			TSharedPtrTS<FProgressListener> ProgressListener = Request->ProgressListener.Pin();
 			if (ProgressListener.IsValid())
@@ -884,10 +884,10 @@ namespace Electra
 			}
 
 			Request->ConnectionInfo.bHasFinished = true;
-			TSharedPtrTS<FReceiveBuffer>	ReceiveBuffer = Request->ReceiveBuffer.Pin();
+			TSharedPtrTS<FWaitableBuffer>	ReceiveBuffer = Request->ReceiveBuffer.Pin();
 			if (ReceiveBuffer.IsValid())
 			{
-				ReceiveBuffer->Buffer.SetEOD();
+				ReceiveBuffer->SetEOD();
 			}
 
 			// Call completion delegate.
@@ -1111,7 +1111,7 @@ namespace Electra
 				// Establish the file handle as "connected"
 				Handle->LocalByteStream->SetConnected(Request);
 				// Read from local file
-				TSharedPtrTS<FReceiveBuffer> ReceiveBuffer = Request->ReceiveBuffer.Pin();
+				TSharedPtrTS<FWaitableBuffer> ReceiveBuffer = Request->ReceiveBuffer.Pin();
 				if (ReceiveBuffer.IsValid())
 				{
 					int32 NumBytesRead = Handle->LocalByteStream->Read(ReceiveBuffer, Request);
@@ -1148,7 +1148,7 @@ namespace Electra
 				// Establish the file handle as "connected"
 				Handle->LocalByteStream->SetConnected(Request);
 				// Read data
-				TSharedPtrTS<FReceiveBuffer> ReceiveBuffer = Request->ReceiveBuffer.Pin();
+				TSharedPtrTS<FWaitableBuffer> ReceiveBuffer = Request->ReceiveBuffer.Pin();
 				if (ReceiveBuffer.IsValid())
 				{
 					int32 NumBytesRead = Handle->LocalByteStream->Read(ReceiveBuffer, Request);
@@ -1442,16 +1442,16 @@ namespace Electra
 						Handle->LastTimeDataReceived = Now;
 
 						// Receive buffer still there?
-						TSharedPtrTS<FReceiveBuffer> ReceiveBuffer = Request->ReceiveBuffer.Pin();
+						TSharedPtrTS<FWaitableBuffer> ReceiveBuffer = Request->ReceiveBuffer.Pin();
 						if (ReceiveBuffer.IsValid())
 						{
 							int64 RequiredBufferSize = ci.ContentLength > 0 ? ci.ContentLength : 0;
-							int64 BufferSizeAfterPush = ReceiveBuffer->Buffer.Num() + NumDataAvailable;
+							int64 BufferSizeAfterPush = ReceiveBuffer->Num() + NumDataAvailable;
 							if (BufferSizeAfterPush > RequiredBufferSize)
 							{
 								RequiredBufferSize = BufferSizeAfterPush;
 							}
-							bool bBufferUsable = ReceiveBuffer->Buffer.EnlargeTo(RequiredBufferSize);
+							bool bBufferUsable = ReceiveBuffer->EnlargeTo(RequiredBufferSize);
 							int64 BufferPushableSize = bBufferUsable ? NumDataAvailable : 0;
 							if (bBufferUsable)
 							{
@@ -1461,7 +1461,7 @@ namespace Electra
 									int64 NewDataSize = 0;
 									Response->GetResponseData().LockBuffer(NewDataPtr, NewDataSize);
 									int64 NumToCopy = BufferPushableSize < NewDataSize ? BufferPushableSize : NewDataSize;
-									bBufferUsable = ReceiveBuffer->Buffer.PushData(NewDataPtr, NumToCopy);
+									bBufferUsable = ReceiveBuffer->PushData(NewDataPtr, NumToCopy);
 									Response->GetResponseData().UnlockBuffer(bBufferUsable ? NumToCopy : 0);
 									if (bBufferUsable)
 									{
@@ -1521,10 +1521,10 @@ namespace Electra
 							if (Handle->ActiveResponse.SizeRemaining() == 0)
 							{
 								// All done now.
-								TSharedPtrTS<FReceiveBuffer> ReceiveBuffer = Request->ReceiveBuffer.Pin();
+								TSharedPtrTS<FWaitableBuffer> ReceiveBuffer = Request->ReceiveBuffer.Pin();
 								if (ReceiveBuffer.IsValid())
 								{
-									ReceiveBuffer->Buffer.SetEOD();
+									ReceiveBuffer->SetEOD();
 								}
 							}
 							else
@@ -1658,16 +1658,16 @@ namespace Electra
 		}
 	}
 
-	int32 FElectraHttpManager::FFileStream::Read(TSharedPtrTS<FReceiveBuffer> RcvBuffer, TSharedPtrTS<FRequest> Request)
+	int32 FElectraHttpManager::FFileStream::Read(TSharedPtrTS<FWaitableBuffer> RcvBuffer, TSharedPtrTS<FRequest> Request)
 	{
 		int64 NumToRead = FileSizeToGo;
 		if (NumToRead)
 		{
-			void* Dst = (void*) RcvBuffer->Buffer.GetLinearWriteData(NumToRead);
+			void* Dst = (void*) RcvBuffer->GetLinearWriteData(NumToRead);
 			if (Dst)
 			{
 				Archive->Serialize(Dst, NumToRead);
-				RcvBuffer->Buffer.AppendedNewData(NumToRead);
+				RcvBuffer->AppendedNewData(NumToRead);
 				Request->ConnectionInfo.BytesReadSoFar += NumToRead;
 				FileSizeToGo -= NumToRead;
 			}
@@ -1716,7 +1716,7 @@ namespace Electra
 		}
 	}
 
-	int32 FElectraHttpManager::FExternalReader::Read(TSharedPtrTS<FReceiveBuffer> RcvBuffer, TSharedPtrTS<FRequest> Request)
+	int32 FElectraHttpManager::FExternalReader::Read(TSharedPtrTS<FWaitableBuffer> RcvBuffer, TSharedPtrTS<FRequest> Request)
 	{
 		if (bCompleted)
 		{
@@ -1725,7 +1725,7 @@ namespace Electra
 			FileSizeToGo = 0;
 			if (FileSize >= 0)
 			{
-				RcvBuffer->Buffer.SetExternalData(MoveTemp(ResponseData));
+				RcvBuffer->SetExternalData(MoveTemp(ResponseData));
 				if (!bRangedRequest)
 				{
 					Request->ConnectionInfo.StatusInfo.HTTPStatus = 200;
@@ -1847,11 +1847,11 @@ namespace Electra
 		}
 	}
 
-	int32 FElectraHttpManager::FDataUrl::Read(TSharedPtrTS<FReceiveBuffer> RcvBuffer, TSharedPtrTS<FRequest> Request)
+	int32 FElectraHttpManager::FDataUrl::Read(TSharedPtrTS<FWaitableBuffer> RcvBuffer, TSharedPtrTS<FRequest> Request)
 	{
-		if (RcvBuffer->Buffer.EnlargeTo(FileSizeToGo))
+		if (RcvBuffer->EnlargeTo(FileSizeToGo))
 		{
-			if (RcvBuffer->Buffer.PushData(Data.GetData() + FileStartOffset, FileSizeToGo))
+			if (RcvBuffer->PushData(Data.GetData() + FileStartOffset, FileSizeToGo))
 			{
 				Request->ConnectionInfo.BytesReadSoFar += FileSizeToGo;
 			}
diff --git a/Engine/Plugins/Media/ElectraPlayer/Source/ElectraPlayerRuntime/Private/Runtime/HTTP/HTTPManager.h b/Engine/Plugins/Media/ElectraPlayer/Source/ElectraPlayerRuntime/Private/Runtime/HTTP/HTTPManager.h
--- a/Engine/Plugins/Media/ElectraPlayer/Source/ElectraPlayerRuntime/Private/Runtime/HTTP/HTTPManager.h
+++ b/Engine/Plugins/Media/ElectraPlayer/Source/ElectraPlayerRuntime/Private/Runtime/HTTP/HTTPManager.h
@@ -222,11 +222,6 @@ namespace Electra
 			FCompletionDelegate		CompletionDelegate;
 		};
 
-		struct FReceiveBuffer
-		{
-			FWaitableBuffer		Buffer;
-		};
-
 		struct FParams
 		{
 			void AddFromHeaderList(const TArray<FString>& InHeaderList)
@@ -261,7 +256,7 @@ namespace Electra
 			FParams								Parameters;
 			FParamDict							Options;
 			HTTP::FConnectionInfo				ConnectionInfo;
-			TWeakPtrTS<FReceiveBuffer>			ReceiveBuffer;
+			TWeakPtrTS<FWaitableBuffer>			ReceiveBuffer;
 			TWeakPtrTS<FProgressListener>		ProgressListener;
 			TSharedPtrTS<IHTTPResponseCache>	ResponseCache;
 			TWeakPtrTS<IExternalDataReader>		ExternalDataReader;
diff --git a/Engine/Plugins/Media/ElectraPlayer/Source/ElectraPlayerRuntime/Private/Runtime/Player/AdaptiveStreamingPlayerResourceRequest.cpp b/Engine/Plugins/Media/ElectraPlayer/Source/ElectraPlayerRuntime/Private/Runtime/Player/AdaptiveStreamingPlayerResourceRequest.cpp
--- a/Engine/Plugins/Media/ElectraPlayer/Source/ElectraPlayerRuntime/Private/Runtime/Player/AdaptiveStreamingPlayerResourceRequest.cpp
+++ b/Engine/Plugins/Media/ElectraPlayer/Source/ElectraPlayerRuntime/Private/Runtime/Player/AdaptiveStreamingPlayerResourceRequest.cpp
@@ -15,7 +15,7 @@ namespace Electra
 FHTTPResourceRequest::FHTTPResourceRequest()
 {
 	Request = MakeSharedTS<IElectraHttpManager::FRequest>();
-	ReceiveBuffer = MakeSharedTS<IElectraHttpManager::FReceiveBuffer>();
+	ReceiveBuffer = MakeSharedTS<FWaitableBuffer>();
 	ProgressListener = MakeSharedTS<IElectraHttpManager::FProgressListener>();
 	ProgressListener->ProgressDelegate   = IElectraHttpManager::FProgressListener::FProgressDelegate::CreateRaw(this, &FHTTPResourceRequest::HTTPProgressCallback);
 	ProgressListener->CompletionDelegate = IElectraHttpManager::FProgressListener::FCompletionDelegate::CreateRaw(this, &FHTTPResourceRequest::HTTPCompletionCallback);
@@ -180,7 +180,7 @@ void FHTTPResourceRequest::StaticDataReady()
 			ci.bHasFinished = true;
 			ci.HTTPVersionReceived = 11;
 			ci.StatusInfo.HTTPStatus = Request->Parameters.Range.IsSet() ? 206 : 200;
-			ci.ContentLength = ci.BytesReadSoFar = ReceiveBuffer.IsValid() ? ReceiveBuffer->Buffer.Num() : 0;
+			ci.ContentLength = ci.BytesReadSoFar = ReceiveBuffer.IsValid() ? ReceiveBuffer->Num() : 0;
 
 			bInCallback = true;
 			CompletedCallback.ExecuteIfBound(AsShared());
diff --git a/Engine/Plugins/Media/ElectraPlayer/Source/ElectraPlayerRuntime/Private/Runtime/Player/AdaptiveStreamingPlayerResourceRequest.h b/Engine/Plugins/Media/ElectraPlayer/Source/ElectraPlayerRuntime/Private/Runtime/Player/AdaptiveStreamingPlayerResourceRequest.h
--- a/Engine/Plugins/Media/ElectraPlayer/Source/ElectraPlayerRuntime/Private/Runtime/Player/AdaptiveStreamingPlayerResourceRequest.h
+++ b/Engine/Plugins/Media/ElectraPlayer/Source/ElectraPlayerRuntime/Private/Runtime/Player/AdaptiveStreamingPlayerResourceRequest.h
@@ -150,7 +150,7 @@ public:
 	virtual TSharedPtrTS<IHTTPResourceRequestObject> GetObject() const
 	{ return UserObject.Pin(); }
 
-	virtual TSharedPtrTS<IElectraHttpManager::FReceiveBuffer> GetResponseBuffer() const
+	virtual TSharedPtrTS<FWaitableBuffer> GetResponseBuffer() const
 	{ return ReceiveBuffer; }
 
 	virtual const HTTP::FConnectionInfo* GetConnectionInfo() const
@@ -192,12 +192,12 @@ private:
 				TSharedPtrTS<FHTTPResourceRequest> p(Owner.Pin());
 				if (p.IsValid())
 				{
-					TSharedPtrTS<IElectraHttpManager::FReceiveBuffer> Buf = p->GetResponseBuffer();
+					TSharedPtrTS<FWaitableBuffer> Buf = p->GetResponseBuffer();
 					if (Buf.IsValid())
 					{
-						Buf->Buffer.Reserve(PlaybackData->Num());
-						Buf->Buffer.PushData(PlaybackData->GetData(), PlaybackData->Num());
-						Buf->Buffer.SetEOD();
+						Buf->Reserve(PlaybackData->Num());
+						Buf->PushData(PlaybackData->GetData(), PlaybackData->Num());
+						Buf->SetEOD();
 					}
 					p->SetStaticDataReady();
 				}
@@ -224,7 +224,7 @@ private:
 	{ bStaticDataReady = true; }
 	void StaticDataReady();
 	TSharedPtrTS<IElectraHttpManager::FRequest> Request;
-	TSharedPtrTS<IElectraHttpManager::FReceiveBuffer> ReceiveBuffer;
+	TSharedPtrTS<FWaitableBuffer> ReceiveBuffer;
 	TSharedPtrTS<IElectraHttpManager::FProgressListener> ProgressListener;
 	FOnRequestCompletedCallback CompletedCallback;
 	TWeakPtrTS<IHTTPResourceRequestObject> UserObject;
diff --git a/Engine/Plugins/Media/ElectraPlayer/Source/ElectraPlayerRuntime/Private/Runtime/Player/DASH/ManifestDASH.cpp b/Engine/Plugins/Media/ElectraPlayer/Source/ElectraPlayerRuntime/Private/Runtime/Player/DASH/ManifestDASH.cpp
--- a/Engine/Plugins/Media/ElectraPlayer/Source/ElectraPlayerRuntime/Private/Runtime/Player/DASH/ManifestDASH.cpp
+++ b/Engine/Plugins/Media/ElectraPlayer/Source/ElectraPlayerRuntime/Private/Runtime/Player/DASH/ManifestDASH.cpp
@@ -102,8 +102,8 @@ namespace DashUtils
 			int64 nr = -1;
 			if (LoadChunk(InNumBytesToRead, InFromOffset) && ResponseBuffer.IsValid())
 			{
-				nr = ResponseBuffer->Buffer.GetLinearReadSize();
-				FMemory::Memcpy(InDestinationBuffer, ResponseBuffer->Buffer.GetLinearReadData(), nr);
+				nr = ResponseBuffer->GetLinearReadSize();
+				FMemory::Memcpy(InDestinationBuffer, ResponseBuffer->GetLinearReadData(), nr);
 			}
 			return nr;
 		}
@@ -160,7 +160,7 @@ namespace DashUtils
 
 
 		TSharedPtrTS<FMPDLoadRequestDASH> LoadRequest;
-		TSharedPtrTS<IElectraHttpManager::FReceiveBuffer> ResponseBuffer;
+		TSharedPtrTS<FWaitableBuffer> ResponseBuffer;
 		int64 FileSize = -1;
 	};
 }
diff --git a/Engine/Plugins/Media/ElectraPlayer/Source/ElectraPlayerRuntime/Private/Runtime/Player/DASH/PlaylistReaderDASH.cpp b/Engine/Plugins/Media/ElectraPlayer/Source/ElectraPlayerRuntime/Private/Runtime/Player/DASH/PlaylistReaderDASH.cpp
--- a/Engine/Plugins/Media/ElectraPlayer/Source/ElectraPlayerRuntime/Private/Runtime/Player/DASH/PlaylistReaderDASH.cpp
+++ b/Engine/Plugins/Media/ElectraPlayer/Source/ElectraPlayerRuntime/Private/Runtime/Player/DASH/PlaylistReaderDASH.cpp
@@ -1080,9 +1080,9 @@ FErrorDetail FPlaylistReaderDASH::GetXMLResponseString(FString& OutXMLString, FR
 {
 	if (FromRequest.IsValid() && FromRequest->Request.IsValid())
 	{
-		TSharedPtrTS<IElectraHttpManager::FReceiveBuffer> ResponseBuffer = FromRequest->Request->GetResponseBuffer();
-		int32 NumResponseBytes = ResponseBuffer->Buffer.Num();
-		const uint8* ResponseBytes = (const uint8*)ResponseBuffer->Buffer.GetLinearReadData();
+		TSharedPtrTS<FWaitableBuffer> ResponseBuffer = FromRequest->Request->GetResponseBuffer();
+		int32 NumResponseBytes = ResponseBuffer->Num();
+		const uint8* ResponseBytes = (const uint8*)ResponseBuffer->GetLinearReadData();
 		// Check for potential BOMs
 		if (NumResponseBytes > 3 && ResponseBytes[0] == 0xEF && ResponseBytes[1] == 0xBB && ResponseBytes[2] == 0xBF)
 		{
@@ -1122,9 +1122,9 @@ FErrorDetail FPlaylistReaderDASH::GetResponseString(FString& OutString, FResourc
 {
 	if (FromRequest.IsValid() && FromRequest->Request.IsValid())
 	{
-		TSharedPtrTS<IElectraHttpManager::FReceiveBuffer> ResponseBuffer = FromRequest->Request->GetResponseBuffer();
-		int32 NumResponseBytes = ResponseBuffer->Buffer.Num();
-		const uint8* ResponseBytes = (const uint8*)ResponseBuffer->Buffer.GetLinearReadData();
+		TSharedPtrTS<FWaitableBuffer> ResponseBuffer = FromRequest->Request->GetResponseBuffer();
+		int32 NumResponseBytes = ResponseBuffer->Num();
+		const uint8* ResponseBytes = (const uint8*)ResponseBuffer->GetLinearReadData();
 		FUTF8ToTCHAR TextConv((const ANSICHAR*)ResponseBytes, NumResponseBytes);
 		FString UTF8Text(TextConv.Length(), TextConv.Get());
 		OutString = MoveTemp(UTF8Text);
diff --git a/Engine/Plugins/Media/ElectraPlayer/Source/ElectraPlayerRuntime/Private/Runtime/Player/DASH/StreamReaderDASH.cpp b/Engine/Plugins/Media/ElectraPlayer/Source/ElectraPlayerRuntime/Private/Runtime/Player/DASH/StreamReaderDASH.cpp
--- a/Engine/Plugins/Media/ElectraPlayer/Source/ElectraPlayerRuntime/Private/Runtime/Player/DASH/StreamReaderDASH.cpp
+++ b/Engine/Plugins/Media/ElectraPlayer/Source/ElectraPlayerRuntime/Private/Runtime/Player/DASH/StreamReaderDASH.cpp
@@ -625,8 +625,8 @@ FErrorDetail FStreamReaderDASH::FStreamHandler::GetInitSegment(TSharedPtrTS<cons
 		else
 		{
 			// Pass an empty buffer so parsing will fail below.
-			TSharedPtrTS<IElectraHttpManager::FReceiveBuffer> EmptyBuffer = MakeSharedTS<IElectraHttpManager::FReceiveBuffer>();
-			EmptyBuffer->Buffer.SetEOD();
+			TSharedPtrTS<FWaitableBuffer> EmptyBuffer = MakeSharedTS<FWaitableBuffer>();
+			EmptyBuffer->SetEOD();
 			StaticDataReader.SetParseData(EmptyBuffer);
 		}
 	}
@@ -873,7 +873,7 @@ FErrorDetail FStreamReaderDASH::FStreamHandler::RetrieveSideloadedFile(TSharedPt
 		return CreateError(FString::Printf(TEXT("Sideloaded media download error: %s"), *Request->ConnectionInfo.StatusInfo.ErrorDetail.GetMessage()), INTERNAL_ERROR_INIT_SEGMENT_DOWNLOAD_ERROR);
 	}
 
-	OutData = MakeSharedTS<const TArray<uint8>>(TArrayView<uint8>(LoadReq->Request->GetResponseBuffer()->Buffer.GetLinearReadData(), LoadReq->Request->GetResponseBuffer()->Buffer.GetLinearReadSize()));
+	OutData = MakeSharedTS<const TArray<uint8>>(TArrayView<uint8>(LoadReq->Request->GetResponseBuffer()->GetLinearReadData(), LoadReq->Request->GetResponseBuffer()->GetLinearReadSize()));
 
 	// Add this to the entity cache in case it needs to be retrieved again.
 	IPlayerEntityCache::FCacheItem CacheItem;
@@ -1175,7 +1175,7 @@ void FStreamReaderDASH::FStreamHandler::HandleRequestMP4()
 		else if (!bIsEmptyFillerSegment)
 		{
 			ReadBuffer.Reset();
-			ReadBuffer.ReceiveBuffer = MakeSharedTS<IElectraHttpManager::FReceiveBuffer>();
+			ReadBuffer.ReceiveBuffer = MakeSharedTS<FWaitableBuffer>();
 
 			// Start downloading the segment.
 			TSharedPtrTS<IElectraHttpManager::FProgressListener>	ProgressListener(new IElectraHttpManager::FProgressListener);
@@ -2017,7 +2017,7 @@ void FStreamReaderDASH::FStreamHandler::HandleRequestMKV()
 		if (!bIsEmptyFillerSegment)
 		{
 			ReadBuffer.Reset();
-			ReadBuffer.ReceiveBuffer = MakeSharedTS<IElectraHttpManager::FReceiveBuffer>();
+			ReadBuffer.ReceiveBuffer = MakeSharedTS<FWaitableBuffer>();
 
 			// Start downloading the segment.
 			TSharedPtrTS<IElectraHttpManager::FProgressListener> ProgressListener(new IElectraHttpManager::FProgressListener);
@@ -2683,7 +2683,7 @@ bool FStreamReaderDASH::FStreamHandler::HasErrored() const
  */
 int64 FStreamReaderDASH::FStreamHandler::ReadData(void* IntoBuffer, int64 NumBytesToRead)
 {
-	FWaitableBuffer& SourceBuffer = ReadBuffer.ReceiveBuffer->Buffer;
+	FWaitableBuffer& SourceBuffer = *ReadBuffer.ReceiveBuffer;
 	// Make sure the buffer will have the amount of data we need.
 	while(1)
 	{
@@ -2769,7 +2769,7 @@ int64 FStreamReaderDASH::FStreamHandler::ReadData(void* IntoBuffer, int64 NumByt
  */
 bool FStreamReaderDASH::FStreamHandler::HasReachedEOF() const
 {
-	const FWaitableBuffer& SourceBuffer = ReadBuffer.ReceiveBuffer->Buffer;
+	const FWaitableBuffer& SourceBuffer = *ReadBuffer.ReceiveBuffer;
 	return !HasErrored() && SourceBuffer.GetEOD() && (ReadBuffer.ParsePos >= SourceBuffer.Num() || ReadBuffer.ParsePos >= ReadBuffer.MaxParsePos);
 }
 
diff --git a/Engine/Plugins/Media/ElectraPlayer/Source/ElectraPlayerRuntime/Private/Runtime/Player/DASH/StreamReaderDASH.h b/Engine/Plugins/Media/ElectraPlayer/Source/ElectraPlayerRuntime/Private/Runtime/Player/DASH/StreamReaderDASH.h
--- a/Engine/Plugins/Media/ElectraPlayer/Source/ElectraPlayerRuntime/Private/Runtime/Player/DASH/StreamReaderDASH.h
+++ b/Engine/Plugins/Media/ElectraPlayer/Source/ElectraPlayerRuntime/Private/Runtime/Player/DASH/StreamReaderDASH.h
@@ -120,9 +120,9 @@ private:
 				ParsePos = 0;
 				MaxParsePos = TNumericLimits<int64>::Max();
 			}
-			TSharedPtrTS<IElectraHttpManager::FReceiveBuffer>	ReceiveBuffer;
-			int64												ParsePos;
-			int64												MaxParsePos;
+			TSharedPtrTS<FWaitableBuffer>	ReceiveBuffer;
+			int64							ParsePos;
+			int64							MaxParsePos;
 		};
 
 
diff --git a/Engine/Plugins/Media/ElectraPlayer/Source/ElectraPlayerRuntime/Private/Runtime/Player/DRM/DRMManager.cpp b/Engine/Plugins/Media/ElectraPlayer/Source/ElectraPlayerRuntime/Private/Runtime/Player/DRM/DRMManager.cpp
--- a/Engine/Plugins/Media/ElectraPlayer/Source/ElectraPlayerRuntime/Private/Runtime/Player/DRM/DRMManager.cpp
+++ b/Engine/Plugins/Media/ElectraPlayer/Source/ElectraPlayerRuntime/Private/Runtime/Player/DRM/DRMManager.cpp
@@ -111,10 +111,10 @@ void FDRMManager::OnHTTPResourceRequestComplete(TSharedPtrTS<FHTTPResourceReques
 		{
 			TArray<uint8> LicenseResponse;
 
-			TSharedPtrTS<IElectraHttpManager::FReceiveBuffer> ResponseBuffer = DrmReq->Request->GetResponseBuffer();
+			TSharedPtrTS<FWaitableBuffer> ResponseBuffer = DrmReq->Request->GetResponseBuffer();
 			if (ResponseBuffer.IsValid())
 			{
-				LicenseResponse.Append((const uint8*)ResponseBuffer->Buffer.GetLinearReadData(), ResponseBuffer->Buffer.Num());
+				LicenseResponse.Append((const uint8*)ResponseBuffer->GetLinearReadData(), ResponseBuffer->Num());
 			}
 			const HTTP::FConnectionInfo* ConnInfo = DrmReq->Request->GetConnectionInfo();
 			int32 HttpResponseCode = ConnInfo ? ConnInfo->StatusInfo.HTTPStatus : 500;
diff --git a/Engine/Plugins/Media/ElectraPlayer/Source/ElectraPlayerRuntime/Private/Runtime/Player/HLS/PlaylistReaderHLS.cpp b/Engine/Plugins/Media/ElectraPlayer/Source/ElectraPlayerRuntime/Private/Runtime/Player/HLS/PlaylistReaderHLS.cpp
--- a/Engine/Plugins/Media/ElectraPlayer/Source/ElectraPlayerRuntime/Private/Runtime/Player/HLS/PlaylistReaderHLS.cpp
+++ b/Engine/Plugins/Media/ElectraPlayer/Source/ElectraPlayerRuntime/Private/Runtime/Player/HLS/PlaylistReaderHLS.cpp
@@ -118,7 +118,7 @@ private:
 			HTTPRequest.Reset();
 		}
 
-		TSharedPtrTS<IElectraHttpManager::FReceiveBuffer> GetReceiveBuffer() const
+		TSharedPtrTS<FWaitableBuffer> GetReceiveBuffer() const
 		{
 			return ReceiveBuffer;
 		}
@@ -221,7 +221,7 @@ private:
 
 		FPlaylistLoadRequestHLS									PlaylistLoadRequest;
 		FTimeValue												ExecuteAtUTC;
-		TSharedPtrTS<IElectraHttpManager::FReceiveBuffer>		ReceiveBuffer;
+		TSharedPtrTS<FWaitableBuffer>							ReceiveBuffer;
 		TSharedPtrTS<IElectraHttpManager::FRequest>				HTTPRequest;
 		TSharedPtrTS<HTTP::FRetryInfo>							RetryInfo;
 		bool													bIsMasterPlaylist;
@@ -632,8 +632,8 @@ void FPlaylistReaderHLS::FPlaylistRequest::PrepareStaticResult()
 	StaticRequestConnectionInfo.HTTPVersionReceived   = 11;
 	StaticRequestConnectionInfo.StatusInfo.HTTPStatus = 200;
 
-	ReceiveBuffer = MakeSharedTS<IElectraHttpManager::FReceiveBuffer>();
-	ReceiveBuffer->Buffer.SetEOD();
+	ReceiveBuffer = MakeSharedTS<FWaitableBuffer>();
+	ReceiveBuffer->SetEOD();
 
 	TSharedPtr<TArray<uint8>, ESPMode::ThreadSafe> DataPtr = StaticRequest->GetData();
 	if (DataPtr.IsValid())
@@ -641,8 +641,8 @@ void FPlaylistReaderHLS::FPlaylistRequest::PrepareStaticResult()
 		TArray<uint8>& DataArray = *DataPtr;
 		if (DataArray.Num())
 		{
-			ReceiveBuffer->Buffer.EnlargeTo(DataArray.Num());
-			ReceiveBuffer->Buffer.PushData(DataArray.GetData(), DataArray.Num());
+			ReceiveBuffer->EnlargeTo(DataArray.Num());
+			ReceiveBuffer->PushData(DataArray.GetData(), DataArray.Num());
 			StaticRequestConnectionInfo.ContentLength = StaticRequestConnectionInfo.BytesReadSoFar = DataArray.Num();
 		}
 	}
@@ -652,7 +652,7 @@ void FPlaylistReaderHLS::FPlaylistRequest::PrepareStaticResult()
 
 void FPlaylistReaderHLS::FPlaylistRequest::Execute(TSharedPtrTS<IElectraHttpManager::FProgressListener> InProgressListener, IPlayerSessionServices* InPlayerSessionServices)
 {
-	ReceiveBuffer = MakeSharedTS<IElectraHttpManager::FReceiveBuffer>();
+	ReceiveBuffer = MakeSharedTS<FWaitableBuffer>();
 	HTTPRequest = MakeSharedTS<IElectraHttpManager::FRequest>();
 	HTTPRequest->Parameters.URL = PlaylistLoadRequest.URL;
 	HTTPRequest->Parameters.AcceptEncoding.Set("");		// setting an empty string enables all supported encodings (ie. gzip)
@@ -1026,8 +1026,8 @@ void FPlaylistReaderHLS::HandleCompletedPlaylistDownloads(FTimeValue& TimeNow)
 FErrorDetail FPlaylistReaderHLS::ParsePlaylist(FPlaylistRequestPtr FromRequest)
 {
 	// Calculate checksum over the response to see if the playlist has changed
-	check(FromRequest->GetReceiveBuffer()->Buffer.GetLinearReadSize() >= FromRequest->GetReceiveBuffer()->Buffer.Num());	// The buffer must be linear!
-	uint32 crc = HashFunctions::CRC::Calc32(FromRequest->GetReceiveBuffer()->Buffer.GetLinearReadData(), FromRequest->GetReceiveBuffer()->Buffer.Num());
+	check(FromRequest->GetReceiveBuffer()->GetLinearReadSize() >= FromRequest->GetReceiveBuffer()->Num());	// The buffer must be linear!
+	uint32 crc = HashFunctions::CRC::Calc32(FromRequest->GetReceiveBuffer()->GetLinearReadData(), FromRequest->GetReceiveBuffer()->Num());
 	bool bNoChange = FromRequest->GetPlaylistLoadRequest().LoadType == FPlaylistLoadRequestHLS::ELoadType::Update && FromRequest->GetPlaylistLoadRequest().LastUpdateCRC32 == crc;
 
 	HLSPlaylistParser::FPlaylist		Playlist;
@@ -1042,9 +1042,9 @@ FErrorDetail FPlaylistReaderHLS::ParsePlaylist(FPlaylistRequestPtr FromRequest)
 	}
 	else
 	{
-		int32 RequestBytes = FromRequest->GetReceiveBuffer()->Buffer.Num();
+		int32 RequestBytes = FromRequest->GetReceiveBuffer()->Num();
 		// The FromRequest buffer is not zero terminated. We need to pick the correct FString constructor for converting the chars into TCHARs while adding the terminating zero!
-		FUTF8ToTCHAR TextConv((const ANSICHAR*)FromRequest->GetReceiveBuffer()->Buffer.GetLinearReadData(), RequestBytes);
+		FUTF8ToTCHAR TextConv((const ANSICHAR*)FromRequest->GetReceiveBuffer()->GetLinearReadData(), RequestBytes);
 		FString UTF8String(TextConv.Length(), TextConv.Get());
 		ParseError = Parser.Parse(UTF8String, Playlist);
 	}
diff --git a/Engine/Plugins/Media/ElectraPlayer/Source/ElectraPlayerRuntime/Private/Runtime/Player/HLS/StreamReaderHLSfmp4.cpp b/Engine/Plugins/Media/ElectraPlayer/Source/ElectraPlayerRuntime/Private/Runtime/Player/HLS/StreamReaderHLSfmp4.cpp
--- a/Engine/Plugins/Media/ElectraPlayer/Source/ElectraPlayerRuntime/Private/Runtime/Player/HLS/StreamReaderHLSfmp4.cpp
+++ b/Engine/Plugins/Media/ElectraPlayer/Source/ElectraPlayerRuntime/Private/Runtime/Player/HLS/StreamReaderHLSfmp4.cpp
@@ -413,7 +413,7 @@ FStreamReaderHLSfmp4::FStreamHandler::ELicenseKeyResult FStreamReaderHLSfmp4::FS
 		{
 			TSharedPtrTS<IElectraHttpManager::FRequest> HTTP(new IElectraHttpManager::FRequest);
 			ReadBuffer.Reset();
-			ReadBuffer.ReceiveBuffer = MakeSharedTS<IElectraHttpManager::FReceiveBuffer>();
+			ReadBuffer.ReceiveBuffer = MakeSharedTS<FWaitableBuffer>();
 
 			Metrics::FSegmentDownloadStats& ds = InRequest->DownloadStats;
 			ds.URL = LicenseKeyInfo->URI;
@@ -439,9 +439,9 @@ FStreamReaderHLSfmp4::FStreamHandler::ELicenseKeyResult FStreamReaderHLSfmp4::FS
 						if (KeyData.IsValid())
 						{
 							// Copy the response over into the receive buffer as if it was received through the http request.
-							ReadBuffer.ReceiveBuffer->Buffer.Reserve(KeyData->Num());
-							ReadBuffer.ReceiveBuffer->Buffer.PushData(KeyData->GetData(), KeyData->Num());
-							ReadBuffer.ReceiveBuffer->Buffer.SetEOD();
+							ReadBuffer.ReceiveBuffer->Reserve(KeyData->Num());
+							ReadBuffer.ReceiveBuffer->PushData(KeyData->GetData(), KeyData->Num());
+							ReadBuffer.ReceiveBuffer->SetEOD();
 							bHaveStaticResponse = true;
 						}
 						break;
@@ -477,11 +477,11 @@ FStreamReaderHLSfmp4::FStreamHandler::ELicenseKeyResult FStreamReaderHLSfmp4::FS
 				PlayerSessionService->SendMessageToPlayer(FLicenseKeyMessage::Create(FLicenseKeyMessage::EReason::LicenseKeyDownload, FErrorDetail(), &HTTP->ConnectionInfo));
 
 				// There is not much we can validate here. The key is the direct key data without any wrapping.
-				if (ReadBuffer.ReceiveBuffer->Buffer.Num() == 16)
+				if (ReadBuffer.ReceiveBuffer->Num() == 16)
 				{
 					LicenseKeyData = MakeShared<TArray<uint8>, ESPMode::ThreadSafe>();
-					LicenseKeyData->AddUninitialized(ReadBuffer.ReceiveBuffer->Buffer.Num());
-					FMemory::Memcpy(LicenseKeyData->GetData(), ReadBuffer.ReceiveBuffer->Buffer.GetLinearReadData(), ReadBuffer.ReceiveBuffer->Buffer.GetLinearReadSize());
+					LicenseKeyData->AddUninitialized(ReadBuffer.ReceiveBuffer->Num());
+					FMemory::Memcpy(LicenseKeyData->GetData(), ReadBuffer.ReceiveBuffer->GetLinearReadData(), ReadBuffer.ReceiveBuffer->GetLinearReadSize());
 					if (InRequest->LicenseKeyCache.IsValid())
 					{
 						InRequest->LicenseKeyCache->AddLicenseKey(LicenseKeyData, LicenseKeyInfo, FTimeValue::GetPositiveInfinity());
@@ -541,7 +541,7 @@ FStreamReaderHLSfmp4::FStreamHandler::EInitSegmentResult FStreamReaderHLSfmp4::F
 			ProgressListener->ProgressDelegate = IElectraHttpManager::FProgressListener::FProgressDelegate::CreateRaw(this, &FStreamHandler::HTTPProgressCallback);
 			ProgressListener->CompletionDelegate = IElectraHttpManager::FProgressListener::FCompletionDelegate::CreateRaw(this, &FStreamHandler::HTTPCompletionCallback);
 			ReadBuffer.Reset();
-			ReadBuffer.ReceiveBuffer = MakeSharedTS<IElectraHttpManager::FReceiveBuffer>();
+			ReadBuffer.ReceiveBuffer = MakeSharedTS<FWaitableBuffer>();
 
 			FString RequestURL = Request->InitSegmentInfo->URI;
 
@@ -819,7 +819,7 @@ void FStreamReaderHLSfmp4::FStreamHandler::HandleRequest()
 			if (!bHasErrored)
 			{
 				ReadBuffer.Reset();
-				ReadBuffer.ReceiveBuffer = MakeSharedTS<IElectraHttpManager::FReceiveBuffer>();
+				ReadBuffer.ReceiveBuffer = MakeSharedTS<FWaitableBuffer>();
 
 				// Start downloading the segment. Clear any stats that may have been set by the init segment download.
 				FString RequestURL = Request->URL;
@@ -1356,7 +1356,7 @@ bool FStreamReaderHLSfmp4::FStreamHandler::HasErrored() const
  */
 int64 FStreamReaderHLSfmp4::FStreamHandler::ReadData(void* IntoBuffer, int64 NumBytesToRead)
 {
-	FWaitableBuffer& SourceBuffer = ReadBuffer.ReceiveBuffer->Buffer;
+	FWaitableBuffer& SourceBuffer = *ReadBuffer.ReceiveBuffer;
 	// Make sure the buffer will have the amount of data we need.
 	while(1)
 	{
@@ -1570,7 +1570,7 @@ int64 FStreamReaderHLSfmp4::FStreamHandler::ReadData(void* IntoBuffer, int64 Num
  */
 bool FStreamReaderHLSfmp4::FStreamHandler::HasReachedEOF() const
 {
-	const FWaitableBuffer& SourceBuffer = ReadBuffer.ReceiveBuffer->Buffer;
+	const FWaitableBuffer& SourceBuffer = *ReadBuffer.ReceiveBuffer;
 	return !HasErrored() && SourceBuffer.GetEOD() && (ReadBuffer.ParsePos >= SourceBuffer.Num() || ReadBuffer.ParsePos >= ReadBuffer.MaxParsePos);
 }
 
diff --git a/Engine/Plugins/Media/ElectraPlayer/Source/ElectraPlayerRuntime/Private/Runtime/Player/HLS/StreamReaderHLSfmp4.h b/Engine/Plugins/Media/ElectraPlayer/Source/ElectraPlayerRuntime/Private/Runtime/Player/HLS/StreamReaderHLSfmp4.h
--- a/Engine/Plugins/Media/ElectraPlayer/Source/ElectraPlayerRuntime/Private/Runtime/Player/HLS/StreamReaderHLSfmp4.h
+++ b/Engine/Plugins/Media/ElectraPlayer/Source/ElectraPlayerRuntime/Private/Runtime/Player/HLS/StreamReaderHLSfmp4.h
@@ -134,10 +134,10 @@ private:
 				MaxParsePos = TNumericLimits<int64>::Max();
 				DecryptedPos = 0;
 			}
-			TSharedPtrTS<IElectraHttpManager::FReceiveBuffer>	ReceiveBuffer;
-			int64												ParsePos;
-			int64												MaxParsePos;
-			int32												DecryptedPos;
+			TSharedPtrTS<FWaitableBuffer>	ReceiveBuffer;
+			int64							ParsePos;
+			int64							MaxParsePos;
+			int32							DecryptedPos;
 		};
 
 
diff --git a/Engine/Plugins/Media/ElectraPlayer/Source/ElectraPlayerRuntime/Private/Runtime/Player/SimpleElectraAudioPlayer.cpp b/Engine/Plugins/Media/ElectraPlayer/Source/ElectraPlayerRuntime/Private/Runtime/Player/SimpleElectraAudioPlayer.cpp
--- a/Engine/Plugins/Media/ElectraPlayer/Source/ElectraPlayerRuntime/Private/Runtime/Player/SimpleElectraAudioPlayer.cpp
+++ b/Engine/Plugins/Media/ElectraPlayer/Source/ElectraPlayerRuntime/Private/Runtime/Player/SimpleElectraAudioPlayer.cpp
@@ -516,10 +516,10 @@ private:
 			{
 				ErrorCode = InRequest->GetError();
 				BlobData = MakeShared<TArray<uint8>, ESPMode::ThreadSafe>();
-				TSharedPtrTS<IElectraHttpManager::FReceiveBuffer> ResponseBuffer = InRequest->GetResponseBuffer();
+				TSharedPtrTS<FWaitableBuffer> ResponseBuffer = InRequest->GetResponseBuffer();
 				if (ResponseBuffer.IsValid())
 				{
-					BlobData->Append((const uint8*)ResponseBuffer->Buffer.GetLinearReadData(), ResponseBuffer->Buffer.Num());
+					BlobData->Append((const uint8*)ResponseBuffer->GetLinearReadData(), ResponseBuffer->Num());
 				}
 			}
 			bIsComplete = true;
diff --git a/Engine/Plugins/Media/ElectraPlayer/Source/ElectraPlayerRuntime/Private/Runtime/Player/mkv/PlaylistReaderMKV.cpp b/Engine/Plugins/Media/ElectraPlayer/Source/ElectraPlayerRuntime/Private/Runtime/Player/mkv/PlaylistReaderMKV.cpp
--- a/Engine/Plugins/Media/ElectraPlayer/Source/ElectraPlayerRuntime/Private/Runtime/Player/mkv/PlaylistReaderMKV.cpp
+++ b/Engine/Plugins/Media/ElectraPlayer/Source/ElectraPlayerRuntime/Private/Runtime/Player/mkv/PlaylistReaderMKV.cpp
@@ -114,7 +114,7 @@ private:
 
 	FCriticalSection Lock;
 	TSharedPtrTS<IElectraHttpManager::FRequest> Request;
-	TSharedPtrTS<IElectraHttpManager::FReceiveBuffer> ReceiveBuffer;
+	TSharedPtrTS<FWaitableBuffer> ReceiveBuffer;
 	TSharedPtrTS<IElectraHttpManager::FProgressListener> ProgressListener;
 	HTTP::FConnectionInfo ConnectionInfo;
 	FMediaEvent RequestFinished;
@@ -300,8 +300,8 @@ void FPlaylistReaderMKV::ReadChunk(uint8* DestinationBuffer, int64 InFromOffset,
 	ProgressListener->CompletionDelegate = IElectraHttpManager::FProgressListener::FCompletionDelegate::CreateRaw(this, &FPlaylistReaderMKV::HTTPCompletionCallback);
 	ProgressListener->ProgressDelegate   = IElectraHttpManager::FProgressListener::FProgressDelegate::CreateRaw(this, &FPlaylistReaderMKV::HTTPProgressCallback);
 
-	ReceiveBuffer = MakeSharedTS<IElectraHttpManager::FReceiveBuffer>();
-	ReceiveBuffer->Buffer.SetExternalBuffer(DestinationBuffer, ChunkSize);
+	ReceiveBuffer = MakeSharedTS<FWaitableBuffer>();
+	ReceiveBuffer->SetExternalBuffer(DestinationBuffer, ChunkSize);
 
 	Request = MakeSharedTS<IElectraHttpManager::FRequest>();
 	Request->Parameters.URL = PlaylistURL;
@@ -393,7 +393,7 @@ int64 FPlaylistReaderMKV::MKVReadData(void* IntoBuffer, int64 NumBytesToRead, in
 	ReadChunk(reinterpret_cast<uint8*>(IntoBuffer), InFromOffset, NumBytesToRead);
 	while(1)
 	{
-		TSharedPtrTS<IElectraHttpManager::FReceiveBuffer> rb = ReceiveBuffer;
+		TSharedPtrTS<FWaitableBuffer> rb = ReceiveBuffer;
 		if (bHasErrored || bAbort || !rb.IsValid())
 		{
 			return -1;
diff --git a/Engine/Plugins/Media/ElectraPlayer/Source/ElectraPlayerRuntime/Private/Runtime/Player/mkv/StreamReaderMKV.cpp b/Engine/Plugins/Media/ElectraPlayer/Source/ElectraPlayerRuntime/Private/Runtime/Player/mkv/StreamReaderMKV.cpp
--- a/Engine/Plugins/Media/ElectraPlayer/Source/ElectraPlayerRuntime/Private/Runtime/Player/mkv/StreamReaderMKV.cpp
+++ b/Engine/Plugins/Media/ElectraPlayer/Source/ElectraPlayerRuntime/Private/Runtime/Player/mkv/StreamReaderMKV.cpp
@@ -362,7 +362,7 @@ void FStreamReaderMKV::HandleRequest()
 	ProgressListener->ProgressDelegate = IElectraHttpManager::FProgressListener::FProgressDelegate::CreateRaw(this, &FStreamReaderMKV::HTTPProgressCallback);
 
 	ReadBuffer.Reset();
-	ReadBuffer.ReceiveBuffer = MakeSharedTS<IElectraHttpManager::FReceiveBuffer>();
+	ReadBuffer.ReceiveBuffer = MakeSharedTS<FWaitableBuffer>();
 	ReadBuffer.SetStartOffset(Request->FileStartOffset);
 	ReadBuffer.SetCurrentOffset(0);
 	ReadBuffer.SetEndOffset(Request->FileEndOffset);
@@ -910,7 +910,7 @@ void FStreamReaderMKV::WorkerThread()
 
 int32 FStreamReaderMKV::FReadBuffer::ReadTo(void* IntoBuffer, int64 NumBytesToRead)
 {
-	FWaitableBuffer& SourceBuffer = ReceiveBuffer->Buffer;
+	FWaitableBuffer& SourceBuffer = *ReceiveBuffer;
 	// Make sure the buffer will have the amount of data we need.
 	while(1)
 	{
diff --git a/Engine/Plugins/Media/ElectraPlayer/Source/ElectraPlayerRuntime/Private/Runtime/Player/mkv/StreamReaderMKV.h b/Engine/Plugins/Media/ElectraPlayer/Source/ElectraPlayerRuntime/Private/Runtime/Player/mkv/StreamReaderMKV.h
--- a/Engine/Plugins/Media/ElectraPlayer/Source/ElectraPlayerRuntime/Private/Runtime/Player/mkv/StreamReaderMKV.h
+++ b/Engine/Plugins/Media/ElectraPlayer/Source/ElectraPlayerRuntime/Private/Runtime/Player/mkv/StreamReaderMKV.h
@@ -187,10 +187,10 @@ private:
 		void Abort()
 		{
 			bAbort = true;
-			TSharedPtrTS<IElectraHttpManager::FReceiveBuffer> Buffer = ReceiveBuffer;
+			TSharedPtrTS<FWaitableBuffer> Buffer = ReceiveBuffer;
 			if (Buffer.IsValid())
 			{
-				Buffer->Buffer.Abort();
+				Buffer->Abort();
 			}
 		}
 		void SetHasErrored()
@@ -198,7 +198,7 @@ private:
 			bHasErrored = true;
 		}
 		int32 ReadTo(void* ToBuffer, int64 NumBytes);
-		TSharedPtrTS<IElectraHttpManager::FReceiveBuffer> ReceiveBuffer;
+		TSharedPtrTS<FWaitableBuffer> ReceiveBuffer;
 		int64 StartOffset = 0;
 		int64 EndOffset = 0;
 		int64 CurrentOffset = 0;
diff --git a/Engine/Plugins/Media/ElectraPlayer/Source/ElectraPlayerRuntime/Private/Runtime/Player/mp4/PlaylistReaderMP4.cpp b/Engine/Plugins/Media/ElectraPlayer/Source/ElectraPlayerRuntime/Private/Runtime/Player/mp4/PlaylistReaderMP4.cpp
--- a/Engine/Plugins/Media/ElectraPlayer/Source/ElectraPlayerRuntime/Private/Runtime/Player/mp4/PlaylistReaderMP4.cpp
+++ b/Engine/Plugins/Media/ElectraPlayer/Source/ElectraPlayerRuntime/Private/Runtime/Player/mp4/PlaylistReaderMP4.cpp
@@ -101,7 +101,7 @@ private:
 	HTTP::FConnectionInfo ConnectionInfo;
 
 	TSharedPtrTS<IParserISO14496_12> MP4Parser;
-	TSharedPtrTS<IElectraHttpManager::FReceiveBuffer> ParseBuffer;
+	TSharedPtrTS<FWaitableBuffer> ParseBuffer;
 	int64 ParsePos = 0;
 	int64 ParseBufferSize = 0;
 
@@ -268,7 +268,7 @@ void FPlaylistReaderMP4::WorkerThread()
 
 			ParseBuffer = BoxInfos.FindByPredicate([InType=UtilsMP4::Make4CC('m','o','o','v')](const UtilsMP4::FMP4RootBoxLocator::FBoxInfo& InBox){return InBox.Type == InType;})->DataBuffer;
 			ParsePos = 0;
-			ParseBufferSize = ParseBuffer->Buffer.Num();
+			ParseBufferSize = ParseBuffer->Num();
 
 			MP4Parser = IParserISO14496_12::CreateParser();
 			UEMediaError parseError = MP4Parser->ParseHeader(this, this, PlayerSessionServices, nullptr);
@@ -334,7 +334,7 @@ int64 FPlaylistReaderMP4::ReadData(void* IntoBuffer, int64 NumBytesToRead)
 	{
 		return 0;
 	}
-	const uint8* Src = ParseBuffer->Buffer.GetLinearReadData() + ParsePos;
+	const uint8* Src = ParseBuffer->GetLinearReadData() + ParsePos;
 	if (IntoBuffer)
 	{
 		FMemory::Memcpy(IntoBuffer, Src, NumBytesToRead);
diff --git a/Engine/Plugins/Media/ElectraPlayer/Source/ElectraPlayerRuntime/Private/Runtime/Player/mp4/StreamReaderMP4.cpp b/Engine/Plugins/Media/ElectraPlayer/Source/ElectraPlayerRuntime/Private/Runtime/Player/mp4/StreamReaderMP4.cpp
--- a/Engine/Plugins/Media/ElectraPlayer/Source/ElectraPlayerRuntime/Private/Runtime/Player/mp4/StreamReaderMP4.cpp
+++ b/Engine/Plugins/Media/ElectraPlayer/Source/ElectraPlayerRuntime/Private/Runtime/Player/mp4/StreamReaderMP4.cpp
@@ -369,7 +369,7 @@ void FStreamReaderMP4::HandleRequest()
 	ProgressListener->ProgressDelegate   = IElectraHttpManager::FProgressListener::FProgressDelegate::CreateRaw(this, &FStreamReaderMP4::HTTPProgressCallback);
 
 	ReadBuffer.Reset();
-	ReadBuffer.ReceiveBuffer = MakeSharedTS<IElectraHttpManager::FReceiveBuffer>();
+	ReadBuffer.ReceiveBuffer = MakeSharedTS<FWaitableBuffer>();
 	ReadBuffer.SetCurrentPos(Request->FileStartOffset);
 
 	TSharedPtrTS<IElectraHttpManager::FRequest> HTTP(new IElectraHttpManager::FRequest);
@@ -745,7 +745,7 @@ void FStreamReaderMP4::WorkerThread()
 
 int32 FStreamReaderMP4::FReadBuffer::ReadTo(void* IntoBuffer, int64 NumBytesToRead)
 {
-	FWaitableBuffer& SourceBuffer = ReceiveBuffer->Buffer;
+	FWaitableBuffer& SourceBuffer = *ReceiveBuffer;
 	// Make sure the buffer will have the amount of data we need.
 	while(1)
 	{
diff --git a/Engine/Plugins/Media/ElectraPlayer/Source/ElectraPlayerRuntime/Private/Runtime/Player/mp4/StreamReaderMP4.h b/Engine/Plugins/Media/ElectraPlayer/Source/ElectraPlayerRuntime/Private/Runtime/Player/mp4/StreamReaderMP4.h
--- a/Engine/Plugins/Media/ElectraPlayer/Source/ElectraPlayerRuntime/Private/Runtime/Player/mp4/StreamReaderMP4.h
+++ b/Engine/Plugins/Media/ElectraPlayer/Source/ElectraPlayerRuntime/Private/Runtime/Player/mp4/StreamReaderMP4.h
@@ -138,10 +138,10 @@ private:
 		void Abort()
 		{
 			bAbort = true;
-			TSharedPtrTS<IElectraHttpManager::FReceiveBuffer>	Buffer = ReceiveBuffer;
+			TSharedPtrTS<FWaitableBuffer> Buffer = ReceiveBuffer;
 			if (Buffer.IsValid())
 			{
-				Buffer->Buffer.Abort();
+				Buffer->Abort();
 			}
 		}
 		void SetHasErrored()
@@ -149,7 +149,7 @@ private:
 			bHasErrored = true;
 		}
 		int32 ReadTo(void* ToBuffer, int64 NumBytes);
-		TSharedPtrTS<IElectraHttpManager::FReceiveBuffer>	ReceiveBuffer;
+		TSharedPtrTS<FWaitableBuffer>						ReceiveBuffer;
 		int64												CurrentPos;
 		int64												ParsePos;
 		bool												bAbort;
diff --git a/Engine/Plugins/Media/ElectraPlayer/Source/ElectraPlayerRuntime/Private/Runtime/Player/mpegaudio/StreamReaderMPEGAudio.cpp b/Engine/Plugins/Media/ElectraPlayer/Source/ElectraPlayerRuntime/Private/Runtime/Player/mpegaudio/StreamReaderMPEGAudio.cpp
--- a/Engine/Plugins/Media/ElectraPlayer/Source/ElectraPlayerRuntime/Private/Runtime/Player/mpegaudio/StreamReaderMPEGAudio.cpp
+++ b/Engine/Plugins/Media/ElectraPlayer/Source/ElectraPlayerRuntime/Private/Runtime/Player/mpegaudio/StreamReaderMPEGAudio.cpp
@@ -75,7 +75,7 @@ struct FStreamReaderMPEGAudio::FLiveRequest : public TSharedFromThis<FStreamRead
 	int32 StatusCode = 0;
 	TMediaMessageQueueWithTimeout<EEvent> Events;
 	TMultiMap<FString, FString> Headers;
-	TWeakPtrTS<IElectraHttpManager::FReceiveBuffer> ReceiveBuffer;
+	TWeakPtrTS<FWaitableBuffer> ReceiveBuffer;
 	bool bCanceled = false;
 	int32 MetaDataEveryNBytes = 0;
 	int32 MetaDataBytesToGo = 0;
@@ -166,10 +166,10 @@ void FStreamReaderMPEGAudio::CancelRequest(EStreamType StreamType, bool bSilent)
 void FStreamReaderMPEGAudio::CancelRequests()
 {
 	bRequestCanceled = true;
-	TSharedPtrTS<IElectraHttpManager::FReceiveBuffer> RcvBuf = ReceiveBuffer;
+	TSharedPtrTS<FWaitableBuffer> RcvBuf = ReceiveBuffer;
 	if (RcvBuf.IsValid())
 	{
-		RcvBuf->Buffer.Abort();
+		RcvBuf->Abort();
 	}
 	TSharedPtrTS<FLiveRequest> lr(LiveRequest);
 	if (lr.IsValid())
@@ -180,8 +180,8 @@ void FStreamReaderMPEGAudio::CancelRequests()
 
 bool FStreamReaderMPEGAudio::HasBeenAborted() const
 {
-	TSharedPtrTS<IElectraHttpManager::FReceiveBuffer> RcvBuf = ReceiveBuffer;
-	return bRequestCanceled || (RcvBuf.IsValid() && RcvBuf->Buffer.WasAborted());
+	TSharedPtrTS<FWaitableBuffer> RcvBuf = ReceiveBuffer;
+	return bRequestCanceled || (RcvBuf.IsValid() && RcvBuf->WasAborted());
 }
 
 bool FStreamReaderMPEGAudio::HasErrored() const
@@ -288,17 +288,17 @@ void FStreamReaderMPEGAudio::FLiveRequest::OnProcessRequestStream(void* InDataPt
 			else
 			{
 				int32 DataBytesNow = BytesUntilNextMetadata > 0 ? Utils::Min((int32) InLength, BytesUntilNextMetadata) : InLength;
-				TSharedPtrTS<IElectraHttpManager::FReceiveBuffer> rb = ReceiveBuffer.Pin();
+				TSharedPtrTS<FWaitableBuffer> rb = ReceiveBuffer.Pin();
 				if (rb.IsValid())
 				{
-					FScopeLock Lock(rb->Buffer.GetLock());
-					int64 BufSizeRequired = rb->Buffer.Num() + DataBytesNow;
-					if (!rb->Buffer.EnlargeTo(BufSizeRequired))
+					FScopeLock Lock(rb->GetLock());
+					int64 BufSizeRequired = rb->Num() + DataBytesNow;
+					if (!rb->EnlargeTo(BufSizeRequired))
 					{
 						InLength = 0;
 						return;
 					}
-					if (!rb->Buffer.PushData(reinterpret_cast<const uint8*>(InDataPtr), DataBytesNow))
+					if (!rb->PushData(reinterpret_cast<const uint8*>(InDataPtr), DataBytesNow))
 					{
 						InLength = 0;
 						return;
@@ -306,7 +306,7 @@ void FStreamReaderMPEGAudio::FLiveRequest::OnProcessRequestStream(void* InDataPt
 
 					// Sanity check that we are not reading excessive data, which is the case when the
 					// player has been paused for instance.
-					if (MaxDataBytes && rb->Buffer.Num() > MaxDataBytes)
+					if (MaxDataBytes && rb->Num() > MaxDataBytes)
 					{
 						bHasFailed = true;
 						InLength = 0;
@@ -381,7 +381,7 @@ void FStreamReaderMPEGAudio::HandleRequest()
 
 	TSharedPtrTS<IElectraHttpManager::FRequest> HTTP(new IElectraHttpManager::FRequest);
 	TSharedPtrTS<IElectraHttpManager::FProgressListener> ProgressListener;
-	ReceiveBuffer = MakeSharedTS<IElectraHttpManager::FReceiveBuffer>();
+	ReceiveBuffer = MakeSharedTS<FWaitableBuffer>();
 	if (!bIsLivePlayback)
 	{
 		ProgressListener = MakeSharedTS<IElectraHttpManager::FProgressListener>();
@@ -451,11 +451,11 @@ void FStreamReaderMPEGAudio::HandleRequest()
 	TArray<int32> SyncMarkerOffsets;
 	while(!bDone && !HasErrored() && !HasBeenAborted() && !bTerminate)
 	{
-		if (ReceiveBuffer->Buffer.WaitUntilSizeAvailable(SyncMarkerCheckPos + FrameCheckSize, 1000 * 20))
+		if (ReceiveBuffer->WaitUntilSizeAvailable(SyncMarkerCheckPos + FrameCheckSize, 1000 * 20))
 		{
-			FScopeLock Lock(ReceiveBuffer->Buffer.GetLock());
-			const uint8* BufferBaseData = ReceiveBuffer->Buffer.GetLinearReadData();
-			int64 BufferDataSize = ReceiveBuffer->Buffer.GetLinearReadSize();
+			FScopeLock Lock(ReceiveBuffer->GetLock());
+			const uint8* BufferBaseData = ReceiveBuffer->GetLinearReadData();
+			int64 BufferDataSize = ReceiveBuffer->GetLinearReadSize();
 			if (!BufferBaseData)
 			{
 				bDone = true;
@@ -483,7 +483,7 @@ void FStreamReaderMPEGAudio::HandleRequest()
 				}
 			}
 			SyncMarkerCheckPos = BufferDataSize - 1;
-			if (ReceiveBuffer->Buffer.GetEOD())
+			if (ReceiveBuffer->GetEOD())
 			{
 				bDone = true;
 			}
@@ -494,8 +494,8 @@ void FStreamReaderMPEGAudio::HandleRequest()
 	// Probe that we are properly locked on to the frames.
 	if (!HasErrored() && !HasBeenAborted() && !bTerminate)
 	{
-		FScopeLock Lock(ReceiveBuffer->Buffer.GetLock());
-		const uint8* BufferBaseData = ReceiveBuffer->Buffer.GetLinearReadData();
+		FScopeLock Lock(ReceiveBuffer->GetLock());
+		const uint8* BufferBaseData = ReceiveBuffer->GetLinearReadData();
 		if (BufferBaseData)
 		{
 			const int32 MaxCheckFrames = FMath::Min(10, SyncMarkerOffsets.Num());
@@ -573,22 +573,22 @@ void FStreamReaderMPEGAudio::HandleRequest()
 		// With Live playback we have to remove the data bytes we already passed along from the start of the buffer.
 		if (LiveRequest.IsValid() && NumLiveStreamBytesToRemove && NextAUBufferOffset >= NumLiveStreamBytesToRemove)
 		{
-			ReceiveBuffer->Buffer.RemoveFromBeginning(NextAUBufferOffset);
+			ReceiveBuffer->RemoveFromBeginning(NextAUBufferOffset);
 			NextAUBufferOffset = 0;
 		}
 
 		// Wait until we get the next AU's data plus the following 4 bytes that are the header of the following frame.
 		const int64 TotalNumNeeded = NextAUBufferOffset + NextAUFrameSize + 4;
-		if (ReceiveBuffer->Buffer.WaitUntilSizeAvailable(TotalNumNeeded, 1000 * 20))
+		if (ReceiveBuffer->WaitUntilSizeAvailable(TotalNumNeeded, 1000 * 20))
 		{
-			FScopeLock Lock(ReceiveBuffer->Buffer.GetLock());
+			FScopeLock Lock(ReceiveBuffer->GetLock());
 
-			const uint8* BufferBaseData = ReceiveBuffer->Buffer.GetLinearReadData();
+			const uint8* BufferBaseData = ReceiveBuffer->GetLinearReadData();
 			if (!BufferBaseData)
 			{
 				break;
 			}
-			int64 BufferDataSize = ReceiveBuffer->Buffer.GetLinearReadSize();
+			int64 BufferDataSize = ReceiveBuffer->GetLinearReadSize();
 
 			int64 NumGot = BufferDataSize - NextAUBufferOffset;
 			// Did we get the next frame's worth?
@@ -758,8 +758,8 @@ void FStreamReaderMPEGAudio::HandleRequest()
 				// Did we also get the next 4 bytes?
 				if (NumGot >= 4)
 				{
-					FScopeLock Lock2(ReceiveBuffer->Buffer.GetLock());
-					HeaderValue = GetUINT32BE(ReceiveBuffer->Buffer.GetLinearReadData() + NextAUBufferOffset + NextAUFrameSize);
+					FScopeLock Lock2(ReceiveBuffer->GetLock());
+					HeaderValue = GetUINT32BE(ReceiveBuffer->GetLinearReadData() + NextAUBufferOffset + NextAUFrameSize);
 					NextAUFrameSize = ElectraDecodersUtil::MPEG::UtilsMPEG123::GetFrameSize(HeaderValue);
 				}
 				else
diff --git a/Engine/Plugins/Media/ElectraPlayer/Source/ElectraPlayerRuntime/Private/Runtime/Player/mpegaudio/StreamReaderMPEGAudio.h b/Engine/Plugins/Media/ElectraPlayer/Source/ElectraPlayerRuntime/Private/Runtime/Player/mpegaudio/StreamReaderMPEGAudio.h
--- a/Engine/Plugins/Media/ElectraPlayer/Source/ElectraPlayerRuntime/Private/Runtime/Player/mpegaudio/StreamReaderMPEGAudio.h
+++ b/Engine/Plugins/Media/ElectraPlayer/Source/ElectraPlayerRuntime/Private/Runtime/Player/mpegaudio/StreamReaderMPEGAudio.h
@@ -125,7 +125,7 @@ private:
 	struct FLiveRequest;
 	TSharedPtrTS<FLiveRequest> LiveRequest;
 	TSharedPtrTS<FStreamSegmentRequestMPEGAudio> CurrentRequest;
-	TSharedPtrTS<IElectraHttpManager::FReceiveBuffer> ReceiveBuffer;
+	TSharedPtrTS<FWaitableBuffer> ReceiveBuffer;
 	FMediaEvent WorkSignal;
 
 	static uint32 UniqueDownloadID;
diff --git a/Engine/Plugins/Media/ElectraPlayer/Source/ElectraPlayerRuntime/Private/Runtime/Utilities/UtilsMP4.cpp b/Engine/Plugins/Media/ElectraPlayer/Source/ElectraPlayerRuntime/Private/Runtime/Utilities/UtilsMP4.cpp
--- a/Engine/Plugins/Media/ElectraPlayer/Source/ElectraPlayerRuntime/Private/Runtime/Utilities/UtilsMP4.cpp
+++ b/Engine/Plugins/Media/ElectraPlayer/Source/ElectraPlayerRuntime/Private/Runtime/Utilities/UtilsMP4.cpp
@@ -595,7 +595,7 @@ bool UtilsMP4::FMP4RootBoxLocator::LocateRootBoxes(TArray<FBoxInfo>& OutBoxInfos
 	});
 
 
-	auto CreateReadRequestAndBuffer = [&](TSharedPtrTS<IElectraHttpManager::FReceiveBuffer>& OutReceiveBuffer, int64 InFromOffset, int64 InNumBytes) -> TSharedPtrTS<IElectraHttpManager::FRequest>
+	auto CreateReadRequestAndBuffer = [&](TSharedPtrTS<FWaitableBuffer>& OutReceiveBuffer, int64 InFromOffset, int64 InNumBytes) -> TSharedPtrTS<IElectraHttpManager::FRequest>
 	{
 		TSharedPtrTS<IElectraHttpManager::FRequest> Req = MakeSharedTS<IElectraHttpManager::FRequest>();
 		Req->Parameters.URL = URL;
@@ -608,8 +608,8 @@ bool UtilsMP4::FMP4RootBoxLocator::LocateRootBoxes(TArray<FBoxInfo>& OutBoxInfos
 		Req->Parameters.Range.SetEndIncluding(LastByte);
 		Req->Parameters.ConnectTimeout = FTimeValue().SetFromMilliseconds(1000 * 8);
 		Req->Parameters.NoDataTimeout = FTimeValue().SetFromMilliseconds(1000 * 6);
-		OutReceiveBuffer = MakeSharedTS<IElectraHttpManager::FReceiveBuffer>();
-		OutReceiveBuffer->Buffer.Reserve(InNumBytes);
+		OutReceiveBuffer = MakeSharedTS<FWaitableBuffer>();
+		OutReceiveBuffer->Reserve(InNumBytes);
 		Req->ReceiveBuffer = OutReceiveBuffer;
 		Req->ProgressListener = ProgressListener;
 		Req->ResponseCache = LocalCache;
@@ -623,7 +623,7 @@ bool UtilsMP4::FMP4RootBoxLocator::LocateRootBoxes(TArray<FBoxInfo>& OutBoxInfos
 	bool bSuccess = false;
 	bool bIsFirst = true;
 	TSharedPtrTS<IElectraHttpManager::FRequest> Request;
-	TSharedPtrTS<IElectraHttpManager::FReceiveBuffer> ReceiveBuffer;
+	TSharedPtrTS<FWaitableBuffer> ReceiveBuffer;
 	while(1)
 	{
 		int64 SizeToRead = CurrentEndOffset < StartOffset+MinRequiredReadSize ? ChunkSize : MinRequiredReadSize;
@@ -634,17 +634,17 @@ bool UtilsMP4::FMP4RootBoxLocator::LocateRootBoxes(TArray<FBoxInfo>& OutBoxInfos
 		{
 			break;
 		}
-		TSharedPtrTS<IElectraHttpManager::FReceiveBuffer> DataBuffer(MoveTemp(ReceiveBuffer));
+		TSharedPtrTS<FWaitableBuffer> DataBuffer(MoveTemp(ReceiveBuffer));
 		InHTTPManager->RemoveRequest(Request, false);
 		Request.Reset();
 
-		if (DataBuffer.IsValid() && DataBuffer->Buffer.Num() >= 8)
+		if (DataBuffer.IsValid() && DataBuffer->Num() >= 8)
 		{
-			const int64 End = StartOffset + DataBuffer->Buffer.Num();
+			const int64 End = StartOffset + DataBuffer->Num();
 			CurrentEndOffset = CurrentEndOffset < End ? End : CurrentEndOffset;
 
 			FBoxInfo bi;
-			const uint32* Data = reinterpret_cast<const uint32*>(DataBuffer->Buffer.GetLinearReadData());
+			const uint32* Data = reinterpret_cast<const uint32*>(DataBuffer->GetLinearReadData());
 			bi.Size = (int64) MEDIA_FROM_BIG_ENDIAN(Data[0]);
 			bi.Type = MEDIA_FROM_BIG_ENDIAN(Data[1]);
 			bi.Offset = StartOffset;
@@ -676,7 +676,7 @@ bool UtilsMP4::FMP4RootBoxLocator::LocateRootBoxes(TArray<FBoxInfo>& OutBoxInfos
 			else if (bi.Size == 1)
 			{
 				// A size of 1 indicates that the size is expressed as a 64 bit value following the box type.
-				if (DataBuffer->Buffer.Num() < 16)
+				if (DataBuffer->Num() < 16)
 				{
 					ErrorMsg = TEXT("Invalid mp4 file: Box requiring 64 bit size value is truncated");
 					break;
@@ -687,7 +687,7 @@ bool UtilsMP4::FMP4RootBoxLocator::LocateRootBoxes(TArray<FBoxInfo>& OutBoxInfos
 
 			if (bi.Type == Make4CC('u','u','i','d'))
 			{
-				if (DataBuffer->Buffer.Num() < BoxInternalOffset + 16)
+				if (DataBuffer->Num() < BoxInternalOffset + 16)
 				{
 					ErrorMsg = TEXT("Invalid mp4 file: UUID box is truncated");
 					break;
diff --git a/Engine/Plugins/Media/ElectraPlayer/Source/ElectraPlayerRuntime/Private/Runtime/Utilities/UtilsMP4.h b/Engine/Plugins/Media/ElectraPlayer/Source/ElectraPlayerRuntime/Private/Runtime/Utilities/UtilsMP4.h
--- a/Engine/Plugins/Media/ElectraPlayer/Source/ElectraPlayerRuntime/Private/Runtime/Utilities/UtilsMP4.h
+++ b/Engine/Plugins/Media/ElectraPlayer/Source/ElectraPlayerRuntime/Private/Runtime/Utilities/UtilsMP4.h
@@ -97,7 +97,7 @@ public:
 			int64 Offset = 0;
 			uint32 Type = 0;
 			uint8 UUID[16] {};
-			TSharedPtrTS<IElectraHttpManager::FReceiveBuffer> DataBuffer;
+			TSharedPtrTS<FWaitableBuffer> DataBuffer;
 		};
 		FMP4RootBoxLocator() = default;
 		~FMP4RootBoxLocator();
