Electra: Unifying the data reader for mp4 and mkv

#rnx

diff --git a/Engine/Plugins/Media/ElectraPlayer/Source/ElectraPlayerRuntime/Private/Runtime/BufferedDataReader.h b/Engine/Plugins/Media/ElectraPlayer/Source/ElectraPlayerRuntime/Private/Runtime/BufferedDataReader.h
--- a/Engine/Plugins/Media/ElectraPlayer/Source/ElectraPlayerRuntime/Private/Runtime/BufferedDataReader.h
+++ b/Engine/Plugins/Media/ElectraPlayer/Source/ElectraPlayerRuntime/Private/Runtime/BufferedDataReader.h
@@ -9,6 +9,52 @@
 
 namespace Electra
 {
+	class IGenericDataReader
+	{
+	public:
+		virtual ~IGenericDataReader() = default;
+		/**
+		 * Read n bytes of data starting at offset o into the provided buffer.
+		 *
+		 * Reading must return the number of bytes asked to get, if necessary by blocking.
+		 * If a read error prevents reading the number of bytes -1 must be returned.
+		 *
+		 * @param IntoBuffer Buffer into which to store the data bytes. If nullptr is passed the data must be skipped over.
+		 * @param NumBytesToRead The number of bytes to read. Must not read more bytes and should be no less than requested.
+		 * @return The number of bytes read or -1 on a read error. If the read would go beyond the size of the file then
+		 *         returning fewer bytes than requested is permitted in this case ONLY.
+		 */
+		virtual int64 ReadData(void* InDestinationBuffer, int64 InNumBytesToRead, int64 InFromOffset) = 0;
+
+		/**
+		 * Returns the current internal read offset.
+		 */
+		virtual int64 GetCurrentOffset() const = 0;
+
+		/**
+		 * Returns the total size of the file.
+		 * The size should be available at least after the first call to ReadData().
+		 * If the total length is not known, return -1.
+		 */
+		virtual int64 GetTotalSize() const = 0;
+
+		/**
+		 * Checks if reading of the file and therefore parsing has been aborted.
+		 *
+		 * @return true if reading/parsing has been aborted, false otherwise.
+		 */
+		virtual bool HasReadBeenAborted() const = 0;
+
+		/**
+		 * Checks if the data source has reached the End Of File (EOF) and cannot provide any additional data.
+		 *
+		 * @return If EOF has been reached returns true, otherwise false.
+		 */
+		virtual bool HasReachedEOF() const = 0;
+	};
+
+
+
 	class FBufferedDataReader
 	{
 	public:
@@ -38,8 +84,8 @@ namespace Electra
 		}
 
 		int64 GetCurrentOffset() const
-		{ 
-			return CurrentOffset; 
+		{
+			return CurrentOffset;
 		}
 
 		int64 GetTotalDataSize() const
diff --git a/Engine/Plugins/Media/ElectraPlayer/Source/ElectraPlayerRuntime/Private/Runtime/Demuxer/ParserISO14496-12.cpp b/Engine/Plugins/Media/ElectraPlayer/Source/ElectraPlayerRuntime/Private/Runtime/Demuxer/ParserISO14496-12.cpp
--- a/Engine/Plugins/Media/ElectraPlayer/Source/ElectraPlayerRuntime/Private/Runtime/Demuxer/ParserISO14496-12.cpp
+++ b/Engine/Plugins/Media/ElectraPlayer/Source/ElectraPlayerRuntime/Private/Runtime/Demuxer/ParserISO14496-12.cpp
@@ -78,7 +78,7 @@ namespace Electra
 	class FMP4BoxReader : private TMediaNoncopyable<FMP4BoxReader>
 	{
 	public:
-		FMP4BoxReader(IParserISO14496_12::IReader* InDataReader, IParserISO14496_12::IBoxCallback* InBoxParseCallback)
+		FMP4BoxReader(IGenericDataReader* InDataReader, IParserISO14496_12::IBoxCallback* InBoxParseCallback)
 			: DataReader(InDataReader)
 			, BoxParseCallback(InBoxParseCallback)
 			, DataDuplicationBuffer(nullptr)
@@ -104,12 +104,13 @@ namespace Electra
 		UEMediaError Read(T& value)
 		{
 			T Temp = 0;
-			int64 NumRead = DataReader->ReadData(&Temp, sizeof(T));
-			if (NumRead == sizeof(T))
+			int64 NumToRead = sizeof(T);
+			int64 NumRead = DataReader->ReadData(&Temp, NumToRead, -1);
+			if (NumRead == NumToRead)
 			{
 				if (DataDuplicationBuffer)
 				{
-					DataDuplicationBuffer->Append((const uint8*)&Temp, sizeof(T));
+					DataDuplicationBuffer->Append((const uint8*)&Temp, NumToRead);
 				}
 				value = Utils::ValueFromBigEndian(Temp);
 				return UEMEDIA_ERROR_OK;
@@ -118,7 +119,7 @@ namespace Electra
 			{
 				return UEMEDIA_ERROR_ABORTED;
 			}
-			else if (NumRead == 0 || DataReader->HasReachedEOF())
+			else if (NumRead < NumToRead || DataReader->HasReachedEOF())
 			{
 				return UEMEDIA_ERROR_INSUFFICIENT_DATA;
 			}
@@ -192,11 +193,11 @@ namespace Electra
 			{
 				int32 BufferSizeNow = DataDuplicationBuffer->Num();
 				DataDuplicationBuffer->AddUninitialized((int32)NumBytes);
-				NumRead = DataReader->ReadData(DataDuplicationBuffer->GetData() + BufferSizeNow, NumBytes);
+				NumRead = DataReader->ReadData(DataDuplicationBuffer->GetData() + BufferSizeNow, NumBytes, -1);
 			}
 			else
 			{
-				NumRead = DataReader->ReadData(Buffer, NumBytes);
+				NumRead = DataReader->ReadData(Buffer, NumBytes, -1);
 				if (NumRead == NumBytes && DataDuplicationBuffer)
 				{
 					DataDuplicationBuffer->Append((const uint8*)Buffer, (int32)NumBytes);
@@ -262,20 +263,20 @@ namespace Electra
 		FMP4BoxReader() = delete;
 		FMP4BoxReader(const FMP4BoxReader&) = delete;
 
-		IParserISO14496_12::IReader* DataReader = nullptr;
+		IGenericDataReader* DataReader = nullptr;
 		IParserISO14496_12::IBoxCallback* BoxParseCallback = nullptr;
 
 		TArray<uint8>* DataDuplicationBuffer = nullptr;
 	};
 
 
-	class FDataBufferReader : public IParserISO14496_12::IReader
+	class FDataBufferReader : public IGenericDataReader
 	{
 	public:
 		FDataBufferReader(const TArray<uint8>& InDataBufferToReadFrom) : DataBufferRef(InDataBufferToReadFrom)
 		{}
 		virtual ~FDataBufferReader() = default;
-		int64 ReadData(void* IntoBuffer, int64 NumBytesToRead) override
+		int64 ReadData(void* IntoBuffer, int64 NumBytesToRead, int64 InFromOffset) override
 		{
 			int64 NumAvail = DataBufferRef.Num() - CurrentOffset;
 			if (NumAvail >= NumBytesToRead)
@@ -295,6 +296,9 @@ namespace Electra
 		{ return false;	}
 		int64 GetCurrentOffset() const override
 		{ return CurrentOffset;	}
+		int64 GetTotalSize() const override
+		{ check(!"this should not be called"); return -1; }
+
 	private:
 		const TArray<uint8>& DataBufferRef;
 		int64 CurrentOffset = 0;
@@ -5715,7 +5719,7 @@ namespace Electra
 	public:
 		FParserISO14496_12();
 		virtual ~FParserISO14496_12();
-		UEMediaError ParseHeader(IReader* DataReader, IBoxCallback* BoxParseCallback, IPlayerSessionServices* PlayerSession, const IParserISO14496_12* OptionalInitSegment) override;
+		UEMediaError ParseHeader(IGenericDataReader* DataReader, IBoxCallback* BoxParseCallback, IPlayerSessionServices* PlayerSession, const IParserISO14496_12* OptionalInitSegment) override;
 
 		UEMediaError PrepareTracks(IPlayerSessionServices* PlayerSession, TSharedPtrTS<const IParserISO14496_12> OptionalMP4InitSegment) override;
 
@@ -7075,7 +7079,7 @@ namespace Electra
 		delete ParsedData;
 	}
 
-	UEMediaError FParserISO14496_12::ParseHeader(IReader* InDataReader, IBoxCallback* InBoxParseCallback, IPlayerSessionServices* PlayerSession, const IParserISO14496_12* InOptionalInitSegment)
+	UEMediaError FParserISO14496_12::ParseHeader(IGenericDataReader* InDataReader, IBoxCallback* InBoxParseCallback, IPlayerSessionServices* PlayerSession, const IParserISO14496_12* InOptionalInitSegment)
 	{
 		if (!InDataReader || !InBoxParseCallback)
 		{
diff --git a/Engine/Plugins/Media/ElectraPlayer/Source/ElectraPlayerRuntime/Private/Runtime/Demuxer/ParserISO14496-12.h b/Engine/Plugins/Media/ElectraPlayer/Source/ElectraPlayerRuntime/Private/Runtime/Demuxer/ParserISO14496-12.h
--- a/Engine/Plugins/Media/ElectraPlayer/Source/ElectraPlayerRuntime/Private/Runtime/Demuxer/ParserISO14496-12.h
+++ b/Engine/Plugins/Media/ElectraPlayer/Source/ElectraPlayerRuntime/Private/Runtime/Demuxer/ParserISO14496-12.h
@@ -9,6 +9,7 @@
 #include "OptionalValue.h"
 #include "StreamTypes.h"
 #include "ParameterDictionary.h"
+#include "BufferedDataReader.h"
 
 #include "ElectraEncryptedSampleInfo.h"
 
@@ -29,49 +30,6 @@ namespace Electra
 	public:
 		virtual ~IParserISO14496_12() = default;
 
-		/**
-		 * Interface for reading data from a source.
-		 */
-		class IReader
-		{
-		public:
-			virtual ~IReader() = default;
-			/**
-			 * Read n bytes of data into the provided buffer.
-			 *
-			 * Reading must return the number of bytes asked to get, if necessary by blocking.
-			 * If a read error prevents reading the number of bytes -1 must be returned.
-			 *
-			 * @param IntoBuffer Buffer into which to store the data bytes. If nullptr is passed the data must be skipped over.
-			 * @param NumBytesToRead The number of bytes to read. Must not read more bytes and no less than requested.
-			 * @return The number of bytes read or -1 on a read error.
-			 */
-			virtual int64 ReadData(void* IntoBuffer, int64 NumBytesToRead) = 0;
-
-			/**
-			 * Checks if the data source has reached the End Of File (EOF) and cannot provide any additional data.
-			 *
-			 * @return If EOF has been reached returns true, otherwise false.
-			 */
-			virtual bool HasReachedEOF() const = 0;
-
-			/**
-			 * Checks if reading of the file and therefor parsing has been aborted.
-			 *
-			 * @return true if reading/parsing has been aborted, false otherwise.
-			 */
-			virtual bool HasReadBeenAborted() const = 0;
-
-			/**
-			 * Returns the current read offset.
-			 *
-			 * The first read offset is not necessarily zero. It could be anywhere inside the source.
-			 *
-			 * @return The current byte offset in the source.
-			 */
-			virtual int64 GetCurrentOffset() const = 0;
-		};
-
 		/** Box type is a 32 bit value in an mp4 file. */
 		typedef uint32 FBoxType;
 
@@ -137,7 +95,7 @@ namespace Electra
 		/**
 		 * Parses the header boxes (all non-MDAT boxes).
 		 */
-		virtual UEMediaError ParseHeader(IReader* DataReader, IBoxCallback* BoxParseCallback, IPlayerSessionServices* PlayerSession, const IParserISO14496_12* OptionalInitSegment) = 0;
+		virtual UEMediaError ParseHeader(IGenericDataReader* DataReader, IBoxCallback* BoxParseCallback, IPlayerSessionServices* PlayerSession, const IParserISO14496_12* OptionalInitSegment) = 0;
 
 
 		/** A brand is a 32 bit value in an mp4 file. */
diff --git a/Engine/Plugins/Media/ElectraPlayer/Source/ElectraPlayerRuntime/Private/Runtime/Demuxer/ParserISO14496-12_Utils.h b/Engine/Plugins/Media/ElectraPlayer/Source/ElectraPlayerRuntime/Private/Runtime/Demuxer/ParserISO14496-12_Utils.h
--- a/Engine/Plugins/Media/ElectraPlayer/Source/ElectraPlayerRuntime/Private/Runtime/Demuxer/ParserISO14496-12_Utils.h
+++ b/Engine/Plugins/Media/ElectraPlayer/Source/ElectraPlayerRuntime/Private/Runtime/Demuxer/ParserISO14496-12_Utils.h
@@ -10,7 +10,7 @@
 namespace Electra
 {
 
-	class FMP4StaticDataReader : public IParserISO14496_12::IReader
+	class FMP4StaticDataReader : public IGenericDataReader
 	{
 	public:
 		FMP4StaticDataReader() = default;
@@ -28,9 +28,9 @@ namespace Electra
 		}
 	private:
 		//----------------------------------------------------------------------
-		// Methods from IParserISO14496_12::IReader
+		// Methods from IGenericDataReader
 		//
-		virtual int64 ReadData(void* IntoBuffer, int64 NumBytesToRead) override
+		int64 ReadData(void* IntoBuffer, int64 NumBytesToRead, int64 InFromOffset) override
 		{
 			if (ResponseBuffer.IsValid() && NumBytesToRead <= DataSize - CurrentOffset)
 			{
@@ -43,18 +43,23 @@ namespace Electra
 			}
 			return -1;
 		}
-		virtual bool HasReachedEOF() const override
+		bool HasReachedEOF() const override
 		{
 			return ResponseBuffer.IsValid() ? ResponseBuffer->GetEOD() && CurrentOffset >= DataSize : true;
 		}
-		virtual bool HasReadBeenAborted() const override
+		bool HasReadBeenAborted() const override
 		{
 			return ResponseBuffer.IsValid() ? ResponseBuffer->WasAborted() : true;
 		}
-		virtual int64 GetCurrentOffset() const override
+		int64 GetCurrentOffset() const override
 		{
 			return CurrentOffset;
 		}
+		int64 GetTotalSize() const override
+		{
+			check(!"this should not be called");
+			return -1;
+		}
 
 		TSharedPtrTS<FWaitableBuffer> ResponseBuffer;
 		const uint8* Data = nullptr;
diff --git a/Engine/Plugins/Media/ElectraPlayer/Source/ElectraPlayerRuntime/Private/Runtime/Demuxer/ParserMKV.cpp b/Engine/Plugins/Media/ElectraPlayer/Source/ElectraPlayerRuntime/Private/Runtime/Demuxer/ParserMKV.cpp
--- a/Engine/Plugins/Media/ElectraPlayer/Source/ElectraPlayerRuntime/Private/Runtime/Demuxer/ParserMKV.cpp
+++ b/Engine/Plugins/Media/ElectraPlayer/Source/ElectraPlayerRuntime/Private/Runtime/Demuxer/ParserMKV.cpp
@@ -527,13 +527,13 @@ namespace Electra
 
 		FParserMKV(IPlayerSessionServices* InPlayerSession);
 
-		FErrorDetail ParseHeader(IReader* DataReader, EParserFlags ParseFlags) override;
+		FErrorDetail ParseHeader(IGenericDataReader* DataReader, EParserFlags ParseFlags) override;
 		FErrorDetail PrepareTracks() override;
 		FTimeValue GetDuration() const override;
 		int32 GetNumberOfTracks() const override;
 		const ITrack* GetTrackByIndex(int32 Index) const override;
 		const ITrack* GetTrackByTrackID(uint64 TrackID) const override;
-		TSharedPtrTS<IClusterParser> CreateClusterParser(IReader* DataReader, const TArray<uint64>& TrackIDsToParse, EClusterParseFlags ParseFlags) const override;
+		TSharedPtrTS<IClusterParser> CreateClusterParser(IGenericDataReader* DataReader, const TArray<uint64>& TrackIDsToParse, EClusterParseFlags ParseFlags) const override;
 		void AddCue(int64 InCueTimestamp, uint64 InTrackID, int64 InCueRelativePosition, uint64 InCueBlockNumber, int64 InClusterPosition) override;
 
 		int64 OnReadAssetData(void* Destination, int64 NumBytes, int64 FromOffset, int64* OutTotalSize) override;
@@ -2232,7 +2232,7 @@ namespace Electra
 			};
 			#undef OVERRIDE_BASE_ACTIONS
 
-			FMKVClusterParser(TSharedPtrTS<const FParserMKV> InParentMKV, IParserMKV::IReader* InDataReader, const TArray<uint64>& InTrackIDsToParse, EClusterParseFlags InParseFlags);
+			FMKVClusterParser(TSharedPtrTS<const FParserMKV> InParentMKV, IGenericDataReader* InDataReader, const TArray<uint64>& InTrackIDsToParse, EClusterParseFlags InParseFlags);
 			virtual ~FMKVClusterParser() {}
 			EParseAction NextParseAction() override;
 			FErrorDetail GetLastError() const override;
@@ -2243,18 +2243,18 @@ namespace Electra
 			class FClusterDataReader : public IMKVElementReader, public FMKVEBMLReader, public IMKVFetcher
 			{
 			public:
-				FClusterDataReader(IParserMKV::IReader* InReader) : FMKVEBMLReader(this), Reader(InReader), StartOffset(InReader->MKVGetCurrentFileOffset()), Offset(InReader->MKVGetCurrentFileOffset()) {}
+				FClusterDataReader(IGenericDataReader* InReader) : FMKVEBMLReader(this), Reader(InReader), StartOffset(InReader->GetCurrentOffset()), Offset(InReader->GetCurrentOffset()) {}
 				virtual ~FClusterDataReader() {}
 				int64 GetStartOffset() const { return StartOffset; }
 				int64 GetCurrentOffset() const { return CurrentOffset(); }
-				int64 GetTotalSize() const { return Reader->MKVGetTotalSize(); }
+				int64 GetTotalSize() const { return Reader->GetTotalSize(); }
 				int64 GetEndOffset() const { return StartOffset + GetTotalSize(); }
-				bool HasReadBeenAborted() const { return bWasAborted || Reader->MKVHasReadBeenAborted(); }
+				bool HasReadBeenAborted() const { return bWasAborted || Reader->HasReadBeenAborted(); }
 				bool ReachedEOS() const { return bReachedEOS; }
 				bool SkipOver(int64 NumBytes) { return Skip(NumBytes); }
 
 				bool Read(uint8& OutValue) override
-				{ return Validate(Reader->MKVReadData(&OutValue, sizeof(OutValue), Offset), sizeof(OutValue)); }
+				{ return Validate(Reader->ReadData(&OutValue, sizeof(OutValue), Offset), sizeof(OutValue)); }
 				bool Read(uint16& OutValue) override
 				{ return ReadValue(OutValue); }
 				bool Read(uint32& OutValue) override
@@ -2262,9 +2262,9 @@ namespace Electra
 				bool Read(uint64& OutValue) override
 				{ return ReadValue(OutValue); }
 				bool Read(TArray<uint8>& OutValue, int64 NumBytes) override
-				{ OutValue.AddUninitialized((int32)NumBytes); return Validate(Reader->MKVReadData(OutValue.GetData(), NumBytes, Offset), NumBytes); }
+				{ OutValue.AddUninitialized((int32)NumBytes); return Validate(Reader->ReadData(OutValue.GetData(), NumBytes, Offset), NumBytes); }
 				bool Skip(int64 NumBytes) override
-				{ return Validate(Reader->MKVReadData(nullptr, NumBytes, Offset), NumBytes); }
+				{ return Validate(Reader->ReadData(nullptr, NumBytes, Offset), NumBytes); }
 
 				void SetOffset(int64 InOffset)
 				{ Offset = InOffset; }
@@ -2290,13 +2290,13 @@ namespace Electra
 				{ check(!"not implemented"); return false;}
 			private:
 				FErrorDetail LastError() const override { return Error; }
-				int64 CurrentOffset() const override { return Reader->MKVGetCurrentFileOffset(); }
+				int64 CurrentOffset() const override { return Reader->GetCurrentOffset(); }
 				bool Prefetch(int64 NumBytes) override { return true; }
 
 				template<typename T>
 				bool ReadValue(T& OutValue)
 				{
-					if (Validate(Reader->MKVReadData(&OutValue, sizeof(T), Offset), sizeof(T)))
+					if (Validate(Reader->ReadData(&OutValue, sizeof(T), Offset), sizeof(T)))
 					{
 					#if PLATFORM_LITTLE_ENDIAN
 						OutValue = Utils::EndianSwap(OutValue);
@@ -2317,11 +2317,11 @@ namespace Electra
 					{
 						bReachedEOS = true;
 					}
-					bWasAborted = Reader->MKVHasReadBeenAborted();
+					bWasAborted = Reader->HasReadBeenAborted();
 					return false;
 				}
 
-				IParserMKV::IReader* Reader = nullptr;
+				IGenericDataReader* Reader = nullptr;
 				int64 StartOffset = 0;
 				int64 Offset = 0;
 				FErrorDetail Error;
@@ -2436,7 +2436,7 @@ namespace Electra
 	private:
 		bool SetupCodecInfo(FStreamCodecInformation& OutCodecInformation, TMKVElementPtr<FMKVTrackEntry> InFromTrack);
 
-		IReader* DataReader = nullptr;
+		IGenericDataReader* DataReader = nullptr;
 		IPlayerSessionServices* PlayerSessionServices = nullptr;
 
 		const TMap<FString,FString> CodecMapping
@@ -2541,12 +2541,12 @@ namespace Electra
 	{
 		if (DataReader)
 		{
-			int64 NumRead = DataReader->MKVReadData(Destination, NumBytes, FromOffset);
+			int64 NumRead = DataReader->ReadData(Destination, NumBytes, FromOffset);
 			if (OutTotalSize)
 			{
-				*OutTotalSize = DataReader->MKVGetTotalSize();
+				*OutTotalSize = DataReader->GetTotalSize();
 			}
-			if (DataReader->MKVHasReadBeenAborted())
+			if (DataReader->HasReadBeenAborted())
 			{
 				return static_cast<int64>(FBufferedDataReader::IDataProvider::EError::Aborted);
 			}
@@ -2564,7 +2564,7 @@ namespace Electra
 		PlayerSessionServices = InPlayerSession;
 	}
 
-	FErrorDetail FParserMKV::ParseHeader(IReader* InDataReader, EParserFlags ParseFlags)
+	FErrorDetail FParserMKV::ParseHeader(IGenericDataReader* InDataReader, EParserFlags ParseFlags)
 	{
 		DataReader = InDataReader;
 
@@ -3716,7 +3716,7 @@ namespace Electra
 
 
 
-	TSharedPtrTS<IParserMKV::IClusterParser> FParserMKV::CreateClusterParser(IParserMKV::IReader* InDataReader, const TArray<uint64>& InTrackIDsToParse, EClusterParseFlags InParseFlags) const
+	TSharedPtrTS<IParserMKV::IClusterParser> FParserMKV::CreateClusterParser(IGenericDataReader* InDataReader, const TArray<uint64>& InTrackIDsToParse, EClusterParseFlags InParseFlags) const
 	{
 		FMKVClusterParser* Parser = new FMKVClusterParser(AsShared(), InDataReader, InTrackIDsToParse, InParseFlags);
 		return TSharedPtrTS<IParserMKV::IClusterParser>(Parser);
@@ -3731,7 +3731,7 @@ namespace Electra
 		Segment->GetOrCreateCues()->AddCue(InCueTimestamp / TimestampScale, InTrackID, InCueRelativePosition, InCueBlockNumber, InClusterPosition, NextCueUniqueID);
 	}
 
-	FParserMKV::FMKVClusterParser::FMKVClusterParser(TSharedPtrTS<const FParserMKV> InParentMKV, IParserMKV::IReader* InDataReader, const TArray<uint64>& InTrackIDsToParse, EClusterParseFlags InParseFlags)
+	FParserMKV::FMKVClusterParser::FMKVClusterParser(TSharedPtrTS<const FParserMKV> InParentMKV, IGenericDataReader* InDataReader, const TArray<uint64>& InTrackIDsToParse, EClusterParseFlags InParseFlags)
 		: ParentMKV(InParentMKV), Reader(new FClusterDataReader(InDataReader))
 	{
 		SegmentBaseOffset = InParentMKV->GetSegment()->GetElementOffset();
diff --git a/Engine/Plugins/Media/ElectraPlayer/Source/ElectraPlayerRuntime/Private/Runtime/Demuxer/ParserMKV.h b/Engine/Plugins/Media/ElectraPlayer/Source/ElectraPlayerRuntime/Private/Runtime/Demuxer/ParserMKV.h
--- a/Engine/Plugins/Media/ElectraPlayer/Source/ElectraPlayerRuntime/Private/Runtime/Demuxer/ParserMKV.h
+++ b/Engine/Plugins/Media/ElectraPlayer/Source/ElectraPlayerRuntime/Private/Runtime/Demuxer/ParserMKV.h
@@ -8,7 +8,7 @@
 #include "ParameterDictionary.h"
 #include "ErrorDetail.h"
 #include "ElectraEncryptedSampleInfo.h"
-
+#include "BufferedDataReader.h"
 
 namespace Electra
 {
@@ -26,44 +26,6 @@ namespace Electra
 	public:
 		virtual ~IParserMKV() = default;
 
-		/**
-		 * Interface for reading data from a source.
-		 */
-		class IReader
-		{
-		public:
-			virtual ~IReader() = default;
-			/**
-			 * Read n bytes of data starting at offset o into the provided buffer.
-			 *
-			 * Reading must return the number of bytes asked to get, if necessary by blocking.
-			 * If a read error prevents reading the number of bytes -1 must be returned.
-			 *
-			 * @param IntoBuffer Buffer into which to store the data bytes. If nullptr is passed the data must be skipped over.
-			 * @param NumBytesToRead The number of bytes to read. Must not read more bytes and no less than requested.
-			 * @return The number of bytes read or -1 on a read error. If the read would go beyond the size of the file then
-			 *         returning fewer bytes than requested is permitted in this case ONLY.
-			 */
-			virtual int64 MKVReadData(void* InDestinationBuffer, int64 InNumBytesToRead, int64 InFromOffset) = 0;
-
-			virtual int64 MKVGetCurrentFileOffset() const = 0;
-
-			/**
-			 * Returns the total size of the file.
-			 * This should be possible after performing the first ReadData().
-			 * If the total length is not known, return -1.
-			 */
-			virtual int64 MKVGetTotalSize() = 0;
-
-			/**
-			 * Checks if reading of the file and therefor parsing has been aborted.
-			 *
-			 * @return true if reading/parsing has been aborted, false otherwise.
-			 */
-			virtual bool MKVHasReadBeenAborted() const = 0;
-		};
-
-
 		static TSharedPtrTS<IParserMKV> CreateParser(IPlayerSessionServices* PlayerSession);
 
 		enum EParserFlags
@@ -77,7 +39,7 @@ namespace Electra
 		/**
 		 * Parses the header boxes.
 		 */
-		virtual FErrorDetail ParseHeader(IReader* DataReader, EParserFlags ParseFlags) = 0;
+		virtual FErrorDetail ParseHeader(IGenericDataReader* DataReader, EParserFlags ParseFlags) = 0;
 
 
 		/*******************************************************************************************************************/
@@ -180,7 +142,7 @@ namespace Electra
 
 			/**
 			 * Performs parsing the cluster content or current frame, returning the next action to take.
-			 * 
+			 *
 			 * Returns the action to perform next.
 			 */
 			virtual EParseAction NextParseAction() = 0;
@@ -189,11 +151,11 @@ namespace Electra
 			 * Returns the error which resulted in returning the next action `Failure`
 			 */
 			virtual FErrorDetail GetLastError() const = 0;
-			
+
 			/**
 			 * Returns the base class pointer of the next action to perform.
 			 * Must be cast to the action indicated by the next action.
-			 * 
+			 *
 			 * The action is owned by the parser and must not be destroyed!
 			 */
 			virtual const IAction* GetAction() const = 0;
@@ -203,7 +165,7 @@ namespace Electra
 			 * This is needed for retries.
 			 */
 			virtual int64 GetClusterPosition() const = 0;
-			
+
 			/**
 			 * Returns the offset of the current block (simple or group) in the cluster.
 			 * This is needed for retries.
@@ -222,7 +184,7 @@ namespace Electra
 			virtual const TArray<uint8>& GetCodecSpecificData() const = 0;
 			virtual const FStreamCodecInformation& GetCodecInformation() const = 0;
 			virtual const FString GetLanguage() const = 0;
-			
+
 			virtual ICueIterator* CreateCueIterator() const = 0;
 		};
 
@@ -241,7 +203,7 @@ namespace Electra
 		 * Create a cluster parser.
 		 * The data reader MUST start reading on a Matroska cluster.
 		 */
-		virtual TSharedPtrTS<IClusterParser> CreateClusterParser(IReader* DataReader, const TArray<uint64>& TrackIDsToParse, EClusterParseFlags ParseFlags) const = 0;
+		virtual TSharedPtrTS<IClusterParser> CreateClusterParser(IGenericDataReader* DataReader, const TArray<uint64>& TrackIDsToParse, EClusterParseFlags ParseFlags) const = 0;
 
 		// Adds a cue of it does not exist yet. This may be called during cluster parsing for sync samples since not all
 		// sync samples may have been added as cues in the multiplexing process.
diff --git a/Engine/Plugins/Media/ElectraPlayer/Source/ElectraPlayerRuntime/Private/Runtime/Demuxer/ParserMKV_Utils.h b/Engine/Plugins/Media/ElectraPlayer/Source/ElectraPlayerRuntime/Private/Runtime/Demuxer/ParserMKV_Utils.h
--- a/Engine/Plugins/Media/ElectraPlayer/Source/ElectraPlayerRuntime/Private/Runtime/Demuxer/ParserMKV_Utils.h
+++ b/Engine/Plugins/Media/ElectraPlayer/Source/ElectraPlayerRuntime/Private/Runtime/Demuxer/ParserMKV_Utils.h
@@ -11,7 +11,7 @@
 namespace Electra
 {
 
-	class FMKVStaticDataReader : public IParserMKV::IReader
+	class FMKVStaticDataReader : public IGenericDataReader
 	{
 	public:
 		FMKVStaticDataReader() = default;
@@ -24,7 +24,7 @@ namespace Electra
 			CurrentOffset = 0;
 		}
 	private:
-		int64 MKVReadData(void* InDestinationBuffer, int64 InNumBytesToRead, int64 InFromOffset) override
+		int64 ReadData(void* InDestinationBuffer, int64 InNumBytesToRead, int64 InFromOffset) override
 		{
 			if (InFromOffset >= DataSize)
 			{
@@ -42,12 +42,14 @@ namespace Electra
 			}
 			return InNumBytesToRead;
 		}
-		int64 MKVGetCurrentFileOffset() const override
+		int64 GetCurrentOffset() const override
 		{ return CurrentOffset; }
-		int64 MKVGetTotalSize() override
+		int64 GetTotalSize() const override
 		{ return DataSize; }
-		bool MKVHasReadBeenAborted() const override
+		bool HasReadBeenAborted() const override
 		{ return false; }
+		bool HasReachedEOF() const override
+		{ check(!"this should not be called"); return false; }
 
 		TSharedPtrTS<FWaitableBuffer> ResponseBuffer;
 		const uint8* Data = nullptr;
diff --git a/Engine/Plugins/Media/ElectraPlayer/Source/ElectraPlayerRuntime/Private/Runtime/Player/DASH/ManifestDASH.cpp b/Engine/Plugins/Media/ElectraPlayer/Source/ElectraPlayerRuntime/Private/Runtime/Player/DASH/ManifestDASH.cpp
--- a/Engine/Plugins/Media/ElectraPlayer/Source/ElectraPlayerRuntime/Private/Runtime/Player/DASH/ManifestDASH.cpp
+++ b/Engine/Plugins/Media/ElectraPlayer/Source/ElectraPlayerRuntime/Private/Runtime/Player/DASH/ManifestDASH.cpp
@@ -88,7 +88,7 @@ namespace DashUtils
 
 
 
-	class FMatroskaParserDataReader : public IParserMKV::IReader, public TSharedFromThis<FMatroskaParserDataReader, ESPMode::ThreadSafe>
+	class FMatroskaParserDataReader : public IGenericDataReader, public TSharedFromThis<FMatroskaParserDataReader, ESPMode::ThreadSafe>
 	{
 	public:
 		virtual ~FMatroskaParserDataReader() { }
@@ -97,7 +97,7 @@ namespace DashUtils
 			: LoadRequest(MoveTemp(InLoadReq))
 		{ }
 
-		int64 MKVReadData(void* InDestinationBuffer, int64 InNumBytesToRead, int64 InFromOffset) override
+		int64 ReadData(void* InDestinationBuffer, int64 InNumBytesToRead, int64 InFromOffset) override
 		{
 			int64 nr = -1;
 			if (LoadChunk(InNumBytesToRead, InFromOffset) && ResponseBuffer.IsValid())
@@ -108,12 +108,14 @@ namespace DashUtils
 			return nr;
 		}
 
-		int64 MKVGetCurrentFileOffset() const override
+		int64 GetCurrentOffset() const override
 		{ check(!"should not be called!"); return -1; }
-		int64 MKVGetTotalSize() override
+		int64 GetTotalSize() const override
 		{ return FileSize; }
-		bool MKVHasReadBeenAborted() const override
+		bool HasReadBeenAborted() const override
 		{ return false; }
+		bool HasReachedEOF() const override
+		{ check(!"this should not be called"); return false; }
 
 	private:
 		bool LoadChunk(int64 InNumBytesToRead, int64 InFromOffset)
diff --git a/Engine/Plugins/Media/ElectraPlayer/Source/ElectraPlayerRuntime/Private/Runtime/Player/DASH/StreamReaderDASH.cpp b/Engine/Plugins/Media/ElectraPlayer/Source/ElectraPlayerRuntime/Private/Runtime/Player/DASH/StreamReaderDASH.cpp
--- a/Engine/Plugins/Media/ElectraPlayer/Source/ElectraPlayerRuntime/Private/Runtime/Player/DASH/StreamReaderDASH.cpp
+++ b/Engine/Plugins/Media/ElectraPlayer/Source/ElectraPlayerRuntime/Private/Runtime/Player/DASH/StreamReaderDASH.cpp
@@ -1471,7 +1471,7 @@ void FStreamReaderDASH::FStreamHandler::HandleRequestMP4()
 									}
 									else if (NumBytesToSkip > 0)
 									{
-										int64 NumSkipped = ReadData(nullptr, NumBytesToSkip);
+										int64 NumSkipped = ReadData(nullptr, NumBytesToSkip, -1);
 										if (NumSkipped != NumBytesToSkip)
 										{
 											FAccessUnit::Release(AccessUnit);
@@ -1489,7 +1489,7 @@ void FStreamReaderDASH::FStreamHandler::HandleRequestMP4()
 										MoofInfo.PayloadStartOffset = GetCurrentOffset();
 									}
 
-									int64 NumRead = ReadData(AccessUnit->AUData, AccessUnit->AUSize);
+									int64 NumRead = ReadData(AccessUnit->AUData, AccessUnit->AUSize, -1);
 									if (NumRead == AccessUnit->AUSize)
 									{
 										MoofInfo.NumKeyframeBytes += AccessUnit->bIsSyncSample ? AccessUnit->AUSize : 0;
@@ -2122,7 +2122,7 @@ void FStreamReaderDASH::FStreamHandler::HandleRequestMKV()
 
 						int64 NumToRead = Action->GetNumBytesToRead();
 						void* ReadTo = PrepareAccessUnit(NumToRead);
-						int64 nr = ReadData(ReadTo, NumToRead);
+						int64 nr = ReadData(ReadTo, NumToRead, -1);
 						if (nr != NumToRead)
 						{
 							bHasErrored = true;
@@ -2224,7 +2224,7 @@ void FStreamReaderDASH::FStreamHandler::HandleRequestMKV()
 						const IParserMKV::IClusterParser::IActionSkipOver* Action = static_cast<const IParserMKV::IClusterParser::IActionSkipOver*>(ClusterParser->GetAction());
 						check(Action);
 						int64 NumBytesToSkip = Action->GetNumBytesToSkip();
-						int64 nr = ReadData(nullptr, NumBytesToSkip);
+						int64 nr = ReadData(nullptr, NumBytesToSkip, -1);
 						if (nr != NumBytesToSkip)
 						{
 							bHasErrored = true;
@@ -2681,7 +2681,7 @@ bool FStreamReaderDASH::FStreamHandler::HasErrored() const
  * @param NumBytesToRead The number of bytes to read. Must not read more bytes and no less than requested.
  * @return The number of bytes read or -1 on a read error.
  */
-int64 FStreamReaderDASH::FStreamHandler::ReadData(void* IntoBuffer, int64 NumBytesToRead)
+int64 FStreamReaderDASH::FStreamHandler::ReadData(void* IntoBuffer, int64 NumBytesToRead, int64 InFromOffset)
 {
 	FWaitableBuffer& SourceBuffer = *ReadBuffer.ReceiveBuffer;
 	// Make sure the buffer will have the amount of data we need.
@@ -2795,6 +2795,19 @@ int64 FStreamReaderDASH::FStreamHandler::GetCurrentOffset() const
 	return ReadBuffer.ParsePos;
 }
 
+int64 FStreamReaderDASH::FStreamHandler::GetTotalSize() const
+{
+	if (CurrentRequest.IsValid())
+	{
+		Metrics::FSegmentDownloadStats& ds = CurrentRequest->DownloadStats;
+		if (ds.ByteSize > 0)
+		{
+			return ds.ByteSize;
+		}
+	}
+	return TNumericLimits<int64>::Max();
+}
+
 
 IParserISO14496_12::IBoxCallback::EParseContinuation FStreamReaderDASH::FStreamHandler::OnFoundBox(IParserISO14496_12::FBoxType Box, int64 BoxSizeInBytes, int64 FileDataOffset, int64 BoxDataOffset)
 {
@@ -2827,42 +2840,5 @@ IParserISO14496_12::IBoxCallback::EParseContinuation FStreamReaderDASH::FStreamH
 	return IParserISO14496_12::IBoxCallback::EParseContinuation::Continue;
 }
 
-
-
-int64 FStreamReaderDASH::FStreamHandler::MKVReadData(void* InDestinationBuffer, int64 InNumBytesToRead, int64 InFromOffset)
-{
-	check(InFromOffset == MKVGetCurrentFileOffset());
-	if (InFromOffset != MKVGetCurrentFileOffset())
-	{
-		return -1;
-	}
-	return ReadData(InDestinationBuffer, InNumBytesToRead);
-}
-
-int64 FStreamReaderDASH::FStreamHandler::MKVGetCurrentFileOffset() const
-{
-	return GetCurrentOffset();
-}
-
-int64 FStreamReaderDASH::FStreamHandler::MKVGetTotalSize()
-{
-	if (CurrentRequest.IsValid())
-	{
-		Metrics::FSegmentDownloadStats& ds = CurrentRequest->DownloadStats;
-		if (ds.ByteSize > 0)
-		{
-			return ds.ByteSize;
-		}
-	}
-	return TNumericLimits<int64>::Max();
-}
-
-bool FStreamReaderDASH::FStreamHandler::MKVHasReadBeenAborted() const
-{
-	return HasReadBeenAborted();
-}
-
-
-
 } // namespace Electra
 
diff --git a/Engine/Plugins/Media/ElectraPlayer/Source/ElectraPlayerRuntime/Private/Runtime/Player/DASH/StreamReaderDASH.h b/Engine/Plugins/Media/ElectraPlayer/Source/ElectraPlayerRuntime/Private/Runtime/Player/DASH/StreamReaderDASH.h
--- a/Engine/Plugins/Media/ElectraPlayer/Source/ElectraPlayerRuntime/Private/Runtime/Player/DASH/StreamReaderDASH.h
+++ b/Engine/Plugins/Media/ElectraPlayer/Source/ElectraPlayerRuntime/Private/Runtime/Player/DASH/StreamReaderDASH.h
@@ -106,7 +106,7 @@ public:
 private:
 	ELECTRA_IMPL_DEFAULT_ERROR_METHODS(DASHStreamReader);
 
-	struct FStreamHandler : public FMediaThread, public IParserISO14496_12::IReader, public IParserISO14496_12::IBoxCallback, public IParserMKV::IReader
+	struct FStreamHandler : public FMediaThread, public IGenericDataReader, public IParserISO14496_12::IBoxCallback
 	{
 		struct FReadBuffer
 		{
@@ -303,20 +303,15 @@ private:
 		void UpdateAUDropState(FAccessUnit* InAU, const TSharedPtrTS<FStreamSegmentRequestDASH>& InRequest);
 
 
-		// Methods from IParserISO14496_12::IReader
-		int64 ReadData(void* IntoBuffer, int64 NumBytesToRead) override;
+		// Methods from IGenericDataReader
+		int64 ReadData(void* IntoBuffer, int64 NumBytesToRead, int64 InFromOffset) override;
 		bool HasReachedEOF() const override;
 		bool HasReadBeenAborted() const override;
 		int64 GetCurrentOffset() const override;
+		int64 GetTotalSize() const override;
 		// Methods from IParserISO14496_12::IBoxCallback
 		IParserISO14496_12::IBoxCallback::EParseContinuation OnFoundBox(IParserISO14496_12::FBoxType Box, int64 BoxSizeInBytes, int64 FileDataOffset, int64 BoxDataOffset) override;
 		IParserISO14496_12::IBoxCallback::EParseContinuation OnEndOfBox(IParserISO14496_12::FBoxType Box, int64 BoxSizeInBytes, int64 FileDataOffset, int64 BoxDataOffset) override;
-
-		// Methods from IParserMKV::IReader
-		int64 MKVReadData(void* InDestinationBuffer, int64 InNumBytesToRead, int64 InFromOffset) override;
-		int64 MKVGetCurrentFileOffset() const override;
-		int64 MKVGetTotalSize() override;
-		bool MKVHasReadBeenAborted() const override;
 	};
 
 	FStreamHandler						StreamHandlers[3];		// 0 = video, 1 = audio, 2 = subtitle
diff --git a/Engine/Plugins/Media/ElectraPlayer/Source/ElectraPlayerRuntime/Private/Runtime/Player/HLS/StreamReaderHLSfmp4.cpp b/Engine/Plugins/Media/ElectraPlayer/Source/ElectraPlayerRuntime/Private/Runtime/Player/HLS/StreamReaderHLSfmp4.cpp
--- a/Engine/Plugins/Media/ElectraPlayer/Source/ElectraPlayerRuntime/Private/Runtime/Player/HLS/StreamReaderHLSfmp4.cpp
+++ b/Engine/Plugins/Media/ElectraPlayer/Source/ElectraPlayerRuntime/Private/Runtime/Player/HLS/StreamReaderHLSfmp4.cpp
@@ -974,7 +974,7 @@ void FStreamReaderHLSfmp4::FStreamHandler::HandleRequest()
 											MoofInfo.PayloadStartOffset = GetCurrentOffset();
 										}
 
-										int64 NumRead = ReadData(AccessUnit->AUData, AccessUnit->AUSize);
+										int64 NumRead = ReadData(AccessUnit->AUData, AccessUnit->AUSize, -1);
 										if (NumRead == AccessUnit->AUSize)
 										{
 											MoofInfo.ContentDuration += Duration;
@@ -1354,7 +1354,7 @@ bool FStreamReaderHLSfmp4::FStreamHandler::HasErrored() const
  * @param NumBytesToRead The number of bytes to read. Must not read more bytes and no less than requested.
  * @return The number of bytes read or -1 on a read error.
  */
-int64 FStreamReaderHLSfmp4::FStreamHandler::ReadData(void* IntoBuffer, int64 NumBytesToRead)
+int64 FStreamReaderHLSfmp4::FStreamHandler::ReadData(void* IntoBuffer, int64 NumBytesToRead, int64 InFromOffset)
 {
 	FWaitableBuffer& SourceBuffer = *ReadBuffer.ReceiveBuffer;
 	// Make sure the buffer will have the amount of data we need.
@@ -1596,6 +1596,12 @@ int64 FStreamReaderHLSfmp4::FStreamHandler::GetCurrentOffset() const
 	return ReadBuffer.ParsePos;
 }
 
+int64 FStreamReaderHLSfmp4::FStreamHandler::GetTotalSize() const
+{
+	check(!"this should not be called");
+	return -1;
+}
+
 
 IParserISO14496_12::IBoxCallback::EParseContinuation FStreamReaderHLSfmp4::FStreamHandler::OnFoundBox(IParserISO14496_12::FBoxType Box, int64 BoxSizeInBytes, int64 FileDataOffset, int64 BoxDataOffset)
 {
diff --git a/Engine/Plugins/Media/ElectraPlayer/Source/ElectraPlayerRuntime/Private/Runtime/Player/HLS/StreamReaderHLSfmp4.h b/Engine/Plugins/Media/ElectraPlayer/Source/ElectraPlayerRuntime/Private/Runtime/Player/HLS/StreamReaderHLSfmp4.h
--- a/Engine/Plugins/Media/ElectraPlayer/Source/ElectraPlayerRuntime/Private/Runtime/Player/HLS/StreamReaderHLSfmp4.h
+++ b/Engine/Plugins/Media/ElectraPlayer/Source/ElectraPlayerRuntime/Private/Runtime/Player/HLS/StreamReaderHLSfmp4.h
@@ -119,7 +119,7 @@ public:
 
 private:
 
-	struct FStreamHandler : public FMediaThread, public IParserISO14496_12::IReader, public IParserISO14496_12::IBoxCallback
+	struct FStreamHandler : public FMediaThread, public IGenericDataReader, public IParserISO14496_12::IBoxCallback
 	{
 		struct FReadBuffer
 		{
@@ -250,11 +250,13 @@ private:
 
 		bool HasErrored() const;
 
-		// Methods from IParserISO14496_12::IReader
-		int64 ReadData(void* IntoBuffer, int64 NumBytesToRead) override;
+		// Methods from IGenericDataReader
+		int64 ReadData(void* IntoBuffer, int64 NumBytesToRead, int64 InFromOffset) override;
 		bool HasReachedEOF() const override;
 		bool HasReadBeenAborted() const override;
 		int64 GetCurrentOffset() const override;
+		int64 GetTotalSize() const override;
+
 		// Methods from IParserISO14496_12::IBoxCallback
 		IParserISO14496_12::IBoxCallback::EParseContinuation OnFoundBox(IParserISO14496_12::FBoxType Box, int64 BoxSizeInBytes, int64 FileDataOffset, int64 BoxDataOffset) override;
 		IParserISO14496_12::IBoxCallback::EParseContinuation OnEndOfBox(IParserISO14496_12::FBoxType Box, int64 BoxSizeInBytes, int64 FileDataOffset, int64 BoxDataOffset) override;
diff --git a/Engine/Plugins/Media/ElectraPlayer/Source/ElectraPlayerRuntime/Private/Runtime/Player/mkv/PlaylistReaderMKV.cpp b/Engine/Plugins/Media/ElectraPlayer/Source/ElectraPlayerRuntime/Private/Runtime/Player/mkv/PlaylistReaderMKV.cpp
--- a/Engine/Plugins/Media/ElectraPlayer/Source/ElectraPlayerRuntime/Private/Runtime/Player/mkv/PlaylistReaderMKV.cpp
+++ b/Engine/Plugins/Media/ElectraPlayer/Source/ElectraPlayerRuntime/Private/Runtime/Player/mkv/PlaylistReaderMKV.cpp
@@ -29,7 +29,7 @@ namespace Electra
 /**
  * This class is responsible for downloading the mkv non-mdat boxes and parsing them.
  */
-class FPlaylistReaderMKV : public IPlaylistReaderMKV, public IParserMKV::IReader, public FMediaThread
+class FPlaylistReaderMKV : public IPlaylistReaderMKV, public IGenericDataReader, public FMediaThread
 {
 public:
 	FPlaylistReaderMKV();
@@ -76,13 +76,18 @@ public:
 
 private:
 	// Methods from IParserMKV::IReader
-	bool MKVHasReadBeenAborted() const override;
-	int64 MKVReadData(void* IntoBuffer, int64 NumBytesToRead, int64 InFromOffset) override;
-	int64 MKVGetCurrentFileOffset() const override;
-	int64 MKVGetTotalSize() override
+	bool HasReadBeenAborted() const override;
+	int64 ReadData(void* IntoBuffer, int64 NumBytesToRead, int64 InFromOffset) override;
+	int64 GetCurrentOffset() const override;
+	int64 GetTotalSize() const override
 	{
 		return FileSize;
 	}
+	bool HasReachedEOF() const override
+	{
+		check(!"this should not be called");
+		return false;
+	}
 
 	void StartWorkerThread();
 	void StopWorkerThread();
@@ -331,7 +336,7 @@ void FPlaylistReaderMKV::WorkerThread()
 	LastErrorDetail = MKVParser->ParseHeader(this, IParserMKV::EParserFlags::ParseFlag_Default);
 	ClearRequest();
 
-	if (!MKVHasReadBeenAborted())
+	if (!HasReadBeenAborted())
 	{
 		// Notify the download of the "master playlist". This indicates the download only, not the parsing thereof.
 		PlayerSessionServices->SendMessageToPlayer(IPlaylistReader::PlaylistDownloadMessage::Create(&ConnectionInfo, Playlist::EListType::Master, Playlist::ELoadType::Initial));
@@ -388,7 +393,7 @@ void FPlaylistReaderMKV::WorkerThread()
  * @param NumBytesToRead The number of bytes to read. Must not read more bytes and no less than requested.
  * @return The number of bytes read or -1 on a read error.
  */
-int64 FPlaylistReaderMKV::MKVReadData(void* IntoBuffer, int64 NumBytesToRead, int64 InFromOffset)
+int64 FPlaylistReaderMKV::ReadData(void* IntoBuffer, int64 NumBytesToRead, int64 InFromOffset)
 {
 	ReadChunk(reinterpret_cast<uint8*>(IntoBuffer), InFromOffset, NumBytesToRead);
 	while(1)
@@ -407,7 +412,7 @@ int64 FPlaylistReaderMKV::MKVReadData(void* IntoBuffer, int64 NumBytesToRead, in
 	return NumBytesToRead;
 }
 
-int64 FPlaylistReaderMKV::MKVGetCurrentFileOffset() const
+int64 FPlaylistReaderMKV::GetCurrentOffset() const
 {
 	check(!"This is not expected to be called using a buffered reader");
 	return -1;
@@ -418,7 +423,7 @@ int64 FPlaylistReaderMKV::MKVGetCurrentFileOffset() const
  *
  * @return true if reading/parsing has been aborted, false otherwise.
  */
-bool FPlaylistReaderMKV::MKVHasReadBeenAborted() const
+bool FPlaylistReaderMKV::HasReadBeenAborted() const
 {
 	return bAbort;
 }
diff --git a/Engine/Plugins/Media/ElectraPlayer/Source/ElectraPlayerRuntime/Private/Runtime/Player/mkv/StreamReaderMKV.cpp b/Engine/Plugins/Media/ElectraPlayer/Source/ElectraPlayerRuntime/Private/Runtime/Player/mkv/StreamReaderMKV.cpp
--- a/Engine/Plugins/Media/ElectraPlayer/Source/ElectraPlayerRuntime/Private/Runtime/Player/mkv/StreamReaderMKV.cpp
+++ b/Engine/Plugins/Media/ElectraPlayer/Source/ElectraPlayerRuntime/Private/Runtime/Player/mkv/StreamReaderMKV.cpp
@@ -667,7 +667,7 @@ void FStreamReaderMKV::HandleRequest()
 			// Anything still not sent off we delete when we were aborted or asked to terminate.
 			// On regular completion we keep the remaining AUs to emit with the next request.
 			// We have to keep them in order to calculate the sample durations.
-			if (bTerminate || MKVHasReadBeenAborted())
+			if (bTerminate || HasReadBeenAborted())
 			{
 				for(auto &ActiveTrk : ActiveTrackMap)
 				{
@@ -733,7 +733,7 @@ FStreamReaderMKV::EEmitResult FStreamReaderMKV::EmitSamples(EEmitType InEmitType
 	EEmitResult Result = EEmitResult::SentNothing;
 	// Emit all remaining pending AUs
 	bool bAllSentOff = false;
-	while(!bAllSentOff && !bTerminate && !MKVHasReadBeenAborted())
+	while(!bAllSentOff && !bTerminate && !HasReadBeenAborted())
 	{
 		bAllSentOff = true;
 		for(auto &ActiveTrk : ActiveTrackMap)
@@ -746,7 +746,7 @@ FStreamReaderMKV::EEmitResult FStreamReaderMKV::EmitSamples(EEmitType InEmitType
 				td.SortedAccessUnitFIFO.Empty();
 			}
 
-			while(td.AccessUnitFIFO.Num() && !MKVHasReadBeenAborted())
+			while(td.AccessUnitFIFO.Num() && !HasReadBeenAborted())
 			{
 				if (td.bNeedToRecalculateDurations)
 				{
@@ -875,7 +875,7 @@ FStreamReaderMKV::EEmitResult FStreamReaderMKV::EmitSamples(EEmitType InEmitType
 
 
 	// Check that buffers are how they are supposed to be
-	if (!bTerminate && !MKVHasReadBeenAborted() && InEmitType == EEmitType::AllRemaining)
+	if (!bTerminate && !HasReadBeenAborted() && InEmitType == EEmitType::AllRemaining)
 	{
 		for(auto &ActiveTrk : ActiveTrackMap)
 		{
@@ -948,10 +948,10 @@ int32 FStreamReaderMKV::FReadBuffer::ReadTo(void* IntoBuffer, int64 NumBytesToRe
 }
 
 
-int64 FStreamReaderMKV::MKVReadData(void* InDestinationBuffer, int64 InNumBytesToRead, int64 InFromOffset)
+int64 FStreamReaderMKV::ReadData(void* InDestinationBuffer, int64 InNumBytesToRead, int64 InFromOffset)
 {
-	check(InFromOffset == MKVGetCurrentFileOffset());
-	if (InFromOffset != MKVGetCurrentFileOffset())
+	check(InFromOffset == GetCurrentOffset());
+	if (InFromOffset != GetCurrentOffset())
 	{
 		return -1;
 	}
@@ -965,18 +965,25 @@ int64 FStreamReaderMKV::MKVReadData(void* InDestinationBuffer, int64 InNumBytesT
 	return nb;
 }
 
-int64 FStreamReaderMKV::MKVGetCurrentFileOffset() const
+int64 FStreamReaderMKV::GetCurrentOffset() const
 {
 	return ReadBuffer.GetStartOffset() + ReadBuffer.GetCurrentOffset();
 }
 
-int64 FStreamReaderMKV::MKVGetTotalSize()
+int64 FStreamReaderMKV::GetTotalSize() const
 {
 	return ReadBuffer.GetTotalSize();
 }
-bool FStreamReaderMKV::MKVHasReadBeenAborted() const
+
+bool FStreamReaderMKV::HasReadBeenAborted() const
 {
 	return HasBeenAborted();
 }
 
+bool FStreamReaderMKV::HasReachedEOF() const
+{
+	check(!"this should not be called");
+	return false;
+}
+
 } // namespace Electra
diff --git a/Engine/Plugins/Media/ElectraPlayer/Source/ElectraPlayerRuntime/Private/Runtime/Player/mkv/StreamReaderMKV.h b/Engine/Plugins/Media/ElectraPlayer/Source/ElectraPlayerRuntime/Private/Runtime/Player/mkv/StreamReaderMKV.h
--- a/Engine/Plugins/Media/ElectraPlayer/Source/ElectraPlayerRuntime/Private/Runtime/Player/mkv/StreamReaderMKV.h
+++ b/Engine/Plugins/Media/ElectraPlayer/Source/ElectraPlayerRuntime/Private/Runtime/Player/mkv/StreamReaderMKV.h
@@ -103,7 +103,7 @@ public:
 /**
  * This class implements an interface to read from an mkv/webm file.
  */
-class FStreamReaderMKV : public IStreamReader, public FMediaThread, public IParserMKV::IReader
+class FStreamReaderMKV : public IStreamReader, public FMediaThread, IGenericDataReader
 {
 public:
 	FStreamReaderMKV() = default;
@@ -114,11 +114,12 @@ public:
 	void CancelRequest(EStreamType StreamType, bool bSilent) override;
 	void CancelRequests() override;
 
-	// Methods from IParserMKV::IReader
-	int64 MKVReadData(void* InDestinationBuffer, int64 InNumBytesToRead, int64 InFromOffset) override;
-	int64 MKVGetCurrentFileOffset() const override;
-	int64 MKVGetTotalSize() override;
-	bool MKVHasReadBeenAborted() const override;
+	// Methods from IGenericDataReader
+	int64 ReadData(void* InDestinationBuffer, int64 InNumBytesToRead, int64 InFromOffset) override;
+	int64 GetCurrentOffset() const override;
+	int64 GetTotalSize() const override;
+	bool HasReadBeenAborted() const override;
+	bool HasReachedEOF() const override;
 
 private:
 	void WorkerThread();
diff --git a/Engine/Plugins/Media/ElectraPlayer/Source/ElectraPlayerRuntime/Private/Runtime/Player/mp4/PlaylistReaderMP4.cpp b/Engine/Plugins/Media/ElectraPlayer/Source/ElectraPlayerRuntime/Private/Runtime/Player/mp4/PlaylistReaderMP4.cpp
--- a/Engine/Plugins/Media/ElectraPlayer/Source/ElectraPlayerRuntime/Private/Runtime/Player/mp4/PlaylistReaderMP4.cpp
+++ b/Engine/Plugins/Media/ElectraPlayer/Source/ElectraPlayerRuntime/Private/Runtime/Player/mp4/PlaylistReaderMP4.cpp
@@ -31,7 +31,7 @@ namespace Electra
 /**
  * This class is responsible for downloading the mp4 non-mdat boxes and parsing them.
  */
-class FPlaylistReaderMP4 : public IPlaylistReaderMP4, public IParserISO14496_12::IReader, public IParserISO14496_12::IBoxCallback, public FMediaThread
+class FPlaylistReaderMP4 : public IPlaylistReaderMP4, IGenericDataReader, public IParserISO14496_12::IBoxCallback, public FMediaThread
 {
 public:
 	FPlaylistReaderMP4();
@@ -76,14 +76,15 @@ public:
 	virtual TSharedPtrTS<IManifest> GetManifest() override;
 
 private:
-	// Methods from IParserISO14496_12::IReader
-	virtual int64 ReadData(void* IntoBuffer, int64 NumBytesToRead) override;
-	virtual bool HasReachedEOF() const override;
-	virtual bool HasReadBeenAborted() const override;
-	virtual int64 GetCurrentOffset() const override;
+	// Methods from IGenericDataReader
+	int64 ReadData(void* IntoBuffer, int64 NumBytesToRead, int64 InFromOffset) override;
+	bool HasReachedEOF() const override;
+	bool HasReadBeenAborted() const override;
+	int64 GetCurrentOffset() const override;
+	int64 GetTotalSize() const override;
 	// Methods from IParserISO14496_12::IBoxCallback
-	virtual IParserISO14496_12::IBoxCallback::EParseContinuation OnFoundBox(IParserISO14496_12::FBoxType Box, int64 BoxSizeInBytes, int64 FileDataOffset, int64 BoxDataOffset) override;
-	virtual IParserISO14496_12::IBoxCallback::EParseContinuation OnEndOfBox(IParserISO14496_12::FBoxType Box, int64 BoxSizeInBytes, int64 FileDataOffset, int64 BoxDataOffset) override;
+	IParserISO14496_12::IBoxCallback::EParseContinuation OnFoundBox(IParserISO14496_12::FBoxType Box, int64 BoxSizeInBytes, int64 FileDataOffset, int64 BoxDataOffset) override;
+	IParserISO14496_12::IBoxCallback::EParseContinuation OnEndOfBox(IParserISO14496_12::FBoxType Box, int64 BoxSizeInBytes, int64 FileDataOffset, int64 BoxDataOffset) override;
 
 	void StartWorkerThread();
 	void StopWorkerThread();
@@ -325,7 +326,7 @@ void FPlaylistReaderMP4::WorkerThread()
  * @param NumBytesToRead The number of bytes to read. Must not read more bytes and no less than requested.
  * @return The number of bytes read or -1 on a read error.
  */
-int64 FPlaylistReaderMP4::ReadData(void* IntoBuffer, int64 NumBytesToRead)
+int64 FPlaylistReaderMP4::ReadData(void* IntoBuffer, int64 NumBytesToRead, int64 InFromOffset)
 {
 	// We have all the data available.
 	check(ParseBuffer.IsValid());
@@ -375,6 +376,13 @@ int64 FPlaylistReaderMP4::GetCurrentOffset() const
 	return ParsePos;
 }
 
+int64 FPlaylistReaderMP4::GetTotalSize() const
+{
+	check(!"this should not be called");
+	return -1;
+}
+
+
 IParserISO14496_12::IBoxCallback::EParseContinuation FPlaylistReaderMP4::OnFoundBox(IParserISO14496_12::FBoxType Box, int64 BoxSizeInBytes, int64 FileDataOffset, int64 BoxDataOffset)
 {
 	return IParserISO14496_12::IBoxCallback::EParseContinuation::Continue;
