Media: Added passing through of timecodes embedded in a video stream

#rnx

diff --git a/Engine/Plugins/Media/ElectraPlayer/Source/ElectraPlayerRuntime/Private/Runtime/Decoder/BitstreamProcessors/VideoDecoderInputBitstreamProcessorH264.cpp b/Engine/Plugins/Media/ElectraPlayer/Source/ElectraPlayerRuntime/Private/Runtime/Decoder/BitstreamProcessors/VideoDecoderInputBitstreamProcessorH264.cpp
--- a/Engine/Plugins/Media/ElectraPlayer/Source/ElectraPlayerRuntime/Private/Runtime/Decoder/BitstreamProcessors/VideoDecoderInputBitstreamProcessorH264.cpp
+++ b/Engine/Plugins/Media/ElectraPlayer/Source/ElectraPlayerRuntime/Private/Runtime/Decoder/BitstreamProcessors/VideoDecoderInputBitstreamProcessorH264.cpp
@@ -3,7 +3,7 @@
 #include "VideoDecoderInputBitstreamProcessorH264.h"
 
 #include "Utilities/Utilities.h"
-#include "Utils/MPEG/ElectraUtilsMPEGVideo.h"
+#include "Utils/MPEG/ElectraUtilsMPEGVideo_H264.h"
 #include "StreamAccessUnitBuffer.h"
 #include "Decoder/VideoDecoderHelpers.h"
 
@@ -29,16 +29,19 @@ private:
 	{
 	public:
 		virtual ~FCodecSpecificMessages() = default;
-		TArray<ElectraDecodersUtil::MPEG::FISO14496_10_seq_parameter_set_data> SPSs;
+		TMap<uint32, ElectraDecodersUtil::MPEG::H264::FSequenceParameterSet> SPSs;
+		TArray<ElectraDecodersUtil::MPEG::FSEIMessage> SEIMessages;
 	};
 
 	bool bReplaceLengthWithStartcode = true;
 
 	TSharedPtr<const FAccessUnit::CodecData, ESPMode::ThreadSafe> PreviousCodecData;
 	TSharedPtr<const FAccessUnit::CodecData, ESPMode::ThreadSafe> CurrentCodecData;
-	TArray<ElectraDecodersUtil::MPEG::FISO14496_10_seq_parameter_set_data> SPSs;
+	TMap<uint32, ElectraDecodersUtil::MPEG::H264::FSequenceParameterSet> SPSs;
 
 	MPEG::FColorimetryHelper Colorimetry;
+
+	bool HandlePicTimingSEI(FParamDict* InOutProperties, const ElectraDecodersUtil::MPEG::FSEIMessage& InSEI, const TMap<uint32, ElectraDecodersUtil::MPEG::H264::FSequenceParameterSet>& InSPSs);
 };
 
 
@@ -77,6 +80,13 @@ IVideoDecoderInputBitstreamProcessor::EProcessResult FVideoDecoderInputBitstream
 	}
 	check(InOutAccessUnit->ESType == EStreamType::Video);
 
+	TSharedPtr<FCodecSpecificMessages> Msgs = StaticCastSharedPtr<FCodecSpecificMessages>(OutBSI.CodecSpecificMessages);
+	if (!Msgs.IsValid())
+	{
+		Msgs = MakeShared<FCodecSpecificMessages>();
+		OutBSI.CodecSpecificMessages = Msgs;
+	}
+
 	//
 	// Extract sequence parameter sets from the codec specific data.
 	//
@@ -97,30 +107,20 @@ IVideoDecoderInputBitstreamProcessor::EProcessResult FVideoDecoderInputBitstream
 				// SPS?
 				if (nal_unit_type == 7)
 				{
-					ElectraDecodersUtil::MPEG::FISO14496_10_seq_parameter_set_data sps;
-					if (ElectraDecodersUtil::MPEG::ParseH264SPS(sps, NALU, NALUs[i].Size))
-					{
-						SPSs.Emplace(MoveTemp(sps));
-					}
+					bool bOk = ElectraDecodersUtil::MPEG::H264::ParseSequenceParameterSet(SPSs, NALU, NALUs[i].Size);
+					check(bOk); (void)bOk;
 				}
 			}
-
-			// Create the message structure if it does not exist yet.
-			if (!OutBSI.CodecSpecificMessages.IsValid())
-			{
-				OutBSI.CodecSpecificMessages = MakeShared<FCodecSpecificMessages>();
-			}
-			// Set the new SPS in the message structure.
-			StaticCastSharedPtr<FCodecSpecificMessages>(OutBSI.CodecSpecificMessages)->SPSs = SPSs;
-
 			Result = IVideoDecoderInputBitstreamProcessor::EProcessResult::CSDChanged;
 		}
 		PreviousCodecData = CurrentCodecData.IsValid() ? CurrentCodecData : InOutAccessUnit->AUCodecData;
 		CurrentCodecData = InOutAccessUnit->AUCodecData;
 	}
 
-	// NOTE: In a second phase we should probably scan the input access unit for inband SPS if the codec is avc3.
+	// Set the messages from the CSD on the access unit.
+	Msgs->SPSs = SPSs;
 
+	// NOTE: In a second phase we should probably scan the input access unit for inband SPS if the codec is avc3.
 
 	// Now go over the NALUs in the access unit and see what is there.
 	OutBSI.bIsDiscardable = true;
@@ -129,6 +129,8 @@ IVideoDecoderInputBitstreamProcessor::EProcessResult FVideoDecoderInputBitstream
 	uint32* LastNALU = (uint32*)Electra::AdvancePointer(NALU, InOutAccessUnit->AUSize);
 	while(NALU < LastNALU)
 	{
+		uint32 naluLen = MEDIA_FROM_BIG_ENDIAN(*NALU);
+
 		// Check the nal_ref_idc in the NAL unit for dependencies.
 		uint8 nal = *(const uint8*)(NALU + 1);
 		check((nal & 0x80) == 0);
@@ -144,15 +146,14 @@ IVideoDecoderInputBitstreamProcessor::EProcessResult FVideoDecoderInputBitstream
 		// SEI message(s)?
 		if ((nal & 0x1f) == 6)
 		{
-			// NOTE: We could get closed captions here.
+			ElectraDecodersUtil::MPEG::ExtractSEIMessages(Msgs->SEIMessages, Electra::AdvancePointer(NALU, 5), naluLen-1, ElectraDecodersUtil::MPEG::ESEIStreamType::H264, false);
 		}
 
-		uint32 naluLen = MEDIA_FROM_BIG_ENDIAN(*NALU) + 4;
 		if (bReplaceLengthWithStartcode)
 		{
 			*NALU = MEDIA_TO_BIG_ENDIAN(0x00000001U);
 		}
-		NALU = Electra::AdvancePointer(NALU, naluLen);
+		NALU = Electra::AdvancePointer(NALU, naluLen + 4);
 	}
 
 	return Result;
@@ -167,31 +168,146 @@ void FVideoDecoderInputBitstreamProcessorH264::SetPropertiesOnOutput(TSharedPtr<
 	}
 
 	TSharedPtr<FCodecSpecificMessages> Msg = StaticCastSharedPtr<FCodecSpecificMessages>(InFromBSI.CodecSpecificMessages);
-
-	const TArray<ElectraDecodersUtil::MPEG::FISO14496_10_seq_parameter_set_data>& SPSList = Msg.IsValid() ? Msg->SPSs : SPSs;
+	const TMap<uint32, ElectraDecodersUtil::MPEG::H264::FSequenceParameterSet>& SPSMap = Msg.IsValid() ? Msg->SPSs : SPSs;
 
 	// We only interact with the first SPS.
-	if (SPSList.Num() > 0)
+	if (SPSMap.Num() > 0)
 	{
+		auto It = SPSMap.CreateConstIterator();
+		const ElectraDecodersUtil::MPEG::H264::FSequenceParameterSet& sps = It->Value;
 		// Set the bit depth and the colorimetry.
 		uint8 colour_primaries=2, transfer_characteristics=2, matrix_coeffs=2;
 		uint8 video_full_range_flag=0, video_format=5;
-		if (SPSList[0].colour_description_present_flag)
+		if (sps.colour_description_present_flag)
 		{
-			colour_primaries = SPSList[0].colour_primaries;
-			transfer_characteristics = SPSList[0].transfer_characteristics;
-			matrix_coeffs = SPSList[0].matrix_coefficients;
+			colour_primaries = sps.colour_primaries;
+			transfer_characteristics = sps.transfer_characteristics;
+			matrix_coeffs = sps.matrix_coefficients;
 		}
-		if (SPSList[0].video_signal_type_present_flag)
+		if (sps.video_signal_type_present_flag)
 		{
-			video_full_range_flag = SPSList[0].video_full_range_flag;
-			video_format = SPSList[0].video_format;
+			video_full_range_flag = sps.video_full_range_flag;
+			video_format = sps.video_format;
 		}
 
 		Colorimetry.Update(colour_primaries, transfer_characteristics, matrix_coeffs, video_full_range_flag, video_format);
 		Colorimetry.UpdateParamDict(*InOutProperties);
 	}
+
+	// Handle SEI messages we are interested in
+	for(int32 i=0, iMax=Msg.IsValid()?Msg->SEIMessages.Num():0; i<iMax; ++i)
+	{
+		const ElectraDecodersUtil::MPEG::FSEIMessage& sei = Msg->SEIMessages[i];
+		switch(sei.PayloadType)
+		{
+			case 1:	// pic_timing()
+			{
+				HandlePicTimingSEI(InOutProperties, sei, SPSMap);
+				break;
+			}
+			default:
+			{
+				break;
+			}
+		}
+	}
 }
 
+bool FVideoDecoderInputBitstreamProcessorH264::HandlePicTimingSEI(FParamDict* InOutProperties, const ElectraDecodersUtil::MPEG::FSEIMessage& InSEI, const TMap<uint32, ElectraDecodersUtil::MPEG::H264::FSequenceParameterSet>& InSPSs)
+{
+	// Parsing the pic_timing() SEI message requires the active SPS.
+	// NOTE: We do NOT know which SPS is active, if there are several as the activation is determined by the slice being decoded.
+	if (InSPSs.Num() != 1)
+	{
+		return false;
+	}
+
+	auto It = InSPSs.CreateConstIterator();
+	const ElectraDecodersUtil::MPEG::H264::FSequenceParameterSet& sps = It->Value;
+
+	ElectraDecodersUtil::MPEG::H264::FBitstreamReader br(InSEI.Message.GetData(), InSEI.Message.Num());
+
+	IVideoDecoderTimecode::FMPEGDefinition ClockTimestamp[3];
+
+	const bool CpbDpbDelaysPresentFlag = sps.nal_hrd_parameters_present_flag || sps.vcl_hrd_parameters_present_flag;
+	if (CpbDpbDelaysPresentFlag)
+	{
+		int32 NumBits1 = 1 + (sps.nal_hrd_parameters_present_flag ? sps.nal_hrd_parameters.cpb_removal_delay_length_minus1 : sps.vcl_hrd_parameters.cpb_removal_delay_length_minus1);
+		int32 NumBits2 = 1 + (sps.nal_hrd_parameters_present_flag ? sps.nal_hrd_parameters.dpb_output_delay_length_minus1 : sps.vcl_hrd_parameters.dpb_output_delay_length_minus1);
+		uint32 cpb_removal_delay = br.GetBits(NumBits1);
+		uint32 dpb_output_delay = br.GetBits(NumBits2);
+		(void)cpb_removal_delay;
+		(void)dpb_output_delay;
+	}
+	if (sps.pic_struct_present_flag)
+	{
+		static const uint8 kNumClockTS[9] = { 1, 1, 1, 2, 2, 3, 3, 2, 3 };
+		uint32 pic_struct = br.GetBits(4);
+		if (pic_struct > 8)
+		{
+			return false;
+		}
+		for(int32 nc=0,ncMax=kNumClockTS[pic_struct]; nc<ncMax; ++nc)
+		{
+			IVideoDecoderTimecode::FMPEGDefinition& ct(ClockTimestamp[nc]);
+			ct.FromH26x = 4;
+			ct.clock_timestamp_flag = (uint8) br.GetBits(1);
+			if (ct.clock_timestamp_flag)
+			{
+				// Set the timing values from the SPS.
+				ct.timing_info_present_flag = sps.timing_info_present_flag;
+				ct.num_units_in_tick = sps.num_units_in_tick;
+				ct.time_scale = sps.time_scale;
+
+				// Read the values from the message.
+				ct.ct_type = (uint8) br.GetBits(2);
+				ct.nuit_field_based_flag = (uint8) br.GetBits(1);
+				ct.counting_type = (uint8) br.GetBits(5);
+				ct.full_timestamp_flag = (uint8) br.GetBits(1);
+				ct.discontinuity_flag = (uint8) br.GetBits(1);
+				ct.cnt_dropped_flag = (uint8) br.GetBits(1);
+				ct.n_frames = (uint16) br.GetBits(8);
+				if (ct.full_timestamp_flag)
+				{
+					ct.seconds_value = (uint8) br.GetBits(6);
+					ct.minutes_value = (uint8) br.GetBits(6);
+					ct.hours_value = (uint8) br.GetBits(5);
+				}
+				else
+				{
+					// seconds_flag
+					if (br.GetBits(1))
+					{
+						ct.seconds_value = (uint8) br.GetBits(6);
+						// minutes_flag
+						if (br.GetBits(1))
+						{
+							ct.minutes_value = (uint8) br.GetBits(6);
+							// hours_flag
+							if (br.GetBits(1))
+							{
+								ct.hours_value = (uint8) br.GetBits(5);
+							}
+						}
+					}
+				}
+				uint32 time_offset_length = sps.nal_hrd_parameters_present_flag ? sps.nal_hrd_parameters.time_offset_length : sps.vcl_hrd_parameters_present_flag ? sps.vcl_hrd_parameters.time_offset_length : 24;
+				uint32 time_offset = br.GetBits(time_offset_length);
+				ct.time_offset = (int32)(time_offset << (32 - time_offset_length)) >> (32 - time_offset_length);
+
+				if (ct.timing_info_present_flag)
+				{
+					ct.clockTimestamp = (int64)(((uint32)ct.hours_value * 60 + ct.minutes_value) * 60 + ct.seconds_value) * ct.time_scale + ct.n_frames * (ct.num_units_in_tick * (ct.nuit_field_based_flag + 1)) + ct.time_offset;
+				}
+			}
+		}
+
+		TSharedPtr<MPEG::FVideoDecoderTimecode, ESPMode::ThreadSafe> NewTimecode = MakeShared<MPEG::FVideoDecoderTimecode, ESPMode::ThreadSafe>();
+		// Set from the first clock since we are only dealing with progressive frames, not interlaced fields.
+		NewTimecode->Update(ClockTimestamp[0]);
+		InOutProperties->Set(IDecoderOutputOptionNames::Timecode, FVariantValue(NewTimecode));
+	}
+	return true;
+}
 
 } // namespace Electra
diff --git a/Engine/Plugins/Media/ElectraPlayer/Source/ElectraPlayerRuntime/Private/Runtime/Decoder/VideoDecoderHelpers.h b/Engine/Plugins/Media/ElectraPlayer/Source/ElectraPlayerRuntime/Private/Runtime/Decoder/VideoDecoderHelpers.h
--- a/Engine/Plugins/Media/ElectraPlayer/Source/ElectraPlayerRuntime/Private/Runtime/Decoder/VideoDecoderHelpers.h
+++ b/Engine/Plugins/Media/ElectraPlayer/Source/ElectraPlayerRuntime/Private/Runtime/Decoder/VideoDecoderHelpers.h
@@ -13,6 +13,20 @@ namespace Electra
 namespace MPEG
 {
 
+class FVideoDecoderTimecode : public IVideoDecoderTimecode
+{
+public:
+	virtual ~FVideoDecoderTimecode() = default;
+	virtual FMPEGDefinition const* GetMPEGDefinition() const override
+	{ return &DecoderTimecode; }
+	void Update(const IVideoDecoderTimecode::FMPEGDefinition& InTimecode)
+	{ DecoderTimecode = InTimecode; }
+private:
+	IVideoDecoderTimecode::FMPEGDefinition DecoderTimecode;
+};
+
+
+
 class FColorimetryHelper
 {
 public:
diff --git a/Engine/Plugins/Media/ElectraUtil/Source/ElectraBase/Public/MediaDecoderOutput.h b/Engine/Plugins/Media/ElectraUtil/Source/ElectraBase/Public/MediaDecoderOutput.h
--- a/Engine/Plugins/Media/ElectraUtil/Source/ElectraBase/Public/MediaDecoderOutput.h
+++ b/Engine/Plugins/Media/ElectraUtil/Source/ElectraBase/Public/MediaDecoderOutput.h
@@ -355,4 +355,5 @@ static const FName AspectH(TEXT("aspect_h"));
 static const FName FPSNumerator(TEXT("fps_num"));
 static const FName FPSDenominator(TEXT("fps_denom"));
 static const FName PixelDataScale(TEXT("pix_datascale"));
+static const FName Timecode(TEXT("timecode"));
 }
diff --git a/Engine/Plugins/Media/ElectraUtil/Source/ElectraBase/Public/MediaVideoDecoderOutput.h b/Engine/Plugins/Media/ElectraUtil/Source/ElectraBase/Public/MediaVideoDecoderOutput.h
--- a/Engine/Plugins/Media/ElectraUtil/Source/ElectraBase/Public/MediaVideoDecoderOutput.h
+++ b/Engine/Plugins/Media/ElectraUtil/Source/ElectraBase/Public/MediaVideoDecoderOutput.h
@@ -43,6 +43,40 @@ protected:
 
 
 
+class IVideoDecoderTimecode : public TSharedFromThis<IVideoDecoderTimecode, ESPMode::ThreadSafe>
+{
+public:
+	struct FMPEGDefinition
+	{
+		// Calculated value as per:
+		//   clockTimestamp = ( ( hH * 60 + mM ) * 60 + sS ) * time_scale + nFrames * ( num_units_in_tick * ( 1 + nuit_field_based_flag ) ) + tOffset,
+		// can only be valid when there is timing information.
+		int64 clockTimestamp = 0;
+		// Values from a pic_timing() SEI in H.264 or from a time_code() SEI in H.265
+		uint32 num_units_in_tick = 0;			// from the SPS
+		uint32 time_scale = 0;					// from the SPS
+		int32 time_offset = 0;
+		uint16 n_frames = 0;
+		uint8 timing_info_present_flag = 0;		// from the SPS
+		uint8 clock_timestamp_flag = 0;
+		uint8 ct_type = 0;
+		uint8 nuit_field_based_flag = 0;
+		uint8 counting_type = 0;
+		uint8 full_timestamp_flag = 0;
+		uint8 discontinuity_flag = 0;
+		uint8 cnt_dropped_flag = 0;
+		uint8 seconds_value = 0;
+		uint8 minutes_value = 0;
+		uint8 hours_value = 0;
+		uint8 FromH26x = 0;		// last digit of the codec this comes from. 4=H.264, 5=H.265, etc. In case values need different interpretation.
+	};
+
+	virtual ~IVideoDecoderTimecode() = default;
+	virtual FMPEGDefinition const* GetMPEGDefinition() const = 0;
+};
+
+
+
 class IVideoDecoderColorimetry : public TSharedFromThis<IVideoDecoderColorimetry, ESPMode::ThreadSafe>
 {
 public:
@@ -289,6 +323,15 @@ public:
 		return nullptr;
 	}
 
+	virtual TSharedPtr<const IVideoDecoderTimecode, ESPMode::ThreadSafe> GetTimecode() const
+	{
+		if (ParamDict && ParamDict->HaveKey(IDecoderOutputOptionNames::Timecode))
+		{
+			return ParamDict->GetValue(IDecoderOutputOptionNames::Timecode).GetSharedPointer<const IVideoDecoderTimecode>();
+		}
+		return nullptr;
+	}
+
 protected:
 	FVideoDecoderOutput() = default;
 
diff --git a/Engine/Plugins/Media/ElectraUtil/Source/ElectraSamples/Private/ElectraTextureSampleBase.cpp b/Engine/Plugins/Media/ElectraUtil/Source/ElectraSamples/Private/ElectraTextureSampleBase.cpp
--- a/Engine/Plugins/Media/ElectraUtil/Source/ElectraSamples/Private/ElectraTextureSampleBase.cpp
+++ b/Engine/Plugins/Media/ElectraUtil/Source/ElectraSamples/Private/ElectraTextureSampleBase.cpp
@@ -11,11 +11,36 @@
 
 // -------------------------------------------------------------------------------------------------------------------------------------------------------
 
+namespace
+{
+
+static TOptional<FTimecode> CreateTimecodeFromMPEGDefinition(TOptional<FFrameRate>& OutFramerate, const IVideoDecoderTimecode::FMPEGDefinition* InMPEGTimecode)
+{
+	if (InMPEGTimecode->timing_info_present_flag)
+	{
+		const FTimespan ts(Electra::FTimeValue(InMPEGTimecode->clockTimestamp, InMPEGTimecode->time_scale).GetAsTimespan());
+		OutFramerate = FFrameRate(InMPEGTimecode->time_scale, InMPEGTimecode->num_units_in_tick);
+		return FTimecode::FromTimespan(ts, OutFramerate.GetValue(), InMPEGTimecode->ct_type > 1, false);
+	}
+	return TOptional<FTimecode>();
+}
+
+}
+
+// -------------------------------------------------------------------------------------------------------------------------------------------------------
+
 void IElectraTextureSampleBase::Initialize(FVideoDecoderOutput* InVideoDecoderOutput)
 {
 	VideoDecoderOutput = StaticCastSharedPtr<FVideoDecoderOutput, IDecoderOutputPoolable, ESPMode::ThreadSafe>(InVideoDecoderOutput->AsShared());
 	HDRInfo = VideoDecoderOutput->GetHDRInformation();
-	Colorimetry = VideoDecoderOutput->GetColorimetry();
+	TSharedPtr<const IVideoDecoderTimecode, ESPMode::ThreadSafe> TimecodePtr = VideoDecoderOutput->GetTimecode();
+
+	if (TimecodePtr.IsValid())
+	{
+		// Store this in case this is needed again.
+		DecoderTimecode = TimecodePtr;
+		Timecode = CreateTimecodeFromMPEGDefinition(Framerate, TimecodePtr->GetMPEGDefinition());
+	}
 
 	// Get various basic MP4-style colorimetry values (we default to video range Rec709 SDR)
 	bool bFullRange = false;
diff --git a/Engine/Plugins/Media/ElectraUtil/Source/ElectraSamples/Public/IElectraTextureSample.h b/Engine/Plugins/Media/ElectraUtil/Source/ElectraSamples/Public/IElectraTextureSample.h
--- a/Engine/Plugins/Media/ElectraUtil/Source/ElectraSamples/Public/IElectraTextureSample.h
+++ b/Engine/Plugins/Media/ElectraUtil/Source/ElectraSamples/Public/IElectraTextureSample.h
@@ -5,6 +5,8 @@
 #include "IMediaTextureSampleConverter.h"
 #include "MediaObjectPool.h"
 #include "Misc/Timespan.h"
+#include "Misc/Timecode.h"
+#include "Misc/Optional.h"
 #include "Templates/SharedPointer.h"
 #include "Templates/RefCounting.h"
 
@@ -33,6 +35,15 @@ public:
 	virtual FMediaTimeStamp GetTime() const override;
 	virtual FTimespan GetDuration() const override;
 
+	virtual TOptional<FTimecode> GetTimecode() const override
+	{
+		return Timecode;
+	}
+	virtual TOptional<FFrameRate> GetFramerate() const override
+	{
+		return Framerate;
+	}
+
 	virtual double GetAspectRatio() const override
 	{
 		return VideoDecoderOutput->GetAspectRatio();
@@ -68,6 +79,10 @@ protected:
 	/** Quick access for some HDR related info */
 	TWeakPtr<const IVideoDecoderHDRInformation, ESPMode::ThreadSafe> HDRInfo;
 	TWeakPtr<const IVideoDecoderColorimetry, ESPMode::ThreadSafe> Colorimetry;
+	/** Optional timecode */
+	TWeakPtr<const IVideoDecoderTimecode> DecoderTimecode;
+	TOptional<FTimecode> Timecode;
+	TOptional<FFrameRate> Framerate;
 
 	/** YUV matrix, adjusted to compensate for decoder output specific scale */
 	FMatrix44f SampleToRgbMtx;
diff --git a/Engine/Source/Runtime/Media/Public/IMediaTextureSample.h b/Engine/Source/Runtime/Media/Public/IMediaTextureSample.h
--- a/Engine/Source/Runtime/Media/Public/IMediaTextureSample.h
+++ b/Engine/Source/Runtime/Media/Public/IMediaTextureSample.h
@@ -12,6 +12,7 @@
 #include "Misc/Optional.h"
 #include "Misc/Timecode.h"
 #include "Misc/Timespan.h"
+#include "Misc/FrameRate.h"
 #include "Templates/SharedPointer.h"
 
 #include "HDRHelper.h"
@@ -263,8 +264,8 @@ public:
 	 * @return texture sample converter
 	 */
 	virtual IMediaTextureSampleConverter* GetMediaTextureSampleConverter()
-	{ 
-		return nullptr; 
+	{
+		return nullptr;
 	}
 
 	/**
@@ -273,8 +274,8 @@ public:
 	 * @return texture sample color converter
 	 */
 	virtual IMediaTextureSampleColorConverter* GetMediaTextureSampleColorConverter()
-	{ 
-		return nullptr; 
+	{
+		return nullptr;
 	}
 
 #endif //WITH_ENGINE
@@ -297,6 +298,15 @@ public:
 	 */
 	virtual TOptional<FTimecode> GetTimecode() const { return TOptional<FTimecode>(); }
 
+	/**
+	 * Get the sample framerate if available. This is the rate in which the timecode
+	 * is measured. It is not necessarily the display frame rate.
+	 *
+	 * @return Sample framerate. May be needed for converting GetTimecode().
+	 * @see GetTime
+	 */
+	virtual TOptional<FFrameRate> GetFramerate() const { return TOptional<FFrameRate>(); }
+
 	/**
 	 * Whether the sample can be held in a cache.
 	 *
@@ -380,7 +390,7 @@ public:
 	{
 		return false;
 	}
-	
+
 	/**
 	* Get complete 4x4 matrix to apply to the sample's pixels to yield RGB data in the sample's gamut
 	 *
@@ -407,7 +417,7 @@ public:
 
 	/**
 	 * Get Colorspace conversion matrix to convert to CIE1931 XYZ space
-	 * 
+	 *
 	 * @return Conversion Matrix
 	 */
 	virtual FMatrix44d GetGamutToXYZMatrix() const
@@ -417,7 +427,7 @@ public:
 
 	/**
 	 * Get white point of color space of the data contain in the sample
-	 * 
+	 *
 	 * @return White point
 	 */
 	virtual FVector2d GetWhitePoint() const
@@ -476,7 +486,7 @@ public:
 	 * Reset sample to empty state
 	 */
 	virtual void Reset() { }
-	
+
 public:
 
 	/** Virtual destructor. */
diff --git a/Engine/Source/Runtime/MediaAssets/Private/Assets/MediaPlayer.cpp b/Engine/Source/Runtime/MediaAssets/Private/Assets/MediaPlayer.cpp
--- a/Engine/Source/Runtime/MediaAssets/Private/Assets/MediaPlayer.cpp
+++ b/Engine/Source/Runtime/MediaAssets/Private/Assets/MediaPlayer.cpp
@@ -396,6 +396,12 @@ UMediaTimeStampInfo* UMediaPlayer::GetDisplayTimeStamp() const
 	return TimeStampInfo;
 }
 
+TOptional<FTimecode> UMediaPlayer::GetVideoTimecode() const
+{
+	return PlayerFacade->GetVideoTimecode();
+}
+
+
 FText UMediaPlayer::GetTrackDisplayName(EMediaPlayerTrack TrackType, int32 TrackIndex) const
 {
 	return PlayerFacade->GetTrackDisplayName((EMediaTrackType)TrackType, TrackIndex);
diff --git a/Engine/Source/Runtime/MediaAssets/Public/MediaPlayer.h b/Engine/Source/Runtime/MediaAssets/Public/MediaPlayer.h
--- a/Engine/Source/Runtime/MediaAssets/Public/MediaPlayer.h
+++ b/Engine/Source/Runtime/MediaAssets/Public/MediaPlayer.h
@@ -16,6 +16,7 @@
 #include "UObject/ScriptMacros.h"
 #include "Misc/Guid.h"
 #include "Misc/Variant.h"
+#include "Misc/Timecode.h"
 #include "Engine/LatentActionManager.h"
 #include "MediaPlayerOptions.h"
 #include "IMediaTimeSource.h"
@@ -413,6 +414,19 @@ public:
 	UFUNCTION(BlueprintCallable, Category = "Media|MediaPlayer")
 	MEDIAASSETS_API UMediaTimeStampInfo* GetDisplayTimeStamp() const;
 
+	/**
+	 * Get the timecode from the video track, if one is available.
+	 * If available, the returned value will be the timecode of the sample most recently
+	 * sent into the video sink.
+	 *
+	 * @return Most recently delivered video timecode from the video track.
+	 *
+	 * @note Even if timecode is available, the value returned reflects the timecode
+	 *       that most recently delivered to the video sink. During a seek operation
+	 *       the timecode will be unset.
+	 */
+	MEDIAASSETS_API TOptional<FTimecode> GetVideoTimecode() const;
+
 	/**
 	 * Get the human readable name of the specified track.
 	 *
diff --git a/Engine/Source/Runtime/MediaUtils/Private/MediaPlayerFacade.cpp b/Engine/Source/Runtime/MediaUtils/Private/MediaPlayerFacade.cpp
--- a/Engine/Source/Runtime/MediaUtils/Private/MediaPlayerFacade.cpp
+++ b/Engine/Source/Runtime/MediaUtils/Private/MediaPlayerFacade.cpp
@@ -495,6 +495,13 @@ FMediaTimeStamp FMediaPlayerFacade::GetDisplayTimeStamp() const
 	return GetTimeStampInternal(true);
 }
 
+TOptional<FTimecode> FMediaPlayerFacade::GetVideoTimecode() const
+{
+	FScopeLock Lock(&LastTimeValuesCS);
+	return MostRecentlyDeliveredVideoFrameTimecode;
+}
+
+
 FMediaTimeStamp FMediaPlayerFacade::GetTimeStampInternal(bool bForDisplay) const
 {
 	TSharedPtr<IMediaPlayer, ESPMode::ThreadSafe> CurrentPlayer(Player);
@@ -1683,6 +1690,7 @@ void FMediaPlayerFacade::Flush(bool bExcludePlayer, bool bOnSeek)
 	MetadataSampleSinks.Flush(RawMediaPlayer);
 	SubtitleSampleSinks.Flush(RawMediaPlayer);
 	VideoSampleSinks.Flush(RawMediaPlayer);
+	MostRecentlyDeliveredVideoFrameTimecode.Reset();
 
 	if (Player.IsValid() && !bExcludePlayer)
 	{
@@ -3225,6 +3233,7 @@ bool FMediaPlayerFacade::ProcessVideoSamples(IMediaSamples& Samples, const TRang
 			FScopeLock Lock(&LastTimeValuesCS);
 			CurrentFrameVideoDisplayTimeStamp = CurrentFrameVideoTimeStamp = SampleTimeRange.GetLowerBoundValue();
 			LastVideoSampleProcessedTimeRange = SampleTimeRange;
+			MostRecentlyDeliveredVideoFrameTimecode = Sample->GetTimecode();
 		}
 
 		UpdateSeekStatus(&CurrentFrameVideoTimeStamp);
@@ -3610,6 +3619,7 @@ void FMediaPlayerFacade::ReInit()
 		CurrentFrameVideoTimeStamp.Invalidate();
 		CurrentFrameVideoDisplayTimeStamp.Invalidate();
 		NextEstVideoTimeAtFrameStart.Invalidate();
+		MostRecentlyDeliveredVideoFrameTimecode.Reset();
 		SeekTargetTime.Invalidate();
 		SeekIndex = 0;
 	}
diff --git a/Engine/Source/Runtime/MediaUtils/Public/MediaPlayerFacade.h b/Engine/Source/Runtime/MediaUtils/Public/MediaPlayerFacade.h
--- a/Engine/Source/Runtime/MediaUtils/Public/MediaPlayerFacade.h
+++ b/Engine/Source/Runtime/MediaUtils/Public/MediaPlayerFacade.h
@@ -25,6 +25,7 @@
 #include "Misc/Guid.h"
 #include "Misc/Optional.h"
 #include "Misc/Timespan.h"
+#include "Misc/Timecode.h"
 #include "Misc/Variant.h"
 #include "Templates/Atomic.h"
 #include "Templates/SharedPointer.h"
@@ -392,6 +393,19 @@ public:
 	 */
 	MEDIAUTILS_API FMediaTimeStamp GetDisplayTimeStamp() const;
 
+	/**
+	 * Get the timecode from the video track, if one is available.
+	 * If available, the returned value will be the timecode of the sample most recently
+	 * sent into the video sink.
+	 *
+	 * @return Most recently delivered video timecode from the video track.
+	 *
+	 * @note Even if timecode is available, the value returned reflects the timecode
+	 *       that most recently delivered to the video sink. During a seek operation
+	 *       the timecode will be unset.
+	 */
+	MEDIAUTILS_API TOptional<FTimecode> GetVideoTimecode() const;
+
 	/**
 	 * Get the human readable name of the specified track.
 	 *
@@ -1063,6 +1077,9 @@ private:
 	/** Estimation for next frame's video timestamp (used when no audio present or active in stream) */
 	FMediaTimeStampSample NextEstVideoTimeAtFrameStart;
 
+	/** Timecode of the video frame most recently delivered to the video sinks */
+	TOptional<FTimecode> MostRecentlyDeliveredVideoFrameTimecode;
+
 	/** Timestamp of seek target location if seek is pending */
 	FMediaTimeStamp SeekTargetTime;
 
