Electra: Fixed a PVS warning

#rnx

diff --git a/Engine/Plugins/Media/ElectraPlayer/Source/ElectraPlayerRuntime/Private/Runtime/Decoder/BitstreamProcessors/VideoDecoderInputBitstreamProcessorH264.cpp b/Engine/Plugins/Media/ElectraPlayer/Source/ElectraPlayerRuntime/Private/Runtime/Decoder/BitstreamProcessors/VideoDecoderInputBitstreamProcessorH264.cpp
--- a/Engine/Plugins/Media/ElectraPlayer/Source/ElectraPlayerRuntime/Private/Runtime/Decoder/BitstreamProcessors/VideoDecoderInputBitstreamProcessorH264.cpp
+++ b/Engine/Plugins/Media/ElectraPlayer/Source/ElectraPlayerRuntime/Private/Runtime/Decoder/BitstreamProcessors/VideoDecoderInputBitstreamProcessorH264.cpp
@@ -297,7 +297,7 @@ bool FVideoDecoderInputBitstreamProcessorH264::HandlePicTimingSEI(FParamDict* In
 
 				if (ct.timing_info_present_flag)
 				{
-					ct.clockTimestamp = (int64)(((uint32)ct.hours_value * 60 + ct.minutes_value) * 60 + ct.seconds_value) * ct.time_scale + ct.n_frames * (ct.num_units_in_tick * (ct.nuit_field_based_flag + 1)) + ct.time_offset;
+					ct.clockTimestamp = (((int64)ct.hours_value * 60 + ct.minutes_value) * 60 + ct.seconds_value) * ct.time_scale + ct.n_frames * (ct.num_units_in_tick * (ct.nuit_field_based_flag + 1)) + ct.time_offset;
 				}
 			}
 		}
