Electra: Maintenance in preparation for upcoming changes

#rnx

diff --git a/Engine/Plugins/Media/ElectraCodecs/Source/ElectraDecoders/Private/Utils/MPEG/ElectraUtilsMPEGVideo_H264.cpp b/Engine/Plugins/Media/ElectraCodecs/Source/ElectraDecoders/Private/Utils/MPEG/ElectraUtilsMPEGVideo_H264.cpp
--- a/Engine/Plugins/Media/ElectraCodecs/Source/ElectraDecoders/Private/Utils/MPEG/ElectraUtilsMPEGVideo_H264.cpp
+++ b/Engine/Plugins/Media/ElectraCodecs/Source/ElectraDecoders/Private/Utils/MPEG/ElectraUtilsMPEGVideo_H264.cpp
@@ -617,9 +617,11 @@ namespace ElectraDecodersUtil
 						sps.max_bits_per_mb_denom = br.ue_v();
 						RANGE_CHECK_FAILURE(sps.max_bits_per_mb_denom <= 16, TEXT("max_bits_per_mb_denom"));
 						sps.log2_max_mv_length_horizontal = br.ue_v();
-						RANGE_CHECK_FAILURE(sps.log2_max_mv_length_horizontal <= 15, TEXT("log2_max_mv_length_horizontal"));
+						// The mv lengths should not be larger than 15 according to the standard, but we have
+						// encountered streams where this is not the case, so we remove the checks for better compatibility.
+						//RANGE_CHECK_FAILURE(sps.log2_max_mv_length_horizontal <= 15, TEXT("log2_max_mv_length_horizontal"));
 						sps.log2_max_mv_length_vertical = br.ue_v();
-						RANGE_CHECK_FAILURE(sps.log2_max_mv_length_vertical <= 15, TEXT("log2_max_mv_length_vertical"));
+						//RANGE_CHECK_FAILURE(sps.log2_max_mv_length_vertical <= 15, TEXT("log2_max_mv_length_vertical"));
 						sps.max_num_reorder_frames = br.ue_v();
 						sps.max_dec_frame_buffering = br.ue_v();
 						RANGE_CHECK_FAILURE(sps.max_dec_frame_buffering <= sps.max_num_ref_frames, TEXT("max_dec_frame_buffering"));
diff --git a/Engine/Plugins/Media/ElectraCodecs/Source/ElectraDecoders/Private/Utils/MPEG/ElectraUtilsMPEGVideo_H265.cpp b/Engine/Plugins/Media/ElectraCodecs/Source/ElectraDecoders/Private/Utils/MPEG/ElectraUtilsMPEGVideo_H265.cpp
--- a/Engine/Plugins/Media/ElectraCodecs/Source/ElectraDecoders/Private/Utils/MPEG/ElectraUtilsMPEGVideo_H265.cpp
+++ b/Engine/Plugins/Media/ElectraCodecs/Source/ElectraDecoders/Private/Utils/MPEG/ElectraUtilsMPEGVideo_H265.cpp
@@ -266,6 +266,38 @@ namespace ElectraDecodersUtil
 				return FFractionalValue();
 			}
 
+			uint64 FSequenceParameterSet::GetConstraintFlags() const
+			{
+				return (uint64)((profile_tier_level.general_progressive_source_flag << 3) | (profile_tier_level.general_interlaced_source_flag << 2) | (profile_tier_level.general_non_packed_constraint_flag << 1) | profile_tier_level.general_frame_only_constraint_flag) << 44;
+			}
+
+			FString FSequenceParameterSet::GetRFC6381(const TCHAR* SampleTypePrefix) const
+			{
+				// As per ISO/IEC 14496-15:2014 Annex E.3
+				uint32 cf = BitReverse32(profile_tier_level.general_profile_compatibility_flags);
+				uint64 ConstraintFlags = GetConstraintFlags();
+				FString cfs;
+				bool bNonZero = false;
+				for(int32 i=5; i>=0; --i, ConstraintFlags >>= 8)
+				{
+					if ((ConstraintFlags & 255) == 0 && !bNonZero)
+					{
+						continue;
+					}
+					bNonZero = true;
+					cfs = FString::Printf(TEXT(".%02X"), (uint8)(ConstraintFlags & 255)) + cfs;
+				}
+
+				if (profile_tier_level.general_profile_space == 0)
+				{
+					return FString::Printf(TEXT("%s.%d.%X.%c%d%s"), SampleTypePrefix, profile_tier_level.general_profile_idc, cf, profile_tier_level.general_tier_flag ? TCHAR('H') : TCHAR('L'), profile_tier_level.general_level_idc, *cfs);
+				}
+				else
+				{
+					return FString::Printf(TEXT("%s.%c%d.%X.%c%d%s"), SampleTypePrefix, TCHAR('A')+profile_tier_level.general_profile_space-1, profile_tier_level.general_profile_idc, cf, profile_tier_level.general_tier_flag ? TCHAR('H') : TCHAR('L'), profile_tier_level.general_level_idc, *cfs);
+				}
+			}
+
 
 			bool FHRDParameters::Parse(FBitstreamReader& br, bool commonInfPresentFlag, uint32 maxNumSubLayersMinus1)
 			{
@@ -578,6 +610,7 @@ namespace ElectraDecodersUtil
 					general_profile_space = br.GetBits(2);
 					general_tier_flag = br.GetBits(1);
 					general_profile_idc = br.GetBits(5);
+					general_profile_compatibility_flags = br.PeekBits(32);
 					for(int32 i=0; i<32; ++i)
 					{
 						general_profile_compatibility_flag[i] = br.GetBits(1);
diff --git a/Engine/Plugins/Media/ElectraCodecs/Source/ElectraDecoders/Public/Utils/ElectraBitstreamReader.h b/Engine/Plugins/Media/ElectraCodecs/Source/ElectraDecoders/Public/Utils/ElectraBitstreamReader.h
--- a/Engine/Plugins/Media/ElectraCodecs/Source/ElectraDecoders/Public/Utils/ElectraBitstreamReader.h
+++ b/Engine/Plugins/Media/ElectraCodecs/Source/ElectraDecoders/Public/Utils/ElectraBitstreamReader.h
@@ -21,6 +21,23 @@ public:
 		SetData(InData, InDataSize, StartBytePosition, StartBitPosition);
 	}
 
+	FElectraBitstreamReader(const FElectraBitstreamReader& rhs)
+	{
+		Data = rhs.Data;
+		DataSize = rhs.DataSize;
+		BytePosition = rhs.BytePosition;
+		BitPosition = rhs.BitPosition;
+	}
+
+	FElectraBitstreamReader& operator = (const FElectraBitstreamReader& rhs)
+	{
+		Data = rhs.Data;
+		DataSize = rhs.DataSize;
+		BytePosition = rhs.BytePosition;
+		BitPosition = rhs.BitPosition;
+		return *this;
+	}
+
 	void SetData(const void* InData, uint64 InDataSize, uint64 StartBytePosition = 0, uint32 StartBitPosition = 0)
 	{
 		check(InData && InDataSize);
@@ -45,6 +62,11 @@ public:
 		return DataSize - BytePosition;
 	}
 
+	uint32 GetBytePosition() const
+	{
+		return BytePosition;
+	}
+
 	uint32 GetBitPosition() const
 	{
 		return BitPosition;
@@ -253,4 +275,72 @@ private:
 	int64 NumBytesToGo = 0;
 };
 
+
+class ELECTRADECODERS_API FElectraBitstreamWriter
+{
+public:
+	FElectraBitstreamWriter()
+	{
+		Buffer.SetNumZeroed(256);
+	}
+	void PutBits(uint32 InValue, uint32 InNumBits)
+	{
+		check(InNumBits != 0);
+		check(InNumBits <= 32);
+		uint32 Mask = 1U << (InNumBits-1);
+		while(Mask)
+		{
+			PutBit((InValue & Mask) != 0 ? 1 : 0);
+			Mask >>= 1;
+		}
+	}
+	void AlignToBytes(uint32 InFillBitsWith)
+	{
+		if (BitPosition)
+		{
+			do
+			{
+				PutBit(InFillBitsWith);
+			} while(BitPosition);
+		}
+	}
+	uint32 GetNumBits() const
+	{
+		return BytePosition * 8 + BitPosition;
+	}
+	uint32 GetNumBytesUsed() const
+	{
+		return (GetNumBits() + 7) / 8;
+	}
+	bool IsByteAligned() const
+	{
+		return BitPosition == 0;
+	}
+	void GetArray(TArray<uint8>& OutArray)
+	{
+		Buffer.SetNum(GetNumBytesUsed());
+		OutArray = MoveTemp(Buffer);
+		BytePosition = 0;
+		BitPosition = 0;
+	}
+private:
+	void PutBit(uint32 InBit)
+	{
+		InBit <<= (7 - BitPosition);
+		check(Buffer.Num());
+		Buffer[BytePosition] |= (uint8)InBit;
+		if (++BitPosition == 8)
+		{
+			BitPosition = 0;
+			if (++BytePosition > (uint32)Buffer.Num())
+			{
+				Buffer.SetNumZeroed(Buffer.Num() + 256);
+			}
+		}
+	}
+	TArray<uint8> Buffer;
+	uint32 BytePosition = 0;
+	uint32 BitPosition = 0;
+};
+
 } // namespace ElectraDecodersUtil
diff --git a/Engine/Plugins/Media/ElectraCodecs/Source/ElectraDecoders/Public/Utils/MPEG/ElectraUtilsMPEGVideo_H265.h b/Engine/Plugins/Media/ElectraCodecs/Source/ElectraDecoders/Public/Utils/MPEG/ElectraUtilsMPEGVideo_H265.h
--- a/Engine/Plugins/Media/ElectraCodecs/Source/ElectraDecoders/Public/Utils/MPEG/ElectraUtilsMPEGVideo_H265.h
+++ b/Engine/Plugins/Media/ElectraCodecs/Source/ElectraDecoders/Public/Utils/MPEG/ElectraUtilsMPEGVideo_H265.h
@@ -127,6 +127,7 @@ namespace ElectraDecodersUtil
 					general_tier_flag = 0;
 					general_profile_idc = 0;
 					FMemory::Memzero(general_profile_compatibility_flag);
+					general_profile_compatibility_flags = 0;
 					general_progressive_source_flag = 0;
 					general_interlaced_source_flag = 0;
 					general_non_packed_constraint_flag = 0;
@@ -206,6 +207,7 @@ namespace ElectraDecodersUtil
 				uint8 general_tier_flag;								// u(1)
 				uint8 general_profile_idc;								// u(5)
 				uint8 general_profile_compatibility_flag[32];			// u(1)
+				uint32 general_profile_compatibility_flags;				// u(32), same as general_profile_compatibility_flag[32] but in a single word
 				uint8 general_progressive_source_flag;					// u(1)
 				uint8 general_interlaced_source_flag;					// u(1)
 				uint8 general_non_packed_constraint_flag;				// u(1)
@@ -621,6 +623,8 @@ namespace ElectraDecodersUtil
 				void GetCrop(int32& OutLeft, int32& OutRight, int32& OutTop, int32& OutBottom) const;
 				void GetAspect(int32& OutSarW, int32& OutSarH) const;
 				FFractionalValue GetTiming() const;
+				uint64 GetConstraintFlags() const;
+				FString GetRFC6381(const TCHAR* SampleTypePrefix) const;
 
 				uint8 sps_video_parameter_set_id;							// u(4)
 				uint8 sps_max_sub_layers_minus1;							// u(3), 0-6
diff --git a/Engine/Plugins/Media/ElectraPlayer/Source/ElectraPlayerRuntime/Private/Runtime/HTTP/HTTPManager.cpp b/Engine/Plugins/Media/ElectraPlayer/Source/ElectraPlayerRuntime/Private/Runtime/HTTP/HTTPManager.cpp
--- a/Engine/Plugins/Media/ElectraPlayer/Source/ElectraPlayerRuntime/Private/Runtime/HTTP/HTTPManager.cpp
+++ b/Engine/Plugins/Media/ElectraPlayer/Source/ElectraPlayerRuntime/Private/Runtime/HTTP/HTTPManager.cpp
@@ -887,6 +887,10 @@ namespace Electra
 			TSharedPtrTS<FWaitableBuffer>	ReceiveBuffer = Request->ReceiveBuffer.Pin();
 			if (ReceiveBuffer.IsValid())
 			{
+				if (Request->ConnectionInfo.StatusInfo.ErrorCode)
+				{
+					ReceiveBuffer->SetHasErrored();
+				}
 				ReceiveBuffer->SetEOD();
 			}
 
diff --git a/Engine/Plugins/Media/ElectraPlayer/Source/ElectraPlayerRuntime/Private/Runtime/Player/AdaptiveStreamingPlayerMetrics.h b/Engine/Plugins/Media/ElectraPlayer/Source/ElectraPlayerRuntime/Private/Runtime/Player/AdaptiveStreamingPlayerMetrics.h
--- a/Engine/Plugins/Media/ElectraPlayer/Source/ElectraPlayerRuntime/Private/Runtime/Player/AdaptiveStreamingPlayerMetrics.h
+++ b/Engine/Plugins/Media/ElectraPlayer/Source/ElectraPlayerRuntime/Private/Runtime/Player/AdaptiveStreamingPlayerMetrics.h
@@ -17,6 +17,7 @@ namespace Metrics
 {
 	enum class ESegmentType
 	{
+		Undefined,
 		Init,
 		Media
 	};
@@ -24,6 +25,8 @@ namespace Metrics
 	{
 		switch(SegmentType)
 		{
+			case ESegmentType::Undefined:
+				return TEXT("Undefined");
 			case ESegmentType::Init:
 				return TEXT("Init");
 			case ESegmentType::Media:
@@ -113,42 +116,6 @@ namespace Metrics
 
 	struct FSegmentDownloadStats
 	{
-		// Inputs from stream request
-		EStreamType		StreamType = EStreamType::Video;	//!< Type of stream
-		ESegmentType	SegmentType = ESegmentType::Media;	//!< Type of segment (init or media)
-		FString			URL;								//!< Effective URL used to download from
-		FString			Range;								//!< Range used to download
-		FString			CDN;								//!< CDN
-		FString			MediaAssetID;
-		FString			AdaptationSetID;
-		FString			RepresentationID;
-		double			PresentationTime = 0.0;				//!< Presentation time on media timeline
-		double			Duration = 0.0;						//!< Duration of segment as specified in manifest
-		int32			Bitrate = 0;						//!< Stream bitrate as specified in manifest
-		int32			QualityIndex = 0;					//!< Quality index of this segment
-		int32			HighestQualityIndex = 0;			//!< The highest quality index that could be had.
-		int32			RetryNumber = 0;
-		bool			bIsMissingSegment = false;			//!< true if the segment was not actually downloaded because it is missing on the timeline.
-
-		// Outputs from stream reader
-		uint32			StatsID = 0;						//!< ID uniquely identifying this download
-		FString			FailureReason;						//!< Human readable failure reason. Only for display purposes.
-		double			AvailibilityDelay = 0.0;			//!< Time the download had to wait for the segment to enter its availability window.
-		double			DurationDownloaded = 0.0;			//!< Duration of content successfully downloaded. May be less than Duration in case of errors.
-		double			DurationDelivered = 0.0;			//!< Duration of content delivered to buffer. If larger than DurationDownloaded indicates dummy data was inserted into buffer.
-		double			TimeToFirstByte = 0.0;				//!< Time in seconds until first data byte was received
-		double			TimeToDownload = 0.0;				//!< Total time in seconds for entire download
-		int64			ByteSize = 0;						//!< Content-Length, may be -1 if unknown (either on error or chunked transfer)
-		int64			NumBytesDownloaded = 0;				//!< Number of bytes successfully downloaded.
-		int32			HTTPStatusCode = 0 ;				//!< HTTP status code (0 if not connected to server yet)
-		bool			bWasSuccessful = false;				//!< true if download was successful, false if not
-		bool			bWasAborted = false;				//!< true if download was aborted by ABR (not by playback!)
-		bool			bDidTimeout = false;				//!< true if a timeout occurred. Only set if timeouts are enabled. Usually the ABR will monitor and abort.
-		bool			bParseFailure = false;				//!< true if the segment could not be parsed
-		bool			bInsertedFillerData = false;
-		bool			bIsCachedResponse = false;
-		bool			bWaitingForRemoteRetryElement = false;
-
 		// Chunk timing
 		struct FMovieChunkInfo
 		{
@@ -159,8 +126,88 @@ namespace Metrics
 			FTimeValue ContentDuration;
 			FMovieChunkInfo() { ContentDuration = FTimeValue::GetZero(); }
 		};
+
+		// Inputs from stream request
+		EStreamType StreamType;					//!< Type of stream
+		ESegmentType SegmentType;				//!< Type of segment (init or media)
+		FString URL;							//!< Effective URL used to download from
+		FString Range;							//!< Range used to download
+		FString CDN;							//!< CDN
+		FString MediaAssetID;
+		FString AdaptationSetID;
+		FString RepresentationID;
+		double PresentationTime;				//!< Presentation time on media timeline
+		double Duration;						//!< Duration of segment as specified in manifest
+		int32 Bitrate;							//!< Stream bitrate as specified in manifest
+		int32 QualityIndex;						//!< Quality index of this segment
+		int32 HighestQualityIndex;				//!< The highest quality index that could be had.
+		int32 RetryNumber;
+		bool bWaitingForRemoteRetryElement;
+
+		// Outputs from stream reader
+		uint32 StatsID;							//!< ID uniquely identifying this download
+		FString FailureReason;					//!< Human readable failure reason. Only for display purposes.
+		double AvailibilityDelay;				//!< Time the download had to wait for the segment to enter its availability window.
+		double DurationDownloaded;				//!< Duration of content successfully downloaded. May be less than Duration in case of errors.
+		double DurationDelivered;				//!< Duration of content delivered to buffer. If larger than DurationDownloaded indicates dummy data was inserted into buffer.
+		double TimeToFirstByte;					//!< Time in seconds until first data byte was received
+		double TimeToDownload;					//!< Total time in seconds for entire download
+		int64 ByteSize;							//!< Content-Length, may be -1 if unknown (either on error or chunked transfer)
+		int64 NumBytesDownloaded;				//!< Number of bytes successfully downloaded.
+		int32 HTTPStatusCode;					//!< HTTP status code (0 if not connected to server yet)
+		bool bWasSuccessful;					//!< true if download was successful, false if not
+		bool bWasAborted;						//!< true if download was aborted by ABR (not by playback!)
+		bool bDidTimeout;						//!< true if a timeout occurred. Only set if timeouts are enabled. Usually the ABR will monitor and abort.
+		bool bParseFailure;						//!< true if the segment could not be parsed
+		bool bIsMissingSegment;					//!< true if the segment was not actually downloaded because it is missing on the timeline.
+		bool bInsertedFillerData;
+		bool bIsCachedResponse;
 		TArray<IElectraHTTPStreamResponse::FTimingTrace> TimingTraces;
 		TArray<FMovieChunkInfo> MovieChunkInfos;
+
+		FSegmentDownloadStats()
+		{ Reset(); }
+		void Reset()
+		{
+			StreamType = EStreamType::Unsupported;
+			SegmentType = ESegmentType::Undefined;
+			URL.Empty();
+			Range.Empty();
+			CDN.Empty();
+			MediaAssetID.Empty();
+			AdaptationSetID.Empty();
+			RepresentationID.Empty();
+			PresentationTime = 0.0;
+			Duration = 0.0;
+			Bitrate = 0;
+			QualityIndex = 0;
+			HighestQualityIndex = 0;
+			RetryNumber = 0;
+			bWaitingForRemoteRetryElement = false;
+			ResetOutput();
+		}
+		void ResetOutput()
+		{
+			StatsID = 0;
+			FailureReason.Empty();
+			AvailibilityDelay = 0.0;
+			DurationDownloaded = 0.0;
+			DurationDelivered = 0.0;
+			TimeToFirstByte = 0.0;
+			TimeToDownload = 0.0;
+			ByteSize = 0;
+			NumBytesDownloaded = 0;
+			HTTPStatusCode = 0 ;
+			bWasSuccessful = false;
+			bWasAborted = false;
+			bDidTimeout = false;
+			bParseFailure = false;
+			bIsMissingSegment = false;
+			bInsertedFillerData = false;
+			bIsCachedResponse = false;
+			TimingTraces.Empty();
+			MovieChunkInfos.Empty();
+		}
 	};
 
 	struct FLicenseKeyStats
diff --git a/Engine/Plugins/Media/ElectraPlayer/Source/ElectraPlayerRuntime/Private/Runtime/Player/AdaptiveStreamingPlayerResourceRequest.cpp b/Engine/Plugins/Media/ElectraPlayer/Source/ElectraPlayerRuntime/Private/Runtime/Player/AdaptiveStreamingPlayerResourceRequest.cpp
--- a/Engine/Plugins/Media/ElectraPlayer/Source/ElectraPlayerRuntime/Private/Runtime/Player/AdaptiveStreamingPlayerResourceRequest.cpp
+++ b/Engine/Plugins/Media/ElectraPlayer/Source/ElectraPlayerRuntime/Private/Runtime/Player/AdaptiveStreamingPlayerResourceRequest.cpp
@@ -152,35 +152,36 @@ void FHTTPResourceRequest::Cancel()
 
 void FHTTPResourceRequest::StaticDataReady()
 {
-	if (Request.IsValid())
+	TSharedPtrTS<IElectraHttpManager::FRequest> Req = Request;
+	if (Req.IsValid())
 	{
 		// Was static data actually set or was there no data provided?
 		if (!bStaticDataReady)
 		{
 			// Do the actual HTTP request now.
 			// Is an external reader to be used?
-			if (!Request->ExternalDataReader.IsValid())
+			if (!Req->ExternalDataReader.IsValid())
 			{
-				Request->ExternalDataReader = PlayerSessionServices->GetExternalDataReader();
+				Req->ExternalDataReader = PlayerSessionServices->GetExternalDataReader();
 			}
 			bWasAdded = true;
 			TSharedPtrTS<IElectraHttpManager> PinnedHTTPManager = HTTPManager.Pin();
 			if (PinnedHTTPManager.IsValid())
 			{
-				PinnedHTTPManager->AddRequest(Request, false);
+				PinnedHTTPManager->AddRequest(Req, false);
 			}
 		}
 		else
 		{
-			HTTP::FConnectionInfo& ci = Request->ConnectionInfo;
-			ci.EffectiveURL = Request->Parameters.URL;
-			ci.bIsConnected = true;
-			ci.bHaveResponseHeaders  = true;
-			ci.bWasAborted = false;
-			ci.bHasFinished = true;
-			ci.HTTPVersionReceived = 11;
-			ci.StatusInfo.HTTPStatus = Request->Parameters.Range.IsSet() ? 206 : 200;
-			ci.ContentLength = ci.BytesReadSoFar = ReceiveBuffer.IsValid() ? ReceiveBuffer->Num() : 0;
+			ConnectionInfo = Req->ConnectionInfo;
+			ConnectionInfo.EffectiveURL = Req->Parameters.URL;
+			ConnectionInfo.bIsConnected = true;
+			ConnectionInfo.bHaveResponseHeaders  = true;
+			ConnectionInfo.bWasAborted = false;
+			ConnectionInfo.bHasFinished = true;
+			ConnectionInfo.HTTPVersionReceived = 11;
+			ConnectionInfo.StatusInfo.HTTPStatus = Req->Parameters.Range.IsSet() ? 206 : 200;
+			ConnectionInfo.ContentLength = ConnectionInfo.BytesReadSoFar = ReceiveBuffer.IsValid() ? ReceiveBuffer->Num() : 0;
 
 			bInCallback = true;
 			CompletedCallback.ExecuteIfBound(AsShared());
@@ -199,29 +200,30 @@ void FHTTPResourceRequest::HTTPCompletionCallback(const IElectraHttpManager::FRe
 	TSharedPtrTS<FHTTPResourceRequest> Self = AsShared();
 	if (Self.IsValid())
 	{
-		if (Request.IsValid())
+		TSharedPtrTS<IElectraHttpManager::FRequest> Req = Request;
+		if (Req.IsValid())
 		{
-			const HTTP::FConnectionInfo& ConnInfo = Request->ConnectionInfo;
-			if (!ConnInfo.bWasAborted)
+			ConnectionInfo = Req->ConnectionInfo;
+			if (!ConnectionInfo.bWasAborted)
 			{
 				Error = 0;
-				if (Request->ConnectionInfo.StatusInfo.ErrorDetail.IsError())
+				if (Req->ConnectionInfo.StatusInfo.ErrorDetail.IsError())
 				{
-					if (ConnInfo.StatusInfo.ConnectionTimeoutAfterMilliseconds)
+					if (ConnectionInfo.StatusInfo.ConnectionTimeoutAfterMilliseconds)
 					{
 						Error = 1;
 					}
-					else if (ConnInfo.StatusInfo.NoDataTimeoutAfterMilliseconds)
+					else if (ConnectionInfo.StatusInfo.NoDataTimeoutAfterMilliseconds)
 					{
 						Error = 2;
 					}
-					else if (ConnInfo.StatusInfo.bReadError)
+					else if (ConnectionInfo.StatusInfo.bReadError)
 					{
 						Error = 3;
 					}
 					else
 					{
-						Error = ConnInfo.StatusInfo.HTTPStatus ? ConnInfo.StatusInfo.HTTPStatus : 4;
+						Error = ConnectionInfo.StatusInfo.HTTPStatus ? ConnectionInfo.StatusInfo.HTTPStatus : 4;
 					}
 				}
 
diff --git a/Engine/Plugins/Media/ElectraPlayer/Source/ElectraPlayerRuntime/Private/Runtime/Player/AdaptiveStreamingPlayerResourceRequest.h b/Engine/Plugins/Media/ElectraPlayer/Source/ElectraPlayerRuntime/Private/Runtime/Player/AdaptiveStreamingPlayerResourceRequest.h
--- a/Engine/Plugins/Media/ElectraPlayer/Source/ElectraPlayerRuntime/Private/Runtime/Player/AdaptiveStreamingPlayerResourceRequest.h
+++ b/Engine/Plugins/Media/ElectraPlayer/Source/ElectraPlayerRuntime/Private/Runtime/Player/AdaptiveStreamingPlayerResourceRequest.h
@@ -10,6 +10,7 @@
 namespace Electra
 {
 class IPlayerSessionServices;
+class FHTTPResourceRequest;
 
 
 class IAdaptiveStreamingPlayerResourceRequest : public TSharedFromThis<IAdaptiveStreamingPlayerResourceRequest, ESPMode::ThreadSafe>
@@ -65,8 +66,38 @@ public:
 	virtual ~IHTTPResourceRequestObject() = default;
 };
 
+
+
+class FHTTPResourceRequestCompletionSignal : public TSharedFromThis<FHTTPResourceRequestCompletionSignal, ESPMode::ThreadSafe>
+{
+public:
+	static TSharedPtr<FHTTPResourceRequestCompletionSignal, ESPMode::ThreadSafe> Create()
+	{ return MakeShared<FHTTPResourceRequestCompletionSignal, ESPMode::ThreadSafe>(); }
+	virtual ~FHTTPResourceRequestCompletionSignal()
+	{ Sig.Signal(); }
+	virtual void CompletionCallback(TSharedPtrTS<FHTTPResourceRequest> InRequest)
+	{ Sig.Signal();	}
+
+	void Wait()
+	{ Sig.Wait(); }
+	void WaitAndReset()
+	{ Sig.WaitAndReset(); }
+	void Signal()
+	{ Sig.Signal(); }
+	void Reset()
+	{ Sig.Reset(); }
+	bool IsSignaled() const
+	{ return Sig.IsSignaled(); }
+	bool WaitTimeout(int64 MicroSeconds)
+	{ return Sig.WaitTimeout(MicroSeconds); }
+	bool WaitTimeoutAndReset(int64 MicroSeconds)
+	{ return Sig.WaitTimeoutAndReset(MicroSeconds); }
+private:
+	FMediaEvent Sig;
+};
+
 /**
- * 
+ *
  */
 class FHTTPResourceRequest : public TSharedFromThis<FHTTPResourceRequest, ESPMode::ThreadSafe>
 {
@@ -89,34 +120,37 @@ public:
 
 	virtual FHTTPResourceRequest& PostData(TArray<uint8> InPostData)
 	{ Request->Parameters.PostData = MoveTemp(InPostData); return *this; }
-	
+
 	virtual FHTTPResourceRequest& Range(const FString& InRange)
 	{ Request->Parameters.Range.Set(InRange); return *this; }
-	
+
 	virtual FHTTPResourceRequest& Headers(const TArray<HTTP::FHTTPHeader>& InHeaders)
 	{ Request->Parameters.RequestHeaders = InHeaders; return *this; }
 
 	virtual FHTTPResourceRequest& Headers(const TArray<FString>& InHeaders)
 	{ Request->Parameters.AddFromHeaderList(InHeaders); return *this; }
-	
+
 	virtual FHTTPResourceRequest& AcceptEncoding(const FString& InAcceptEncoding)
 	{ Request->Parameters.AcceptEncoding.Set(InAcceptEncoding); return *this; }
-	
+
 	virtual FHTTPResourceRequest& ConnectionTimeout(const FTimeValue& InTimeoutAfter)
 	{ Request->Parameters.ConnectTimeout = InTimeoutAfter; return *this; }
-	
+
 	virtual FHTTPResourceRequest& NoDataTimeout(const FTimeValue& InTimeoutAfter)
 	{ Request->Parameters.NoDataTimeout = InTimeoutAfter; return *this; }
-	
+
 	virtual FHTTPResourceRequest& AllowStaticQuery(IAdaptiveStreamingPlayerResourceRequest::EPlaybackResourceType InStaticQueryType)
 	{ StaticQueryType.Set(InStaticQueryType); return *this; }
-	
+
 	virtual FHTTPResourceRequest& Object(TSharedPtrTS<IHTTPResourceRequestObject> InUserObject)
 	{ UserObject = InUserObject; return *this; }
 
 	virtual FHTTPResourceRequest& StreamTypeAndQuality(EStreamType InStreamType, int32 InQualityIndex, int32 InMaxQualityIndex)
 	{ Request->Parameters.StreamType = InStreamType; Request->Parameters.QualityIndex = InQualityIndex; Request->Parameters.MaxQualityIndex = InMaxQualityIndex; return *this; }
 
+	virtual FHTTPResourceRequest& CompletionSignal(const TSharedPtr<FHTTPResourceRequestCompletionSignal, ESPMode::ThreadSafe>& InSig)
+	{ CompletedCallback.BindThreadSafeSP(InSig.ToSharedRef(), &FHTTPResourceRequestCompletionSignal::CompletionCallback); return *this; }
+
 	virtual FOnRequestCompletedCallback& Callback()
 	{ return CompletedCallback; }
 
@@ -137,7 +171,7 @@ public:
 		else if (Error >= 100 && Error < 600)	return FString::Printf(TEXT("HTTP status code %d"), Error);
 		else									return FString::Printf(TEXT("Unknown <code %d>"), Error);
 	}
-	
+
 	virtual bool GetWasCanceled() const
 	{ return bWasCanceled; }
 
@@ -154,8 +188,8 @@ public:
 	{ return ReceiveBuffer; }
 
 	virtual const HTTP::FConnectionInfo* GetConnectionInfo() const
-	{ return Request.IsValid() ? &Request->ConnectionInfo : nullptr; }
-
+	{ return &ConnectionInfo; }
+	
 	virtual TSharedPtrTS<IElectraHttpManager::FRequest> GetRequest() const
 	{ return Request; }
 
@@ -223,6 +257,7 @@ private:
 	void SetStaticDataReady()
 	{ bStaticDataReady = true; }
 	void StaticDataReady();
+	HTTP::FConnectionInfo ConnectionInfo;
 	TSharedPtrTS<IElectraHttpManager::FRequest> Request;
 	TSharedPtrTS<FWaitableBuffer> ReceiveBuffer;
 	TSharedPtrTS<IElectraHttpManager::FProgressListener> ProgressListener;
@@ -239,8 +274,6 @@ private:
 	int32 Error = 0;
 };
 
-
-
 } // namespace Electra
 
 
diff --git a/Engine/Plugins/Media/ElectraPlayer/Source/ElectraPlayerRuntime/Private/Runtime/StreamDataBuffer.h b/Engine/Plugins/Media/ElectraPlayer/Source/ElectraPlayerRuntime/Private/Runtime/StreamDataBuffer.h
--- a/Engine/Plugins/Media/ElectraPlayer/Source/ElectraPlayerRuntime/Private/Runtime/StreamDataBuffer.h
+++ b/Engine/Plugins/Media/ElectraPlayer/Source/ElectraPlayerRuntime/Private/Runtime/StreamDataBuffer.h
@@ -67,9 +67,10 @@ namespace Electra
 		void Reset()
 		{
 			FScopeLock Lock(&AccessLock);
-			WritePos = ReadPos = 0;
+			WritePos = 0;
 			bEOD = false;
 			bWasAborted = false;
+			bHasErrored = false;
 			// Do not modify the "waiter" members. When waiting for data to arrive it needs to continue
 			// doing so even when we are resetting the buffer.
 				//WaitingForSize = 0;
@@ -83,11 +84,11 @@ namespace Electra
 			return DataSize;
 		}
 
-		// Returns the number of bytes in the buffer (amount that can be popped)
+		// Returns the number of bytes in the buffer
 		int64 Num() const
 		{
 			FScopeLock Lock(&AccessLock);
-			return WritePos - ReadPos;
+			return WritePos;
 		}
 
 		// Returns the number of free bytes in the buffer (amount that can be pushed)
@@ -200,35 +201,13 @@ namespace Electra
 			SizeAvailableSignal.Signal();
 		}
 
-#if 0
-		// "Pops" data from the buffer to a destination. At most the specified number of bytes are popped, or fewer if not as many are available.
-		int64 PopData(uint8* OutData, int64 MaxElementsWanted)
-		{
-			FScopeLock Lock(&AccessLock);
-			int64 size = Num();
-
-			if (MaxElementsWanted > size)
-			{
-				MaxElementsWanted = size;
-			}
-			// Copy out or skip over?
-			if (OutData)
-			{
-				CopyData(OutData, GetBufferBase() + ReadPos, MaxElementsWanted);
-			}
-			ReadPos += MaxElementsWanted;
-			return MaxElementsWanted;
-		}
-#endif
-
 		void RemoveFromBeginning(int64 InNumBytesToRemove)
 		{
 			FScopeLock Lock(&AccessLock);
 			check(ExternalBuffer == nullptr);
 			check(InNumBytesToRemove >= 0);
-			check(ReadPos == 0);
 			check(InNumBytesToRemove <= (int64)WritePos);
-			int64 InNow = WritePos - ReadPos;
+			const int64 InNow = WritePos;
 			check(InNow - InNumBytesToRemove >= 0);
 			if (InNumBytesToRemove > 0)
 			{
@@ -257,17 +236,17 @@ namespace Electra
 		int64 GetLinearReadSize() const
 		{
 			FScopeLock Lock(&AccessLock);
-			return WritePos - ReadPos;
+			return WritePos;
 		}
 
 		// Must control Lock()/Unlock() externally!
 		const uint8* GetLinearReadData() const
 		{
-			return GetBufferBase() ? GetBufferBase() + ReadPos : nullptr;
+			return GetBufferBase() ? GetBufferBase() : nullptr;
 		}
 		uint8* GetLinearReadData()
 		{
-			return GetBufferBase() ? GetBufferBase() + ReadPos : nullptr;
+			return GetBufferBase() ? GetBufferBase() : nullptr;
 		}
 
 		uint8* GetLinearWriteData(int64 InNumBytesToAppend)
@@ -323,6 +302,16 @@ namespace Electra
 			return bWasAborted;
 		}
 
+		void SetHasErrored()
+		{
+			bHasErrored = true;
+		}
+
+		bool HasErrored() const
+		{
+			return bHasErrored;
+		}
+
 	protected:
 		uint8* GetBufferBase() const
 		{
@@ -342,7 +331,6 @@ namespace Electra
 				Buffer = MakeShared<TArray<uint8>, ESPMode::ThreadSafe>();
 				Buffer->AddUninitialized(InSize);
 				WritePos = 0;
-				ReadPos = 0;
 			}
 			return true;
 		}
@@ -352,7 +340,6 @@ namespace Electra
 			Buffer.Reset();
 			DataSize = 0;
 			WritePos = 0;
-			ReadPos = 0;
 			ExternalBuffer = nullptr;
 		}
 
@@ -383,8 +370,6 @@ namespace Electra
 		uint64 DataSize = 0;
 		// Offset into buffer where to add new data
 		uint64 WritePos = 0;
-		// Offset into buffer from where to read the next data.
-		uint64 ReadPos = 0;
 		// Amount of data necessary to be present for `SizeAvailableSignal` to get set.
 		uint64 WaitingForSize = 0;
 		// If set a buffer is provided externally to read into directly.
@@ -393,8 +378,8 @@ namespace Electra
 		volatile bool bEOD = false;
 		// Flag indicating that reading into the buffer has been aborted.
 		volatile bool bWasAborted = false;
+		// Flag indicating that filling the buffer from the source has encountered an error.
+		volatile bool bHasErrored = false;
 	};
 
 } // namespace Electra
-
-
diff --git a/Engine/Plugins/Media/ElectraUtil/Source/ElectraBase/Private/PlayerTime.cpp b/Engine/Plugins/Media/ElectraUtil/Source/ElectraBase/Private/PlayerTime.cpp
--- a/Engine/Plugins/Media/ElectraUtil/Source/ElectraBase/Private/PlayerTime.cpp
+++ b/Engine/Plugins/Media/ElectraUtil/Source/ElectraBase/Private/PlayerTime.cpp
@@ -90,7 +90,7 @@ namespace Electra
 			{
 				HNS 		= Numerator;
 				bIsValid	= true;
-				bIsInfinity = false;
+				bIsInfinity = HNS == 0x7fffffffffffffffLL || HNS == -0x7fffffffffffffffLL;
 			}
 			else if (Numerator >= -922337203685LL && Numerator <= 922337203685LL)
 			{
diff --git a/Engine/Plugins/Media/ElectraUtil/Source/ElectraBase/Private/StringHelpers.cpp b/Engine/Plugins/Media/ElectraUtil/Source/ElectraBase/Private/StringHelpers.cpp
--- a/Engine/Plugins/Media/ElectraUtil/Source/ElectraBase/Private/StringHelpers.cpp
+++ b/Engine/Plugins/Media/ElectraUtil/Source/ElectraBase/Private/StringHelpers.cpp
@@ -7,6 +7,39 @@ namespace Electra
 	namespace StringHelpers
 	{
 
+		FString ISO_8859_1_ToFString(const uint8* InStringToConvert, int32 InNumCharsToConvert)
+		{
+			if (InNumCharsToConvert == 0)
+			{
+				return FString();
+			}
+			else if (InNumCharsToConvert < 0)
+			{
+				for(InNumCharsToConvert=0; InStringToConvert[InNumCharsToConvert] && InNumCharsToConvert < 16384; ++InNumCharsToConvert)
+				{}
+			}
+			TArray<uint8> ConvBuf;
+			ConvBuf.Reserve(InNumCharsToConvert*2);
+			for(int32 i=0; i<InNumCharsToConvert; ++i, ++InStringToConvert)
+			{
+				if (*InStringToConvert == 0x00)
+				{
+					break;
+				}
+				if (*InStringToConvert >= 0x20 && *InStringToConvert < 0x7f)
+				{
+					ConvBuf.Add(*InStringToConvert);
+				}
+				else  if (*InStringToConvert >= 0xa0)
+				{
+					// We can convert straight from ISO 8859-1 to UTF8 by doing this:
+					ConvBuf.Add(0xc0 | (*InStringToConvert >> 6));
+					ConvBuf.Add(0x80 | (*InStringToConvert & 0x3f));
+				}
+			}
+			auto Cnv = StringCast<TCHAR>(reinterpret_cast<const UTF8CHAR*>(ConvBuf.GetData()), ConvBuf.Num());
+			return FString(Cnv.Length(), Cnv.Get());
+		}
 
 		int32 FindFirstOf(const FString& InString, const FString& SplitAt, int32 FirstPos)
 		{
@@ -112,13 +145,13 @@ namespace Electra
 
 
 		bool StringEquals(const TCHAR * const s1, const TCHAR * const s2)
-		{ 
-			return FPlatformString::Strcmp(s1, s2) == 0; 
+		{
+			return FPlatformString::Strcmp(s1, s2) == 0;
 		}
 
 		bool StringStartsWith(const TCHAR * const s1, const TCHAR * const s2, SIZE_T n)
-		{ 
-			return FPlatformString::Strncmp(s1, s2, n) == 0; 
+		{
+			return FPlatformString::Strncmp(s1, s2, n) == 0;
 		}
 
 		void StringToArray(TArray<uint8>& OutArray, const FString& InString)
diff --git a/Engine/Plugins/Media/ElectraUtil/Source/ElectraBase/Public/PlayerTime.h b/Engine/Plugins/Media/ElectraUtil/Source/ElectraBase/Public/PlayerTime.h
--- a/Engine/Plugins/Media/ElectraUtil/Source/ElectraBase/Public/PlayerTime.h
+++ b/Engine/Plugins/Media/ElectraUtil/Source/ElectraBase/Public/PlayerTime.h
@@ -359,7 +359,10 @@ public:
 			{
 				if (!bIsInfinity && !rhs.bIsInfinity)
 				{
-					HNS += rhs.HNS;
+					if (WillOverflow(HNS, HNS, rhs.HNS))
+					{
+						SetToPositiveInfinity();
+					}
 				}
 				else
 				{
@@ -430,8 +433,11 @@ public:
 		{
 			if (!bIsInfinity && !rhs.bIsInfinity)
 			{
-				Result.HNS = HNS + rhs.HNS;
 				Result.bIsValid = true;
+				if (WillOverflow(Result.HNS, HNS, rhs.HNS))
+				{
+					Result.SetToPositiveInfinity();
+				}
 			}
 			else
 			{
@@ -518,6 +524,12 @@ public:
 	}
 
 private:
+	static inline bool WillOverflow(int64& OutTemp, const int64 InA, const int64 InB)
+	{
+		int64 Temp = InA + InB;
+		OutTemp = Temp;
+		return (InA >= 0 && InB >= 0 && Temp < 0) || (InA < 0 && InB < 0 && Temp >= 0);
+	}
 	int64	HNS;
 	int64	SequenceIndex;
 	bool	bIsValid;
diff --git a/Engine/Plugins/Media/ElectraUtil/Source/ElectraBase/Public/Utilities/StringHelpers.h b/Engine/Plugins/Media/ElectraUtil/Source/ElectraBase/Public/Utilities/StringHelpers.h
--- a/Engine/Plugins/Media/ElectraUtil/Source/ElectraBase/Public/Utilities/StringHelpers.h
+++ b/Engine/Plugins/Media/ElectraUtil/Source/ElectraBase/Public/Utilities/StringHelpers.h
@@ -14,6 +14,7 @@ namespace Electra
 
 	namespace StringHelpers
 	{
+		FString ELECTRABASE_API ISO_8859_1_ToFString(const uint8* InStringToConvert, int32 InNumCharsToConvert=-1);
 
 		int32 ELECTRABASE_API FindFirstOf(const FString& InString, const FString& SplitAt, int32 FirstPos = 0);
 
