Electra: Disabled blocking when paused.

This greatly improves in-editor/sequencer usability. Scrubbing a timeline issues new seeks, but only the most recent is served and any prior seeks get canceled.

diff --git a/Engine/Plugins/Media/ElectraPlayer/Source/ElectraPlayerPlugin/Private/ElectraPlayerPlugin.cpp b/Engine/Plugins/Media/ElectraPlayer/Source/ElectraPlayerPlugin/Private/ElectraPlayerPlugin.cpp
--- a/Engine/Plugins/Media/ElectraPlayer/Source/ElectraPlayerPlugin/Private/ElectraPlayerPlugin.cpp
+++ b/Engine/Plugins/Media/ElectraPlayer/Source/ElectraPlayerPlugin/Private/ElectraPlayerPlugin.cpp
@@ -855,7 +855,7 @@ bool FElectraPlayerPlugin::CanControl(EMediaControl Control) const
 	EMediaState CurrentState = GetState();
 	if (Control == EMediaControl::BlockOnFetch)
 	{
-		return CurrentState == EMediaState::Paused || CurrentState == EMediaState::Playing;
+		return CurrentState == EMediaState::Playing;
 	}
 	else if (Control == EMediaControl::Pause)
 	{
