ElectraPlayer: Fixed incorrect tagging of latest possible presentation timestamp in segment loading of Live streams

#rnx

diff --git a/Engine/Plugins/Media/ElectraPlayer/Source/ElectraPlayerRuntime/Private/Runtime/Player/DASH/ManifestBuilderDASH.cpp b/Engine/Plugins/Media/ElectraPlayer/Source/ElectraPlayerRuntime/Private/Runtime/Player/DASH/ManifestBuilderDASH.cpp
--- a/Engine/Plugins/Media/ElectraPlayer/Source/ElectraPlayerRuntime/Private/Runtime/Player/DASH/ManifestBuilderDASH.cpp
+++ b/Engine/Plugins/Media/ElectraPlayer/Source/ElectraPlayerRuntime/Private/Runtime/Player/DASH/ManifestBuilderDASH.cpp
@@ -2342,7 +2342,7 @@ FTimeRange FManifestDASHInternal::GetTotalTimeRange() const
 		bool bIsUpdating = AreUpdatesExpected();
 		FTimeValue ast = GetAnchorTime();
 		FTimeValue Now = PlayerSessionServices->GetSynchronizedUTCTime()->GetTime();
-		FTimeValue LastEnd = GetLastPeriodEndTime();
+		FTimeValue LastEnd = GetLastPeriodEndTime(true);
 		TotalTimeRange.End = bIsUpdating ? Now : LastEnd;
 		if (TotalTimeRange.End.IsValid())
 		{
@@ -2499,7 +2499,7 @@ FTimeValue FManifestDASHInternal::GetTimeshiftBufferDepth() const
 }
 
 
-FTimeValue FManifestDASHInternal::GetLastPeriodEndTime() const
+FTimeValue FManifestDASHInternal::GetLastPeriodEndTime(bool bForTimeline) const
 {
 	// As per Annex A.3.2
 	FTimeValue ast = GetAnchorTime();
@@ -2522,7 +2522,7 @@ FTimeValue FManifestDASHInternal::GetLastPeriodEndTime() const
 	// If MUP is zero then it is expected that InbandEventStream is used to signal when to update the MPD.
 	// Since the MPD will not update through MUP in this case we need to return that the period goes up
 	// to mediaPresentationDuration, in this case infinity.
-	if (MPDRoot->GetMinimumUpdatePeriod().IsValid() && MPDRoot->GetMinimumUpdatePeriod() > FTimeValue::GetZero())
+	if (bForTimeline && MPDRoot->GetMinimumUpdatePeriod().IsValid() && MPDRoot->GetMinimumUpdatePeriod() > FTimeValue::GetZero())
 	{
 		FTimeValue CheckTime = FetchTime + MPDRoot->GetMinimumUpdatePeriod();
 		return CheckTime < End ? CheckTime : End;
@@ -2571,7 +2571,7 @@ FTimeRange FManifestDASHInternal::GetSeekableTimeRange() const
 		bool bIsUpdating = AreUpdatesExpected() || EpicEventType == EEpicEventType::Dynamic;
 		FTimeValue Distance = bIsUpdating ? CalculateDistanceToLiveEdge() : FixedSeekEndDistance;
 		FTimeValue Now = PlayerSessionServices->GetSynchronizedUTCTime()->GetTime();
-		FTimeValue LastEnd = GetLastPeriodEndTime();
+		FTimeValue LastEnd = GetLastPeriodEndTime(true);
 		FTimeValue CurrentEnd = Now;
 		if (LastEnd.IsValid() && CurrentEnd > LastEnd)
 		{
diff --git a/Engine/Plugins/Media/ElectraPlayer/Source/ElectraPlayerRuntime/Private/Runtime/Player/DASH/ManifestDASH.cpp b/Engine/Plugins/Media/ElectraPlayer/Source/ElectraPlayerRuntime/Private/Runtime/Player/DASH/ManifestDASH.cpp
--- a/Engine/Plugins/Media/ElectraPlayer/Source/ElectraPlayerRuntime/Private/Runtime/Player/DASH/ManifestDASH.cpp
+++ b/Engine/Plugins/Media/ElectraPlayer/Source/ElectraPlayerRuntime/Private/Runtime/Player/DASH/ManifestDASH.cpp
@@ -545,7 +545,7 @@ IManifest::FResult FManifestDASH::FindPlayPeriod(TSharedPtrTS<IPlayPeriod>& OutP
 	PlayRangeEnd -= Manifest->GetAnchorTime();
 
 	// Quick out if the time falls outside the presentation.
-	FTimeValue TotalEndTime = Manifest->GetLastPeriodEndTime();
+	FTimeValue TotalEndTime = Manifest->GetLastPeriodEndTime(true);
 	TotalEndTime -= Manifest->GetAnchorTime();
 	if (PlayRangeEnd.IsValid() && TotalEndTime.IsValid() && PlayRangeEnd < TotalEndTime)
 	{
@@ -635,7 +635,7 @@ IManifest::FResult FManifestDASH::FindPlayPeriod(TSharedPtrTS<IPlayPeriod>& OutP
 		if (Manifest->IsDynamicEpicEvent())
 		{
 			FTimeValue Now = PlayerSessionServices->GetSynchronizedUTCTime()->GetTime();
-			FTimeValue End = Manifest->GetLastPeriodEndTime();
+			FTimeValue End = Manifest->GetLastPeriodEndTime(true);
 			if (Now >= End)
 			{
 				return IManifest::FResult(IManifest::FResult::EType::PastEOS);
@@ -1337,7 +1337,7 @@ IManifest::FResult FDASHPlayPeriod::GetStartingSegment(TSharedPtrTS<IStreamSegme
 			SearchOpt.PeriodDuration = Period->GetDuration();
 			if (!SearchOpt.PeriodDuration.IsValid() || SearchOpt.PeriodDuration.IsPositiveInfinity())
 			{
-				SearchOpt.PeriodDuration = Manifest->GetLastPeriodEndTime() - AST - Period->GetStart();
+				SearchOpt.PeriodDuration = Manifest->GetLastPeriodEndTime(false) - AST - Period->GetStart();
 			}
 			SearchOpt.PeriodPresentationEnd = PlayRangeEnd;
 			SearchOpt.bHasFollowingPeriod = Period->GetHasFollowingPeriod();
@@ -1616,7 +1616,7 @@ IManifest::FResult FDASHPlayPeriod::GetNextOrRetrySegment(TSharedPtrTS<IStreamSe
 	SearchOpt.PeriodDuration = Period->GetDuration();
 	if (!SearchOpt.PeriodDuration.IsValid() || SearchOpt.PeriodDuration.IsPositiveInfinity())
 	{
-		SearchOpt.PeriodDuration = Manifest->GetLastPeriodEndTime() - AST;
+		SearchOpt.PeriodDuration = Manifest->GetLastPeriodEndTime(false) - AST;
 	}
 	SearchOpt.QualityIndex = ActiveQualityIndex.Index;
 	SearchOpt.MaxQualityIndex = ActiveQualityIndex.MaxIndex;
diff --git a/Engine/Plugins/Media/ElectraPlayer/Source/ElectraPlayerRuntime/Private/Runtime/Player/DASH/PlaylistReaderDASH_Internal.h b/Engine/Plugins/Media/ElectraPlayer/Source/ElectraPlayerRuntime/Private/Runtime/Player/DASH/PlaylistReaderDASH_Internal.h
--- a/Engine/Plugins/Media/ElectraPlayer/Source/ElectraPlayerRuntime/Private/Runtime/Player/DASH/PlaylistReaderDASH_Internal.h
+++ b/Engine/Plugins/Media/ElectraPlayer/Source/ElectraPlayerRuntime/Private/Runtime/Player/DASH/PlaylistReaderDASH_Internal.h
@@ -904,7 +904,7 @@ public:
 	FTimeValue GetDesiredLiveLatency() const;
 
 	FTimeValue GetMPDValidityEndTime() const;
-	FTimeValue GetLastPeriodEndTime() const;
+	FTimeValue GetLastPeriodEndTime(bool bForTimeline) const;
 	FTimeValue GetMinimumUpdatePeriod() const;
 	bool AreUpdatesExpected() const;
 	bool IsStaticType() const;
