ElectraPlayer: fixing a PVS warning

#rnx

diff --git a/Engine/Plugins/Media/ElectraPlayer/Source/ElectraPlayerRuntime/Private/Runtime/Demuxer/ParserISO13818-1.cpp b/Engine/Plugins/Media/ElectraPlayer/Source/ElectraPlayerRuntime/Private/Runtime/Demuxer/ParserISO13818-1.cpp
--- a/Engine/Plugins/Media/ElectraPlayer/Source/ElectraPlayerRuntime/Private/Runtime/Demuxer/ParserISO13818-1.cpp
+++ b/Engine/Plugins/Media/ElectraPlayer/Source/ElectraPlayerRuntime/Private/Runtime/Demuxer/ParserISO13818-1.cpp
@@ -1214,11 +1214,14 @@ FParserISO13818_1::ETableResult FParserISO13818_1::HandlePSITable(const FPayload
 		}
 		return ETableResult::Continue;
 	}
-	// Forbidden ?
-	else if (table_id == 0xff)
-	{
-		return ETableResult::Continue;
-	}
+	/*
+	Already checked above.
+		// Forbidden ?
+		else if (table_id == 0xff)
+		{
+			return ETableResult::Continue;
+		}
+	*/
 	else
 	{
 		/*
