Fixed some 'deprecated' FString usage.

diff --git a/Engine/Plugins/2D/Paper2D/Source/PaperSpriteSheetImporter/Private/PaperSpriteSheetImportFactory.cpp b/Engine/Plugins/2D/Paper2D/Source/PaperSpriteSheetImporter/Private/PaperSpriteSheetImportFactory.cpp
--- a/Engine/Plugins/2D/Paper2D/Source/PaperSpriteSheetImporter/Private/PaperSpriteSheetImportFactory.cpp
+++ b/Engine/Plugins/2D/Paper2D/Source/PaperSpriteSheetImporter/Private/PaperSpriteSheetImportFactory.cpp
@@ -65,7 +65,7 @@ UObject* UPaperSpriteSheetImportFactory::FactoryCreateText(UClass* InClass, UObj
 	UPaperSpriteSheet* Result = nullptr;
 	
 	const FString NameForErrors(InName.ToString());
-	const FString FileContent(BufferEnd - Buffer, Buffer);
+	const FString FileContent = FString::ConstructFromPtrSize(Buffer, BufferEnd - Buffer);
 	
 	if (Importer.ImportFromString(FileContent, NameForErrors, /*bSilent=*/ false) &&
 		Importer.ImportTextures(LongPackagePath, CurrentSourcePath))
diff --git a/Engine/Plugins/2D/Paper2D/Source/PaperTiledImporter/Private/PaperTiledImporterFactory.cpp b/Engine/Plugins/2D/Paper2D/Source/PaperTiledImporter/Private/PaperTiledImporterFactory.cpp
--- a/Engine/Plugins/2D/Paper2D/Source/PaperTiledImporter/Private/PaperTiledImporterFactory.cpp
+++ b/Engine/Plugins/2D/Paper2D/Source/PaperTiledImporter/Private/PaperTiledImporterFactory.cpp
@@ -148,7 +148,7 @@ UObject* UPaperTiledImporterFactory::FactoryCreateText(UClass* InClass, UObject*
  	const FString LongPackagePath = FPackageName::GetLongPackagePath(InParent->GetOutermost()->GetPathName());
  
  	const FString NameForErrors(InName.ToString());
- 	const FString FileContent(BufferEnd - Buffer, Buffer);
+ 	const FString FileContent = FString::ConstructFromPtrSize(Buffer, BufferEnd - Buffer);
  	TSharedPtr<FJsonObject> DescriptorObject = ParseJSON(FileContent, NameForErrors);
 
  	UPaperTileMap* Result = nullptr;
diff --git a/Engine/Plugins/Editor/AssetReferenceRestrictions/Source/AssetReferenceRestrictions/Private/AssetReferencingDomains.cpp b/Engine/Plugins/Editor/AssetReferenceRestrictions/Source/AssetReferenceRestrictions/Private/AssetReferencingDomains.cpp
--- a/Engine/Plugins/Editor/AssetReferenceRestrictions/Source/AssetReferenceRestrictions/Private/AssetReferencingDomains.cpp
+++ b/Engine/Plugins/Editor/AssetReferenceRestrictions/Source/AssetReferenceRestrictions/Private/AssetReferencingDomains.cpp
@@ -84,7 +84,7 @@ public:
 		if (RemainingPath.FindChar(TEXT('/'), /*out*/ DirectorySeparatorIndex))
 		{
 			check(DirectorySeparatorIndex > 0);
-			const FString DirectoryName(DirectorySeparatorIndex, RemainingPath.GetData());
+			const FString DirectoryName = FString::ConstructFromPtrSize(RemainingPath.GetData(), DirectorySeparatorIndex);
 
 			TSharedPtr<FDomainPathNode>& ChildFolder = SubFolders.FindOrAdd(DirectoryName);
 			if (!ChildFolder.IsValid())
diff --git a/Engine/Plugins/Editor/AssetSearch/Source/Private/AssetSearchManager.cpp b/Engine/Plugins/Editor/AssetSearch/Source/Private/AssetSearchManager.cpp
--- a/Engine/Plugins/Editor/AssetSearch/Source/Private/AssetSearchManager.cpp
+++ b/Engine/Plugins/Editor/AssetSearch/Source/Private/AssetSearchManager.cpp
@@ -937,7 +937,7 @@ void FAssetSearchManager::StoreIndexForAsset_DDC(const UObject* InAsset)
 void FAssetSearchManager::LoadDDCContentIntoDatabase(const FAssetData& InAsset, const TArray<uint8>& Content, const FString& DerivedDataKey)
 {
 	FUTF8ToTCHAR WByteBuffer((const ANSICHAR*)Content.GetData(), Content.Num());
-	FString IndexedJson(WByteBuffer.Length(), WByteBuffer.Get());
+	FString IndexedJson = FString::ConstructFromPtrSize(WByteBuffer.Get(), WByteBuffer.Length());
 
 	AddOrUpdateAsset(InAsset, IndexedJson, DerivedDataKey);
 }
diff --git a/Engine/Plugins/Experimental/RemoteSession/Source/RemoteSession/Private/Channels/RemoteSessionChannel.cpp b/Engine/Plugins/Experimental/RemoteSession/Source/RemoteSession/Private/Channels/RemoteSessionChannel.cpp
--- a/Engine/Plugins/Experimental/RemoteSession/Source/RemoteSession/Private/Channels/RemoteSessionChannel.cpp
+++ b/Engine/Plugins/Experimental/RemoteSession/Source/RemoteSession/Private/Channels/RemoteSessionChannel.cpp
@@ -33,7 +33,7 @@ void FRemoteSessionChannelRegistry::RemoveChannelFactory(TWeakPtr<IRemoteSession
 
 TSharedPtr<IRemoteSessionChannel> FRemoteSessionChannelRegistry::CreateChannel(const FStringView InChannelName, ERemoteSessionChannelMode InMode, TBackChannelSharedPtr<IBackChannelConnection> InConnection)
 {
-	FString Name = FString(InChannelName.Len(), InChannelName.GetData());
+	FString Name = FString::ConstructFromPtrSize(InChannelName.GetData(), InChannelName.Len());
 
 	if (!RegisteredFactories.Contains(Name))
 	{
diff --git a/Engine/Plugins/Experimental/RemoteSession/Source/RemoteSession/Private/RemoteSessionModule.cpp b/Engine/Plugins/Experimental/RemoteSession/Source/RemoteSession/Private/RemoteSessionModule.cpp
--- a/Engine/Plugins/Experimental/RemoteSession/Source/RemoteSession/Private/RemoteSessionModule.cpp
+++ b/Engine/Plugins/Experimental/RemoteSession/Source/RemoteSession/Private/RemoteSessionModule.cpp
@@ -174,7 +174,7 @@ bool FRemoteSessionModule::HandleSettingsSaved()
 
 void FRemoteSessionModule::AddChannelFactory(const FStringView InChannelName, ERemoteSessionChannelMode InHostMode, TWeakPtr<IRemoteSessionChannelFactoryWorker> Worker)
 {
-	FRemoteSessionChannelRegistry::Get().RegisterChannelFactory(*FString(InChannelName.Len(), InChannelName.GetData()), InHostMode, Worker);
+	FRemoteSessionChannelRegistry::Get().RegisterChannelFactory(*FString::ConstructFromPtrSize(InChannelName.GetData(), InChannelName.Len()), InHostMode, Worker);
 }
 
 void FRemoteSessionModule::RemoveChannelFactory(TWeakPtr<IRemoteSessionChannelFactoryWorker> Worker)
diff --git a/Engine/Plugins/FX/Niagara/Source/NiagaraEditor/Private/NiagaraNodeConvert.cpp b/Engine/Plugins/FX/Niagara/Source/NiagaraEditor/Private/NiagaraNodeConvert.cpp
--- a/Engine/Plugins/FX/Niagara/Source/NiagaraEditor/Private/NiagaraNodeConvert.cpp
+++ b/Engine/Plugins/FX/Niagara/Source/NiagaraEditor/Private/NiagaraNodeConvert.cpp
@@ -676,7 +676,7 @@ void UNiagaraNodeConvert::AutowireNewNode(UEdGraphPin* FromPin)
 		for (int32 i = 0; i < AutowireSwizzle.Len(); ++i)
 		{
 			TCHAR Char = AutowireSwizzle[i];
-			FString CharStr = FString(1, &Char);
+			FString CharStr = FString::ConstructFromPtrSize(&Char, 1);
 			SrcPath.Empty(SrcPath.Num());
 			DestPath.Empty(DestPath.Num());
 
diff --git a/Engine/Plugins/Media/ElectraCDM/Source/ElectraCDM/Private/ElectraCDMUtils.cpp b/Engine/Plugins/Media/ElectraCDM/Source/ElectraCDM/Private/ElectraCDMUtils.cpp
--- a/Engine/Plugins/Media/ElectraCDM/Source/ElectraCDM/Private/ElectraCDMUtils.cpp
+++ b/Engine/Plugins/Media/ElectraCDM/Source/ElectraCDM/Private/ElectraCDMUtils.cpp
@@ -56,7 +56,7 @@ namespace ElectraCDMUtils
 	FString ArrayToString(const TArray<uint8>& InArray, int32 InStartAt)
 	{
 		FUTF8ToTCHAR cnv((const ANSICHAR*)InArray.GetData() + InStartAt, InArray.Num() - InStartAt);
-		FString UTF8Text(cnv.Length(), cnv.Get());
+		FString UTF8Text = FString::ConstructFromPtrSize(cnv.Get(), cnv.Length());
 		return MoveTemp(UTF8Text);
 	}
 
diff --git a/Engine/Plugins/Media/ElectraPlayer/Source/ElectraPlayerPlugin/Private/ElectraPlayerPlugin.cpp b/Engine/Plugins/Media/ElectraPlayer/Source/ElectraPlayerPlugin/Private/ElectraPlayerPlugin.cpp
--- a/Engine/Plugins/Media/ElectraPlayer/Source/ElectraPlayerPlugin/Private/ElectraPlayerPlugin.cpp
+++ b/Engine/Plugins/Media/ElectraPlayer/Source/ElectraPlayerPlugin/Private/ElectraPlayerPlugin.cpp
@@ -218,7 +218,7 @@ public:
 	virtual FText GetText() const override
 	{
 		FUTF8ToTCHAR cnv((const ANSICHAR*)Subtitle->GetData().GetData(), Subtitle->GetData().Num());
-		FString UTF8Text(cnv.Length(), cnv.Get());
+		FString UTF8Text = FString::ConstructFromPtrSize(cnv.Get(), cnv.Length());
 		return FText::FromString(UTF8Text);
 	}
 
diff --git a/Engine/Plugins/Media/ElectraPlayer/Source/ElectraPlayerRuntime/Private/Runtime/Demuxer/ParserISO14496-12.cpp b/Engine/Plugins/Media/ElectraPlayer/Source/ElectraPlayerRuntime/Private/Runtime/Demuxer/ParserISO14496-12.cpp
--- a/Engine/Plugins/Media/ElectraPlayer/Source/ElectraPlayerRuntime/Private/Runtime/Demuxer/ParserISO14496-12.cpp
+++ b/Engine/Plugins/Media/ElectraPlayer/Source/ElectraPlayerRuntime/Private/Runtime/Demuxer/ParserISO14496-12.cpp
@@ -136,7 +136,7 @@ namespace Electra
 				if (bIsUTF8)
 				{
 					FUTF8ToTCHAR cnv((const ANSICHAR*)SrcBuf.GetData(), SrcBuf.Num());
-					return FString(cnv.Length(), cnv.Get());
+					return FString::ConstructFromPtrSize(cnv.Get(), cnv.Length());
 				}
 				else
 				{
@@ -1520,7 +1520,7 @@ namespace Electra
 			Language[1] = (char)(0x60 + ((Value16 & 0x03e0) >> 5));
 			Language[2] = (char)(0x60 + (Value16 & 0x001f));
 			// Try to map the ISO-639-2T language code to the shorter ISO-639-1 code if possible.
-			Lang639_1 = ISO639::MapTo639_1(FString(FMEDIA_STATIC_ARRAY_COUNT(Language), Language));
+			Lang639_1 = ISO639::MapTo639_1(FString::ConstructFromPtrSize(Language, FMEDIA_STATIC_ARRAY_COUNT(Language)));
 			RETURN_IF_ERROR(ParseInfo->Reader()->Read(Value16));				// pre_defined (in QuickTime this held 'Quality')
 			return Error;
 		}
@@ -2423,7 +2423,7 @@ namespace Electra
 						return UEMEDIA_ERROR_FORMAT_ERROR;
 					}
 					FUTF8ToTCHAR cnv((const ANSICHAR*)FontName, FontNameLength);
-					fr.FontName = FString(cnv.Length(), cnv.Get());
+					fr.FontName = FString::ConstructFromPtrSize(cnv.Get(), cnv.Length());
 				}
 			}
 
@@ -5051,7 +5051,7 @@ namespace Electra
 			const FFillerData* Filler = NameBox ? NameBox->GetFillerData() : nullptr;
 			if (Filler)
 			{
-				FString name(Filler->Size, static_cast<const ANSICHAR*>(Filler->Data));
+				FString name = FString::ConstructFromPtrSize(static_cast<const ANSICHAR*>(Filler->Data), Filler->Size);
 				return TMediaOptionalValue<FString>(name);
 			}
 			return TMediaOptionalValue<FString>();
diff --git a/Engine/Plugins/Media/ElectraPlayer/Source/ElectraPlayerRuntime/Private/Runtime/Player/DASH/PlaylistReaderDASH.cpp b/Engine/Plugins/Media/ElectraPlayer/Source/ElectraPlayerRuntime/Private/Runtime/Player/DASH/PlaylistReaderDASH.cpp
--- a/Engine/Plugins/Media/ElectraPlayer/Source/ElectraPlayerRuntime/Private/Runtime/Player/DASH/PlaylistReaderDASH.cpp
+++ b/Engine/Plugins/Media/ElectraPlayer/Source/ElectraPlayerRuntime/Private/Runtime/Player/DASH/PlaylistReaderDASH.cpp
@@ -1111,7 +1111,7 @@ FErrorDetail FPlaylistReaderDASH::GetXMLResponseString(FString& OutXMLString, FR
 			return CreateErrorAndLog(FString::Printf(TEXT("Document has unsupported UTF-32 LE BOM!")), ERRCODE_DASH_MPD_UNSUPPORTED_DOCUMENT_ENCODING);
 		}
 		FUTF8ToTCHAR TextConv((const ANSICHAR*)ResponseBytes, NumResponseBytes);
-		FString XML(TextConv.Length(), TextConv.Get());
+		FString XML = FString::ConstructFromPtrSize(TextConv.Get(), TextConv.Length());
 		OutXMLString = MoveTemp(XML);
 	}
 	return FErrorDetail();
@@ -1126,7 +1126,7 @@ FErrorDetail FPlaylistReaderDASH::GetResponseString(FString& OutString, FResourc
 		int32 NumResponseBytes = ResponseBuffer->Num();
 		const uint8* ResponseBytes = (const uint8*)ResponseBuffer->GetLinearReadData();
 		FUTF8ToTCHAR TextConv((const ANSICHAR*)ResponseBytes, NumResponseBytes);
-		FString UTF8Text(TextConv.Length(), TextConv.Get());
+		FString UTF8Text = FString::ConstructFromPtrSize(TextConv.Get(), TextConv.Length());
 		OutString = MoveTemp(UTF8Text);
 	}
 	return FErrorDetail();
diff --git a/Engine/Plugins/Media/ElectraPlayer/Source/ElectraPlayerRuntime/Private/Runtime/Player/HLS/PlaylistReaderHLS.cpp b/Engine/Plugins/Media/ElectraPlayer/Source/ElectraPlayerRuntime/Private/Runtime/Player/HLS/PlaylistReaderHLS.cpp
--- a/Engine/Plugins/Media/ElectraPlayer/Source/ElectraPlayerRuntime/Private/Runtime/Player/HLS/PlaylistReaderHLS.cpp
+++ b/Engine/Plugins/Media/ElectraPlayer/Source/ElectraPlayerRuntime/Private/Runtime/Player/HLS/PlaylistReaderHLS.cpp
@@ -1045,7 +1045,7 @@ FErrorDetail FPlaylistReaderHLS::ParsePlaylist(FPlaylistRequestPtr FromRequest)
 		int32 RequestBytes = FromRequest->GetReceiveBuffer()->Num();
 		// The FromRequest buffer is not zero terminated. We need to pick the correct FString constructor for converting the chars into TCHARs while adding the terminating zero!
 		FUTF8ToTCHAR TextConv((const ANSICHAR*)FromRequest->GetReceiveBuffer()->GetLinearReadData(), RequestBytes);
-		FString UTF8String(TextConv.Length(), TextConv.Get());
+		FString UTF8String = FString::ConstructFromPtrSize(TextConv.Get(), TextConv.Length());
 		ParseError = Parser.Parse(UTF8String, Playlist);
 	}
 
diff --git a/Engine/Plugins/Media/ElectraPlayer/Source/ElectraPlayerRuntime/Private/Runtime/Utilities/UtilsMP4.cpp b/Engine/Plugins/Media/ElectraPlayer/Source/ElectraPlayerRuntime/Private/Runtime/Utilities/UtilsMP4.cpp
--- a/Engine/Plugins/Media/ElectraPlayer/Source/ElectraPlayerRuntime/Private/Runtime/Utilities/UtilsMP4.cpp
+++ b/Engine/Plugins/Media/ElectraPlayer/Source/ElectraPlayerRuntime/Private/Runtime/Utilities/UtilsMP4.cpp
@@ -222,7 +222,7 @@ void UtilsMP4::FMetadataParser::ParseBoxDataList(const FString& AsCategory, cons
 			Lang[0] = (char)(0x60 + ((LanguageIndicator & 0x7c00) >> 10));
 			Lang[1] = (char)(0x60 + ((LanguageIndicator & 0x03e0) >> 5));
 			Lang[2] = (char)(0x60 + (LanguageIndicator & 0x001f));
-			Item->Language = FString(3, Lang);
+			Item->Language = FString::ConstructFromPtrSize(Lang, 3);
 		}
 
 		bool bSet = true;
diff --git a/Engine/Plugins/Media/ElectraPlayer/Source/ElectraPlayerRuntime/Private/Runtime/Utilities/UtilsMPEG.cpp b/Engine/Plugins/Media/ElectraPlayer/Source/ElectraPlayerRuntime/Private/Runtime/Utilities/UtilsMPEG.cpp
--- a/Engine/Plugins/Media/ElectraPlayer/Source/ElectraPlayerRuntime/Private/Runtime/Utilities/UtilsMPEG.cpp
+++ b/Engine/Plugins/Media/ElectraPlayer/Source/ElectraPlayerRuntime/Private/Runtime/Utilities/UtilsMPEG.cpp
@@ -236,7 +236,7 @@ namespace Electra
 							*OutPtr = SrcPtr;
 						}
 						auto Cnv = StringCast<TCHAR>(reinterpret_cast<const UTF8CHAR*>(ConvBuf.GetData()), ConvBuf.Num());
-						return FString(Cnv.Length(), Cnv.Get());
+						return FString::ConstructFromPtrSize(Cnv.Get(), Cnv.Length());
 					}
 					else if (Encoding == 1 && NumBytes > 1)
 					{
@@ -251,7 +251,7 @@ namespace Electra
 						if (BOM == 0xfffe)
 						{
 							auto Cnv = StringCast<TCHAR>(reinterpret_cast<const UCS2CHAR*>(SrcPtr+2), NumSrcChars/2-1);
-							return FString(Cnv.Length(), Cnv.Get());
+							return FString::ConstructFromPtrSize(Cnv.Get(), Cnv.Length());
 						}
 						else if (BOM == 0xfeff)
 						{
@@ -263,7 +263,7 @@ namespace Electra
 								be[k+1] = SrcPtr[k+0];
 							}
 							auto Cnv = StringCast<TCHAR>(reinterpret_cast<const UCS2CHAR*>(be.GetData()+2), NumSrcChars/2-1);
-							return FString(Cnv.Length(), Cnv.Get());
+							return FString::ConstructFromPtrSize(Cnv.Get(), Cnv.Length());
 						}
 					}
 				}
diff --git a/Engine/Plugins/Media/ElectraUtil/Source/ElectraBase/Private/StringHelpers.cpp b/Engine/Plugins/Media/ElectraUtil/Source/ElectraBase/Private/StringHelpers.cpp
--- a/Engine/Plugins/Media/ElectraUtil/Source/ElectraBase/Private/StringHelpers.cpp
+++ b/Engine/Plugins/Media/ElectraUtil/Source/ElectraBase/Private/StringHelpers.cpp
@@ -38,7 +38,7 @@ namespace Electra
 				}
 			}
 			auto Cnv = StringCast<TCHAR>(reinterpret_cast<const UTF8CHAR*>(ConvBuf.GetData()), ConvBuf.Num());
-			return FString(Cnv.Length(), Cnv.Get());
+			return FString::ConstructFromPtrSize(Cnv.Get(), Cnv.Length());
 		}
 
 		int32 FindFirstOf(const FString& InString, const FString& SplitAt, int32 FirstPos)
diff --git a/Engine/Plugins/Online/SocketSubsystemEOS/Source/SocketSubsystemEOS/Private/InternetAddrEOS.cpp b/Engine/Plugins/Online/SocketSubsystemEOS/Source/SocketSubsystemEOS/Private/InternetAddrEOS.cpp
--- a/Engine/Plugins/Online/SocketSubsystemEOS/Source/SocketSubsystemEOS/Private/InternetAddrEOS.cpp
+++ b/Engine/Plugins/Online/SocketSubsystemEOS/Source/SocketSubsystemEOS/Private/InternetAddrEOS.cpp
@@ -108,7 +108,7 @@ void FInternetAddrEOS::SetRawIp(const TArray<uint8>& RawAddr)
 {
 	// Need auto here, as might give us different return type depending on size of TCHAR
 	auto ConvertedTCHARData = StringCast<TCHAR>(reinterpret_cast<const ANSICHAR*>(RawAddr.GetData()), RawAddr.Num());
-	const FString IpAsString(ConvertedTCHARData.Length(), ConvertedTCHARData.Get());
+	const FString IpAsString = FString::ConstructFromPtrSize(ConvertedTCHARData.Get(), ConvertedTCHARData.Length());
 
 	bool bUnused;
 	SetIp(*IpAsString, bUnused);
diff --git a/Engine/Plugins/Runtime/Harmonix/Source/HarmonixDspEditor/Private/FusionPatchAssetFactory.cpp b/Engine/Plugins/Runtime/Harmonix/Source/HarmonixDspEditor/Private/FusionPatchAssetFactory.cpp
--- a/Engine/Plugins/Runtime/Harmonix/Source/HarmonixDspEditor/Private/FusionPatchAssetFactory.cpp
+++ b/Engine/Plugins/Runtime/Harmonix/Source/HarmonixDspEditor/Private/FusionPatchAssetFactory.cpp
@@ -242,7 +242,7 @@ UObject* UFusionPatchAssetFactory::FactoryCreateText(UClass* InClass, UObject* I
 	
 	const FString SourceFile = GetCurrentFilename();
 	FString JsonString;
-	const FString DtaString(BufferEnd - Buffer, Buffer);
+	const FString DtaString = FString::ConstructFromPtrSize(Buffer, BufferEnd - Buffer);
 
 	FDtaParser::DtaStringToJsonString(DtaString, JsonString);
 	FString ErrorMessage;
diff --git a/Engine/Source/Developer/Android/AndroidPlatformEditor/Private/AndroidLicenseDialog.cpp b/Engine/Source/Developer/Android/AndroidPlatformEditor/Private/AndroidLicenseDialog.cpp
--- a/Engine/Source/Developer/Android/AndroidPlatformEditor/Private/AndroidLicenseDialog.cpp
+++ b/Engine/Source/Developer/Android/AndroidPlatformEditor/Private/AndroidLicenseDialog.cpp
@@ -87,7 +87,7 @@ void SAndroidLicenseDialog::Construct(const FArguments& InArgs)
 
 				{
 					const FUTF8ToTCHAR ConvertedString(reinterpret_cast<ANSICHAR*>(LicenseStart), LicenseLength);
-					LicenseText = FString(ConvertedString.Length(), ConvertedString.Get());
+					LicenseText = FString::ConstructFromPtrSize(ConvertedString.Get(), ConvertedString.Length());
 				}
 
 				FSHA1::HashBuffer(LicenseStart, LicenseLength, LicenseHash.Hash);
diff --git a/Engine/Source/Developer/TraceInsights/Private/Insights/MemoryProfiler/Widgets/SMemAllocTableTreeView.cpp b/Engine/Source/Developer/TraceInsights/Private/Insights/MemoryProfiler/Widgets/SMemAllocTableTreeView.cpp
--- a/Engine/Source/Developer/TraceInsights/Private/Insights/MemoryProfiler/Widgets/SMemAllocTableTreeView.cpp
+++ b/Engine/Source/Developer/TraceInsights/Private/Insights/MemoryProfiler/Widgets/SMemAllocTableTreeView.cpp
@@ -1895,7 +1895,7 @@ bool SMemAllocTableTreeView::BuildOpenSourceSubMenuItems(FMenuBuilder& MenuBuild
 				const int32 FileNameLen = FCString::Strlen(Frame->Symbol->File);
 				if (FileNameLen > MaxFileNameLen)
 				{
-					FString FileNameStr = TEXT("...") + FString(MaxFileNameLen, Frame->Symbol->File + (FileNameLen - MaxFileNameLen));
+					FString FileNameStr = TEXT("...") + FString::ConstructFromPtrSize(Frame->Symbol->File + (FileNameLen - MaxFileNameLen), MaxFileNameLen);
 					FileName = FText::FromString(FileNameStr);
 				}
 				else
@@ -1908,7 +1908,7 @@ bool SMemAllocTableTreeView::BuildOpenSourceSubMenuItems(FMenuBuilder& MenuBuild
 				const int32 SymbolNameLen = FCString::Strlen(Frame->Symbol->Name);
 				if (SymbolNameLen > MaxSymbolNameLen)
 				{
-					FString SymbolNameStr = TEXT("...") + FString(MaxSymbolNameLen, Frame->Symbol->Name + (SymbolNameLen - MaxSymbolNameLen));
+					FString SymbolNameStr = TEXT("...") + FString::ConstructFromPtrSize(Frame->Symbol->Name + (SymbolNameLen - MaxSymbolNameLen), MaxSymbolNameLen);
 					SymbolName = FText::FromString(SymbolNameStr);
 				}
 				else
diff --git a/Engine/Source/Developer/TraceInsights/Private/Insights/StoreService/DiagnosticsSessionAnalyzer.cpp b/Engine/Source/Developer/TraceInsights/Private/Insights/StoreService/DiagnosticsSessionAnalyzer.cpp
--- a/Engine/Source/Developer/TraceInsights/Private/Insights/StoreService/DiagnosticsSessionAnalyzer.cpp
+++ b/Engine/Source/Developer/TraceInsights/Private/Insights/StoreService/DiagnosticsSessionAnalyzer.cpp
@@ -37,15 +37,15 @@ bool FDiagnosticsSessionAnalyzer::OnEvent(uint16 RouteId, EStyle, const FOnEvent
 		uint8 AppNameOffset = EventData.GetValue<uint8>("AppNameOffset");
 		uint8 CommandLineOffset = EventData.GetValue<uint8>("CommandLineOffset");
 
-		Platform = FString(AppNameOffset, (const ANSICHAR*)Attachment);
+		Platform = FString::ConstructFromPtrSize((const ANSICHAR*)Attachment, AppNameOffset);
 
 		Attachment += AppNameOffset;
 		int32 AppNameLength = CommandLineOffset - AppNameOffset;
-		AppName = FString(AppNameLength, (const ANSICHAR*)Attachment);
+		AppName = FString::ConstructFromPtrSize((const ANSICHAR*)Attachment, AppNameLength);
 
 		Attachment += AppNameLength;
 		int32 CommandLineLength = EventData.GetAttachmentSize() - CommandLineOffset;
-		CommandLine = FString(CommandLineLength, (const ANSICHAR*)Attachment);
+		CommandLine = FString::ConstructFromPtrSize((const ANSICHAR*)Attachment, CommandLineLength);
 
 		ConfigurationType = (EBuildConfiguration) EventData.GetValue<uint8>("ConfigurationType");
 		TargetType = (EBuildTargetType) EventData.GetValue<uint8>("TargetType");
diff --git a/Engine/Source/Editor/UMGEditor/Private/WidgetBlueprintEditorUtils.cpp b/Engine/Source/Editor/UMGEditor/Private/WidgetBlueprintEditorUtils.cpp
--- a/Engine/Source/Editor/UMGEditor/Private/WidgetBlueprintEditorUtils.cpp
+++ b/Engine/Source/Editor/UMGEditor/Private/WidgetBlueprintEditorUtils.cpp
@@ -2564,7 +2564,7 @@ FString RemoveSuffixFromName(const FString OldName)
 			}
 		}
 	}
-	return FString(NameLen, *OldName);
+	return FString::ConstructFromPtrSize(*OldName, NameLen);
 }
 
 FString FWidgetBlueprintEditorUtils::FindNextValidName(UWidgetTree* WidgetTree, const FString& Name)
diff --git a/Engine/Source/Runtime/ExternalRPCRegistry/Private/ExternalRpcRegistry.cpp b/Engine/Source/Runtime/ExternalRPCRegistry/Private/ExternalRpcRegistry.cpp
--- a/Engine/Source/Runtime/ExternalRPCRegistry/Private/ExternalRpcRegistry.cpp
+++ b/Engine/Source/Runtime/ExternalRPCRegistry/Private/ExternalRpcRegistry.cpp
@@ -285,7 +285,7 @@ void UExternalRpcRegistry::AddRequestToLedger(const FHttpServerRequest& Request)
 		FRpcLedgerEntry NewEntry;
 		NewEntry.RpcName = Request.Headers[TEXT("rpcname")][0];
 		FUTF8ToTCHAR WByteBuffer(reinterpret_cast<const ANSICHAR*>(Request.Body.GetData()), Request.Body.Num());
-		NewEntry.RequestBody = FString(WByteBuffer.Length(), WByteBuffer.Get());
+		NewEntry.RequestBody = FString::ConstructFromPtrSize(WByteBuffer.Get(), WByteBuffer.Length());
 		NewEntry.RequestTime = FDateTime::UtcNow();
 		RequestLedger.Add(NewEntry);
 	}
diff --git a/Engine/Source/Runtime/Online/Stomp/Private/StompMessage.cpp b/Engine/Source/Runtime/Online/Stomp/Private/StompMessage.cpp
--- a/Engine/Source/Runtime/Online/Stomp/Private/StompMessage.cpp
+++ b/Engine/Source/Runtime/Online/Stomp/Private/StompMessage.cpp
@@ -37,7 +37,7 @@ FString FStompMessage::GetBodyAsString() const
 	// TODO get encoding from content-type header. Currently assuming utf8
 	FStompBuffer& Body = Frame->GetBody();
 	FUTF8ToTCHAR Convert((ANSICHAR*)Body.GetData(), Body.Num());
-	return FString(Convert.Length(), Convert.Get());
+	return FString::ConstructFromPtrSize(Convert.Get(), Convert.Length());
 }
 
 const uint8* FStompMessage::GetRawBody() const
diff --git a/Engine/Source/Runtime/WebBrowser/Private/Native/NativeJSScripting.cpp b/Engine/Source/Runtime/WebBrowser/Private/Native/NativeJSScripting.cpp
--- a/Engine/Source/Runtime/WebBrowser/Private/Native/NativeJSScripting.cpp
+++ b/Engine/Source/Runtime/WebBrowser/Private/Native/NativeJSScripting.cpp
@@ -199,7 +199,7 @@ int32 ParseParams(const FString& ParamStr, TArray<FString>& OutArray)
 		int32 DelimLimit = 4;
 		while (const TCHAR *At = FCString::Strstr(Start, TEXT("/")))
 		{
-			OutArray.Emplace(At - Start, Start);
+			OutArray.Emplace(FString::ConstructFromPtrSize(Start, At - Start));
 			Start = At + 1;
 			if (--DelimLimit == 0)
 			{
