Electra: Fixed MP3 stream reader using the given response length in place and thus modifying it which now indicates a write error to the HTTP module

#rnx

diff --git a/Engine/Plugins/Media/ElectraPlayer/Source/ElectraPlayerRuntime/Private/Runtime/Player/mpegaudio/StreamReaderMPEGAudio.cpp b/Engine/Plugins/Media/ElectraPlayer/Source/ElectraPlayerRuntime/Private/Runtime/Player/mpegaudio/StreamReaderMPEGAudio.cpp
--- a/Engine/Plugins/Media/ElectraPlayer/Source/ElectraPlayerRuntime/Private/Runtime/Player/mpegaudio/StreamReaderMPEGAudio.cpp
+++ b/Engine/Plugins/Media/ElectraPlayer/Source/ElectraPlayerRuntime/Private/Runtime/Player/mpegaudio/StreamReaderMPEGAudio.cpp
@@ -253,10 +253,11 @@ void FStreamReaderMPEGAudio::FLiveRequest::OnStatusCodeReceived(FHttpRequestPtr
 	}
 }
 
-void FStreamReaderMPEGAudio::FLiveRequest::OnProcessRequestStream(void* InDataPtr, int64& InLength)
+void FStreamReaderMPEGAudio::FLiveRequest::OnProcessRequestStream(void* InDataPtr, int64& InOutLength)
 {
 	if (StatusCode == 200)
 	{
+		int64 InLength = InOutLength;
 		while(InLength > 0)
 		{
 			// Are we receiving metadata right now?
@@ -295,12 +296,12 @@ void FStreamReaderMPEGAudio::FLiveRequest::OnProcessRequestStream(void* InDataPt
 					int64 BufSizeRequired = rb->Num() + DataBytesNow;
 					if (!rb->EnlargeTo(BufSizeRequired))
 					{
-						InLength = 0;
+						InOutLength = 0;
 						return;
 					}
 					if (!rb->PushData(reinterpret_cast<const uint8*>(InDataPtr), DataBytesNow))
 					{
-						InLength = 0;
+						InOutLength = 0;
 						return;
 					}
 
@@ -309,7 +310,7 @@ void FStreamReaderMPEGAudio::FLiveRequest::OnProcessRequestStream(void* InDataPt
 					if (MaxDataBytes && rb->Num() > MaxDataBytes)
 					{
 						bHasFailed = true;
-						InLength = 0;
+						InOutLength = 0;
 						return;
 					}
 				}
