Electra: Updating utility functions

#rnx

diff --git a/Engine/Plugins/Media/ElectraCodecs/Source/ElectraDecoders/Private/Utils/MPEG/ElectraUtilsMPEGVideo.cpp b/Engine/Plugins/Media/ElectraCodecs/Source/ElectraDecoders/Private/Utils/MPEG/ElectraUtilsMPEGVideo.cpp
--- a/Engine/Plugins/Media/ElectraCodecs/Source/ElectraDecoders/Private/Utils/MPEG/ElectraUtilsMPEGVideo.cpp
+++ b/Engine/Plugins/Media/ElectraCodecs/Source/ElectraDecoders/Private/Utils/MPEG/ElectraUtilsMPEGVideo.cpp
@@ -2,6 +2,7 @@
 
 #include "Utils/MPEG/ElectraUtilsMPEGVideo.h"
 #include "Utils/ElectraBitstreamReader.h"
+#include "Utils/MPEG/ElectraUtilsMPEGVideo_H265.h"
 #include "ElectraDecodersUtils.h"
 
 namespace ElectraDecodersUtil
@@ -257,7 +258,7 @@ namespace ElectraDecodersUtil
 		{
 			return ParsedSPSs.Num();
 		}
-		
+
 		const FISO23008_2_seq_parameter_set_data& FHEVCDecoderConfigurationRecord::GetParsedSPS(int32 SpsIndex) const
 		{
 			return ParsedSPSs[SpsIndex];
@@ -288,6 +289,118 @@ namespace ElectraDecodersUtil
 			ParsedSPSs.Empty();
 		}
 
+		bool FHEVCDecoderConfigurationRecord::CreateFromCodecSpecificData(const TArray<uint8>& InFromCSD)
+		{
+			if (!InFromCSD.IsEmpty())
+			{
+				CodecSpecificData = InFromCSD;
+				TArray<ElectraDecodersUtil::MPEG::H265::FNaluInfo> Nalus;
+				if (!ElectraDecodersUtil::MPEG::H265::ParseBitstreamForNALUs(Nalus, InFromCSD.GetData(), InFromCSD.Num()))
+				{
+					return false;
+				}
+				FArray aVPS, aSPS, aPPS;
+				for(int32 i=0; i<Nalus.Num(); ++i)
+				{
+					if (Nalus[i].Type == 32)
+					{
+						aVPS.NALUnitType = 32;
+						aVPS.NALUs.Emplace_GetRef() = MakeArrayView<const uint8>(InFromCSD.GetData() + Nalus[i].Offset + Nalus[i].UnitLength, Nalus[i].Size);
+					}
+					else if (Nalus[i].Type == 33)
+					{
+						TMap<uint32, ElectraDecodersUtil::MPEG::H265::FSequenceParameterSet> SPSs;
+						if (ElectraDecodersUtil::MPEG::H265::ParseSequenceParameterSet(SPSs, InFromCSD.GetData() + Nalus[i].Offset + Nalus[i].UnitLength, Nalus[i].Size))
+						{
+							ElectraDecodersUtil::MPEG::H265::FSequenceParameterSet sps(SPSs.CreateConstIterator().Value());
+							GeneralProfileSpace = sps.profile_tier_level.general_profile_space;
+							GeneralTierFlag = sps.profile_tier_level.general_tier_flag;
+							GeneralProfileIDC = sps.profile_tier_level.general_profile_idc;
+							GeneralProfileCompatibilityFlags = sps.profile_tier_level.general_profile_compatibility_flags;
+							GeneralConstraintIndicatorFlags = sps.profile_tier_level.general_constraint_indicator_flags;
+							GeneralLevelIDC = sps.profile_tier_level.general_level_idc;
+							check(sps.vui_parameters_present_flag);
+							MinSpatialSegmentationIDC = sps.vui_parameters.min_spatial_segmentation_idc;
+							ParallelismType = 0;	// for simplicity as 0 indicates `unknown`.
+							ChromaFormat = sps.chroma_format_idc;
+							BitDepthLumaMinus8 = sps.bit_depth_luma_minus8;
+							BitDepthChromaMinus8 = sps.bit_depth_chroma_minus8;
+							AverageFrameRate = 0;	// unknown
+							ConstantFrameRate = 0;	// may or may not
+							NumTemporalLayers = 0;	// unknown
+							TemporalIdNested = 0;	// may or may not
+							NALUnitLengthMinus1 = 3;
+						}
+						else
+						{
+							return false;
+						}
+						aSPS.NALUnitType = 33;
+						aSPS.NALUs.Emplace_GetRef() = MakeArrayView<const uint8>(InFromCSD.GetData() + Nalus[i].Offset + Nalus[i].UnitLength, Nalus[i].Size);
+					}
+					else if (Nalus[i].Type == 34)
+					{
+						aPPS.NALUnitType = 34;
+						aPPS.NALUs.Emplace_GetRef() = MakeArrayView<const uint8>(InFromCSD.GetData() + Nalus[i].Offset + Nalus[i].UnitLength, Nalus[i].Size);
+					}
+				}
+				if (aVPS.NALUs.IsEmpty() || aSPS.NALUs.IsEmpty() || aPPS.NALUs.IsEmpty())
+				{
+					return false;
+				}
+				Arrays.Emplace(MoveTemp(aVPS));
+				Arrays.Emplace(MoveTemp(aSPS));
+				Arrays.Emplace(MoveTemp(aPPS));
+
+				// Create the raw configuration record now.
+				ConfigurationVersion = 1;
+				FElectraBitstreamWriter wr;
+				wr.PutBits(ConfigurationVersion, 8);
+				wr.PutBits(GeneralProfileSpace, 2);
+				wr.PutBits(GeneralTierFlag, 1);
+				wr.PutBits(GeneralProfileIDC, 5);
+				wr.PutBits(GeneralProfileCompatibilityFlags, 32);
+				wr.PutBits64(GeneralConstraintIndicatorFlags, 48);
+				wr.PutBits(GeneralLevelIDC, 8);
+				wr.PutBits(15U, 4);
+				wr.PutBits(MinSpatialSegmentationIDC, 12);
+				wr.PutBits(63U, 6);
+				wr.PutBits(ParallelismType, 2);
+				wr.PutBits(63U, 6);
+				wr.PutBits(ChromaFormat, 2);
+				wr.PutBits(31U, 5);
+				wr.PutBits(BitDepthLumaMinus8, 3);
+				wr.PutBits(31U, 5);
+				wr.PutBits(BitDepthChromaMinus8, 3);
+				wr.PutBits(AverageFrameRate, 16);
+				wr.PutBits(ConstantFrameRate, 2);
+				wr.PutBits(NumTemporalLayers, 3);
+				wr.PutBits(TemporalIdNested, 1);
+				wr.PutBits(NALUnitLengthMinus1, 2);
+				wr.PutBits(static_cast<uint32>(Arrays.Num()), 8);
+				for(int32 i=0; i<Arrays.Num(); ++i)
+				{
+					const FArray& a(Arrays[i]);
+					wr.PutBits(a.Completeness, 1);
+					wr.PutBits(0U, 1);
+					wr.PutBits(a.NALUnitType, 6);
+					wr.PutBits(static_cast<uint32>(a.NALUs.Num()), 16);
+					for(int32 j=0; j<a.NALUs.Num(); ++j)
+					{
+						const TArray<uint8>& d(a.NALUs[j]);
+						wr.PutBits(static_cast<uint32>(d.Num()), 16);
+						for(int32 k=0; k<d.Num(); ++k)
+						{
+							wr.PutBits(d[k], 8);
+						}
+					}
+				}
+				wr.GetArray(RawData);
+				return true;
+			}
+			return false;
+		}
+
 		bool FHEVCDecoderConfigurationRecord::Parse()
 		{
 			CodecSpecificData.Empty();
@@ -933,7 +1046,7 @@ namespace ElectraDecodersUtil
 					return false;
 				}
 			}
-			
+
 			uint8 long_term_ref_pics_present_flag = BitReader.GetBits(1);
 			if (long_term_ref_pics_present_flag)
 			{
@@ -944,7 +1057,7 @@ namespace ElectraDecodersUtil
 					BitReader.SkipBits(1);							//	used_by_curr_pic_lt_sps_flag[ i ]
 				}
 			}
-			
+
 			OutSPS.sps_temporal_mvp_enabled_flag = BitReader.GetBits(1);
 			OutSPS.strong_intra_smoothing_enabled_flag = BitReader.GetBits(1);
 			OutSPS.vui_parameters_present_flag = BitReader.GetBits(1);
@@ -1114,7 +1227,7 @@ namespace ElectraDecodersUtil
 									payload_bit_equal_to_zero()
 								}
 							}
-					
+
 						H.264:
 							if (!byte_aligned())
 							{
diff --git a/Engine/Plugins/Media/ElectraCodecs/Source/ElectraDecoders/Private/Utils/MPEG/ElectraUtilsMPEGVideo_H265.cpp b/Engine/Plugins/Media/ElectraCodecs/Source/ElectraDecoders/Private/Utils/MPEG/ElectraUtilsMPEGVideo_H265.cpp
--- a/Engine/Plugins/Media/ElectraCodecs/Source/ElectraDecoders/Private/Utils/MPEG/ElectraUtilsMPEGVideo_H265.cpp
+++ b/Engine/Plugins/Media/ElectraCodecs/Source/ElectraDecoders/Private/Utils/MPEG/ElectraUtilsMPEGVideo_H265.cpp
@@ -615,6 +615,7 @@ namespace ElectraDecodersUtil
 					{
 						general_profile_compatibility_flag[i] = br.GetBits(1);
 					}
+					general_constraint_indicator_flags = FBitstreamReader(br).GetBits64(48);	// elaborate PeekBits(48)
 					general_progressive_source_flag = br.GetBits(1);
 					general_interlaced_source_flag = br.GetBits(1);
 					general_non_packed_constraint_flag = br.GetBits(1);
diff --git a/Engine/Plugins/Media/ElectraCodecs/Source/ElectraDecoders/Public/Utils/ElectraBitstreamReader.h b/Engine/Plugins/Media/ElectraCodecs/Source/ElectraDecoders/Public/Utils/ElectraBitstreamReader.h
--- a/Engine/Plugins/Media/ElectraCodecs/Source/ElectraDecoders/Public/Utils/ElectraBitstreamReader.h
+++ b/Engine/Plugins/Media/ElectraCodecs/Source/ElectraDecoders/Public/Utils/ElectraBitstreamReader.h
@@ -294,6 +294,18 @@ public:
 			Mask >>= 1;
 		}
 	}
+	void PutBits64(uint64 InValue, uint32 InNumBits)
+	{
+		uint32 Upper = static_cast<uint32>(InValue >> 32U);
+		uint32 Lower = static_cast<uint32>(InValue);
+		if (InNumBits > 32)
+		{
+			const uint32 nb = InNumBits - 32;
+			PutBits(Upper, nb);
+			InNumBits -= nb;
+		}
+		PutBits(Lower, InNumBits);
+	}
 	void AlignToBytes(uint32 InFillBitsWith)
 	{
 		if (BitPosition)
diff --git a/Engine/Plugins/Media/ElectraCodecs/Source/ElectraDecoders/Public/Utils/MPEG/ElectraUtilsMPEGVideo.h b/Engine/Plugins/Media/ElectraCodecs/Source/ElectraDecoders/Public/Utils/MPEG/ElectraUtilsMPEGVideo.h
--- a/Engine/Plugins/Media/ElectraCodecs/Source/ElectraDecoders/Public/Utils/MPEG/ElectraUtilsMPEGVideo.h
+++ b/Engine/Plugins/Media/ElectraCodecs/Source/ElectraDecoders/Public/Utils/MPEG/ElectraUtilsMPEGVideo.h
@@ -375,6 +375,7 @@ namespace ElectraDecodersUtil
 		public:
 			void SetRawData(const void* Data, int64 Size);
 			const TArray<uint8>& GetRawData() const;
+			bool CreateFromCodecSpecificData(const TArray<uint8>& InFromCSD);
 
 			void Reset();
 
diff --git a/Engine/Plugins/Media/ElectraCodecs/Source/ElectraDecoders/Public/Utils/MPEG/ElectraUtilsMPEGVideo_H265.h b/Engine/Plugins/Media/ElectraCodecs/Source/ElectraDecoders/Public/Utils/MPEG/ElectraUtilsMPEGVideo_H265.h
--- a/Engine/Plugins/Media/ElectraCodecs/Source/ElectraDecoders/Public/Utils/MPEG/ElectraUtilsMPEGVideo_H265.h
+++ b/Engine/Plugins/Media/ElectraCodecs/Source/ElectraDecoders/Public/Utils/MPEG/ElectraUtilsMPEGVideo_H265.h
@@ -128,6 +128,7 @@ namespace ElectraDecodersUtil
 					general_profile_idc = 0;
 					FMemory::Memzero(general_profile_compatibility_flag);
 					general_profile_compatibility_flags = 0;
+					general_constraint_indicator_flags = 0;
 					general_progressive_source_flag = 0;
 					general_interlaced_source_flag = 0;
 					general_non_packed_constraint_flag = 0;
@@ -208,6 +209,7 @@ namespace ElectraDecodersUtil
 				uint8 general_profile_idc;								// u(5)
 				uint8 general_profile_compatibility_flag[32];			// u(1)
 				uint32 general_profile_compatibility_flags;				// u(32), same as general_profile_compatibility_flag[32] but in a single word
+				uint64 general_constraint_indicator_flags;				// u(48), bits starting with `general_progressive_source_flag` and extending through the 44 constraint bits
 				uint8 general_progressive_source_flag;					// u(1)
 				uint8 general_interlaced_source_flag;					// u(1)
 				uint8 general_non_packed_constraint_flag;				// u(1)
diff --git a/Engine/Plugins/Media/ElectraPlayer/Source/ElectraPlayerRuntime/Private/Runtime/Demuxer/ParserISO13818-1.cpp b/Engine/Plugins/Media/ElectraPlayer/Source/ElectraPlayerRuntime/Private/Runtime/Demuxer/ParserISO13818-1.cpp
--- a/Engine/Plugins/Media/ElectraPlayer/Source/ElectraPlayerRuntime/Private/Runtime/Demuxer/ParserISO13818-1.cpp
+++ b/Engine/Plugins/Media/ElectraPlayer/Source/ElectraPlayerRuntime/Private/Runtime/Demuxer/ParserISO13818-1.cpp
@@ -209,8 +209,8 @@ private:
 
 	struct FDTSPTS
 	{
-		TOptional<int64> DTS;
-		TOptional<int64> PTS;
+		TOptional<uint64> DTS;
+		TOptional<uint64> PTS;
 	};
 
 	struct FResidualPESData
@@ -238,10 +238,10 @@ private:
 	void DeselectAllPESStreams();
 	void ActivateUserStreamSelection();
 
-	EPESPacketResult ParseADTSAAC(TArray<FESPacket>& OutPackets, ElectraDecodersUtil::FElectraBitstreamReader& InOutBR, const FDTSPTS& InDTSPTS, FResidualPESData* InResidualData);
-	EPESPacketResult ParseMPEGAudio(TArray<FESPacket>& OutPackets, ElectraDecodersUtil::FElectraBitstreamReader& InOutBR, const FDTSPTS& InDTSPTS, FResidualPESData* InResidualData);
-	EPESPacketResult ParseAVC(TArray<FESPacket>& OutPackets, ElectraDecodersUtil::FElectraBitstreamReader& InOutBR, const FDTSPTS& InDTSPTS, FResidualPESData* InResidualData);
-	EPESPacketResult ParseHEVC(TArray<FESPacket>& OutPackets, ElectraDecodersUtil::FElectraBitstreamReader& InOutBR, const FDTSPTS& InDTSPTS, FResidualPESData* InResidualData);
+	EPESPacketResult ParseADTSAAC(TArray<FESPacket>& OutPackets, ElectraDecodersUtil::FElectraBitstreamReader& InOutBR, const FDTSPTS& InDTSPTS, FResidualPESData* InResidualData, bool bFlushResiduals);
+	EPESPacketResult ParseMPEGAudio(TArray<FESPacket>& OutPackets, ElectraDecodersUtil::FElectraBitstreamReader& InOutBR, const FDTSPTS& InDTSPTS, FResidualPESData* InResidualData, bool bFlushResiduals);
+	EPESPacketResult ParseAVC(TArray<FESPacket>& OutPackets, ElectraDecodersUtil::FElectraBitstreamReader& InOutBR, const FDTSPTS& InDTSPTS, FResidualPESData* InResidualData, bool bFlushResiduals);
+	EPESPacketResult ParseHEVC(TArray<FESPacket>& OutPackets, ElectraDecodersUtil::FElectraBitstreamReader& InOutBR, const FDTSPTS& InDTSPTS, FResidualPESData* InResidualData, bool bFlushResiduals);
 
 	TUniquePtr<FStaticInitSegReader> InitSegReader;
 
@@ -256,13 +256,14 @@ private:
 	TSharedPtrTS<FProgramTable> CurrentProgramTable;
 
 	TOptional<FUserProgramSelection> PendingUserProgramSelection;
-	TSharedPtrTS<FPESData> AvailablePESPacket;
+	TArray<TSharedPtrTS<FPESData>> AvailablePESPackets;
 
 	TMap<int32, FResidualPESData> ResidualPESDataMap;
 
 	FErrorDetail ErrorDetail;
 	EParseState ParseState = EParseState::Continue;
 	int64 FileOffset = 0;
+	uint64 TimestampOffset = 0;
 
 
 	static const uint32 CRCTable[256];
@@ -321,11 +322,12 @@ IParserISO13818_1::EParseState FParserISO13818_1::BeginParsing(IPlayerSessionSer
 	PIDStreamData.Empty();
 	CurrentProgramTable.Reset();
 	PendingUserProgramSelection.Reset();
-	AvailablePESPacket.Reset();
+	AvailablePESPackets.Empty();
 	Current.Reset();
 	ErrorDetail.Clear();
 	ParseState = EParseState::Continue;
 	FileOffset = 0;
+	TimestampOffset = InSourceInfo.TimestampOffset;
 
 	if (InSourceInfo.InitSegmentData.IsValid() && InSourceInfo.InitSegmentData->Num())
 	{
@@ -583,7 +585,6 @@ IParserISO13818_1::EParseState FParserISO13818_1::ParseNextPacket(IPlayerSession
 							{
 								ResidualPESDataMap.Remove((int32) PID);
 							}
-							check(Payloads.Num() <= 1);	// if there are several, which should only happen in PSI sections, do we need to break up their handling?
 							for(int32 i=0; i<Payloads.Num(); ++i)
 							{
 								switch(Payloads[i].Type)
@@ -599,7 +600,7 @@ IParserISO13818_1::EParseState FParserISO13818_1::ParseNextPacket(IPlayerSession
 									}
 									case EPayloadType::PES:
 									{
-										AvailablePESPacket = MoveTemp(Payloads[i].PESData);
+										AvailablePESPackets.Emplace(MoveTemp(Payloads[i].PESData));
 										NextParseState = EParseState::HavePESPacket;
 										break;
 									}
@@ -637,7 +638,7 @@ IParserISO13818_1::EParseState FParserISO13818_1::ParseNextPacket(IPlayerSession
 		}
 		else
 		{
-			NextParseState = EParseState::Failed;
+			NextParseState = EParseState::ReadError;
 		}
 	}
 
@@ -682,12 +683,10 @@ IParserISO13818_1::EParseState FParserISO13818_1::Parse(IPlayerSessionServices*
 		}
 		case EParseState::HavePESPacket:
 		{
-			if (AvailablePESPacket.IsValid())
+			if (AvailablePESPackets.IsEmpty())
 			{
-				UE_LOG(LogElectraMPEGTSParser, Log, TEXT("Received new PES packet, but user did not handle it."));
-				AvailablePESPacket.Reset();
+				ParseState = EParseState::Continue;
 			}
-			ParseState = EParseState::Continue;
 			break;
 		}
 		case EParseState::Failed:
@@ -704,9 +703,23 @@ IParserISO13818_1::EParseState FParserISO13818_1::Parse(IPlayerSessionServices*
 				pidPes->FinishCurrentPESPacket(PESPacket);
 				if (PESPacket.IsValid())
 				{
-					AvailablePESPacket = MoveTemp(PESPacket);
+					AvailablePESPackets.Emplace(MoveTemp(PESPacket));
 					return EParseState::HavePESPacket;
 				}
+				else
+				{
+					// Are there residuals that, at least for video PES streams will most likely
+					// contain the last frame? We assume the residuals to be a complete frame and
+					// not partial data that carries over into the next segment.
+					if (ResidualPESDataMap.Contains(pesIt.Key))
+					{
+						PESPacket = MakeShared<FPESData, ESPMode::ThreadSafe>();
+						PESPacket->PID = pesIt.Key;
+						PESPacket->StreamType = pidPes->StreamInfo.StreamType;
+						AvailablePESPackets.Emplace(MoveTemp(PESPacket));
+						return EParseState::HavePESPacket;
+					}
+				}
 			}
 			// Flush any residuals. We do not expect consecutive segments to need them.
 			ResidualPESDataMap.Empty();
@@ -796,7 +809,10 @@ void FParserISO13818_1::SelectProgramStreams(int32 InProgramNumber, const TArray
 
 TSharedPtrTS<FParserISO13818_1::FPESData> FParserISO13818_1::GetPESPacket()
 {
-	return MoveTemp(AvailablePESPacket);
+	check(AvailablePESPackets.Num());
+	TSharedPtrTS<FParserISO13818_1::FPESData> Next(AvailablePESPackets[0]);
+	AvailablePESPackets.RemoveAt(0);
+	return MoveTemp(Next);
 }
 
 
@@ -933,11 +949,7 @@ void FParserISO13818_1::FPIDStream::ExtractValidPESPackets(TArray<TSharedPtrTS<F
 				InNumBytesAddedNow = 0;
 			}
 
-			// The size is unknown. This means that we collect data until we get a packet that does not have `payload_unit_start_indicator` set
-			// AND uses an adaptation field to skip over initial bytes, because the data is to end with the last byte of the packet.
-			// Safety check that the accumulation buffer does not contain more than we have just added prior to getting here.
-			// If that is so, then the payload probably ended exactly on the last byte of the previous packet, with or
-			// without having an adaptation field.
+			// The size is unknown. This means that we collect data until we get a packet that has `payload_unit_start_indicator` set.
 			if (InCurrent.bIsStart && AccumulationBuffer.Num() > InNumBytesAddedNow)
 			{
 				GatheringSection.TotalSize = AccumulationBuffer.Num() - InNumBytesAddedNow;
@@ -956,26 +968,6 @@ void FParserISO13818_1::FPIDStream::ExtractValidPESPackets(TArray<TSharedPtrTS<F
 				GatheringSection.PCR = InCurrent.PCR;
 				continue;
 			}
-
-			// Last block?
-			bool bLastBlock = InCurrent.BytesSkippedUntilPayload > 0;
-			// We also treat a packet that is not a start and has an adaptation field that is all zero as the last packet for this PES stream.
-			// This may not be quite right, but this appears to be a common case and an adaptation field that is not setting anything seems to
-			// be pointless, other than to realize a 2-byte "stuffing".
-			bLastBlock |= !InCurrent.bIsStart && InCurrent.AdaptationFieldSize == 1 && InCurrent.AdaptationFieldFirstByte == 0;
-			if (bLastBlock)
-			{
-				// Move the entire accumulation buffer over.
-				TSharedPtrTS<FPESData> NewPES = MakeSharedTS<FPESData>();
-				NewPES->PID = PID;
-				NewPES->StreamType = StreamInfo.StreamType;
-				NewPES->PacketData = MoveTemp(PacketDataBuffer);
-				NewPES->bRandomAccessIndicator = GatheringSection.bRandomAccessIndicator;
-				NewPES->PCR = GatheringSection.PCR;
-				OutPESSections.Emplace(MoveTemp(NewPES));
-				GatheringSection.Reset();
-				PacketDataBuffer = MakeSharedTS<TArray<uint8>>();
-			}
 			return;
 		}
 	}
@@ -1598,136 +1590,175 @@ void FParserISO13818_1::ProcessDescriptor(ElectraDecodersUtil::FElectraBitstream
 
 IParserISO13818_1::EPESPacketResult FParserISO13818_1::ParsePESPacket(TArray<IParserISO13818_1::FESPacket>& OutPackets, TSharedPtrTS<IParserISO13818_1::FPESData> InPESPacket)
 {
-	if (!InPESPacket.IsValid() || !InPESPacket->PacketData.IsValid() || InPESPacket->PacketData->Num() < 6)
+	if (!InPESPacket.IsValid() || (InPESPacket->PacketData.IsValid() && InPESPacket->PacketData->Num() < 6))
 	{
 		return EPESPacketResult::Invalid;
 	}
-	// Check that this packet has the proper start code.
-	const TArray<uint8>& PESData = *(InPESPacket->PacketData);
-	const uint8* Data = PESData.GetData();
-	if (Data[0] != 0 || Data[1] != 0 || Data[2] != 1)
+	if (!InPESPacket->PacketData.IsValid() && ParseState != EParseState::EOS)
 	{
 		return EPESPacketResult::Invalid;
 	}
-	const uint8 stream_id = Data[3];
-	const int32 PES_packet_length = (int32) (((uint32)Data[4] << 8) | Data[5]);
-	if (PES_packet_length && PES_packet_length+6 != PESData.Num())
+
+	FDTSPTS DTSPTS;
+	ElectraDecodersUtil::FElectraBitstreamReader br;
+	uint8 stream_id = 0xbe;
+	bool bHandleStreamId = false;
+	if (InPESPacket->PacketData.IsValid())
 	{
-		// Packet size mismatch.
-		return PES_packet_length+6 > PESData.Num() ? EPESPacketResult::Truncated : EPESPacketResult::Invalid;
+		// Check that this packet has the proper start code.
+		const TArray<uint8>& PESData = *(InPESPacket->PacketData);
+		const uint8* Data = PESData.GetData();
+		if (Data[0] != 0 || Data[1] != 0 || Data[2] != 1)
+		{
+			return EPESPacketResult::Invalid;
+		}
+		stream_id = Data[3];
+		const int32 PES_packet_length = (int32) (((uint32)Data[4] << 8) | Data[5]);
+		if (PES_packet_length && PES_packet_length+6 != PESData.Num())
+		{
+			// Packet size mismatch.
+			return PES_packet_length+6 > PESData.Num() ? EPESPacketResult::Truncated : EPESPacketResult::Invalid;
+		}
+		if (stream_id != 0xbc /* program_stream_map */ &&
+			stream_id != 0xbe /* padding_stream */ &&
+			stream_id != 0xbf /* private_stream_2 */ &&
+			stream_id != 0xf0 /* ECM */ &&
+			stream_id != 0xf1 /* EMM */ &&
+			stream_id != 0xff /* program_stream_directory */ &&
+			stream_id != 0xf2 /* Rec. ITU-T H.222.0 | ISO/IEC 13818-1 Annex A or ISO/IEC 13818-6_DSMCC_stream */ &&
+			stream_id != 0xf8 /* Rec. ITU-T H.222.1 type E */)
+		{
+			br.SetData(Data, PESData.NumBytes(), 6);
+			const uint32 OneZero = br.GetBits(2);
+			const uint32 PES_scrambling_control = br.GetBits(2);
+			const uint32 PES_priority = br.GetBits(1);
+			const uint32 data_alignment_indicator = br.GetBits(1);
+			const uint32 copyright = br.GetBits(1);
+			const uint32 original_or_copy = br.GetBits(1);
+			const uint32 PTS_DTS_flags = br.GetBits(2);
+			const uint32 ESCR_flag = br.GetBits(1);
+			const uint32 ES_rate_flag = br.GetBits(1);
+			const uint32 DSM_trick_mode_flag = br.GetBits(1);
+			const uint32 additional_copy_info_flag = br.GetBits(1);
+			const uint32 PES_CRC_flag = br.GetBits(1);
+			const uint32 PES_extension_flag = br.GetBits(1);
+			const uint32 PES_header_data_length = br.GetBits(8);
+			ElectraDecodersUtil::FElectraBitstreamReader opt(br);
+			br.SkipBytes(PES_header_data_length);
+			if (PTS_DTS_flags == 2)
+			{
+				const uint32 ZeroZeroOneZero = opt.GetBits(4);
+				const uint32 PTS_32_30 = opt.GetBits(3);
+				opt.SkipBits(1);
+				const uint32 PTS_29_15 = opt.GetBits(15);
+				opt.SkipBits(1);
+				const uint32 PTS_14_0 = opt.GetBits(15);
+				opt.SkipBits(1);
+				DTSPTS.PTS = ((((uint64)PTS_32_30 << 30U) | ((uint64)PTS_29_15 << 15U) | (uint64)PTS_14_0) + TimestampOffset) & 0x1ffffffffULL;
+			}
+			else if (PTS_DTS_flags == 3)
+			{
+				const uint32 ZeroZeroOneOne = opt.GetBits(4);
+				const uint32 PTS_32_30 = opt.GetBits(3);
+				opt.SkipBits(1);
+				const uint32 PTS_29_15 = opt.GetBits(15);
+				opt.SkipBits(1);
+				const uint32 PTS_14_0 = opt.GetBits(15);
+				opt.SkipBits(1);
+				const uint32 ZeroZeroZeroOne = opt.GetBits(4);
+				const uint32 DTS_32_30 = opt.GetBits(3);
+				opt.SkipBits(1);
+				const uint32 DTS_29_15 = opt.GetBits(15);
+				opt.SkipBits(1);
+				const uint32 DTS_14_0 = opt.GetBits(15);
+				opt.SkipBits(1);
+				DTSPTS.PTS = ((((uint64)PTS_32_30 << 30U) | ((uint64)PTS_29_15 << 15U) | (uint64)PTS_14_0) + TimestampOffset) & 0x1ffffffffULL;
+				DTSPTS.DTS = ((((uint64)DTS_32_30 << 30U) | ((uint64)DTS_29_15 << 15U) | (uint64)DTS_14_0) + TimestampOffset) & 0x1ffffffffULL;
+			}
+			if (ESCR_flag)
+			{
+				opt.SkipBits(2); // reserved
+				const uint32 ESCR_base_32_30 = opt.GetBits(3);
+				opt.SkipBits(1);
+				const uint32 ESCR_base_29_15 = opt.GetBits(15);
+				opt.SkipBits(1);
+				const uint32 ESCR_base_14_0 = opt.GetBits(15);
+				opt.SkipBits(1);
+				const uint32 ESCR_extension = opt.GetBits(9);
+				opt.SkipBits(1);
+			}
+			if (ES_rate_flag)
+			{
+				opt.SkipBits(1);
+				const uint32 ES_rate = opt.GetBits(22);
+				opt.SkipBits(1);
+			}
+			if (DSM_trick_mode_flag)
+			{
+				const uint32 trick_mode_control = opt.GetBits(3);
+				// ...
+				opt.SkipBits(5);
+			}
+			if (additional_copy_info_flag)
+			{
+				opt.SkipBits(1);
+				const uint32 additional_copy_info = opt.GetBits(7);
+			}
+			if (PES_CRC_flag)
+			{
+				const uint32 previous_PES_packet_CRC = opt.GetBits(16);
+			}
+
+			bHandleStreamId = true;
+		}
 	}
-	if (stream_id != 0xbc /* program_stream_map */ &&
-		stream_id != 0xbe /* padding_stream */ &&
-		stream_id != 0xbf /* private_stream_2 */ &&
-		stream_id != 0xf0 /* ECM */ &&
-		stream_id != 0xf1 /* EMM */ &&
-		stream_id != 0xff /* program_stream_directory */ &&
-		stream_id != 0xf2 /* Rec. ITU-T H.222.0 | ISO/IEC 13818-1 Annex A or ISO/IEC 13818-6_DSMCC_stream */ &&
-		stream_id != 0xf8 /* Rec. ITU-T H.222.1 type E */)
+	else
 	{
-		FDTSPTS DTSPTS;
-		ElectraDecodersUtil::FElectraBitstreamReader br(Data, PESData.NumBytes(), 6);
-		const uint32 OneZero = br.GetBits(2);
-		const uint32 PES_scrambling_control = br.GetBits(2);
-		const uint32 PES_priority = br.GetBits(1);
-		const uint32 data_alignment_indicator = br.GetBits(1);
-		const uint32 copyright = br.GetBits(1);
-		const uint32 original_or_copy = br.GetBits(1);
-		const uint32 PTS_DTS_flags = br.GetBits(2);
-		const uint32 ESCR_flag = br.GetBits(1);
-		const uint32 ES_rate_flag = br.GetBits(1);
-		const uint32 DSM_trick_mode_flag = br.GetBits(1);
-		const uint32 additional_copy_info_flag = br.GetBits(1);
-		const uint32 PES_CRC_flag = br.GetBits(1);
-		const uint32 PES_extension_flag = br.GetBits(1);
-		const uint32 PES_header_data_length = br.GetBits(8);
-		ElectraDecodersUtil::FElectraBitstreamReader opt(br);
-		br.SkipBytes(PES_header_data_length);
-		if (PTS_DTS_flags == 2)
-		{
-			const uint32 ZeroZeroOneZero = opt.GetBits(4);
-			const uint32 PTS_32_30 = opt.GetBits(3);
-			opt.SkipBits(1);
-			const uint32 PTS_29_15 = opt.GetBits(15);
-			opt.SkipBits(1);
-			const uint32 PTS_14_0 = opt.GetBits(15);
-			opt.SkipBits(1);
-			DTSPTS.PTS = (int64) (((uint64)PTS_32_30 << 30U) | ((uint64)PTS_29_15 << 15U) | (uint64)PTS_14_0);
-		}
-		else if (PTS_DTS_flags == 3)
-		{
-			const uint32 ZeroZeroOneOne = opt.GetBits(4);
-			const uint32 PTS_32_30 = opt.GetBits(3);
-			opt.SkipBits(1);
-			const uint32 PTS_29_15 = opt.GetBits(15);
-			opt.SkipBits(1);
-			const uint32 PTS_14_0 = opt.GetBits(15);
-			opt.SkipBits(1);
-			const uint32 ZeroZeroZeroOne = opt.GetBits(4);
-			const uint32 DTS_32_30 = opt.GetBits(3);
-			opt.SkipBits(1);
-			const uint32 DTS_29_15 = opt.GetBits(15);
-			opt.SkipBits(1);
-			const uint32 DTS_14_0 = opt.GetBits(15);
-			opt.SkipBits(1);
-			DTSPTS.PTS = (int64) (((uint64)PTS_32_30 << 30U) | ((uint64)PTS_29_15 << 15U) | (uint64)PTS_14_0);
-			DTSPTS.DTS = (int64) (((uint64)DTS_32_30 << 30U) | ((uint64)DTS_29_15 << 15U) | (uint64)DTS_14_0);
-		}
-		if (ESCR_flag)
-		{
-			opt.SkipBits(2); // reserved
-			const uint32 ESCR_base_32_30 = opt.GetBits(3);
-			opt.SkipBits(1);
-			const uint32 ESCR_base_29_15 = opt.GetBits(15);
-			opt.SkipBits(1);
-			const uint32 ESCR_base_14_0 = opt.GetBits(15);
-			opt.SkipBits(1);
-			const uint32 ESCR_extension = opt.GetBits(9);
-			opt.SkipBits(1);
-		}
-		if (ES_rate_flag)
-		{
-			opt.SkipBits(1);
-			const uint32 ES_rate = opt.GetBits(22);
-			opt.SkipBits(1);
-		}
-		if (DSM_trick_mode_flag)
-		{
-			const uint32 trick_mode_control = opt.GetBits(3);
-			// ...
-			opt.SkipBits(5);
-		}
-		if (additional_copy_info_flag)
-		{
-			opt.SkipBits(1);
-			const uint32 additional_copy_info = opt.GetBits(7);
-		}
-		if (PES_CRC_flag)
-		{
-			const uint32 previous_PES_packet_CRC = opt.GetBits(16);
-		}
+		// This is called to process the residuals from a previous packet, so this needs to be handled.
+		bHandleStreamId = true;
+	}
 
+	if (bHandleStreamId)
+	{
 		TArray<FESPacket> NewPackets;
 		EPESPacketResult PESResult = EPESPacketResult::NotSupported;
 		bool bSuccess = false;
-		check(br.IsByteAligned());
+		bool bFlushResiduals = ParseState == EParseState::EOS;
+		check(br.IsByteAligned());	// even an unset bitstream reader is aligned, so this works when flushing residuals as well
 		switch(InPESPacket->StreamType)
 		{
 			// MPEG audio (layer 1, 2 or 3)
 			case 0x03:
 			// MPEG audio in ADTS format (AAC)
 			case 0x0f:
+			// AVC
+			case 0x1b:
+			// HEVC
+			case 0x24:
 			{
 				FResidualPESData* Residuals = ResidualPESDataMap.Find(InPESPacket->PID);
 				if (InPESPacket->StreamType == 0x0f)
 				{
-					PESResult = ParseADTSAAC(NewPackets, br, DTSPTS, Residuals);
+					PESResult = ParseADTSAAC(NewPackets, br, DTSPTS, Residuals, bFlushResiduals);
+				}
+				else if (InPESPacket->StreamType == 0x03)
+				{
+					PESResult = ParseMPEGAudio(NewPackets, br, DTSPTS, Residuals, bFlushResiduals);
+				}
+				else if (InPESPacket->StreamType == 0x1b)
+				{
+					PESResult = ParseAVC(NewPackets, br, DTSPTS, Residuals, bFlushResiduals);
 				}
 				else
 				{
-					PESResult = ParseMPEGAudio(NewPackets, br, DTSPTS, Residuals);
+					PESResult = ParseHEVC(NewPackets, br, DTSPTS, Residuals, bFlushResiduals);
 				}
-				if (PESResult == EPESPacketResult::Truncated && br.GetRemainingByteLength())
+
+				if (PESResult == EPESPacketResult::Ok)
+				{
+					ResidualPESDataMap.Remove(InPESPacket->PID);
+				}
+				else if (PESResult == EPESPacketResult::Truncated && br.GetRemainingByteLength())
 				{
 					if (!Residuals)
 					{
@@ -1754,18 +1785,10 @@ IParserISO13818_1::EPESPacketResult FParserISO13818_1::ParsePESPacket(TArray<IPa
 					}
 					PESResult = EPESPacketResult::Ok;
 				}
-				break;
-			}
-			// AVC
-			case 0x1b:
-			{
-				PESResult = ParseAVC(NewPackets, br, DTSPTS, nullptr);
-				break;
-			}
-			// HEVC
-			case 0x24:
-			{
-				PESResult = ParseHEVC(NewPackets, br, DTSPTS, nullptr);
+				if (bFlushResiduals)
+				{
+					ResidualPESDataMap.Remove(InPESPacket->PID);
+				}
 				break;
 			}
 			// Dolby
@@ -1802,7 +1825,7 @@ IParserISO13818_1::EPESPacketResult FParserISO13818_1::ParsePESPacket(TArray<IPa
 	return EPESPacketResult::Invalid;
 }
 
-FParserISO13818_1::EPESPacketResult FParserISO13818_1::ParseADTSAAC(TArray<FESPacket>& OutPackets, ElectraDecodersUtil::FElectraBitstreamReader& InOutBR, const FDTSPTS& InDTSPTS, FResidualPESData* InResidualData)
+FParserISO13818_1::EPESPacketResult FParserISO13818_1::ParseADTSAAC(TArray<FESPacket>& OutPackets, ElectraDecodersUtil::FElectraBitstreamReader& InOutBR, const FDTSPTS& InDTSPTS, FResidualPESData* InResidualData, bool bFlushResiduals)
 {
 	ElectraDecodersUtil::FElectraBitstreamReader br(InOutBR);
 	FDTSPTS DTSPTS(InDTSPTS);
@@ -1848,10 +1871,6 @@ FParserISO13818_1::EPESPacketResult FParserISO13818_1::ParseADTSAAC(TArray<FESPa
 		const int32 FrameSize = frame_length - (prot_absent ? 7 : 9);
 		if (FrameSize < 0 || br.GetRemainingByteLength() < FrameSize)
 		{
-			if (!InResidualData)
-			{
-				UE_LOG(LogElectraMPEGTSParser, Warning, TEXT("Remaining PES packet data too small for a complete ADTS frame. Incorrect multiplex?"));
-			}
 			return EPESPacketResult::Truncated;
 		}
 		if (num_frames > 0)
@@ -1907,7 +1926,7 @@ FParserISO13818_1::EPESPacketResult FParserISO13818_1::ParseADTSAAC(TArray<FESPa
 	return EPESPacketResult::Ok;
 }
 
-FParserISO13818_1::EPESPacketResult FParserISO13818_1::ParseMPEGAudio(TArray<FESPacket>& OutPackets, ElectraDecodersUtil::FElectraBitstreamReader& InOutBR, const FDTSPTS& InDTSPTS, FResidualPESData* InResidualData)
+FParserISO13818_1::EPESPacketResult FParserISO13818_1::ParseMPEGAudio(TArray<FESPacket>& OutPackets, ElectraDecodersUtil::FElectraBitstreamReader& InOutBR, const FDTSPTS& InDTSPTS, FResidualPESData* InResidualData, bool bFlushResiduals)
 {
 	ElectraDecodersUtil::FElectraBitstreamReader br(InOutBR);
 	FDTSPTS DTSPTS(InDTSPTS);
@@ -2005,174 +2024,295 @@ FParserISO13818_1::EPESPacketResult FParserISO13818_1::ParseMPEGAudio(TArray<FES
 	return EPESPacketResult::Ok;
 }
 
-FParserISO13818_1::EPESPacketResult FParserISO13818_1::ParseAVC(TArray<FESPacket>& OutPackets, ElectraDecodersUtil::FElectraBitstreamReader& InOutBR, const FDTSPTS& InDTSPTS, FResidualPESData* InResidualData)
+FParserISO13818_1::EPESPacketResult FParserISO13818_1::ParseAVC(TArray<FESPacket>& OutPackets, ElectraDecodersUtil::FElectraBitstreamReader& InOutBR, const FDTSPTS& InDTSPTS, FResidualPESData* InResidualData, bool bFlushResiduals)
 {
-	TArray<ElectraDecodersUtil::MPEG::H264::FNaluInfo> NALUs;
-	const uint8* InData = reinterpret_cast<const uint8*>(InOutBR.GetRemainingData());
-	uint64 InDataLength = InOutBR.GetRemainingByteLength();
-	InOutBR.SkipBytes(InDataLength);
-	// We get an Annex-B stream here which we need to decompose, remove the AUD NALU and separate the SPS and PPS NALUs.
-	if (ElectraDecodersUtil::MPEG::H264::ParseBitstreamForNALUs(NALUs, InData, InDataLength))
+	ElectraDecodersUtil::FElectraBitstreamReader br(InOutBR);
+	FDTSPTS DTSPTS(InDTSPTS);
+
+	// Are we dealing with residuals?
+	if (InResidualData && InResidualData->RemainingData.Num())
 	{
-		// In a first pass calculate the size of the final data.
-		const int32 kSizeOfSizeField = 4;
-		int32 SizeCSD = 0;
-		int32 SizeData = 0;
-		bool bIsIDR = false;
-		bool bHaveAUD = false;
-		for(int32 i=0; i<NALUs.Num(); ++i)
-		{
-			if (NALUs[i].Type == 9)
-			{
-				// AUD
-				if (bHaveAUD)
-				{
-					UE_LOG(LogElectraMPEGTSParser, Warning, TEXT("More than one AUD NALU found in AVC packet"));
-				}
-				bHaveAUD = true;
-			}
-			else if (NALUs[i].Type == 7 || NALUs[i].Type == 8)
+		if (!bFlushResiduals)
+		{
+			InResidualData->RemainingData.Append(reinterpret_cast<const uint8*>(InOutBR.GetRemainingData()), (int32)InOutBR.GetRemainingByteLength());
+		}
+		br.SetData(InResidualData->RemainingData.GetData(), InResidualData->RemainingData.Num());
+		DTSPTS = InResidualData->PreviousDTSPTS;
+	}
+
+	// Deal with potentially multiple frames
+	for(int32 npkt=0; ; ++npkt)
+	{
+		TArray<ElectraDecodersUtil::MPEG::H264::FNaluInfo> NALUs;
+		const uint8* InData = reinterpret_cast<const uint8*>(br.GetRemainingData());
+		uint64 InDataLength = br.GetRemainingByteLength();
+		if (InDataLength == 0)
+		{
+			break;
+		}
+		// We get an Annex-B stream here which we need to decompose, remove the AUD NALU and separate the SPS and PPS NALUs.
+		if (ElectraDecodersUtil::MPEG::H264::ParseBitstreamForNALUs(NALUs, InData, InDataLength))
+		{
+			// Because of some streams splitting video across multiple PES packets of smaller sizes (instead of 0)
+			// and therefore several start flag and PES headers, we need to reassemble the packets.
+			// In order to do this we need to take data enclosed between two AUD NALUs and thus need to have
+			// an additional frame (AUD denotes the start of a frame, not the end).
+			int32 NumAUDNalus = 0;
+			for(int32 i=0; i<NALUs.Num(); ++i)
 			{
-				// SPS or PPS
-				SizeCSD += NALUs[i].Size + kSizeOfSizeField;
+				NumAUDNalus += NALUs[i].Type == 9 ? 1 : 0;
 			}
-			else if (NALUs[i].Type == 12)
+			if (NumAUDNalus >= (bFlushResiduals ? 1 : 2))
 			{
-				// Filler data
+				int32 FirstNALUIndex = -1;
+				int32 LastNALUIndex = -1;
+				const int32 kSizeOfSizeField = 4;
+				int32 SizeCSD = 0;
+				int32 SizeData = 0;
+				bool bIsIDR = false;
+				// In a first pass calculate the size of the final data.
+				for(int32 i=0; i<NALUs.Num(); ++i)
+				{
+					// First NALU we need must be AUD. If not, skip it.
+					if (FirstNALUIndex < 0 && NALUs[i].Type != 9)
+					{
+						continue;
+					}
+					// Take note of the indices of the first and second AUD.
+					if (NALUs[i].Type == 9)
+					{
+						FirstNALUIndex = FirstNALUIndex < 0 ? i : FirstNALUIndex;
+						LastNALUIndex = LastNALUIndex < 0 && i > FirstNALUIndex ? i : LastNALUIndex;
+						if (LastNALUIndex > 0)
+						{
+							break;
+						}
+						continue;
+					}
+
+					if (NALUs[i].Type == 7 || NALUs[i].Type == 8)
+					{
+						// SPS or PPS
+						SizeCSD += NALUs[i].Size + kSizeOfSizeField;
+					}
+					else if (NALUs[i].Type == 12)
+					{
+						// Filler data
+					}
+					else
+					{
+						// Other
+						SizeData += NALUs[i].Size + kSizeOfSizeField;
+						if (NALUs[i].Type == 5)
+						{
+							bIsIDR = true;
+						}
+					}
+				}
+				if (FirstNALUIndex != 0)
+				{
+					UE_LOG(LogElectraMPEGTSParser, Warning, TEXT("First NALU in AVC packet is not an AUD"));
+				}
+
+				FESPacket& pkt = OutPackets.Emplace_GetRef();
+				pkt.bIsSyncFrame = bIsIDR;
+				pkt.SubPacketNum = npkt;
+				pkt.DTS = DTSPTS.DTS;
+				pkt.PTS = DTSPTS.PTS;
+				pkt.CSD = MakeSharedTS<TArray<uint8>>();
+				pkt.CSD->SetNumUninitialized(SizeCSD);
+				pkt.Data = MakeSharedTS<TArray<uint8>>();
+				pkt.Data->SetNumUninitialized(SizeData);
+				uint8* CSDPtr = pkt.CSD->GetData();
+				uint8* DataPtr = pkt.Data->GetData();
+
+				// Second pass, copy the data.
+				for(int32 i=FirstNALUIndex, iMax=bFlushResiduals?NALUs.Num():LastNALUIndex; i<iMax; ++i)
+				{
+					if (NALUs[i].Type == 9 || NALUs[i].Type == 12)
+					{
+						// Skip AUD and filler data
+						continue;
+					}
+					const bool bCSD = NALUs[i].Type == 7 || NALUs[i].Type == 8;
+					uint8** DstPtr = bCSD ? &CSDPtr : &DataPtr;
+					uint32 Size = (uint32)NALUs[i].Size;
+					*((uint32*)(*DstPtr)) = bCSD ? MEDIA_TO_BIG_ENDIAN((uint32)1) : MEDIA_TO_BIG_ENDIAN(Size);
+					(*DstPtr) += 4;
+					uint32 Pos = NALUs[i].Offset + NALUs[i].UnitLength;
+					FMemory::Memcpy((*DstPtr), InData + Pos, Size);
+					(*DstPtr) += Size;
+				}
+				// Remove the data we processed from the residuals
+				int32 ConsumedSize = bFlushResiduals ? (int32) br.GetRemainingByteLength() : (int32) NALUs[LastNALUIndex].Offset;
+				br.SkipBytes(ConsumedSize);
+				InOutBR = br;
+				// Update the DTS and PTS in the residual data
+				if (InResidualData)
+				{
+					InResidualData->PreviousDTSPTS = InDTSPTS;
+				}
 			}
 			else
 			{
-				// Other
-				SizeData += NALUs[i].Size + kSizeOfSizeField;
-				if (NALUs[i].Type == 5)
-				{
-					bIsIDR = true;
-				}
+				// Not enough data yet. Need an additional AUD NALU.
+				InOutBR = br;
+				return EPESPacketResult::Truncated;
 			}
 		}
-
-		FESPacket& pkt = OutPackets.Emplace_GetRef();
-		pkt.bIsSyncFrame = bIsIDR;
-		pkt.SubPacketNum = 0;
-		pkt.DTS = InDTSPTS.DTS;
-		pkt.PTS = InDTSPTS.PTS;
-		pkt.CSD = MakeSharedTS<TArray<uint8>>();
-		pkt.CSD->SetNumUninitialized(SizeCSD);
-		pkt.Data = MakeSharedTS<TArray<uint8>>();
-		pkt.Data->SetNumUninitialized(SizeData);
-		uint8* CSDPtr = pkt.CSD->GetData();
-		uint8* DataPtr = pkt.Data->GetData();
-
-		// Second pass, copy the data.
-		for(int32 i=0; i<NALUs.Num(); ++i)
+		else
 		{
-			if (NALUs[i].Type == 9 || NALUs[i].Type == 12)
-			{
-				// Skip AUD and filler data
-				continue;
-			}
-			const bool bCSD = NALUs[i].Type == 7 || NALUs[i].Type == 8;
-			uint8** DstPtr = bCSD ? &CSDPtr : &DataPtr;
-			uint32 Size = (uint32)NALUs[i].Size;
-			*((uint32*)(*DstPtr)) = bCSD ? MEDIA_TO_BIG_ENDIAN((uint32)1) : MEDIA_TO_BIG_ENDIAN(Size);
-			(*DstPtr) += 4;
-			uint32 Pos = NALUs[i].Offset + NALUs[i].UnitLength;
-			FMemory::Memcpy((*DstPtr), InData + Pos, Size);
-			(*DstPtr) += Size;
+			UE_LOG(LogElectraMPEGTSParser, Warning, TEXT("Failed to parse the AVC packet for NALUs"));
+			return EPESPacketResult::Invalid;
 		}
-		return EPESPacketResult::Ok;
-	}
-	else
-	{
-		UE_LOG(LogElectraMPEGTSParser, Warning, TEXT("Failed to parse the AVC packet for NALUs"));
 	}
-	return EPESPacketResult::Invalid;
+	return EPESPacketResult::Ok;
 }
 
-FParserISO13818_1::EPESPacketResult FParserISO13818_1::ParseHEVC(TArray<FESPacket>& OutPackets, ElectraDecodersUtil::FElectraBitstreamReader& InOutBR, const FDTSPTS& InDTSPTS, FResidualPESData* InResidualData)
+FParserISO13818_1::EPESPacketResult FParserISO13818_1::ParseHEVC(TArray<FESPacket>& OutPackets, ElectraDecodersUtil::FElectraBitstreamReader& InOutBR, const FDTSPTS& InDTSPTS, FResidualPESData* InResidualData, bool bFlushResiduals)
 {
-	TArray<ElectraDecodersUtil::MPEG::H265::FNaluInfo> NALUs;
-	const uint8* InData = reinterpret_cast<const uint8*>(InOutBR.GetRemainingData());
-	uint64 InDataLength = InOutBR.GetRemainingByteLength();
-	InOutBR.SkipBytes(InDataLength);
-	// We get an Annex-B stream here which we need to decompose, remove the AUD NUT and separate the VPS, SPS and PPS NUTs.
-	if (ElectraDecodersUtil::MPEG::H265::ParseBitstreamForNALUs(NALUs, InData, (uint64)InDataLength))
+	ElectraDecodersUtil::FElectraBitstreamReader br(InOutBR);
+	FDTSPTS DTSPTS(InDTSPTS);
+
+	// Are we dealing with residuals?
+	if (InResidualData && InResidualData->RemainingData.Num())
 	{
-		// In a first pass calculate the size of the final data.
-		const int32 kSizeOfSizeField = 4;
-		int32 SizeCSD = 0;
-		int32 SizeData = 0;
-		bool bIsSync = false;
-		bool bHaveAUD = false;
-		for(int32 i=0; i<NALUs.Num(); ++i)
-		{
-			if (NALUs[i].Type == 35)
-			{
-				// AUD
-				if (bHaveAUD)
-				{
-					UE_LOG(LogElectraMPEGTSParser, Warning, TEXT("More than one AUD NUT found in HEVC packet"));
-				}
-				bHaveAUD = true;
-			}
-			else if (NALUs[i].Type == 32 || NALUs[i].Type == 33 || NALUs[i].Type == 34)
+		if (!bFlushResiduals)
+		{
+			InResidualData->RemainingData.Append(reinterpret_cast<const uint8*>(InOutBR.GetRemainingData()), (int32)InOutBR.GetRemainingByteLength());
+		}
+		br.SetData(InResidualData->RemainingData.GetData(), InResidualData->RemainingData.Num());
+		DTSPTS = InResidualData->PreviousDTSPTS;
+	}
+
+	// Deal with potentially multiple frames
+	for(int32 npkt=0; ; ++npkt)
+	{
+		TArray<ElectraDecodersUtil::MPEG::H265::FNaluInfo> NALUs;
+		const uint8* InData = reinterpret_cast<const uint8*>(br.GetRemainingData());
+		uint64 InDataLength = br.GetRemainingByteLength();
+		if (InDataLength == 0)
+		{
+			break;
+		}
+		// We get an Annex-B stream here which we need to decompose, remove the AUD NUT and separate the VPS, SPS and PPS NUTs.
+		if (ElectraDecodersUtil::MPEG::H265::ParseBitstreamForNALUs(NALUs, InData, InDataLength))
+		{
+			// Because of some streams splitting video across multiple PES packets of smaller sizes (instead of 0)
+			// and therefore several start flag and PES headers, we need to reassemble the packets.
+			// In order to do this we need to take data enclosed between two AUD NUTs and thus need to have
+			// an additional frame (AUD denotes the start of a frame, not the end).
+			int32 NumAUDNalus = 0;
+			for(int32 i=0; i<NALUs.Num(); ++i)
 			{
-				// VPS, SPS or PPS
-				SizeCSD += NALUs[i].Size + kSizeOfSizeField;
+				NumAUDNalus += NALUs[i].Type == 35 ? 1 : 0;
 			}
-			else if (NALUs[i].Type == 38)
+			if (NumAUDNalus >= (bFlushResiduals ? 1 : 2))
 			{
-				// Filler data.
-			}
-			else
-			{
-				// Other
-				SizeData += NALUs[i].Size + kSizeOfSizeField;
+				int32 FirstNALUIndex = -1;
+				int32 LastNALUIndex = -1;
+				const int32 kSizeOfSizeField = 4;
+				int32 SizeCSD = 0;
+				int32 SizeData = 0;
+				bool bIsSync = false;
+				// In a first pass calculate the size of the final data.
+				for(int32 i=0; i<NALUs.Num(); ++i)
+				{
+					// First NUT we need must be AUD. If not, skip it.
+					if (FirstNALUIndex < 0 && NALUs[i].Type != 35)
+					{
+						continue;
+					}
+					// Take note of the indices of the first and second AUD.
+					if (NALUs[i].Type == 35)
+					{
+						FirstNALUIndex = FirstNALUIndex < 0 ? i : FirstNALUIndex;
+						LastNALUIndex = LastNALUIndex < 0 && i > FirstNALUIndex ? i : LastNALUIndex;
+						if (LastNALUIndex > 0)
+						{
+							break;
+						}
+						continue;
+					}
 
-				// IDR, CRA or BLA frame?
-				if (NALUs[i].Type >= 16 && NALUs[i].Type <= 21)
+					if (NALUs[i].Type == 32 || NALUs[i].Type == 33 || NALUs[i].Type == 34)
+					{
+						// VPS, SPS or PPS
+						SizeCSD += NALUs[i].Size + kSizeOfSizeField;
+					}
+					else if (NALUs[i].Type == 38)
+					{
+						// Filler data
+					}
+					else
+					{
+						// Other
+						SizeData += NALUs[i].Size + kSizeOfSizeField;
+						// IDR, CRA or BLA frame?
+						if (NALUs[i].Type >= 16 && NALUs[i].Type <= 21)
+						{
+							bIsSync = true;
+						}
+					}
+				}
+				if (FirstNALUIndex != 0)
 				{
-					bIsSync = true;
+					UE_LOG(LogElectraMPEGTSParser, Warning, TEXT("First NUT in HEVC packet is not an AUD"));
 				}
-			}
-		}
-
-		FESPacket& pkt = OutPackets.Emplace_GetRef();
-		pkt.bIsSyncFrame = bIsSync;
-		pkt.SubPacketNum = 0;
-		pkt.DTS = InDTSPTS.DTS;
-		pkt.PTS = InDTSPTS.PTS;
-		pkt.CSD = MakeSharedTS<TArray<uint8>>();
-		pkt.CSD->SetNumUninitialized(SizeCSD);
-		pkt.Data = MakeSharedTS<TArray<uint8>>();
-		pkt.Data->SetNumUninitialized(SizeData);
-		uint8* CSDPtr = pkt.CSD->GetData();
-		uint8* DataPtr = pkt.Data->GetData();
 
-		// Second pass, copy the data.
-		for(int32 i=0; i<NALUs.Num(); ++i)
-		{
-			if (NALUs[i].Type == 35 || NALUs[i].Type == 38)
+				FESPacket& pkt = OutPackets.Emplace_GetRef();
+				pkt.bIsSyncFrame = bIsSync;
+				pkt.SubPacketNum = npkt;
+				pkt.DTS = DTSPTS.DTS;
+				pkt.PTS = DTSPTS.PTS;
+				pkt.CSD = MakeSharedTS<TArray<uint8>>();
+				pkt.CSD->SetNumUninitialized(SizeCSD);
+				pkt.Data = MakeSharedTS<TArray<uint8>>();
+				pkt.Data->SetNumUninitialized(SizeData);
+				uint8* CSDPtr = pkt.CSD->GetData();
+				uint8* DataPtr = pkt.Data->GetData();
+
+				// Second pass, copy the data.
+				for(int32 i=FirstNALUIndex, iMax=bFlushResiduals?NALUs.Num():LastNALUIndex; i<iMax; ++i)
+				{
+					if (NALUs[i].Type == 35 || NALUs[i].Type == 38)
+					{
+						// Skip AUD and filler data
+						continue;
+					}
+					const bool bCSD = NALUs[i].Type == 32 || NALUs[i].Type == 33 || NALUs[i].Type == 34;
+					uint8** DstPtr = bCSD ? &CSDPtr : &DataPtr;
+					uint32 Size = (uint32)NALUs[i].Size;
+					*((uint32*)(*DstPtr)) = bCSD ? MEDIA_TO_BIG_ENDIAN((uint32)1) : MEDIA_TO_BIG_ENDIAN(Size);
+					(*DstPtr) += 4;
+					uint32 Pos = NALUs[i].Offset + NALUs[i].UnitLength;
+					FMemory::Memcpy((*DstPtr), InData + Pos, Size);
+					(*DstPtr) += Size;
+				}
+				// Remove the data we processed from the residuals
+				int32 ConsumedSize = bFlushResiduals ? (int32) br.GetRemainingByteLength() : (int32) NALUs[LastNALUIndex].Offset;
+				br.SkipBytes(ConsumedSize);
+				InOutBR = br;
+				// Update the DTS and PTS in the residual data
+				if (InResidualData)
+				{
+					InResidualData->PreviousDTSPTS = InDTSPTS;
+				}
+			}
+			else
 			{
-				// Skip AUD and filler data
-				continue;
+				// Not enough data yet. Need an additional AUD NALU.
+				InOutBR = br;
+				return EPESPacketResult::Truncated;
 			}
-			const bool bCSD = NALUs[i].Type == 32 || NALUs[i].Type == 33 || NALUs[i].Type == 34;
-			uint8** DstPtr = bCSD ? &CSDPtr : &DataPtr;
-			uint32 Size = (uint32)NALUs[i].Size;
-			*((uint32*)(*DstPtr)) = bCSD ? MEDIA_TO_BIG_ENDIAN((uint32)1) : MEDIA_TO_BIG_ENDIAN(Size);
-			(*DstPtr) += 4;
-			uint32 Pos = NALUs[i].Offset + NALUs[i].UnitLength;
-			FMemory::Memcpy((*DstPtr), InData + Pos, Size);
-			(*DstPtr) += Size;
 		}
-		return EPESPacketResult::Ok;
-	}
-	else
-	{
-		UE_LOG(LogElectraMPEGTSParser, Warning, TEXT("Failed to parse the HEVC packet for NUTs"));
+		else
+		{
+			UE_LOG(LogElectraMPEGTSParser, Warning, TEXT("Failed to parse the HEVC packet for NUTs"));
+			return EPESPacketResult::Invalid;
+		}
 	}
-	return EPESPacketResult::Invalid;
+	return EPESPacketResult::Ok;
 }
 
 
@@ -2198,7 +2338,7 @@ bool FParserISO13818_1::ParseCSD(FStreamCodecInformation& OutParsedCSD, const FE
 				OutParsedCSD.SetCodec4CC(Utils::Make4CC('m','p','g','a'));
 				OutParsedCSD.SetProfile(ElectraDecodersUtil::MPEG::UtilsMPEG123::GetVersion(HeaderValue));
 				OutParsedCSD.SetProfileLevel(ElectraDecodersUtil::MPEG::UtilsMPEG123::GetLayer(HeaderValue));
-				OutParsedCSD.SetCodecSpecifierRFC6381(TEXT("mp4a.6b"));
+				OutParsedCSD.SetCodecSpecifierRFC6381(TEXT("mp4a.6b"));	// alternatively "mp4a.40.34"
 				OutParsedCSD.SetSamplingRate(ElectraDecodersUtil::MPEG::UtilsMPEG123::GetSamplingRate(HeaderValue));
 				OutParsedCSD.SetNumberOfChannels(ElectraDecodersUtil::MPEG::UtilsMPEG123::GetChannelCount(HeaderValue));
 				OutParsedCSD.GetExtras().Set(StreamCodecInformationOptions::SamplesPerBlock, FVariantValue((int64)ElectraDecodersUtil::MPEG::UtilsMPEG123::GetSamplesPerFrame(HeaderValue)));
diff --git a/Engine/Plugins/Media/ElectraPlayer/Source/ElectraPlayerRuntime/Private/Runtime/Demuxer/ParserISO13818-1.h b/Engine/Plugins/Media/ElectraPlayer/Source/ElectraPlayerRuntime/Private/Runtime/Demuxer/ParserISO13818-1.h
--- a/Engine/Plugins/Media/ElectraPlayer/Source/ElectraPlayerRuntime/Private/Runtime/Demuxer/ParserISO13818-1.h
+++ b/Engine/Plugins/Media/ElectraPlayer/Source/ElectraPlayerRuntime/Private/Runtime/Demuxer/ParserISO13818-1.h
@@ -34,7 +34,6 @@ namespace Electra
 			ParseFlag_Default = 0,
 			// Ignore PAT/PMT from the actual stream and rely on info from the init segment (if present)
 			ParseFlag_IgnoreProgramStream = 1 << 0,
-			ParseFlag_Continuation = 1 << 1,
 		};
 
 		struct FSourceInfo
@@ -42,6 +41,7 @@ namespace Electra
 			TSharedPtr<const TArray<uint8>, ESPMode::ThreadSafe> InitSegmentData;
 			int64 InFirstFileByteOffset = -1;
 			int64 InLastFileByteOffset = -1;
+			uint64 TimestampOffset = 0;
 		};
 
 		enum class EParseState
@@ -54,6 +54,8 @@ namespace Electra
 			HavePESPacket,
 			// Failed
 			Failed,
+			// Read error
+			ReadError,
 			// Reached the end of the stream.
 			EOS
 		};
@@ -143,8 +145,8 @@ namespace Electra
 
 		struct FESPacket
 		{
-			TOptional<int64> DTS;
-			TOptional<int64> PTS;
+			TOptional<uint64> DTS;
+			TOptional<uint64> PTS;
 			TSharedPtrTS<TArray<uint8>> CSD;
 			TSharedPtrTS<TArray<uint8>> Data;
 			int32 SubPacketNum = 0;
diff --git a/Engine/Plugins/Media/ElectraPlayer/Source/ElectraPlayerRuntime/Private/Runtime/Demuxer/ParserISO14496-12.cpp b/Engine/Plugins/Media/ElectraPlayer/Source/ElectraPlayerRuntime/Private/Runtime/Demuxer/ParserISO14496-12.cpp
--- a/Engine/Plugins/Media/ElectraPlayer/Source/ElectraPlayerRuntime/Private/Runtime/Demuxer/ParserISO14496-12.cpp
+++ b/Engine/Plugins/Media/ElectraPlayer/Source/ElectraPlayerRuntime/Private/Runtime/Demuxer/ParserISO14496-12.cpp
@@ -5737,6 +5737,7 @@ namespace Electra
 		const IMetadata* GetMetadata(EBaseBoxType InFromBox) const override;
 
 		TSharedPtrTS<IAllTrackIterator> CreateAllTrackIteratorByFilePos(int64 InFromFilePos) const override;
+		TSharedPtrTS<IAllTrackIterator> CreateAllTrackIteratorForTrackIDs(const TArray<uint32>& InTrackIDs) const override;
 
 	private:
 		class FTrack;
@@ -8308,7 +8309,7 @@ namespace Electra
 						if (TFHDBox)
 						{
 							FTrack* Track = NewParsedTrackInfo->GetTrackByID(TFHDBox->GetTrackID());
-							check(Track);
+							// We may not find it because the track is not of a supported type.
 							if (Track)
 							{
 								// The fragmented track needs to have a trex box.
@@ -8411,8 +8412,9 @@ namespace Electra
 							}
 							else
 							{
-								// There needs to be a track with that ID in the moov box
-								return UEMEDIA_ERROR_FORMAT_ERROR;
+								// Either the track is not a supported type or the moov really does not specify
+								// a track with that ID, which would be an error.
+								// We assume that the file itself is ok and that the track is not supported.
 							}
 						}
 						else
@@ -8633,6 +8635,38 @@ namespace Electra
 		return ti;
 	}
 
+	TSharedPtrTS<IParserISO14496_12::IAllTrackIterator> FParserISO14496_12::CreateAllTrackIteratorForTrackIDs(const TArray<uint32>& InTrackIDs) const
+	{
+		TSharedPtrTS<FParserISO14496_12::FAllTrackIterator> ti = MakeShared<FParserISO14496_12::FAllTrackIterator, ESPMode::ThreadSafe>();
+
+		int64 LowestFilePos = TNumericLimits<int64>::Max();
+		for(int32 nTrk=0, nTrkMax=InTrackIDs.Num(); nTrk<nTrkMax; ++nTrk)
+		{
+			const FTrack* Track = ParsedTrackInfo->GetTrackByID(InTrackIDs[nTrk]);
+			if (!Track)
+			{
+				continue;
+			}
+			FTrackIterator* TrkIt = static_cast<FTrackIterator*>(Track->CreateIterator());
+			TSharedPtrTS<ITrackIterator> SafeTrkIt(TrkIt);
+			ti->TrackIterators.Add(SafeTrkIt);
+			UEMediaError err = TrkIt->StartAtFirstInteral();
+			if (err == UEMEDIA_ERROR_OK)
+			{
+				if (TrkIt->GetSampleFileOffset() < LowestFilePos)
+				{
+					LowestFilePos = TrkIt->GetSampleFileOffset();
+					ti->CurrentIterator = SafeTrkIt;
+				}
+			}
+			else if (err == UEMEDIA_ERROR_END_OF_STREAM)
+			{
+				ti->NewlyReachedEOS.Add(TrkIt);
+			}
+		}
+		return ti;
+	}
+
 
 	/*********************************************************************************************************************/
 
diff --git a/Engine/Plugins/Media/ElectraPlayer/Source/ElectraPlayerRuntime/Private/Runtime/Demuxer/ParserISO14496-12.h b/Engine/Plugins/Media/ElectraPlayer/Source/ElectraPlayerRuntime/Private/Runtime/Demuxer/ParserISO14496-12.h
--- a/Engine/Plugins/Media/ElectraPlayer/Source/ElectraPlayerRuntime/Private/Runtime/Demuxer/ParserISO14496-12.h
+++ b/Engine/Plugins/Media/ElectraPlayer/Source/ElectraPlayerRuntime/Private/Runtime/Demuxer/ParserISO14496-12.h
@@ -294,6 +294,8 @@ namespace Electra
 
 		virtual TSharedPtrTS<IAllTrackIterator> CreateAllTrackIteratorByFilePos(int64 InFromFilePos) const = 0;
 
+		virtual TSharedPtrTS<IAllTrackIterator> CreateAllTrackIteratorForTrackIDs(const TArray<uint32>& InTrackIDs) const = 0;
+
 		virtual const ITrack* GetTrackByIndex(int32 Index) const = 0;
 		virtual const ITrack* GetTrackByTrackID(int32 TrackID) const = 0;
 
diff --git a/Engine/Plugins/Media/ElectraPlayer/Source/ElectraPlayerRuntime/Private/Runtime/ErrorDetail.h b/Engine/Plugins/Media/ElectraPlayer/Source/ElectraPlayerRuntime/Private/Runtime/ErrorDetail.h
--- a/Engine/Plugins/Media/ElectraPlayer/Source/ElectraPlayerRuntime/Private/Runtime/ErrorDetail.h
+++ b/Engine/Plugins/Media/ElectraPlayer/Source/ElectraPlayerRuntime/Private/Runtime/ErrorDetail.h
@@ -166,14 +166,35 @@ static inline FErrorDetail PostError(IPlayerSessionServices* InPlayerSessionServ
 }
 
 
-
-class FErrorReporter
-{
-public:
-	FErrorReporter();
-	~FErrorReporter();
-};
-
+#define ELECTRA_CLASS_DEFAULT_ERROR_METHODS(SessionService, ErrorFacility)															\
+inline FErrorDetail CreateError(const FString& InMessage, uint16 InCode)															\
+{																																	\
+	FErrorDetail Error;																												\
+	Error.SetError(UEMEDIA_ERROR_DETAIL);																							\
+	Error.SetFacility(Facility::EFacility::ErrorFacility);																			\
+	Error.SetCode(InCode);																											\
+	Error.SetMessage(InMessage);																									\
+	return Error;																													\
+}																																	\
+inline FErrorDetail CreateErrorAndLog(const FString& InMessage, uint16 InCode)														\
+{																																	\
+	FErrorDetail Error = CreateError(InMessage, InCode);																			\
+	SessionService->PostLog(Facility::EFacility::ErrorFacility, IInfoLog::ELevel::Error, Error.GetPrintable());						\
+	return Error;																													\
+}																																	\
+inline void LogMessage(IInfoLog::ELevel InLevel, const FString& InMessage)															\
+{																																	\
+	SessionService->PostLog(Facility::EFacility::ErrorFacility, InLevel, InMessage);												\
+}																																	\
+inline FErrorDetail PostError(const FErrorDetail& InError)																			\
+{																																	\
+	SessionService->PostError(InError);																								\
+	return InError;																													\
+}																																	\
+inline FErrorDetail PostError(const FString& InMessage, uint16 InCode)																\
+{																																	\
+	return PostError(CreateError(InMessage, InCode));																				\
+}
 
 
 } // namespace Electra
diff --git a/Engine/Plugins/Media/ElectraPlayer/Source/ElectraPlayerRuntime/Private/Runtime/Player/AdaptiveStreamingPlayerInternal.h b/Engine/Plugins/Media/ElectraPlayer/Source/ElectraPlayerRuntime/Private/Runtime/Player/AdaptiveStreamingPlayerInternal.h
--- a/Engine/Plugins/Media/ElectraPlayer/Source/ElectraPlayerRuntime/Private/Runtime/Player/AdaptiveStreamingPlayerInternal.h
+++ b/Engine/Plugins/Media/ElectraPlayer/Source/ElectraPlayerRuntime/Private/Runtime/Player/AdaptiveStreamingPlayerInternal.h
@@ -55,22 +55,6 @@ namespace Electra
 {
 class FAdaptiveStreamingPlayer;
 
-inline int32 StreamTypeToArrayIndex(EStreamType StreamType)
-{
-	switch(StreamType)
-	{
-		case EStreamType::Video:
-			return 0;
-		case EStreamType::Audio:
-			return 1;
-		case EStreamType::Subtitle:
-			return 2;
-		default:
-			return 3;
-	}
-}
-
-
 
 class IAdaptiveStreamingWrappedRenderer : public IMediaRenderer
 {
diff --git a/Engine/Plugins/Media/ElectraPlayer/Source/ElectraPlayerRuntime/Private/Runtime/Player/AdaptiveStreamingPlayerMetrics.h b/Engine/Plugins/Media/ElectraPlayer/Source/ElectraPlayerRuntime/Private/Runtime/Player/AdaptiveStreamingPlayerMetrics.h
--- a/Engine/Plugins/Media/ElectraPlayer/Source/ElectraPlayerRuntime/Private/Runtime/Player/AdaptiveStreamingPlayerMetrics.h
+++ b/Engine/Plugins/Media/ElectraPlayer/Source/ElectraPlayerRuntime/Private/Runtime/Player/AdaptiveStreamingPlayerMetrics.h
@@ -160,6 +160,8 @@ namespace Metrics
 		bool bDidTimeout;						//!< true if a timeout occurred. Only set if timeouts are enabled. Usually the ABR will monitor and abort.
 		bool bParseFailure;						//!< true if the segment could not be parsed
 		bool bIsMissingSegment;					//!< true if the segment was not actually downloaded because it is missing on the timeline.
+		bool bWasSkipped;						//!< true if the segment was skipped over due to internal timestamps being less than expected
+		bool bWasFalloffSegment;				//!< true if the segment was no longer present on the timeline.
 		bool bInsertedFillerData;
 		bool bIsCachedResponse;
 		TArray<IElectraHTTPStreamResponse::FTimingTrace> TimingTraces;
@@ -203,6 +205,8 @@ namespace Metrics
 			bDidTimeout = false;
 			bParseFailure = false;
 			bIsMissingSegment = false;
+			bWasSkipped = false;
+			bWasFalloffSegment = false;
 			bInsertedFillerData = false;
 			bIsCachedResponse = false;
 			TimingTraces.Empty();
diff --git a/Engine/Plugins/Media/ElectraPlayer/Source/ElectraPlayerRuntime/Private/Runtime/Player/AdaptiveStreamingPlayerResourceRequest.h b/Engine/Plugins/Media/ElectraPlayer/Source/ElectraPlayerRuntime/Private/Runtime/Player/AdaptiveStreamingPlayerResourceRequest.h
--- a/Engine/Plugins/Media/ElectraPlayer/Source/ElectraPlayerRuntime/Private/Runtime/Player/AdaptiveStreamingPlayerResourceRequest.h
+++ b/Engine/Plugins/Media/ElectraPlayer/Source/ElectraPlayerRuntime/Private/Runtime/Player/AdaptiveStreamingPlayerResourceRequest.h
@@ -169,15 +169,25 @@ public:
 		else if (Error == 3)					return FString(TEXT("Connection closed"));
 		else if (Error == 4)					return FString(TEXT("No connection"));
 		else if (Error >= 100 && Error < 600)	return FString::Printf(TEXT("HTTP status code %d"), Error);
+		else if (Error == -1)					return FString(TEXT("Could not connect"));
 		else									return FString::Printf(TEXT("Unknown <code %d>"), Error);
 	}
 
+	virtual bool GetHasFinished() const
+	{ return bHasFinished; }
+
 	virtual bool GetWasCanceled() const
 	{ return bWasCanceled; }
 
 	virtual FString GetURL() const
 	{ return Request.IsValid() ? Request->Parameters.URL : FString(); }
 
+	virtual FTimeValue GetConnectionTimeout() const
+	{ return Request.IsValid() ? Request->Parameters.ConnectTimeout : FTimeValue(); }
+
+	virtual FTimeValue GetNoDataTimeout() const
+	{ return Request.IsValid() ? Request->Parameters.NoDataTimeout : FTimeValue(); }
+
 	virtual IAdaptiveStreamingPlayerResourceRequest::EPlaybackResourceType GetStaticQuery() const
 	{ return StaticQueryType.GetWithDefault(IAdaptiveStreamingPlayerResourceRequest::EPlaybackResourceType::Empty); }
 
@@ -189,7 +199,7 @@ public:
 
 	virtual const HTTP::FConnectionInfo* GetConnectionInfo() const
 	{ return &ConnectionInfo; }
-	
+
 	virtual TSharedPtrTS<IElectraHttpManager::FRequest> GetRequest() const
 	{ return Request; }
 
@@ -275,5 +285,3 @@ private:
 };
 
 } // namespace Electra
-
-
diff --git a/Engine/Plugins/Media/ElectraPlayer/Source/ElectraPlayerRuntime/Private/Runtime/Player/Playlist.h b/Engine/Plugins/Media/ElectraPlayer/Source/ElectraPlayerRuntime/Private/Runtime/Player/Playlist.h
--- a/Engine/Plugins/Media/ElectraPlayer/Source/ElectraPlayerRuntime/Private/Runtime/Player/Playlist.h
+++ b/Engine/Plugins/Media/ElectraPlayer/Source/ElectraPlayerRuntime/Private/Runtime/Player/Playlist.h
@@ -4,29 +4,29 @@
 
 #include "PlayerCore.h"
 
-
-
-
 namespace Electra
 {
-
 	namespace Playlist
 	{
-
 		enum class EListType
 		{
-			Master,
-			Variant
+			Main,
+			Variant,
+			Master = Main
 		};
 
 		static const TCHAR* const GetPlaylistTypeString(EListType InListType)
 		{
-			switch (InListType)
+			switch(InListType)
 			{
-				case EListType::Master:
-					return TEXT("Master");
+				case EListType::Main:
+				{
+					return TEXT("Main");
+				}
 				case EListType::Variant:
+				{
 					return TEXT("Variant");
+				}
 			}
 			return TEXT("n/a");
 		}
@@ -43,15 +43,15 @@ namespace Electra
 			switch (InLoadType)
 			{
 				case ELoadType::Initial:
+				{
 					return TEXT("Initial");
+				}
 				case ELoadType::Update:
+				{
 					return TEXT("Update");
+				}
 			}
 			return TEXT("n/a");
 		}
-
 	}
-
 } // namespace Electra
-
-
diff --git a/Engine/Plugins/Media/ElectraPlayer/Source/ElectraPlayerRuntime/Private/Runtime/StreamTypes.cpp b/Engine/Plugins/Media/ElectraPlayer/Source/ElectraPlayerRuntime/Private/Runtime/StreamTypes.cpp
--- a/Engine/Plugins/Media/ElectraPlayer/Source/ElectraPlayerRuntime/Private/Runtime/StreamTypes.cpp
+++ b/Engine/Plugins/Media/ElectraPlayer/Source/ElectraPlayerRuntime/Private/Runtime/StreamTypes.cpp
@@ -100,10 +100,13 @@ namespace Electra
 				// Profile and level follow?
 				if (CodecOTI.Len() > 5 && CodecOTI[4] == TCHAR('.'))
 				{
+					int32 DotPos;
+					CodecOTI.FindLastChar(TCHAR('.'), DotPos);
+
 					FString Temp;
 					int32 TempValue;
-					// We recognize the expected format avcC.xxyyzz and for legacy reasons also avcC.xx.zz
-					if (CodecOTI.Len() == 11)
+					// We recognize the expected format avcC.xxyyzz and for legacy reasons also avcC.xxx.zz
+					if (CodecOTI.Len() == 11 && DotPos == 4)
 					{
 						Temp = CodecOTI.Mid(5, 2);
 						LexFromStringHex(TempValue, *Temp);
@@ -115,13 +118,13 @@ namespace Electra
 						LexFromStringHex(TempValue, *Temp);
 						SetProfileLevel(TempValue);
 					}
-					else if (CodecOTI.Len() == 10 && CodecOTI[7] == TCHAR('.'))
+					else if (DotPos != INDEX_NONE)
 					{
-						Temp = CodecOTI.Mid(5, 2);
-						LexFromStringHex(TempValue, *Temp);
+						Temp = CodecOTI.Mid(5, DotPos-5);
+						LexFromString(TempValue, *Temp);
 						SetProfile(TempValue);
-						Temp = CodecOTI.Mid(8, 2);
-						LexFromStringHex(TempValue, *Temp);
+						Temp = CodecOTI.Mid(DotPos+1);
+						LexFromString(TempValue, *Temp);
 						SetProfileLevel(TempValue);
 						// Change the string to the expected format.
 						SetCodecSpecifierRFC6381(FString::Printf(TEXT("avc%c.%02x00%02x"), CodecOTI[3], GetProfile(), GetProfileLevel()));
@@ -211,6 +214,15 @@ namespace Electra
 			}
 			return false;
 		}
+		else if (CodecOTI.StartsWith(TEXT("dvh1")) || CodecOTI.StartsWith(TEXT("dvhe")))
+		{
+			// Dolby Vision only recognized as a generic Video 4CC for now.
+			StreamType = EStreamType::Video;
+			CodecSpecifier = CodecOTI;
+			Codec = ECodec::Video4CC;
+			Codec4CC = CodecOTI[3] == TCHAR('1') ? Make4CC('d','v','h','1') : Make4CC('d','v','h','e');
+			return true;
+		}
 		else if (CodecOTI.StartsWith(TEXT("mp4a")))
 		{
 			StreamType = EStreamType::Audio;
@@ -219,20 +231,39 @@ namespace Electra
 			// Object and profile follow?
 			if (CodecOTI.Len() > 6 && CodecOTI[4] == TCHAR('.'))
 			{
-				// mp4a.40.d is recognized.
+				// mp4a.40.d and mp4a.6b are recognized.
 				FString OT, Profile;
 				int32 DotPos = CodecOTI.Find(TEXT("."), ESearchCase::CaseSensitive, ESearchDir::FromStart, 5);
 				OT = CodecOTI.Mid(5, DotPos != INDEX_NONE ? DotPos - 5 : DotPos);
 				Profile = CodecOTI.Mid(DotPos != INDEX_NONE ? DotPos + 1 : DotPos);
-				if (!OT.Equals(TEXT("40")))
+				if (OT.Equals(TEXT("40")))
 				{
-					return false;
+					int32 ProfileValue = 0;
+					LexFromString(ProfileValue, *Profile);
+					SetProfile(ProfileValue);
+					// AAC-LC, AAC-HE (SBR), AAC-HEv2 (PS), MP3
+					if (!(ProfileValue == 2 || ProfileValue == 5 || ProfileValue == 29 || ProfileValue == 34))
+					{
+						return false;
+					}
+					if (ProfileValue == 34)
+					{
+						Codec = ECodec::Audio4CC;
+						Codec4CC = Make4CC('m','p','g','a');
+						MimeType = TEXT("audio/mpeg");
+						SetProfile(1);
+						SetProfileLevel(3);
+					}
 				}
-				int32 ProfileValue = 0;
-				LexFromString(ProfileValue, *Profile);
-				SetProfile(ProfileValue);
-				// AAC-LC, AAC-HE (SBR), AAC-HEv2 (PS)
-				if (!(ProfileValue == 2 || ProfileValue == 5 || ProfileValue == 29))
+				else if (OT.Equals(TEXT("6b"), ESearchCase::IgnoreCase))
+				{
+					Codec = ECodec::Audio4CC;
+					Codec4CC = Make4CC('m','p','g','a');
+					MimeType = TEXT("audio/mpeg");
+					SetProfile(1);
+					SetProfileLevel(3);
+				}
+				else
 				{
 					return false;
 				}
@@ -247,6 +278,14 @@ namespace Electra
 			// Presently not supported.
 			return false;
 		}
+		else if (CodecOTI.StartsWith(TEXT("ac-3")) || CodecOTI.StartsWith(TEXT("ac3")))
+		{
+			StreamType = EStreamType::Audio;
+			CodecSpecifier = CodecOTI;
+			Codec = ECodec::AC3;
+			// Presently not supported.
+			return false;
+		}
 		else if (CodecOTI.Equals(TEXT("opus"), ESearchCase::IgnoreCase))
 		{
 			StreamType = EStreamType::Audio;
@@ -305,7 +344,6 @@ namespace Electra
 			SetCodecSpecifierRFC6381(FString::Printf(TEXT("vp09.%02d.%02d.%02d"), 0, 0, 8));
 			return true;
 		}
-
 		else if (CodecOTI.Equals(TEXT("wvtt")))
 		{
 			StreamType = EStreamType::Subtitle;
@@ -673,6 +711,12 @@ namespace Electra
 
 	int32 FCodecSelectionPriorities::GetClassPriority(const FString& CodecSpecifierRFC6381) const
 	{
+		// If no priorities are given then all have the same priority of 0.
+		if (!ClassPriorities.Num())
+		{
+			return 0;
+		}
+		// Otherwise apply the priority filter. If no match then return -1.
 		for(auto &CodecClass : ClassPriorities)
 		{
 			if (CodecSpecifierRFC6381.StartsWith(CodecClass.Prefix, ESearchCase::IgnoreCase))
diff --git a/Engine/Plugins/Media/ElectraPlayer/Source/ElectraPlayerRuntime/Private/Runtime/StreamTypes.h b/Engine/Plugins/Media/ElectraPlayer/Source/ElectraPlayerRuntime/Private/Runtime/StreamTypes.h
--- a/Engine/Plugins/Media/ElectraPlayer/Source/ElectraPlayerRuntime/Private/Runtime/StreamTypes.h
+++ b/Engine/Plugins/Media/ElectraPlayer/Source/ElectraPlayerRuntime/Private/Runtime/StreamTypes.h
@@ -34,6 +34,18 @@ namespace Electra
 		return TEXT("n/a");
 	}
 
+	static inline constexpr int32 StreamTypeToArrayIndex(EStreamType StreamType)
+	{
+		return StreamType == EStreamType::Video ? 0 :
+			   StreamType == EStreamType::Audio ? 1 :
+			   StreamType == EStreamType::Subtitle ? 2 : 3;
+	}
+
+	static inline constexpr int32 StreamTypeToArrayIndex0(EStreamType StreamType)
+	{
+		check(StreamType != EStreamType::Unsupported);
+		return StreamType == EStreamType::Audio ? 1 : StreamType == EStreamType::Subtitle ? 2 : 0;
+	}
 
 
 	class FStreamCodecInformation
@@ -54,6 +66,7 @@ namespace Electra
 			// --- Audio ---
 			AAC = 100,
 			EAC3,
+			AC3,
 			Audio4CC,
 			// --- Subtitle / Caption ---
 			WebVTT = 200,
@@ -125,6 +138,7 @@ namespace Electra
 			{
 				case ECodec::AAC:
 				case ECodec::EAC3:
+				case ECodec::AC3:
 				case ECodec::Audio4CC:
 					return true;
 				default:
@@ -146,6 +160,26 @@ namespace Electra
 			}
 		}
 
+		bool IsCodec(EStreamType InType) const
+		{
+			switch(InType)
+			{
+				case EStreamType::Video:
+				{
+					return IsVideoCodec();
+				}
+				case EStreamType::Audio:
+				{
+					return IsAudioCodec();
+				}
+				case EStreamType::Subtitle:
+				{
+					return IsSubtitleCodec();
+				}
+			}
+			return false;
+		}
+
 		const FString& GetCodecSpecifierRFC6381() const
 		{
 			return CodecSpecifier;
@@ -708,9 +742,10 @@ namespace Electra
 		FStreamCodecInformation		CodecInformation;					//!< Stream codec information
 		FString						ID;									//!< ID of this stream
 		int32						Bandwidth;							//!< Bandwidth required for this stream in bits per second
+		int32						QualityIndex = 0;
 		bool Equals(const FStreamMetadata& Other) const
 		{
-			return ID == Other.ID && Bandwidth == Other.Bandwidth && CodecInformation.Equals(Other.CodecInformation);
+			return ID == Other.ID && Bandwidth == Other.Bandwidth && QualityIndex == Other.QualityIndex && CodecInformation.Equals(Other.CodecInformation);
 		}
 	};
 
diff --git a/Engine/Plugins/Media/ElectraPlayer/Source/ElectraPlayerRuntime/Private/Runtime/Utilities/UtilsMPEG.cpp b/Engine/Plugins/Media/ElectraPlayer/Source/ElectraPlayerRuntime/Private/Runtime/Utilities/UtilsMPEG.cpp
--- a/Engine/Plugins/Media/ElectraPlayer/Source/ElectraPlayerRuntime/Private/Runtime/Utilities/UtilsMPEG.cpp
+++ b/Engine/Plugins/Media/ElectraPlayer/Source/ElectraPlayerRuntime/Private/Runtime/Utilities/UtilsMPEG.cpp
@@ -337,6 +337,19 @@ namespace Electra
 							}
 							break;
 						}
+						// Private data
+						case Utils::Make4CC('P','R','I','V'):
+						{
+							int32 TextEncoding = 0;
+							const uint8* NextData = InData + 10;
+							FString Owner(GetString(NextData, FrameEnd - NextData, &NextData, &TextEncoding));
+							int32 PrivateSize = FrameEnd - NextData;
+							TArray<uint8> PrivateData(NextData, PrivateSize);
+							FItem& Item = PrivateItems.Emplace_GetRef();
+							Item.MimeType = MoveTemp(Owner);
+							Item.Value = FVariant(MoveTemp(PrivateData));
+							break;
+						}
 						// Recognized text fields
 						case Utils::Make4CC('T','A','L','B'):
 						case Utils::Make4CC('T','C','O','M'):
@@ -401,6 +414,11 @@ namespace Electra
 			return Tags;
 		}
 
+		const TArray<FID3V2Metadata::FItem>& FID3V2Metadata::GetPrivateItems() const
+		{
+			return PrivateItems;
+		}
+
 	} // namespace MPEG
 } // namespace Electra
 
diff --git a/Engine/Plugins/Media/ElectraPlayer/Source/ElectraPlayerRuntime/Private/Runtime/Utilities/UtilsMPEG.h b/Engine/Plugins/Media/ElectraPlayer/Source/ElectraPlayerRuntime/Private/Runtime/Utilities/UtilsMPEG.h
--- a/Engine/Plugins/Media/ElectraPlayer/Source/ElectraPlayerRuntime/Private/Runtime/Utilities/UtilsMPEG.h
+++ b/Engine/Plugins/Media/ElectraPlayer/Source/ElectraPlayerRuntime/Private/Runtime/Utilities/UtilsMPEG.h
@@ -87,7 +87,7 @@ namespace Electra
 			struct FItem
 			{
 				FString Language;				// ISO 639-2; if not set (all zero) the default entry for all languages
-				FString MimeType;
+				FString MimeType;				// Mime type or Owner ID for private item.
 				FVariant Value;
 				int32 ItemType = -1;
 			};
@@ -97,8 +97,10 @@ namespace Electra
 			bool GetTag(FItem& OutValue, uint32 InTag);
 			const TMap<uint32, FItem>& GetTags() const;
 			TMap<uint32, FItem>& GetTags();
+			const TArray<FItem>& GetPrivateItems() const;
 		private:
 			TMap<uint32, FItem> Tags;
+			TArray<FItem> PrivateItems;
 		};
 
 	} // namespace MPEG
diff --git a/Engine/Plugins/Media/ElectraUtil/Source/ElectraBase/Private/PlayerTime.cpp b/Engine/Plugins/Media/ElectraUtil/Source/ElectraBase/Private/PlayerTime.cpp
--- a/Engine/Plugins/Media/ElectraUtil/Source/ElectraBase/Private/PlayerTime.cpp
+++ b/Engine/Plugins/Media/ElectraUtil/Source/ElectraBase/Private/PlayerTime.cpp
@@ -170,6 +170,11 @@ namespace Electra
 			{
 				LexFromString(Numerator, *(InString.Mid(0, DotIndex)));
 				FString frc = InString.Mid(DotIndex + 1);
+				bool bIsNeg = Numerator < 0;
+				if (bIsNeg)
+				{
+					Numerator = -Numerator;
+				}
 				// Remove all trailing zeros
 				int32 last0 = TimeStringHelpers::FindLastNotOf(frc, kTextZero);
 				if (last0 != INDEX_NONE)
@@ -182,6 +187,10 @@ namespace Electra
 						Denominator *= 10;
 					}
 				}
+				if (bIsNeg)
+				{
+					Numerator = -Numerator;
+				}
 				bIsValid = true;
 			}
 		}
diff --git a/Engine/Plugins/Media/ElectraUtil/Source/ElectraBase/Private/StringHelpers.cpp b/Engine/Plugins/Media/ElectraUtil/Source/ElectraBase/Private/StringHelpers.cpp
--- a/Engine/Plugins/Media/ElectraUtil/Source/ElectraBase/Private/StringHelpers.cpp
+++ b/Engine/Plugins/Media/ElectraUtil/Source/ElectraBase/Private/StringHelpers.cpp
@@ -166,13 +166,73 @@ namespace Electra
 		{
 			FUTF8ToTCHAR cnv((const ANSICHAR*)InArray.GetData(), InArray.Num());
 			FString UTF8Text = FString::ConstructFromPtrSize(cnv.Get(), cnv.Length());
-			//return (TCHAR*)FUTF8ToTCHAR((const ANSICHAR*)InArray.GetData(), InArray.Num()).Get();
-			//FString UTF8Text(InArray.Num(), (TCHAR*)FUTF8ToTCHAR((const ANSICHAR*)InArray.GetData(), InArray.Num()).Get());
 			return MoveTemp(UTF8Text);
 		}
 
+		bool ArrayToString(FString& OutString, const TConstArrayView<const uint8>& InArray)
+		{
+			int32 nb = InArray.Num();
+			const uint8* Data = InArray.GetData();
+			// Check for potential BOMs
+			if (nb > 3 && Data[0] == 0xEF && Data[1] == 0xBB && Data[2] == 0xBF)
+			{
+				// UTF-8 BOM
+				nb -= 3;
+				Data += 3;
+			}
+			else if (nb >= 2 && Data[0] == 0xFE && Data[1] == 0xFF)
+			{
+				// UTF-16 BE BOM
+				return false;
+			}
+			else if (nb >= 2 && Data[0] == 0xFF && Data[1] == 0xFE)
+			{
+				// UTF-16 LE BOM
+				auto cnv = StringCast<TCHAR>(reinterpret_cast<const UCS2CHAR*>(Data+2), nb/2-1);
+				FString UTF8Text = FString::ConstructFromPtrSize(cnv.Get(), cnv.Length());
+				OutString = MoveTemp(UTF8Text);
+				return true;
+			}
+			else if (nb >= 4 && Data[0] == 0x00 && Data[1] == 0x00 && Data[2] == 0xFE && Data[3] == 0xFF)
+			{
+				// UTF-32 BE BOM
+				return false;
+			}
+			else if (nb >= 4 && Data[0] == 0xFF && Data[1] == 0xFE && Data[2] == 0x00 && Data[3] == 0x00)
+			{
+				// UTF-32 LE BOM
+				return false;
+			}
+			FUTF8ToTCHAR cnv((const ANSICHAR*)Data, nb);
+			FString UTF8Text = FString::ConstructFromPtrSize(cnv.Get(), cnv.Length());
+			OutString = MoveTemp(UTF8Text);
+			return true;
+		}
 
-	} // namespace StringHelpers
-} // namespace Electra
 
+		FString GetLongestCommonPrefix(TArray<FString>& InOutTempArrayOfInputs)
+		{
+			if (InOutTempArrayOfInputs.IsEmpty())
+			{
+				return FString();
+			}
+			else if (InOutTempArrayOfInputs.Num() == 1)
+			{
+				return InOutTempArrayOfInputs[0];
+			}
+			InOutTempArrayOfInputs.Sort();
+			const int32 l1 =InOutTempArrayOfInputs[0].Len();
+			const int32 l2 =InOutTempArrayOfInputs.Last().Len();
+			const int32 CommonSize = l1 < l2 ? l1 : l2;
+			const TCHAR* s1 = *InOutTempArrayOfInputs[0];
+			const TCHAR* s2 = *InOutTempArrayOfInputs.Last();
+			int32 NumSame = 0;
+			while(NumSame < CommonSize && s1[NumSame] == s2[NumSame])
+			{
+				++NumSame;
+			}
+			return FString::ConstructFromPtrSize(s1, NumSame);
+		}
 
+	} // namespace StringHelpers
+} // namespace Electra
diff --git a/Engine/Plugins/Media/ElectraUtil/Source/ElectraBase/Public/PlayerTime.h b/Engine/Plugins/Media/ElectraUtil/Source/ElectraBase/Public/PlayerTime.h
--- a/Engine/Plugins/Media/ElectraUtil/Source/ElectraBase/Public/PlayerTime.h
+++ b/Engine/Plugins/Media/ElectraUtil/Source/ElectraBase/Public/PlayerTime.h
@@ -523,6 +523,23 @@ public:
 		return Result;
 	}
 
+	inline FTimeValue operator - () const
+	{
+		FTimeValue Result(*this);
+		Result.HNS = -Result.HNS;
+		return Result;
+	}
+
+	inline FTimeValue Abs() const
+	{
+		FTimeValue Result(*this);
+		if (Result.HNS < 0)
+		{
+			Result.HNS = -Result.HNS;
+		}
+		return Result;
+	}
+
 private:
 	static inline bool WillOverflow(int64& OutTemp, const int64 InA, const int64 InB)
 	{
diff --git a/Engine/Plugins/Media/ElectraUtil/Source/ElectraBase/Public/Utilities/StringHelpers.h b/Engine/Plugins/Media/ElectraUtil/Source/ElectraBase/Public/Utilities/StringHelpers.h
--- a/Engine/Plugins/Media/ElectraUtil/Source/ElectraBase/Public/Utilities/StringHelpers.h
+++ b/Engine/Plugins/Media/ElectraUtil/Source/ElectraBase/Public/Utilities/StringHelpers.h
@@ -32,6 +32,10 @@ namespace Electra
 
 		FString ELECTRABASE_API ArrayToString(const TArray<uint8>& InArray);
 
+		bool ELECTRABASE_API ArrayToString(FString& OutString, const TConstArrayView<const uint8>& InArray);
+
+		FString ELECTRABASE_API GetLongestCommonPrefix(TArray<FString>& InOutTempArrayOfInputs);
+
 		/**
 		 * There is a known anomaly in the FString::TConstIterator. It iterates all TCHARs in the string *including* the terminating zero character.
 		 * This is not the behaviour we want and setup some helper iterator here which is not including the terminating zero.
