ElectraPlayer: Moved live latency calculation out into the respective media format manifests

#rnx

diff --git a/Engine/Plugins/Media/ElectraPlayer/Source/ElectraPlayerRuntime/Private/Runtime/Player/AdaptiveStreamingPlayer.cpp b/Engine/Plugins/Media/ElectraPlayer/Source/ElectraPlayerRuntime/Private/Runtime/Player/AdaptiveStreamingPlayer.cpp
--- a/Engine/Plugins/Media/ElectraPlayer/Source/ElectraPlayerRuntime/Private/Runtime/Player/AdaptiveStreamingPlayer.cpp
+++ b/Engine/Plugins/Media/ElectraPlayer/Source/ElectraPlayerRuntime/Private/Runtime/Player/AdaptiveStreamingPlayer.cpp
@@ -2078,7 +2078,7 @@ void FAdaptiveStreamingPlayer::HandleSessionMessage(TSharedPtrTS<IPlayerMessage>
 			{
 				if (!Result.WasAborted())
 				{
-					if (pMsg->GetListType() == Playlist::EListType::Master && pMsg->GetLoadType() == Playlist::ELoadType::Initial)
+					if (pMsg->GetListType() == Playlist::EListType::Main && pMsg->GetLoadType() == Playlist::ELoadType::Initial)
 					{
 						DispatchEvent(FMetricEvent::ReportReceivedMasterPlaylist(pMsg->GetConnectionInfo().EffectiveURL));
 					}
@@ -2112,8 +2112,12 @@ void FAdaptiveStreamingPlayer::HandleSessionMessage(TSharedPtrTS<IPlayerMessage>
 		Metrics::FLicenseKeyStats stats;
 		stats.bWasSuccessful = !pMsg->GetResult().IsError();
 		stats.URL   		 = pMsg->GetConnectionInfo().EffectiveURL;
-		stats.FailureReason  = pMsg->GetResult().IsError() ? pMsg->GetResult().GetPrintable() : FString(); //pMsg->GetConnectionInfo().StatusInfo.ErrorDetail.GetMessage();
+		stats.FailureReason  = pMsg->GetResult().IsError() ? pMsg->GetResult().GetPrintable() : FString();
 		DispatchEvent(FMetricEvent::ReportLicenseKey(stats));
+		if (pMsg->GetResult().IsError())
+		{
+			PostError(pMsg->GetResult());
+		}
 	}
 	// Decoder message?
 	else if (SessionMessage->GetType() == FDecoderMessage::Type())
@@ -2138,30 +2142,7 @@ FTimeValue FAdaptiveStreamingPlayer::CalculateCurrentLiveLatency(bool bViaLatenc
 	FTimeValue LiveLatency;
 	if (Manifest.IsValid())
 	{
-		if (Manifest->GetPresentationType() != IManifest::EType::OnDemand)
-		{
-			FTimeValue UTCNow = GetSynchronizedUTCTime()->GetTime();
-			FTimeValue PlayPosNow = PlaybackState.GetPlayPosition();
-			LiveLatency = UTCNow - PlayPosNow;
-
-			if (bViaLatencyElement)
-			{
-				TSharedPtrTS<const FLowLatencyDescriptor> llDesc = Manifest->GetLowLatencyDescriptor();
-				if (llDesc.IsValid())
-				{
-					// Low latency Live
-					TSharedPtrTS<IProducerReferenceTimeInfo> ProdRefTime = Manifest->GetProducerReferenceTimeInfo(llDesc->Latency.ReferenceID);
-					if (ProdRefTime.IsValid())
-					{
-						FTimeValue EncoderLatency = PlaybackState.GetEncoderLatency();
-						if (EncoderLatency.IsValid())
-						{
-							LiveLatency += EncoderLatency;
-						}
-					}
-				}
-			}
-		}
+		LiveLatency = Manifest->CalculateCurrentLiveLatency(PlaybackState.GetPlayPosition(), PlaybackState.GetEncoderLatency(), bViaLatencyElement);
 	}
 	return LiveLatency;
 }
diff --git a/Engine/Plugins/Media/ElectraPlayer/Source/ElectraPlayerRuntime/Private/Runtime/Player/AdaptiveStreamingPlayerManifest.cpp b/Engine/Plugins/Media/ElectraPlayer/Source/ElectraPlayerRuntime/Private/Runtime/Player/AdaptiveStreamingPlayerManifest.cpp
--- a/Engine/Plugins/Media/ElectraPlayer/Source/ElectraPlayerRuntime/Private/Runtime/Player/AdaptiveStreamingPlayerManifest.cpp
+++ b/Engine/Plugins/Media/ElectraPlayer/Source/ElectraPlayerRuntime/Private/Runtime/Player/AdaptiveStreamingPlayerManifest.cpp
@@ -41,6 +41,7 @@ FString FixLocalFileSchemeURL(const FString& InURL)
 	if (URL.StartsWith("file:", ESearchCase::IgnoreCase))
 	{
 		URL.ReplaceCharInline(TCHAR('\\'), TCHAR('/'));
+		URL.ReplaceInline(TEXT(":///"), TEXT("://"));
 		URL.ReplaceInline(TEXT(" "), TEXT("%20"));
 	}
 	return URL;
diff --git a/Engine/Plugins/Media/ElectraPlayer/Source/ElectraPlayerRuntime/Private/Runtime/Player/DASH/ManifestDASH.cpp b/Engine/Plugins/Media/ElectraPlayer/Source/ElectraPlayerRuntime/Private/Runtime/Player/DASH/ManifestDASH.cpp
--- a/Engine/Plugins/Media/ElectraPlayer/Source/ElectraPlayerRuntime/Private/Runtime/Player/DASH/ManifestDASH.cpp
+++ b/Engine/Plugins/Media/ElectraPlayer/Source/ElectraPlayerRuntime/Private/Runtime/Player/DASH/ManifestDASH.cpp
@@ -351,6 +351,34 @@ TSharedPtrTS<const FLowLatencyDescriptor> FManifestDASH::GetLowLatencyDescriptor
 	return Manifest.IsValid() ? Manifest->GetLowLatencyDescriptor() : nullptr;
 }
 
+FTimeValue FManifestDASH::CalculateCurrentLiveLatency(const FTimeValue& InCurrentPlaybackPosition, const FTimeValue& InEncoderLatency, bool bViaLatencyElement) const
+{
+	FTimeValue LiveLatency;
+	if (GetPresentationType() != IManifest::EType::OnDemand)
+	{
+		FTimeValue UTCNow = PlayerSessionServices->GetSynchronizedUTCTime()->GetTime();
+		LiveLatency = UTCNow - InCurrentPlaybackPosition;
+
+		if (bViaLatencyElement)
+		{
+			TSharedPtrTS<const FLowLatencyDescriptor> llDesc = GetLowLatencyDescriptor();
+			if (llDesc.IsValid())
+			{
+				// Low latency Live
+				TSharedPtrTS<IProducerReferenceTimeInfo> ProdRefTime = GetProducerReferenceTimeInfo(llDesc->Latency.ReferenceID);
+				if (ProdRefTime.IsValid())
+				{
+					if (InEncoderLatency.IsValid())
+					{
+						LiveLatency += InEncoderLatency;
+					}
+				}
+			}
+		}
+	}
+	return LiveLatency;
+}
+
 FTimeValue FManifestDASH::GetAnchorTime() const
 {
 	TSharedPtrTS<FManifestDASHInternal> Manifest(CurrentManifest);
diff --git a/Engine/Plugins/Media/ElectraPlayer/Source/ElectraPlayerRuntime/Private/Runtime/Player/DASH/ManifestDASH.h b/Engine/Plugins/Media/ElectraPlayer/Source/ElectraPlayerRuntime/Private/Runtime/Player/DASH/ManifestDASH.h
--- a/Engine/Plugins/Media/ElectraPlayer/Source/ElectraPlayerRuntime/Private/Runtime/Player/DASH/ManifestDASH.h
+++ b/Engine/Plugins/Media/ElectraPlayer/Source/ElectraPlayerRuntime/Private/Runtime/Player/DASH/ManifestDASH.h
@@ -20,6 +20,7 @@ public:
 	virtual ~FManifestDASH();
 	EType GetPresentationType() const override;
 	TSharedPtrTS<const FLowLatencyDescriptor> GetLowLatencyDescriptor() const override;
+	FTimeValue CalculateCurrentLiveLatency(const FTimeValue& InCurrentPlaybackPosition, const FTimeValue& InEncoderLatency, bool bViaLatencyElement) const override;
 	FTimeValue GetAnchorTime() const override;
 	FTimeRange GetTotalTimeRange() const override;
 	FTimeRange GetSeekableTimeRange() const override;
diff --git a/Engine/Plugins/Media/ElectraPlayer/Source/ElectraPlayerRuntime/Private/Runtime/Player/DASH/StreamReaderDASH.cpp b/Engine/Plugins/Media/ElectraPlayer/Source/ElectraPlayerRuntime/Private/Runtime/Player/DASH/StreamReaderDASH.cpp
--- a/Engine/Plugins/Media/ElectraPlayer/Source/ElectraPlayerRuntime/Private/Runtime/Player/DASH/StreamReaderDASH.cpp
+++ b/Engine/Plugins/Media/ElectraPlayer/Source/ElectraPlayerRuntime/Private/Runtime/Player/DASH/StreamReaderDASH.cpp
@@ -1449,8 +1449,6 @@ void FStreamReaderDASH::FStreamHandler::HandleRequestMP4()
 									AccessUnit->LatestPTS += TimeOffset;
 									AccessUnit->LatestPTS.SetSequenceIndex(Request->TimestampSequenceIndex);
 
-									AccessUnit->OffsetFromSegmentStart = TimelineOffset;
-
 									AccessUnit->ProducerReferenceTime = ProducerTime.Base + FTimeValue(AUDTS - ProducerTime.Media, TrackTimescale);
 
 									ElectraCDM::FMediaCDMSampleInfo SampleEncryptionInfo;
diff --git a/Engine/Plugins/Media/ElectraPlayer/Source/ElectraPlayerRuntime/Private/Runtime/Player/HLS/ManifestHLS.cpp b/Engine/Plugins/Media/ElectraPlayer/Source/ElectraPlayerRuntime/Private/Runtime/Player/HLS/ManifestHLS.cpp
--- a/Engine/Plugins/Media/ElectraPlayer/Source/ElectraPlayerRuntime/Private/Runtime/Player/HLS/ManifestHLS.cpp
+++ b/Engine/Plugins/Media/ElectraPlayer/Source/ElectraPlayerRuntime/Private/Runtime/Player/HLS/ManifestHLS.cpp
@@ -146,6 +146,17 @@ TSharedPtrTS<const FLowLatencyDescriptor> FManifestHLS::GetLowLatencyDescriptor(
 	return nullptr;
 }
 
+FTimeValue FManifestHLS::CalculateCurrentLiveLatency(const FTimeValue& InCurrentPlaybackPosition, const FTimeValue& InEncoderLatency, bool bViaLatencyElement) const
+{
+	FTimeValue LiveLatency;
+	if (GetPresentationType() != IManifest::EType::OnDemand)
+	{
+		FTimeValue UTCNow = SessionServices->GetSynchronizedUTCTime()->GetTime();
+		LiveLatency = UTCNow - InCurrentPlaybackPosition;
+	}
+	return LiveLatency;
+}
+
 FTimeValue FManifestHLS::GetAnchorTime() const
 {
 	return FTimeValue::GetZero();
diff --git a/Engine/Plugins/Media/ElectraPlayer/Source/ElectraPlayerRuntime/Private/Runtime/Player/HLS/ManifestHLS.h b/Engine/Plugins/Media/ElectraPlayer/Source/ElectraPlayerRuntime/Private/Runtime/Player/HLS/ManifestHLS.h
--- a/Engine/Plugins/Media/ElectraPlayer/Source/ElectraPlayerRuntime/Private/Runtime/Player/HLS/ManifestHLS.h
+++ b/Engine/Plugins/Media/ElectraPlayer/Source/ElectraPlayerRuntime/Private/Runtime/Player/HLS/ManifestHLS.h
@@ -18,6 +18,7 @@ public:
 	virtual ~FManifestHLS();
 	EType GetPresentationType() const override;
 	TSharedPtrTS<const FLowLatencyDescriptor> GetLowLatencyDescriptor() const override;
+	FTimeValue CalculateCurrentLiveLatency(const FTimeValue& InCurrentPlaybackPosition, const FTimeValue& InEncoderLatency, bool bViaLatencyElement) const override;
 	FTimeValue GetAnchorTime() const override;
 	FTimeRange GetTotalTimeRange() const override;
 	FTimeRange GetSeekableTimeRange() const override;
diff --git a/Engine/Plugins/Media/ElectraPlayer/Source/ElectraPlayerRuntime/Private/Runtime/Player/Manifest.h b/Engine/Plugins/Media/ElectraPlayer/Source/ElectraPlayerRuntime/Private/Runtime/Player/Manifest.h
--- a/Engine/Plugins/Media/ElectraPlayer/Source/ElectraPlayerRuntime/Private/Runtime/Player/Manifest.h
+++ b/Engine/Plugins/Media/ElectraPlayer/Source/ElectraPlayerRuntime/Private/Runtime/Player/Manifest.h
@@ -188,6 +188,15 @@ namespace Electra
 				RetryAgainAtTime = MEDIAutcTime::Current() + FTimeValue().SetFromMilliseconds(Milliseconds);
 				return *this;
 			}
+			FResult& RetryAfter(const FTimeValue& InAfter)
+			{
+				if (InAfter.IsValid())
+				{
+					Type = EType::TryAgainLater;
+					RetryAgainAtTime = MEDIAutcTime::Current() + InAfter;
+				}
+				return *this;
+			}
 			FResult& SetErrorDetail(const FErrorDetail& InErrorDetail)
 			{
 				ErrorDetail = InErrorDetail;
@@ -248,6 +257,9 @@ namespace Electra
 		//! Returns the low-latency descriptor, if any. May return nullptr if there is none.
 		virtual TSharedPtrTS<const FLowLatencyDescriptor> GetLowLatencyDescriptor() const = 0;
 
+		//! Calculates the live latency given the current playback position and optional encoder latency.
+		virtual FTimeValue CalculateCurrentLiveLatency(const FTimeValue& InCurrentPlaybackPosition, const FTimeValue& InEncoderLatency, bool bViaLatencyElement) const = 0;
+
 
 		//-------------------------------------------------------------------------
 		// Presentation and play time related functions.
diff --git a/Engine/Plugins/Media/ElectraPlayer/Source/ElectraPlayerRuntime/Private/Runtime/Player/mkv/ManifestMKV.h b/Engine/Plugins/Media/ElectraPlayer/Source/ElectraPlayerRuntime/Private/Runtime/Player/mkv/ManifestMKV.h
--- a/Engine/Plugins/Media/ElectraPlayer/Source/ElectraPlayerRuntime/Private/Runtime/Player/mkv/ManifestMKV.h
+++ b/Engine/Plugins/Media/ElectraPlayer/Source/ElectraPlayerRuntime/Private/Runtime/Player/mkv/ManifestMKV.h
@@ -32,6 +32,8 @@ public:
 	EType GetPresentationType() const override;
 	TSharedPtrTS<const FLowLatencyDescriptor> GetLowLatencyDescriptor() const override
 	{ return nullptr; }
+	FTimeValue CalculateCurrentLiveLatency(const FTimeValue& InCurrentPlaybackPosition, const FTimeValue& InEncoderLatency, bool bViaLatencyElement) const override
+	{ return FTimeValue(); }
 	FTimeValue GetAnchorTime() const override
 	{ return FTimeValue::GetZero(); }
 	FTimeRange GetTotalTimeRange() const override
diff --git a/Engine/Plugins/Media/ElectraPlayer/Source/ElectraPlayerRuntime/Private/Runtime/Player/mp4/ManifestMP4.h b/Engine/Plugins/Media/ElectraPlayer/Source/ElectraPlayerRuntime/Private/Runtime/Player/mp4/ManifestMP4.h
--- a/Engine/Plugins/Media/ElectraPlayer/Source/ElectraPlayerRuntime/Private/Runtime/Player/mp4/ManifestMP4.h
+++ b/Engine/Plugins/Media/ElectraPlayer/Source/ElectraPlayerRuntime/Private/Runtime/Player/mp4/ManifestMP4.h
@@ -33,6 +33,8 @@ public:
 	EType GetPresentationType() const override;
 	TSharedPtrTS<const FLowLatencyDescriptor> GetLowLatencyDescriptor() const override
 	{ return nullptr; }
+	FTimeValue CalculateCurrentLiveLatency(const FTimeValue& InCurrentPlaybackPosition, const FTimeValue& InEncoderLatency, bool bViaLatencyElement) const override
+	{ return FTimeValue(); }
 	FTimeValue GetAnchorTime() const override
 	{ return FTimeValue::GetZero(); }
 	FTimeRange GetTotalTimeRange() const override
diff --git a/Engine/Plugins/Media/ElectraPlayer/Source/ElectraPlayerRuntime/Private/Runtime/Player/mpegaudio/ManifestMPEGAudio.cpp b/Engine/Plugins/Media/ElectraPlayer/Source/ElectraPlayerRuntime/Private/Runtime/Player/mpegaudio/ManifestMPEGAudio.cpp
--- a/Engine/Plugins/Media/ElectraPlayer/Source/ElectraPlayerRuntime/Private/Runtime/Player/mpegaudio/ManifestMPEGAudio.cpp
+++ b/Engine/Plugins/Media/ElectraPlayer/Source/ElectraPlayerRuntime/Private/Runtime/Player/mpegaudio/ManifestMPEGAudio.cpp
@@ -229,6 +229,34 @@ TSharedPtrTS<const FLowLatencyDescriptor> FManifestMPEGAudioInternal::GetLowLate
 	return LatencyDescriptor;
 }
 
+FTimeValue FManifestMPEGAudioInternal::CalculateCurrentLiveLatency(const FTimeValue& InCurrentPlaybackPosition, const FTimeValue& InEncoderLatency, bool bViaLatencyElement) const
+{
+	FTimeValue LiveLatency;
+	if (GetPresentationType() != IManifest::EType::OnDemand)
+	{
+		FTimeValue UTCNow = PlayerSessionServices->GetSynchronizedUTCTime()->GetTime();
+		LiveLatency = UTCNow - InCurrentPlaybackPosition;
+
+		if (bViaLatencyElement)
+		{
+			TSharedPtrTS<const FLowLatencyDescriptor> llDesc = GetLowLatencyDescriptor();
+			if (llDesc.IsValid())
+			{
+				// Low latency Live
+				TSharedPtrTS<IProducerReferenceTimeInfo> ProdRefTime = GetProducerReferenceTimeInfo(llDesc->Latency.ReferenceID);
+				if (ProdRefTime.IsValid())
+				{
+					if (InEncoderLatency.IsValid())
+					{
+						LiveLatency += InEncoderLatency;
+					}
+				}
+			}
+		}
+	}
+	return LiveLatency;
+}
+
 
 TSharedPtrTS<IProducerReferenceTimeInfo> FManifestMPEGAudioInternal::GetProducerReferenceTimeInfo(int64 ID) const
 {
diff --git a/Engine/Plugins/Media/ElectraPlayer/Source/ElectraPlayerRuntime/Private/Runtime/Player/mpegaudio/ManifestMPEGAudio.h b/Engine/Plugins/Media/ElectraPlayer/Source/ElectraPlayerRuntime/Private/Runtime/Player/mpegaudio/ManifestMPEGAudio.h
--- a/Engine/Plugins/Media/ElectraPlayer/Source/ElectraPlayerRuntime/Private/Runtime/Player/mpegaudio/ManifestMPEGAudio.h
+++ b/Engine/Plugins/Media/ElectraPlayer/Source/ElectraPlayerRuntime/Private/Runtime/Player/mpegaudio/ManifestMPEGAudio.h
@@ -73,6 +73,7 @@ public:
 
 	EType GetPresentationType() const override;
 	TSharedPtrTS<const FLowLatencyDescriptor> GetLowLatencyDescriptor() const override;
+	FTimeValue CalculateCurrentLiveLatency(const FTimeValue& InCurrentPlaybackPosition, const FTimeValue& InEncoderLatency, bool bViaLatencyElement) const override;
 	FTimeValue GetAnchorTime() const override
 	{ return FTimeValue::GetZero(); }
 	FTimeRange GetTotalTimeRange() const override
diff --git a/Engine/Plugins/Media/ElectraPlayer/Source/ElectraPlayerRuntime/Private/Runtime/StreamAccessUnitBuffer.h b/Engine/Plugins/Media/ElectraPlayer/Source/ElectraPlayerRuntime/Private/Runtime/StreamAccessUnitBuffer.h
--- a/Engine/Plugins/Media/ElectraPlayer/Source/ElectraPlayerRuntime/Private/Runtime/StreamAccessUnitBuffer.h
+++ b/Engine/Plugins/Media/ElectraPlayer/Source/ElectraPlayerRuntime/Private/Runtime/StreamAccessUnitBuffer.h
@@ -76,7 +76,6 @@ namespace Electra
 		FTimeValue					PTO;							//!< Media local presentation time offset.
 		FTimeValue					EarliestPTS;					//!< Earliest PTS at which to present samples. If this is larger than PTS the sample is not to be presented.
 		FTimeValue					LatestPTS;						//!< Latest PTS at which to present samples. If this is less than PTS the sample is not to be presented.
-		FTimeValue					OffsetFromSegmentStart;			//!< If set, the difference between the first segment AU's PTS and the expected time according to the playlist.
 		FTimeValue					ProducerReferenceTime;			//!< If set, the wallclock time of the producer when this AU was encoded or captured
 		int64						SequenceIndex;
 		uint32						AUSize;							//!< Size of this access unit
@@ -165,7 +164,6 @@ namespace Electra
 			PTO.SetToZero();
 			EarliestPTS.SetToInvalid();
 			LatestPTS.SetToInvalid();
-			OffsetFromSegmentStart.SetToInvalid();
 			SequenceIndex = 0;
 			ESType = EStreamType::Unsupported;
 			AUSize = 0;
@@ -521,56 +519,6 @@ namespace Electra
 			}
 		}
 
-		// Tags all AUs such that the decoders will skip over them until the desired decode time.
-		FTimeValue PrepareForDecodeStartingAt(FTimeValue DecodeStartTime)
-		{
-			FTimeValue TaggedDuration(FTimeValue::GetZero());
-			FScopeLock Lock(&AccessLock);
-			int32 RemoveUpTo = -1;
-			for(int32 i=0; i<AccessUnits.Num(); ++i)
-			{
-				FAccessUnit* AU = AccessUnits[i];
-
-				// Check if the first AU has a negative offset compared to the expected segment start time.
-				// If so, we have to adjust the given decode start time by that offset as to not discard the data
-				// that has an internal timestamp deviation from the timeline of the playlist.
-				// Positive offsets are currently assumed to not be an issue, unless they were unreasonably large,
-				// but what constitutes "reasonably" is not clear. We're leaving things at that.
-				if (i == 0 && AU->OffsetFromSegmentStart.IsValid() && AU->OffsetFromSegmentStart < FTimeValue::GetZero())
-				{
-					DecodeStartTime += AU->OffsetFromSegmentStart;
-				}
-
-				if (AU->PTS < DecodeStartTime)
-				{
-					// If the AU was not set with a drop flag yet we do it now.
-					// Drop-flagged AUs do not count towards playable duration, so we need to adjust this now.
-					if (AU->DropState == 0)
-					{
-						PlayableDuration -= AU->Duration;
-						TaggedDuration += AU->Duration;
-					}
-					AU->DropState |= FAccessUnit::EDropState::TooEarly;
-
-					AU->EarliestPTS = DecodeStartTime;
-
-					// If this is a sync sample we can remove all preceding AUs until this one.
-					if (AU->bIsSyncSample)
-					{
-						RemoveUpTo = i;
-					}
-				}
-			}
-			for(int32 i=0; i<RemoveUpTo; ++i)
-			{
-				FAccessUnit* AU = nullptr;
-				Pop(AU);
-				FAccessUnit::Release(AU);
-			}
-			return TaggedDuration;
-		}
-
-
 		//! Waits for data to arrive. Returns true if data is present. False if not and timeout expired.
 		bool WaitForData(int64 waitForMicroseconds = -1)
 		{
diff --git a/Engine/Plugins/Media/ElectraPlayer/Source/ElectraPlayerRuntime/Private/Runtime/StreamDataBuffer.h b/Engine/Plugins/Media/ElectraPlayer/Source/ElectraPlayerRuntime/Private/Runtime/StreamDataBuffer.h
--- a/Engine/Plugins/Media/ElectraPlayer/Source/ElectraPlayerRuntime/Private/Runtime/StreamDataBuffer.h
+++ b/Engine/Plugins/Media/ElectraPlayer/Source/ElectraPlayerRuntime/Private/Runtime/StreamDataBuffer.h
@@ -239,6 +239,12 @@ namespace Electra
 			return WritePos;
 		}
 
+		void SetLinearReadSize(int64 InNewSize)
+		{
+			FScopeLock Lock(&AccessLock);
+			WritePos = InNewSize;
+		}
+
 		// Must control Lock()/Unlock() externally!
 		const uint8* GetLinearReadData() const
 		{
