Electra: Added additional codec format for mp3; sleeping video decoder during dummy decoding

#rnx

diff --git a/Engine/Plugins/Media/ElectraCodecs/Source/ElectraDecoders/Private/ElectraMediaMP3Decoder.cpp b/Engine/Plugins/Media/ElectraCodecs/Source/ElectraDecoders/Private/ElectraMediaMP3Decoder.cpp
--- a/Engine/Plugins/Media/ElectraCodecs/Source/ElectraDecoders/Private/ElectraMediaMP3Decoder.cpp
+++ b/Engine/Plugins/Media/ElectraCodecs/Source/ElectraDecoders/Private/ElectraMediaMP3Decoder.cpp
@@ -253,7 +253,7 @@ public:
 	static TArray<FString> Permitted4CCs;
 };
 TSharedPtr<FElectraCommonAudioFlacDecoderFactory, ESPMode::ThreadSafe> FElectraCommonAudioFlacDecoderFactory::Self;
-TArray<FString> FElectraCommonAudioFlacDecoderFactory::Permitted4CCs = { TEXT("mp4a.6b"), TEXT(".mp3") };
+TArray<FString> FElectraCommonAudioFlacDecoderFactory::Permitted4CCs = { TEXT("mp4a.6b"), TEXT("mp4a.40.34"), TEXT(".mp3") };
 
 /*********************************************************************************************************************/
 /*********************************************************************************************************************/
diff --git a/Engine/Plugins/Media/ElectraPlayer/Source/ElectraPlayerRuntime/Private/Runtime/Decoder/VideoDecoder.cpp b/Engine/Plugins/Media/ElectraPlayer/Source/ElectraPlayerRuntime/Private/Runtime/Decoder/VideoDecoder.cpp
--- a/Engine/Plugins/Media/ElectraPlayer/Source/ElectraPlayerRuntime/Private/Runtime/Decoder/VideoDecoder.cpp
+++ b/Engine/Plugins/Media/ElectraPlayer/Source/ElectraPlayerRuntime/Private/Runtime/Decoder/VideoDecoder.cpp
@@ -1247,6 +1247,19 @@ bool FVideoDecoderImpl::HandleDummyDecoding()
 		DummyBufferSampleProperties.Set(RenderOptionKeys::DummyBufferFlag, FVariantValue(true));
 		Renderer->ReturnBuffer(CurrentOutputBuffer, true, DummyBufferSampleProperties);
 		CurrentOutputBuffer = nullptr;
+		// We must not drain the source buffer too quickly. While our counterpart code in the audio decoder actually
+		// produces a usable sample containing silence, we cannot create a usable dummy frame because we have to
+		// keep the last good frame on screen. Our sample we have just returned will not actually be sent into the
+		// media sample queue and thus any next call to `Renderer->CanReceiveOutputFrames(1)` above will always
+		// return `true` because the sample queue will not be full, and as a result we race and take new source
+		// samples from the buffer so quickly, that the buffer will underrun.
+		// To prevent this we put ourselves to sleep for a while. Not the entire sample duration though, but for
+		// enough time to hopefully not cause an underrun.
+		// NOTE: Technically speaking this is not a good solution because we should not really sleep here as
+		//       that is only acceptable at 1x play rate. If playing faster we would need to sleep for a shorter
+		//       duration here or not at all. Filler data on missing media segments should not really happen
+		//       though, so I'm hoping we're getting by.
+		MinLoopSleepTimeMsec = CurrentAccessUnit->AdjustedDuration.GetAsMilliseconds() - 1;
 		return true;
 	}
 	return true;
