ElectraPlayer: Fixed PVS warning

#rnx

diff --git a/Engine/Plugins/Media/ElectraPlayer/Source/ElectraPlayerRuntime/Private/Runtime/Player/StreamSegmentReaderCommon.cpp b/Engine/Plugins/Media/ElectraPlayer/Source/ElectraPlayerRuntime/Private/Runtime/Player/StreamSegmentReaderCommon.cpp
--- a/Engine/Plugins/Media/ElectraPlayer/Source/ElectraPlayerRuntime/Private/Runtime/Player/StreamSegmentReaderCommon.cpp
+++ b/Engine/Plugins/Media/ElectraPlayer/Source/ElectraPlayerRuntime/Private/Runtime/Player/StreamSegmentReaderCommon.cpp
@@ -2033,6 +2033,7 @@ FStreamSegmentReaderCommon::FStreamHandler::EHandleResult FStreamSegmentReaderCo
 							FAccessUnit::Release(AccessUnit);
 							AccessUnit = nullptr;
 						}
+						/*
 						else if ((AccessUnit->DropState & FAccessUnit::EDropState::TooLate) == FAccessUnit::EDropState::TooLate)
 						{
 							// Tag the last one and send it off, but stop doing so for the remainder of the segment.
@@ -2040,6 +2041,7 @@ FStreamSegmentReaderCommon::FStreamHandler::EHandleResult FStreamSegmentReaderCo
 							AccessUnit->bIsLastInPeriod = true;
 							CurrentlyActiveTrackData->bReadPastLastPTS = true;
 						}
+						*/
 					}
 
 					if (AccessUnit)
