ElectraPlayer: Updated ABR to ignore nonexisting segments

#rnx

diff --git a/Engine/Plugins/Media/ElectraPlayer/Source/ElectraPlayerRuntime/Private/Runtime/Player/ABRRules/ABRLiveStream.cpp b/Engine/Plugins/Media/ElectraPlayer/Source/ElectraPlayerRuntime/Private/Runtime/Player/ABRRules/ABRLiveStream.cpp
--- a/Engine/Plugins/Media/ElectraPlayer/Source/ElectraPlayerRuntime/Private/Runtime/Player/ABRRules/ABRLiveStream.cpp
+++ b/Engine/Plugins/Media/ElectraPlayer/Source/ElectraPlayerRuntime/Private/Runtime/Player/ABRRules/ABRLiveStream.cpp
@@ -325,6 +325,7 @@ private:
 	bool bIsRebuffering = false;
 	bool bIsSeeking = false;
 	bool bIsPaused = true;
+	bool bFalloffDetected = false;
 
 	// Work variables used throughout the PrepareStreamCandidateList() -> EvaluateForError() -> EvaluateForQuality() -> PerformSelection() chain.
 	bool bRetryIfPossible = false;
@@ -409,6 +410,7 @@ FABRLiveStream::FABRLiveStream(IABRInfoInterface* InInfo, EMediaFormatType InFor
 	bIsBuffering = false;
 	bIsSeeking = false;
 	bIsPaused = Info->ABRGetPlaySpeed() == FTimeValue::GetZero();
+	bFalloffDetected = false;
 }
 
 FABRLiveStream::~FABRLiveStream()
@@ -646,6 +648,16 @@ FABRDownloadProgressDecision FABRLiveStream::ReportDownloadProgress(const Metric
 
 void FABRLiveStream::ReportDownloadEnd(const Metrics::FSegmentDownloadStats& SegmentDownloadStats)
 {
+	if (SegmentDownloadStats.bWasSkipped)
+	{
+		return;
+	}
+	else if (SegmentDownloadStats.bWasFalloffSegment)
+	{
+		bFalloffDetected = true;
+		return;
+	}
+
 	FStreamWorkVars* WorkVars = GetWorkVars(SegmentDownloadStats.StreamType);
 
 	if (WorkVars && SegmentDownloadStats.SegmentType == Metrics::ESegmentType::Media)
@@ -796,6 +808,7 @@ void FABRLiveStream::ReportBufferingEnd(Metrics::EBufferingReason BufferingReaso
 	bIsBuffering = false;
 	bIsRebuffering = false;
 	bIsSeeking = false;
+	bFalloffDetected = false;
 	LatencyExceededVars.Reset();
 }
 
@@ -814,6 +827,7 @@ void FABRLiveStream::ReportPlaybackEnded()
 	bIsBuffering = false;
 	bIsSeeking = false;
 	bIsPaused = true;
+	bFalloffDetected = false;
 }
 
 void FABRLiveStream::RepresentationsChanged(EStreamType InStreamType, TSharedPtrTS<IManifest::IPlayPeriod> InCurrentPlayPeriod)
@@ -1669,7 +1683,9 @@ IAdaptiveStreamSelector::FRebufferAction FABRLiveStream::GetRebufferAction()
 		}
 		else
 		{
-			Action.Action = IAdaptiveStreamSelector::FRebufferAction::EAction::GoToLive;
+			ApplyRebufferingPenalty();
+			Action.Action = Info->ABRShouldPlayOnLiveEdge() ? IAdaptiveStreamSelector::FRebufferAction::EAction::GoToLive : IAdaptiveStreamSelector::FRebufferAction::EAction::Restart;
+			//Action.Action = IAdaptiveStreamSelector::FRebufferAction::EAction::GoToLive;
 		}
 	}
 	return Action;
@@ -1690,63 +1706,76 @@ IAdaptiveStreamSelector::EHandlingAction FABRLiveStream::PeriodicHandle()
 	// Perform catch up / slow down of playback when playing (not paused and not buffering).
 	if (!bIsPaused && !bIsBuffering && Info->ABRGetPlaySpeed() != FTimeValue::GetZero())
 	{
-		const EStreamType StreamType = GetPrimaryStreamType();
-		const FStreamWorkVars* WorkVars = GetWorkVars(StreamType);
-		if (WorkVars)
+		if (bFalloffDetected)
 		{
-			bool bEOS = false;
-			const double AvailableBufferedDuration = GetPlayablePlayerDuration(bEOS, StreamType);
-			const double CurrentLatency = Info->ABRGetLatency().GetAsSeconds();
-			const double TargetLatency = lc.TargetLatency;
-			const double NetworkLatency = WorkVars->AverageLatency.GetWeightedMax(DefaultNetworkLatency);
-			const double Distance = CurrentLatency - TargetLatency;
-			const double AbsDistance = Utils::AbsoluteValue(Distance);
-			const double CurrentPlayRate = Info->ABRGetRenderRateScale();
-			const bool bOvertime = WorkVars->bWentIntoOvertime;
-
-			int32 Gain, Trend;
-			const bool bStable = HasStableBuffer(Gain, Trend, StreamType, 0.0);
-			const bool bMaybeSpeedUp = bStable || Trend>=0;
-
-			// Slow down because of insufficient buffered data?
-			bool bBufferSlowDown = bOvertime || (!bEOS && AvailableBufferedDuration < Utils::Max(lc.LowBufferContentBackoff, NetworkLatency));
-			// Do not slow down if we are supposed to play on the Live edge and are already too far behind.
-			if (Info->ABRShouldPlayOnLiveEdge() && Distance > lc.LowBufferMaxTargetLatencyDistance)
+			bFalloffDetected = false;
+			if (!LatencyExceededVars.bLiveSeekRequestIssued)
 			{
-				bBufferSlowDown = false;
+				LatencyExceededVars.bLiveSeekRequestIssued = true;
+				NextAction = IAdaptiveStreamSelector::EHandlingAction::SeekToLive;
+				Info->LogMessage(IInfoLog::ELevel::Info, FString::Printf(TEXT("Timeline fall off detected, requesting seek to Live edge")));
 			}
-			if (bBufferSlowDown)
-			{
-				const double NewRate = lc.LowBufferMinMinPlayRate;
-				SetRenderRateScale(NewRate);
-			}
-			// Should we be playing on the Live edge?
-			else if (Info->ABRShouldPlayOnLiveEdge() && AbsDistance > lc.ActivationThreshold * TargetLatency)
+		}
+		else
+		{
+			const EStreamType StreamType = GetPrimaryStreamType();
+			const FStreamWorkVars* WorkVars = GetWorkVars(StreamType);
+			if (WorkVars)
 			{
-				double NewRate = CalculateCatchupPlayRate(lc, Distance);
-				if (Utils::AbsoluteValue(NewRate - CurrentPlayRate) > lc.MinRateChangeUseThreshold)
+				bool bEOS = false;
+				const double AvailableBufferedDuration = GetPlayablePlayerDuration(bEOS, StreamType);
+				const double CurrentLatency = Info->ABRGetLatency().GetAsSeconds();
+				const double TargetLatency = lc.TargetLatency;
+				const double NetworkLatency = WorkVars->AverageLatency.GetWeightedMax(DefaultNetworkLatency);
+				const double Distance = CurrentLatency - TargetLatency;
+				const double AbsDistance = Utils::AbsoluteValue(Distance);
+				const double CurrentPlayRate = Info->ABRGetRenderRateScale();
+				const bool bOvertime = WorkVars->bWentIntoOvertime;
+
+				int32 Gain, Trend;
+				const bool bStable = HasStableBuffer(Gain, Trend, StreamType, 0.0);
+				const bool bMaybeSpeedUp = bStable || Trend>=0;
+
+				// Slow down because of insufficient buffered data?
+				bool bBufferSlowDown = bOvertime || (!bEOS && AvailableBufferedDuration < Utils::Max(lc.LowBufferContentBackoff, NetworkLatency));
+				// Do not slow down if we are supposed to play on the Live edge and are already too far behind.
+				if (Info->ABRShouldPlayOnLiveEdge() && Distance > lc.LowBufferMaxTargetLatencyDistance)
+				{
+					bBufferSlowDown = false;
+				}
+				if (bBufferSlowDown)
 				{
-					// If we are to speed up to catch up with the Live edge we do so only if the buffer is stable.
-					if (NewRate > 1.0 && !bMaybeSpeedUp)
+					const double NewRate = lc.LowBufferMinMinPlayRate;
+					SetRenderRateScale(NewRate);
+				}
+				// Should we be playing on the Live edge?
+				else if (Info->ABRShouldPlayOnLiveEdge() && AbsDistance > lc.ActivationThreshold * TargetLatency)
+				{
+					double NewRate = CalculateCatchupPlayRate(lc, Distance);
+					if (Utils::AbsoluteValue(NewRate - CurrentPlayRate) > lc.MinRateChangeUseThreshold)
 					{
-						NewRate = 1.0;
+						// If we are to speed up to catch up with the Live edge we do so only if the buffer is stable.
+						if (NewRate > 1.0 && !bMaybeSpeedUp)
+						{
+							NewRate = 1.0;
+						}
+						SetRenderRateScale(NewRate);
 					}
-					SetRenderRateScale(NewRate);
 				}
-			}
-			else if (CurrentPlayRate != 1.0)
-			{
-				SetRenderRateScale(1.0);
-			}
+				else if (CurrentPlayRate != 1.0)
+				{
+					SetRenderRateScale(1.0);
+				}
 
-			// Is the maximum latency exceeded and a jump to the Live edge required?
-			if (Info->ABRShouldPlayOnLiveEdge() && CurrentLatency > lc.MaxLatency)
-			{
-				if (!LatencyExceededVars.bLiveSeekRequestIssued)
+				// Is the maximum latency exceeded and a jump to the Live edge required?
+				if (Info->ABRShouldPlayOnLiveEdge() && CurrentLatency > lc.MaxLatency)
 				{
-					LatencyExceededVars.bLiveSeekRequestIssued = true;
-					NextAction = IAdaptiveStreamSelector::EHandlingAction::SeekToLive;
-					Info->LogMessage(IInfoLog::ELevel::Info, FString::Printf(TEXT("Maximum latency exceeded, requesting seek to Live edge")));
+					if (!LatencyExceededVars.bLiveSeekRequestIssued)
+					{
+						LatencyExceededVars.bLiveSeekRequestIssued = true;
+						NextAction = IAdaptiveStreamSelector::EHandlingAction::SeekToLive;
+						Info->LogMessage(IInfoLog::ELevel::Info, FString::Printf(TEXT("Maximum latency exceeded, requesting seek to Live edge")));
+					}
 				}
 			}
 		}
diff --git a/Engine/Plugins/Media/ElectraPlayer/Source/ElectraPlayerRuntime/Private/Runtime/Player/ABRRules/ABRVoDPlus.cpp b/Engine/Plugins/Media/ElectraPlayer/Source/ElectraPlayerRuntime/Private/Runtime/Player/ABRRules/ABRVoDPlus.cpp
--- a/Engine/Plugins/Media/ElectraPlayer/Source/ElectraPlayerRuntime/Private/Runtime/Player/ABRRules/ABRVoDPlus.cpp
+++ b/Engine/Plugins/Media/ElectraPlayer/Source/ElectraPlayerRuntime/Private/Runtime/Player/ABRRules/ABRVoDPlus.cpp
@@ -430,6 +430,10 @@ FABRDownloadProgressDecision FABROnDemandPlus::ReportDownloadProgress(const Metr
 
 void FABROnDemandPlus::ReportDownloadEnd(const Metrics::FSegmentDownloadStats& SegmentDownloadStats)
 {
+	if (SegmentDownloadStats.bWasSkipped)
+	{
+		return;
+	}
 	FStreamWorkVars* WorkVars = GetWorkVars(SegmentDownloadStats.StreamType);
 
 	if (WorkVars && SegmentDownloadStats.SegmentType == Metrics::ESegmentType::Media)
@@ -853,7 +857,7 @@ IAdaptiveStreamSelector::ESegmentAction FABROnDemandPlus::PerformSelection(const
 	{
 		FStreamWorkVars* WorkVars = GetWorkVars(StreamType);
 
-		int32 NewQualityIndex = 0;
+		int32 NewQualityIndex = -1;
 
 		bool bEOS = false;
 		const double AvailableDuration = GetPlayablePlayerDuration(bEOS, StreamType);
@@ -906,6 +910,13 @@ IAdaptiveStreamSelector::ESegmentAction FABROnDemandPlus::PerformSelection(const
 				}
 			}
 
+			// None of the candidates is deemed feasible. This may happen if all low quality streams have issues and ended up
+			// on the deny list. Pick the lowest quality one, even if this is one that may not be playable without rebuffering.
+			if (NewQualityIndex < 0)
+			{
+				NewQualityIndex = InCandidates[0]->QualityIndex;
+			}
+
 			/*
 				Check if we are in the lower configured range of all qualities and if so, hold back the download
 				for a while to avoid polluting the buffer with low quality data in case the bandwidth will recover soon
