ElectraPlayer: Change to new HLS handler

#rnx

diff --git a/Engine/Plugins/Media/ElectraPlayer/Source/ElectraPlayerRuntime/Private/Runtime/Player/AdaptiveStreamingPlayer.cpp b/Engine/Plugins/Media/ElectraPlayer/Source/ElectraPlayerRuntime/Private/Runtime/Player/AdaptiveStreamingPlayer.cpp
--- a/Engine/Plugins/Media/ElectraPlayer/Source/ElectraPlayerRuntime/Private/Runtime/Player/AdaptiveStreamingPlayer.cpp
+++ b/Engine/Plugins/Media/ElectraPlayer/Source/ElectraPlayerRuntime/Private/Runtime/Player/AdaptiveStreamingPlayer.cpp
@@ -1917,7 +1917,7 @@ void FAdaptiveStreamingPlayer::HandleSeeking()
 void FAdaptiveStreamingPlayer::HandleMetadataChanges()
 {
 	// The timeline can change dynamically. Refresh it on occasion.
-	if (Manifest.IsValid() && ManifestType == EMediaFormatType::DASH)
+	if (Manifest.IsValid() && (ManifestType == EMediaFormatType::DASH || ManifestType == EMediaFormatType::HLS))
 	{
 		PlaybackState.SetSeekableRange(Manifest->GetSeekableTimeRange());
 		PlaybackState.SetTimelineRange(Manifest->GetTotalTimeRange());
diff --git a/Engine/Plugins/Media/ElectraPlayer/Source/ElectraPlayerRuntime/Private/Runtime/Player/AdaptiveStreamingPlayerManifest.cpp b/Engine/Plugins/Media/ElectraPlayer/Source/ElectraPlayerRuntime/Private/Runtime/Player/AdaptiveStreamingPlayerManifest.cpp
--- a/Engine/Plugins/Media/ElectraPlayer/Source/ElectraPlayerRuntime/Private/Runtime/Player/AdaptiveStreamingPlayerManifest.cpp
+++ b/Engine/Plugins/Media/ElectraPlayer/Source/ElectraPlayerRuntime/Private/Runtime/Player/AdaptiveStreamingPlayerManifest.cpp
@@ -4,7 +4,7 @@
 #include "Misc/Paths.h"
 #include "Player/AdaptiveStreamingPlayerInternal.h"
 #include "Player/AdaptivePlayerOptionKeynames.h"
-#include "Player/HLS/PlaylistReaderHLS.h"
+#include "Player/HLS/PlaylistHandlerHLS.h"
 #include "Player/mp4/PlaylistReaderMP4.h"
 #include "Player/DASH/PlaylistReaderDASH.h"
 #include "Player/mkv/PlaylistReaderMKV.h"
@@ -221,7 +221,7 @@ void FAdaptiveStreamingPlayer::InternalLoadManifest(const FString& InURL, const
 			CurrentState = EPlayerState::eState_ParsingManifest;
 			if (mimeType == Playlist::MIMETypeHLS)
 			{
-				ManifestReader = IPlaylistReaderHLS::Create(this);
+				ManifestReader = IPlaylistHandlerHLS::Create(this);
 				ManifestType = EMediaFormatType::HLS;
 			}
 			else if (mimeType == Playlist::MIMETypeMP4 || mimeType == Playlist::MIMETypeMP4A || mimeType == Playlist::MIMETypeQuickTime)
