ElectraPlayer: Fixed incorrect representation selection in HLS handler when variant filtering is in use.

#rnx

diff --git a/Engine/Plugins/Media/ElectraPlayer/Source/ElectraPlayerRuntime/Private/Runtime/Player/HLS/PlaylistHLS.cpp b/Engine/Plugins/Media/ElectraPlayer/Source/ElectraPlayerRuntime/Private/Runtime/Player/HLS/PlaylistHLS.cpp
--- a/Engine/Plugins/Media/ElectraPlayer/Source/ElectraPlayerRuntime/Private/Runtime/Player/HLS/PlaylistHLS.cpp
+++ b/Engine/Plugins/Media/ElectraPlayer/Source/ElectraPlayerRuntime/Private/Runtime/Player/HLS/PlaylistHLS.cpp
@@ -1421,7 +1421,7 @@ FErrorDetail FActiveHLSPlaylist::FTimelineMediaAsset::GetVariantPlaylist(TShared
 			MultiVariantPlaylist->StreamInfs[vIdx].PathwayId.Equals(InPathway))
 		{
 			si = &MultiVariantPlaylist->StreamInfs[vIdx];
-			Representation = AdaptationSet->GetRepresentationByIndex(vIdx);
+			Representation = AdaptationSet->GetRepresentationByUniqueIdentifier(VariantID);
 			break;
 		}
 	}
@@ -2442,7 +2442,7 @@ FString FActiveHLSPlaylist::FPlayPeriod::GetSelectedAdaptationSetID(EStreamType
 
 IManifest::IPlayPeriod::ETrackChangeResult FActiveHLSPlaylist::FPlayPeriod::ChangeTrackStreamPreference(EStreamType InStreamType, const FStreamSelectionAttributes& InStreamAttributes)
 {
-check(!"TODO");
+	check(!"track changes not supported yet");
 	return IManifest::IPlayPeriod::ETrackChangeResult::NotChanged;
 }
 
