ElectraPlayer: Removed option to suppress a seek if another one is ongoing as well as not seeking if the target time is already within a threshold of the current time. Trying to suppress seeks inside the media player is a problem because the facade expects a different sequence index on the decoder outputs with every seek, and not performing a seek results in a mismatch.

#rnx

diff --git a/Engine/Plugins/Media/ElectraPlayer/Source/ElectraPlayerPlugin/Private/ElectraPlayerPlugin.cpp b/Engine/Plugins/Media/ElectraPlayer/Source/ElectraPlayerPlugin/Private/ElectraPlayerPlugin.cpp
--- a/Engine/Plugins/Media/ElectraPlayer/Source/ElectraPlayerPlugin/Private/ElectraPlayerPlugin.cpp
+++ b/Engine/Plugins/Media/ElectraPlayer/Source/ElectraPlayerPlugin/Private/ElectraPlayerPlugin.cpp
@@ -654,7 +654,6 @@ bool FElectraPlayerPlugin::Open(const FString& Url, const IMediaOptions* Options
 	if (Environment == MediaPlayerOptionValues::Environment_Preview() || Environment == MediaPlayerOptionValues::Environment_Sequencer())
 	{
 		PlayerOptions.Set(TEXT("optimize_seek_for_scrubbing"), Electra::FVariantValue(true));
-		PlayerOptions.Set(TEXT("new_scrubbing_seek_cancels_current"), Electra::FVariantValue(true));
 		//PlayerOptions.Set(TEXT("do_not_hold_back_first_frame"), Electra::FVariantValue(true));
 		PlayerOptions.Set(TEXT("always_emit_samples_when_paused"), Electra::FVariantValue(true));
 		PlayerOptions.Set(TEXT("worker_threads"), Electra::FVariantValue(FString(TEXT("worker"))));
diff --git a/Engine/Plugins/Media/ElectraPlayer/Source/ElectraPlayerRuntime/Private/ElectraPlayer.cpp b/Engine/Plugins/Media/ElectraPlayer/Source/ElectraPlayerRuntime/Private/ElectraPlayer.cpp
--- a/Engine/Plugins/Media/ElectraPlayer/Source/ElectraPlayerRuntime/Private/ElectraPlayer.cpp
+++ b/Engine/Plugins/Media/ElectraPlayer/Source/ElectraPlayerRuntime/Private/ElectraPlayer.cpp
@@ -1272,7 +1272,6 @@ bool FElectraPlayer::Seek(const FTimespan& Time, const FSeekParam& Param)
 		seek.NewSequenceIndex = Param.SequenceIndex;
 		seek.StartingBitrate = Param.StartingBitrate;
 		seek.bOptimizeForScrubbing = Param.bOptimizeForScrubbing;
-		seek.DistanceThreshold = Param.DistanceThreshold;
 		bInitialSeekPerformed = true;
 		bDiscardOutputUntilCleanStart = true;
 		CurrentPlayer->AdaptivePlayer->SeekTo(seek);
diff --git a/Engine/Plugins/Media/ElectraPlayer/Source/ElectraPlayerRuntime/Private/Runtime/Player/AdaptivePlayerOptionKeynames.h b/Engine/Plugins/Media/ElectraPlayer/Source/ElectraPlayerRuntime/Private/Runtime/Player/AdaptivePlayerOptionKeynames.h
--- a/Engine/Plugins/Media/ElectraPlayer/Source/ElectraPlayerRuntime/Private/Runtime/Player/AdaptivePlayerOptionKeynames.h
+++ b/Engine/Plugins/Media/ElectraPlayer/Source/ElectraPlayerRuntime/Private/Runtime/Player/AdaptivePlayerOptionKeynames.h
@@ -36,9 +36,6 @@ namespace Electra
 	//! (bool) true to optimize seeking for faster frame scrubbing, false to optimize for playback.
 	const FName OptionKeyFrameOptimizeSeekForScrubbing(TEXT("optimize_seek_for_scrubbing"));
 
-	//! (bool) true to allow a new scrubbing seek to cancel an ongoing scrubbing seek. Non-scrubbing seeks always cancel pending seeks.
-	const FName OptionKeyNewScrubbingSeekCancelsCurrent(TEXT("new_scrubbing_seek_cancels_current"));
-
 	//! (bool) true to emit the first decoded video frame while prerolling so it can be displayed while scrubbing.
 	const FName OptionKeyDoNotHoldBackFirstVideoFrame(TEXT("do_not_hold_back_first_frame"));
 
diff --git a/Engine/Plugins/Media/ElectraPlayer/Source/ElectraPlayerRuntime/Private/Runtime/Player/AdaptiveStreamingPlayer.cpp b/Engine/Plugins/Media/ElectraPlayer/Source/ElectraPlayerRuntime/Private/Runtime/Player/AdaptiveStreamingPlayer.cpp
--- a/Engine/Plugins/Media/ElectraPlayer/Source/ElectraPlayerRuntime/Private/Runtime/Player/AdaptiveStreamingPlayer.cpp
+++ b/Engine/Plugins/Media/ElectraPlayer/Source/ElectraPlayerRuntime/Private/Runtime/Player/AdaptiveStreamingPlayer.cpp
@@ -1818,57 +1818,27 @@ void FAdaptiveStreamingPlayer::HandleSeeking()
 		return;
 	}
 
-	// When playing, the last successfully seeked position is irrelevant and cannot be used as a reference any more.
-	if (PlaybackState.GetIsPlaying())
-	{
-		SeekVars.InvalidateLastFinished();
-	}
-
 	FScopeLock lock(&SeekVars.Lock);
 	// Is there a pending request?
 	if (SeekVars.PendingRequest.IsSet())
 	{
-		// If there is an active request and the new request is for scrubbing we let the active request finish first.
-		bool bIsForScrubbing = SeekVars.PendingRequest.GetValue().bOptimizeForScrubbing.Get(PlayerOptions.GetValue(OptionKeyFrameOptimizeSeekForScrubbing).SafeGetBool(false));
-		bool bNewScrubSeekCancelsCurrent = PlayerOptions.GetValue(OptionKeyNewScrubbingSeekCancelsCurrent).SafeGetBool(false);
-		if (SeekVars.ActiveRequest.IsSet() && bIsForScrubbing && !bNewScrubSeekCancelsCurrent)
-		{
-			return;
-		}
-
 		// Not playing on the Live edge. This may get set to true later on handling the start segment.
 		PlaybackState.SetShouldPlayOnLiveEdge(false);
-		// On a user-induced seek the sequence index is supposed to be increased.
-		// We need to do this even when we would not actually issue the seek.
 
-		// Adjust seek index as indicated
-		// (we even do this for seeks marked as not affecting the seek index as we must still add any priovious accumulated calls that had not been flagged in such a way)
+		// Adjust seek index as necessary.
 		if (SeekVars.PendingRequest.GetValue().NewSequenceIndex.IsSet())
 		{
 			CurrentPlaybackSequenceState.PrimaryIndex = SeekVars.PendingRequest.GetValue().NewSequenceIndex.GetValue();
 		}
+		// Contractually a seek implies resetting the loop count, which is represented in the 2ndary index.
 		CurrentPlaybackSequenceState.SecondaryIndex = 0;
-
-		// And since it is a seek on purpose the loop counter is reset as well.
+		// Reset the loop counter in the state as well.
 		FInternalLoopState LoopStateNow;
 		PlaybackState.GetLoopState(LoopStateNow);
 		LoopStateNow.Count = 0;
 		PlaybackState.SetLoopState(LoopStateNow);
 		CurrentLoopState.Count = 0;
 
-		// Check the distance to the last seek performed, if there is one.
-		if (SeekVars.LastFinishedRequest.IsSet() && SeekVars.PendingRequest.GetValue().DistanceThreshold.IsSet())
-		{
-			double Distance = Utils::AbsoluteValue(SeekVars.LastFinishedRequest.GetValue().Time.GetAsSeconds() - SeekVars.PendingRequest.GetValue().Time.GetAsSeconds());
-			if (Distance <= SeekVars.PendingRequest.GetValue().DistanceThreshold.Get(0.0))
-			{
-				// Already there
-				SeekVars.PendingRequest.Reset();
-				DispatchEvent(FMetricEvent::ReportSeekCompleted(true));
-				return;
-			}
-		}
-
 		// Trigger the seek.
 		FSeekParam SeekParam = SeekVars.PendingRequest.GetValue();
 		FTimeRange SeekPlayrangeOnRequest = SeekVars.PlayrangeOnRequest;
@@ -1893,6 +1863,7 @@ void FAdaptiveStreamingPlayer::HandleSeeking()
 			CurrentState = EPlayerState::eState_Seeking;
 			SeekVars.Lock.Lock();
 
+			bool bIsForScrubbing = SeekParam.bOptimizeForScrubbing.Get(PlayerOptions.GetValue(OptionKeyFrameOptimizeSeekForScrubbing).SafeGetBool(false));
 			SeekVars.bForScrubbing = bIsForScrubbing;
 
 			// When seeking (other than the initial playstart seek) any default end time that
@@ -3695,9 +3666,8 @@ void FAdaptiveStreamingPlayer::InternalStartoverAtCurrentPosition()
 		// Is there a pending request?
 		if (SeekVars.PendingRequest.IsSet())
 		{
-			// To ensure the seek will execute we clear out the active and last finished requests.
+			// To ensure the seek will execute we clear out the active requests.
 			SeekVars.ActiveRequest.Reset();
-			SeekVars.LastFinishedRequest.Reset();
 			return;
 		}
 	}
diff --git a/Engine/Plugins/Media/ElectraPlayer/Source/ElectraPlayerRuntime/Private/Runtime/Player/AdaptiveStreamingPlayer.h b/Engine/Plugins/Media/ElectraPlayer/Source/ElectraPlayerRuntime/Private/Runtime/Player/AdaptiveStreamingPlayer.h
--- a/Engine/Plugins/Media/ElectraPlayer/Source/ElectraPlayerRuntime/Private/Runtime/Player/AdaptiveStreamingPlayer.h
+++ b/Engine/Plugins/Media/ElectraPlayer/Source/ElectraPlayerRuntime/Private/Runtime/Player/AdaptiveStreamingPlayer.h
@@ -134,15 +134,13 @@ public:
 			bOptimizeForScrubbing.Reset();
 		}
 		// Time to seek to.
-		FTimeValue	Time;
+		FTimeValue Time;
 		// New sequence index to associate with the newly decoded samples.
 		TOptional<int32> NewSequenceIndex;
 		// Maximum stream bitrate to use when seeking.
 		TOptional<int32> StartingBitrate;
 		// Optimize for frame scrubbing (faster display of frame at target time)?
 		TOptional<bool> bOptimizeForScrubbing;
-		// Allowed distance to last performed seek to save a redundant new seek.
-		TOptional<double> DistanceThreshold;
 	};
 
 	/**
diff --git a/Engine/Plugins/Media/ElectraPlayer/Source/ElectraPlayerRuntime/Private/Runtime/Player/AdaptiveStreamingPlayerInternal.h b/Engine/Plugins/Media/ElectraPlayer/Source/ElectraPlayerRuntime/Private/Runtime/Player/AdaptiveStreamingPlayerInternal.h
--- a/Engine/Plugins/Media/ElectraPlayer/Source/ElectraPlayerRuntime/Private/Runtime/Player/AdaptiveStreamingPlayerInternal.h
+++ b/Engine/Plugins/Media/ElectraPlayer/Source/ElectraPlayerRuntime/Private/Runtime/Player/AdaptiveStreamingPlayerInternal.h
@@ -1892,7 +1892,6 @@ private:
 			//       triggered while already processing a seek due to the asynchronous processing of the request.
 			//PendingRequest.Reset();
 			ActiveRequest.Reset();
-			LastFinishedRequest.Reset();
 			PlayrangeOnRequest.Reset();
 			bForScrubbing = false;
 			bScrubPrerollDone = false;
@@ -1900,21 +1899,14 @@ private:
 
 		void SetFinished()
 		{
-			LastFinishedRequest = ActiveRequest;
 			ActiveRequest.Reset();
 		}
 
-		void InvalidateLastFinished()
-		{
-			LastFinishedRequest.Reset();
-		}
-
-		bool	bForScrubbing = false;
-		bool	bScrubPrerollDone = false;
-		bool	bIsPlayStart = true;
+		bool bForScrubbing = false;
+		bool bScrubPrerollDone = false;
+		bool bIsPlayStart = true;
 
 		TOptional<FSeekParam> ActiveRequest;
-		TOptional<FSeekParam> LastFinishedRequest;
 
 		mutable FCriticalSection Lock;
 		// The pending request must be accessed under lock since it is written to by the main thread and read from the worker thread!
diff --git a/Engine/Plugins/Media/ElectraPlayer/Source/ElectraPlayerRuntime/Public/IElectraPlayerInterface.h b/Engine/Plugins/Media/ElectraPlayer/Source/ElectraPlayerRuntime/Public/IElectraPlayerInterface.h
--- a/Engine/Plugins/Media/ElectraPlayer/Source/ElectraPlayerRuntime/Public/IElectraPlayerInterface.h
+++ b/Engine/Plugins/Media/ElectraPlayer/Source/ElectraPlayerRuntime/Public/IElectraPlayerInterface.h
@@ -266,7 +266,6 @@ public:
 		TOptional<int32> SequenceIndex;
 		TOptional<int32> StartingBitrate;
 		TOptional<bool> bOptimizeForScrubbing;
-		TOptional<double> DistanceThreshold;
 	};
 
 	virtual bool Seek(const FTimespan& Time) = 0;
