Electra: Retaining the resource delegate during the creation of the video decoder

#rnx

diff --git a/Engine/Plugins/Media/ElectraPlayer/Source/ElectraPlayerRuntime/Private/Runtime/Decoder/VideoDecoder.cpp b/Engine/Plugins/Media/ElectraPlayer/Source/ElectraPlayerRuntime/Private/Runtime/Decoder/VideoDecoder.cpp
--- a/Engine/Plugins/Media/ElectraPlayer/Source/ElectraPlayerRuntime/Private/Runtime/Decoder/VideoDecoder.cpp
+++ b/Engine/Plugins/Media/ElectraPlayer/Source/ElectraPlayerRuntime/Private/Runtime/Decoder/VideoDecoder.cpp
@@ -555,8 +555,9 @@ bool FVideoDecoderImpl::InternalDecoderCreate()
 	}
 
 	// Create platform specifics.
+	TSharedPtr<IVideoDecoderResourceDelegate, ESPMode::ThreadSafe> PinnedVideoResourceDelegate = VideoResourceDelegate.Pin();
 	TMap<FString, FVariant> PlatformSpecificCfg(DecoderFactoryAddtlCfg);
-	PlatformSpecificCfg.Emplace(TEXT("VideoResourceDelegate"), FVariant((uint64)VideoResourceDelegate.Pin().Get()));
+	PlatformSpecificCfg.Emplace(TEXT("VideoResourceDelegate"), FVariant((uint64)PinnedVideoResourceDelegate.Get()));
 	PlatformResource = FPlatformElectraDecoderResourceManager::GetDelegate()->CreatePlatformResource(this, IElectraDecoderResourceDelegateBase::EDecoderPlatformResourceType::Video, PlatformSpecificCfg);
 
 	// Put a pointer to the renderer into the decoder creation configuration.
