ElectraPlayer: Audio track changes with HLS are now supported

#rnx

diff --git a/Engine/Plugins/Media/ElectraPlayer/Source/ElectraPlayerRuntime/Private/Runtime/Player/AdaptiveStreamingPlayer.cpp b/Engine/Plugins/Media/ElectraPlayer/Source/ElectraPlayerRuntime/Private/Runtime/Player/AdaptiveStreamingPlayer.cpp
--- a/Engine/Plugins/Media/ElectraPlayer/Source/ElectraPlayerRuntime/Private/Runtime/Player/AdaptiveStreamingPlayer.cpp
+++ b/Engine/Plugins/Media/ElectraPlayer/Source/ElectraPlayerRuntime/Private/Runtime/Player/AdaptiveStreamingPlayer.cpp
@@ -1593,19 +1593,27 @@ bool FAdaptiveStreamingPlayer::InternalHandleThreadMessages()
 					switch(ev.StreamType)
 					{
 						case EStreamType::Video:
+						{
 							PendingTrackSelectionVid = MakeSharedTS<FStreamSelectionAttributes>(ev.TrackAttributes);
 							bIsVideoDeselected = false;
 							break;
+						}
 						case EStreamType::Audio:
+						{
 							PendingTrackSelectionAud = MakeSharedTS<FStreamSelectionAttributes>(ev.TrackAttributes);
 							bIsAudioDeselected = false;
 							break;
+						}
 						case EStreamType::Subtitle:
+						{
 							PendingTrackSelectionTxt = MakeSharedTS<FStreamSelectionAttributes>(ev.TrackAttributes);
 							bIsTextDeselected = false;
 							break;
+						}
 						default:
+						{
 							break;
+						}
 					}
 				}
 				break;
@@ -3637,24 +3645,32 @@ void FAdaptiveStreamingPlayer::InternalDeselectStream(EStreamType StreamType)
 	switch(StreamType)
 	{
 		case EStreamType::Video:
+		{
 			bIsVideoDeselected = true;
 			SelectedStreamAttributesVid.Deselect();
 			VideoDecoder.Flush();
 			VideoRender.Flush(false);
 			break;
+		}
 		case EStreamType::Audio:
+		{
 			bIsAudioDeselected = true;
 			SelectedStreamAttributesAud.Deselect();
 			AudioDecoder.Flush();
 			AudioRender.Flush();
 			break;
+		}
 		case EStreamType::Subtitle:
+		{
 			bIsTextDeselected = true;
 			SelectedStreamAttributesTxt.Deselect();
 			SubtitleDecoder.Flush();
 			break;
+		}
 		default:
+		{
 			break;
+		}
 	}
 }
 
@@ -3693,10 +3709,11 @@ void FAdaptiveStreamingPlayer::InternalStartoverAtLiveEdge()
 void FAdaptiveStreamingPlayer::InternalHandleSegmentTrackChanges(const FTimeValue& CurrentTime)
 {
 	bool bChangeMade = false;
+	bool bForceStartover = false;
 
 	if (ManifestType == EMediaFormatType::ISOBMFF || ManifestType == EMediaFormatType::MKV)
 	{
-		auto HandleISOBMFFTrackSelection = [&bChangeMade, this](FStreamSelectionAttributes& OutSelectionAttributes, FInternalStreamSelectionAttributes& InOutSelectedAttributes, EStreamType InType,
+		auto HandleMultiplexedTrackSelection = [&bChangeMade, this](FStreamSelectionAttributes& OutSelectionAttributes, FInternalStreamSelectionAttributes& InOutSelectedAttributes, EStreamType InType,
 																TSharedPtrTS<FStreamSelectionAttributes>& InPendingSelection, TSharedPtrTS<IManifest::IPlayPeriod> InCurrentPeriod) -> void
 		{
 			if (InPendingSelection.IsValid() && InCurrentPeriod.IsValid())
@@ -3719,37 +3736,18 @@ void FAdaptiveStreamingPlayer::InternalHandleSegmentTrackChanges(const FTimeValu
 			}
 		};
 
-		HandleISOBMFFTrackSelection(StreamSelectionAttributesVid, SelectedStreamAttributesVid, EStreamType::Video, PendingTrackSelectionVid, CurrentPlayPeriodVideo);
-		HandleISOBMFFTrackSelection(StreamSelectionAttributesAud, SelectedStreamAttributesAud, EStreamType::Audio, PendingTrackSelectionAud, CurrentPlayPeriodAudio);
-		HandleISOBMFFTrackSelection(StreamSelectionAttributesTxt, SelectedStreamAttributesTxt, EStreamType::Subtitle, PendingTrackSelectionTxt, CurrentPlayPeriodText);
+		HandleMultiplexedTrackSelection(StreamSelectionAttributesVid, SelectedStreamAttributesVid, EStreamType::Video, PendingTrackSelectionVid, CurrentPlayPeriodVideo);
+		HandleMultiplexedTrackSelection(StreamSelectionAttributesAud, SelectedStreamAttributesAud, EStreamType::Audio, PendingTrackSelectionAud, CurrentPlayPeriodAudio);
+		HandleMultiplexedTrackSelection(StreamSelectionAttributesTxt, SelectedStreamAttributesTxt, EStreamType::Subtitle, PendingTrackSelectionTxt, CurrentPlayPeriodText);
 	}
-	else if (ManifestType == EMediaFormatType::HLS)
-	{
-		if (PendingTrackSelectionAud.IsValid() && CurrentPlayPeriodAudio.IsValid())
-		{
-			TSharedPtrTS<FBufferSourceInfo> BufferSourceInfo = CurrentPlayPeriodAudio->GetSelectedStreamBufferSourceInfo(EStreamType::Audio);
-			TSharedPtrTS<FMultiTrackAccessUnitBuffer> StreamRcvBuf = GetCurrentReceiveStreamBuffer(EStreamType::Audio);
-			if (StreamRcvBuf.IsValid())
-			{
-				StreamRcvBuf->SelectTrackWhenAvailable(CurrentPlaybackSequenceID[StreamTypeToArrayIndex(EStreamType::Audio)], BufferSourceInfo);
-			}
-			StreamSelectionAttributesAud = *PendingTrackSelectionAud;
-			SelectedStreamAttributesAud.UpdateWith(BufferSourceInfo->Kind, BufferSourceInfo->Language, BufferSourceInfo->Codec, BufferSourceInfo->HardIndex);
-			PendingTrackSelectionAud.Reset();
-			bChangeMade = true;
-		}
-		// Ignore video and subtitle changes for now.
-		PendingTrackSelectionVid.Reset();
-		PendingTrackSelectionTxt.Reset();
-	}
-	else if (ManifestType == EMediaFormatType::DASH)
+	else if (ManifestType == EMediaFormatType::DASH || ManifestType == EMediaFormatType::HLS)
 	{
 		if (PendingStartRequest.IsValid() || PendingFirstSegmentRequest.IsValid())
 		{
 			return;
 		}
 
-		auto HandleDASHTrackSelection = [&bChangeMade, this](FStreamSelectionAttributes& OutSelectionAttributes, FInternalStreamSelectionAttributes& InOutSelectedAttributes, EStreamType InType,
+		auto HandleElementaryStreamTrackSelection = [&bChangeMade, &bForceStartover, this](FStreamSelectionAttributes& OutSelectionAttributes, FInternalStreamSelectionAttributes& InOutSelectedAttributes, EStreamType InType,
 															 TSharedPtrTS<FStreamSelectionAttributes>& InPendingSelection, TSharedPtrTS<IManifest::IPlayPeriod> InCurrentPeriod) -> void
 		{
 			if (InPendingSelection.IsValid() && InCurrentPeriod.IsValid())
@@ -3776,6 +3774,13 @@ void FAdaptiveStreamingPlayer::InternalHandleSegmentTrackChanges(const FTimeValu
 
 						bChangeMade = true;
 					}
+					else if (TrackChangeResult == IManifest::IPlayPeriod::ETrackChangeResult::StartOver)
+					{
+						OutSelectionAttributes = *InPendingSelection;
+						InOutSelectedAttributes.Select();
+						bChangeMade = true;
+						bForceStartover = true;
+					}
 				}
 				InOutSelectedAttributes.Select();
 				InPendingSelection.Reset();
@@ -3783,8 +3788,11 @@ void FAdaptiveStreamingPlayer::InternalHandleSegmentTrackChanges(const FTimeValu
 		};
 
 		PendingTrackSelectionVid.Reset();	// Ignore video changes for now.
-		HandleDASHTrackSelection(StreamSelectionAttributesAud, SelectedStreamAttributesAud, EStreamType::Audio, PendingTrackSelectionAud, CurrentPlayPeriodAudio);
-		HandleDASHTrackSelection(StreamSelectionAttributesTxt, SelectedStreamAttributesTxt, EStreamType::Subtitle, PendingTrackSelectionTxt, CurrentPlayPeriodText);
+		HandleElementaryStreamTrackSelection(StreamSelectionAttributesAud, SelectedStreamAttributesAud, EStreamType::Audio, PendingTrackSelectionAud, CurrentPlayPeriodAudio);
+		if (ManifestType == EMediaFormatType::DASH)
+		{
+			HandleElementaryStreamTrackSelection(StreamSelectionAttributesTxt, SelectedStreamAttributesTxt, EStreamType::Subtitle, PendingTrackSelectionTxt, CurrentPlayPeriodText);
+		}
 	}
 	else
 	{
@@ -3798,7 +3806,7 @@ void FAdaptiveStreamingPlayer::InternalHandleSegmentTrackChanges(const FTimeValu
 	if (bChangeMade)
 	{
 		DataBuffersCriticalSection.Lock();
-		bool bNeedStartover = (ActiveDataOutputBuffers.IsValid() && NextDataBuffers.Num() > 0) || NextDataBuffers.Num() > 1;
+		bool bNeedStartover = bForceStartover || (ActiveDataOutputBuffers.IsValid() && NextDataBuffers.Num() > 0) || NextDataBuffers.Num() > 1;
 		DataBuffersCriticalSection.Unlock();
 		if (bNeedStartover)
 		{
diff --git a/Engine/Plugins/Media/ElectraPlayer/Source/ElectraPlayerRuntime/Private/Runtime/Player/HLS/PlaylistHLS.cpp b/Engine/Plugins/Media/ElectraPlayer/Source/ElectraPlayerRuntime/Private/Runtime/Player/HLS/PlaylistHLS.cpp
--- a/Engine/Plugins/Media/ElectraPlayer/Source/ElectraPlayerRuntime/Private/Runtime/Player/HLS/PlaylistHLS.cpp
+++ b/Engine/Plugins/Media/ElectraPlayer/Source/ElectraPlayerRuntime/Private/Runtime/Player/HLS/PlaylistHLS.cpp
@@ -1488,7 +1488,9 @@ FErrorDetail FActiveHLSPlaylist::FTimelineMediaAsset::GetVariantPlaylist(TShared
 				if ((*Renditions)[k].StableRenditionId.Equals(tm.Rendition.GetValue().StableRenditionId))
 				{
 					Rendition = &(*Renditions)[k];
-					Representation = AdaptationSet->GetRepresentationByIndex(k);
+					// Adaptation sets created from renditions only have a single rendition in them.
+					check(AdaptationSet->GetNumberOfRepresentations() == 1);
+					Representation = AdaptationSet->GetRepresentationByIndex(0);
 					break;
 				}
 			}
@@ -1501,7 +1503,9 @@ FErrorDetail FActiveHLSPlaylist::FTimelineMediaAsset::GetVariantPlaylist(TShared
 				if ((*Renditions)[k].Name.Equals(tm.Meta.Label))
 				{
 					Rendition = &(*Renditions)[k];
-					Representation = AdaptationSet->GetRepresentationByIndex(k);
+					// Adaptation sets created from renditions only have a single rendition in them.
+					check(AdaptationSet->GetNumberOfRepresentations() == 1);
+					Representation = AdaptationSet->GetRepresentationByIndex(0);
 					break;
 				}
 			}
@@ -2328,6 +2332,11 @@ void FActiveHLSPlaylist::FPlayPeriod::SetStreamPreferences(EStreamType InStreamT
 
 IManifest::IPlayPeriod::EReadyState FActiveHLSPlaylist::FPlayPeriod::GetReadyState()
 {
+	// While the state is preparing, call PrepareForPlay() again to check on media playlist load progress.
+	if (CurrentReadyState == IManifest::IPlayPeriod::EReadyState::Preparing)
+	{
+		PrepareForPlay();
+	}
 	return CurrentReadyState;
 }
 
@@ -2338,61 +2347,121 @@ void FActiveHLSPlaylist::FPlayPeriod::Load()
 
 void FActiveHLSPlaylist::FPlayPeriod::PrepareForPlay()
 {
-	int64 StartingBitrate = PlayerSessionServices->GetOptionValue(OptionKeyCurrentAvgStartingVideoBitrate).SafeGetInt64(2*1000*1000);
+	//int64 StartingBitrate = PlayerSessionServices->GetOptionValue(OptionKeyCurrentAvgStartingVideoBitrate).SafeGetInt64(2*1000*1000);
 
-/*
-	TODO: - based on the initial stream selection locate the track that best matches the type/language
-	      - request the media playlist if not available yet. return `Preparing`. once available change state to `IsReady`
-		  - take note of the selected track and stream index.
-*/
+	TArray<TSharedPtr<FLoadRequestHLSPlaylist, ESPMode::ThreadSafe>> NewLoadReq;
+	int32 NumPending = 0;
 
-	if (TimelineMediaAsset->VideoTracks.Num() && TimelineMediaAsset->VideoTracks[0].Meta.StreamDetails.Num())
+	// Select streams by preference, or the first one of the type if no preference is given.
+	for(int32 nStreamTypeIdx=0; nStreamTypeIdx<3; ++nStreamTypeIdx)
 	{
-		FSelectedTrackStream& st = SelectedTrackStream[StreamTypeToArrayIndex(EStreamType::Video)];
-		st.MetaID = TimelineMediaAsset->VideoTracks[0].Meta.ID;
-		st.TrackIndex = 0;
-		st.StreamIndex = 0;
-		st.bIsSelected = true;
+		FSelectedTrackStream& st = SelectedTrackStream[nStreamTypeIdx];
+		const TArray<FActiveHLSPlaylist::FInternalTrackMetadata>* Tracks = nullptr;
+		const EStreamType StreamType = StreamArrayIndexToType(nStreamTypeIdx);
 
-		st.BufferSourceInfo = MakeSharedTS<FBufferSourceInfo>();
-		st.BufferSourceInfo->PeriodID = TimelineMediaAsset->GetAssetIdentifier();
-		st.BufferSourceInfo->PeriodAdaptationSetID = st.MetaID;
-		st.BufferSourceInfo->Kind = TimelineMediaAsset->VideoTracks[0].Meta.Kind;
-		st.BufferSourceInfo->Language = TimelineMediaAsset->VideoTracks[0].Meta.Language;
-		st.BufferSourceInfo->Codec = Util::GetBaseCodec(TimelineMediaAsset->VideoTracks[0].Meta.HighestBandwidthCodec.GetCodecSpecifierRFC6381());
-		st.BufferSourceInfo->HardIndex = st.StreamIndex;
-
-		FStreamLoadRequest LoadReq;
-		if (TimelineMediaAsset->GetVariantPlaylist(LoadReq.Request, PlayerSessionServices, EStreamType::Video, TimelineMediaAsset->CurrentPathwayId, 0,0,0,0).IsOK())
+		if (nStreamTypeIdx == 0 && TimelineMediaAsset->VideoTracks.Num())
 		{
-			st.ActivePlaylist = TimelineMediaAsset->GetExistingMediaPlaylistFromLoadRequest(LoadReq.Request);
+			Tracks = &TimelineMediaAsset->VideoTracks;
+		}
+		else if (nStreamTypeIdx == 1 && TimelineMediaAsset->AudioTracks.Num())
+		{
+			Tracks = &TimelineMediaAsset->AudioTracks;
+		}
+		else if (nStreamTypeIdx == 2 && TimelineMediaAsset->SubtitleTracks.Num())
+		{
+		// Not setting up subtitle tracks for now as we have no WebVTT parser.
+			//Tracks = &TimelineMediaAsset->SubtitleTracks;
 		}
-	}
 
-	if (TimelineMediaAsset->AudioTracks.Num() && TimelineMediaAsset->AudioTracks[0].Meta.StreamDetails.Num())
-	{
-		FSelectedTrackStream& st = SelectedTrackStream[StreamTypeToArrayIndex(EStreamType::Audio)];
-		st.MetaID = TimelineMediaAsset->AudioTracks[0].Meta.ID;
-		st.TrackIndex = 0;
-		st.StreamIndex = 0;
-		st.bIsSelected = true;
+		int32 SelectedTrackIndex = 0;
+		if (Tracks && StreamSelectionAttributes[nStreamTypeIdx].IsSet() && StreamSelectionAttributes[nStreamTypeIdx].GetValue().IsSet())
+		{
+			const FStreamSelectionAttributes& Sel(StreamSelectionAttributes[nStreamTypeIdx].GetValue());
 
-		st.BufferSourceInfo = MakeSharedTS<FBufferSourceInfo>();
-		st.BufferSourceInfo->PeriodID = TimelineMediaAsset->GetAssetIdentifier();
-		st.BufferSourceInfo->PeriodAdaptationSetID = st.MetaID;
-		st.BufferSourceInfo->Kind = TimelineMediaAsset->AudioTracks[0].Meta.Kind;
-		st.BufferSourceInfo->Language = TimelineMediaAsset->AudioTracks[0].Meta.Language;
-		st.BufferSourceInfo->Codec = Util::GetBaseCodec(TimelineMediaAsset->AudioTracks[0].Meta.HighestBandwidthCodec.GetCodecSpecifierRFC6381());
-		st.BufferSourceInfo->HardIndex = st.StreamIndex;
+			// Is this a hard choice?
+			if (Sel.OverrideIndex.IsSet() && Sel.OverrideIndex.GetValue() < Tracks->Num())
+			{
+				SelectedTrackIndex = Sel.OverrideIndex.GetValue();
+			}
+			else
+			{
+				TArray<int32> CandidateIndices;
+				// Choose language?
+				if (Sel.Language_ISO639.IsSet())
+				{
+					for(int32 i=0; i<Tracks->Num(); ++i)
+					{
+						if ((*Tracks)[i].Meta.Language.Equals(Sel.Language_ISO639.GetValue()))
+						{
+							CandidateIndices.Emplace(i);
+						}
+					}
+				}
+				// If there are multiple language candidates narrow the list down by kind.
+				if (CandidateIndices.Num() && Sel.Kind.IsSet())
+				{
+					TArray<int32> TempList;
+					for(int32 i=0; i<CandidateIndices.Num(); ++i)
+					{
+						if ((*Tracks)[CandidateIndices[i]].Meta.Kind.Equals(Sel.Kind.GetValue()))
+						{
+							TempList.Emplace(i);
+						}
+					}
+					// If there are new candidates update the list. If everything is filtered out, keep the previous list.
+					if (TempList.Num())
+					{
+						Swap(TempList, CandidateIndices);
+					}
+				}
+				// TODO: In the future we could narrow the list down by codec if necessary.
 
-		FStreamLoadRequest LoadReq;
-		if (TimelineMediaAsset->GetVariantPlaylist(LoadReq.Request, PlayerSessionServices, EStreamType::Audio, TimelineMediaAsset->CurrentPathwayId, 0,0,0,0).IsOK())
+				// Use the first candidate's index even if there are several possibilities. If there are none, use the first track.
+				SelectedTrackIndex = CandidateIndices.Num() ? CandidateIndices[0] : 0;
+			}
+		}
+		if (Tracks && SelectedTrackIndex < Tracks->Num() && (*Tracks)[SelectedTrackIndex].Meta.StreamDetails.Num())
 		{
-			st.ActivePlaylist = TimelineMediaAsset->GetExistingMediaPlaylistFromLoadRequest(LoadReq.Request);
+			const FActiveHLSPlaylist::FInternalTrackMetadata& Track((*Tracks)[SelectedTrackIndex]);
+
+			st.MetaID = Track.Meta.ID;
+			st.TrackIndex = SelectedTrackIndex;
+			st.StreamIndex = 0;
+			st.bIsSelected = true;
+
+			st.BufferSourceInfo = MakeSharedTS<FBufferSourceInfo>();
+			st.BufferSourceInfo->PeriodID = TimelineMediaAsset->GetAssetIdentifier();
+			st.BufferSourceInfo->PeriodAdaptationSetID = st.MetaID;
+			st.BufferSourceInfo->Kind = Track.Meta.Kind;
+			st.BufferSourceInfo->Language = Track.Meta.Language;
+			st.BufferSourceInfo->Codec = Util::GetBaseCodec(Track.Meta.HighestBandwidthCodec.GetCodecSpecifierRFC6381());
+			st.BufferSourceInfo->HardIndex = st.StreamIndex;
+
+			FStreamLoadRequest LoadReq;
+			if (TimelineMediaAsset->GetVariantPlaylist(LoadReq.Request, PlayerSessionServices, StreamType, TimelineMediaAsset->CurrentPathwayId, SelectedTrackIndex,0, SelectedTrackStream[0].TrackIndex,0).IsOK())
+			{
+				st.ActivePlaylist = TimelineMediaAsset->GetExistingMediaPlaylistFromLoadRequest(LoadReq.Request);
+				if (!st.ActivePlaylist.IsValid())
+				{
+					NewLoadReq.Add(LoadReq.Request);
+					++NumPending;
+				}
+				else if (!st.ActivePlaylist->ActivateIsReady())
+				{
+					++NumPending;
+				}
+			}
 		}
 	}
-
-	CurrentReadyState = IManifest::IPlayPeriod::EReadyState::IsReady;
+	if (NumPending)
+	{
+		TimelineMediaAsset->AddNewMediaPlaylistLoadRequests(NewLoadReq);
+		CurrentReadyState = IManifest::IPlayPeriod::EReadyState::Preparing;
+	}
+	else
+	{
+		CurrentReadyState = IManifest::IPlayPeriod::EReadyState::IsReady;
+	}
 }
 
 int64 FActiveHLSPlaylist::FPlayPeriod::GetDefaultStartingBitrate() const
@@ -2442,8 +2511,42 @@ FString FActiveHLSPlaylist::FPlayPeriod::GetSelectedAdaptationSetID(EStreamType
 
 IManifest::IPlayPeriod::ETrackChangeResult FActiveHLSPlaylist::FPlayPeriod::ChangeTrackStreamPreference(EStreamType InStreamType, const FStreamSelectionAttributes& InStreamAttributes)
 {
-	check(!"track changes not supported yet");
-	return IManifest::IPlayPeriod::ETrackChangeResult::NotChanged;
+	// Video cannot be switched seamlessly as this might also contain audio and subtitles.
+	if (InStreamType == EStreamType::Video)
+	{
+		return IManifest::IPlayPeriod::ETrackChangeResult::StartOver;
+	}
+	// On ongoing Live presentation without PDT has no information to locate the startover segment because
+	// there will no "previous" segment request to get the information from. See GetContinuationSegment()
+	if (TimelineMediaAsset->InitialPlaylistType == FPlaylistParserHLS::EPlaylistType::Live && !TimelineMediaAsset->bHasProgramDateTime && !TimelineMediaAsset->bInitialHasEndList)
+	{
+		return IManifest::IPlayPeriod::ETrackChangeResult::StartOver;
+	}
+
+	// Create a temporary period and prepare it for playback. This may result in media playlist load requests!
+	FPlayPeriod TempPeriod(PlayerSessionServices, TimelineMediaAsset);
+	TempPeriod.SetStreamPreferences(InStreamType, InStreamAttributes);
+	TempPeriod.PrepareForPlay();
+	const TArray<FActiveHLSPlaylist::FInternalTrackMetadata>* Tracks = nullptr;
+	const int32 stIdx = StreamTypeToArrayIndex(InStreamType);
+	if (stIdx == 0 && TimelineMediaAsset->VideoTracks.Num())
+	{
+		Tracks = &TimelineMediaAsset->VideoTracks;
+	}
+	else if (stIdx == 1 && TimelineMediaAsset->AudioTracks.Num())
+	{
+		Tracks = &TimelineMediaAsset->AudioTracks;
+	}
+	else if (stIdx == 2 && TimelineMediaAsset->SubtitleTracks.Num())
+	{
+		Tracks = &TimelineMediaAsset->SubtitleTracks;
+	}
+	// If either the stream we are leaving or the one we want to switch to is a variant we have to start over.
+	if (!Tracks || ((*Tracks)[SelectedTrackStream[stIdx].TrackIndex].bIsVariant || (*Tracks)[TempPeriod.SelectedTrackStream[stIdx].TrackIndex].bIsVariant))
+	{
+		return IManifest::IPlayPeriod::ETrackChangeResult::StartOver;
+	}
+	return IManifest::IPlayPeriod::ETrackChangeResult::NewPeriodNeeded;
 }
 
 TSharedPtrTS<ITimelineMediaAsset> FActiveHLSPlaylist::FPlayPeriod::GetMediaAsset() const
@@ -2791,8 +2894,9 @@ IManifest::FResult FActiveHLSPlaylist::FPlayPeriod::GetSegment(TSharedPtrTS<IStr
 	int32 SelectedTrackTypeIndex = StreamTypeToArrayIndex(InSegment->StreamType);
 	if (!SelectedTrackStream[SelectedTrackTypeIndex].bIsSelected)
 	{
-		// If the track was deselected then we should not actually be called.
-		return IManifest::FResult(IManifest::FResult::EType::NotFound).SetErrorDetail(FErrorDetail().SetMessage("Track got deselected"));
+		// The track may not be selected, which is ok and happens when switching tracks as start-over requests are made for all
+		// track types that are not the ones being switched.
+		return IManifest::FResult(IManifest::FResult::EType::NotFound);
 	}
 
 	FStreamLoadRequest LoadReq;
@@ -2853,6 +2957,7 @@ IManifest::FResult FActiveHLSPlaylist::FPlayPeriod::GetSegment(TSharedPtrTS<IStr
 	bool bSetNextExpectedTime = false;
 	if (LoadReq.Playlist == InSegment->HLS.Playlist)
 	{
+		check(InNextType != ENextSegType::StartOver);	// starting over has no current information on sequence index, so we must not get here.
 		SegParam.MediaSequenceIndex = InSegment->Segment.Number;
 		// Not actually needed, but set as a safe value.
 		SegParam.Start.Time.SetFromHNS(InSegment->Segment.Time + InSegment->Segment.Duration);
@@ -2877,6 +2982,12 @@ IManifest::FResult FActiveHLSPlaylist::FPlayPeriod::GetSegment(TSharedPtrTS<IStr
 				SegParam.Start.Time.SetFromHNS(InSegment->Segment.Time + InSegment->Segment.Duration*3/4);
 				SegParam.Start.Time += InSegment->AdditionalAdjustmentTime;
 			}
+			else if (InNextType == ENextSegType::StartOver)
+			{
+				SegParam.SearchType = IManifest::ESearchType::Before;
+				SegParam.Start.Time.SetFromHNS(InSegment->Segment.Time);
+				SegParam.Start.Time += InSegment->AdditionalAdjustmentTime;
+			}
 			else
 			{
 				SegParam.SearchType = IManifest::ESearchType::Same;
@@ -2886,6 +2997,7 @@ IManifest::FResult FActiveHLSPlaylist::FPlayPeriod::GetSegment(TSharedPtrTS<IStr
 		}
 		else
 		{
+			check(InNextType != ENextSegType::StartOver);	// starting over has no current information on sequence index, so we must not get here.
 			// Otherwise - in a Live presentation that has no #EXT-X-PROGRAM-DATE-TIME values - things are a bit more difficult.
 			SegParam.DiscontinuityIndex = InSegment->HLS.DiscontinuitySequence;
 			// We got back one segment on purpose and risk loading a segment we already got.
@@ -3051,8 +3163,11 @@ IManifest::FResult FActiveHLSPlaylist::FPlayPeriod::GetSegment(TSharedPtrTS<IStr
 
 IManifest::FResult FActiveHLSPlaylist::FPlayPeriod::GetContinuationSegment(TSharedPtrTS<IStreamSegment>& OutSegment, EStreamType InStreamType, const FPlayerSequenceState& InSequenceState, const FPlayStartPosition& InStartPosition, ESearchType InSearchType)
 {
-	check(!"track changes not supported yet");
-	return IManifest::FResult(IManifest::FResult::EType::NotFound);
+	FStreamSegmentRequestCommon DummyReq;
+	DummyReq.StreamType = InStreamType;
+	DummyReq.Segment.Time = InStartPosition.Time.GetAsHNS();
+	DummyReq.TimestampSequenceIndex = InSequenceState.GetSequenceIndex();
+	return GetSegment(OutSegment, &DummyReq, InStartPosition.Options, ENextSegType::StartOver);
 }
 
 IManifest::FResult FActiveHLSPlaylist::FPlayPeriod::GetLoopingSegment(TSharedPtrTS<IStreamSegment>& OutSegment, const FPlayerSequenceState& InSequenceState, const FPlayStartPosition& InStartPosition, ESearchType InSearchType)
diff --git a/Engine/Plugins/Media/ElectraPlayer/Source/ElectraPlayerRuntime/Private/Runtime/Player/Manifest.h b/Engine/Plugins/Media/ElectraPlayer/Source/ElectraPlayerRuntime/Private/Runtime/Player/Manifest.h
--- a/Engine/Plugins/Media/ElectraPlayer/Source/ElectraPlayerRuntime/Private/Runtime/Player/Manifest.h
+++ b/Engine/Plugins/Media/ElectraPlayer/Source/ElectraPlayerRuntime/Private/Runtime/Player/Manifest.h
@@ -448,7 +448,8 @@ namespace Electra
 			{
 				Changed,
 				NotChanged,
-				NewPeriodNeeded
+				NewPeriodNeeded,
+				StartOver
 			};
 
 			/**
diff --git a/Engine/Plugins/Media/ElectraPlayer/Source/ElectraPlayerRuntime/Private/Runtime/StreamTypes.h b/Engine/Plugins/Media/ElectraPlayer/Source/ElectraPlayerRuntime/Private/Runtime/StreamTypes.h
--- a/Engine/Plugins/Media/ElectraPlayer/Source/ElectraPlayerRuntime/Private/Runtime/StreamTypes.h
+++ b/Engine/Plugins/Media/ElectraPlayer/Source/ElectraPlayerRuntime/Private/Runtime/StreamTypes.h
@@ -47,6 +47,12 @@ namespace Electra
 		return StreamType == EStreamType::Audio ? 1 : StreamType == EStreamType::Subtitle ? 2 : 0;
 	}
 
+	static inline constexpr EStreamType StreamArrayIndexToType(int32 InIndex)
+	{
+		return InIndex == 0 ? EStreamType::Video :
+			   InIndex == 1 ? EStreamType::Audio :
+			   InIndex == 2 ? EStreamType::Subtitle : EStreamType::Unsupported;
+	}
 
 	class FStreamCodecInformation
 	{
@@ -813,6 +819,11 @@ namespace Electra
 
 		virtual ~FStreamSelectionAttributes() = default;
 
+		virtual bool IsSet() const
+		{
+			return Kind.IsSet() || Language_ISO639.IsSet() || Codec.IsSet() || OverrideIndex.IsSet();
+		}
+
 		virtual bool IsCompatibleWith(const FStreamSelectionAttributes& Other)
 		{
 			if (OverrideIndex.IsSet() && Other.OverrideIndex.IsSet() && OverrideIndex.GetValue() >= 0 && Other.OverrideIndex.GetValue() >= 0 && OverrideIndex.GetValue() != Other.OverrideIndex.GetValue())
