[ElectraPlayer] Adding support parsing time code track in quick time container. Doesn't support variable frame rate.

The first sample of the timecode track is used to determine the start timecode for the track. This is then either copied in the meta data of referencing video tracks or globally if a timecode track is not referenced by any video tracks.

Timecode information is not parsed by default. The new player option that must be specified to request parsing timecode info is "ParseTimecodeInfo".

The new player's timecode info metadata keys:
"StartTimecodeValue" -> String in SMPTE HH::MM::SS::FF format (supports DF).
"StartTimecodeFrameRate" -> String that can be parsed by ParseFrameRate.

diff --git a/Engine/Plugins/Media/ElectraPlayer/Source/ElectraPlayerPlugin/Private/ElectraPlayerPlugin.cpp b/Engine/Plugins/Media/ElectraPlayer/Source/ElectraPlayerPlugin/Private/ElectraPlayerPlugin.cpp
--- a/Engine/Plugins/Media/ElectraPlayer/Source/ElectraPlayerPlugin/Private/ElectraPlayerPlugin.cpp
+++ b/Engine/Plugins/Media/ElectraPlayer/Source/ElectraPlayerPlugin/Private/ElectraPlayerPlugin.cpp
@@ -588,6 +588,15 @@ bool FElectraPlayerPlugin::Open(const FString& Url, const IMediaOptions* Options
 		}
 	}
 
+	if (InPlayerOptions)
+	{
+		if (InPlayerOptions->InternalCustomOptions.Find(MediaPlayerOptionValues::ParseTimecodeInfo()))
+		{
+			static const FName OptionKeyParseTimecodeInfo(TEXT("parse_timecode_info"));
+			PlayerOptions.Set(OptionKeyParseTimecodeInfo, FVariantValue());
+		}
+	}
+
 	// Check for one-time initialization options that can't be changed during playback.
 	int64 InitialStreamBitrate = Options->GetMediaOption(TEXT("ElectraInitialBitrate"), (int64)-1);
 	if (InitialStreamBitrate > 0)
@@ -756,6 +765,11 @@ void FElectraPlayerPlugin::TickInput(FTimespan DeltaTime, FTimespan Timecode)
 	Player->Tick(DeltaTime, Timecode);
 }
 
+FVariant FElectraPlayerPlugin::GetMediaInfo(FName InInfoName) const
+{
+	return Player.IsValid() ?  Player->GetMediaInfo(InInfoName).ToFVariant() : FVariant();
+}
+
 //-----------------------------------------------------------------------------
 /**
 	Returns the current metadata, if any.
diff --git a/Engine/Plugins/Media/ElectraPlayer/Source/ElectraPlayerPlugin/Private/ElectraPlayerPlugin.h b/Engine/Plugins/Media/ElectraPlayer/Source/ElectraPlayerPlugin/Private/ElectraPlayerPlugin.h
--- a/Engine/Plugins/Media/ElectraPlayer/Source/ElectraPlayerPlugin/Private/ElectraPlayerPlugin.h
+++ b/Engine/Plugins/Media/ElectraPlayer/Source/ElectraPlayerPlugin/Private/ElectraPlayerPlugin.h
@@ -99,6 +99,8 @@ public:
 	void TickInput(FTimespan DeltaTime, FTimespan Timecode) override;
 	void SetLastAudioRenderedSampleTime(FTimespan SampleTime) override;
 
+	virtual FVariant GetMediaInfo(FName InInfoName) const override;
+
 	TSharedPtr<TMap<FString, TArray<TUniquePtr<IMediaMetadataItem>>>, ESPMode::ThreadSafe> GetMediaMetadata() const override;
 
 	bool GetPlayerFeatureFlag(EFeatureFlag flag) const override;
diff --git a/Engine/Plugins/Media/ElectraPlayer/Source/ElectraPlayerRuntime/Private/ElectraPlayer.cpp b/Engine/Plugins/Media/ElectraPlayer/Source/ElectraPlayerRuntime/Private/ElectraPlayer.cpp
--- a/Engine/Plugins/Media/ElectraPlayer/Source/ElectraPlayerRuntime/Private/ElectraPlayer.cpp
+++ b/Engine/Plugins/Media/ElectraPlayer/Source/ElectraPlayerRuntime/Private/ElectraPlayer.cpp
@@ -1400,6 +1400,15 @@ TRange<FTimespan> FElectraPlayer::GetPlaybackRange(ETimeRangeType InRangeToGet)
 	return Range;
 }
 
+Electra::FVariantValue FElectraPlayer::GetMediaInfo(FName InInfoName) const
+{
+	const TSharedPtr<FInternalPlayerImpl, ESPMode::ThreadSafe> LockedPlayer = CurrentPlayer;
+	if (LockedPlayer.IsValid() && LockedPlayer->AdaptivePlayer.IsValid())
+	{
+		return LockedPlayer->AdaptivePlayer->GetMediaInfo(InInfoName);
+	}
+	return Electra::FVariantValue();
+}
 
 TSharedPtr<TMap<FString, TArray<TSharedPtr<Electra::IMediaStreamMetadata::IItem, ESPMode::ThreadSafe>>>, ESPMode::ThreadSafe> FElectraPlayer::GetMediaMetadata() const
 {
diff --git a/Engine/Plugins/Media/ElectraPlayer/Source/ElectraPlayerRuntime/Private/ElectraPlayer.h b/Engine/Plugins/Media/ElectraPlayer/Source/ElectraPlayerRuntime/Private/ElectraPlayer.h
--- a/Engine/Plugins/Media/ElectraPlayer/Source/ElectraPlayerRuntime/Private/ElectraPlayer.h
+++ b/Engine/Plugins/Media/ElectraPlayer/Source/ElectraPlayerRuntime/Private/ElectraPlayer.h
@@ -137,6 +137,7 @@ public:
 	int32 GetNumVideoStreams(int32 TrackIndex) const override;
 	bool GetVideoStreamFormat(FVideoStreamFormat& OutFormat, int32 InTrackIndex, int32 InStreamIndex) const override;
 	bool GetActiveVideoStreamFormat(FVideoStreamFormat& OutFormat) const override;
+	virtual Electra::FVariantValue GetMediaInfo(FName InInfoName) const override;
 	TSharedPtr<TMap<FString, TArray<TSharedPtr<Electra::IMediaStreamMetadata::IItem, ESPMode::ThreadSafe>>>, ESPMode::ThreadSafe> GetMediaMetadata() const override;
 
 	void NotifyOfOptionChange() override;
diff --git a/Engine/Plugins/Media/ElectraPlayer/Source/ElectraPlayerRuntime/Private/Runtime/Demuxer/ParserISO14496-12.cpp b/Engine/Plugins/Media/ElectraPlayer/Source/ElectraPlayerRuntime/Private/Runtime/Demuxer/ParserISO14496-12.cpp
--- a/Engine/Plugins/Media/ElectraPlayer/Source/ElectraPlayerRuntime/Private/Runtime/Demuxer/ParserISO14496-12.cpp
+++ b/Engine/Plugins/Media/ElectraPlayer/Source/ElectraPlayerRuntime/Private/Runtime/Demuxer/ParserISO14496-12.cpp
@@ -3,17 +3,23 @@
 #include "Misc/Build.h"
 #include "PlayerCore.h"
 #include "Containers/ChunkedArray.h"
+#include "Misc/FrameRate.h"
+#include "Misc/Timecode.h"
 
 #include "Demuxer/ParserISO14496-12.h"
 
 #include "Utilities/Utilities.h"
+#include "Utilities/UtilsMP4.h"
 #include "Utilities/UtilsMPEG.h"
 #include "Utils/MPEG/ElectraUtilsMPEGAudio.h"
 #include "Utils/MPEG/ElectraUtilsMPEGVideo.h"
 #include "Utilities/ISO639-Map.h"
 #include "InfoLog.h"
+#include "Player/AdaptivePlayerMetadataKeynames.h"
+#include "Player/AdaptivePlayerOptionKeynames.h"
 #include "Player/PlayerSessionServices.h"
 #include "Player/PlayerStreamFilter.h"
+#include "Player/PlaylistReader.h"
 #include "Decoder/SubtitleDecoder.h"
 
 #include "IElectraCodecFactoryModule.h"
@@ -401,6 +407,11 @@ namespace Electra
 				PlayerSession->PostLog(Facility::EFacility::MP4Parser, InLevel, InMessage);
 			}
 		}
+
+		bool ShouldParseTimecodeInfo() const
+		{
+			return bShouldParseTimecodeInfo;
+		}
 	private:
 		FMP4ParseInfo(const FMP4ParseInfo&) = delete;
 
@@ -481,6 +492,7 @@ namespace Electra
 
 		IPlayerSessionServices* PlayerSession = nullptr;
 		int32 NumTotalBoxesParsed = 0;
+		bool bShouldParseTimecodeInfo = false;
 	};
 
 
@@ -520,6 +532,18 @@ namespace Electra
 			return BoxType;
 		}
 
+#if MEDIA_DEBUG_HAS_BOX_NAMES
+		static FString GetBoxNameString(const char InBoxName[4])
+		{
+			return FString::Printf(TEXT("%c%c%c%c"), TCHAR(InBoxName[0]?InBoxName[0]:'?'), TCHAR(InBoxName[1]?InBoxName[1]:'?'), TCHAR(InBoxName[2]?InBoxName[2]:'?'), TCHAR(InBoxName[3]?InBoxName[3]:'?'));
+		}
+		
+		FString GetBoxName() const
+		{
+			return GetBoxNameString(DbgBoxName);
+		}
+#endif
+
 		bool IsLeafBox() const
 		{
 			return bIsLeafBox;
@@ -791,7 +815,13 @@ namespace Electra
 		static const IParserISO14496_12::FBoxType kSample_wvtt = MAKE_BOX_ATOM('w', 'v', 't', 't');
 		static const IParserISO14496_12::FBoxType kSample_enca = MAKE_BOX_ATOM('e', 'n', 'c', 'a');
 		static const IParserISO14496_12::FBoxType kSample_encv = MAKE_BOX_ATOM('e', 'n', 'c', 'v');
+		static const IParserISO14496_12::FBoxType kSample_tmcd = MAKE_BOX_ATOM('t', 'm', 'c', 'd');
 
+		/**
+		 * Track Reference types - https://developer.apple.com/documentation/quicktime-file-format/track_reference_type_atom
+		 */
+		static const IParserISO14496_12::FBoxType kTrackRef_tmcd = MAKE_BOX_ATOM('t', 'm', 'c', 'd');
+		
 		/**
 		 * Sample grouping types
 		 */
@@ -2230,6 +2260,79 @@ namespace Electra
 		uint16		SampleSize = 0;
 	};
 
+	/**
+	 * Timecode sample entry
+	 * Reference: https://developer.apple.com/documentation/quicktime-file-format/timecode_sample_description
+	 */
+	class FMP4BoxTimecodeSampleEntry : public FMP4BoxSampleEntry
+	{
+	public:
+		FMP4BoxTimecodeSampleEntry(IParserISO14496_12::FBoxType InBoxType, int64 InBoxSize, int64 InStartOffset, int64 InDataOffset, bool bInIsLeafBox)
+			: FMP4BoxSampleEntry(InBoxType, InBoxSize, InStartOffset, InDataOffset, bInIsLeafBox)
+		{
+		}
+
+		virtual ~FMP4BoxTimecodeSampleEntry() override = default;
+
+		struct FFlags
+		{
+			/** Indicates whether the timecode is drop frame. Set it to 1 if the timecode is drop frame. */
+			static constexpr uint32 DropFrame = 0x0001;
+			/** Indicates whether the timecode wraps after 24 hours. Set it to 1 if the timecode wraps. */
+			static constexpr uint32 Max24Hour = 0x0002;
+			/** Indicates whether negative time values are allowed. Set it to 1 if the timecode supports negative values. */
+			static constexpr uint32 AllowNegativeTimes = 0x0004;
+			/** Indicates whether the time value corresponds to a tape counter value. Set it to 1 if the timecode values are tape counter values.*/
+			static constexpr uint32 Counter = 0x0008;
+		};
+
+		uint32 GetFlags() const { return Flags; }
+		uint32 GetTimeScale() const { return TimeScale; }
+		uint32 GetFrameDuration() const { return FrameDuration; }
+		uint32 GetNumberOfFrames() const { return NumberOfFrames; }
+		
+	private:
+		FMP4BoxTimecodeSampleEntry() = delete;
+		FMP4BoxTimecodeSampleEntry(const FMP4BoxTimecodeSampleEntry&) = delete;
+
+	protected:
+		virtual UEMediaError ReadAndParseAttributes(FMP4ParseInfo* ParseInfo) override
+		{
+			UEMediaError Error = UEMEDIA_ERROR_OK;
+			RETURN_IF_ERROR(FMP4BoxSampleEntry::ReadAndParseAttributes(ParseInfo));
+
+			uint32 Reserved32 = 0;
+			uint8 Reserved8 = 0;
+
+			RETURN_IF_ERROR(ParseInfo->Reader()->Read(Reserved32));
+			RETURN_IF_ERROR(ParseInfo->Reader()->Read(Flags));
+			RETURN_IF_ERROR(ParseInfo->Reader()->Read(TimeScale));
+			RETURN_IF_ERROR(ParseInfo->Reader()->Read(FrameDuration));
+			RETURN_IF_ERROR(ParseInfo->Reader()->Read(NumberOfFrames));
+			RETURN_IF_ERROR(ParseInfo->Reader()->Read(Reserved8));
+			
+			// Next would be (optional) source reference atom.
+			const int64 NowAt = ParseInfo->Reader()->GetCurrentReadOffset();
+			int32 BytesRemaining = static_cast<int32>(BoxSize - (NowAt - StartOffset));
+			if (BytesRemaining > 0)
+			{
+				bIsLeafBox = false;
+				RETURN_IF_ERROR(ParseInfo->ReadAndParseNextBox(this, NowAt + BytesRemaining));
+				BytesRemaining = static_cast<int32>(BoxSize - (ParseInfo->Reader()->GetCurrentReadOffset() - StartOffset));
+				check(BytesRemaining >= 0);
+				if (BytesRemaining > 0)
+				{
+					RETURN_IF_ERROR(ParseInfo->Reader()->ReadBytes(nullptr, BytesRemaining));
+				}
+			}
+			return Error;
+		}
+		
+		uint32 Flags = 0;
+		uint32 TimeScale = 0;
+		uint32 FrameDuration = 0;
+		uint8 NumberOfFrames = 0;
+	};
 
 	/**
 	 * Plain Text sample entry (12.5.3 - Sample Entry)
@@ -3271,6 +3374,7 @@ namespace Electra
 			Video,
 			Audio,
 			Subtitles,
+			Timecode,
 			Ignored
 		};
 
@@ -3322,6 +3426,10 @@ namespace Electra
 						SampleType = ESampleType::Subtitles;
 						break;
 					case FMP4Box::kBox_gmhd:
+						// Remark: Could parse gmhd box, check if it has 'tmcd' box to make sure it is a timecode.
+						// For now assume it is a timecode, but will expect 'tmcd' sample description below (and ignore if not).
+						SampleType = ParseInfo->ShouldParseTimecodeInfo() ? ESampleType::Timecode : ESampleType::Ignored;
+						break;
 					default:
 						SampleType = ESampleType::Ignored;
 						break;
@@ -3398,6 +3506,17 @@ namespace Electra
 							NextBox = new FMP4BoxIgnored(ChildBoxType, ChildBoxSize, BoxStartOffset, BoxDataOffset, true);
 							break;
 						*/
+				case ESampleType::Timecode:
+						if (ChildBoxType == FMP4Box::kSample_tmcd)
+						{
+							NextBox = new FMP4BoxTimecodeSampleEntry(ChildBoxType, ChildBoxSize, BoxStartOffset, BoxDataOffset, true);
+						}
+						else
+						{
+							SampleType = ESampleType::Ignored; // Ignoring non-tmcd samples.
+							NextBox = new FMP4BoxIgnored(ChildBoxType, ChildBoxSize, BoxStartOffset, BoxDataOffset, true);	
+						}
+						break;
 					default:
 						NextBox = new FMP4BoxIgnored(ChildBoxType, ChildBoxSize, BoxStartOffset, BoxDataOffset, true);
 						break;
@@ -5028,7 +5147,106 @@ namespace Electra
 		TArray<FSampleGroupDescriptionEntry> Entries;
 	};
 
+	/**
+	 * Track reference entry (within an 'tref' box; Track Reference Box)
+	*  The box type is the reference type (See kTrackRef_...), along with a list of TrackIds for that relation type.
+	 * https://developer.apple.com/documentation/quicktime-file-format/track_reference_type_atom
+	 */
+	class FMP4BoxTrackReferenceEntry : public FMP4BoxBasic
+	{
+	public:
+		FMP4BoxTrackReferenceEntry(IParserISO14496_12::FBoxType InBoxType, int64 InBoxSize, int64 InStartOffset, int64 InDataOffset, bool bInIsLeafBox)
+			: FMP4BoxBasic(InBoxType, InBoxSize, InStartOffset, InDataOffset, bInIsLeafBox)
+		{
+		}
+
+		virtual ~FMP4BoxTrackReferenceEntry() override = default;
+
+		const TArray<int32>& GetTrackIds() const
+		{
+			return TrackIds;
+		}
+		
+	private:
+		FMP4BoxTrackReferenceEntry() = delete;
+		FMP4BoxTrackReferenceEntry(const FMP4BoxTrackReferenceEntry&) = delete;
+
+	protected:
+		virtual UEMediaError ReadAndParseAttributes(FMP4ParseInfo* ParseInfo) override
+		{
+			UEMediaError Error = UEMEDIA_ERROR_OK;
+			const int64 NowAt = ParseInfo->Reader()->GetCurrentReadOffset();
+			const int32 BytesRemaining = static_cast<int32>(BoxSize - (NowAt - StartOffset));
+			const int32 NumTrackIds = BytesRemaining/sizeof(int32);
+			TrackIds.AddZeroed(NumTrackIds);
+			RETURN_IF_ERROR(ParseInfo->Reader()->ReadBytes(TrackIds.GetData(), NumTrackIds*sizeof(int32)));
+			return Error;
+		}
 
+		TArray<int32> TrackIds;
+	};
+
+	/**
+	 * 'tref' box. Track Reference Atom
+	 *  https://developer.apple.com/documentation/quicktime-file-format/track_reference_atom
+	 *
+	 *  Contains a list of relational track references per type. 
+	 */
+	class FMP4BoxTREF : public FMP4BoxBasic
+	{
+	public:
+		FMP4BoxTREF(IParserISO14496_12::FBoxType InBoxType, int64 InBoxSize, int64 InStartOffset, int64 InDataOffset, bool bInIsLeafBox)
+			: FMP4BoxBasic(InBoxType, InBoxSize, InStartOffset, InDataOffset, bInIsLeafBox)
+		{
+		}
+
+		virtual ~FMP4BoxTREF() override = default;
+
+		const TArray<int32>& GetTrackIdsForType(IParserISO14496_12::FBoxType InReferenceType) const
+		{
+			for (const FMP4Box* ChildBox : ChildBoxes)
+			{
+				if (ChildBox->GetType() == InReferenceType)
+				{
+					return static_cast<const FMP4BoxTrackReferenceEntry*>(ChildBox)->GetTrackIds();
+				}
+			}
+			static const TArray<int32> Empty;
+			return Empty;
+		}
+		
+	private:
+		FMP4BoxTREF() = delete;
+		FMP4BoxTREF(const FMP4BoxTREF&) = delete;
+	
+	protected:
+		virtual UEMediaError ReadAndParseAttributes(FMP4ParseInfo* ParseInfo) override
+		{
+			UEMediaError Error = UEMEDIA_ERROR_OK;
+			
+			// Read the list of TrackReference atoms
+			int32 BytesRemaining = static_cast<int32>(BoxSize - (ParseInfo->Reader()->GetCurrentReadOffset() - StartOffset));
+			while (BytesRemaining > 0)
+			{
+				bIsLeafBox = false;
+
+				IParserISO14496_12::FBoxType	ChildBoxType;
+				int64							ChildBoxSize;
+				uint8							ChildBoxUUID[16];
+
+				const int64 BoxStartOffset = ParseInfo->Reader()->GetCurrentReadOffset();
+				RETURN_IF_ERROR(ParseInfo->ReadBoxTypeAndSize(ChildBoxType, ChildBoxSize, ChildBoxUUID));
+				const int64 BoxDataOffset = ParseInfo->Reader()->GetCurrentReadOffset();
+				FMP4Box* NextBox = new FMP4BoxTrackReferenceEntry(ChildBoxType, ChildBoxSize, BoxStartOffset, BoxDataOffset, true);
+				AddChildBox(NextBox);
+				RETURN_IF_ERROR(ParseInfo->ReadAndParseNextBox(NextBox));
+
+				// Where are we at now?
+				BytesRemaining = static_cast<int32>(BoxSize - (ParseInfo->Reader()->GetCurrentReadOffset() - StartOffset));
+			}
+			return Error;
+		}
+	};
 
 	/**
 	 * 'udta' box. ISO/IEC 14496-12:2015 - 8.10.1 User Data Box
@@ -5251,6 +5469,7 @@ namespace Electra
 	{
 		PlayerSession = InPlayerSession;
 		OptionalMoovParseInfo = InOptionalMoovParseInfo;
+		bShouldParseTimecodeInfo = InPlayerSession ? InPlayerSession->HaveOptionValue(Electra::OptionKeyParseTimecodeInfo) : false;
 
 		// New parse or continuing a previous?
 		if (RootBox == nullptr)
@@ -5313,7 +5532,7 @@ namespace Electra
 #if MEDIA_DEBUG_HAS_BOX_NAMES
 				char boxName[4];
 				*((uint32*)&boxName) = MEDIA_TO_BIG_ENDIAN(BoxType);
-				//UE_LOG(LogElectraMP4Parser, Log, TEXT("Next box: '%c%c%c%c' %lld B @ %lld"), boxName[0]?boxName[0]:'?', boxName[1]?boxName[1]:'?', boxName[2]?boxName[2]:'?', boxName[3]?boxName[3]:'?', (long long int)BoxSize, (long long int)BoxStartOffset);
+				//UE_LOG(LogElectraMP4Parser, Log, TEXT("Next box: '%s' %lld B @ %lld (Parent: '%s')"), *FMP4Box::GetBoxNameString(boxName), (long long int)BoxSize, (long long int)BoxStartOffset, ParentBox ? *ParentBox->GetBoxName() : TEXT("****"));
 #endif
 
 				int64 BoxDataOffset = BoxReader->GetCurrentReadOffset();
@@ -5520,8 +5739,10 @@ namespace Electra
 						case FMP4Box::kBox_sgpd:
 							NextBox = new FMP4BoxSGPD(BoxType, BoxSize, BoxStartOffset, BoxDataOffset, true);
 							break;
-						// Boxes we ignore for now.
 						case FMP4Box::kBox_tref:
+							NextBox = new FMP4BoxTREF(BoxType, BoxSize, BoxStartOffset, BoxDataOffset, true);
+							break;
+						// Boxes we ignore for now.
 						case FMP4Box::kBox_colr:
 						case FMP4Box::kBox_fiel:
 						case FMP4Box::kBox_trep:
@@ -5599,7 +5820,7 @@ namespace Electra
 
 						default:
 #if MEDIA_DEBUG_HAS_BOX_NAMES
-//							UE_LOG(LogElectraMP4Parser, Log, TEXT("Ignoring mp4 box: '%c%c%c%c' %lld B @ %lld"), boxName[0]?boxName[0]:'?', boxName[1]?boxName[1]:'?', boxName[2]?boxName[2]:'?', boxName[3]?boxName[3]:'?', (long long int)BoxSize, (long long int)BoxStartOffset);
+//							UE_LOG(LogElectraMP4Parser, Log, TEXT("Ignoring mp4 box: '%s' %lld B @ %lld"), *FMP4Box::GetBoxNameString(boxName), (long long int)BoxSize, (long long int)BoxStartOffset);
 #endif
 							NextBox = new FMP4BoxIgnored(BoxType, BoxSize, BoxStartOffset, BoxDataOffset, true);
 							break;
@@ -5723,6 +5944,8 @@ namespace Electra
 
 		UEMediaError PrepareTracks(IPlayerSessionServices* PlayerSession, TSharedPtrTS<const IParserISO14496_12> OptionalMP4InitSegment) override;
 
+		virtual UEMediaError ResolveTimecodeTracks(IPlayerSessionServices* InPlayerSession, FCancellationCheckDelegate InCancellationCheckDelegate) override;
+
 		TMediaOptionalValue<FTimeFraction> GetMovieDuration() const override;
 		int32 GetNumberOfTracks() const override;
 		int32 GetNumberOfSegmentIndices() const override;
@@ -5770,7 +5993,7 @@ namespace Electra
 			int64 GetEmptyEditOffset() const override;
 			bool GetEncryptionInfo(ElectraCDM::FMediaCDMSampleInfo& OutSampleEncryptionInfo) const override;
 
-			UEMediaError StartAtFirstInteral();
+			UEMediaError StartAtFirstInternal();
 			void SetTrack(const FTrack* InTrack);
 			bool IsValid() const
 			{
@@ -5951,6 +6174,7 @@ namespace Electra
 			const FMP4BoxSTCO* STCOBox = nullptr;
 			const FMP4BoxSTSS* STSSBox = nullptr;
 			const FMP4BoxUDTA* UDTABox = nullptr;
+			const FMP4BoxTREF* TREFBox = nullptr;
 			// Groupings
 			TArray<const FMP4BoxSGPD*> SGPDBoxes;
 			TArray<const FMP4BoxSBGP*> SBGPBoxes;
@@ -5981,6 +6205,7 @@ namespace Electra
 			MPEG::FESDescriptor						CodecSpecificDataMP4A;
 			TArray<uint8>							CodecSpecificDataRAW;
 			FBitrateInfo							BitrateInfo;
+			TOptional<FTimecode>					StartTimecode;
 		};
 
 
@@ -6022,6 +6247,11 @@ namespace Electra
 				return Index >= 0 && Index < GetNumberOfTracks() ? TrackList[Index] : nullptr;
 			}
 
+			FTrack* GetTrackByIndexMutable(int32 Index) const
+			{
+				return Index >= 0 && Index < GetNumberOfTracks() ? TrackList[Index] : nullptr;
+			}
+
 			FTrack* GetTrackByID(uint32 TrackID) const
 			{
 				for(int32 i = 0; i < TrackList.Num(); ++i)
@@ -6074,6 +6304,8 @@ namespace Electra
 		UEMediaError ParseTX3GSampleType(IPlayerSessionServices* PlayerSession, FTrack* Track, const FMP4Box* SampleBox);
 		UEMediaError ParseWVTTSampleType(IPlayerSessionServices* PlayerSession, FTrack* Track, const FMP4Box* SampleBox);
 		UEMediaError ParseSTPPSampleType(IPlayerSessionServices* PlayerSession, FTrack* Track, const FMP4Box* SampleBox);
+		UEMediaError ParseTMCDSampleType(IPlayerSessionServices* InPlayerSession, FTrack* InTrack, const FMP4Box* InSampleBox);
+		UEMediaError ReadTMCDSample(IPlayerSessionServices* InPlayerSession, FTrack* InTrack, const FMP4Box* InSampleBox, FCancellationCheckDelegate InCheckCancellationDelegate);
 		UEMediaError ParseGenericSampleType(IPlayerSessionServices* PlayerSession, FTrack* Track, const FMP4BoxSTSD* STSDBox, const FMP4BoxFRMA* FRMABox);
 
 		FMP4ParseInfo* ParsedData;
@@ -6350,7 +6582,7 @@ namespace Electra
 	}
 
 
-	UEMediaError FParserISO14496_12::FTrackIterator::StartAtFirstInteral()
+	UEMediaError FParserISO14496_12::FTrackIterator::StartAtFirstInternal()
 	{
 		// NOTE: All boxes that are referenced here have been checked to exist so accessing them is safe.
 
@@ -6862,7 +7094,7 @@ namespace Electra
 
 	UEMediaError FParserISO14496_12::FTrackIterator::StartAtFirst(bool bNeedSyncSample)
 	{
-		UEMediaError err = StartAtFirstInteral();
+		UEMediaError err = StartAtFirstInternal();
 		if (err == UEMEDIA_ERROR_OK)
 		{
 			for(; err == UEMEDIA_ERROR_OK; err = Next())
@@ -6885,7 +7117,7 @@ namespace Electra
 		//       Otherwise for fragmented files it's tricky to get at the sample flags to find the sync samples unless we would assume there is only
 		//       a single sync sample at the start of the fragment.
 		// So for now we take the shortcut to just iterate.
-		UEMediaError err = StartAtFirstInteral();
+		UEMediaError err = StartAtFirstInternal();
 		if (err == UEMEDIA_ERROR_OK)
 		{
 			FTrackIterator Best;
@@ -7689,6 +7921,88 @@ namespace Electra
 		return UEMEDIA_ERROR_NOT_SUPPORTED;
 	}
 
+	UEMediaError FParserISO14496_12::ParseTMCDSampleType(IPlayerSessionServices* InPlayerSession, FTrack* InTrack, const FMP4Box* InSampleBox)
+	{
+		check(InSampleBox->GetType() == FMP4Box::kSample_tmcd);
+		const FMP4BoxTimecodeSampleEntry* TMCDBox = static_cast<const FMP4BoxTimecodeSampleEntry*>(InSampleBox);
+				
+		const bool bDropFrame = TMCDBox->GetFlags() & FMP4BoxTimecodeSampleEntry::FFlags::DropFrame ? true : false;
+		const FFrameRate FrameRate = bDropFrame ? FFrameRate(TMCDBox->GetTimeScale(), TMCDBox->GetFrameDuration()) : FFrameRate(TMCDBox->GetNumberOfFrames(), 1);
+
+		InTrack->CodecInformation.SetFrameRate(FTimeFraction(FrameRate.Numerator, FrameRate.Denominator));
+		InTrack->CodecInformation.SetCodec4CC(TMCDBox->GetType()); // Identify as 'tmcd' track, used to find again in second pass.
+		return UEMEDIA_ERROR_OK;
+	}
+
+	UEMediaError FParserISO14496_12::ReadTMCDSample(IPlayerSessionServices* InPlayerSession, FTrack* InTrack, const FMP4Box* InSampleBox, FCancellationCheckDelegate InCheckCancellationDelegate)
+	{
+		if (!InPlayerSession || !InPlayerSession->GetManifestReader() || !InSampleBox)
+		{
+			return UEMEDIA_ERROR_NOT_SUPPORTED;
+		}
+		
+		// We will need the stco, stsc and stsz boxes. We only need the first entry.
+		if (!InTrack || !InTrack->STCOBox || InTrack->STCOBox->GetNumberOfEntries() == 0
+			|| !InTrack->STSCBox || InTrack->STSCBox->GetNumberOfEntries() == 0
+			|| !InTrack->STSZBox)
+		{
+			return UEMEDIA_ERROR_NOT_SUPPORTED;
+		}
+
+		check(InSampleBox->GetType() == FMP4Box::kSample_tmcd);
+		const FMP4BoxTimecodeSampleEntry* TMCDBox = static_cast<const FMP4BoxTimecodeSampleEntry*>(InSampleBox);
+
+		// Obtain the first sample only -> Fetch the first chunk and read the first sample.
+		
+		const int64 ChunkOffset = InTrack->STCOBox->GetEntry(0);
+		//const uint32 SamplesPerChunk = InTrack->STSCBox->GetEntry(0).SamplesPerChunk;
+		const uint32 SampleSizeInBytes = InTrack->STSZBox->GetSampleSize(0);
+		// Note: since we need only the first sample, load only "sample size" to get the first one.
+		const uint32 ChunkSize = SampleSizeInBytes; //SamplesPerChunk * SampleSizeInBytes;
+		
+		const FString ManifestURL = InPlayerSession->GetManifestReader()->GetURL();
+		
+		UtilsMP4::FMP4ChunkLoader ChunkLoader;
+
+		// Propagate the worker thread's abort flag.
+		const UtilsMP4::FMP4ChunkLoader::FCancellationCheckDelegate CancellationCheckDelegate
+			= UtilsMP4::FMP4ChunkLoader::FCancellationCheckDelegate::CreateLambda([InCheckCancellationDelegate]()
+			{
+				return InCheckCancellationDelegate.Execute();
+			});
+
+		if (const TSharedPtrTS<FWaitableBuffer> ChunkData = ChunkLoader.LoadChunk(ChunkOffset, ChunkSize, InPlayerSession->GetHTTPManager(), InPlayerSession->GetHTTPResponseCache(), ManifestURL, CancellationCheckDelegate))
+		{
+			if (ChunkData.IsValid() && ChunkData->Num() >= SampleSizeInBytes)
+			{
+				const uint32* Data = reinterpret_cast<const uint32*>(ChunkData->GetLinearReadData());
+
+				// Note: ignoring counter flag. Even if set to 0, which would mean QT int32 format (HHMMSSFF),
+				// the timecode is still a frame number format (as if Counter set to 1).
+				const FFrameNumber FrameNumber(static_cast<int32>(MEDIA_FROM_BIG_ENDIAN(Data[0])));
+				
+				const bool bDropFrame = TMCDBox->GetFlags() & FMP4BoxTimecodeSampleEntry::FFlags::DropFrame ? true : false;
+				const bool bRollover = TMCDBox->GetFlags() & FMP4BoxTimecodeSampleEntry::FFlags::Max24Hour ? true : false;
+
+				const FFrameRate FrameRate = bDropFrame ? FFrameRate(TMCDBox->GetTimeScale(), TMCDBox->GetFrameDuration()) : FFrameRate(TMCDBox->GetNumberOfFrames(), 1);
+
+				// Need to convert to seconds first for rollover implementation below.
+				const double TimecodeInSeconds = FrameRate.AsSeconds(FFrameTime(FrameNumber));
+
+				// Convert to time code (apply roll over, etc).
+				InTrack->StartTimecode = FTimecode(TimecodeInSeconds, FrameRate, bDropFrame, bRollover);
+
+				// Expose the timecode in the track's metadata.
+				const FVariantValue StartTimecodeValue(InTrack->StartTimecode->ToString());
+				const FVariantValue StartTimecodeFrameRate(FrameRate.ToPrettyText().ToString());
+				InTrack->CodecInformation.GetExtras().Set(MetaDataKeyStartTimecodeValue, StartTimecodeValue);
+				InTrack->CodecInformation.GetExtras().Set(MetaDataKeyStartTimecodeFrameRate, StartTimecodeFrameRate);
+				return UEMEDIA_ERROR_OK;
+			}
+		}
+		return UEMEDIA_ERROR_NOT_SUPPORTED;	
+	}
+
 	UEMediaError FParserISO14496_12::ParseGenericSampleType(IPlayerSessionServices* PlayerSession, FTrack* Track, const FMP4BoxSTSD* STSDBox, const FMP4BoxFRMA* FRMABox)
 	{
 		if (Track && STSDBox)
@@ -7923,6 +8237,7 @@ namespace Electra
 							return UEMEDIA_ERROR_FORMAT_ERROR;
 						}
 						Track->UDTABox = static_cast<const FMP4BoxUDTA*>(Box->GetBoxPath(FMP4Box::kBox_udta));
+						Track->TREFBox = static_cast<const FMP4BoxTREF*>(Box->GetBoxPath(FMP4Box::kBox_tref));
 						// TODO: Check the dinf->dref->url box to not reference any external media!
 
 						// Get the composition time offset and the empty edit from the edit list. Those are necessary to correct the PTS values from
@@ -8247,6 +8562,19 @@ namespace Electra
 								}
 								break;
 							}
+							case FMP4Box::kSample_tmcd:
+							{
+								UEMediaError Error = ParseTMCDSampleType(PlayerSession, Track.Get(), STSDFirstChildBox);
+								if (Error == UEMEDIA_ERROR_NOT_SUPPORTED)
+								{
+									bIsSupported = false;
+								}
+								else if (Error != UEMEDIA_ERROR_OK)
+								{
+									return Error;
+								}
+								break;
+							}
 							default:
 							{
 								UEMediaError Error = ParseGenericSampleType(PlayerSession, Track.Get(), Track->STSDBox, nullptr);
@@ -8289,7 +8617,6 @@ namespace Electra
 			}
 		}
 
-
 		// Fragmented?
 		const FMP4Box* MOOFBox = ParsedData->GetCurrentMoofBox();
 		if (MOOFBox)
@@ -8437,6 +8764,79 @@ namespace Electra
 		return MOOVBox || MOOFBox || ParsedData->GetCurrentSidxBox() ? UEMEDIA_ERROR_OK : UEMEDIA_ERROR_FORMAT_ERROR;
 	}
 
+	UEMediaError FParserISO14496_12::ResolveTimecodeTracks(IPlayerSessionServices* InPlayerSession, FCancellationCheckDelegate InCancellationCheckDelegate)
+	{
+		if (!ParsedData->ShouldParseTimecodeInfo())
+		{
+			return UEMEDIA_ERROR_OK;
+		}
+		
+		// Read timecode sample data for timecode tracks.
+		for(int32 TrackIndex = 0; TrackIndex < ParsedTrackInfo->GetNumberOfTracks(); ++TrackIndex)
+		{
+			FTrack* Track = ParsedTrackInfo->GetTrackByIndexMutable(TrackIndex);
+			if (Track && Track->CodecInformation.GetCodec4CC() == FMP4Box::kSample_tmcd)
+			{
+				ReadTMCDSample(InPlayerSession, Track, Track->STSDBox->FindBox(FMP4Box::kSample_tmcd), InCancellationCheckDelegate);
+			}
+		}
+			
+		// Copy timecode meta data from tmcd tracks to the related video tracks.
+		for(int32 TrackIndex = 0; TrackIndex < ParsedTrackInfo->GetNumberOfTracks(); ++TrackIndex)
+		{
+			FTrack* Track = ParsedTrackInfo->GetTrackByIndexMutable(TrackIndex);
+			if (Track && Track->CodecInformation.IsVideoCodec() && Track->TREFBox)
+			{
+				const TArray<int32>& TimecodeTrackIds = Track->TREFBox->GetTrackIdsForType(FMP4Box::kTrackRef_tmcd);
+				if (!TimecodeTrackIds.IsEmpty())
+				{
+					if (const FTrack* TimecodeTrack = ParsedTrackInfo->GetTrackByID(TimecodeTrackIds[0]))
+					{
+						if (TimecodeTrack->StartTimecode.IsSet())
+						{
+							Track->StartTimecode = TimecodeTrack->StartTimecode;
+							FParamDict& TrackExtras = Track->CodecInformation.GetExtras();
+							TrackExtras.SetValueFrom(MetaDataKeyStartTimecodeValue, TimecodeTrack->CodecInformation.GetExtras());
+							TrackExtras.SetValueFrom(MetaDataKeyStartTimecodeFrameRate, TimecodeTrack->CodecInformation.GetExtras());
+						}
+					}
+				}
+			}
+		}
+
+		// Look for a tmcd track not referenced by any video track and export it globally.
+		for(int32 TrackIndex = 0; TrackIndex < ParsedTrackInfo->GetNumberOfTracks(); ++TrackIndex)
+		{
+			const FTrack* Track = ParsedTrackInfo->GetTrackByIndex(TrackIndex);
+			if (Track && Track->CodecInformation.GetCodec4CC() == FMP4Box::kSample_tmcd && Track->StartTimecode.IsSet())
+			{
+				auto IsTimecodeTrackReferenced = [](const FTrackInfo* InTrackInfo, int32 InTimecodeTrackId)
+				{
+					for(int32 TrackIndex = 0; TrackIndex < InTrackInfo->GetNumberOfTracks(); ++TrackIndex)
+					{
+						const FTrack* Track = InTrackInfo->GetTrackByIndex(TrackIndex);
+						if (Track && Track->CodecInformation.IsVideoCodec() && Track->TREFBox)
+						{
+							const TArray<int32>& TimecodeTrackIds = Track->TREFBox->GetTrackIdsForType(FMP4Box::kTrackRef_tmcd);
+							if (TimecodeTrackIds.Contains(InTimecodeTrackId))
+							{
+								return true;
+							}
+						}
+					}
+					return false;
+				};
+
+				if (!IsTimecodeTrackReferenced(ParsedTrackInfo, Track->GetID()))
+				{
+					InPlayerSession->GetMutableOptions().SetValueFrom(MetaDataKeyStartTimecodeValue, Track->CodecInformation.GetExtras());
+					InPlayerSession->GetMutableOptions().SetValueFrom(MetaDataKeyStartTimecodeFrameRate, Track->CodecInformation.GetExtras());
+					break;
+				}
+			}
+		}
+		return UEMEDIA_ERROR_OK;
+	}
 
 	TMediaOptionalValue<FTimeFraction> FParserISO14496_12::GetMovieDuration() const
 	{
@@ -8606,7 +9006,7 @@ namespace Electra
 			FTrackIterator* TrkIt = static_cast<FTrackIterator*>(Track->CreateIterator());
 			TSharedPtrTS<ITrackIterator> SafeTrkIt(TrkIt);
 			ti->TrackIterators.Add(SafeTrkIt);
-			UEMediaError err = TrkIt->StartAtFirstInteral();
+			UEMediaError err = TrkIt->StartAtFirstInternal();
 			if (err == UEMEDIA_ERROR_OK)
 			{
 				while(TrkIt->GetSampleFileOffset() < InFromFilePos)
@@ -8650,7 +9050,7 @@ namespace Electra
 			FTrackIterator* TrkIt = static_cast<FTrackIterator*>(Track->CreateIterator());
 			TSharedPtrTS<ITrackIterator> SafeTrkIt(TrkIt);
 			ti->TrackIterators.Add(SafeTrkIt);
-			UEMediaError err = TrkIt->StartAtFirstInteral();
+			UEMediaError err = TrkIt->StartAtFirstInternal();
 			if (err == UEMEDIA_ERROR_OK)
 			{
 				if (TrkIt->GetSampleFileOffset() < LowestFilePos)
diff --git a/Engine/Plugins/Media/ElectraPlayer/Source/ElectraPlayerRuntime/Private/Runtime/Demuxer/ParserISO14496-12.h b/Engine/Plugins/Media/ElectraPlayer/Source/ElectraPlayerRuntime/Private/Runtime/Demuxer/ParserISO14496-12.h
--- a/Engine/Plugins/Media/ElectraPlayer/Source/ElectraPlayerRuntime/Private/Runtime/Demuxer/ParserISO14496-12.h
+++ b/Engine/Plugins/Media/ElectraPlayer/Source/ElectraPlayerRuntime/Private/Runtime/Demuxer/ParserISO14496-12.h
@@ -115,6 +115,10 @@ namespace Electra
 
 		virtual UEMediaError PrepareTracks(IPlayerSessionServices* PlayerSession, TSharedPtrTS<const IParserISO14496_12> OptionalMP4InitSegment) = 0;
 
+		// Cancellation delegate for potentially slow tasks.
+		DECLARE_DELEGATE_RetVal(bool, FCancellationCheckDelegate);
+		virtual UEMediaError ResolveTimecodeTracks(IPlayerSessionServices* InPlayerSession, FCancellationCheckDelegate InCancellationCheckDelegate) = 0;
+
 		virtual TMediaOptionalValue<FTimeFraction> GetMovieDuration() const = 0;
 
 		virtual int32 GetNumberOfTracks() const = 0;
diff --git a/Engine/Plugins/Media/ElectraPlayer/Source/ElectraPlayerRuntime/Private/Runtime/Player/AdaptivePlayerMetadataKeynames.h b/Engine/Plugins/Media/ElectraPlayer/Source/ElectraPlayerRuntime/Private/Runtime/Player/AdaptivePlayerMetadataKeynames.h
new file mode 100644
--- /dev/null
+++ b/Engine/Plugins/Media/ElectraPlayer/Source/ElectraPlayerRuntime/Private/Runtime/Player/AdaptivePlayerMetadataKeynames.h
@@ -0,0 +1,13 @@
+// Copyright Epic Games, Inc. All Rights Reserved.
+
+#pragma once
+
+#include "PlayerCore.h"
+
+namespace Electra
+{
+	const FName MetaDataKeyStartTimecodeValue(TEXT("StartTimecodeValue"));
+	const FName MetaDataKeyStartTimecodeFrameRate(TEXT("StartTimecodeFrameRate"));
+} // namespace Electra
+
+
diff --git a/Engine/Plugins/Media/ElectraPlayer/Source/ElectraPlayerRuntime/Private/Runtime/Player/AdaptivePlayerOptionKeynames.h b/Engine/Plugins/Media/ElectraPlayer/Source/ElectraPlayerRuntime/Private/Runtime/Player/AdaptivePlayerOptionKeynames.h
--- a/Engine/Plugins/Media/ElectraPlayer/Source/ElectraPlayerRuntime/Private/Runtime/Player/AdaptivePlayerOptionKeynames.h
+++ b/Engine/Plugins/Media/ElectraPlayer/Source/ElectraPlayerRuntime/Private/Runtime/Player/AdaptivePlayerOptionKeynames.h
@@ -62,6 +62,8 @@ namespace Electra
 
 	const FName OptionKeyResponseCacheMaxEntries(TEXT("httpcache_max_entries"));
 	const FName OptionKeyResponseCacheMaxByteSize(TEXT("httpcache_max_bytesize"));
+
+	const FName OptionKeyParseTimecodeInfo(TEXT("parse_timecode_info"));
 } // namespace Electra
 
 
diff --git a/Engine/Plugins/Media/ElectraPlayer/Source/ElectraPlayerRuntime/Private/Runtime/Player/AdaptiveStreamingPlayer.cpp b/Engine/Plugins/Media/ElectraPlayer/Source/ElectraPlayerRuntime/Private/Runtime/Player/AdaptiveStreamingPlayer.cpp
--- a/Engine/Plugins/Media/ElectraPlayer/Source/ElectraPlayerRuntime/Private/Runtime/Player/AdaptiveStreamingPlayer.cpp
+++ b/Engine/Plugins/Media/ElectraPlayer/Source/ElectraPlayerRuntime/Private/Runtime/Player/AdaptiveStreamingPlayer.cpp
@@ -267,6 +267,10 @@ void FAdaptiveStreamingPlayer::ModifyOptions(const FParamDict& InOptionsToSetOrC
 	WorkerThread.SendOptionChangeMessage(InOptionsToSetOrChange, InOptionsToClear);
 }
 
+FVariantValue FAdaptiveStreamingPlayer::GetMediaInfo(FName InKey) const
+{
+	return PlayerOptions.GetValue(InKey);	// thread safe
+}
 
 void FAdaptiveStreamingPlayer::SetStaticResourceProviderCallback(const TSharedPtr<IAdaptiveStreamingPlayerResourceProvider, ESPMode::ThreadSafe>& InStaticResourceProvider)
 {
diff --git a/Engine/Plugins/Media/ElectraPlayer/Source/ElectraPlayerRuntime/Private/Runtime/Player/AdaptiveStreamingPlayer.h b/Engine/Plugins/Media/ElectraPlayer/Source/ElectraPlayerRuntime/Private/Runtime/Player/AdaptiveStreamingPlayer.h
--- a/Engine/Plugins/Media/ElectraPlayer/Source/ElectraPlayerRuntime/Private/Runtime/Player/AdaptiveStreamingPlayer.h
+++ b/Engine/Plugins/Media/ElectraPlayer/Source/ElectraPlayerRuntime/Private/Runtime/Player/AdaptiveStreamingPlayer.h
@@ -83,6 +83,11 @@ public:
 	 */
 	virtual void ModifyOptions(const FParamDict& InOptionsToSetOrChange, const FParamDict& InOptionsToClear) = 0;
 
+	/**
+	 * Returns a player option value.
+	 */
+	virtual FVariantValue GetMediaInfo(FName InKey) const = 0;
+	
 	/**
 	 * Sets the attributes for the stream to start buffering for and playing.
 	 * This must be set before calling SeekTo() or LoadManifest() to have an immediate effect.
diff --git a/Engine/Plugins/Media/ElectraPlayer/Source/ElectraPlayerRuntime/Private/Runtime/Player/AdaptiveStreamingPlayerInternal.h b/Engine/Plugins/Media/ElectraPlayer/Source/ElectraPlayerRuntime/Private/Runtime/Player/AdaptiveStreamingPlayerInternal.h
--- a/Engine/Plugins/Media/ElectraPlayer/Source/ElectraPlayerRuntime/Private/Runtime/Player/AdaptiveStreamingPlayerInternal.h
+++ b/Engine/Plugins/Media/ElectraPlayer/Source/ElectraPlayerRuntime/Private/Runtime/Player/AdaptiveStreamingPlayerInternal.h
@@ -1055,6 +1055,7 @@ public:
 
 	void Initialize(const FParamDict& Options) override;
 	void ModifyOptions(const FParamDict& InOptionsToSetOrChange, const FParamDict& InOptionsToClear) override;
+	virtual FVariantValue GetMediaInfo(FName InKey) const override;
 
 	void SetInitialStreamAttributes(EStreamType StreamType, const FStreamSelectionAttributes& InitialSelection) override;
 	void EnableFrameAccurateSeeking(bool bEnabled) override;
diff --git a/Engine/Plugins/Media/ElectraPlayer/Source/ElectraPlayerRuntime/Private/Runtime/Player/mp4/PlaylistReaderMP4.cpp b/Engine/Plugins/Media/ElectraPlayer/Source/ElectraPlayerRuntime/Private/Runtime/Player/mp4/PlaylistReaderMP4.cpp
--- a/Engine/Plugins/Media/ElectraPlayer/Source/ElectraPlayerRuntime/Private/Runtime/Player/mp4/PlaylistReaderMP4.cpp
+++ b/Engine/Plugins/Media/ElectraPlayer/Source/ElectraPlayerRuntime/Private/Runtime/Player/mp4/PlaylistReaderMP4.cpp
@@ -279,6 +279,11 @@ void FPlaylistReaderMP4::WorkerThread()
 				parseError = MP4Parser->PrepareTracks(PlayerSessionServices, TSharedPtrTS<const IParserISO14496_12>());
 				if (parseError == UEMEDIA_ERROR_OK)
 				{
+					MP4Parser->ResolveTimecodeTracks(PlayerSessionServices, IParserISO14496_12::FCancellationCheckDelegate::CreateLambda([&]()
+					{
+						return bAbort;
+					}));
+
 					Manifest = MakeSharedTS<FManifestMP4Internal>(PlayerSessionServices);
 
 					TArray<FURL_RFC3986::FQueryParam> URLFragmentComponents;
diff --git a/Engine/Plugins/Media/ElectraPlayer/Source/ElectraPlayerRuntime/Private/Runtime/Utilities/UtilsMP4.cpp b/Engine/Plugins/Media/ElectraPlayer/Source/ElectraPlayerRuntime/Private/Runtime/Utilities/UtilsMP4.cpp
--- a/Engine/Plugins/Media/ElectraPlayer/Source/ElectraPlayerRuntime/Private/Runtime/Utilities/UtilsMP4.cpp
+++ b/Engine/Plugins/Media/ElectraPlayer/Source/ElectraPlayerRuntime/Private/Runtime/Utilities/UtilsMP4.cpp
@@ -745,6 +745,74 @@ bool UtilsMP4::FMP4RootBoxLocator::LocateRootBoxes(TArray<FBoxInfo>& OutBoxInfos
 	return bSuccess;
 }
 
+TSharedPtrTS<FWaitableBuffer> UtilsMP4::FMP4ChunkLoader::LoadChunk(const int64 InOffset, const int64 InSize, const TSharedPtrTS<IElectraHttpManager>& InHTTPManager, const TSharedPtrTS<IHTTPResponseCache>& InHttpResponseCache, const FString& InURL, FCancellationCheckDelegate InCheckCancellationDelegate)
+{
+	FMediaEvent ReadCompleted;
+	FString URL(InURL);
+	volatile bool bAbort = false;
+	FileSize = -1;
+	bHasErrored = false;
+	
+	TSharedPtrTS<IElectraHttpManager::FProgressListener> ProgressListener = MakeSharedTS<IElectraHttpManager::FProgressListener>();
+	ProgressListener->ProgressDelegate = IElectraHttpManager::FProgressListener::FProgressDelegate::CreateLambda([&](const IElectraHttpManager::FRequest* InRequest)->int32
+	{
+		bAbort = InCheckCancellationDelegate.Execute();
+		return bAbort ? 1 : 0;
+	});
+	ProgressListener->CompletionDelegate = IElectraHttpManager::FProgressListener::FCompletionDelegate::CreateLambda([&](const IElectraHttpManager::FRequest* InRequest)
+	{
+		const bool bFailed = InRequest->ConnectionInfo.StatusInfo.ErrorDetail.IsError();
+		ConnectionInfo = InRequest->ConnectionInfo;
+		if (!bFailed)
+		{
+			// Set the size of the resource if we don't have it yet.
+			if (FileSize < 0)
+			{
+				ElectraHTTPStream::FHttpRange crh;
+				if (crh.ParseFromContentRangeResponse(InRequest->ConnectionInfo.ContentRangeHeader))
+				{
+					FileSize = crh.GetDocumentSize();
+				}
+			}
+
+			if (ConnectionInfo.EffectiveURL.Len())
+			{
+				URL = ConnectionInfo.EffectiveURL;
+			}
+		}
+		bHasErrored = bFailed;
+		ReadCompleted.Signal();
+	});
+
+	auto CreateReadRequestAndBuffer = [&](TSharedPtrTS<FWaitableBuffer>& OutReceiveBuffer, int64 InFromOffset, int64 InNumBytes) -> TSharedPtrTS<IElectraHttpManager::FRequest>
+	{
+		TSharedPtrTS<IElectraHttpManager::FRequest> Req = MakeSharedTS<IElectraHttpManager::FRequest>();
+		Req->Parameters.URL = URL;
+		Req->Parameters.Range.SetStart(InFromOffset);
+		int64 LastByte = InFromOffset + InNumBytes - 1;
+		if (FileSize >= 0 && LastByte > FileSize-1)
+		{
+			LastByte = FileSize - 1;
+		}
+		Req->Parameters.Range.SetEndIncluding(LastByte);
+		Req->Parameters.ConnectTimeout = FTimeValue().SetFromMilliseconds(1000 * 8);
+		Req->Parameters.NoDataTimeout = FTimeValue().SetFromMilliseconds(1000 * 6);
+		OutReceiveBuffer = MakeSharedTS<FWaitableBuffer>();
+		OutReceiveBuffer->Reserve(InNumBytes);
+		Req->ReceiveBuffer = OutReceiveBuffer;
+		Req->ProgressListener = ProgressListener;
+		Req->ResponseCache = InHttpResponseCache;
+		return Req;
+	};
+
+	TSharedPtrTS<FWaitableBuffer> ReceiveBuffer;
+	const TSharedPtrTS<IElectraHttpManager::FRequest> Request = CreateReadRequestAndBuffer(ReceiveBuffer, InOffset, InSize);
+	InHTTPManager->AddRequest(Request, false);
+	ReadCompleted.WaitAndReset();
+	InHTTPManager->RemoveRequest(Request, false);
+
+	return !(bAbort || bHasErrored) ? ReceiveBuffer : nullptr;
+}
 
 }
 
diff --git a/Engine/Plugins/Media/ElectraPlayer/Source/ElectraPlayerRuntime/Private/Runtime/Utilities/UtilsMP4.h b/Engine/Plugins/Media/ElectraPlayer/Source/ElectraPlayerRuntime/Private/Runtime/Utilities/UtilsMP4.h
--- a/Engine/Plugins/Media/ElectraPlayer/Source/ElectraPlayerRuntime/Private/Runtime/Utilities/UtilsMP4.h
+++ b/Engine/Plugins/Media/ElectraPlayer/Source/ElectraPlayerRuntime/Private/Runtime/Utilities/UtilsMP4.h
@@ -120,6 +120,30 @@ public:
 		volatile bool bHasErrored = false;
 	};
 
+	class FMP4ChunkLoader
+	{
+	public:		
+		FMP4ChunkLoader() = default;
+		~FMP4ChunkLoader() = default;
+		
+		DECLARE_DELEGATE_RetVal(bool, FCancellationCheckDelegate);
+		TSharedPtrTS<FWaitableBuffer> LoadChunk(const int64 InOffset, const int64 InSize, const TSharedPtrTS<IElectraHttpManager>& InHTTPManager, const TSharedPtrTS<IHTTPResponseCache>& InHttpResponseCache, const FString& InURL, FCancellationCheckDelegate InCheckCancellationDelegate);
+
+		int64 GetFileSize() const
+		{ return FileSize; }
+		bool DidDownloadFail() const
+		{ return bHasErrored; }
+		const HTTP::FConnectionInfo& GetConnectionInfo() const
+		{ return ConnectionInfo; }
+
+		const FString& GetErrorMessage() const
+		{ return ErrorMsg; }
+	private:
+		HTTP::FConnectionInfo ConnectionInfo;
+		FString ErrorMsg;
+		int64 FileSize = -1;
+		volatile bool bHasErrored = false;
+	};
 };
 
 } // namespace Electra
diff --git a/Engine/Plugins/Media/ElectraPlayer/Source/ElectraPlayerRuntime/Public/IElectraPlayerInterface.h b/Engine/Plugins/Media/ElectraPlayer/Source/ElectraPlayerRuntime/Public/IElectraPlayerInterface.h
--- a/Engine/Plugins/Media/ElectraPlayer/Source/ElectraPlayerRuntime/Public/IElectraPlayerInterface.h
+++ b/Engine/Plugins/Media/ElectraPlayer/Source/ElectraPlayerRuntime/Public/IElectraPlayerInterface.h
@@ -320,6 +320,7 @@ public:
 	virtual bool GetVideoStreamFormat(FVideoStreamFormat& OutFormat, int32 InTrackIndex, int32 InStreamIndex) const = 0;
 	virtual bool GetActiveVideoStreamFormat(FVideoStreamFormat& OutFormat) const = 0;
 
+	virtual Electra::FVariantValue GetMediaInfo(FName InInfoName) const = 0;
 	virtual TSharedPtr<TMap<FString, TArray<TSharedPtr<Electra::IMediaStreamMetadata::IItem, ESPMode::ThreadSafe>>>, ESPMode::ThreadSafe> GetMediaMetadata() const = 0;
 
 	virtual void NotifyOfOptionChange() = 0;
diff --git a/Engine/Plugins/Media/ElectraUtil/Source/ElectraBase/Private/ParameterDictionary.cpp b/Engine/Plugins/Media/ElectraUtil/Source/ElectraBase/Private/ParameterDictionary.cpp
--- a/Engine/Plugins/Media/ElectraUtil/Source/ElectraBase/Private/ParameterDictionary.cpp
+++ b/Engine/Plugins/Media/ElectraUtil/Source/ElectraBase/Private/ParameterDictionary.cpp
@@ -118,6 +118,29 @@ void FVariantValue::CopyInternal(const FVariantValue& FromOther)
 	}
 }
 
+FVariant FVariantValue::ToFVariant() const
+{
+	switch(GetDataType())
+	{
+	case EDataType::TypeFString:
+		return FVariant(GetFString());
+	case EDataType::TypeDouble:
+		return FVariant(GetDouble());
+	case EDataType::TypeInt64:
+		return FVariant(GetInt64());
+	case EDataType::TypeBoolean:
+		return FVariant(GetBool());
+	case EDataType::TypeTimeValue:
+		return FVariant(GetTimeValue().GetAsTimespan());
+	case EDataType::TypeVoidPointer:
+		return FVariant(reinterpret_cast<uint64>(GetPointer()));
+	case EDataType::TypeU8Array:
+		return FVariant(GetArray());
+	case EDataType::TypeSharedPointer: // Can't be converted.
+		default:
+			return FVariant();
+	}
+}
 
 void FVariantValue::Clear()
 {
@@ -457,12 +480,28 @@ void FParamDict::Set(const FName& Key, const FVariantValue& Value)
 	Dictionary.Emplace(Key, Value); 
 }
 
+void FParamDict::Set(const FName& Key, FVariantValue&& Value)
+{
+	Dictionary.Emplace(Key, MoveTemp(Value));
+}
+
 void FParamDict::GetKeys(TArray<FName>& OutKeys) const
 {
 	OutKeys.Empty();
 	Dictionary.GenerateKeyArray(OutKeys);
 }
 
+bool FParamDict::SetValueFrom(FName InKey, const FParamDict& InOther)
+{
+	FVariantValue OtherValue = InOther.GetValue(InKey);
+	const bool bOtherHasKey = OtherValue.IsValid();
+	if (bOtherHasKey)
+	{
+		Set(InKey, MoveTemp(OtherValue));
+	}
+	return bOtherHasKey;
+}
+
 void FParamDict::ConvertKeysStartingWithTo(TMap<FString, FVariant>& OutVariantMap, const FString& InKeyStartsWith, const FString& InAddPrefixToKey) const
 {
 	OutVariantMap.Reserve(Dictionary.Num());
@@ -478,47 +517,11 @@ void FParamDict::ConvertKeysStartingWithTo(TMap<FString, FVariant>& OutVariantMa
 
 		NewKey = InAddPrefixToKey;
 		NewKey.Append(s);
-		switch(Pair.Value.GetDataType())
+
+		FVariant ConvertedValue = Pair.Value.ToFVariant();
+		if (!ConvertedValue.IsEmpty())
 		{
-			case FVariantValue::EDataType::TypeFString:
-			{
-				OutVariantMap.Emplace(NewKey, Pair.Value.GetFString());
-				break;
-			}
-			case FVariantValue::EDataType::TypeDouble:
-			{
-				OutVariantMap.Emplace(NewKey, Pair.Value.GetDouble());
-				break;
-			}
-			case FVariantValue::EDataType::TypeInt64:
-			{
-				OutVariantMap.Emplace(NewKey, Pair.Value.GetInt64());
-				break;
-			}
-			case FVariantValue::EDataType::TypeBoolean:
-			{
-				OutVariantMap.Emplace(NewKey, Pair.Value.GetInt64());
-				break;
-			}
-			case FVariantValue::EDataType::TypeTimeValue:
-			{
-				OutVariantMap.Emplace(NewKey, Pair.Value.GetTimeValue().GetAsTimespan());
-				break;
-			}
-			case FVariantValue::EDataType::TypeVoidPointer:
-			{
-				OutVariantMap.Emplace(NewKey, (uint64) Pair.Value.GetPointer());
-				break;
-			}
-			case FVariantValue::EDataType::TypeU8Array:
-			{
-				OutVariantMap.Emplace(NewKey, Pair.Value.GetArray());
-				break;
-			}
-			default:
-			{
-				break;
-			}
+			OutVariantMap.Emplace(NewKey, MoveTemp(ConvertedValue));
 		}
 	}
 }
diff --git a/Engine/Plugins/Media/ElectraUtil/Source/ElectraBase/Public/ParameterDictionary.h b/Engine/Plugins/Media/ElectraUtil/Source/ElectraBase/Public/ParameterDictionary.h
--- a/Engine/Plugins/Media/ElectraUtil/Source/ElectraBase/Public/ParameterDictionary.h
+++ b/Engine/Plugins/Media/ElectraUtil/Source/ElectraBase/Public/ParameterDictionary.h
@@ -131,6 +131,8 @@ public:
 		return DataType == type;
 	}
 
+	FVariant ToFVariant() const;
+
 private:
 
 	union FUnionLayout
@@ -163,10 +165,12 @@ public:
 	~FParamDict() = default;
 	void Clear();
 	void Set(const FName& Key, const FVariantValue& Value);
+	void Set(const FName& Key, FVariantValue&& Value);
 	void GetKeys(TArray<FName>& OutKeys) const;
 	bool HaveKey(const FName& Key) const;
 	FVariantValue GetValue(const FName& Key) const;
 	void Remove(const FName& Key);
+	bool SetValueFrom(FName InKey, const FParamDict& InOther);
 
 	void ConvertTo(TMap<FString, FVariant>& OutVariantMap, const FString& InAddPrefixToKey) const;
 	void ConvertKeysStartingWithTo(TMap<FString, FVariant>& OutVariantMap, const FString& InKeyStartsWith, const FString& InAddPrefixToKey) const;
@@ -217,6 +221,11 @@ public:
 		FScopeLock lock(&Lock);
 		Dictionary.Set(Key, Value);
 	}
+	void Set(const FName& Key, FVariantValue&& Value)
+	{
+		FScopeLock lock(&Lock);
+		Dictionary.Set(Key, MoveTemp(Value));
+	}
 	void GetKeys(TArray<FName>& OutKeys) const
 	{
 		FScopeLock lock(&Lock);
@@ -237,6 +246,11 @@ public:
 		FScopeLock lock(&Lock);
 		Dictionary.Remove(Key);
 	}
+	bool SetValueFrom(FName InKey, const FParamDict& InOther)
+	{
+		FScopeLock lock(&Lock);
+		return Dictionary.SetValueFrom(InKey, InOther);
+	}
 	void ConvertTo(TMap<FString, FVariant>& OutVariantMap, const FString& InAddPrefixToKey) const
 	{
 		FScopeLock lock(&Lock);
diff --git a/Engine/Source/Runtime/MediaUtils/Public/MediaPlayerOptions.h b/Engine/Source/Runtime/MediaUtils/Public/MediaPlayerOptions.h
--- a/Engine/Source/Runtime/MediaUtils/Public/MediaPlayerOptions.h
+++ b/Engine/Source/Runtime/MediaUtils/Public/MediaPlayerOptions.h
@@ -165,4 +165,10 @@ namespace MediaPlayerOptionValues
 	{ static FName OptName(TEXT("ImgMediaSmartCacheEnabled")); return OptName; }
 	inline const FName& ImgMediaSmartCacheTimeToLookAhead() 
 	{ static FName OptName(TEXT("ImgMediaSmartCacheTimeToLookAhead")); return OptName; }
+
+	/**
+	 * Requests parsing timecode meta data if available. 
+	 */
+	inline const FName& ParseTimecodeInfo()
+	{ static FName OptName(TEXT("ParseTimecodeInfo")); return OptName; }
 }
