MediaServices: Fix GC issue introduced by 35867793. There was a circular reference inside the Facade. A TSharedPtr mixed with GC.
#rnx

diff --git a/Engine/Plugins/Media/ElectraPlayer/Source/ElectraPlayerPlugin/Public/IElectraPlayerPluginModule.h b/Engine/Plugins/Media/ElectraPlayer/Source/ElectraPlayerPlugin/Public/IElectraPlayerPluginModule.h
--- a/Engine/Plugins/Media/ElectraPlayer/Source/ElectraPlayerPlugin/Public/IElectraPlayerPluginModule.h
+++ b/Engine/Plugins/Media/ElectraPlayer/Source/ElectraPlayerPlugin/Public/IElectraPlayerPluginModule.h
@@ -10,6 +10,7 @@
 #include "IMediaOptions.h"
 #include "UObject/Object.h"
 #include "UObject/StrongObjectPtrTemplates.h"
+#include "UObject/WeakObjectPtr.h"
 
 class IAnalyticsProviderET;
 class IMediaEventSink;
@@ -81,7 +82,7 @@ public:
 	{
 		FScopeLock lock(&OwnerLock);
 		Owner = nullptr;
-		OwnerObject = nullptr;
+		OwnerObject.Reset();
 	}
 	virtual void Lock() override
 	{
@@ -93,12 +94,12 @@ public:
 	}
 	virtual IMediaOptions* GetMediaOptionInterface() override
 	{
-		return Owner;
+		return OwnerObject.IsStale(true, true) ? nullptr : Owner;
 	}
 private:
 	FCriticalSection OwnerLock;
 	IMediaOptions* Owner = nullptr;
-	TStrongObjectPtr<const UObject> OwnerObject;
+	FWeakObjectPtr OwnerObject;
 };
 
 
diff --git a/Engine/Source/Runtime/MediaUtils/Private/MediaPlayerFacade.cpp b/Engine/Source/Runtime/MediaUtils/Private/MediaPlayerFacade.cpp
--- a/Engine/Source/Runtime/MediaUtils/Private/MediaPlayerFacade.cpp
+++ b/Engine/Source/Runtime/MediaUtils/Private/MediaPlayerFacade.cpp
@@ -768,7 +768,7 @@ public:
 
 	virtual const IMediaOptions* GetOptions() const override
 	{
-		return OptionsObject.IsValid() ? Options : nullptr;
+		return OptionsObject.IsStale() ? nullptr : Options;
 	}
 
 	virtual const FMediaPlayerOptions* GetPlayerOptions() const override
