ElectraPlayer: Disabling emitting of decoded frames into the sample queue when paused. While this speeds up single frame stepping, it also sends additional frames into the queue from which the facade can pick the one "best fitting" the current game time. This will cause desynchronization in sequencer with multiple player instances due to the asynchronous arrival of decoded frames, giving one instance a choice of frames to select from while others get no choice, resulting in inconsistent selection of the frame to be used.

#rnx

diff --git a/Engine/Plugins/Media/ElectraPlayer/Source/ElectraPlayerPlugin/Private/ElectraPlayerPlugin.cpp b/Engine/Plugins/Media/ElectraPlayer/Source/ElectraPlayerPlugin/Private/ElectraPlayerPlugin.cpp
--- a/Engine/Plugins/Media/ElectraPlayer/Source/ElectraPlayerPlugin/Private/ElectraPlayerPlugin.cpp
+++ b/Engine/Plugins/Media/ElectraPlayer/Source/ElectraPlayerPlugin/Private/ElectraPlayerPlugin.cpp
@@ -664,7 +664,7 @@ bool FElectraPlayerPlugin::Open(const FString& Url, const IMediaOptions* Options
 	{
 		PlayerOptions.Set(TEXT("optimize_seek_for_scrubbing"), Electra::FVariantValue(true));
 		//PlayerOptions.Set(TEXT("do_not_hold_back_first_frame"), Electra::FVariantValue(true));
-		PlayerOptions.Set(TEXT("always_emit_samples_when_paused"), Electra::FVariantValue(true));
+		//PlayerOptions.Set(TEXT("always_emit_samples_when_paused"), Electra::FVariantValue(true));
 		PlayerOptions.Set(TEXT("worker_threads"), Electra::FVariantValue(FString(TEXT("worker"))));
 	}
 	else
