Repaired DXVA decoder validation layer issue with videos that don't match decoder size requirements
#rnx

diff --git a/Engine/Plugins/Media/ElectraPlayer/Source/ElectraPlayerRuntime/Private/Runtime/Decoder/Windows/ElectraVideoDecoder_PC.cpp b/Engine/Plugins/Media/ElectraPlayer/Source/ElectraPlayerRuntime/Private/Runtime/Decoder/Windows/ElectraVideoDecoder_PC.cpp
--- a/Engine/Plugins/Media/ElectraPlayer/Source/ElectraPlayerRuntime/Private/Runtime/Decoder/Windows/ElectraVideoDecoder_PC.cpp
+++ b/Engine/Plugins/Media/ElectraPlayer/Source/ElectraPlayerRuntime/Private/Runtime/Decoder/Windows/ElectraVideoDecoder_PC.cpp
@@ -172,6 +172,7 @@ void FElectraPlayerVideoDecoderOutputPC::InitializeWithResource(const TRefCountP
 	ResourceDesc = InResource->GetDesc();
 
 	const bool bIsTexture = (ResourceDesc.Dimension == D3D12_RESOURCE_DIMENSION_TEXTURE2D);
+	const FIntPoint TextureDim = bIsTexture ? FIntPoint{ static_cast<int>(MaxWidth), static_cast<int>(MaxHeight) } : SampleDim;
 
 	const FIntPoint Dim = GetDim();
 	const EPixelFormat PixFmt = GetFormat();
@@ -216,7 +217,7 @@ void FElectraPlayerVideoDecoderOutputPC::InitializeWithResource(const TRefCountP
 	* - We need to copy this from the output buffers ASAP to not throttle buffer availability (and hence performance) for the decoder
 	* - As RHI limits resource creation to the renderthread, which would delay creation, we need to use a platform API level texture we can create right here and now
 	*/
-	if (!TextureDX12.IsValid() || SampleDim != TextureDX12Dim)
+	if (!TextureDX12.IsValid() || TextureDim != TextureDX12Dim)
 	{
 		// We get here only after the instance came back from the pool! Hence we can be sure any old texture is no longer actively used.
 		TextureDX12 = nullptr;
@@ -233,7 +234,7 @@ void FElectraPlayerVideoDecoderOutputPC::InitializeWithResource(const TRefCountP
 		}
 
 		FElectraMediaDecoderOutputBufferPool_DX12::FOutputData OutputData;
-		if (!D3D12ResourcePool->AllocateOutputDataAsTexture(OutputData, SampleDim.X, SampleDim.Y, DXGIFmt))
+		if (!D3D12ResourcePool->AllocateOutputDataAsTexture(OutputData, TextureDim.X, TextureDim.Y, DXGIFmt))
 		{
 			UE_LOG(LogElectraPlayer, Error, TEXT("Could not allocate texture resource from heap!"));
 			return;
@@ -246,7 +247,7 @@ void FElectraPlayerVideoDecoderOutputPC::InitializeWithResource(const TRefCountP
 		D3DFence = OutputData.Fence;
 		FenceValue = OutputData.FenceValue;
 
-		TextureDX12Dim = SampleDim;
+		TextureDX12Dim = TextureDim;
 	}
 	else
 	{
