ElectraPlayer: Replacing the MPD[REDACTED_HANDLE] on updated MPD's with the one from the initial MPD. The attribute must not change according to the specification, but we see it change with certain streams, causing playback problems. We ignore the changed value now to make such streams play.

#rnx

diff --git a/Engine/Plugins/Media/ElectraPlayer/Source/ElectraPlayerRuntime/Private/Runtime/Player/DASH/PlaylistReaderDASH.cpp b/Engine/Plugins/Media/ElectraPlayer/Source/ElectraPlayerRuntime/Private/Runtime/Player/DASH/PlaylistReaderDASH.cpp
--- a/Engine/Plugins/Media/ElectraPlayer/Source/ElectraPlayerRuntime/Private/Runtime/Player/DASH/PlaylistReaderDASH.cpp
+++ b/Engine/Plugins/Media/ElectraPlayer/Source/ElectraPlayerRuntime/Private/Runtime/Player/DASH/PlaylistReaderDASH.cpp
@@ -1405,6 +1405,12 @@ void FPlaylistReaderDASH::ManifestUpdateDownloadCompleted(FResourceLoadRequestPt
 				LastErrorDetail = Builder->BuildFromMPD(NewManifest, XML.GetCharArray().GetData(), EffectiveURL, ETag);
 				if (LastErrorDetail.IsOK() || LastErrorDetail.IsTryAgain())
 				{
+					// Set the availability start time of the new MPD to be that of the current MPD.
+					// It is not supposed to change in accordance with ISO/IEC 23009-1:2022 Section 8.4.2, but we
+					// have seen MPDs where it *does* change. In an attempt to make these MPDs usable we keep
+					// the initial AST.
+					NewManifest->GetMPDRoot()->SetAvailabilityStartTime(Manifest->GetMPDRoot()->GetAvailabilityStartTime());
+
 					// Copy over the initial document URL fragments.
 					if (NewManifest.IsValid())
 					{
