Disable ALLOW_MFSAMPLE_WITH_DX12 for all Windows SDK, as the current implementation does not work reliably on AMD GPUs
#rnx

diff --git a/Engine/Plugins/Media/ElectraCodecs/Source/ElectraDecoders/Private/Windows/WindowsPlatformHeaders_Video_DX.h b/Engine/Plugins/Media/ElectraCodecs/Source/ElectraDecoders/Private/Windows/WindowsPlatformHeaders_Video_DX.h
--- a/Engine/Plugins/Media/ElectraCodecs/Source/ElectraDecoders/Private/Windows/WindowsPlatformHeaders_Video_DX.h
+++ b/Engine/Plugins/Media/ElectraCodecs/Source/ElectraDecoders/Private/Windows/WindowsPlatformHeaders_Video_DX.h
@@ -18,7 +18,7 @@ THIRD_PARTY_INCLUDES_START
 #include "mferror.h"
 #include "mfidl.h"
 
-#if defined(NTDDI_WIN10_NI)
+#if 0 // defined(NTDDI_WIN10_NI) // Currently disabled as this causes trouble on AMD GPU
 #include "mfd3d12.h"
 #define ALLOW_MFSAMPLE_WITH_DX12	1	// Windows SDK 22621 and up do feature APIs to support DX12 texture resources with WMF transforms
 #else
diff --git a/Engine/Plugins/Media/ElectraPlayer/Source/ElectraPlayerRuntime/Private/Runtime/Decoder/Windows/ElectraVideoDecoder_PC.cpp b/Engine/Plugins/Media/ElectraPlayer/Source/ElectraPlayerRuntime/Private/Runtime/Decoder/Windows/ElectraVideoDecoder_PC.cpp
--- a/Engine/Plugins/Media/ElectraPlayer/Source/ElectraPlayerRuntime/Private/Runtime/Decoder/Windows/ElectraVideoDecoder_PC.cpp
+++ b/Engine/Plugins/Media/ElectraPlayer/Source/ElectraPlayerRuntime/Private/Runtime/Decoder/Windows/ElectraVideoDecoder_PC.cpp
@@ -14,7 +14,7 @@ THIRD_PARTY_INCLUDES_START
 #include "mfapi.h"
 #include "mferror.h"
 #include "mfidl.h"
-#if defined(NTDDI_WIN10_NI)
+#if 0 // defined(NTDDI_WIN10_NI) // Currently disabled as this causes trouble on AMD GPU
 #include "mfd3d12.h"
 #define ALLOW_MFSAMPLE_WITH_DX12	1	// Windows SDK 22621 and up do feature APIs to support DX12 texture resources with WMF transforms
 #else
