Fail FElectraPlayerPlugin::Open when Options == nullptr instead of crashing

#rnx

diff --git a/Engine/Plugins/Media/ElectraPlayer/Source/ElectraPlayerPlugin/Private/ElectraPlayerPlugin.cpp b/Engine/Plugins/Media/ElectraPlayer/Source/ElectraPlayerPlugin/Private/ElectraPlayerPlugin.cpp
--- a/Engine/Plugins/Media/ElectraPlayer/Source/ElectraPlayerPlugin/Private/ElectraPlayerPlugin.cpp
+++ b/Engine/Plugins/Media/ElectraPlayer/Source/ElectraPlayerPlugin/Private/ElectraPlayerPlugin.cpp
@@ -532,6 +532,19 @@ bool FElectraPlayerPlugin::Open(const FString& Url, const IMediaOptions* Options
 
 bool FElectraPlayerPlugin::Open(const FString& Url, const IMediaOptions* Options, const FMediaPlayerOptions* InPlayerOptions)
 {
+	if (Options == nullptr)
+	{
+		UE_LOG(LogElectraPlayerPlugin, Error, TEXT("[%p] IMediaPlayer::Open: Options == nullptr"), this);
+
+		FScopeLock lock(&CallbackPointerLock);
+		if (EventSink)
+		{
+			EventSink->ReceiveMediaEvent(EMediaEvent::MediaOpenFailed);
+		}
+
+		return false;
+	}
+
 	// Get the safe option interface to poll for changes during playback.
 	CallbackPointerLock.Lock();
 	OptionInterface = StaticCastSharedPtr<IElectraSafeMediaOptionInterface>(Options->GetMediaOption(TEXT("GetSafeMediaOptions"), TSharedPtr<IElectraSafeMediaOptionInterface, ESPMode::ThreadSafe>()));
