Electra: Changed parser to use the newer decoder configuration record structure

#rnx

diff --git a/Engine/Plugins/Media/ElectraCodecs/Source/ElectraDecoders/Private/Utils/MPEG/ElectraUtilsMPEGVideo_H265.cpp b/Engine/Plugins/Media/ElectraCodecs/Source/ElectraDecoders/Private/Utils/MPEG/ElectraUtilsMPEGVideo_H265.cpp
--- a/Engine/Plugins/Media/ElectraCodecs/Source/ElectraDecoders/Private/Utils/MPEG/ElectraUtilsMPEGVideo_H265.cpp
+++ b/Engine/Plugins/Media/ElectraCodecs/Source/ElectraDecoders/Private/Utils/MPEG/ElectraUtilsMPEGVideo_H265.cpp
@@ -1057,6 +1057,17 @@ namespace ElectraDecodersUtil
 				#undef RANGE_CHECK_FAILURE
 			}
 
+			bool ParseVideoParameterSet(FVideoParameterSet& OutVideoParameterSet, const uint8* InBitstream, uint64 InBitstreamLenInBytes)
+			{
+				TMap<uint32, FVideoParameterSet> vpss;
+				if (!ParseVideoParameterSet(vpss, InBitstream, InBitstreamLenInBytes))
+				{
+					return false;
+				}
+				OutVideoParameterSet = vpss.CreateConstIterator().Value();
+				return true;
+			}
+
 			bool ParseSequenceParameterSet(FSequenceParameterSet& OutSequenceParameterSet, const uint8* InBitstream, uint64 InBitstreamLenInBytes)
 			{
 				TMap<uint32, FSequenceParameterSet> spss;
@@ -2664,8 +2675,21 @@ namespace ElectraDecodersUtil
 							CSD.Append(a.NALUs[j]);
 						}
 
+						// VPS nut?
+						if (a.NALUnitType == 32)
+						{
+							TArray<uint8>& vps(VideoParameterSets.Emplace_GetRef());
+							for(int32 k=0; k<a.NALUs.Num(); ++k)
+							{
+								vps.Append(a.NALUs[k]);
+								if (!ParseVideoParameterSet(ParsedVideoParameterSets.Emplace_GetRef(), a.NALUs[k].GetData(), a.NALUs[k].Num()))
+								{
+									ParsedVideoParameterSets.RemoveAt(ParsedVideoParameterSets.Num()-1);
+								}
+							}
+						}
 						// SPS nut?
-						if (a.NALUnitType == 33)
+						else if (a.NALUnitType == 33)
 						{
 							TArray<uint8>& sps(SequenceParameterSets.Emplace_GetRef());
 							for(int32 k=0; k<a.NALUs.Num(); ++k)
@@ -2677,6 +2701,24 @@ namespace ElectraDecodersUtil
 								}
 							}
 						}
+						// Prefix nut?
+						else if (a.NALUnitType == 39)
+						{
+							TArray<uint8>& Prefix(PrefixNUTs.Emplace_GetRef());
+							for(int32 k=0; k<a.NALUs.Num(); ++k)
+							{
+								Prefix.Append(a.NALUs[k]);
+							}
+						}
+						// Suffix nut?
+						else if (a.NALUnitType == 40)
+						{
+							TArray<uint8>& Suffix(SuffixNUTs.Emplace_GetRef());
+							for(int32 k=0; k<a.NALUs.Num(); ++k)
+							{
+								Suffix.Append(a.NALUs[k]);
+							}
+						}
 					}
 					return true;
 				}
diff --git a/Engine/Plugins/Media/ElectraCodecs/Source/ElectraDecoders/Public/Utils/MPEG/ElectraUtilsMPEGVideo_H265.h b/Engine/Plugins/Media/ElectraCodecs/Source/ElectraDecoders/Public/Utils/MPEG/ElectraUtilsMPEGVideo_H265.h
--- a/Engine/Plugins/Media/ElectraCodecs/Source/ElectraDecoders/Public/Utils/MPEG/ElectraUtilsMPEGVideo_H265.h
+++ b/Engine/Plugins/Media/ElectraCodecs/Source/ElectraDecoders/Public/Utils/MPEG/ElectraUtilsMPEGVideo_H265.h
@@ -1363,13 +1363,16 @@ namespace ElectraDecodersUtil
 			bool ELECTRADECODERS_API ParsePictureParameterSet(TMap<uint32, FPictureParameterSet>& InOutPictureParameterSets, const TMap<uint32, FSequenceParameterSet>& InSequenceParameterSets, const uint8* InBitstream, uint64 InBitstreamLenInBytes);
 			bool ELECTRADECODERS_API ParseSliceHeader(TUniquePtr<FRBSP>& OutRBSP, FBitstreamReader& OutRBSPReader, FSliceSegmentHeader& OutSlice, const TMap<uint32, FVideoParameterSet>& InVideoParameterSets, const TMap<uint32, FSequenceParameterSet>& InSequenceParameterSets, const TMap<uint32, FPictureParameterSet>& InPictureParameterSets, const uint8* InBitstream, uint64 InBitstreamLenInBytes);
 
+			bool ELECTRADECODERS_API ParseVideoParameterSet(FVideoParameterSet& OutSequenceParameterSet, const uint8* InBitstream, uint64 InBitstreamLenInBytes);
+			inline bool ParseVideoParameterSet(FVideoParameterSet& OutSequenceParameterSet, const TArray<uint8>& InData)
+			{ return ParseVideoParameterSet(OutSequenceParameterSet, InData.GetData(), InData.Num()); }
+
 			bool ELECTRADECODERS_API ParseSequenceParameterSet(FSequenceParameterSet& OutSequenceParameterSet, const uint8* InBitstream, uint64 InBitstreamLenInBytes);
 			inline bool ParseSequenceParameterSet(FSequenceParameterSet& OutSequenceParameterSet, const TArray<uint8>& InData)
 			{ return ParseSequenceParameterSet(OutSequenceParameterSet, InData.GetData(), InData.Num()); }
 
 
 
-
 			class FHEVCDecoderConfigurationRecord
 			{
 			public:
@@ -1381,8 +1384,12 @@ namespace ElectraDecodersUtil
 				inline const TArray<uint8>& GetCodecSpecificData() const
 				{ return CSD; }
 
+				inline const TArray<TArray<uint8>>& GetVideoParameterSets() const
+				{ return VideoParameterSets; }
 				inline const TArray<TArray<uint8>>& GetSequenceParameterSets() const
 				{ return SequenceParameterSets; }
+				inline const TArray<FSequenceParameterSet>& GetParsedSequenceParameterSets() const
+				{ return ParsedSequenceParameterSets; }
 				inline uint8 GetGeneralProfileSpace() const
 				{ return GeneralProfileSpace; }
 				inline uint8 GetGeneralTierFlag() const
@@ -1401,6 +1408,10 @@ namespace ElectraDecodersUtil
 				{ return BitDepthLumaMinus8; }
 				inline uint8 GetBitDepthChromaMinus8() const
 				{ return BitDepthChromaMinus8; }
+				inline const TArray<TArray<uint8>>& GetPrefixNUTs() const
+				{ return PrefixNUTs; }
+				inline const TArray<TArray<uint8>>& GetSuffixNUTs() const
+				{ return SuffixNUTs; }
 
 				ELECTRADECODERS_API FString GetCodecSpecifierRFC6381(const TCHAR* InPrefix) const;
 				ELECTRADECODERS_API FString GetFormatInfo() const;
@@ -1413,7 +1424,11 @@ namespace ElectraDecodersUtil
 					uint8 NALUnitType = 0;
 				};
 				TArray<FArray> Arrays;
+				TArray<TArray<uint8>> VideoParameterSets;
 				TArray<TArray<uint8>> SequenceParameterSets;
+				TArray<TArray<uint8>> PrefixNUTs;
+				TArray<TArray<uint8>> SuffixNUTs;
+				TArray<FVideoParameterSet> ParsedVideoParameterSets;
 				TArray<FSequenceParameterSet> ParsedSequenceParameterSets;
 				TArray<uint8> RawData;
 				TArray<uint8> CSD;
diff --git a/Engine/Plugins/Media/ElectraPlayer/Source/ElectraPlayerRuntime/Private/Runtime/Demuxer/ParserISO14496-12.cpp b/Engine/Plugins/Media/ElectraPlayer/Source/ElectraPlayerRuntime/Private/Runtime/Demuxer/ParserISO14496-12.cpp
--- a/Engine/Plugins/Media/ElectraPlayer/Source/ElectraPlayerRuntime/Private/Runtime/Demuxer/ParserISO14496-12.cpp
+++ b/Engine/Plugins/Media/ElectraPlayer/Source/ElectraPlayerRuntime/Private/Runtime/Demuxer/ParserISO14496-12.cpp
@@ -3,6 +3,7 @@
 #include "Misc/Build.h"
 #include "PlayerCore.h"
 #include "Containers/ChunkedArray.h"
+#include "Misc/Optional.h"
 #include "Misc/FrameRate.h"
 #include "Misc/Timecode.h"
 
@@ -12,7 +13,8 @@
 #include "Utilities/UtilsMP4.h"
 #include "Utilities/UtilsMPEG.h"
 #include "Utils/MPEG/ElectraUtilsMPEGAudio.h"
-#include "Utils/MPEG/ElectraUtilsMPEGVideo.h"
+#include "Utils/MPEG/ElectraUtilsMPEGVideo_H264.h"
+#include "Utils/MPEG/ElectraUtilsMPEGVideo_H265.h"
 #include "Utilities/ISO639-Map.h"
 #include "InfoLog.h"
 #include "Player/AdaptivePlayerMetadataKeynames.h"
@@ -537,7 +539,7 @@ namespace Electra
 		{
 			return FString::Printf(TEXT("%c%c%c%c"), TCHAR(InBoxName[0]?InBoxName[0]:'?'), TCHAR(InBoxName[1]?InBoxName[1]:'?'), TCHAR(InBoxName[2]?InBoxName[2]:'?'), TCHAR(InBoxName[3]?InBoxName[3]:'?'));
 		}
-		
+
 		FString GetBoxName() const
 		{
 			return GetBoxNameString(DbgBoxName);
@@ -821,7 +823,7 @@ namespace Electra
 		 * Track Reference types - https://developer.apple.com/documentation/quicktime-file-format/track_reference_type_atom
 		 */
 		static const IParserISO14496_12::FBoxType kTrackRef_tmcd = MAKE_BOX_ATOM('t', 'm', 'c', 'd');
-		
+
 		/**
 		 * Sample grouping types
 		 */
@@ -2290,7 +2292,7 @@ namespace Electra
 		uint32 GetTimeScale() const { return TimeScale; }
 		uint32 GetFrameDuration() const { return FrameDuration; }
 		uint32 GetNumberOfFrames() const { return NumberOfFrames; }
-		
+
 	private:
 		FMP4BoxTimecodeSampleEntry() = delete;
 		FMP4BoxTimecodeSampleEntry(const FMP4BoxTimecodeSampleEntry&) = delete;
@@ -2310,7 +2312,7 @@ namespace Electra
 			RETURN_IF_ERROR(ParseInfo->Reader()->Read(FrameDuration));
 			RETURN_IF_ERROR(ParseInfo->Reader()->Read(NumberOfFrames));
 			RETURN_IF_ERROR(ParseInfo->Reader()->Read(Reserved8));
-			
+
 			// Next would be (optional) source reference atom.
 			const int64 NowAt = ParseInfo->Reader()->GetCurrentReadOffset();
 			int32 BytesRemaining = static_cast<int32>(BoxSize - (NowAt - StartOffset));
@@ -2327,7 +2329,7 @@ namespace Electra
 			}
 			return Error;
 		}
-		
+
 		uint32 Flags = 0;
 		uint32 TimeScale = 0;
 		uint32 FrameDuration = 0;
@@ -2838,7 +2840,7 @@ namespace Electra
 		{
 		}
 
-		const ElectraDecodersUtil::MPEG::FAVCDecoderConfigurationRecord& GetDecoderConfigurationRecord() const
+		const TOptional<ElectraDecodersUtil::MPEG::H264::FAVCDecoderConfigurationRecord>& GetDecoderConfigurationRecord() const
 		{
 			return ConfigurationRecord;
 		}
@@ -2854,7 +2856,7 @@ namespace Electra
 			uint32 BytesRemaining = BoxSize - (ParseInfo->Reader()->GetCurrentReadOffset() - StartOffset);
 			if (BytesRemaining)
 			{
-				void* TempBuffer = FMemory::Malloc(BytesRemaining);
+				uint8* TempBuffer = (uint8*)FMemory::Malloc(BytesRemaining);
 				if (!TempBuffer)
 				{
 					return UEMEDIA_ERROR_OOM;
@@ -2862,14 +2864,18 @@ namespace Electra
 				Error = ParseInfo->Reader()->ReadBytes(TempBuffer, BytesRemaining);
 				if (Error == UEMEDIA_ERROR_OK)
 				{
-					ConfigurationRecord.SetRawData(TempBuffer, BytesRemaining);
+					ElectraDecodersUtil::MPEG::H264::FAVCDecoderConfigurationRecord dcr;
+					if (dcr.Parse(MakeConstArrayView(TempBuffer, BytesRemaining)))
+					{
+						ConfigurationRecord = MoveTemp(dcr);
+					}
 				}
 				FMemory::Free(TempBuffer);
 			}
 			return Error;
 		}
 
-		ElectraDecodersUtil::MPEG::FAVCDecoderConfigurationRecord		ConfigurationRecord;
+		TOptional<ElectraDecodersUtil::MPEG::H264::FAVCDecoderConfigurationRecord> ConfigurationRecord;
 	};
 
 	/**
@@ -2887,7 +2893,7 @@ namespace Electra
 		{
 		}
 
-		const ElectraDecodersUtil::MPEG::FHEVCDecoderConfigurationRecord& GetDecoderConfigurationRecord() const
+		const TOptional<ElectraDecodersUtil::MPEG::H265::FHEVCDecoderConfigurationRecord>& GetDecoderConfigurationRecord() const
 		{
 			return ConfigurationRecord;
 		}
@@ -2903,7 +2909,7 @@ namespace Electra
 			uint32 BytesRemaining = BoxSize - (ParseInfo->Reader()->GetCurrentReadOffset() - StartOffset);
 			if (BytesRemaining)
 			{
-				void* TempBuffer = FMemory::Malloc(BytesRemaining);
+				uint8* TempBuffer = (uint8*)FMemory::Malloc(BytesRemaining);
 				if (!TempBuffer)
 				{
 					return UEMEDIA_ERROR_OOM;
@@ -2911,14 +2917,18 @@ namespace Electra
 				Error = ParseInfo->Reader()->ReadBytes(TempBuffer, BytesRemaining);
 				if (Error == UEMEDIA_ERROR_OK)
 				{
-					ConfigurationRecord.SetRawData(TempBuffer, BytesRemaining);
+					ElectraDecodersUtil::MPEG::H265::FHEVCDecoderConfigurationRecord dcr;
+					if (dcr.Parse(MakeConstArrayView(TempBuffer, BytesRemaining)))
+					{
+						ConfigurationRecord = MoveTemp(dcr);
+					}
 				}
 				FMemory::Free(TempBuffer);
 			}
 			return Error;
 		}
 
-		ElectraDecodersUtil::MPEG::FHEVCDecoderConfigurationRecord ConfigurationRecord;
+		TOptional<ElectraDecodersUtil::MPEG::H265::FHEVCDecoderConfigurationRecord> ConfigurationRecord;
 	};
 
 
@@ -3469,12 +3479,17 @@ namespace Electra
 				switch(SampleType)
 				{
 					case ESampleType::Video:
+					{
 						NextBox = new FMP4BoxVisualSampleEntry(ChildBoxType, ChildBoxSize, BoxStartOffset, BoxDataOffset, true);
 						break;
+					}
 					case ESampleType::Audio:
+					{
 						NextBox = new FMP4BoxAudioSampleEntry(ChildBoxType, ChildBoxSize, BoxStartOffset, BoxDataOffset, true, Version);
 						break;
+					}
 					case ESampleType::Subtitles:
+					{
 						if (ChildBoxType == FMP4Box::kSample_tx3g)
 						{
 							NextBox = new FMP4BoxTX3GSampleEntry(ChildBoxType, ChildBoxSize, BoxStartOffset, BoxDataOffset, true, Version);
@@ -3506,7 +3521,9 @@ namespace Electra
 							NextBox = new FMP4BoxIgnored(ChildBoxType, ChildBoxSize, BoxStartOffset, BoxDataOffset, true);
 							break;
 						*/
-				case ESampleType::Timecode:
+					}
+					case ESampleType::Timecode:
+					{
 						if (ChildBoxType == FMP4Box::kSample_tmcd)
 						{
 							NextBox = new FMP4BoxTimecodeSampleEntry(ChildBoxType, ChildBoxSize, BoxStartOffset, BoxDataOffset, true);
@@ -3514,12 +3531,15 @@ namespace Electra
 						else
 						{
 							SampleType = ESampleType::Ignored; // Ignoring non-tmcd samples.
-							NextBox = new FMP4BoxIgnored(ChildBoxType, ChildBoxSize, BoxStartOffset, BoxDataOffset, true);	
+							NextBox = new FMP4BoxIgnored(ChildBoxType, ChildBoxSize, BoxStartOffset, BoxDataOffset, true);
 						}
 						break;
+					}
 					default:
+					{
 						NextBox = new FMP4BoxIgnored(ChildBoxType, ChildBoxSize, BoxStartOffset, BoxDataOffset, true);
 						break;
+					}
 				}
 				AddChildBox(NextBox);
 				RETURN_IF_ERROR(ParseInfo->ReadAndParseNextBox(NextBox));
@@ -5166,7 +5186,7 @@ namespace Electra
 		{
 			return TrackIds;
 		}
-		
+
 	private:
 		FMP4BoxTrackReferenceEntry() = delete;
 		FMP4BoxTrackReferenceEntry(const FMP4BoxTrackReferenceEntry&) = delete;
@@ -5179,7 +5199,10 @@ namespace Electra
 			const int32 BytesRemaining = static_cast<int32>(BoxSize - (NowAt - StartOffset));
 			const int32 NumTrackIds = BytesRemaining/sizeof(int32);
 			TrackIds.AddZeroed(NumTrackIds);
-			RETURN_IF_ERROR(ParseInfo->Reader()->ReadBytes(TrackIds.GetData(), NumTrackIds*sizeof(int32)));
+			for(int32 i=0; i<NumTrackIds; ++i)
+			{
+				RETURN_IF_ERROR(ParseInfo->Reader()->Read(TrackIds[i]));
+			}
 			return Error;
 		}
 
@@ -5190,7 +5213,7 @@ namespace Electra
 	 * 'tref' box. Track Reference Atom
 	 *  https://developer.apple.com/documentation/quicktime-file-format/track_reference_atom
 	 *
-	 *  Contains a list of relational track references per type. 
+	 *  Contains a list of relational track references per type.
 	 */
 	class FMP4BoxTREF : public FMP4BoxBasic
 	{
@@ -5204,7 +5227,7 @@ namespace Electra
 
 		const TArray<int32>& GetTrackIdsForType(IParserISO14496_12::FBoxType InReferenceType) const
 		{
-			for (const FMP4Box* ChildBox : ChildBoxes)
+			for(const FMP4Box* ChildBox : ChildBoxes)
 			{
 				if (ChildBox->GetType() == InReferenceType)
 				{
@@ -5214,16 +5237,16 @@ namespace Electra
 			static const TArray<int32> Empty;
 			return Empty;
 		}
-		
+
 	private:
 		FMP4BoxTREF() = delete;
 		FMP4BoxTREF(const FMP4BoxTREF&) = delete;
-	
+
 	protected:
 		virtual UEMediaError ReadAndParseAttributes(FMP4ParseInfo* ParseInfo) override
 		{
 			UEMediaError Error = UEMEDIA_ERROR_OK;
-			
+
 			// Read the list of TrackReference atoms
 			int32 BytesRemaining = static_cast<int32>(BoxSize - (ParseInfo->Reader()->GetCurrentReadOffset() - StartOffset));
 			while (BytesRemaining > 0)
@@ -6195,17 +6218,17 @@ namespace Electra
 			TArray<ITrack::FProducerReferenceTime> PRFTBoxes;
 			// Future boxes: sbgp, sgpd, subs
 
-			FStreamCodecInformation					CodecInformation;
+			FStreamCodecInformation CodecInformation;
 
-			int64									EmptyEditDurationInMediaTimeUnits = 0;
-			int64									CompositionTimeEditOffset = 0;
+			int64 EmptyEditDurationInMediaTimeUnits = 0;
+			int64 CompositionTimeEditOffset = 0;
 
-			ElectraDecodersUtil::MPEG::FAVCDecoderConfigurationRecord	CodecSpecificDataAVC;
-			ElectraDecodersUtil::MPEG::FHEVCDecoderConfigurationRecord	CodecSpecificDataHEVC;
-			MPEG::FESDescriptor						CodecSpecificDataMP4A;
-			TArray<uint8>							CodecSpecificDataRAW;
-			FBitrateInfo							BitrateInfo;
-			TOptional<FTimecode>					StartTimecode;
+			ElectraDecodersUtil::MPEG::H264::FAVCDecoderConfigurationRecord CodecSpecificDataAVC;
+			ElectraDecodersUtil::MPEG::H265::FHEVCDecoderConfigurationRecord CodecSpecificDataHEVC;
+			MPEG::FESDescriptor CodecSpecificDataMP4A;
+			TArray<uint8> CodecSpecificDataRAW;
+			FBitrateInfo BitrateInfo;
+			TOptional<FTimecode> StartTimecode;
 		};
 
 
@@ -7466,23 +7489,20 @@ namespace Electra
 					case FMP4Box::kBox_avcC:
 					{
 						const FMP4BoxAVCC* AVCCBox = static_cast<const FMP4BoxAVCC*>(SampleEntry);
-						Track->CodecSpecificDataAVC = AVCCBox->GetDecoderConfigurationRecord();
-						bool bOk = Track->CodecSpecificDataAVC.Parse();
-						check(bOk);
-						if (bOk)
+						if (AVCCBox->GetDecoderConfigurationRecord().IsSet())
 						{
+							Track->CodecSpecificDataAVC = AVCCBox->GetDecoderConfigurationRecord().GetValue();
 							Track->CodecInformation.SetStreamType(EStreamType::Video);
 							Track->CodecInformation.SetCodec(FStreamCodecInformation::ECodec::H264);
-							Track->CodecInformation.SetCodecSpecificData(Track->CodecSpecificDataRAW);
 							Track->CodecInformation.SetStreamLanguageCode(Track->GetLanguage());
-							if (Track->CodecSpecificDataAVC.GetNumberOfSPS() == 0)
+							if (Track->CodecSpecificDataAVC.GetSequenceParameterSets().IsEmpty())
 							{
 								if (InAvcType == 3)
 								{
 									Track->CodecInformation.SetProfile(Track->CodecSpecificDataAVC.GetAVCProfileIndication());
 									Track->CodecInformation.SetProfileLevel(Track->CodecSpecificDataAVC.GetAVCLevelIndication());
 									Track->CodecInformation.SetProfileConstraints(Track->CodecSpecificDataAVC.GetProfileCompatibility());
-									Track->CodecInformation.SetCodecSpecifierRFC6381(FString::Printf(TEXT("avc%d.%02x%02x%02x"), InAvcType, Track->CodecSpecificDataAVC.GetAVCProfileIndication(), Track->CodecSpecificDataAVC.GetProfileCompatibility(), Track->CodecSpecificDataAVC.GetAVCLevelIndication()));
+									Track->CodecInformation.SetCodecSpecifierRFC6381(Track->CodecSpecificDataAVC.GetCodecSpecifierRFC6381(TEXT("avc3")));
 									if (Track->TKHDBox)
 									{
 										Track->CodecInformation.SetResolution(FStreamCodecInformation::FResolution(Track->TKHDBox->GetWidth(), Track->TKHDBox->GetHeight()));
@@ -7499,21 +7519,28 @@ namespace Electra
 							}
 							else
 							{
-								const ElectraDecodersUtil::MPEG::FISO14496_10_seq_parameter_set_data& sps = Track->CodecSpecificDataAVC.GetParsedSPS(0);
-								int32 CropL, CropR, CropT, CropB;
-								sps.GetCrop(CropL, CropR, CropT, CropB);
-								Track->CodecInformation.SetResolution(FStreamCodecInformation::FResolution(sps.GetWidth() - CropL - CropR, sps.GetHeight() - CropT - CropB));
-								Track->CodecInformation.SetCrop(FStreamCodecInformation::FCrop(CropL, CropT, CropR, CropB));
-								FStreamCodecInformation::FAspectRatio ar;
-								sps.GetAspect(ar.Width, ar.Height);
-								Track->CodecInformation.SetAspectRatio(ar);
-								Track->CodecInformation.SetFrameRate(sps.GetTiming().Denom ? FTimeFraction(sps.GetTiming().Num, sps.GetTiming().Denom) : FTimeFraction());
-								Track->CodecInformation.SetProfile(sps.profile_idc);
-								Track->CodecInformation.SetProfileLevel(sps.level_idc);
-								uint8 Constraints = (sps.constraint_set0_flag << 7) | (sps.constraint_set1_flag << 6) | (sps.constraint_set2_flag << 5) | (sps.constraint_set3_flag << 4) | (sps.constraint_set4_flag << 3) | (sps.constraint_set5_flag << 2);
-								Track->CodecInformation.SetProfileConstraints(Constraints);
-								Track->CodecInformation.SetCodecSpecifierRFC6381(FString::Printf(TEXT("avc%d.%02x%02x%02x"), InAvcType, sps.profile_idc, Constraints, sps.level_idc));
-								bGotVideoFormat = true;
+								ElectraDecodersUtil::MPEG::H264::FSequenceParameterSet sps;
+								if (ElectraDecodersUtil::MPEG::H264::ParseSequenceParameterSet(sps, Track->CodecSpecificDataAVC.GetSequenceParameterSets()[0]))
+								{
+									int32 CropL, CropR, CropT, CropB;
+									sps.GetCrop(CropL, CropR, CropT, CropB);
+									Track->CodecInformation.SetResolution(FStreamCodecInformation::FResolution(sps.GetWidth() - CropL - CropR, sps.GetHeight() - CropT - CropB));
+									Track->CodecInformation.SetCrop(FStreamCodecInformation::FCrop(CropL, CropT, CropR, CropB));
+									FStreamCodecInformation::FAspectRatio ar;
+									sps.GetAspect(ar.Width, ar.Height);
+									Track->CodecInformation.SetAspectRatio(ar);
+									Track->CodecInformation.SetFrameRate(sps.GetTiming().Denom ? FTimeFraction(sps.GetTiming().Num, sps.GetTiming().Denom) : FTimeFraction());
+									Track->CodecInformation.SetProfile(sps.profile_idc);
+									Track->CodecInformation.SetProfileLevel(sps.level_idc);
+									uint8 Constraints = (sps.constraint_set0_flag << 7) | (sps.constraint_set1_flag << 6) | (sps.constraint_set2_flag << 5) | (sps.constraint_set3_flag << 4) | (sps.constraint_set4_flag << 3) | (sps.constraint_set5_flag << 2);
+									Track->CodecInformation.SetProfileConstraints(Constraints);
+									Track->CodecInformation.SetCodecSpecifierRFC6381(FString::Printf(TEXT("avc%d.%02x%02x%02x"), InAvcType, sps.profile_idc, Constraints, sps.level_idc));
+									bGotVideoFormat = true;
+								}
+								else
+								{
+									return UEMEDIA_ERROR_FORMAT_ERROR;
+								}
 							}
 						}
 						else
@@ -7585,62 +7612,54 @@ namespace Electra
 					case FMP4Box::kBox_hvcC:
 					{
 						const FMP4BoxHVCC* HVCCBox = static_cast<const FMP4BoxHVCC*>(SampleEntry);
-						Track->CodecSpecificDataHEVC = HVCCBox->GetDecoderConfigurationRecord();
-						bool bOk = Track->CodecSpecificDataHEVC.Parse();
-						check(bOk);
-						if (bOk)
+						if (HVCCBox->GetDecoderConfigurationRecord().IsSet())
 						{
-							ElectraDecodersUtil::MPEG::FISO23008_2_seq_parameter_set_data TempSPS;
+							Track->CodecSpecificDataHEVC = HVCCBox->GetDecoderConfigurationRecord().GetValue();
 							Track->CodecInformation.SetStreamType(EStreamType::Video);
 							Track->CodecInformation.SetCodec(FStreamCodecInformation::ECodec::H265);
-							Track->CodecInformation.SetCodecSpecificData(Track->CodecSpecificDataRAW);
 							Track->CodecInformation.SetStreamLanguageCode(Track->GetLanguage());
-							const ElectraDecodersUtil::MPEG::FISO23008_2_seq_parameter_set_data* sps = nullptr;
-							const TCHAR* HevType = TEXT("hvc1");
-							if (Track->CodecSpecificDataHEVC.GetNumberOfSPS() == 0)
+							Track->CodecInformation.SetProfileSpace(Track->CodecSpecificDataHEVC.GetGeneralProfileSpace());
+							Track->CodecInformation.SetProfileTier(Track->CodecSpecificDataHEVC.GetGeneralTierFlag());
+							Track->CodecInformation.SetProfile(Track->CodecSpecificDataHEVC.GetGeneralProfileIDC());
+							Track->CodecInformation.SetProfileLevel(Track->CodecSpecificDataHEVC.GetGeneralLevelIDC());
+							Track->CodecInformation.SetProfileConstraints(Track->CodecSpecificDataHEVC.GetGeneralConstraintIndicatorFlags());
+							Track->CodecInformation.SetProfileCompatibilityFlags(Track->CodecSpecificDataHEVC.GetGeneralProfileCompatibilityFlags());
+
+							if (Track->CodecSpecificDataHEVC.GetSequenceParameterSets().IsEmpty())
+							{
+								if (Track->TKHDBox)
+								{
+									Track->CodecInformation.SetResolution(FStreamCodecInformation::FResolution(Track->TKHDBox->GetWidth(), Track->TKHDBox->GetHeight()));
+									Track->CodecInformation.SetAspectRatio(FStreamCodecInformation::FAspectRatio(1, 1));
+									// Crop values are not easily determined as there is no single macroblock (coding tree unit) size.
+									// Do not set any.
+										//Track->CodecInformation.SetCrop(FStreamCodecInformation::FCrop(0, 0, Track->TKHDBox->GetWidth() % 16, Track->TKHDBox->GetHeight() % 16));
+									Track->CodecInformation.SetCodecSpecifierRFC6381(Track->CodecSpecificDataHEVC.GetCodecSpecifierRFC6381(TEXT("hev1")));
+									bGotVideoFormat = true;
+								}
+								else
+								{
+									return UEMEDIA_ERROR_FORMAT_ERROR;
+								}
+							}
+							else if (!Track->CodecSpecificDataHEVC.GetParsedSequenceParameterSets().IsEmpty())
 							{
-								TempSPS.general_profile_space = Track->CodecSpecificDataHEVC.GetGeneralProfileSpace();
-								TempSPS.general_tier_flag = Track->CodecSpecificDataHEVC.GetGeneralTierFlag();
-								TempSPS.general_profile_idc = Track->CodecSpecificDataHEVC.GetGeneralProfileIDC();
-								TempSPS.general_level_idc = Track->CodecSpecificDataHEVC.GetGeneralLevelIDC();
-								TempSPS.general_profile_compatibility_flag = Track->CodecSpecificDataHEVC.GetGeneralProfileCompatibilityFlags();
-								uint64 flags = Track->CodecSpecificDataHEVC.GetGeneralConstraintIndicatorFlags();
-								TempSPS.general_reserved_zero_44bits = (flags & 0xfffffffffffUL);
-								flags >>= 44;
-								TempSPS.general_frame_only_constraint_flag = flags & 1;
-								flags >>= 1;
-								TempSPS.general_non_packed_constraint_flag = flags & 1;
-								flags >>= 1;
-								TempSPS.general_interlaced_source_flag = flags & 1;
-								flags >>= 1;
-								TempSPS.general_progressive_source_flag = flags & 1;
-								TempSPS.pic_width_in_luma_samples = Track->TKHDBox->GetWidth();
-								TempSPS.pic_height_in_luma_samples = Track->TKHDBox->GetHeight();
-								TempSPS.bit_depth_luma_minus8 = Track->CodecSpecificDataHEVC.GetBitDepthLumaMinus8();
-								TempSPS.bit_depth_chroma_minus8 = Track->CodecSpecificDataHEVC.GetBitDepthChromaMinus8();
-								sps = &TempSPS;
-								HevType = TEXT("hev1");
+								const ElectraDecodersUtil::MPEG::H265::FSequenceParameterSet& sps = Track->CodecSpecificDataHEVC.GetParsedSequenceParameterSets()[0];
+								int32 CropL, CropR, CropT, CropB;
+								sps.GetCrop(CropL, CropR, CropT, CropB);
+								Track->CodecInformation.SetResolution(FStreamCodecInformation::FResolution(sps.GetWidth() - CropL - CropR, sps.GetHeight() - CropT - CropB));
+								Track->CodecInformation.SetCrop(FStreamCodecInformation::FCrop(CropL, CropT, CropR, CropB));
+								FStreamCodecInformation::FAspectRatio ar;
+								sps.GetAspect(ar.Width, ar.Height);
+								Track->CodecInformation.SetAspectRatio(ar);
+								Track->CodecInformation.SetFrameRate(sps.GetTiming().Denom ? FTimeFraction(sps.GetTiming().Num, sps.GetTiming().Denom) : FTimeFraction());
+								Track->CodecInformation.SetCodecSpecifierRFC6381(sps.GetRFC6381(TEXT("hvc1")));
+								bGotVideoFormat = true;
 							}
 							else
 							{
-								sps = &Track->CodecSpecificDataHEVC.GetParsedSPS(0);
+								return UEMEDIA_ERROR_FORMAT_ERROR;
 							}
-							int32 CropL, CropR, CropT, CropB;
-							sps->GetCrop(CropL, CropR, CropT, CropB);
-							Track->CodecInformation.SetResolution(FStreamCodecInformation::FResolution(sps->GetWidth() - CropL - CropR, sps->GetHeight() - CropT - CropB));
-							Track->CodecInformation.SetCrop(FStreamCodecInformation::FCrop(CropL, CropT, CropR, CropB));
-							FStreamCodecInformation::FAspectRatio ar;
-							sps->GetAspect(ar.Width, ar.Height);
-							Track->CodecInformation.SetAspectRatio(ar);
-							Track->CodecInformation.SetFrameRate(sps->GetTiming().Denom ? FTimeFraction(sps->GetTiming().Num, sps->GetTiming().Denom) : FTimeFraction());
-							Track->CodecInformation.SetProfileSpace(sps->general_profile_space);
-							Track->CodecInformation.SetProfileTier(sps->general_tier_flag);
-							Track->CodecInformation.SetProfile(sps->general_profile_idc);
-							Track->CodecInformation.SetProfileLevel(sps->general_level_idc);
-							Track->CodecInformation.SetProfileConstraints(sps->GetConstraintFlags());
-							Track->CodecInformation.SetProfileCompatibilityFlags(sps->general_profile_compatibility_flag);
-							Track->CodecInformation.SetCodecSpecifierRFC6381(sps->GetRFC6381(HevType));
-							bGotVideoFormat = true;
 						}
 						else
 						{
@@ -7714,7 +7733,6 @@ namespace Electra
 				{
 					Track->CodecInformation.SetStreamType(EStreamType::Audio);
 					Track->CodecInformation.SetCodec(FStreamCodecInformation::ECodec::AAC);
-					Track->CodecInformation.SetCodecSpecificData(Track->CodecSpecificDataRAW);
 					Track->CodecInformation.SetStreamLanguageCode(Track->GetLanguage());
 
 					ElectraDecodersUtil::MPEG::FAACDecoderConfigurationRecord ConfigRecord;
@@ -7758,7 +7776,6 @@ namespace Electra
 					Track->CodecInformation.SetCodec(FStreamCodecInformation::ECodec::Audio4CC);
 					Track->CodecInformation.SetCodec4CC(Utils::Make4CC('m','p','g','a'));
 					Track->CodecInformation.SetProfile(1);
-					Track->CodecInformation.SetCodecSpecificData(Track->CodecSpecificDataRAW);
 					Track->CodecInformation.SetStreamLanguageCode(Track->GetLanguage());
 					Track->CodecInformation.SetCodecSpecifierRFC6381(FString::Printf(TEXT("mp4a.6b")));
 					Track->CodecInformation.SetSamplingRate((int32) AudioSampleEntry->GetSampleRate());
@@ -7925,7 +7942,7 @@ namespace Electra
 	{
 		check(InSampleBox->GetType() == FMP4Box::kSample_tmcd);
 		const FMP4BoxTimecodeSampleEntry* TMCDBox = static_cast<const FMP4BoxTimecodeSampleEntry*>(InSampleBox);
-				
+
 		const bool bDropFrame = TMCDBox->GetFlags() & FMP4BoxTimecodeSampleEntry::FFlags::DropFrame ? true : false;
 		const FFrameRate FrameRate = bDropFrame ? FFrameRate(TMCDBox->GetTimeScale(), TMCDBox->GetFrameDuration()) : FFrameRate(TMCDBox->GetNumberOfFrames(), 1);
 
@@ -7940,7 +7957,7 @@ namespace Electra
 		{
 			return UEMEDIA_ERROR_NOT_SUPPORTED;
 		}
-		
+
 		// We will need the stco, stsc and stsz boxes. We only need the first entry.
 		if (!InTrack || !InTrack->STCOBox || InTrack->STCOBox->GetNumberOfEntries() == 0
 			|| !InTrack->STSCBox || InTrack->STSCBox->GetNumberOfEntries() == 0
@@ -7953,15 +7970,15 @@ namespace Electra
 		const FMP4BoxTimecodeSampleEntry* TMCDBox = static_cast<const FMP4BoxTimecodeSampleEntry*>(InSampleBox);
 
 		// Obtain the first sample only -> Fetch the first chunk and read the first sample.
-		
+
 		const int64 ChunkOffset = InTrack->STCOBox->GetEntry(0);
 		//const uint32 SamplesPerChunk = InTrack->STSCBox->GetEntry(0).SamplesPerChunk;
 		const uint32 SampleSizeInBytes = InTrack->STSZBox->GetSampleSize(0);
 		// Note: since we need only the first sample, load only "sample size" to get the first one.
 		const uint32 ChunkSize = SampleSizeInBytes; //SamplesPerChunk * SampleSizeInBytes;
-		
+
 		const FString ManifestURL = InPlayerSession->GetManifestReader()->GetURL();
-		
+
 		UtilsMP4::FMP4ChunkLoader ChunkLoader;
 
 		// Propagate the worker thread's abort flag.
@@ -7980,7 +7997,7 @@ namespace Electra
 				// Note: ignoring counter flag. Even if set to 0, which would mean QT int32 format (HHMMSSFF),
 				// the timecode is still a frame number format (as if Counter set to 1).
 				const FFrameNumber FrameNumber(static_cast<int32>(MEDIA_FROM_BIG_ENDIAN(Data[0])));
-				
+
 				const bool bDropFrame = TMCDBox->GetFlags() & FMP4BoxTimecodeSampleEntry::FFlags::DropFrame ? true : false;
 				const bool bRollover = TMCDBox->GetFlags() & FMP4BoxTimecodeSampleEntry::FFlags::Max24Hour ? true : false;
 
@@ -8000,7 +8017,7 @@ namespace Electra
 				return UEMEDIA_ERROR_OK;
 			}
 		}
-		return UEMEDIA_ERROR_NOT_SUPPORTED;	
+		return UEMEDIA_ERROR_NOT_SUPPORTED;
 	}
 
 	UEMediaError FParserISO14496_12::ParseGenericSampleType(IPlayerSessionServices* PlayerSession, FTrack* Track, const FMP4BoxSTSD* STSDBox, const FMP4BoxFRMA* FRMABox)
@@ -8770,7 +8787,7 @@ namespace Electra
 		{
 			return UEMEDIA_ERROR_OK;
 		}
-		
+
 		// Read timecode sample data for timecode tracks.
 		for(int32 TrackIndex = 0; TrackIndex < ParsedTrackInfo->GetNumberOfTracks(); ++TrackIndex)
 		{
@@ -8780,7 +8797,7 @@ namespace Electra
 				ReadTMCDSample(InPlayerSession, Track, Track->STSDBox->FindBox(FMP4Box::kSample_tmcd), InCancellationCheckDelegate);
 			}
 		}
-			
+
 		// Copy timecode meta data from tmcd tracks to the related video tracks.
 		for(int32 TrackIndex = 0; TrackIndex < ParsedTrackInfo->GetNumberOfTracks(); ++TrackIndex)
 		{
