Added guard to prevent FStream::SourceBufferInfo nullptr access
Cherry picked change from Release-31.41

#rnx

diff --git a/Engine/Plugins/Media/ElectraPlayer/Source/ElectraPlayerRuntime/Private/Runtime/Player/DASH/StreamReaderDASH.cpp b/Engine/Plugins/Media/ElectraPlayer/Source/ElectraPlayerRuntime/Private/Runtime/Player/DASH/StreamReaderDASH.cpp
--- a/Engine/Plugins/Media/ElectraPlayer/Source/ElectraPlayerRuntime/Private/Runtime/Player/DASH/StreamReaderDASH.cpp
+++ b/Engine/Plugins/Media/ElectraPlayer/Source/ElectraPlayerRuntime/Private/Runtime/Player/DASH/StreamReaderDASH.cpp
@@ -1029,6 +1029,13 @@ void FStreamReaderDASH::FStreamHandler::HandleRequestMP4()
 	ActiveTrackData.Reset();
 	ActiveTrackData.CSD = MakeSharedTS<FAccessUnit::CodecData>();
 	// Copy the source buffer info into a new instance and set the playback sequence ID in it.
+	if (!ensure(Request->SourceBufferInfo))
+	{
+		LogMessage(IInfoLog::ELevel::Error, TEXT("SourceBuffer is a nullptr"));
+		bHasErrored = true;
+		return;
+	}
+
 	ActiveTrackData.BufferSourceInfo = MakeSharedTS<FBufferSourceInfo>(*Request->SourceBufferInfo);
 	ActiveTrackData.BufferSourceInfo->PlaybackSequenceID = Request->GetPlaybackSequenceID();
 	ActiveTrackData.StreamType = Request->GetType();
@@ -1898,6 +1905,13 @@ void FStreamReaderDASH::FStreamHandler::HandleRequestMKV()
 	ActiveTrackData.Reset();
 	ActiveTrackData.CSD = MakeSharedTS<FAccessUnit::CodecData>();
 	// Copy the source buffer info into a new instance and set the playback sequence ID in it.
+	if (!ensure(Request->SourceBufferInfo))
+	{
+		LogMessage(IInfoLog::ELevel::Error, TEXT("SourceBuffer is a nullptr"));
+		bHasErrored = true;
+		return;
+	}
+
 	ActiveTrackData.BufferSourceInfo = MakeSharedTS<FBufferSourceInfo>(*Request->SourceBufferInfo);
 	ActiveTrackData.BufferSourceInfo->PlaybackSequenceID = Request->GetPlaybackSequenceID();
 	// Video needs to be held back to recalculate the frame duration and possibly the DTS.
