ElectraPlayer: Making internal resource delegate accessible for external use (Windows only)

#rnx

diff --git a/Engine/Plugins/Media/ElectraPlayer/Source/ElectraPlayerPlugin/Private/ElectraPlayerPlugin.h b/Engine/Plugins/Media/ElectraPlayer/Source/ElectraPlayerPlugin/Private/ElectraPlayerPlugin.h
--- a/Engine/Plugins/Media/ElectraPlayer/Source/ElectraPlayerPlugin/Private/ElectraPlayerPlugin.h
+++ b/Engine/Plugins/Media/ElectraPlayer/Source/ElectraPlayerPlugin/Private/ElectraPlayerPlugin.h
@@ -109,6 +109,8 @@ public:
 	uint32 GetNewResourcesOnOpen() const override;
 
 	void SetMetadataChanged();
+
+	static IElectraPlayerResourceDelegate* PlatformCreateStaticPlayerResourceDelegate();
 private:
 	friend class FElectraPlayerPluginModule;
 	friend class FElectraPlayerResourceDelegate;
diff --git a/Engine/Plugins/Media/ElectraPlayer/Source/ElectraPlayerPlugin/Private/ElectraPlayerPluginModule.cpp b/Engine/Plugins/Media/ElectraPlayer/Source/ElectraPlayerPlugin/Private/ElectraPlayerPluginModule.cpp
--- a/Engine/Plugins/Media/ElectraPlayer/Source/ElectraPlayerPlugin/Private/ElectraPlayerPluginModule.cpp
+++ b/Engine/Plugins/Media/ElectraPlayer/Source/ElectraPlayerPlugin/Private/ElectraPlayerPluginModule.cpp
@@ -12,6 +12,7 @@
 #include "ParameterDictionary.h"
 #include "SimpleElectraAudioPlayer.h"
 #include "IElectraDecoderResourceDelegateBase.h"
+#include "IElectraPlayerDecoderResourceManager.h"
 
 #define LOCTEXT_NAMESPACE "ElectraPlayerPluginModule"
 
@@ -26,12 +27,12 @@ class FElectraPlayerPluginModule : public IElectraPlayerPluginModule
 public:
 	// IElectraPlayerPluginModule interface
 
-	virtual bool IsInitialized() const override
+	bool IsInitialized() const override
 	{
 		return bInitialized;
 	}
 
-	virtual TSharedPtr<IMediaPlayer, ESPMode::ThreadSafe> CreatePlayer(IMediaEventSink& EventSink) override
+	TSharedPtr<IMediaPlayer, ESPMode::ThreadSafe> CreatePlayer(IMediaEventSink& EventSink) override
 	{
 		if (bInitialized)
 		{
@@ -45,27 +46,61 @@ public:
 		return nullptr;
 	}
 
-	virtual void SendAnalyticMetrics(const TSharedPtr<IAnalyticsProviderET>& AnalyticsProvider, const FGuid& PlayerGuid) override
+	void SendAnalyticMetrics(const TSharedPtr<IAnalyticsProviderET>& AnalyticsProvider, const FGuid& PlayerGuid) override
 	{
 		SendAnalyticMetricsDelegate.Broadcast(AnalyticsProvider, PlayerGuid);
 	}
 
-	virtual void SendAnalyticMetricsPerMinute(const TSharedPtr<IAnalyticsProviderET>& AnalyticsProvider) override
+	void SendAnalyticMetricsPerMinute(const TSharedPtr<IAnalyticsProviderET>& AnalyticsProvider) override
 	{
 		SendAnalyticMetricsPerMinuteDelegate.Broadcast(AnalyticsProvider);
 		ISimpleElectraAudioPlayer::SendAnalyticMetrics(AnalyticsProvider);
 	}
 
-	virtual void ReportVideoStreamingError(const FGuid& PlayerGuid, const FString& LastError) override
+	void ReportVideoStreamingError(const FGuid& PlayerGuid, const FString& LastError) override
 	{
 		ReportVideoStreamingErrorDelegate.Broadcast(PlayerGuid, LastError);
 	}
 
-	virtual void ReportSubtitlesMetrics(const FGuid& PlayerGuid, const FString& URL, double ResponseTime, const FString& LastError)
+	void ReportSubtitlesMetrics(const FGuid& PlayerGuid, const FString& URL, double ResponseTime, const FString& LastError) override
 	{
 		ReportSubtitlesMetricsDelegate.Broadcast(PlayerGuid, URL, ResponseTime, LastError);
 	}
 
+	// Create a suitable video decoder resource delegate for and via the Electra Player runtime to be used with it by external means.
+	TSharedPtr<Electra::IVideoDecoderResourceDelegate, ESPMode::ThreadSafe> CreatePlatformVideoDecoderResourceDelegate() override
+	{
+	#if PLATFORM_WINDOWS
+		// This is a bit of a convoluted process due to this whole plugin/adapter/delegate nature.
+		class FDummyAdapter : public IElectraPlayerAdapterDelegate
+		{
+		public:
+			FDummyAdapter(TSharedPtr<IElectraPlayerResourceDelegate, ESPMode::ThreadSafe> InPlayerResourceDlg) : PlayerResourceDlg(MoveTemp(InPlayerResourceDlg)) { }
+			virtual ~FDummyAdapter() { }
+			Electra::FVariantValue QueryOptions(EOptionType Type, const Electra::FVariantValue& Param = Electra::FVariantValue()) override { return Electra::FVariantValue(); }
+			void BlobReceived(const TSharedPtr<TArray<uint8>, ESPMode::ThreadSafe>& InBlobData, EBlobResultType InResultType, int32 InResultCode, const Electra::FParamDict* InExtraInfo) override { }
+			void SendMediaEvent(EPlayerEvent Event) override { }
+			void OnVideoFlush() override { }
+			void OnAudioFlush() override { }
+			void OnSubtitleFlush() override { }
+			void PresentVideoFrame(const FVideoDecoderOutputPtr& InVideoFrame) override { }
+			void PresentAudioFrame(const IAudioDecoderOutputPtr& InAudioFrame) override { }
+			void PresentSubtitleSample(const ISubtitleDecoderOutputPtr& InSubtitleSample) override { }
+			void PresentMetadataSample(const IMetaDataDecoderOutputPtr& InMetadataSample) override { }
+			bool CanReceiveVideoSamples(int32 NumFrames) override { return false; }
+			bool CanReceiveAudioSamples(int32 NumFrames) override { return false; }
+			void PrepareForDecoderShutdown() override { }
+			FString GetVideoAdapterName() const override { return FString(); }
+			TSharedPtr<IElectraPlayerResourceDelegate, ESPMode::ThreadSafe> GetResourceDelegate() const override { return PlayerResourceDlg; }
+		private:
+			TSharedPtr<IElectraPlayerResourceDelegate, ESPMode::ThreadSafe> PlayerResourceDlg;
+		};
+		return FElectraPlayerDecoderResourceManager::CreatePlatformVideoDecoderResourceDelegate(MakeShareable(new FDummyAdapter(MakeShareable(FElectraPlayerPlugin::PlatformCreateStaticPlayerResourceDelegate()))));
+	#else
+		return nullptr;
+	#endif
+	}
+
 public:
 	// IModuleInterface interface
 	static void GetDynamicRHIInfo(void** OutGDynamicRHI, int64* OutGDynamicRHIType)
diff --git a/Engine/Plugins/Media/ElectraPlayer/Source/ElectraPlayerPlugin/Private/Windows/ElectraPlayerPlugin_Windows.cpp b/Engine/Plugins/Media/ElectraPlayer/Source/ElectraPlayerPlugin/Private/Windows/ElectraPlayerPlugin_Windows.cpp
--- a/Engine/Plugins/Media/ElectraPlayer/Source/ElectraPlayerPlugin/Private/Windows/ElectraPlayerPlugin_Windows.cpp
+++ b/Engine/Plugins/Media/ElectraPlayer/Source/ElectraPlayerPlugin/Private/Windows/ElectraPlayerPlugin_Windows.cpp
@@ -28,4 +28,9 @@ IElectraPlayerResourceDelegate* FElectraPlayerPlugin::PlatformCreatePlayerResour
 	return new FPlayerResourceDelegateWindows();
 }
 
+IElectraPlayerResourceDelegate* FElectraPlayerPlugin::PlatformCreateStaticPlayerResourceDelegate()
+{
+	return new FPlayerResourceDelegateWindows();
+}
+
 #endif // PLATFORM_WINDOWS
diff --git a/Engine/Plugins/Media/ElectraPlayer/Source/ElectraPlayerPlugin/Public/IElectraPlayerPluginModule.h b/Engine/Plugins/Media/ElectraPlayer/Source/ElectraPlayerPlugin/Public/IElectraPlayerPluginModule.h
--- a/Engine/Plugins/Media/ElectraPlayer/Source/ElectraPlayerPlugin/Public/IElectraPlayerPluginModule.h
+++ b/Engine/Plugins/Media/ElectraPlayer/Source/ElectraPlayerPlugin/Public/IElectraPlayerPluginModule.h
@@ -17,6 +17,10 @@ class IMediaEventSink;
 class IMediaPlayer;
 class IElectraPlayerDataCache;
 class UObject;
+namespace Electra
+{
+class IVideoDecoderResourceDelegate;
+}
 
 /**
  * This class is used to get safe access to an IMediaOptions interface.
@@ -57,7 +61,7 @@ public:
 				SafeMediaOptionInterface->Unlock();
 			}
 		}
-	
+
 	private:
 		TSharedPtr<IElectraSafeMediaOptionInterface, ESPMode::ThreadSafe> SafeMediaOptionInterface;
 	};
@@ -150,7 +154,8 @@ public:
 	virtual void ReportVideoStreamingError(const FGuid& PlayerGuid, const FString& LastError) = 0;
 	virtual void ReportSubtitlesMetrics(const FGuid& PlayerGuid, const FString& URL, double ResponseTime, const FString& LastError) = 0;
 
-
+	// Create a suitable video decoder resource delegate for and via the Electra Player runtime to be used with it by external means.
+	virtual TSharedPtr<Electra::IVideoDecoderResourceDelegate, ESPMode::ThreadSafe> CreatePlatformVideoDecoderResourceDelegate() = 0;
 public:
 	virtual ~IElectraPlayerPluginModule() { }
 };
