ElectraPlayer: Fixed file: scheme handling on Linux based platforms

#rnx

diff --git a/Engine/Plugins/Media/ElectraPlayer/Source/ElectraPlayerRuntime/Private/Runtime/Player/AdaptiveStreamingPlayerManifest.cpp b/Engine/Plugins/Media/ElectraPlayer/Source/ElectraPlayerRuntime/Private/Runtime/Player/AdaptiveStreamingPlayerManifest.cpp
--- a/Engine/Plugins/Media/ElectraPlayer/Source/ElectraPlayerRuntime/Private/Runtime/Player/AdaptiveStreamingPlayerManifest.cpp
+++ b/Engine/Plugins/Media/ElectraPlayer/Source/ElectraPlayerRuntime/Private/Runtime/Player/AdaptiveStreamingPlayerManifest.cpp
@@ -41,7 +41,6 @@ FString FixLocalFileSchemeURL(const FString& InURL)
 	if (URL.StartsWith("file:", ESearchCase::IgnoreCase))
 	{
 		URL.ReplaceCharInline(TCHAR('\\'), TCHAR('/'));
-		URL.ReplaceInline(TEXT(":///"), TEXT("://"));
 		URL.ReplaceInline(TEXT(" "), TEXT("%20"));
 	}
 	return URL;
