Electra Protron: Added first version of Protron

#rnx

diff --git a/Engine/Plugins/Media/ElectraPlayer/ElectraPlayer.uplugin b/Engine/Plugins/Media/ElectraPlayer/ElectraPlayer.uplugin
--- a/Engine/Plugins/Media/ElectraPlayer/ElectraPlayer.uplugin
+++ b/Engine/Plugins/Media/ElectraPlayer/ElectraPlayer.uplugin
@@ -1,69 +1,91 @@
 {
-	"FileVersion": 1,
-	"Version": 1,
-	"VersionName": "1.0",
-	"FriendlyName": "Electra Player",
-	"Description": "Next-generation playback capability.",
-	"Category": "Media Players",
-	"CreatedBy": "Epic Games, Inc.",
-	"CreatedByURL": "https://epicgames.com",
-	"DocsURL": "https://docs.unrealengine.com/en-US/Engine/MediaFramework/Overview",
-	"MarketplaceURL": "",
-	"SupportURL": "",
-	"EnabledByDefault": false,
-	"CanContainContent": false,
-	"IsBetaVersion": false,
-	"Installed": false,
-	"Modules": [
-		{
-			"Name": "ElectraPlayerPluginHandler",
-			"Type": "RuntimeNoCommandlet",
-			"LoadingPhase": "PreLoadingScreen",
-			"TargetDenyList": [ "Server" ]
-		},
-		{
-			"Name": "ElectraPlayerRuntime",
-			"Type": "RuntimeNoCommandlet",
-			"LoadingPhase": "PreLoadingScreen",
-			"PlatformAllowList": [ "Win64", "Mac", "IOS", "TVOS", "Android", "Linux" ],
-			"TargetDenyList": [ "Server" ]
-		},
-		{
-			"Name": "ElectraPlayerPlugin",
-			"Type": "RuntimeNoCommandlet",
-			"LoadingPhase": "PreLoadingScreen",
-			"PlatformAllowList": [ "Win64", "Mac", "IOS", "TVOS", "Android", "Linux" ],
-			"TargetDenyList": [ "Server" ]
-		},
-		{
-			"Name": "ElectraPlayerFactory",
-			"Type": "Editor",
-			"LoadingPhase": "PostEngineInit",
-			"PlatformAllowList": [ "Win64", "Mac", "Linux" ],
-			"TargetDenyList": [ "Server" ]
-		},
-		{
-			"Name": "ElectraPlayerFactory",
-			"Type": "RuntimeNoCommandlet",
-			"LoadingPhase": "PostEngineInit",
-			"PlatformAllowList": [ "Win64", "Mac", "IOS", "TVOS", "Android", "Linux" ],
-			"TargetDenyList": [ "Server" ]
-		}
+  "FileVersion": 1,
+  "Version": 1,
+  "VersionName": "1.0",
+  "FriendlyName": "Electra Player",
+  "Description": "Cross platform media player for local files and internet streaming.\nAlso provides optimized local mp4 file only player (Protron) for desktop machines.",
+  "Category": "Media Players",
+  "CreatedBy": "Epic Games, Inc.",
+  "CreatedByURL": "https://epicgames.com",
+  "DocsURL": "https://docs.unrealengine.com/en-US/Engine/MediaFramework/Overview",
+  "MarketplaceURL": "",
+  "SupportURL": "",
+  "EnabledByDefault": false,
+  "CanContainContent": false,
+  "IsBetaVersion": false,
+  "Installed": false,
+  "Modules": [
+    {
+      "Name": "ElectraPlayerPluginHandler",
+      "Type": "RuntimeNoCommandlet",
+      "LoadingPhase": "PreLoadingScreen",
+      "TargetDenyList": [ "Server" ]
+    },
+    {
+      "Name": "ElectraPlayerRuntime",
+      "Type": "RuntimeNoCommandlet",
+      "LoadingPhase": "PreLoadingScreen",
+      "PlatformAllowList": [ "Win64", "Mac", "IOS", "TVOS", "Android", "Linux" ],
+      "TargetDenyList": [ "Server" ]
+    },
+    {
+      "Name": "ElectraPlayerPlugin",
+      "Type": "RuntimeNoCommandlet",
+      "LoadingPhase": "PreLoadingScreen",
+      "PlatformAllowList": [ "Win64", "Mac", "IOS", "TVOS", "Android", "Linux" ],
+      "TargetDenyList": [ "Server" ]
+    },
+    {
+      "Name": "ElectraPlayerFactory",
+      "Type": "Editor",
+      "LoadingPhase": "PostEngineInit",
+      "PlatformAllowList": [ "Win64", "Mac", "Linux" ],
+      "TargetDenyList": [ "Server" ]
+    },
+    {
+      "Name": "ElectraPlayerFactory",
+      "Type": "RuntimeNoCommandlet",
+      "LoadingPhase": "PostEngineInit",
+      "PlatformAllowList": [ "Win64", "Mac", "IOS", "TVOS", "Android", "Linux" ],
+      "TargetDenyList": [ "Server" ]
+    },
+
+    {
+      "Name": "ElectraProtron",
+      "Type": "RuntimeNoCommandlet",
+      "LoadingPhase": "PreLoadingScreen",
+      "PlatformAllowList": [ "Win64", "Mac" ],
+      "TargetDenyList": [ "Server" ]
+    },
+    {
+      "Name": "ElectraProtronFactory",
+      "Type": "Editor",
+      "LoadingPhase": "PostEngineInit",
+      "PlatformAllowList": [ "Win64", "Mac" ],
+      "TargetDenyList": [ "Server" ]
+    },
+    {
+      "Name": "ElectraProtronFactory",
+      "Type": "RuntimeNoCommandlet",
+      "LoadingPhase": "PostEngineInit",
+      "PlatformAllowList": [ "Win64", "Mac" ],
+      "TargetDenyList": [ "Server" ]
+    }
+  ],
 
-	],
-	"Plugins": [
-		{
-			"Name": "ElectraUtil",
-			"Enabled": true
-		},
-		{
-			"Name": "ElectraCodecs",
-			"Enabled": true
-		},
-		{
-			"Name": "ElectraCDM",
-			"Enabled": true
-		},
+  "Plugins": [
+    {
+      "Name": "ElectraUtil",
+      "Enabled": true
+    },
+    {
+      "Name": "ElectraCodecs",
+      "Enabled": true
+    },
+    {
+      "Name": "ElectraCDM",
+      "Enabled": true
+    },
     {
       "Name": "ElectraSubtitles",
       "Enabled": true
@@ -76,5 +98,5 @@
       "Name": "D3D12VideoDecodersElectra",
       "Enabled": true
     }
-	]
+  ]
 }
diff --git a/Engine/Plugins/Media/ElectraPlayer/Source/ElectraProtron/ElectraProtron.Build.cs b/Engine/Plugins/Media/ElectraPlayer/Source/ElectraProtron/ElectraProtron.Build.cs
new file mode 100644
--- /dev/null
+++ b/Engine/Plugins/Media/ElectraPlayer/Source/ElectraProtron/ElectraProtron.Build.cs
@@ -0,0 +1,41 @@
+// Copyright Epic Games, Inc. All Rights Reserved.
+
+namespace UnrealBuildTool.Rules
+{
+	public class ElectraProtron : ModuleRules
+	{
+		public ElectraProtron(ReadOnlyTargetRules Target) : base(Target)
+		{
+			DynamicallyLoadedModuleNames.AddRange(
+				new string[] {
+					"Media",
+				});
+
+			PrivateDependencyModuleNames.AddRange(
+				new string[] {
+					"ElectraProtronFactory",
+					"Core",
+					"Engine",
+					"MediaUtils",
+					"RenderCore",
+					"RHI",
+					"ElectraBase",
+                    "ElectraDecoders",
+                    "ElectraPlayerPlugin",      // for early start and init of what we want to access
+                    "ElectraPlayerRuntime",     // sharing the resource manager
+					"ElectraSamples",           // and renderer
+				});
+
+			PrivateIncludePathModuleNames.AddRange(
+				new string[] {
+					"Media",
+				});
+
+			if (Target.Platform.IsInGroup(UnrealPlatformGroup.Windows))
+			{
+				PrivateDependencyModuleNames.Add("D3D12RHI");
+				AddEngineThirdPartyPrivateStaticDependencies(Target, "DX12");
+			}
+		}
+	}
+}
diff --git a/Engine/Plugins/Media/ElectraPlayer/Source/ElectraProtron/Private/ElectraProtronModule.cpp b/Engine/Plugins/Media/ElectraPlayer/Source/ElectraProtron/Private/ElectraProtronModule.cpp
new file mode 100644
--- /dev/null
+++ b/Engine/Plugins/Media/ElectraPlayer/Source/ElectraProtron/Private/ElectraProtronModule.cpp
@@ -0,0 +1,36 @@
+// Copyright Epic Games, Inc. All Rights Reserved.
+
+#include "ElectraProtronPrivate.h"
+
+#include "Modules/ModuleManager.h"
+
+#include "IElectraProtronModule.h"
+#include "Player/ElectraProtronPlayer.h"
+
+DEFINE_LOG_CATEGORY(LogElectraProtron);
+
+class FElectraProtronModule
+	: public IElectraProtronModule
+{
+public:
+
+	//~ IElectraProtronModule interface
+
+	virtual TSharedPtr<IMediaPlayer, ESPMode::ThreadSafe> CreatePlayer(IMediaEventSink& EventSink) override
+	{
+		return MakeShared<FElectraProtronPlayer, ESPMode::ThreadSafe>(EventSink);
+	}
+
+public:
+	virtual void StartupModule() override
+	{
+		// We use resource functionality from the Electra Player, so we make sure the modules are loaded.
+		FModuleManager::Get().GetModule("ElectraPlayerPlugin");
+		FModuleManager::Get().GetModule("ElectraPlayerRuntime");
+	}
+	virtual void ShutdownModule() override
+	{ }
+
+};
+
+IMPLEMENT_MODULE(FElectraProtronModule, ElectraProtron)
diff --git a/Engine/Plugins/Media/ElectraPlayer/Source/ElectraProtron/Private/ElectraProtronPrivate.h b/Engine/Plugins/Media/ElectraPlayer/Source/ElectraProtron/Private/ElectraProtronPrivate.h
new file mode 100644
--- /dev/null
+++ b/Engine/Plugins/Media/ElectraPlayer/Source/ElectraProtron/Private/ElectraProtronPrivate.h
@@ -0,0 +1,7 @@
+// Copyright Epic Games, Inc. All Rights Reserved.
+
+#pragma once
+
+#include "Logging/LogMacros.h"
+
+DECLARE_LOG_CATEGORY_EXTERN(LogElectraProtron, Log, All);
diff --git a/Engine/Plugins/Media/ElectraPlayer/Source/ElectraProtron/Private/Player/ElectraProtronPlayer.cpp b/Engine/Plugins/Media/ElectraPlayer/Source/ElectraProtron/Private/Player/ElectraProtronPlayer.cpp
new file mode 100644
--- /dev/null
+++ b/Engine/Plugins/Media/ElectraPlayer/Source/ElectraProtron/Private/Player/ElectraProtronPlayer.cpp
@@ -0,0 +1,470 @@
+// Copyright Epic Games, Inc. All Rights Reserved.
+
+#include "ElectraProtronPlayer.h"
+#include "ElectraProtronPlayerImpl.h"
+#include "ElectraProtronPrivate.h"
+#include "MediaSamples.h"
+#include "Utilities/UtilitiesMP4.h"
+#include "Utilities/MP4Boxes/MP4Boxes.h"
+#include "ElectraTextureSample.h"
+#include "IElectraAudioSample.h"
+
+
+
+FElectraProtronPlayer::FElectraProtronPlayer(IMediaEventSink& InEventSink)
+	: EventSink(InEventSink)
+{
+	EmptySampleQueue = MakeShared<FMediaSamples, ESPMode::ThreadSafe>();
+
+	CurrentState = EMediaState::Closed;
+	CurrentStatus = EMediaStatus::None;
+}
+
+FElectraProtronPlayer::~FElectraProtronPlayer()
+{
+}
+
+void FElectraProtronPlayer::Close()
+{
+	TSharedPtr<FImpl, ESPMode::ThreadSafe> ToClose;
+	Swap(CurrentPlayer, ToClose);
+	CurrentState = EMediaState::Closed;
+	CurrentStatus = EMediaStatus::None;
+	CurrentURL.Empty();
+	CurrentInternalState = EInternalState::Closed;
+	SampleQueueLock.Lock();
+	CurrentSampleQueue.Reset();
+	SampleQueueLock.Unlock();
+	if (ToClose.IsValid())
+	{
+		ToClose->Close(FImpl::FCompletionDelegate());
+		ToClose.Reset();
+	}
+	CurrentTexturePool.Reset();
+	CurrentAudioSamplePool.Reset();
+}
+
+IMediaCache& FElectraProtronPlayer::GetCache()
+{
+	return *this;
+}
+
+IMediaControls& FElectraProtronPlayer::GetControls()
+{
+	return *this;
+}
+
+FGuid FElectraProtronPlayer::GetPlayerPluginGUID() const
+{
+	// Same GUID as in the factory!
+	static FGuid PlayerPluginGUID(0x2899727b, 0xfc934ccb, 0x94119db7, 0x185741d8);
+	return PlayerPluginGUID;
+}
+
+IMediaSamples& FElectraProtronPlayer::GetSamples()
+{
+	FScopeLock lock(&SampleQueueLock);
+	if (CurrentSampleQueue.IsValid())
+	{
+		return *CurrentSampleQueue;
+	}
+	return *EmptySampleQueue;
+}
+
+IMediaTracks& FElectraProtronPlayer::GetTracks()
+{
+	return *this;
+}
+
+FString FElectraProtronPlayer::GetUrl() const
+{
+	return CurrentURL;
+}
+
+IMediaView& FElectraProtronPlayer::GetView()
+{
+	return *this;
+}
+
+bool FElectraProtronPlayer::Open(const FString& InUrl, const IMediaOptions* InOptions)
+{
+	// We expect this player to be used with a FileMediaSource, so the URL needs to
+	// start with "file://" and the remainder is a filename as-is, without escaped URL characters.
+	if (!InUrl.StartsWith(TEXT("file://")))
+	{
+		LastError = FString::Printf(TEXT("File to open does not start with the file:// scheme"));
+		UE_LOG(LogElectraProtron, Error, TEXT("%s"), *LastError);
+		return false;
+	}
+
+	if (CurrentPlayer.IsValid())
+	{
+		Close();
+	}
+
+	CurrentURL = InUrl;
+
+	// Remove the "file://" prefix.
+	FString Filename = CurrentURL.RightChop(7);
+
+	// Create a new implementation player.
+	CurrentPlayer = MakeShared<FImpl, ESPMode::ThreadSafe>();
+	CurrentState = EMediaState::Preparing;
+	CurrentStatus = EMediaStatus::Connecting;
+	CurrentInternalState = EInternalState::Opening;
+	CurrentSequenceIndex = 0;
+
+	CurrentTexturePool = MakeShareable(new FElectraTextureSamplePool);
+	CurrentAudioSamplePool = MakeShareable(new FElectraAudioSamplePool);
+
+	TSharedPtr<FMediaSamples, ESPMode::ThreadSafe> NewSampleQueue = MakeShared<FMediaSamples, ESPMode::ThreadSafe>();
+	SampleQueueLock.Lock();
+	CurrentSampleQueue = NewSampleQueue;
+	SampleQueueLock.Unlock();
+	EventSink.ReceiveMediaEvent(EMediaEvent::MediaConnecting);
+	FImpl::FOpenParam op;
+	op.Filename = Filename;
+	op.SampleQueue = NewSampleQueue;
+	op.TexturePool = CurrentTexturePool;
+	op.AudioSamplePool = CurrentAudioSamplePool;
+	CurrentPlayer->Open(op, FImpl::FCompletionDelegate::CreateSPLambda(AsShared(), [this](FImpl::ImplPointer InImpl)
+	{
+		if (InImpl == CurrentPlayer)
+		{
+			// Change the state to opened, whether successful or not, and let
+			// the media thread do the state transition.
+			CurrentInternalState = EInternalState::Opened;
+		}
+	}));
+	return true;
+}
+
+bool FElectraProtronPlayer::Open(const TSharedRef<FArchive, ESPMode::ThreadSafe>& Archive, const FString& OriginalUrl, const IMediaOptions* Options)
+{
+	return false;
+}
+
+bool FElectraProtronPlayer::Open(const FString& InUrl, const IMediaOptions* InOptions, const FMediaPlayerOptions* InPlayerOptions)
+{
+	// Not currently supported, use player-option-less Open() instead.
+	return Open(InUrl, InOptions);
+}
+
+FVariant FElectraProtronPlayer::GetMediaInfo(FName InInfoName) const
+{
+	TSharedPtr<FImpl, ESPMode::ThreadSafe> Player(GetCurrentPlayer());
+	return Player.IsValid() ? Player->GetMediaInfo(InInfoName) : FVariant();
+}
+
+
+bool FElectraProtronPlayer::GetPlayerFeatureFlag(EFeatureFlag InWhichFeature) const
+{
+	switch(InWhichFeature)
+	{
+		case EFeatureFlag::AllowShutdownOnClose:
+		{
+			return true;
+		}
+		case EFeatureFlag::UsePlaybackTimingV2:
+		{
+			return true;
+		}
+		case EFeatureFlag::PlayerUsesInternalFlushOnSeek:
+		{
+			return true;
+		}
+		case EFeatureFlag::IsTrackSwitchSeamless:
+		{
+			return true;
+		}
+		default:
+		{
+			return IMediaPlayer::GetPlayerFeatureFlag(InWhichFeature);
+		}
+	}
+
+}
+
+
+void FElectraProtronPlayer::CheckForStateChanges()
+{
+	// Check for open completion on the media thread.
+	TSharedPtr<FImpl, ESPMode::ThreadSafe> Player(GetCurrentPlayer());
+	if (CurrentInternalState == EInternalState::Opened)
+	{
+		// Did the player fail to open?
+		LastError = Player->GetLastError();
+		if (LastError.IsEmpty())
+		{
+			CurrentInternalState = EInternalState::Ready;
+			CurrentState = EMediaState::Stopped;
+			CurrentStatus = EMediaStatus::None;
+			EventSink.ReceiveMediaEvent(EMediaEvent::TracksChanged);
+			EventSink.ReceiveMediaEvent(EMediaEvent::MediaOpened);
+		}
+		else
+		{
+			CurrentInternalState = EInternalState::Failed;
+			CurrentState = EMediaState::Error;
+			CurrentStatus = EMediaStatus::None;
+			EventSink.ReceiveMediaEvent(EMediaEvent::MediaOpenFailed);
+			UE_LOG(LogElectraProtron, Error, TEXT("%s"), *LastError);
+		}
+	}
+
+	// Check for state changes when the player is ready.
+	if (CurrentInternalState == EInternalState::Ready)
+	{
+		// Got an error?
+		if (LastError.IsEmpty())
+		{
+			LastError = Player->GetLastError();
+			if (!LastError.IsEmpty())
+			{
+				CurrentInternalState = EInternalState::Failed;
+				CurrentState = EMediaState::Error;
+				CurrentStatus = EMediaStatus::None;
+				EventSink.ReceiveMediaEvent(EMediaEvent::MediaClosed);
+				UE_LOG(LogElectraProtron, Error, TEXT("%s"), *LastError);
+				return;
+			}
+		}
+
+		// Rate changes?
+		float CurrentRate = Player->GetRate();
+		if (CurrentRate == 0.0f && CurrentState == EMediaState::Playing)
+		{
+			CurrentState = EMediaState::Paused;
+			EventSink.ReceiveMediaEvent(EMediaEvent::PlaybackSuspended);
+		}
+		else if (CurrentRate != 0.0f && (CurrentState == EMediaState::Paused || CurrentState == EMediaState::Stopped))
+		{
+			CurrentState = EMediaState::Playing;
+			EventSink.ReceiveMediaEvent(EMediaEvent::PlaybackResumed);
+		}
+
+		// Reached end of playback?
+		if (CurrentState == EMediaState::Playing)
+		{
+			if (Player->HasReachedEnd())
+			{
+				Player->SetRate(0.0f);
+				CurrentState = EMediaState::Paused;
+				EventSink.ReceiveMediaEvent(EMediaEvent::PlaybackEndReached);
+				EventSink.ReceiveMediaEvent(EMediaEvent::PlaybackSuspended);
+			}
+		}
+	}
+}
+
+
+void FElectraProtronPlayer::TickFetch(FTimespan InDeltaTime, FTimespan InTimecode)
+{
+	TSharedPtr<FImpl, ESPMode::ThreadSafe> Player(GetCurrentPlayer());
+	if (Player.IsValid())
+	{
+		Player->TickFetch(InDeltaTime, InTimecode);
+	}
+}
+
+void FElectraProtronPlayer::TickInput(FTimespan InDeltaTime, FTimespan InTimecode)
+{
+	TSharedPtr<FImpl, ESPMode::ThreadSafe> Player(GetCurrentPlayer());
+	if (Player.IsValid())
+	{
+		Player->TickInput(InDeltaTime, InTimecode);
+	}
+	CheckForStateChanges();
+}
+
+FString FElectraProtronPlayer::GetInfo() const
+{
+	return FString();
+}
+
+FString FElectraProtronPlayer::GetStats() const
+{
+	return FString();
+}
+
+bool FElectraProtronPlayer::CanControl(EMediaControl InControl) const
+{
+	if (InControl == EMediaControl::BlockOnFetch)
+	{
+		return CurrentState == EMediaState::Playing;
+	}
+	else if (InControl == EMediaControl::Pause)
+	{
+		return CurrentState == EMediaState::Playing;
+	}
+	else if (InControl == EMediaControl::Resume)
+	{
+		return CurrentState == EMediaState::Paused || CurrentState == EMediaState::Stopped;
+	}
+	else if (InControl == EMediaControl::Seek || InControl == EMediaControl::Scrub)
+	{
+		return CurrentState == EMediaState::Playing || CurrentState == EMediaState::Paused || CurrentState == EMediaState::Stopped;
+	}
+	return false;
+}
+
+FTimespan FElectraProtronPlayer::GetDuration() const
+{
+	TSharedPtr<FImpl, ESPMode::ThreadSafe> Player(GetCurrentPlayer());
+	return Player.IsValid() ? Player->GetDuration() : FTimespan();
+}
+
+float FElectraProtronPlayer::GetRate() const
+{
+	TSharedPtr<FImpl, ESPMode::ThreadSafe> Player(GetCurrentPlayer());
+	return Player.IsValid() ? Player->GetRate() : 0.0f;
+}
+
+EMediaState FElectraProtronPlayer::GetState() const
+{
+	return CurrentState;
+}
+
+EMediaStatus FElectraProtronPlayer::GetStatus() const
+{
+	return CurrentStatus;
+}
+
+TRangeSet<float> FElectraProtronPlayer::GetSupportedRates(EMediaRateThinning InThinning) const
+{
+	TSharedPtr<FImpl, ESPMode::ThreadSafe> Player(GetCurrentPlayer());
+	return Player.IsValid() ? Player->GetSupportedRates(InThinning) : TRangeSet<float>();
+}
+
+FTimespan FElectraProtronPlayer::GetTime() const
+{
+	TSharedPtr<FImpl, ESPMode::ThreadSafe> Player(GetCurrentPlayer());
+	return Player.IsValid() ? Player->GetTime() : FTimespan();
+}
+
+bool FElectraProtronPlayer::IsLooping() const
+{
+	TSharedPtr<FImpl, ESPMode::ThreadSafe> Player(GetCurrentPlayer());
+	return Player.IsValid() ? Player->IsLooping() : false;
+}
+
+bool FElectraProtronPlayer::Seek(const FTimespan& InTime)
+{
+	// Create a new sample queue into which the new decoded samples will be added.
+	TSharedPtr<FMediaSamples, ESPMode::ThreadSafe> SampleQueue = MakeShared<FMediaSamples, ESPMode::ThreadSafe>();
+	SampleQueueLock.Lock();
+	Swap(CurrentSampleQueue, SampleQueue);
+	SampleQueueLock.Unlock();
+
+	TSharedPtr<FImpl, ESPMode::ThreadSafe> Player(GetCurrentPlayer());
+	++CurrentSequenceIndex;
+	bool bOk = Player.IsValid() ? Player->Seek(InTime, CurrentSampleQueue, CurrentSequenceIndex) : false;
+	if (bOk)
+	{
+		// Send a seek complete event even if that is not really true and the seek is in progress or
+		// has not even started. This is mostly to satisfy code that cannot handle async seeking or
+		// seeks that override an ongoing seek (so the number of completed seeks will not match the
+		// number of issued seeks).
+		EventSink.ReceiveMediaEvent(EMediaEvent::SeekCompleted);
+	}
+	else
+	{
+		--CurrentSequenceIndex;
+	}
+	return bOk;
+}
+
+bool FElectraProtronPlayer::SetLooping(bool bInLooping)
+{
+	TSharedPtr<FImpl, ESPMode::ThreadSafe> Player(GetCurrentPlayer());
+	return Player.IsValid() ? Player->SetLooping(bInLooping) : false;
+}
+
+bool FElectraProtronPlayer::SetRate(float InRate)
+{
+	TSharedPtr<FImpl, ESPMode::ThreadSafe> Player(GetCurrentPlayer());
+	return Player.IsValid() ? Player->SetRate(InRate) : false;
+}
+
+bool FElectraProtronPlayer::SetNativeVolume(float Volume)
+{
+	return false;
+}
+
+int32 FElectraProtronPlayer::GetNumTracks(EMediaTrackType InTrackType) const
+{
+	TSharedPtr<FImpl, ESPMode::ThreadSafe> Player(GetCurrentPlayer());
+	return Player.IsValid() ? Player->GetNumTracks(InTrackType) : 0;
+}
+
+int32 FElectraProtronPlayer::GetNumTrackFormats(EMediaTrackType InTrackType, int32 InTrackIndex) const
+{
+	TSharedPtr<FImpl, ESPMode::ThreadSafe> Player(GetCurrentPlayer());
+	return Player.IsValid() ? Player->GetNumTrackFormats(InTrackType, InTrackIndex) : 0;
+}
+
+int32 FElectraProtronPlayer::GetSelectedTrack(EMediaTrackType InTrackType) const
+{
+	TSharedPtr<FImpl, ESPMode::ThreadSafe> Player(GetCurrentPlayer());
+	return Player.IsValid() ? Player->GetSelectedTrack(InTrackType) : -1;
+}
+
+FText FElectraProtronPlayer::GetTrackDisplayName(EMediaTrackType InTrackType, int32 InTrackIndex) const
+{
+	TSharedPtr<FImpl, ESPMode::ThreadSafe> Player(GetCurrentPlayer());
+	return Player.IsValid() ? Player->GetTrackDisplayName(InTrackType, InTrackIndex) : FText();
+}
+
+int32 FElectraProtronPlayer::GetTrackFormat(EMediaTrackType InTrackType, int32 InTrackIndex) const
+{
+	TSharedPtr<FImpl, ESPMode::ThreadSafe> Player(GetCurrentPlayer());
+	return Player.IsValid() ? Player->GetTrackFormat(InTrackType, InTrackIndex) : 0;
+}
+
+FString FElectraProtronPlayer::GetTrackLanguage(EMediaTrackType InTrackType, int32 InTrackIndex) const
+{
+	TSharedPtr<FImpl, ESPMode::ThreadSafe> Player(GetCurrentPlayer());
+	return Player.IsValid() ? Player->GetTrackLanguage(InTrackType, InTrackIndex) : FString();
+}
+
+FString FElectraProtronPlayer::GetTrackName(EMediaTrackType InTrackType, int32 InTrackIndex) const
+{
+	TSharedPtr<FImpl, ESPMode::ThreadSafe> Player(GetCurrentPlayer());
+	return Player.IsValid() ? Player->GetTrackName(InTrackType, InTrackIndex) : FString();
+}
+
+bool FElectraProtronPlayer::GetVideoTrackFormat(int32 InTrackIndex, int32 InFormatIndex, FMediaVideoTrackFormat& OutFormat) const
+{
+	TSharedPtr<FImpl, ESPMode::ThreadSafe> Player(GetCurrentPlayer());
+	return Player.IsValid() ? Player->GetVideoTrackFormat(InTrackIndex, InFormatIndex, OutFormat) : false;
+}
+
+bool FElectraProtronPlayer::GetAudioTrackFormat(int32 InTrackIndex, int32 InFormatIndex, FMediaAudioTrackFormat& OutFormat) const
+{
+	TSharedPtr<FImpl, ESPMode::ThreadSafe> Player(GetCurrentPlayer());
+	return Player.IsValid() ? Player->GetAudioTrackFormat(InTrackIndex, InFormatIndex, OutFormat) : false;
+}
+
+bool FElectraProtronPlayer::SelectTrack(EMediaTrackType InTrackType, int32 InTrackIndex)
+{
+	TSharedPtr<FImpl, ESPMode::ThreadSafe> Player(GetCurrentPlayer());
+	return Player.IsValid() ? Player->SelectTrack(InTrackType, InTrackIndex) : false;
+}
+
+bool FElectraProtronPlayer::SetTrackFormat(EMediaTrackType InTrackType, int32 InTrackIndex, int32 InFormatIndex)
+{
+	TSharedPtr<FImpl, ESPMode::ThreadSafe> Player(GetCurrentPlayer());
+	return Player.IsValid() ? Player->SetTrackFormat(InTrackType, InTrackIndex, InFormatIndex) : false;
+}
+
+bool FElectraProtronPlayer::QueryCacheState(EMediaCacheState InState, TRangeSet<FTimespan>& OutTimeRanges) const
+{
+	TSharedPtr<FImpl, ESPMode::ThreadSafe> Player(GetCurrentPlayer());
+	return Player.IsValid() ? Player->QueryCacheState(InState, OutTimeRanges) : false;
+}
+
+int32 FElectraProtronPlayer::GetSampleCount(EMediaCacheState InState) const
+{
+	TSharedPtr<FImpl, ESPMode::ThreadSafe> Player(GetCurrentPlayer());
+	return Player.IsValid() ? Player->GetSampleCount(InState) : 0;
+}
diff --git a/Engine/Plugins/Media/ElectraPlayer/Source/ElectraProtron/Private/Player/ElectraProtronPlayer.h b/Engine/Plugins/Media/ElectraPlayer/Source/ElectraProtron/Private/Player/ElectraProtronPlayer.h
new file mode 100644
--- /dev/null
+++ b/Engine/Plugins/Media/ElectraPlayer/Source/ElectraProtron/Private/Player/ElectraProtronPlayer.h
@@ -0,0 +1,145 @@
+// Copyright Epic Games, Inc. All Rights Reserved.
+
+#pragma once
+
+#include "CoreTypes.h"
+#include "HAL/CriticalSection.h"
+#include "Containers/Array.h"
+#include "Containers/UnrealString.h"
+#include "Misc/Guid.h"
+#include "IMediaCache.h"
+#include "IMediaControls.h"
+#include "IMediaPlayer.h"
+#include "IMediaTracks.h"
+#include "IMediaView.h"
+#include "IMediaEventSink.h"
+#include "Misc/Timespan.h"
+#include "Misc/Timecode.h"
+#include "Templates/SharedPointer.h"
+#include "ElectraTextureSample.h"
+
+
+class FMediaSamples;
+class FElectraAudioSamplePool;
+class IMediaEventSink;
+
+class FElectraProtronPlayer : public TSharedFromThis<FElectraProtronPlayer, ESPMode::ThreadSafe>
+	, public IMediaPlayer
+	, protected IMediaCache
+	, protected IMediaControls
+	, protected IMediaTracks
+	, protected IMediaView
+{
+public:
+
+	/**
+	 * Create and initialize a new instance.
+	 *
+	 * @param InEventSink The object that receives media events from this player.
+	 */
+	FElectraProtronPlayer(IMediaEventSink& InEventSink);
+
+	/** Virtual destructor. */
+	virtual ~FElectraProtronPlayer();
+
+public:
+
+	//~ IMediaPlayer interface
+
+	void Close() override;
+	IMediaCache& GetCache() override;
+	IMediaControls& GetControls() override;
+	FString GetInfo() const override;
+	FGuid GetPlayerPluginGUID() const override;
+	IMediaSamples& GetSamples() override;
+	FString GetStats() const override;
+	IMediaTracks& GetTracks() override;
+	FString GetUrl() const override;
+	IMediaView& GetView() override;
+	bool Open(const FString& InUrl, const IMediaOptions* InOptions) override;
+	bool Open(const TSharedRef<FArchive, ESPMode::ThreadSafe>& InArchive, const FString& InOriginalUrl, const IMediaOptions* InOptions) override;
+	bool Open(const FString& InUrl, const IMediaOptions* InOptions, const FMediaPlayerOptions* InPlayerOptions) override;
+	FVariant GetMediaInfo(FName InInfoName) const override;
+	bool GetPlayerFeatureFlag(EFeatureFlag /*flag*/) const override;
+
+	void TickFetch(FTimespan InDeltaTime, FTimespan InTimecode) override;
+	void TickInput(FTimespan InDeltaTime, FTimespan InTimecode) override;
+
+protected:
+
+	//~ IMediaControls interface
+
+	bool CanControl(EMediaControl InControl) const override;
+	FTimespan GetDuration() const override;
+	float GetRate() const override;
+	EMediaState GetState() const override;
+	EMediaStatus GetStatus() const override;
+	TRangeSet<float> GetSupportedRates(EMediaRateThinning InThinning) const override;
+	FTimespan GetTime() const override;
+	bool IsLooping() const override;
+	bool Seek(const FTimespan& InTime) override;
+	bool SetLooping(bool bInLooping) override;
+	bool SetRate(float InRate) override;
+	bool SetNativeVolume(float InVolume) override;
+
+protected:
+
+	//~ IMediaTracks interface
+
+	bool GetAudioTrackFormat(int32 InTrackIndex, int32 InFormatIndex, FMediaAudioTrackFormat& OutFormat) const override;
+	int32 GetNumTracks(EMediaTrackType InTrackType) const override;
+	int32 GetNumTrackFormats(EMediaTrackType InTrackType, int32 InTrackIndex) const override;
+	int32 GetSelectedTrack(EMediaTrackType InTrackType) const override;
+	FText GetTrackDisplayName(EMediaTrackType InTrackType, int32 InTrackIndex) const override;
+	int32 GetTrackFormat(EMediaTrackType InTrackType, int32 InTrackIndex) const override;
+	FString GetTrackLanguage(EMediaTrackType InTrackType, int32 InTrackIndex) const override;
+	FString GetTrackName(EMediaTrackType InTrackType, int32 InTrackIndex) const override;
+	bool GetVideoTrackFormat(int32 InTrackIndex, int32 InFormatIndex, FMediaVideoTrackFormat& OutFormat) const override;
+	bool SelectTrack(EMediaTrackType InTrackType, int32 InTrackIndex) override;
+	bool SetTrackFormat(EMediaTrackType InTrackType, int32 InTrackIndex, int32 InFormatIndex) override;
+
+protected:
+
+	//~ IMediaCache interface
+	bool QueryCacheState(EMediaCacheState InState, TRangeSet<FTimespan>& OutTimeRanges) const override;
+	int32 GetSampleCount(EMediaCacheState InState) const override;
+
+private:
+	class FImpl;
+
+
+	TSharedPtr<FImpl, ESPMode::ThreadSafe> GetCurrentPlayer() const
+	{
+		return CurrentPlayer;
+	}
+
+	enum class EInternalState
+	{
+		Closed,
+		Opening,
+		Opened,
+		Ready,
+		Failed
+	};
+	void CheckForStateChanges();
+
+
+	TSharedPtr<FImpl, ESPMode::ThreadSafe> CurrentPlayer;
+	int32 CurrentSequenceIndex = 0;
+	EInternalState CurrentInternalState = EInternalState::Closed;
+
+	IMediaEventSink& EventSink;
+
+	FString LastError;
+	FString CurrentURL;
+
+	FCriticalSection SampleQueueLock;
+	TSharedPtr<FMediaSamples, ESPMode::ThreadSafe> EmptySampleQueue;
+	TSharedPtr<FMediaSamples, ESPMode::ThreadSafe> CurrentSampleQueue;
+	TSharedPtr<FElectraTextureSamplePool, ESPMode::ThreadSafe> CurrentTexturePool;
+	TSharedPtr<FElectraAudioSamplePool, ESPMode::ThreadSafe> CurrentAudioSamplePool;
+
+	/** Current player state. */
+	EMediaState CurrentState;
+	EMediaStatus CurrentStatus;
+};
diff --git a/Engine/Plugins/Media/ElectraPlayer/Source/ElectraProtron/Private/Player/ElectraProtronPlayerControls.cpp b/Engine/Plugins/Media/ElectraPlayer/Source/ElectraProtron/Private/Player/ElectraProtronPlayerControls.cpp
new file mode 100644
--- /dev/null
+++ b/Engine/Plugins/Media/ElectraPlayer/Source/ElectraProtron/Private/Player/ElectraProtronPlayerControls.cpp
@@ -0,0 +1,158 @@
+// Copyright Epic Games, Inc. All Rights Reserved.
+
+#include "ElectraProtronPlayerImpl.h"
+#include "ElectraProtronPrivate.h"
+#include "IElectraCodecFactoryModule.h"
+#include "IElectraCodecFactory.h"
+#include "IElectraDecoderFeaturesAndOptions.h"
+#include "TrackFormatInfo.h"
+
+
+TRangeSet<float> FElectraProtronPlayer::FImpl::GetSupportedRates(EMediaRateThinning InThinning)
+{
+	// For now we do not handle thinned rates differently from unthinned rates.
+	// If we have a selected audio track or a selected video track that does not have keyframes only
+	// we do not support reverse playback.
+	if (!bAreRatesValid)
+	{
+		bAreRatesValid = true;
+
+		TRangeSet<float> TempRates;
+		// Pause and 1x forward are always supported.
+		TempRates.Add(TRange<float>(0.0f));
+		TempRates.Add(TRange<float>(1.0f));
+
+		// No audio?
+		if (TrackSelection.SelectedTrackIndex[CodecTypeIndex(ElectraProtronUtils::FCodecInfo::EType::Audio)] == -1)
+		{
+			// Video selected?
+			const auto TypeIdx = CodecTypeIndex(ElectraProtronUtils::FCodecInfo::EType::Video);
+			if (TrackSelection.SelectedTrackIndex[TypeIdx] != -1)
+			{
+				auto Track = Tracks[UsableTrackArrayIndicesByType[TypeIdx][TrackSelection.SelectedTrackIndex[TypeIdx]]];
+				if (Track.IsValid() && Track->bIsKeyframeOnlyFormat)
+				{
+					TempRates.Add(TRange<float>::Inclusive(0.0f, 64.0f));
+					TempRates.Add(TRange<float>::Inclusive(-64.0f, 0.0f));
+				}
+			}
+		}
+
+		ThinnedRates = TempRates;
+		UnthinnedRates = TempRates;
+	}
+	return InThinning == EMediaRateThinning::Unthinned ? UnthinnedRates : ThinnedRates;
+}
+
+void FElectraProtronPlayer::FImpl::HandleActiveTrackChanges()
+{
+	if (TrackSelection.bChanged)
+	{
+		auto ChangeTrack = [&](int32 InCodecType, FDecoderThread& InDecoderThread, FLoaderThread& InLoaderThread)
+		{
+			if (TrackSelection.SelectedTrackIndex[InCodecType] != TrackSelection.ActiveTrackIndex[InCodecType])
+			{
+				if (TrackSelection.SelectedTrackIndex[InCodecType] >= 0)
+				{
+					// Select.
+					auto Track = Tracks[UsableTrackArrayIndicesByType[InCodecType][TrackSelection.SelectedTrackIndex[InCodecType]]];
+					check(Track.IsValid());
+					if (Track.IsValid() && TrackSampleBuffers.Contains(Track->TrackID))
+					{
+						auto tsb = TrackSampleBuffers[Track->TrackID];
+						if (tsb.IsValid())
+						{
+							InDecoderThread.SetSampleBuffer(tsb, FGetSampleDlg::CreateRaw(&InLoaderThread, &FLoaderThread::GetSample));
+							TrackSelection.ActiveTrackIndex[InCodecType] = TrackSelection.SelectedTrackIndex[InCodecType];
+						}
+					}
+				}
+				else
+				{
+					// Deselect.
+					InDecoderThread.DisconnectSampleBuffer();
+					TrackSelection.ActiveTrackIndex[InCodecType] = TrackSelection.SelectedTrackIndex[InCodecType];
+				}
+			}
+		};
+
+		auto videoCI = CodecTypeIndex(ElectraProtronUtils::FCodecInfo::EType::Video);
+		auto audioCI = CodecTypeIndex(ElectraProtronUtils::FCodecInfo::EType::Audio);
+		// See GetSupportedRates() above. When we are currently playing in reverse and switching tracks would result in
+		// reverse playback being disabled, we need to either set the rate to forward play or pause.
+		if (CurrentRate < 0.0f || SharedPlayParams->PlaybackDirection < 0.0f)
+		{
+			bool bReversePossible = false;
+			if (TrackSelection.SelectedTrackIndex[audioCI] == -1 && TrackSelection.SelectedTrackIndex[videoCI] != -1)
+			{
+				auto Track = Tracks[UsableTrackArrayIndicesByType[videoCI][TrackSelection.SelectedTrackIndex[videoCI]]];
+				bReversePossible = Track.IsValid() && Track->bIsKeyframeOnlyFormat;
+			}
+			if (!bReversePossible)
+			{
+				UE_LOG(LogElectraProtron, Warning, TEXT("New track selection disallows reverse playback. Switching to pause."));
+				// Pause.
+				IntendedRate = 0.0f;
+				HandleRateChanges();
+			}
+		}
+
+		ChangeTrack(videoCI, VideoDecoderThread, VideoLoaderThread);
+		ChangeTrack(audioCI, AudioDecoderThread, AudioLoaderThread);
+		TrackSelection.bChanged = !(TrackSelection.SelectedTrackIndex[videoCI] == TrackSelection.ActiveTrackIndex[videoCI] &&
+									TrackSelection.SelectedTrackIndex[audioCI] == TrackSelection.ActiveTrackIndex[audioCI]);
+	}
+}
+
+void FElectraProtronPlayer::FImpl::HandleRateChanges()
+{
+	if (IntendedRate != CurrentRate)
+	{
+		VideoDecoderThread.SetRate(IntendedRate);
+		AudioDecoderThread.SetRate(IntendedRate);
+
+		// Update the general playback direction. When going into pause retain
+		// the last direction for the loader to know into which direction
+		// the operation went before pausing.
+		// Other than at start the playback direction should not be zero.
+		if (IntendedRate == 0.0f && CurrentRate != 0.0f)
+		{
+			SharedPlayParams->PlaybackDirection = CurrentRate;
+		}
+		else if (IntendedRate != 0.0f)
+		{
+			SharedPlayParams->PlaybackDirection = IntendedRate;
+		}
+		SharedPlayParams->DesiredPlayRate = IntendedRate;
+
+		CurrentRate = IntendedRate;
+	}
+}
+
+void FElectraProtronPlayer::FImpl::HandleSeekRequest(const FSeekRequest& InSeekRequest)
+{
+	// Stop and flush the decoder threads
+	TSharedPtr<FMediaEvent, ESPMode::ThreadSafe> VidFlushed = MakeShared<FMediaEvent, ESPMode::ThreadSafe>();
+	TSharedPtr<FMediaEvent, ESPMode::ThreadSafe> AudFlushed = MakeShared<FMediaEvent, ESPMode::ThreadSafe>();
+
+	VideoDecoderThread.PauseForSeek();
+	AudioDecoderThread.PauseForSeek();
+
+	VideoDecoderThread.Flush(VidFlushed);
+	AudioDecoderThread.Flush(AudFlushed);
+	VidFlushed->Wait();
+	AudFlushed->Wait();
+
+
+	CurrentPlayPosTime = InSeekRequest.NewTime;
+	UpdateTrackLoader(CodecTypeIndex(ElectraProtronUtils::FCodecInfo::EType::Video));
+	UpdateTrackLoader(CodecTypeIndex(ElectraProtronUtils::FCodecInfo::EType::Audio));
+
+	VideoDecoderThread.SetSampleQueue(InSeekRequest.NewSampleQueue);
+	VideoDecoderThread.SetTime(InSeekRequest.NewTime, InSeekRequest.NewSequenceIndex, 0);
+	VideoDecoderThread.ResumeAfterSeek();
+
+	AudioDecoderThread.SetSampleQueue(InSeekRequest.NewSampleQueue);
+	AudioDecoderThread.SetTime(InSeekRequest.NewTime, InSeekRequest.NewSequenceIndex, 0);
+	AudioDecoderThread.ResumeAfterSeek();
+}
diff --git a/Engine/Plugins/Media/ElectraPlayer/Source/ElectraProtron/Private/Player/ElectraProtronPlayerDecoder.cpp b/Engine/Plugins/Media/ElectraPlayer/Source/ElectraProtron/Private/Player/ElectraProtronPlayerDecoder.cpp
new file mode 100644
--- /dev/null
+++ b/Engine/Plugins/Media/ElectraPlayer/Source/ElectraProtron/Private/Player/ElectraProtronPlayerDecoder.cpp
@@ -0,0 +1,1102 @@
+// Copyright Epic Games, Inc. All Rights Reserved.
+
+#include "ElectraProtronPlayerImpl.h"
+#include "ElectraProtronPrivate.h"
+#include "IElectraCodecFactoryModule.h"
+#include "IElectraCodecFactory.h"
+#include "IElectraDecoderFeaturesAndOptions.h"
+#include "Utils/MPEG/ElectraUtilsMPEGVideo_H264.h"
+#include "Utils/MPEG/ElectraUtilsMPEGVideo_H265.h"
+#include "Utils/MPEG/ElectraUtilsMPEGVideo.h"
+#include "Utils/MPEG/ElectraUtilsMPEGAudio.h"
+#include "TrackFormatInfo.h"
+
+#include "IElectraPlayerDecoderResourceManager.h"
+#include "IElectraDecoder.h"
+#include "IElectraDecoderOutputVideo.h"
+#include "IElectraDecoderOutputAudio.h"
+#include "IElectraDecoderFeaturesAndOptions.h"
+#include "IElectraPlayerPluginModule.h"
+#include "ElectraDecodersUtils.h"
+
+#include "MediaSamples.h"
+#include "MediaRendererBase.h"
+#include "MediaDecoderOutput.h"
+#include "ElectraTextureSample.h"
+#include "IElectraAudioSample.h"
+
+
+namespace ProtronDecoderHelpers
+{
+	class FDecoderOutputPointerWrapper : public Electra::IMediaRenderer::IBuffer
+	{
+	public:
+		FDecoderOutputPointerWrapper(FVideoDecoderOutputPtr InPtr) : Ptr(MoveTemp(InPtr))
+		{ HelperDict.Set(Electra::RenderOptionKeys::Texture, Electra::FVariantValue(Ptr)); }
+		const Electra::FParamDict& GetBufferProperties() const override
+		{ return HelperDict; }
+		Electra::FParamDict& GetMutableBufferProperties() override
+		{ return HelperDict; }
+	private:
+		Electra::FParamDict HelperDict;
+		FVideoDecoderOutputPtr Ptr;
+	};
+}
+
+
+
+void FElectraProtronPlayer::FImpl::FDecoderThread::StartThread(const FOpenParam& InParam, const TSharedPtr<FSharedPlayParams, ESPMode::ThreadSafe>& InSharedPlayParams)
+{
+	Params = InParam;
+	SharedPlayParams = InSharedPlayParams;
+	// Video decoder does not need to hold the reference to the audio sample pool.
+	if (DecoderTypeIndex == CodecTypeIndex(ElectraProtronUtils::FCodecInfo::EType::Video))
+	{
+		Params.AudioSamplePool.Reset();
+	}
+	// Audio decoder does not need to hold the reference to the texture pool.
+	else if (DecoderTypeIndex == CodecTypeIndex(ElectraProtronUtils::FCodecInfo::EType::Audio))
+	{
+		Params.TexturePool.Reset();
+	}
+
+	if (!Thread)
+	{
+		Thread = FRunnableThread::Create(this, TEXT("Electra Protron Decoder"), 0, TPri_Normal);
+	}
+}
+
+void FElectraProtronPlayer::FImpl::FDecoderThread::StopThread()
+{
+	if (Thread)
+	{
+		bTerminate = true;
+		Thread->WaitForCompletion();
+		delete Thread;
+		Thread = nullptr;
+	}
+}
+
+uint32 FElectraProtronPlayer::FImpl::FDecoderThread::Run()
+{
+	while(!bTerminate)
+	{
+		// When handling to the target seek time we want to go as quickly as possible,
+		// so we only yield here but not wait.
+		if (SeekTimeToHandleTo.IsSet())
+		{
+			FPlatformProcess::YieldThread();
+		}
+		else
+		{
+			WorkSignal.WaitTimeoutAndReset(1000 * 5);
+		}
+		if (bFlushPending)
+		{
+			bFlushPending = false;
+			PerformFlush();
+			FlushedSignal->Signal();
+		}
+		HandleOutputFrame();
+
+		// If an error occurred, do not do anything to make matters worse.
+		// Just wait for the player to be closed.
+		if (!LastErrorMessage.IsEmpty())
+		{
+			continue;
+		}
+
+		// Is a track buffer change pending?
+		if (PendingBufferChange.bIsSet)
+		{
+			FScopeLock lock(&PendingBufferChange.Lock);
+			PendingBufferChange.bIsSet = false;
+			if (PendingBufferChange.NewTrackSampleBuffer.IsValid())
+			{
+				TSharedPtr<FMP4TrackSampleBuffer, ESPMode::ThreadSafe> NewTrackSampleBuffer(MoveTemp(PendingBufferChange.NewTrackSampleBuffer));
+				FGetSampleDlg NewGetSampleDelegate(MoveTemp(PendingBufferChange.NewGetSampleDelegate));
+				lock.Unlock();
+				DestroyDecoder();
+				TrackSampleBuffer = MoveTemp(NewTrackSampleBuffer);
+				GetSampleDelegate = MoveTemp(NewGetSampleDelegate);
+				UpdateTrackIterator();
+			}
+			else
+			{
+				lock.Unlock();
+				PerformFlush();
+				TrackSampleBuffer.Reset();
+				GetSampleDelegate.Unbind();
+				TrackIterator.Reset();
+				bReachedEnd = true;
+				continue;
+			}
+		}
+
+		// Check for rate changes.
+		if (IntendedRate != CurrentRate)
+		{
+			// Is this a direction change after we reached the end?
+			if (bReachedEnd && IntendedRate != 0.0f && TrackIterator.IsValid())
+			{
+				bool bDirectionChange = (IntendedRate < 0.0f && PlaybackDirection >= 0.0f) || (IntendedRate > 0.0f && PlaybackDirection < 0.0f);
+				if (bDirectionChange)
+				{
+					PlaybackDirection = IntendedRate;
+					bReachedEnd = false;
+					StepTrackIterator(false);
+				}
+			}
+
+			// Going into pause?
+			if (IntendedRate == 0.0f)
+			{
+				if (CurrentRate != 0.0f)
+				{
+					PlaybackDirection = CurrentRate;
+				}
+				if (!IsPaused())
+				{
+					Pause();
+				}
+			}
+			else
+			{
+				PlaybackDirection = IntendedRate;
+				if (IsPaused())
+				{
+					Resume();
+				}
+			}
+			CurrentRate = IntendedRate;
+		}
+
+		// Are we to do something?
+		if (IsPausedForSeek() || (IsPaused() && !SeekTimeToHandleTo.IsSet() && !PendingSeek.bIsSet))
+		{
+			continue;
+		}
+
+		// Moving to a new time?
+		if (PendingSeek.bIsSet)
+		{
+			PendingSeek.Lock.Lock();
+			FTimespan NewTime = PendingSeek.NewTime;
+			int32 NewSequenceIndex = PendingSeek.NewSeqIdx;
+			int32 NewLoopIndex = PendingSeek.NewLoopIdx;
+			PendingSeek.bIsSet = false;
+			PendingSeek.Lock.Unlock();
+
+			TimeLock.Lock();
+			CurrentTime = NewTime;
+			SequenceIndex = NewSequenceIndex;
+			LoopIndex = NewLoopIndex;
+			TimeLock.Unlock();
+
+			UpdateTrackIterator();
+			// Handle decoding to this time now.
+			SeekTimeToHandleTo = CurrentTime;
+			SeekTimeToDecodeTo = CurrentTime;
+			SeekTimeNumFramesDecoded = 0;
+			SeekTimeNumFramesSkipped = 0;
+		}
+
+		DecodeOneFrame();
+	}
+	DestroyDecoder();
+	return 0;
+}
+
+
+void FElectraProtronPlayer::FImpl::FDecoderThread::SetRate(float InNewRate)
+{
+	bool bChange = IntendedRate != InNewRate;
+	IntendedRate = InNewRate;
+	if (bChange)
+	{
+		WorkSignal.Signal();
+	}
+}
+
+bool FElectraProtronPlayer::FImpl::FDecoderThread::SetLooping(bool bInLooping)
+{
+	if (bInLooping && !bShouldLoop)
+	{
+		// Check if play direction is reverse and non-keyframe only.
+		if ((PlaybackDirection < 0.0f || CurrentRate < 0.0f) && TrackSampleBuffer.IsValid() && !TrackSampleBuffer->TrackAndCodecInfo->bIsKeyframeOnlyFormat)
+		{
+			UE_LOG(LogElectraProtron, Warning, TEXT("Cannot enable looping while playing in reverse for tracks that use a non-keyframe-only codec!"));
+			return false;
+		}
+		// Enable
+		bShouldLoop = true;
+	}
+	else if (!bInLooping && bShouldLoop)
+	{
+		// Disable
+		bShouldLoop = false;
+	}
+	return true;
+}
+
+bool FElectraProtronPlayer::FImpl::FDecoderThread::HasReachedEnd()
+{
+	return bReachedEnd;
+}
+
+void FElectraProtronPlayer::FImpl::FDecoderThread::Pause()
+{
+	bIsPaused = true;
+}
+
+void FElectraProtronPlayer::FImpl::FDecoderThread::Resume()
+{
+	bIsPaused = false;
+}
+
+bool FElectraProtronPlayer::FImpl::FDecoderThread::IsPaused()
+{
+	return bIsPaused;
+}
+
+void FElectraProtronPlayer::FImpl::FDecoderThread::PauseForSeek()
+{
+	bPausedForSeek = true;
+}
+
+void FElectraProtronPlayer::FImpl::FDecoderThread::ResumeAfterSeek()
+{
+	bPausedForSeek = false;
+}
+
+bool FElectraProtronPlayer::FImpl::FDecoderThread::IsPausedForSeek()
+{
+	return bPausedForSeek;
+}
+
+void FElectraProtronPlayer::FImpl::FDecoderThread::SetSampleQueue(const TSharedPtr<FMediaSamples, ESPMode::ThreadSafe>& InNewSampleQueue)
+{
+	Params.SampleQueue = InNewSampleQueue;
+}
+
+void FElectraProtronPlayer::FImpl::FDecoderThread::SetSampleBuffer(const TSharedPtr<FElectraProtronPlayer::FImpl::FMP4TrackSampleBuffer, ESPMode::ThreadSafe>& InTrackSampleBuffer, FElectraProtronPlayer::FImpl::FGetSampleDlg InGetSampleDelegate)
+{
+	FScopeLock lock(&PendingBufferChange.Lock);
+	check(InGetSampleDelegate.IsBound());
+	PendingBufferChange.NewTrackSampleBuffer = InTrackSampleBuffer;
+	PendingBufferChange.NewGetSampleDelegate = MoveTemp(InGetSampleDelegate);
+	PendingBufferChange.bIsSet = true;
+}
+
+void FElectraProtronPlayer::FImpl::FDecoderThread::DisconnectSampleBuffer()
+{
+	FScopeLock lock(&PendingBufferChange.Lock);
+	PendingBufferChange.NewTrackSampleBuffer = nullptr;
+	PendingBufferChange.bIsSet = true;
+}
+
+
+void FElectraProtronPlayer::FImpl::FDecoderThread::SetTime(FTimespan InTime, int32 InSeqIdx, int32 InLoopIdx)
+{
+	check(IsPausedForSeek());
+	FScopeLock lock(&PendingSeek.Lock);
+	PendingSeek.NewTime = InTime;
+	PendingSeek.NewSeqIdx = InSeqIdx;
+	PendingSeek.NewLoopIdx = InLoopIdx;
+	PendingSeek.bIsSet = true;
+}
+
+void FElectraProtronPlayer::FImpl::FDecoderThread::SetEstimatedPlaybackTime(FTimespan InTime)
+{
+	FScopeLock lock(&TimeLock);
+	CurrentTime = InTime;
+}
+
+
+void FElectraProtronPlayer::FImpl::FDecoderThread::Flush(const TSharedPtr<FMediaEvent, ESPMode::ThreadSafe>& InFlushedSignal)
+{
+	FlushedSignal = InFlushedSignal;
+	bFlushPending = true;
+}
+
+void FElectraProtronPlayer::FImpl::FDecoderThread::UpdateTrackIterator()
+{
+	if (!TrackSampleBuffer.IsValid())
+	{
+		return;
+	}
+
+	TimeLock.Lock();
+	FTimespan Time = CurrentTime;
+	TimeLock.Unlock();
+	TSharedPtr<Electra::UtilitiesMP4::FMP4Track::FIterator> TkIt =
+		TrackSampleBuffer->TrackAndCodecInfo->MP4Track->CreateIteratorAtKeyframe(Electra::FTimeValue().SetFromTimespan(Time),
+			Electra::FTimeValue().SetFromMilliseconds(TrackSampleBuffer->TrackAndCodecInfo->bIsKeyframeOnlyFormat ? 0 : Config.NextKeyframeThresholdMillis));
+	if (!TkIt.IsValid())
+	{
+		LastErrorMessage = TrackSampleBuffer->TrackAndCodecInfo->MP4Track->GetLastError();
+		UE_LOG(LogElectraProtron, Error, TEXT("%s"), *LastErrorMessage);
+		return;
+	}
+	// Get the timestamps of the media local timeline needed to adjust the decoded samples.
+	FirstSampleEffectiveStartTime = FTimespan(0);
+	LatestSampleEffectiveEndTime = FirstSampleEffectiveStartTime + TkIt->GetEffectiveTrackDuration().GetAsTimespan();
+	// Use the new iterator.
+	TrackIterator = MoveTemp(TkIt);
+	bReachedEnd = false;
+}
+
+
+void FElectraProtronPlayer::FImpl::FDecoderThread::PerformFlush()
+{
+	CurrentDecoderOutput.Reset();
+	CurrentInputSample.Reset();
+	InputForCurrentDecoderOutput.Reset();
+	InDecoderInput.Empty();
+	TrackIterator.Reset();
+	SeekTimeToHandleTo.Reset();
+	SeekTimeToDecodeTo.Reset();
+	SeekTimeNumFramesDecoded = 0;
+	SeekTimeNumFramesSkipped = 0;
+	bIsDrainingAtEOS = false;
+	bReachedEnd = false;
+	bWaitForSyncSample = true;
+	bWarnedMissingSyncSample = false;
+	if (DecoderInstance.IsValid())
+	{
+		DecoderInstance->ResetToCleanStart();
+	}
+	if (DecoderBitstreamProcessor.IsValid())
+	{
+		DecoderBitstreamProcessor->Clear();
+	}
+}
+
+
+void FElectraProtronPlayer::FImpl::FDecoderThread::FlushForEndOrLooping()
+{
+	// Non keyframe-only formats, and all audio formats, require that we drain the decoder to get
+	// all the pending output before we start decoding the next (loop) sequence.
+	bool bNeedsDrain = true;
+	if (DecoderTypeIndex == CodecTypeIndex(ElectraProtronUtils::FCodecInfo::EType::Video) &&
+		TrackSampleBuffer->TrackAndCodecInfo->bIsKeyframeOnlyFormat)
+	{
+		bNeedsDrain = false;
+	}
+
+	if (bNeedsDrain && DecoderInstance.IsValid())
+	{
+		IElectraDecoder::EDecoderError DecErr = DecoderInstance->SendEndOfData();
+		if (ensure(DecErr == IElectraDecoder::EDecoderError::None))
+		{
+			bIsDrainingAtEOS = true;
+			bWaitForSyncSample = true;
+			bWarnedMissingSyncSample = false;
+		}
+		else
+		{
+			LastErrorMessage = FString::Printf(TEXT("Failed to flush decoder at the end of the stream"));
+		}
+	}
+}
+
+
+void FElectraProtronPlayer::FImpl::FDecoderThread::DecodeOneFrame()
+{
+	// Do not decode new data when draining at end-of-stream.
+	if (bIsDrainingAtEOS)
+	{
+		return;
+	}
+	// If we are still hanging on to output we could not deliver we will not
+	// decode new data that will likewise not be able to deliver.
+	if (CurrentDecoderOutput.IsValid())
+	{
+		return;
+	}
+	// If the end of the stream has been reached we are done.
+	if (bReachedEnd)
+	{
+		return;
+	}
+
+	if (TrackSampleBuffer.IsValid())
+	{
+		if (!CreateDecoder())
+		{
+			return;
+		}
+		if (TrackIterator.IsValid())
+		{
+			bool bSkipDecoding = false;
+
+			if (!CurrentInputSample.IsValid())
+			{
+				CurrentInputSample = MakeUnique<FInDecoder>();
+			}
+			if (!CurrentInputSample->Sample.IsValid())
+			{
+				// Get the frame the iterator is pointing at from the buffer.
+				FMP4SamplePtr Sample = GetSampleDelegate.Execute(TrackSampleBuffer, TrackIterator, 0);
+				if (!Sample.IsValid())
+				{
+					return;
+				}
+
+				CurrentInputSample->Sample = MoveTemp(Sample);
+				// Do we have to make a copy of the sample data so we can modify it without affecting the original in the cache?
+				if (DecoderBitstreamProcessor.IsValid() && DecoderBitstreamProcessor->WillModifyBitstreamInPlace())
+				{
+					CurrentInputSample->DataCopy = CurrentInputSample->Sample->Data;
+					CurrentInputSample->DecAU.Data = CurrentInputSample->DataCopy.GetData();
+					CurrentInputSample->DecAU.DataSize = CurrentInputSample->DataCopy.Num();
+				}
+				else
+				{
+					CurrentInputSample->DecAU.Data = CurrentInputSample->Sample->Data.GetData();
+					CurrentInputSample->DecAU.DataSize = CurrentInputSample->Sample->Data.Num();
+				}
+				CurrentInputSample->DecAU.DTS = CurrentInputSample->Sample->DTS;
+				CurrentInputSample->DecAU.PTS = CurrentInputSample->Sample->PTS;
+				CurrentInputSample->DecAU.Duration = CurrentInputSample->Sample->Duration;
+				CurrentInputSample->DecAU.UserValue = NextUserValue;
+				CurrentInputSample->DecAU.Flags = CurrentInputSample->Sample->bIsSyncOrRap ? EElectraDecoderFlags::IsSyncSample : EElectraDecoderFlags::None;
+				if (CurrentInputSample->Sample->bIsSyncOrRap)
+				{
+					CurrentInputSample->CSDOptions = CurrentCodecSpecificData;
+				}
+				if (DecoderBitstreamProcessor.IsValid())
+				{
+					IElectraDecoderBitstreamProcessor::EProcessResult BSResult = DecoderBitstreamProcessor->ProcessInputForDecoding(CurrentInputSample->BSI, CurrentInputSample->DecAU, CurrentInputSample->CSDOptions);
+					if (BSResult == IElectraDecoderBitstreamProcessor::EProcessResult::Error)
+					{
+						LastErrorMessage = DecoderBitstreamProcessor->GetLastError();
+						UE_LOG(LogElectraProtron, Error, TEXT("%s"), *LastErrorMessage);
+						return;
+					}
+				}
+
+				// If handling to a specific time see if this sample is discardable and does not require decoding.
+				if (SeekTimeToDecodeTo.IsSet())
+				{
+					FTimespan Start = CurrentInputSample->Sample->EffectivePTS;
+					FTimespan End = CurrentInputSample->Sample->EffectivePTS + CurrentInputSample->Sample->Duration;
+					FTimespan X = SeekTimeToDecodeTo.GetValue();
+					if (X >= Start && X < End)
+					{
+						SeekTimeToDecodeTo.Reset();
+					}
+					else if ((CurrentInputSample->DecAU.Flags & EElectraDecoderFlags::IsDiscardable) != EElectraDecoderFlags::None)
+					{
+						bSkipDecoding = true;
+						CurrentInputSample.Reset();
+						++SeekTimeNumFramesSkipped;
+					}
+				}
+			}
+
+			// Is this really a sync sample if we need one?
+			if (!bSkipDecoding && bWaitForSyncSample && (CurrentInputSample->DecAU.Flags & EElectraDecoderFlags::IsSyncSample) == EElectraDecoderFlags::None)
+			{
+				if (!bWarnedMissingSyncSample)
+				{
+					bWarnedMissingSyncSample = true;
+					UE_LOG(LogElectraProtron, Warning, TEXT("Expected a sync sample at PTS %lld, but did not get one. The stream may be packaged incorrectly. Dropping frames until one arrives, which may take a while. Please wait!"), (long long int)CurrentInputSample->Sample->EffectivePTS.GetTicks());
+				}
+				bSkipDecoding = true;
+				CurrentInputSample.Reset();
+			}
+
+			if (!bSkipDecoding)
+			{
+				IElectraDecoder::EDecoderError DecErr = DecoderInstance->DecodeAccessUnit(CurrentInputSample->DecAU, CurrentInputSample->CSDOptions);
+				if (DecErr == IElectraDecoder::EDecoderError::None)
+				{
+					if ((CurrentInputSample->DecAU.Flags & EElectraDecoderFlags::DoNotOutput) == EElectraDecoderFlags::None)
+					{
+						// The copied data is no longer necessary to keep around.
+						CurrentInputSample->DataCopy.Empty();
+						// Set the associated indices
+						CurrentInputSample->SequenceIndex = SequenceIndex;
+						CurrentInputSample->LoopIndex = LoopIndex;
+						InDecoderInput.Emplace(NextUserValue, MoveTemp(CurrentInputSample));
+						++NextUserValue;
+					}
+					bWaitForSyncSample = false;
+					bWarnedMissingSyncSample = false;
+				}
+				else if (DecErr == IElectraDecoder::EDecoderError::NoBuffer)
+				{
+					// Try again later...
+					return;
+				}
+				else if (DecErr == IElectraDecoder::EDecoderError::LostDecoder)
+				{
+					if (DecoderInstance->GetError().IsSet())
+					{
+						LastErrorMessage = DecoderInstance->GetError().GetMessage();
+					}
+					else
+					{
+						LastErrorMessage = FString::Printf(TEXT("Lost the decoder"));
+					}
+					UE_LOG(LogElectraProtron, Error, TEXT("%s"), *LastErrorMessage);
+					return;
+				}
+				else
+				{
+					if (DecoderInstance->GetError().IsSet())
+					{
+						LastErrorMessage = DecoderInstance->GetError().GetMessage();
+					}
+					else
+					{
+						LastErrorMessage = FString::Printf(TEXT("Decoder error"));
+					}
+					UE_LOG(LogElectraProtron, Error, TEXT("%s"), *LastErrorMessage);
+					return;
+				}
+			}
+
+			// Move to the next or previous frame
+			StepTrackIterator(true);
+		}
+	}
+}
+
+void FElectraProtronPlayer::FImpl::FDecoderThread::StepTrackIterator(bool bPerformFlush)
+{
+	if (!TrackSampleBuffer.IsValid() || !TrackIterator.IsValid())
+	{
+		return;
+	}
+
+	// Going forward or backwards?
+	if (PlaybackDirection >= 0.0f)
+	{
+		if (!TrackIterator->NextEffective())
+		{
+			if (bPerformFlush)
+			{
+				FlushForEndOrLooping();
+			}
+			if (bShouldLoop)
+			{
+				TSharedPtr<Electra::UtilitiesMP4::FMP4Track::FIterator> TkIt = TrackSampleBuffer->TrackAndCodecInfo->MP4Track->CreateIterator(TrackIterator->GetFirstSampleInfo().SyncSampleNumber);
+				check(TkIt.IsValid());
+				TrackIterator = MoveTemp(TkIt);
+				++LoopIndex;
+			}
+			else
+			{
+				bReachedEnd = true;
+			}
+		}
+	}
+	else
+	{
+		if (!TrackIterator->PrevEffective())
+		{
+			if (bPerformFlush)
+			{
+				FlushForEndOrLooping();
+			}
+			if (bShouldLoop)
+			{
+				// Re-check this is keyframe-only. For non keyframe-only we cannot go back to the last frame since we need to
+				// go to whatever the last keyframe is to decode up to the last frame from there.
+				if (TrackSampleBuffer.IsValid() && TrackSampleBuffer->TrackAndCodecInfo->bIsKeyframeOnlyFormat)
+				{
+					TSharedPtr<Electra::UtilitiesMP4::FMP4Track::FIterator> TkIt = TrackSampleBuffer->TrackAndCodecInfo->MP4Track->CreateIterator(TrackIterator->GetLastSampleInfo().LastSampleNumber);
+					check(TkIt.IsValid());
+					TrackIterator = MoveTemp(TkIt);
+					--LoopIndex;
+				}
+				else
+				{
+					bReachedEnd = true;
+				}
+			}
+			else
+			{
+				bReachedEnd = true;
+			}
+		}
+	}
+}
+
+
+void FElectraProtronPlayer::FImpl::FDecoderThread::HandleOutputFrame()
+{
+	if (!DecoderInstance.IsValid())
+	{
+		return;
+	}
+
+	// If we still have unhandled output, do not get new output yet.
+	if (!CurrentDecoderOutput.IsValid())
+	{
+		InputForCurrentDecoderOutput.Reset();
+		// Does the decoder have new output?
+		IElectraDecoder::EOutputStatus OutputStatus = DecoderInstance->HaveOutput();
+		if (OutputStatus == IElectraDecoder::EOutputStatus::Available)
+		{
+			TSharedPtr<IElectraDecoderOutput, ESPMode::ThreadSafe> Out = DecoderInstance->GetOutput();
+			if (!Out.IsValid())
+			{
+				// No output although advertised?
+				return;
+			}
+			if (DecoderTypeIndex == CodecTypeIndex(ElectraProtronUtils::FCodecInfo::EType::Video) && Out->GetType() == IElectraDecoderOutput::EType::Video)
+			{
+				CurrentDecoderOutput = MoveTemp(Out);
+			}
+			else if (DecoderTypeIndex == CodecTypeIndex(ElectraProtronUtils::FCodecInfo::EType::Audio) && Out->GetType() == IElectraDecoderOutput::EType::Audio)
+			{
+				CurrentDecoderOutput = MoveTemp(Out);
+			}
+			// Unhandled output type?
+			if (Out.IsValid())
+			{
+				Out.Reset();
+				return;
+			}
+			// Find the matching input sample
+			if (!InDecoderInput.RemoveAndCopyValue(CurrentDecoderOutput->GetUserValue(), InputForCurrentDecoderOutput))
+			{
+				check(!"no input for this output?");
+				CurrentDecoderOutput.Reset();
+				return;
+			}
+		}
+		else if (OutputStatus == IElectraDecoder::EOutputStatus::EndOfData)
+		{
+			check(bIsDrainingAtEOS);
+			bIsDrainingAtEOS = false;
+		}
+	}
+	// Handle the output.
+	if (CurrentDecoderOutput.IsValid() && InputForCurrentDecoderOutput.IsValid() && Params.SampleQueue.IsValid())
+	{
+		TOptional<FTimespan> AdjustToTime;
+		bool bSendToOutput = true;
+
+		auto AdjustForSeek = [&]() -> void
+		{
+			// Is this the frame we were supposed to decode up to?
+			if (SeekTimeToHandleTo.IsSet())
+			{
+				FTimespan Start = InputForCurrentDecoderOutput->Sample->EffectivePTS;
+				FTimespan End = InputForCurrentDecoderOutput->Sample->EffectivePTS + InputForCurrentDecoderOutput->Sample->Duration;
+				FTimespan X = SeekTimeToHandleTo.GetValue();
+				if (X >= Start && X < End)
+				{
+					AdjustToTime = X;
+					SeekTimeToHandleTo.Reset();
+					if (DecoderTypeIndex == CodecTypeIndex(ElectraProtronUtils::FCodecInfo::EType::Video) && (SeekTimeNumFramesDecoded + SeekTimeNumFramesSkipped) > 0)
+					{
+						UE_LOG(LogElectraProtron, Log, TEXT("Processed %d frames to reach seek point, of which %d could be skipped"), SeekTimeNumFramesDecoded + SeekTimeNumFramesSkipped, SeekTimeNumFramesSkipped);
+					}
+				}
+				else
+				{
+					bSendToOutput = false;
+					++SeekTimeNumFramesDecoded;
+				}
+			}
+		};
+
+		if (DecoderTypeIndex == CodecTypeIndex(ElectraProtronUtils::FCodecInfo::EType::Video))
+		{
+			AdjustForSeek();
+
+			// When sending the output we need to wait until there is room to receive it.
+			// If the output is not to be sent we can handle it quickly without waiting.
+			if (!bSendToOutput || Params.SampleQueue->CanReceiveVideoSamples(1))
+			{
+				// First perform any processing that needs to be done regardless of whether the output will be sent or not.
+
+				// Get information from the bitstream that the decoder does not provide.
+				if (DecoderBitstreamProcessor.IsValid())
+				{
+					TMap<FString, FVariant> BSIProperties;
+					DecoderBitstreamProcessor->SetPropertiesOnOutput(BSIProperties, InputForCurrentDecoderOutput->BSI);
+					if (BSIProperties.Num())
+					{
+						// Colorimetry?
+						TArray<uint8> CommonColorimetry(ElectraDecodersUtil::GetVariantValueUInt8Array(BSIProperties, IElectraDecoderBitstreamProcessorInfo::CommonColorimetry));
+						if (CommonColorimetry.Num() == sizeof(ElectraDecodersUtil::MPEG::FCommonColorimetry))
+						{
+							const ElectraDecodersUtil::MPEG::FCommonColorimetry& Colorimetry(*reinterpret_cast<const ElectraDecodersUtil::MPEG::FCommonColorimetry*>(CommonColorimetry.GetData()));
+							if (!CurrentColorimetry.IsSet())
+							{
+								CurrentColorimetry = Electra::MPEG::FColorimetryHelper();
+							}
+							CurrentColorimetry.GetValue().Update(Colorimetry.colour_primaries, Colorimetry.transfer_characteristics, Colorimetry.matrix_coeffs, Colorimetry.video_full_range_flag, Colorimetry.video_format);
+						}
+
+						// HDR parameters?
+						TArray<uint8> Mdcv(ElectraDecodersUtil::GetVariantValueUInt8Array(BSIProperties, IElectraDecoderBitstreamProcessorInfo::SeiMasteringDisplayColorVolume));
+						if (Mdcv.Num() == sizeof(ElectraDecodersUtil::MPEG::FSEImastering_display_colour_volume))
+						{
+							if (!CurrentHDR.IsSet())
+							{
+								CurrentHDR = Electra::MPEG::FHDRHelper();
+							}
+							CurrentHDR.GetValue().UpdateWith(*reinterpret_cast<const ElectraDecodersUtil::MPEG::FSEImastering_display_colour_volume*>(Mdcv.GetData()));
+						}
+						TArray<uint8> Clli(ElectraDecodersUtil::GetVariantValueUInt8Array(BSIProperties, IElectraDecoderBitstreamProcessorInfo::SeiContentLightLeveInfo));
+						if (Clli.Num() == sizeof(ElectraDecodersUtil::MPEG::FSEIcontent_light_level_info))
+						{
+							if (!CurrentHDR.IsSet())
+							{
+								CurrentHDR = Electra::MPEG::FHDRHelper();
+							}
+							CurrentHDR.GetValue().UpdateWith(*reinterpret_cast<const ElectraDecodersUtil::MPEG::FSEIcontent_light_level_info*>(Clli.GetData()));
+						}
+						TArray<uint8> Altc(ElectraDecodersUtil::GetVariantValueUInt8Array(BSIProperties, IElectraDecoderBitstreamProcessorInfo::SeiAlternateTransferCharacteristics));
+						if (Altc.Num() == sizeof(ElectraDecodersUtil::MPEG::FSEIalternative_transfer_characteristics))
+						{
+							if (!CurrentHDR.IsSet())
+							{
+								CurrentHDR = Electra::MPEG::FHDRHelper();
+							}
+							CurrentHDR.GetValue().UpdateWith(*reinterpret_cast<const ElectraDecodersUtil::MPEG::FSEIalternative_transfer_characteristics*>(Altc.GetData()));
+						}
+					}
+				}
+
+				FTimespan PTS(InputForCurrentDecoderOutput->Sample->EffectivePTS);
+				FTimespan Dur(InputForCurrentDecoderOutput->Sample->Duration);
+				if (AdjustToTime.IsSet())
+				{
+					FTimespan Diff(AdjustToTime.GetValue() - InputForCurrentDecoderOutput->Sample->EffectivePTS);
+					PTS += Diff;
+					Dur -= Diff;
+				}
+				// Need to trim at the start?
+				else if (PTS < FirstSampleEffectiveStartTime)
+				{
+					FTimespan Diff = FirstSampleEffectiveStartTime - PTS;
+					PTS += Diff;
+					Dur -= Diff;
+				}
+				// Need to trim at the end?
+				if (PTS + Dur > LatestSampleEffectiveEndTime)
+				{
+					FTimespan Diff = PTS + Dur - LatestSampleEffectiveEndTime;
+					Dur -= Diff;
+				}
+
+				// Set the current time to be the end of this decoded frame, which is the PTS plus the
+				// frame's duration when playing forward, or just the PTS when playing in reverse.
+				TimeLock.Lock();
+				CurrentTime = PTS + (PlaybackDirection >= 0.0f ? Dur : FTimespan(0));
+				TimeLock.Unlock();
+
+				// Send the output along?
+				if (bSendToOutput)
+				{
+					TSharedPtr<IElectraDecoderVideoOutput, ESPMode::ThreadSafe> VideoDecoderOutput = StaticCastSharedPtr<IElectraDecoderVideoOutput>(CurrentDecoderOutput);
+					FVideoDecoderOutputPtr NewDecoderOutput = VideoDecoderOutputPool->AcquireShared();
+					// Note: If we need to be informed when the decoder output is no longer needed we
+					//       could set ourselves or some notification wrapper as the owner (derived from IDecoderOutputOwner)
+					//       to receive a SampleReleasedToPool() callback.
+					//NewDecoderOutput->SetOwner(SharedThis(this));
+					ProtronDecoderHelpers::FDecoderOutputPointerWrapper DecOutputPtrWrapper(NewDecoderOutput);
+					TSharedPtr<Electra::FParamDict, ESPMode::ThreadSafe> BufferProperties = MakeShared<Electra::FParamDict, ESPMode::ThreadSafe>();
+
+					Electra::FTimeValue OutputPTS, OutputDur;
+					OutputDur.SetFromTimespan(Dur);
+					uint64 seqIdx = (uint32) InputForCurrentDecoderOutput->LoopIndex;
+					seqIdx |= (uint64)((uint32)InputForCurrentDecoderOutput->SequenceIndex) << 32U;
+					OutputPTS.SetFromTimespan(PTS, (int64) seqIdx);
+					BufferProperties->Set(IDecoderOutputOptionNames::PTS, Electra::FVariantValue(OutputPTS));
+					BufferProperties->Set(IDecoderOutputOptionNames::Duration, Electra::FVariantValue(OutputDur));
+					if (InputForCurrentDecoderOutput->Sample->AssociatedTimecode.IsSet())
+					{
+						BufferProperties->Set(IDecoderOutputOptionNames::TMCDTimecode, Electra::FVariantValue(InputForCurrentDecoderOutput->Sample->AssociatedTimecode.GetValue()));
+						BufferProperties->Set(IDecoderOutputOptionNames::TMCDFramerate, Electra::FVariantValue(InputForCurrentDecoderOutput->Sample->AssociatedTimecodeFramerate.Get(FFrameRate())));
+					}
+
+					// Set the colorimetry, if available, on the output properties.
+					if (CurrentColorimetry.IsSet())
+					{
+						CurrentColorimetry.GetValue().UpdateParamDict(*BufferProperties);
+						// Also HDR information (which requires colorimetry!) if available.
+						if (CurrentHDR.IsSet())
+						{
+							CurrentHDR.GetValue().SetHDRType(VideoDecoderOutput->GetNumberOfBits(), CurrentColorimetry.GetValue());
+							CurrentHDR.GetValue().UpdateParamDict(*BufferProperties);
+						}
+					}
+
+					if (ensure(FElectraPlayerDecoderResourceManager::SetupRenderBufferFromDecoderOutput(&DecOutputPtrWrapper, BufferProperties, VideoDecoderOutput, PlatformResource)))
+					{
+						FElectraTextureSampleRef TextureSample = Params.TexturePool->AcquireShared();
+						TextureSample->Initialize(NewDecoderOutput.Get());
+						Params.SampleQueue->AddVideo(TextureSample);
+					}
+				}
+				CurrentDecoderOutput.Reset();
+				InputForCurrentDecoderOutput.Reset();
+			}
+		}
+		else if (DecoderTypeIndex == CodecTypeIndex(ElectraProtronUtils::FCodecInfo::EType::Audio))
+		{
+			AdjustForSeek();
+
+			// When sending the output we need to wait until there is room to receive it.
+			// If the output is not to be sent we can handle it quickly without waiting.
+			if (!bSendToOutput || Params.SampleQueue->CanReceiveAudioSamples(1))
+			{
+				TSharedPtr<IElectraDecoderAudioOutput, ESPMode::ThreadSafe> AudioDecoderOutput = StaticCastSharedPtr<IElectraDecoderAudioOutput>(CurrentDecoderOutput);
+
+				// Check that we get what we support. At the moment that is the only provided output format anyway, so this should pass.
+				if (AudioDecoderOutput->GetType() == IElectraDecoderOutput::EType::Audio &&
+					AudioDecoderOutput->IsInterleaved() &&
+					AudioDecoderOutput->GetSampleFormat() == IElectraDecoderAudioOutput::ESampleFormat::Float)
+				{
+					const float* PCMBuffer = reinterpret_cast<const float*>(AudioDecoderOutput->GetData(0));
+					const int32 SamplingRate = AudioDecoderOutput->GetSampleRate();
+					const int32 NumberOfChannels = AudioDecoderOutput->GetNumChannels();
+					const int32 NumBytesPerFrame = AudioDecoderOutput->GetBytesPerFrame();
+					const int32 NumBytesPerSample = AudioDecoderOutput->GetBytesPerSample();
+					int32 NumSamplesProduced = AudioDecoderOutput->GetNumFrames();
+
+					if (AudioChannelMapper.IsInitialized() || (!AudioChannelMapper.IsInitialized() && AudioChannelMapper.Initialize(AudioDecoderOutput)))
+					{
+						int32 ByteOffsetToFirstSample = 0;
+						FElectraAudioSampleRef AudioSample = Params.AudioSamplePool->AcquireShared();
+						FMediaTimeStamp PTS(InputForCurrentDecoderOutput->Sample->EffectivePTS, FMediaTimeStamp::MakeSequenceIndex(InputForCurrentDecoderOutput->SequenceIndex, InputForCurrentDecoderOutput->LoopIndex));
+						FTimespan Dur(InputForCurrentDecoderOutput->Sample->Duration);
+						FTimespan TrimDurAtStart, TrimDurAtEnd;
+						if (AdjustToTime.IsSet())
+						{
+							TrimDurAtStart = AdjustToTime.GetValue() - InputForCurrentDecoderOutput->Sample->EffectivePTS;
+							PTS += TrimDurAtStart;
+							Dur -= TrimDurAtStart;
+						}
+						// Need to trim at the start?
+						else if (PTS.Time < FirstSampleEffectiveStartTime)
+						{
+							TrimDurAtStart = FirstSampleEffectiveStartTime - PTS.Time;
+							PTS += TrimDurAtStart;
+							Dur -= TrimDurAtStart;
+						}
+						// Need to trim at the end?
+						if (PTS.Time + Dur > LatestSampleEffectiveEndTime)
+						{
+							TrimDurAtEnd = PTS.Time + Dur - LatestSampleEffectiveEndTime;
+							Dur -= TrimDurAtEnd;
+						}
+
+						// Set the current time to be the end of this decoded frame, which is the PTS plus the
+						// frame's duration when playing forward, or just the PTS when playing in reverse.
+						TimeLock.Lock();
+						CurrentTime = PTS.Time + (PlaybackDirection >= 0.0f ? Dur : FTimespan(0));
+						TimeLock.Unlock();
+
+						// Send the output along?
+						if (bSendToOutput)
+						{
+							int32 SkipStartSampleNum = (int32) (TrimDurAtStart.GetTicks() * SamplingRate / ETimespan::TicksPerSecond);
+							int32 SkipEndSampleNum = (int32) (TrimDurAtEnd.GetTicks() * SamplingRate / ETimespan::TicksPerSecond);
+
+							if (SkipStartSampleNum + SkipEndSampleNum < NumSamplesProduced)
+							{
+								ByteOffsetToFirstSample = SkipStartSampleNum * NumBytesPerFrame;
+								NumSamplesProduced -= SkipStartSampleNum;
+								NumSamplesProduced -= SkipEndSampleNum;
+							}
+							else
+							{
+								NumSamplesProduced = 0;
+							}
+
+							// Anything remaining to pass along?
+							if (NumSamplesProduced)
+							{
+								// Pre-initialize the sample without copying data.
+								if (AudioSample->Initialize(nullptr, EMediaAudioSampleFormat::Float, (uint32)NumberOfChannels, (uint32)NumSamplesProduced, (uint32)SamplingRate, PTS, Dur))
+								{
+
+									AudioChannelMapper.MapChannels(const_cast<void*>(AudioSample->GetBuffer()), (int32)AudioSample->GetAllocatedSize(),
+																	ElectraDecodersUtil::AdvancePointer(PCMBuffer, ByteOffsetToFirstSample), NumSamplesProduced * NumBytesPerFrame, NumSamplesProduced);
+									Params.SampleQueue->AddAudio(AudioSample);
+								}
+								else
+								{
+									// Out of memory. Leave gracefully. Bad things will happen soon somewhere, but at least not here.
+								}
+							}
+						}
+					}
+					else
+					{
+						LastErrorMessage = FString::Printf(TEXT("Could not initialize the channel mapper"));
+						UE_LOG(LogElectraProtron, Error, TEXT("%s"), *LastErrorMessage);
+					}
+				}
+				else
+				{
+					LastErrorMessage = FString::Printf(TEXT("Unsupported audio output format"));
+					UE_LOG(LogElectraProtron, Error, TEXT("%s"), *LastErrorMessage);
+				}
+				CurrentDecoderOutput.Reset();
+				InputForCurrentDecoderOutput.Reset();
+			}
+		}
+	}
+}
+
+
+
+void FElectraProtronPlayer::FImpl::FDecoderThread::DestroyDecoder()
+{
+	PerformFlush();
+	CurrentCodecSpecificData.Empty();
+	CurrentColorimetry.Reset();
+	CurrentHDR.Reset();
+	AudioChannelMapper.Reset();
+
+	if (DecoderInstance.IsValid())
+	{
+		DecoderInstance->Close();
+		DecoderInstance.Reset();
+	}
+	if (DecoderBitstreamProcessor.IsValid())
+	{
+		DecoderBitstreamProcessor->Clear();
+		DecoderBitstreamProcessor.Reset();
+	}
+	if (PlatformResource)
+	{
+		FElectraPlayerDecoderResourceManager::GetDelegate()->ReleasePlatformResource(this, PlatformResource);
+		PlatformResource = nullptr;
+	}
+	VideoResourceDelegate.Reset();
+}
+
+
+bool FElectraProtronPlayer::FImpl::FDecoderThread::CreateDecoder()
+{
+	if (DecoderInstance.IsValid())
+	{
+		return true;
+	}
+	check(TrackSampleBuffer.IsValid());
+	if (!TrackSampleBuffer.IsValid())
+	{
+		LastErrorMessage = FString::Printf(TEXT("Internal error, no track sample buffer is set when creating a decoder."));
+		return false;
+	}
+
+	// Get the decoder factory module.
+	IElectraCodecFactoryModule* FactoryModule = static_cast<IElectraCodecFactoryModule*>(FModuleManager::Get().GetModule(TEXT("ElectraCodecFactory")));
+	check(FactoryModule);
+	// And since we piggyback on the Electra Player decoder output handling and conversion we need the plugin as well.
+	IElectraPlayerPluginModule* PluginModule = static_cast<IElectraPlayerPluginModule*>(FModuleManager::Get().GetModule(TEXT("ElectraPlayerPlugin")));
+	check(PluginModule && PluginModule->IsInitialized());
+
+
+	CurrentCodecSpecificData.Empty();
+	CurrentColorimetry.Reset();
+	CurrentHDR.Reset();
+	TMap<FString, FVariant> DecoderCfgParams;
+	const ElectraProtronUtils::FCodecInfo& ci(TrackSampleBuffer->TrackAndCodecInfo->CodecInfo);
+	if (DecoderTypeIndex == CodecTypeIndex(ElectraProtronUtils::FCodecInfo::EType::Video))
+	{
+		check(ci.Type == ElectraProtronUtils::FCodecInfo::EType::Video);
+		const ElectraProtronUtils::FCodecInfo::FVideo& vi(ci.Properties.Get<ElectraProtronUtils::FCodecInfo::FVideo>());
+		DecoderCfgParams.Add(TEXT("width"), FVariant((uint32) vi.Width));
+		DecoderCfgParams.Add(TEXT("height"), FVariant((uint32) vi.Height));
+		DecoderCfgParams.Add(TEXT("max_width"), FVariant((uint32) vi.Width));
+		DecoderCfgParams.Add(TEXT("max_height"), FVariant((uint32) vi.Height));
+		DecoderCfgParams.Add(TEXT("max_output_buffers"), FVariant((uint32) 8));
+		if (vi.FrameRate.IsValid())
+		{
+			DecoderCfgParams.Add(TEXT("fps"), FVariant((double)vi.FrameRate.GetAsDouble()));
+			DecoderCfgParams.Add(TEXT("fps_n"), FVariant((int64)vi.FrameRate.GetNumerator()));
+			DecoderCfgParams.Add(TEXT("fps_d"), FVariant((uint32)vi.FrameRate.GetDenominator()));
+		}
+		else
+		{
+			DecoderCfgParams.Add(TEXT("fps"), FVariant((double)0.0));
+			DecoderCfgParams.Add(TEXT("fps_n"), FVariant((int64)0));
+			DecoderCfgParams.Add(TEXT("fps_d"), FVariant((uint32)1));
+		}
+		VideoResourceDelegate = PluginModule->CreatePlatformVideoDecoderResourceDelegate();
+		// Set the video resource delegate with the decoder configuration options as a plain pointer.
+		DecoderCfgParams.Emplace(TEXT("VideoResourceDelegate"), FVariant((uint64)VideoResourceDelegate.Get()));
+	}
+	else if (DecoderTypeIndex == CodecTypeIndex(ElectraProtronUtils::FCodecInfo::EType::Audio))
+	{
+		check(ci.Type == ElectraProtronUtils::FCodecInfo::EType::Audio);
+		const ElectraProtronUtils::FCodecInfo::FAudio& ai(ci.Properties.Get<ElectraProtronUtils::FCodecInfo::FAudio>());
+		DecoderCfgParams.Add(TEXT("channel_configuration"), FVariant((uint32) ai.ChannelConfiguration));
+		DecoderCfgParams.Add(TEXT("num_channels"), FVariant((int32) ai.NumChannels));
+		DecoderCfgParams.Add(TEXT("sample_rate"), FVariant((int32) ai.SampleRate));
+	}
+
+	if (ci.DCR.Num())
+	{
+		CurrentCodecSpecificData.Add(TEXT("dcr"), FVariant(ci.DCR));
+	}
+	if (ci.CSD.Num())
+	{
+		CurrentCodecSpecificData.Add(TEXT("csd"), FVariant(ci.CSD));
+	}
+
+	FString DecoderFormat = ci.RFC6381;
+	DecoderCfgParams.Add(TEXT("codec_name"), FVariant(DecoderFormat));
+	DecoderCfgParams.Add(TEXT("codec_4cc"), FVariant((uint32) ci.FourCC));
+	// Add the extra boxes found in the sample description.
+	for(auto& boxIt : ci.ExtraBoxes)
+	{
+		FString BoxName = FString::Printf(TEXT("$%s_box"), *Electra::UtilitiesMP4::Printable4CC(boxIt.Key));
+		CurrentCodecSpecificData.Add(BoxName, FVariant(boxIt.Value));
+	}
+	DecoderCfgParams.Append(CurrentCodecSpecificData);
+
+	// Try to find a decoder factory for the format.
+	// This should succeed since we already checked the supported formats earlier.
+	TMap<FString, FVariant> FormatInfo;
+	TSharedPtr<IElectraCodecFactory, ESPMode::ThreadSafe> DecoderFactory = FactoryModule->GetBestFactoryForFormat(FormatInfo, DecoderFormat, false, DecoderCfgParams);
+	check(DecoderFactory.IsValid());
+	if (!DecoderFactory.IsValid())
+	{
+		LastErrorMessage = FString::Printf(TEXT("No decoder factory found for format \"%s\"."), *DecoderFormat);
+		return false;
+	}
+
+	// Create a decoder instance, for which we need to get delegates from the Electra Player.
+	TSharedPtr<IElectraDecoderResourceDelegate, ESPMode::ThreadSafe> ElectraDecoderResourceDelegate = FElectraPlayerDecoderResourceManager::GetDelegate();
+	// Create platform specific resources to be used with the new decoder.
+	check(PlatformResource == nullptr);
+	if (DecoderTypeIndex == CodecTypeIndex(ElectraProtronUtils::FCodecInfo::EType::Video))
+	{
+		PlatformResource = ElectraDecoderResourceDelegate->CreatePlatformResource(this, IElectraDecoderResourceDelegateBase::EDecoderPlatformResourceType::Video, DecoderCfgParams);
+	}
+	//DecoderCreateCfg.Emplace(TEXT("platformresource"), FVariant((uint64)PlatformResource));
+
+	// Try to create the decoder now.
+	DecoderInstance = DecoderFactory->CreateDecoderForFormat(DecoderFormat, DecoderCfgParams, ElectraDecoderResourceDelegate);
+	if (!DecoderInstance.IsValid())
+	{
+		ElectraDecoderResourceDelegate->ReleasePlatformResource(this, PlatformResource);
+		PlatformResource = nullptr;
+		LastErrorMessage = FString::Printf(TEXT("Failed to create decoder for format \"%s\"."), *DecoderFormat);
+		return false;
+	}
+	if (DecoderInstance->GetError().IsSet())
+	{
+		ElectraDecoderResourceDelegate->ReleasePlatformResource(this, PlatformResource);
+		PlatformResource = nullptr;
+		LastErrorMessage = DecoderInstance->GetError().GetMessage();
+		return false;
+	}
+	// Get the bitstream processor for this decoder, if it requires one.
+	DecoderBitstreamProcessor = DecoderInstance->CreateBitstreamProcessor();
+	return true;
+}
diff --git a/Engine/Plugins/Media/ElectraPlayer/Source/ElectraProtron/Private/Player/ElectraProtronPlayerImpl.cpp b/Engine/Plugins/Media/ElectraPlayer/Source/ElectraProtron/Private/Player/ElectraProtronPlayerImpl.cpp
new file mode 100644
--- /dev/null
+++ b/Engine/Plugins/Media/ElectraPlayer/Source/ElectraProtron/Private/Player/ElectraProtronPlayerImpl.cpp
@@ -0,0 +1,1330 @@
+// Copyright Epic Games, Inc. All Rights Reserved.
+
+#include "ElectraProtronPlayerImpl.h"
+#include "ElectraProtronPrivate.h"
+#include "IElectraCodecFactoryModule.h"
+#include "IElectraCodecFactory.h"
+#include "IElectraDecoderFeaturesAndOptions.h"
+#include "Utils/MPEG/ElectraUtilsMPEGVideo_H264.h"
+#include "Utils/MPEG/ElectraUtilsMPEGVideo_H265.h"
+#include "Utils/MPEG/ElectraUtilsMPEGAudio.h"
+#include "TrackFormatInfo.h"
+#include "Core/MediaThreads.h"
+
+
+namespace ElectraProtronOptionNames
+{
+	const FName StartTimecodeValue(TEXT("StartTimecodeValue"));				// maybe use:  UMediaPlayer::MediaInfoNameSourceNumTiles.Resolve()
+	const FName StartTimecodeFrameRate(TEXT("StartTimecodeFrameRate"));
+}
+
+
+FElectraProtronPlayer::FImpl::FImpl()
+{
+	// Create the playback parameter structure that has members changing at any moment in time.
+	// This information is shared with the frame loader.
+	SharedPlayParams = MakeShared<FSharedPlayParams, ESPMode::ThreadSafe>();
+}
+
+FElectraProtronPlayer::FImpl::~FImpl()
+{
+}
+
+void FElectraProtronPlayer::FImpl::StartThread()
+{
+	if (!Thread)
+	{
+		Thread = FRunnableThread::Create(this, TEXT("Electra Protron"), 0, TPri_Normal);
+	}
+}
+
+
+void FElectraProtronPlayer::FImpl::Open(const FOpenParam& InParam, FElectraProtronPlayer::FImpl::FCompletionDelegate InCompletionDelegate)
+{
+	FWorkerThreadMessage Msg;
+	FWorkerThreadMessage::FParamOpen open { InParam };
+	Msg.Param.Emplace<FWorkerThreadMessage::FParamOpen>(MoveTemp(open));
+	Msg.Self = AsShared();
+	Msg.Type = FWorkerThreadMessage::EType::Open;
+	Msg.CompletionDelegate = MoveTemp(InCompletionDelegate);
+	SendWorkerThreadMessage(MoveTemp(Msg));
+	StartThread();
+}
+
+void FElectraProtronPlayer::FImpl::Close(FElectraProtronPlayer::FImpl::FCompletionDelegate InCompletionDelegate)
+{
+	bAbort = true;
+	if (Thread)
+	{
+		FWorkerThreadMessage Msg;
+		Msg.Param.Emplace<FWorkerThreadMessage::FParamTerminate>(FWorkerThreadMessage::FParamTerminate());
+		Msg.Self = AsShared();
+		Msg.Type = FWorkerThreadMessage::EType::Terminate;
+		Msg.CompletionDelegate = MoveTemp(InCompletionDelegate);
+		SendWorkerThreadMessage(MoveTemp(Msg));
+	}
+	else
+	{
+		InCompletionDelegate.ExecuteIfBound(AsShared());
+	}
+}
+
+FString FElectraProtronPlayer::FImpl::GetLastError()
+{
+	return LastErrorMessage;
+}
+
+bool FElectraProtronPlayer::FImpl::HasReachedEnd()
+{
+	const bool bIsVideoActive = TrackSelection.ActiveTrackIndex[CodecTypeIndex(ElectraProtronUtils::FCodecInfo::EType::Video)] != -1;
+	const bool bIsAudioActive = TrackSelection.ActiveTrackIndex[CodecTypeIndex(ElectraProtronUtils::FCodecInfo::EType::Audio)] != -1;
+
+	bool bAllReachedEnd = true;
+	if (bIsVideoActive && !VideoDecoderThread.HasReachedEnd())
+	{
+		bAllReachedEnd = false;
+	}
+	if (bIsAudioActive && !AudioDecoderThread.HasReachedEnd())
+	{
+		bAllReachedEnd = false;
+	}
+	return bAllReachedEnd;
+}
+
+
+void FElectraProtronPlayer::FImpl::SendWorkerThreadMessage(FElectraProtronPlayer::FImpl::FWorkerThreadMessage&& InMessage)
+{
+	WorkMessages.Enqueue(MoveTemp(InMessage));
+	WorkMessageSignal.Signal();
+}
+
+void FElectraProtronPlayer::FImpl::Exit()
+{
+	// We are still within our own thread here, so we cannot wait for completion.
+	// Use an async task to this if possible.
+	if (GIsRunning)
+	{
+		FMediaRunnable::EnqueueAsyncTask([this]()
+		{
+			Thread->WaitForCompletion();
+			delete Thread;
+			SelfDuringTerminate.Reset();
+		});
+	}
+	else
+	{
+		// Leave the thread dangling, we can't clean it up here.
+		SelfDuringTerminate.Reset();
+	}
+}
+
+uint32 FElectraProtronPlayer::FImpl::Run()
+{
+	bool bDone = false;
+	while(!bDone)
+	{
+		WorkMessageSignal.WaitTimeoutAndReset(1000 * 20);
+		FWorkerThreadMessage Msg;
+		while(WorkMessages.Dequeue(Msg))
+		{
+			switch(Msg.Type)
+			{
+				case FWorkerThreadMessage::EType::Open:
+				{
+					const FWorkerThreadMessage::FParamOpen& open(Msg.Param.Get<FWorkerThreadMessage::FParamOpen>());
+					InternalOpen(open.Param.Filename);
+					// Start loader threads when opening was successful.
+					if (LastErrorMessage.IsEmpty())
+					{
+						VideoLoaderThread.StartThread(open.Param.Filename, SharedPlayParams);
+						AudioLoaderThread.StartThread(open.Param.Filename, SharedPlayParams);
+
+						VideoDecoderThread.StartThread(open.Param, SharedPlayParams);
+						AudioDecoderThread.StartThread(open.Param, SharedPlayParams);
+					}
+					break;
+				}
+				case FWorkerThreadMessage::EType::Terminate:
+				{
+					bDone = true;
+					// Hold on to ourselves while we exit the loop.
+					// Otherwise, if there are no other owners we may get destroyed too soon on our way out.
+					SelfDuringTerminate = AsShared();
+
+					// Stop decoder threads
+					AudioDecoderThread.StopThread();
+					VideoDecoderThread.StopThread();
+
+					// Stop loader threads
+					AudioLoaderThread.StopThread();
+					VideoLoaderThread.StopThread();
+					break;
+				}
+				default:
+				{
+					unimplemented();
+					break;
+				}
+			}
+			Msg.CompletionDelegate.ExecuteIfBound(AsShared());
+		}
+
+		// Is there a new seek request pending?
+		SeekRequestLock.Lock();
+		TOptional<FSeekRequest> NewSeekRequest(PendingSeekRequest);
+		PendingSeekRequest.Reset();
+		SeekRequestLock.Unlock();
+		if (NewSeekRequest.IsSet())
+		{
+			HandleSeekRequest(NewSeekRequest.GetValue());
+		}
+	}
+	return 0;
+}
+
+
+/**
+ * Opens the given file and verifies that it can be used.
+ */
+void FElectraProtronPlayer::FImpl::InternalOpen(const FString& InFilename)
+{
+	// Open the file.
+	TSharedPtr<Electra::IFileDataReader, ESPMode::ThreadSafe> Reader = Electra::IFileDataReader::Create();
+	if (!Reader->Open(InFilename))
+	{
+		LastErrorMessage = Reader->GetLastError();
+		return;
+	}
+
+	// Read the mp4 box structure.
+	Electra::UtilitiesMP4::FMP4BoxLocatorReader BoxLocator;
+	const TArray<uint32> FirstBoxes { Electra::UtilitiesMP4::MakeBoxAtom('f','t','y','p'), Electra::UtilitiesMP4::MakeBoxAtom('s','t','y','p'), Electra::UtilitiesMP4::MakeBoxAtom('s','i','d','x'), Electra::UtilitiesMP4::MakeBoxAtom('f','r','e','e'), Electra::UtilitiesMP4::MakeBoxAtom('s','k','i','p') };
+	const TArray<uint32> ReadBoxes;			// Empty means to read in all boxes except `mdat`
+	const TArray<uint32> StopAfterBoxes;	// Empty means to read all boxes and not stop after a specific one.
+	TArray<TSharedPtr<Electra::UtilitiesMP4::FMP4BoxData, ESPMode::ThreadSafe>> RootBoxes;
+	if (!BoxLocator.LocateAndReadRootBoxes(RootBoxes, Reader, FirstBoxes, StopAfterBoxes, ReadBoxes, Electra::IBaseDataReader::FCancellationCheckDelegate::CreateLambda([&](){return bAbort;})))
+	{
+		LastErrorMessage = BoxLocator.GetLastError();
+		return;
+	}
+
+	// In order to be usable the mp4 needs to have a `moov` box.
+	if (!RootBoxes.FindByPredicate([](const TSharedPtr<Electra::UtilitiesMP4::FMP4BoxData, ESPMode::ThreadSafe>& In){ return In->Type == Electra::UtilitiesMP4::MakeBoxAtom('m','o','o','v'); }))
+	{
+		LastErrorMessage = FString::Printf(TEXT("No `moov` box found in this file. It does not appear to be an mp4/mov file."));
+		return;
+	}
+
+	// Parse all the root boxes. There are typically no more than 4, maybe 6 unless it is a fragmented mp4.
+	ParsedRootBoxes.SetNum(RootBoxes.Num());
+	for(int32 i=0; i<RootBoxes.Num(); ++i)
+	{
+		if (!ParsedRootBoxes[i].ParseBoxTree(RootBoxes[i]))
+		{
+			LastErrorMessage = FString::Printf(TEXT("Failed to parse the file's box structure."));
+			return;
+		}
+	}
+
+	// We need the `moov` box now to determine what is inside this file.
+	// Yes, we looked at the presence of it above already as a quick reject before parsing the boxes.
+	// Now we need it for real and we know that it needs to be there, so here we go
+	TSharedPtr<Electra::UtilitiesMP4::FMP4BoxMOOV> MoovBox = StaticCastSharedPtr<Electra::UtilitiesMP4::FMP4BoxMOOV>(
+															(ParsedRootBoxes.FindByPredicate([](const Electra::UtilitiesMP4::FMP4BoxTreeParser& In)
+															{
+																return In.GetBoxTree()->GetType() == Electra::UtilitiesMP4::MakeBoxAtom('m','o','o','v');
+															})->GetBoxTree()));
+	// Get the `mvhd` box for the duration of the movie and the timescale other values are measured in.
+	auto MvhdBox = MoovBox->FindBoxRecursive<Electra::UtilitiesMP4::FMP4BoxMVHD>(Electra::UtilitiesMP4::MakeBoxAtom('m','v','h','d'), 0);
+	if (!MvhdBox.IsValid())
+	{
+		LastErrorMessage = FString::Printf(TEXT("No `mvhd` box found. This is not a usable mp4/mov file."));
+		return;
+	}
+	MovieDuration = MvhdBox->GetDuration();
+	if (!MovieDuration.IsValid())
+	{
+		LastErrorMessage = FString::Printf(TEXT("Duration in `mvhd` box is set to indefinite or is not valid. This is not a usable mp4/mov file."));
+		return;
+	}
+	Duration = MovieDuration.GetAsTimebase(ETimespan::TicksPerSecond);
+
+
+	IElectraCodecFactoryModule* FactoryModule = static_cast<IElectraCodecFactoryModule*>(FModuleManager::Get().GetModule(TEXT("ElectraCodecFactory")));
+	check(FactoryModule);
+	if (!FactoryModule)
+	{
+		LastErrorMessage = FString::Printf(TEXT("Electra decoder factory not found. Unable to use any track."));
+		return;
+	}
+
+	// Get all the tracks.
+	TArray<TSharedPtr<Electra::UtilitiesMP4::FMP4BoxTRAK>> AllTracks;
+	MoovBox->GetAllBoxInstances(AllTracks, Electra::UtilitiesMP4::MakeBoxAtom('t','r','a','k'));
+	// Prepare the internal track structure and check which tracks we can use and which ones we cannot.
+	Tracks.SetNum(AllTracks.Num());
+	int32 NumTimecodeTracks = 0;
+	int64 LongestTrackDuration = -1;
+	int64 LongestSupportedTrackDuration = -1;
+	for(int32 i=0; i<AllTracks.Num(); ++i)
+	{
+		Tracks[i] = MakeShared<FTrackInfo>();
+		Tracks[i]->TrackBox = AllTracks[i];
+
+		auto Tkhd = AllTracks[i]->FindBoxRecursive<Electra::UtilitiesMP4::FMP4BoxTKHD>(Electra::UtilitiesMP4::MakeBoxAtom('t','k','h','d'), 1);
+		if (!Tkhd)
+		{
+			// If not found the file is broken.
+			LastErrorMessage = FString::Printf(TEXT("No `tkhd` box found on track %d. This file cannot be used."), i);
+			Tracks.Empty();
+			return;
+		}
+		uint32 TrackID = Tracks[i]->TrackID = Tkhd->GetTrackID();
+		/*
+			From the standard:
+				"Tracks that are marked as not enabled (track_enabled set to 0) shall be ignored and treated as if
+				not present."
+		*/
+		if (!Tkhd->IsEnabled())
+		{
+			UE_LOG(LogElectraProtron, Warning, TEXT("Track #%u is flagged as disabled, ignoring this track."), TrackID);
+			continue;
+		}
+
+		// This track's duration must not be indefinite.
+		if (Tkhd->GetDuration() == TNumericLimits<int64>::Max())
+		{
+			UE_LOG(LogElectraProtron, Warning, TEXT("Track #%u has indefinite duration, ignoring this track."), TrackID);
+			continue;
+		}
+
+		// Take note of the track with the longest duration (any track, even unsupported ones)
+		LongestTrackDuration = LongestTrackDuration < Tkhd->GetDuration() ? Tkhd->GetDuration() : LongestTrackDuration;
+
+		ElectraProtronUtils::FCodecInfo& CodecInfo(Tracks[i]->CodecInfo);
+		GetTrackCodecInfo(CodecInfo, AllTracks[i], TrackID);
+
+		// Timecode tracks are inherently usable.
+		if (CodecInfo.Type == ElectraProtronUtils::FCodecInfo::EType::Timecode)
+		{
+			Tracks[i]->bIsUsable = true;
+			Tracks[i]->bIsKeyframeOnlyFormat = true;
+			++NumTimecodeTracks;
+		}
+		// Check with the decoder factory if this format can be decoded.
+		else if (CodecInfo.Type == ElectraProtronUtils::FCodecInfo::EType::Video || CodecInfo.Type == ElectraProtronUtils::FCodecInfo::EType::Audio)
+		{
+			TMap<FString, FVariant> Params;
+			switch(CodecInfo.Type)
+			{
+				case ElectraProtronUtils::FCodecInfo::EType::Video:
+				{
+					const ElectraProtronUtils::FCodecInfo::FVideo& Video(CodecInfo.Properties.Get<ElectraProtronUtils::FCodecInfo::FVideo>());
+					Params.Add(TEXT("width"), FVariant(Video.Width));
+					Params.Add(TEXT("height"), FVariant(Video.Height));
+					if (Video.FrameRate.IsValid())
+					{
+						Params.Add(TEXT("fps"), FVariant((double)Video.FrameRate.GetAsDouble()));
+						Params.Add(TEXT("fps_n"), FVariant((int64)Video.FrameRate.GetNumerator()));
+						Params.Add(TEXT("fps_d"), FVariant((uint32)Video.FrameRate.GetDenominator()));
+					}
+					else
+					{
+						Params.Add(TEXT("fps"), FVariant((double)0.0));
+						Params.Add(TEXT("fps_n"), FVariant((int64)0));
+						Params.Add(TEXT("fps_d"), FVariant((uint32)1));
+					}
+					break;
+				}
+				case ElectraProtronUtils::FCodecInfo::EType::Audio:
+				{
+					const ElectraProtronUtils::FCodecInfo::FAudio& Audio(CodecInfo.Properties.Get<ElectraProtronUtils::FCodecInfo::FAudio>());
+					Params.Add(TEXT("channel_configuration"), FVariant(Audio.ChannelConfiguration));
+					Params.Add(TEXT("num_channels"), FVariant((int32) Audio.NumChannels));
+					Params.Add(TEXT("sample_rate"), FVariant((int32) Audio.SampleRate));
+					break;
+				}
+			}
+			Params.Add(TEXT("dcr"), FVariant(CodecInfo.DCR));
+			Params.Add(TEXT("codec_4cc"), FVariant(CodecInfo.FourCC));
+			Params.Add(TEXT("codec_name"), FVariant(CodecInfo.RFC6381));
+			// Add children box data
+			for(auto& Ch : CodecInfo.ExtraBoxes)
+			{
+				FString BoxName = FString::Printf(TEXT("$%s_box"), *Electra::UtilitiesMP4::GetPrintableBoxAtom(Ch.Key));
+				Params.Add(BoxName, FVariant(Ch.Value));
+			}
+
+			TMap<FString, FVariant> FormatInfo;
+			auto Factory = FactoryModule->GetBestFactoryForFormat(FormatInfo, CodecInfo.RFC6381, false, Params);
+			if (Factory.IsValid())
+			{
+				Tracks[i]->bIsUsable = true;
+				TMap<FString, FVariant> ConfigOptions;
+				Factory->GetConfigurationOptions(ConfigOptions);
+				// Every non-video format is assumed to be keyframe-only. For video we ask the factory.
+				Tracks[i]->bIsKeyframeOnlyFormat = CodecInfo.Type != ElectraProtronUtils::FCodecInfo::EType::Video || !!ElectraDecodersUtil::GetVariantValueSafeI64(FormatInfo, IElectraDecoderFormatInfo::IsEveryFrameKeyframe, 0);
+				// See if the decoder provides a nice name for the format used.
+				Tracks[i]->HumanReadableCodecFormat = ElectraDecodersUtil::GetVariantValueFString(FormatInfo, IElectraDecoderFormatInfo::HumanReadableFormatName);
+				if (Tracks[i]->HumanReadableCodecFormat.IsEmpty())
+				{
+					Tracks[i]->HumanReadableCodecFormat = CodecInfo.HumanReadableFormatInfo;
+					if (Tracks[i]->HumanReadableCodecFormat.IsEmpty())
+					{
+						Tracks[i]->HumanReadableCodecFormat = CodecInfo.RFC6381;
+					}
+				}
+
+				// Take note of the supported track with the longest duration.
+				LongestSupportedTrackDuration = LongestSupportedTrackDuration < Tkhd->GetDuration() ? Tkhd->GetDuration() : LongestSupportedTrackDuration;
+			}
+			else
+			{
+				UE_LOG(LogElectraProtron, Warning, TEXT("No decoder to handle sample type \"%s\" of track #%u, ignoring this track."), *CodecInfo.RFC6381, TrackID);
+			}
+		}
+	}
+
+	// Check that the duration given in the `mvhd` box is not larger than the longest track is.
+	int64 MvhdDur = MovieDuration.GetNumerator();
+	if (MvhdDur > LongestTrackDuration)
+	{
+		UE_LOG(LogElectraProtron, Warning, TEXT("Movie duration in `mvhd` box (%lld) is larger than the longest track (%lld) in the file, adjusting movie duration down."), (long long int)MvhdDur, (long long int)LongestTrackDuration);
+		MovieDuration.SetFromND(LongestTrackDuration, MovieDuration.GetDenominator());
+		Duration = MovieDuration.GetAsTimebase(ETimespan::TicksPerSecond);
+		MvhdDur = LongestTrackDuration;
+	}
+	// Check that the movie duration is not larger than the longest supported track.
+	if (LongestSupportedTrackDuration > 0 && MvhdDur > LongestSupportedTrackDuration)
+	{
+		UE_LOG(LogElectraProtron, Warning, TEXT("Movie duration in `mvhd` box (%lld) is larger than the longest supported track (%lld) in the file, adjusting movie duration down."), (long long int)MvhdDur, (long long int)LongestSupportedTrackDuration);
+		MovieDuration.SetFromND(LongestSupportedTrackDuration, MovieDuration.GetDenominator());
+		Duration = MovieDuration.GetAsTimebase(ETimespan::TicksPerSecond);
+	}
+
+
+	// If there are timecode tracks, find which tracks references them.
+	// If there are none, then - if there is only a single timecode track - we apply it to all other tracks.
+	// Other references are of no interest at the moment.
+	if (NumTimecodeTracks)
+	{
+		bool bAnyTrackReferencesTimecodeExplicitly = false;
+		for(int32 TrkNum=0; TrkNum<Tracks.Num(); ++TrkNum)
+		{
+			// We need to check every track, even the ones we cannot use. Otherwise we
+			// may think the timecode is not referenced and use it for everything!
+			// We do not check if a timecode track references another one as this would be silly.
+			if (Tracks[TrkNum]->CodecInfo.Type != ElectraProtronUtils::FCodecInfo::EType::Timecode)
+			{
+				auto Tref = Tracks[TrkNum]->TrackBox->FindBoxRecursive<Electra::UtilitiesMP4::FMP4BoxTREF>(Electra::UtilitiesMP4::MakeBoxAtom('t','r','e','f'), 1);
+				if (Tref.IsValid())
+				{
+					// Get timecode references, if any.
+					auto References = Tref->GetEntriesOfType(Electra::UtilitiesMP4::MakeBoxAtom('t','m','c','d'));
+					if (References.Num())
+					{
+						if (References.Num() > 1)
+						{
+							UE_LOG(LogElectraProtron, Warning, TEXT("Track #u contains more than one `tmcd` reference box. Using first reference only."), Tracks[TrkNum]->TrackID);
+						}
+						if (References[0].TrackIDs.Num())
+						{
+							if (References[0].TrackIDs.Num() > 1)
+							{
+								UE_LOG(LogElectraProtron, Warning, TEXT("Track #u references more than one timecode track. Using first reference only."), Tracks[TrkNum]->TrackID);
+							}
+							// Whether the reference is actually valid or not, a track makes an explicit reference
+							// and so we must not assign unreferenced timecode tracks later.
+							bAnyTrackReferencesTimecodeExplicitly = true;
+							// Either way we need to tag every referenced track.
+							for(auto& RefTkId : References[0].TrackIDs)
+							{
+								TSharedPtr<FTrackInfo>* ReferencedTrack = Tracks.FindByPredicate([RefTkId](const TSharedPtr<FTrackInfo>& e){ return e->TrackID == RefTkId; });
+								if (ReferencedTrack)
+								{
+									if (!Tracks[TrkNum]->ReferencedTimecodeTrack.IsValid())
+									{
+										Tracks[TrkNum]->ReferencedTimecodeTrack = *ReferencedTrack;
+									}
+									(*ReferencedTrack)->IsReferencedByTracks.Add(Tracks[TrkNum]);
+								}
+								else
+								{
+									UE_LOG(LogElectraProtron, Warning, TEXT("Track #u references a non-existing timecode track #%u. Ignoring."), Tracks[TrkNum]->TrackID, RefTkId);
+									Tracks[TrkNum]->ReferencedTimecodeTrack.Reset();
+								}
+							}
+						}
+					}
+				}
+			}
+		}
+
+		// Now check for tracks that are not explicitly referencing a timecode track, but only when NO track
+		// makes an explicit reference. If some do and other don't we do not interfere.
+		if (!bAnyTrackReferencesTimecodeExplicitly)
+		{
+			if (NumTimecodeTracks == 1)
+			{
+				// Which track is the timecode?
+				TSharedPtr<FTrackInfo>* TimecodeTrack = Tracks.FindByPredicate([](const TSharedPtr<FTrackInfo>& In){ return In->CodecInfo.Type == ElectraProtronUtils::FCodecInfo::EType::Timecode; });
+				check(TimecodeTrack); // there has to be one, otherwise we would not even get here, but for safeties sake...
+				if (TimecodeTrack)
+				{
+					// We set this up for video tracks only.
+					for(int32 i=0; i<Tracks.Num(); ++i)
+					{
+						if (Tracks[i]->CodecInfo.Type != ElectraProtronUtils::FCodecInfo::EType::Video)
+						{
+							continue;
+						}
+						Tracks[i]->ReferencedTimecodeTrack = *TimecodeTrack;
+						(*TimecodeTrack)->IsReferencedByTracks.Add(Tracks[i]);
+					}
+				}
+			}
+			else
+			{
+				UE_LOG(LogElectraProtron, Warning, TEXT("There are %d timecode tracks that are not referenced by any other track. Ignoring all of them."), NumTimecodeTracks);
+			}
+		}
+	}
+	else
+	{
+		// See if there is an `udta` box in the `moov` that contains `(c)TIM` and `(C)TSC` boxes.
+		auto UdtaBox = MoovBox->FindBoxRecursive<Electra::UtilitiesMP4::FMP4BoxUDTA>(Electra::UtilitiesMP4::MakeBoxAtom('u','d','t','a'), 0);
+		if (UdtaBox.IsValid())
+		{
+			auto CTIMBox = UdtaBox->FindBoxRecursive<Electra::UtilitiesMP4::FMP4BoxBase>(Electra::UtilitiesMP4::MakeBoxAtom(0xa9U,'T','I','M'), 0);
+			auto CTSCBox = UdtaBox->FindBoxRecursive<Electra::UtilitiesMP4::FMP4BoxBase>(Electra::UtilitiesMP4::MakeBoxAtom(0xa9U,'T','S','C'), 0);
+			//auto CTSZBox = UdtaBox->FindBoxRecursive<Electra::UtilitiesMP4::FMP4BoxBase>(Electra::UtilitiesMP4::MakeBoxAtom(0xa9U,'T','S','Z'), 0);
+			if (CTIMBox.IsValid() && CTSCBox.IsValid())
+			{
+				auto GetValue = [](const TSharedPtr<Electra::UtilitiesMP4::FMP4BoxBase>& InBox) -> FString
+				{
+					Electra::UtilitiesMP4::FMP4AtomReaderHelper timReader(InBox->GetBoxData());
+					FString String;
+					uint16 StringLength, UnknownValue;
+					if (timReader.Read(StringLength)&&  timReader.Read(UnknownValue))
+					{
+						if (timReader.ReadString(String, StringLength))
+						{
+							return String;
+						}
+					}
+					return FString();
+				};
+				FTrackInfo::FFirstSampleTimecode tc;
+				FFrameRate fr;
+				tc.Timecode = GetValue(CTIMBox);
+				tc.Framerate = GetValue(CTSCBox);
+				TOptional<FTimecode> ptc = FTimecode::ParseTimecode(tc.Timecode);
+				if (ptc.IsSet() && TryParseString(fr, *tc.Framerate))
+				{
+					FFrameNumber fn = ptc.GetValue().ToFrameNumber(fr);
+					tc.TimecodeValue = (uint32) fn.Value;
+				}
+				for(int32 i=0; i<Tracks.Num(); ++i)
+				{
+					if (Tracks[i]->bIsUsable)
+					{
+						Tracks[i]->FirstSampleTimecode = tc;
+					}
+				}
+			}
+		}
+	}
+
+	// One last pass to count how many usable tracks per type we have.
+	const int32 IndexMap[(int32) ElectraProtronUtils::FCodecInfo::EType::MAX] =
+	{
+		CodecTypeIndex(ElectraProtronUtils::FCodecInfo::EType::Video),
+		CodecTypeIndex(ElectraProtronUtils::FCodecInfo::EType::Audio),
+		CodecTypeIndex(ElectraProtronUtils::FCodecInfo::EType::Subtitle),
+		CodecTypeIndex(ElectraProtronUtils::FCodecInfo::EType::Timecode)
+	};
+	UsableTrackArrayIndicesByType.SetNum(UE_ARRAY_COUNT(IndexMap));
+	for(int32 TkTypIdx=0; TkTypIdx<UE_ARRAY_COUNT(IndexMap); ++TkTypIdx)
+	{
+		for(int32 TkIdx=0; TkIdx<Tracks.Num(); ++TkIdx)
+		{
+			if (Tracks[TkIdx]->bIsUsable && TkTypIdx==(int32) Tracks[TkIdx]->CodecInfo.Type)
+			{
+				UsableTrackArrayIndicesByType[TkTypIdx].Emplace(TkIdx);
+			}
+		}
+	}
+
+	// Neither video nor audio?
+	if (UsableTrackArrayIndicesByType[CodecTypeIndex(ElectraProtronUtils::FCodecInfo::EType::Video)].Num() == 0 &&
+		UsableTrackArrayIndicesByType[CodecTypeIndex(ElectraProtronUtils::FCodecInfo::EType::Audio)].Num() == 0)
+	{
+		LastErrorMessage = FString::Printf(TEXT("This file contains no playable video or audio tracks."));
+		return;
+	}
+
+	// Prepare the tracks
+	for(int32 i=0; i<Tracks.Num(); ++i)
+	{
+		if (!Tracks[i]->bIsUsable)
+		{
+			continue;
+		}
+		Tracks[i]->MP4Track = Electra::UtilitiesMP4::FMP4Track::Create(Tracks[i]->TrackBox);
+		if (!Tracks[i]->MP4Track->Prepare(MovieDuration))
+		{
+			LastErrorMessage = Tracks[i]->MP4Track->GetLastError();
+			check(!LastErrorMessage.IsEmpty());
+			return;
+		}
+
+		// If this is a timecode track we may want to read in the first timecode.
+		if (Config.bReadFirstTimecode && Tracks[i]->CodecInfo.Type == ElectraProtronUtils::FCodecInfo::EType::Timecode &&
+			Tracks[i]->CodecInfo.FourCC == Electra::UtilitiesMP4::MakeBoxAtom('t','m','c','d'))
+		{
+			auto It = Tracks[i]->MP4Track->CreateIterator();
+			if (!It.IsValid())
+			{
+				LastErrorMessage = Tracks[i]->MP4Track->GetLastError();
+				check(!LastErrorMessage.IsEmpty());
+				return;
+			}
+			// Read the first sample.
+			int64 SampleSize = It->GetSampleSize();
+			int64 SampleFileOffset = It->GetSampleFileOffset();
+			check(SampleSize == 4);		// If that triggers this is not a tmcd sample.
+			TArray<uint32> TimecodeBuffer;
+			TimecodeBuffer.SetNumUninitialized(Align(SampleSize, sizeof(uint32)));
+			int64 NumRead = Reader->ReadData(TimecodeBuffer.GetData(), SampleSize, SampleFileOffset, Electra::IBaseDataReader::FCancellationCheckDelegate::CreateLambda([&](){return bAbort;}));
+			if (NumRead != SampleSize)
+			{
+				LastErrorMessage = FString::Printf(TEXT("Failed to read first timecode sample."));
+				return;
+			}
+			// Get the timecode description from the codec info.
+			ElectraProtronUtils::FCodecInfo::FTMCDTimecode TimecodeInfo(Tracks[i]->CodecInfo.Properties.Get<ElectraProtronUtils::FCodecInfo::FTMCDTimecode>());
+			// Set the first timecode sample on the track.
+			FTrackInfo::FFirstSampleTimecode firstTC;
+			firstTC.TimecodeValue = Electra::GetFromBigEndian(TimecodeBuffer[0]);
+			firstTC.Framerate = TimecodeInfo.GetFrameRate().ToPrettyText().ToString();
+			firstTC.Timecode = TimecodeInfo.ConvertToTimecode(firstTC.TimecodeValue).ToString();
+			Tracks[i]->FirstSampleTimecode.Emplace(MoveTemp(firstTC));
+
+			// For convenience also set this on the tracks that reference this timecode track
+			for(auto ReferencingTrack : Tracks[i]->IsReferencedByTracks)
+			{
+				if (auto PinnedRefTrk = ReferencingTrack.Pin())
+				{
+					PinnedRefTrk->FirstSampleTimecode = Tracks[i]->FirstSampleTimecode;
+				}
+			}
+		}
+	}
+}
+
+void FElectraProtronPlayer::FImpl::GetTrackCodecInfo(ElectraProtronUtils::FCodecInfo& OutCodecInfo, const TSharedPtr<Electra::UtilitiesMP4::FMP4BoxTRAK>& InTrack, uint32 InTrackID)
+{
+	// There needs to be an `stsd` box in this track. We do not try the expected path of `trak`->`mdia`->`minf`->`stbl`->`stsd` as
+	// this is not that much faster and if the file is somewhat ill-formed we may not find it if it's grouped under elsewhere.
+	auto Stsd = InTrack->FindBoxRecursive<Electra::UtilitiesMP4::FMP4BoxSTSD>(Electra::UtilitiesMP4::MakeBoxAtom('s','t','s','d'), 6);
+	if (!Stsd)
+	{
+		// If not found we just ignore the track. That's a warning but not an error.
+		UE_LOG(LogElectraProtron, Warning, TEXT("No `stsd` box found, ignoring track #%u."), InTrackID);
+		return;
+	}
+
+	auto AddChildren = [&OutCodecInfo](const TSharedPtr<Electra::UtilitiesMP4::FMP4BoxBase>& InSample) -> void
+	{
+		for(auto& Ch : InSample->GetChildren())
+		{
+			OutCodecInfo.ExtraBoxes.Emplace(Ch->GetType(), Ch->GetBoxData());
+		}
+	};
+
+	// Get the sample type of this track. This call is necessary to trigger parsing of child nodes.
+	Electra::UtilitiesMP4::FMP4BoxSampleEntry::ESampleType SampleType = Stsd->GetSampleType();
+	// If already knows to not be supported, skip it.
+	if (SampleType == Electra::UtilitiesMP4::FMP4BoxSampleEntry::ESampleType::Unsupported)
+	{
+		UE_LOG(LogElectraProtron, Warning, TEXT("Unsupported sample type, ignoring track #%u."), InTrackID);
+		return;
+	}
+	// Several entries are permitted, but we need this to be unambiguous.
+	if (Stsd->GetChildren().Num() > 1)
+	{
+		UE_LOG(LogElectraProtron, Warning, TEXT("Multiple sample descriptions found in `stsd` box, ignoring track #%u."), InTrackID);
+		return;
+	}
+	else if (Stsd->GetChildren().Num() == 0)
+	{
+		UE_LOG(LogElectraProtron, Warning, TEXT("No sample description found in `stsd` box, ignoring track #%u."), InTrackID);
+		return;
+	}
+
+	bool bIsEncrypted = false;
+	// Based on the sample type, get it and see if it is using codec we support.
+	if (SampleType == Electra::UtilitiesMP4::FMP4BoxSampleEntry::ESampleType::Video)
+	{
+		auto SetResoAndFPSFromBox = [InTrack](ElectraProtronUtils::FCodecInfo::FVideo& InOutVideo, const TSharedPtr<Electra::UtilitiesMP4::FMP4BoxVisualSampleEntry>& InVisual) -> void
+		{
+			InOutVideo.Width = InVisual->GetWidth();
+			InOutVideo.Height = InVisual->GetHeight();
+			auto Mdhd = InTrack->FindBoxRecursive<Electra::UtilitiesMP4::FMP4BoxMDHD>(Electra::UtilitiesMP4::MakeBoxAtom('m','d','h','d'), 2);
+			auto Stts = InTrack->FindBoxRecursive<Electra::UtilitiesMP4::FMP4BoxSTTS>(Electra::UtilitiesMP4::MakeBoxAtom('s','t','t','s'), 5);
+			if (Mdhd.IsValid() && Stts.IsValid() && Stts->GetEntries().Num())
+			{
+				InOutVideo.FrameRate.SetFromND(Mdhd->GetTimescale(), Stts->GetEntries()[0].sample_delta);
+			}
+		};
+
+		TArray<TSharedPtr<Electra::UtilitiesMP4::FMP4BoxVisualSampleEntry>> Visuals;
+		Stsd->GetSampleDescriptions(Visuals);
+		check(Visuals.Num() == 1);
+		switch(Visuals[0]->GetType())
+		{
+			case Electra::UtilitiesMP4::MakeBoxAtom('e','n','c','v'):
+			{
+				bIsEncrypted = true;
+				break;
+			}
+			case Electra::UtilitiesMP4::MakeBoxAtom('a','v','c','1'):
+			case Electra::UtilitiesMP4::MakeBoxAtom('a','v','c','3'):
+			{
+				if (Visuals[0]->GetFrameCount() == 1 && Visuals[0]->GetDepth() == 0x18)
+				{
+					auto Avcc = Visuals[0]->FindBoxRecursive<Electra::UtilitiesMP4::FMP4BoxAVCC>(Electra::UtilitiesMP4::MakeBoxAtom('a','v','c','C'), 0);
+					if (Avcc.IsValid())
+					{
+						ElectraDecodersUtil::MPEG::H264::FAVCDecoderConfigurationRecord dcr;
+						if (dcr.Parse(Avcc->GetAVCDecoderConfigurationRecord()))
+						{
+							ElectraProtronUtils::FCodecInfo::FVideo Video;
+							ElectraDecodersUtil::MPEG::H264::FSequenceParameterSet sps;
+							const TCHAR* Prefix = nullptr;
+							if (dcr.GetSequenceParameterSets().Num() && ElectraDecodersUtil::MPEG::H264::ParseSequenceParameterSet(sps, dcr.GetSequenceParameterSets()[0]))
+							{
+								ElectraDecodersUtil::FFractionalValue fr = sps.GetTiming();
+								sps.GetDisplaySize(Video.Width, Video.Height);
+								Video.FrameRate.SetFromND(fr.Num, fr.Denom);
+								Prefix = TEXT("avc1");
+							}
+							else
+							{
+								SetResoAndFPSFromBox(Video, Visuals[0]);
+								Prefix = TEXT("avc3");
+							}
+							OutCodecInfo.Properties.Emplace<ElectraProtronUtils::FCodecInfo::FVideo>(MoveTemp(Video));
+							OutCodecInfo.RFC6381 = dcr.GetCodecSpecifierRFC6381(Prefix);
+							OutCodecInfo.FourCC = Visuals[0]->GetType();
+							OutCodecInfo.DCR = Avcc->GetAVCDecoderConfigurationRecord();
+							OutCodecInfo.CSD = dcr.GetCodecSpecificData();
+							OutCodecInfo.HumanReadableFormatInfo = dcr.GetFormatInfo();
+							OutCodecInfo.Type = ElectraProtronUtils::FCodecInfo::EType::Video;
+							AddChildren(Visuals[0]);
+						}
+					}
+				}
+				break;
+			}
+			case Electra::UtilitiesMP4::MakeBoxAtom('h','v','c','1'):
+			case Electra::UtilitiesMP4::MakeBoxAtom('h','e','v','1'):
+			{
+				if (Visuals[0]->GetFrameCount() == 1 && Visuals[0]->GetDepth() == 0x18)
+				{
+					auto Hvcc = Visuals[0]->FindBoxRecursive<Electra::UtilitiesMP4::FMP4BoxHVCC>(Electra::UtilitiesMP4::MakeBoxAtom('h','v','c','C'), 0);
+					if (Hvcc.IsValid())
+					{
+						ElectraDecodersUtil::MPEG::H265::FHEVCDecoderConfigurationRecord dcr;
+						if (dcr.Parse(Hvcc->GetHEVCDecoderConfigurationRecord()))
+						{
+							ElectraProtronUtils::FCodecInfo::FVideo Video;
+							ElectraDecodersUtil::MPEG::H265::FSequenceParameterSet sps;
+							const TCHAR* Prefix = nullptr;
+							if (dcr.GetSequenceParameterSets().Num() && ElectraDecodersUtil::MPEG::H265::ParseSequenceParameterSet(sps, dcr.GetSequenceParameterSets()[0]))
+							{
+								ElectraDecodersUtil::FFractionalValue fr = sps.GetTiming();
+								sps.GetDisplaySize(Video.Width, Video.Height);
+								Video.FrameRate.SetFromND(fr.Num, fr.Denom);
+								Prefix = TEXT("hvc1");
+							}
+							else
+							{
+								SetResoAndFPSFromBox(Video, Visuals[0]);
+								Prefix = TEXT("hev1");
+							}
+							OutCodecInfo.Properties.Emplace<ElectraProtronUtils::FCodecInfo::FVideo>(MoveTemp(Video));
+							OutCodecInfo.RFC6381 = dcr.GetCodecSpecifierRFC6381(Prefix);
+							OutCodecInfo.FourCC = Visuals[0]->GetType();
+							OutCodecInfo.DCR = Hvcc->GetHEVCDecoderConfigurationRecord();
+							OutCodecInfo.CSD = dcr.GetCodecSpecificData();
+							OutCodecInfo.HumanReadableFormatInfo = dcr.GetFormatInfo();
+							OutCodecInfo.Type = ElectraProtronUtils::FCodecInfo::EType::Video;
+							AddChildren(Visuals[0]);
+						}
+					}
+				}
+				break;
+			}
+			default:
+			{
+				ElectraProtronUtils::FCodecInfo::FVideo Video;
+				SetResoAndFPSFromBox(Video, Visuals[0]);
+				OutCodecInfo.Properties.Emplace<ElectraProtronUtils::FCodecInfo::FVideo>(MoveTemp(Video));
+				OutCodecInfo.RFC6381 = Electra::UtilitiesMP4::GetPrintableBoxAtom(Visuals[0]->GetType());
+				OutCodecInfo.FourCC = Visuals[0]->GetType();
+				OutCodecInfo.DCR = Visuals[0]->GetBoxData();
+				OutCodecInfo.Type = ElectraProtronUtils::FCodecInfo::EType::Video;
+				AddChildren(Visuals[0]);
+				break;
+			}
+		}
+	}
+	else if (SampleType == Electra::UtilitiesMP4::FMP4BoxSampleEntry::ESampleType::Audio)
+	{
+		TArray<TSharedPtr<Electra::UtilitiesMP4::FMP4BoxAudioSampleEntry>> Audios;
+		Stsd->GetSampleDescriptions(Audios);
+		check(Audios.Num() == 1);
+
+		switch(Audios[0]->GetType())
+		{
+			case Electra::UtilitiesMP4::MakeBoxAtom('e','n','c','a'):
+			{
+				bIsEncrypted = true;
+				break;
+			}
+			case Electra::UtilitiesMP4::MakeBoxAtom('m','p','4','a'):
+			{
+				// Search down one extra level since the `esds` might be contained within a `wave` in QuickTime.
+				auto Esds = Audios[0]->FindBoxRecursive<Electra::UtilitiesMP4::FMP4BoxESDS>(Electra::UtilitiesMP4::MakeBoxAtom('e','s','d','s'), 1);
+				ElectraDecodersUtil::MPEG::FESDescriptor esds;
+				if (Esds.IsValid() && esds.Parse(Esds->GetESDescriptor()))
+				{
+					// AAC audio?
+					if (esds.GetObjectTypeID() == ElectraDecodersUtil::MPEG::FESDescriptor::FObjectTypeID::MPEG4_Audio &&
+						esds.GetStreamType() == ElectraDecodersUtil::MPEG::FESDescriptor::FStreamType::AudioStream)
+					{
+						ElectraDecodersUtil::MPEG::FAACDecoderConfigurationRecord dcr;
+						if (dcr.Parse(esds.GetCodecSpecificData()))
+						{
+							ElectraProtronUtils::FCodecInfo::FAudio Audio;
+							Audio.SampleRate = dcr.SamplingRate;
+							Audio.ChannelConfiguration = dcr.ChannelConfiguration;
+							Audio.NumChannels = ElectraDecodersUtil::MPEG::AACUtils::GetNumberOfChannelsFromChannelConfiguration(dcr.ChannelConfiguration);
+							OutCodecInfo.Properties.Emplace<ElectraProtronUtils::FCodecInfo::FAudio>(MoveTemp(Audio));
+							OutCodecInfo.RFC6381 = dcr.GetCodecSpecifierRFC6381();
+							OutCodecInfo.FourCC = Audios[0]->GetType();
+							OutCodecInfo.DCR = esds.GetCodecSpecificData();
+							OutCodecInfo.CSD = dcr.GetCodecSpecificData();
+							OutCodecInfo.HumanReadableFormatInfo = dcr.GetFormatInfo();
+							OutCodecInfo.Type = ElectraProtronUtils::FCodecInfo::EType::Audio;
+							AddChildren(Audios[0]);
+						}
+					}
+					// MPEG audio?
+					else if (esds.GetObjectTypeID() == ElectraDecodersUtil::MPEG::FESDescriptor::FObjectTypeID::MPEG1_Audio &&
+							 esds.GetStreamType() == ElectraDecodersUtil::MPEG::FESDescriptor::FStreamType::AudioStream)
+					{
+						ElectraProtronUtils::FCodecInfo::FAudio Audio;
+						Audio.SampleRate = Audios[0]->GetSampleRate();
+						Audio.NumChannels = Audios[0]->GetChannelCount();
+						OutCodecInfo.Properties.Emplace<ElectraProtronUtils::FCodecInfo::FAudio>(MoveTemp(Audio));
+						OutCodecInfo.RFC6381 = TEXT("mp4a.40.34");
+						OutCodecInfo.FourCC = Electra::UtilitiesMP4::MakeBoxAtom('m','p','g','a');
+						OutCodecInfo.HumanReadableFormatInfo = TEXT("MPEG audio");
+						OutCodecInfo.Type = ElectraProtronUtils::FCodecInfo::EType::Audio;
+						AddChildren(Audios[0]);
+					}
+				}
+				break;
+			}
+			default:
+			{
+				ElectraProtronUtils::FCodecInfo::FAudio Audio;
+				Audio.SampleRate = Audios[0]->GetSampleRate();
+				Audio.ChannelConfiguration = 0;
+				Audio.NumChannels = Audios[0]->GetChannelCount();
+				OutCodecInfo.Properties.Emplace<ElectraProtronUtils::FCodecInfo::FAudio>(MoveTemp(Audio));
+				OutCodecInfo.RFC6381 = Electra::UtilitiesMP4::GetPrintableBoxAtom(Audios[0]->GetType());
+				OutCodecInfo.FourCC = Audios[0]->GetType();
+				OutCodecInfo.DCR = Audios[0]->GetBoxData();
+				OutCodecInfo.Type = ElectraProtronUtils::FCodecInfo::EType::Audio;
+				AddChildren(Audios[0]);
+				break;
+			}
+		}
+
+	}
+	else if (SampleType == Electra::UtilitiesMP4::FMP4BoxSampleEntry::ESampleType::QTFFTimecode)
+	{
+		TArray<TSharedPtr<Electra::UtilitiesMP4::FMP4BoxQTFFTimecodeSampleEntry>> Timecodes;
+		Stsd->GetSampleDescriptions(Timecodes);
+		check(Timecodes.Num() == 1);
+		switch(Timecodes[0]->GetType())
+		{
+			case Electra::UtilitiesMP4::MakeBoxAtom('t','m','c','d'):
+			{
+				ElectraProtronUtils::FCodecInfo::FTMCDTimecode Timecode;
+
+				Timecode.Flags = Timecodes[0]->GetFlags();
+				Timecode.Timescale = Timecodes[0]->GetTimescale();
+				Timecode.FrameDuration = Timecodes[0]->GetFrameDuration();
+				Timecode.NumberOfFrames = Timecodes[0]->GetNumberOfFrames();
+				OutCodecInfo.Properties.Emplace<ElectraProtronUtils::FCodecInfo::FTMCDTimecode>(MoveTemp(Timecode));
+				OutCodecInfo.RFC6381 = Electra::UtilitiesMP4::GetPrintableBoxAtom(Timecodes[0]->GetType());
+				OutCodecInfo.FourCC = Timecodes[0]->GetType();
+				OutCodecInfo.DCR = Timecodes[0]->GetBoxData();
+				OutCodecInfo.Type = ElectraProtronUtils::FCodecInfo::EType::Timecode;
+				AddChildren(Timecodes[0]);
+				break;
+			}
+		}
+	}
+
+	// If the track has not been identified as usable so far, remove it.
+	if (OutCodecInfo.Type == ElectraProtronUtils::FCodecInfo::EType::Invalid)
+	{
+		if (!bIsEncrypted)
+		{
+			TArray<TSharedPtr<Electra::UtilitiesMP4::FMP4BoxSampleEntry>> Entries;
+			Stsd->GetSampleDescriptions(Entries);
+			check(Entries.Num() == 1);
+			UE_LOG(LogElectraProtron, Warning, TEXT("Track of sample type \"%s\" is not supported, ignoring track #%u."), *Electra::UtilitiesMP4::GetPrintableBoxAtom(Entries[0]->GetType()), InTrackID);
+		}
+		else
+		{
+			UE_LOG(LogElectraProtron, Warning, TEXT("Track is using encryption, ignoring track #%u."), InTrackID);
+		}
+	}
+}
+
+
+
+FVariant FElectraProtronPlayer::FImpl::GetMediaInfo(FName InInfoName)
+{
+	if (InInfoName == ElectraProtronOptionNames::StartTimecodeValue || InInfoName == ElectraProtronOptionNames::StartTimecodeFrameRate)
+	{
+		auto ci = CodecTypeIndex(ElectraProtronUtils::FCodecInfo::EType::Video);
+		int32 SelectedVideoTrackIndex = TrackSelection.SelectedTrackIndex[ci];
+		if (SelectedVideoTrackIndex >= 0 && SelectedVideoTrackIndex < UsableTrackArrayIndicesByType[ci].Num())
+		{
+			const FTrackInfo& ti(*Tracks[UsableTrackArrayIndicesByType[ci][SelectedVideoTrackIndex]]);
+			if (ti.FirstSampleTimecode.IsSet())
+			{
+				if (InInfoName == ElectraProtronOptionNames::StartTimecodeValue)
+				{
+					return FVariant(ti.FirstSampleTimecode.GetValue().Timecode);
+				}
+				else //if (InInfoName == ElectraProtronOptionNames::StartTimecodeFrameRate)
+				{
+					return FVariant(ti.FirstSampleTimecode.GetValue().Framerate);
+				}
+			}
+		}
+	}
+	return FVariant();
+}
+
+
+
+int32 FElectraProtronPlayer::FImpl::GetNumTracks(EMediaTrackType InTrackType)
+{
+	switch(InTrackType)
+	{
+		case EMediaTrackType::Video:
+		{
+			return UsableTrackArrayIndicesByType[CodecTypeIndex(ElectraProtronUtils::FCodecInfo::EType::Video)].Num();
+		}
+		case EMediaTrackType::Audio:
+		{
+			return UsableTrackArrayIndicesByType[CodecTypeIndex(ElectraProtronUtils::FCodecInfo::EType::Audio)].Num();
+		}
+	}
+	return 0;
+}
+int32 FElectraProtronPlayer::FImpl::GetNumTrackFormats(EMediaTrackType InTrackType, int32 InTrackIndex)
+{
+	if (InTrackIndex >= 0)
+	{
+		// Every track this player supports, if the track exists, only has a single format.
+		switch(InTrackType)
+		{
+			case EMediaTrackType::Video:
+			{
+				return InTrackIndex < UsableTrackArrayIndicesByType[CodecTypeIndex(ElectraProtronUtils::FCodecInfo::EType::Video)].Num() ? 1 : 0;
+			}
+			case EMediaTrackType::Audio:
+			{
+				return InTrackIndex < UsableTrackArrayIndicesByType[CodecTypeIndex(ElectraProtronUtils::FCodecInfo::EType::Audio)].Num() ? 1 : 0;
+			}
+		}
+	}
+	return 0;
+}
+
+int32 FElectraProtronPlayer::FImpl::GetTrackFormat(EMediaTrackType InTrackType, int32 InTrackIndex)
+{
+	if (InTrackIndex >= 0)
+	{
+		// Every track this player supports, if the track exists, only has a single format.
+		switch(InTrackType)
+		{
+			case EMediaTrackType::Video:
+			{
+				return InTrackIndex < UsableTrackArrayIndicesByType[CodecTypeIndex(ElectraProtronUtils::FCodecInfo::EType::Video)].Num() ? 0 : -1;
+			}
+			case EMediaTrackType::Audio:
+			{
+				return InTrackIndex < UsableTrackArrayIndicesByType[CodecTypeIndex(ElectraProtronUtils::FCodecInfo::EType::Audio)].Num() ? 0 : -1;
+			}
+		}
+	}
+	return -1;
+}
+
+FText FElectraProtronPlayer::FImpl::GetTrackDisplayName(EMediaTrackType InTrackType, int32 InTrackIndex)
+{
+	if (InTrackIndex >= 0)
+	{
+		switch(InTrackType)
+		{
+			case EMediaTrackType::Video:
+			{
+				if (InTrackIndex < UsableTrackArrayIndicesByType[CodecTypeIndex(ElectraProtronUtils::FCodecInfo::EType::Video)].Num())
+				{
+					const FTrackInfo& ti(*Tracks[UsableTrackArrayIndicesByType[CodecTypeIndex(ElectraProtronUtils::FCodecInfo::EType::Video)][InTrackIndex]]);
+					FString Name(ti.MP4Track->GetCommonMetadata().Name);
+					if (Name.IsEmpty())
+					{
+						Name = FString::Printf(TEXT("Video track #%u"), ti.TrackID);
+					}
+					return FText::FromString(Name);
+				}
+				break;
+			}
+			case EMediaTrackType::Audio:
+			{
+				if (InTrackIndex < UsableTrackArrayIndicesByType[CodecTypeIndex(ElectraProtronUtils::FCodecInfo::EType::Audio)].Num())
+				{
+					const FTrackInfo& ti(*Tracks[UsableTrackArrayIndicesByType[CodecTypeIndex(ElectraProtronUtils::FCodecInfo::EType::Audio)][InTrackIndex]]);
+					FString Name(ti.MP4Track->GetCommonMetadata().Name);
+					if (Name.IsEmpty())
+					{
+						Name = FString::Printf(TEXT("Audio track #%u"), ti.TrackID);
+					}
+					return FText::FromString(Name);
+				}
+				break;
+			}
+		}
+	}
+	return FText();
+}
+
+FString FElectraProtronPlayer::FImpl::GetTrackLanguage(EMediaTrackType InTrackType, int32 InTrackIndex)
+{
+	if (InTrackIndex >= 0)
+	{
+		switch(InTrackType)
+		{
+			case EMediaTrackType::Video:
+			{
+				if (InTrackIndex < UsableTrackArrayIndicesByType[CodecTypeIndex(ElectraProtronUtils::FCodecInfo::EType::Video)].Num())
+				{
+					return Tracks[UsableTrackArrayIndicesByType[CodecTypeIndex(ElectraProtronUtils::FCodecInfo::EType::Video)][InTrackIndex]]->MP4Track->GetCommonMetadata().Language;
+				}
+				break;
+			}
+			case EMediaTrackType::Audio:
+			{
+				if (InTrackIndex < UsableTrackArrayIndicesByType[CodecTypeIndex(ElectraProtronUtils::FCodecInfo::EType::Audio)].Num())
+				{
+					return Tracks[UsableTrackArrayIndicesByType[CodecTypeIndex(ElectraProtronUtils::FCodecInfo::EType::Audio)][InTrackIndex]]->MP4Track->GetCommonMetadata().Language;
+				}
+				break;
+			}
+		}
+	}
+	return FString();
+}
+
+FString FElectraProtronPlayer::FImpl::GetTrackName(EMediaTrackType InTrackType, int32 InTrackIndex)
+{
+	if (InTrackIndex >= 0)
+	{
+		switch(InTrackType)
+		{
+			case EMediaTrackType::Video:
+			{
+				if (InTrackIndex < UsableTrackArrayIndicesByType[CodecTypeIndex(ElectraProtronUtils::FCodecInfo::EType::Video)].Num())
+				{
+					return FString::Printf(TEXT("%u"), Tracks[UsableTrackArrayIndicesByType[CodecTypeIndex(ElectraProtronUtils::FCodecInfo::EType::Video)][InTrackIndex]]->TrackID);
+				}
+				break;
+			}
+			case EMediaTrackType::Audio:
+			{
+				if (InTrackIndex < UsableTrackArrayIndicesByType[CodecTypeIndex(ElectraProtronUtils::FCodecInfo::EType::Audio)].Num())
+				{
+					return FString::Printf(TEXT("%u"), Tracks[UsableTrackArrayIndicesByType[CodecTypeIndex(ElectraProtronUtils::FCodecInfo::EType::Audio)][InTrackIndex]]->TrackID);
+				}
+				break;
+			}
+		}
+	}
+	return FString();
+}
+
+bool FElectraProtronPlayer::FImpl::GetVideoTrackFormat(int32 InTrackIndex, int32 InFormatIndex, FMediaVideoTrackFormat& OutFormat)
+{
+	if (InTrackIndex >= 0 && InTrackIndex < UsableTrackArrayIndicesByType[CodecTypeIndex(ElectraProtronUtils::FCodecInfo::EType::Video)].Num() && InFormatIndex == 0)
+	{
+		const FTrackInfo& ti(*Tracks[UsableTrackArrayIndicesByType[CodecTypeIndex(ElectraProtronUtils::FCodecInfo::EType::Video)][InTrackIndex]]);
+		const ElectraProtronUtils::FCodecInfo::FVideo& vi(ti.CodecInfo.Properties.Get<ElectraProtronUtils::FCodecInfo::FVideo>());
+		OutFormat.Dim.X = (int32) vi.Width;
+		OutFormat.Dim.Y = (int32) vi.Height;
+		OutFormat.FrameRate = (float) vi.FrameRate.GetAsDouble();
+		OutFormat.FrameRates = TRange<float>{ OutFormat.FrameRate };
+		OutFormat.TypeName = ti.HumanReadableCodecFormat;
+		return true;
+	}
+	return false;
+}
+
+bool FElectraProtronPlayer::FImpl::GetAudioTrackFormat(int32 InTrackIndex, int32 InFormatIndex, FMediaAudioTrackFormat& OutFormat)
+{
+	if (InTrackIndex >= 0 && InTrackIndex < UsableTrackArrayIndicesByType[CodecTypeIndex(ElectraProtronUtils::FCodecInfo::EType::Audio)].Num() && InFormatIndex == 0)
+	{
+		const FTrackInfo& ti(*Tracks[UsableTrackArrayIndicesByType[CodecTypeIndex(ElectraProtronUtils::FCodecInfo::EType::Audio)][InTrackIndex]]);
+		const ElectraProtronUtils::FCodecInfo::FAudio& ai(ti.CodecInfo.Properties.Get<ElectraProtronUtils::FCodecInfo::FAudio>());
+		OutFormat.BitsPerSample = 16;
+		OutFormat.NumChannels = ai.NumChannels;
+		OutFormat.SampleRate = ai.SampleRate;
+		OutFormat.TypeName = ti.HumanReadableCodecFormat;
+		return true;
+	}
+	return false;
+}
+
+int32 FElectraProtronPlayer::FImpl::GetSelectedTrack(EMediaTrackType InTrackType)
+{
+	switch(InTrackType)
+	{
+		case EMediaTrackType::Video:
+		{
+			return TrackSelection.SelectedTrackIndex[CodecTypeIndex(ElectraProtronUtils::FCodecInfo::EType::Video)];
+		}
+		case EMediaTrackType::Audio:
+		{
+			return TrackSelection.SelectedTrackIndex[CodecTypeIndex(ElectraProtronUtils::FCodecInfo::EType::Audio)];
+		}
+	}
+	return -1;
+}
+
+
+bool FElectraProtronPlayer::FImpl::SelectTrack(EMediaTrackType InTrackType, int32 InTrackIndex)
+{
+	if (InTrackIndex < -1)
+	{
+		InTrackIndex = -1;
+	}
+	switch(InTrackType)
+	{
+		case EMediaTrackType::Video:
+		{
+			const auto TypeIndex = CodecTypeIndex(ElectraProtronUtils::FCodecInfo::EType::Video);
+			if (InTrackIndex < UsableTrackArrayIndicesByType[TypeIndex].Num())
+			{
+				if (TrackSelection.SelectedTrackIndex[TypeIndex] != InTrackIndex)
+				{
+					TrackSelection.SelectedTrackIndex[TypeIndex] = InTrackIndex;
+					TrackSelection.bChanged = true;
+					bAreRatesValid = false;
+					UpdateTrackLoader(TypeIndex);
+					HandleActiveTrackChanges();
+				}
+				return true;
+			}
+			break;
+		}
+		case EMediaTrackType::Audio:
+		{
+			const auto TypeIndex = CodecTypeIndex(ElectraProtronUtils::FCodecInfo::EType::Audio);
+			if (InTrackIndex < UsableTrackArrayIndicesByType[TypeIndex].Num())
+			{
+				if (TrackSelection.SelectedTrackIndex[TypeIndex] != InTrackIndex)
+				{
+					TrackSelection.SelectedTrackIndex[TypeIndex] = InTrackIndex;
+					TrackSelection.bChanged = true;
+					bAreRatesValid = false;
+					UpdateTrackLoader(TypeIndex);
+					HandleActiveTrackChanges();
+				}
+				return true;
+			}
+			break;
+		}
+		default:
+		{
+			break;
+		}
+	}
+	return false;
+}
+
+bool FElectraProtronPlayer::FImpl::SetTrackFormat(EMediaTrackType InTrackType, int32 InTrackIndex, int32 InFormatIndex)
+{
+	return false;
+}
+
+bool FElectraProtronPlayer::FImpl::QueryCacheState(EMediaCacheState InState, TRangeSet<FTimespan>& OutTimeRanges)
+{
+	if (InState == EMediaCacheState::Loading)
+	{
+		if (TrackSelection.SelectedTrackIndex[CodecTypeIndex(ElectraProtronUtils::FCodecInfo::EType::Video)] >= 0)
+		{
+			OutTimeRanges = VideoLoaderThread.GetTimeRangesToLoad();
+			return true;
+		}
+		else if (TrackSelection.SelectedTrackIndex[CodecTypeIndex(ElectraProtronUtils::FCodecInfo::EType::Audio)] >= 0)
+		{
+			OutTimeRanges = AudioLoaderThread.GetTimeRangesToLoad();
+			return true;
+		}
+	}
+	return false;
+}
+
+int32 FElectraProtronPlayer::FImpl::GetSampleCount(EMediaCacheState InState)
+{
+	return 0;
+}
+
+float FElectraProtronPlayer::FImpl::GetRate()
+{
+	return CurrentRate;
+}
+
+bool FElectraProtronPlayer::FImpl::SetRate(float InRate)
+{
+	HandleActiveTrackChanges();
+
+	IntendedRate = InRate;
+	HandleRateChanges();
+
+	return true;
+}
+
+FTimespan FElectraProtronPlayer::FImpl::GetTime()
+{
+	return CurrentPlayPosTime;
+}
+
+bool FElectraProtronPlayer::FImpl::SetLooping(bool bInLooping)
+{
+	bool bOk = true;
+	if (bOk && !VideoDecoderThread.SetLooping(bInLooping))
+	{
+		bOk = false;
+	}
+	if (bOk && !AudioDecoderThread.SetLooping(bInLooping))
+	{
+		bOk = false;
+	}
+	SharedPlayParams->bShouldLoop = bInLooping && bOk;
+	return bOk;
+}
+
+bool FElectraProtronPlayer::FImpl::IsLooping()
+{
+	return SharedPlayParams->bShouldLoop;
+}
+
+bool FElectraProtronPlayer::FImpl::Seek(const FTimespan& InTime, const TSharedPtr<FMediaSamples, ESPMode::ThreadSafe>& InNewSampleQueue, uint32 InNewSequenceIndex)
+{
+	FSeekRequest sr;
+	sr.NewSampleQueue = InNewSampleQueue;
+	sr.NewTime = InTime;
+	sr.NewSequenceIndex = InNewSequenceIndex;
+	SeekRequestLock.Lock();
+	PendingSeekRequest = MoveTemp(sr);
+	SeekRequestLock.Unlock();
+	WorkMessageSignal.Signal();
+	return true;
+}
+
+void FElectraProtronPlayer::FImpl::TickFetch(FTimespan InDeltaTime, FTimespan InTimecode)
+{
+}
+
+void FElectraProtronPlayer::FImpl::TickInput(FTimespan InDeltaTime, FTimespan InTimecode)
+{
+	if (LastErrorMessage.IsEmpty())
+	{
+		if (!VideoDecoderThread.GetLastError().IsEmpty())
+		{
+			LastErrorMessage = VideoDecoderThread.GetLastError();
+		}
+		else if (!AudioDecoderThread.GetLastError().IsEmpty())
+		{
+			LastErrorMessage = AudioDecoderThread.GetLastError();
+		}
+		else if (!VideoLoaderThread.GetLastError().IsEmpty())
+		{
+			LastErrorMessage = VideoLoaderThread.GetLastError();
+		}
+		else if (!AudioLoaderThread.GetLastError().IsEmpty())
+		{
+			LastErrorMessage = AudioLoaderThread.GetLastError();
+		}
+	}
+
+	double Rate = CurrentRate;
+	if (Rate != 0.0)
+	{
+		FTimespan ScaledDeltaTime = InDeltaTime * Rate;
+		FTimespan NewPos = CurrentPlayPosTime + ScaledDeltaTime;
+		if (NewPos < FTimespan::Zero())
+		{
+			NewPos = IsLooping() ? Duration - (NewPos % Duration) : FTimespan::Zero();
+		}
+		else if (NewPos > Duration)
+		{
+			NewPos = IsLooping() ? NewPos % Duration : Duration;
+		}
+		VideoDecoderThread.SetEstimatedPlaybackTime(NewPos);
+		AudioDecoderThread.SetEstimatedPlaybackTime(NewPos);
+		CurrentPlayPosTime = NewPos;
+	}
+}
+
+void FElectraProtronPlayer::FImpl::UpdateTrackLoader(int32 InCodecTypeIndex)
+{
+	if (TrackSelection.SelectedTrackIndex[InCodecTypeIndex] >= 0)
+	{
+		auto Track = Tracks[UsableTrackArrayIndicesByType[InCodecTypeIndex][TrackSelection.SelectedTrackIndex[InCodecTypeIndex]]];
+		check(Track.IsValid());
+		if (Track.IsValid())
+		{
+			// Do we have a track sample buffer for this track?
+			if (!TrackSampleBuffers.Contains(Track->TrackID))
+			{
+				// No, create it now.
+				TSharedPtr<FMP4TrackSampleBuffer, ESPMode::ThreadSafe> tsb = MakeShared<FMP4TrackSampleBuffer, ESPMode::ThreadSafe>();
+				tsb->TrackAndCodecInfo = Track;
+				tsb->TrackID = Track->TrackID;
+				TrackSampleBuffers.Emplace(Track->TrackID, MoveTemp(tsb));
+			}
+
+			if (InCodecTypeIndex == CodecTypeIndex(ElectraProtronUtils::FCodecInfo::EType::Video))
+			{
+				VideoLoaderThread.RequestLoad(TrackSampleBuffers[Track->TrackID], CurrentPlayPosTime);
+			}
+			else if (InCodecTypeIndex == CodecTypeIndex(ElectraProtronUtils::FCodecInfo::EType::Audio))
+			{
+				AudioLoaderThread.RequestLoad(TrackSampleBuffers[Track->TrackID], CurrentPlayPosTime);
+			}
+		}
+	}
+}
diff --git a/Engine/Plugins/Media/ElectraPlayer/Source/ElectraProtron/Private/Player/ElectraProtronPlayerImpl.h b/Engine/Plugins/Media/ElectraPlayer/Source/ElectraProtron/Private/Player/ElectraProtronPlayerImpl.h
new file mode 100644
--- /dev/null
+++ b/Engine/Plugins/Media/ElectraPlayer/Source/ElectraProtron/Private/Player/ElectraProtronPlayerImpl.h
@@ -0,0 +1,490 @@
+// Copyright Epic Games, Inc. All Rights Reserved.
+
+#pragma once
+
+#include "ElectraProtronPlayer.h"
+#include "HAL/Runnable.h"
+#include "HAL/RunnableThread.h"
+#include "Core/MediaEventSignal.h"
+#include "Containers/Queue.h"
+#include "Misc/TVariant.h"
+#include "Misc/Optional.h"
+#include "Math/Range.h"
+#include "Math/RangeSet.h"
+#include "Utilities/UtilitiesMP4.h"
+#include "Utilities/MP4Boxes/MP4Boxes.h"
+#include "Utilities/MP4Boxes/MP4Track.h"
+#include "TrackFormatInfo.h"
+#include "IElectraDecoder.h"
+#include "IElectraDecoderResourceDelegateBase.h"
+#include "IElectraPlayerDecoderResourceManager.h"
+#include "Utils/MPEG/ElectraUtilsMPEGVideo.h"
+#include "Utils/AudioChannelMapper.h"
+#include "Decoder/VideoDecoderHelpers.h"
+
+
+/**
+ * Private player implementation managed through a thread-safe pointer and a worker thread
+ * to not block the game thread.
+ */
+class FElectraProtronPlayer::FImpl : public TSharedFromThis<FElectraProtronPlayer::FImpl, ESPMode::ThreadSafe>, public FRunnable
+{
+public:
+	// Define pointer type to get around the ',' in the name that confuses the delegate declaration macro.
+	using ImplPointer = TSharedPtr<FImpl, ESPMode::ThreadSafe>;
+	DECLARE_DELEGATE_OneParam(FCompletionDelegate, ImplPointer);
+
+	FImpl();
+	~FImpl();
+	struct FOpenParam
+	{
+		FString Filename;
+		TSharedPtr<FMediaSamples, ESPMode::ThreadSafe> SampleQueue;
+		TSharedPtr<FElectraTextureSamplePool, ESPMode::ThreadSafe> TexturePool;
+		TSharedPtr<FElectraAudioSamplePool, ESPMode::ThreadSafe> AudioSamplePool;
+	};
+	void Open(const FOpenParam& InParam, FCompletionDelegate InCompletionDelegate);
+	void Close(FCompletionDelegate InCompletionDelegate);
+	FString GetLastError();
+	bool HasReachedEnd();
+
+	// Implementation methods from FElectraProtronPlayer
+	FTimespan GetDuration()
+	{ return Duration; }
+	FVariant GetMediaInfo(FName InInfoName);
+	bool GetAudioTrackFormat(int32 InTrackIndex, int32 InFormatIndex, FMediaAudioTrackFormat& OutFormat);
+	int32 GetNumTracks(EMediaTrackType InTrackType);
+	int32 GetNumTrackFormats(EMediaTrackType InTrackType, int32 InTrackIndex);
+	int32 GetSelectedTrack(EMediaTrackType InTrackType);
+	FText GetTrackDisplayName(EMediaTrackType InTrackType, int32 InTrackIndex);
+	int32 GetTrackFormat(EMediaTrackType InTrackType, int32 InTrackIndex);
+	FString GetTrackLanguage(EMediaTrackType InTrackType, int32 InTrackIndex);
+	FString GetTrackName(EMediaTrackType InTrackType, int32 InTrackIndex);
+	bool GetVideoTrackFormat(int32 InTrackIndex, int32 InFormatIndex, FMediaVideoTrackFormat& OutFormat);
+	bool SelectTrack(EMediaTrackType InTrackType, int32 InTrackIndex);
+	bool SetTrackFormat(EMediaTrackType InTrackType, int32 InTrackIndex, int32 InFormatIndex);
+	bool QueryCacheState(EMediaCacheState InState, TRangeSet<FTimespan>& OutTimeRanges);
+	int32 GetSampleCount(EMediaCacheState InState);
+
+	TRangeSet<float> GetSupportedRates(EMediaRateThinning InThinning);
+	float GetRate();
+	bool SetRate(float InRate);
+	FTimespan GetTime();
+	bool SetLooping(bool bInLooping);
+	bool IsLooping();
+	bool Seek(const FTimespan& InTime, const TSharedPtr<FMediaSamples, ESPMode::ThreadSafe>& InNewSampleQueue, uint32 InNewSequenceIndex);
+
+	void TickFetch(FTimespan InDeltaTime, FTimespan InTimecode);
+	void TickInput(FTimespan InDeltaTime, FTimespan InTimecode);
+
+
+	// Inherited from FRunnable
+	uint32 Run() override;
+	void Exit() override;
+
+private:
+	static constexpr int32 CodecTypeIndex(const ElectraProtronUtils::FCodecInfo::EType InType)
+	{ return (int32) InType; }
+
+	struct FConfig
+	{
+		FTimespan DurationCacheAhead { ETimespan::TicksPerSecond * 2 };
+		FTimespan DurationCacheBehind { ETimespan::TicksPerSecond * 1 };
+		int32 NextKeyframeThresholdMillis { 2 };
+		bool bReadFirstTimecode = true;
+		bool bReadSampleTimecode = true;
+	};
+
+	struct FSharedPlayParams
+	{
+		float DesiredPlayRate = 0.0f;
+		float PlaybackDirection = 0.0f;
+		bool bShouldLoop = false;
+	};
+
+	struct FSeekRequest
+	{
+		TSharedPtr<FMediaSamples, ESPMode::ThreadSafe> NewSampleQueue;
+		FTimespan NewTime;
+		int32 NewSequenceIndex = 0;
+	};
+
+	struct FWorkerThreadMessage
+	{
+		enum class EType
+		{
+			Nop,
+			Open,
+			Terminate
+		};
+		struct FParamOpen
+		{
+			FOpenParam Param;
+		};
+		struct FParamTerminate
+		{
+		};
+
+		EType Type = EType::Nop;
+		ImplPointer Self;
+		FCompletionDelegate CompletionDelegate;
+		TVariant<FParamOpen, FParamTerminate> Param;
+	};
+
+	struct FTrackInfo
+	{
+		TArray<TWeakPtr<FTrackInfo>> IsReferencedByTracks;
+		ElectraProtronUtils::FCodecInfo CodecInfo;
+		FString HumanReadableCodecFormat;
+		TSharedPtr<Electra::UtilitiesMP4::FMP4BoxTRAK> TrackBox;
+		TSharedPtr<Electra::UtilitiesMP4::FMP4Track> MP4Track;
+		TWeakPtr<FTrackInfo> ReferencedTimecodeTrack;
+		uint32 TrackID = 0;
+		bool bIsUsable = false;
+		bool bIsKeyframeOnlyFormat = false;
+		struct FFirstSampleTimecode
+		{
+			FString Timecode;
+			FString Framerate;
+			uint32 TimecodeValue = 0;
+		};
+		TOptional<FFirstSampleTimecode> FirstSampleTimecode;
+	};
+
+	struct FTrackSelection
+	{
+		int32 SelectedTrackIndex[4] {-1, -1, -1, -1};
+		int32 ActiveTrackIndex[4] {-1, -1, -1, -1};
+		bool bChanged = false;
+	};
+
+	struct FMP4Sample
+	{
+		TArray<uint8> Data;
+		FTimespan DTS;
+		FTimespan PTS;
+		FTimespan EffectiveDTS;
+		FTimespan EffectivePTS;
+		FTimespan Duration;
+		int64 SizeInBytes;
+		int64 OffsetInFile;
+		uint32 TrackID;
+		uint32 SampleNumber;
+		bool bIsSyncOrRap;
+		TOptional<FTimecode> AssociatedTimecode;
+		TOptional<FFrameRate> AssociatedTimecodeFramerate;
+	};
+	using FMP4SamplePtr = TSharedPtr<FMP4Sample, ESPMode::ThreadSafe>;
+
+	struct FMP4TrackSampleBuffer
+	{
+		FCriticalSection Lock;
+		TRangeSet<uint32> SampleRanges;
+		TSortedMap<uint32, FMP4SamplePtr> SampleMap;
+		TSharedPtr<FTrackInfo> TrackAndCodecInfo;
+		uint32 TrackID;
+	};
+	using FMP4TrackSampleBufferPtr = TSharedPtr<FMP4TrackSampleBuffer, ESPMode::ThreadSafe>;
+
+	using FTrackIterator = TSharedPtr<Electra::UtilitiesMP4::FMP4Track::FIterator>;
+
+	void StartThread();
+	void SendWorkerThreadMessage(FWorkerThreadMessage&& InMessage);
+
+	void InternalOpen(const FString& InFilename);
+	void GetTrackCodecInfo(ElectraProtronUtils::FCodecInfo& OutCodecInfo, const TSharedPtr<Electra::UtilitiesMP4::FMP4BoxTRAK>& InTrack, uint32 InTrackID);
+
+	void UpdateTrackLoader(int32 InCodecTypeIndex);
+
+	void HandleActiveTrackChanges();
+	void HandleRateChanges();
+	void HandleSeekRequest(const FSeekRequest& InSeekRequest);
+
+
+	FConfig Config;
+
+	FRunnableThread* Thread = nullptr;
+	ImplPointer SelfDuringTerminate;
+	FMediaEvent WorkMessageSignal;
+	TQueue<FWorkerThreadMessage, EQueueMode::Mpsc> WorkMessages;
+	FString LastErrorMessage;
+	bool bAbort = false;
+	TArray<Electra::UtilitiesMP4::FMP4BoxTreeParser> ParsedRootBoxes;
+
+	TArray<TSharedPtr<FTrackInfo>> Tracks;
+	TArray<TArray<int32>> UsableTrackArrayIndicesByType;
+
+	Electra::FTimeFraction MovieDuration;
+	FTimespan Duration;						// already converted `MovieDuration` for faster access
+	FTrackSelection TrackSelection;
+	TRangeSet<float> UnthinnedRates;
+	TRangeSet<float> ThinnedRates;
+	bool bAreRatesValid = false;
+	FTimespan CurrentPlayPosTime;
+	float CurrentRate = 0.0f;
+	float IntendedRate = 0.0f;
+
+	FCriticalSection SeekRequestLock;
+	TOptional<FSeekRequest> PendingSeekRequest;
+
+	TMap<uint32, FMP4TrackSampleBufferPtr> TrackSampleBuffers;
+
+	TSharedPtr<FSharedPlayParams, ESPMode::ThreadSafe> SharedPlayParams;
+
+	// ===== Loader and decoder threads =====
+
+	DECLARE_DELEGATE_RetVal_ThreeParams(FMP4SamplePtr, FGetSampleDlg, FMP4TrackSampleBufferPtr /*FromBuffer*/, const FTrackIterator& /*AtIterator*/, int32 /*WaitMicros*/);
+
+	class FLoaderThread : public FRunnable
+	{
+	public:
+		FLoaderThread(const FConfig& InConfig, int32 InCodecTypeIndex)
+			: Config(InConfig), LoaderTypeIndex(InCodecTypeIndex)
+		{ }
+
+		FString GetLastError()
+		{ return LastErrorMessage; }
+
+		// Inherited from FRunnable
+		uint32 Run() override;
+		struct FMessage
+		{
+			enum class EType
+			{
+				Nop,
+				Open,
+				Load,
+				Update,
+				Terminate
+			};
+			struct FParamOpen
+			{
+				FString Filename;
+				TSharedPtr<FSharedPlayParams, ESPMode::ThreadSafe> SharedPlayParams;
+			};
+			struct FParamNull
+			{
+			};
+
+			EType Type = EType::Nop;
+			TVariant<FParamOpen, FParamNull> Param;
+		};
+		void SendThreadMessage(FMessage&& InMessage)
+		{
+			Messages.Enqueue(MoveTemp(InMessage));
+			MessageSignal.Signal();
+		}
+		void StartThread(const FString& InFilename, const TSharedPtr<FSharedPlayParams, ESPMode::ThreadSafe>& InSharedPlayParams);
+		void StopThread();
+
+		// Called by the player
+		void RequestLoad(FMP4TrackSampleBufferPtr InTrackSampleBuffer, FTimespan InTime);
+		TRangeSet<FTimespan> GetTimeRangesToLoad();
+
+		// Called by the decoder
+		FMP4SamplePtr GetSample(FMP4TrackSampleBufferPtr InFromBuffer, const FTrackIterator& InAtIterator, int32 InWaitMicroseconds);
+	private:
+		const FConfig& Config;
+		int32 LoaderTypeIndex;
+
+		FRunnableThread* Thread = nullptr;
+		FMediaEvent MessageSignal;
+		TQueue<FMessage, EQueueMode::Mpsc> Messages;
+
+		TSharedPtr<FSharedPlayParams, ESPMode::ThreadSafe> SharedPlayParams;
+		TSharedPtr<Electra::IFileDataReader, ESPMode::ThreadSafe> Reader;
+		FString LastErrorMessage;
+
+		struct FLoadRequest
+		{
+			void Empty()
+			{
+				TrackSampleBuffer.Reset();
+				StartAtIterator.Reset();
+				UpdateAtIterator.Reset();
+			}
+			FMP4TrackSampleBufferPtr TrackSampleBuffer;
+			FTrackIterator StartAtIterator;
+			FTrackIterator UpdateAtIterator;
+		};
+
+		FCriticalSection LoadRequestLock;
+		FLoadRequest PendingLoadRequest;
+		FLoadRequest ActiveLoadRequest;
+		FCriticalSection TimeRangeLock;
+		TRangeSet<FTimespan> TimeRangesToLoad;
+		volatile int32 LoadRequestDirty = -1;
+
+		enum class ELoadResult
+		{
+			Ok,
+			Error,
+			Canceled
+		};
+
+		ELoadResult Load(FMP4TrackSampleBufferPtr InTrackSampleBuffer, FTrackIterator InAtIterator);
+		FMP4SamplePtr RetrieveSample(const FMP4TrackSampleBufferPtr& InTrackSampleBuffer, const FTrackIterator& InSampleIt, const FTrackIterator& InOptionalTimecodeIt, const ElectraProtronUtils::FCodecInfo::FTMCDTimecode& InTimecodeInfo);
+
+		struct FSampleRange
+		{
+			TRangeSet<FTimespan> TimeRanges;
+			TRangeSet<uint32> SampleRanges;
+			int32 NumSamplesAfter = 0;
+			int32 NumSamplesBefore = 0;
+			int32 NumRemainingToLoadAfter = 0;
+			int32 NumRemainingToLoadBefore = 0;
+		};
+		void CalcRangeToLoad(FSampleRange& OutRange, const FMP4TrackSampleBufferPtr& InTrackSampleBuffer, const FTrackIterator& InSampleIt);
+
+		void GetUnreferencedFrames(TArray<uint32>& OutFramesToRemove, const FMP4TrackSampleBufferPtr& InTrackSampleBuffer, const FSampleRange& InActiveSampleRange);
+	};
+
+	class FDecoderThread : public FRunnable
+	{
+	public:
+		FDecoderThread(const FConfig& InConfig, int32 InCodecTypeIndex)
+			: Config(InConfig), DecoderTypeIndex(InCodecTypeIndex)
+		{
+			VideoDecoderOutputPool = FVideoPool::Create();
+		}
+
+		// Inherited from FRunnable
+		uint32 Run() override;
+
+		void StartThread(const FOpenParam& InParam, const TSharedPtr<FSharedPlayParams, ESPMode::ThreadSafe>& InSharedPlayParams);
+		void StopThread();
+
+		FString GetLastError()
+		{ return LastErrorMessage; }
+
+		// Sets playback rate.
+		void SetRate(float InNewRate);
+		// Whether or not to loop.
+		bool SetLooping(bool bInLooping);
+
+		// Did decoding reach the last sample (or first sample when going in reverse)?
+		bool HasReachedEnd();
+
+		// Pauses decoding. Used when switching buffers or setting a new position.
+		void Pause();
+		// Resumes decoding. Decoder starts in paused state and needs to be resumed after setting buffer and time.
+		void Resume();
+		// Set the buffer to get the samples to decode from.
+		void SetSampleBuffer(const FMP4TrackSampleBufferPtr& InTrackSampleBuffer, FGetSampleDlg InGetSampleDelegate);
+		void DisconnectSampleBuffer();
+
+		bool IsPaused();
+		void PauseForSeek();
+		void ResumeAfterSeek();
+		bool IsPausedForSeek();
+		void SetTime(FTimespan InTime, int32 InSeqIdx, int32 InLoopIdx);
+		void SetEstimatedPlaybackTime(FTimespan InTime);
+
+		void Flush(const TSharedPtr<FMediaEvent, ESPMode::ThreadSafe>& InFlushedSignal);
+		void SetSampleQueue(const TSharedPtr<FMediaSamples, ESPMode::ThreadSafe>& InNewSampleQueue);
+	private:
+
+		void UpdateTrackIterator();
+		void StepTrackIterator(bool bPerformFlush);
+		bool CreateDecoder();
+		void DestroyDecoder();
+		void DecodeOneFrame();
+		void HandleOutputFrame();
+		void FlushForEndOrLooping();
+		void PerformFlush();
+
+		struct FInDecoder
+		{
+			TMap<FString, FVariant> CSDOptions;
+			IElectraDecoder::FInputAccessUnit DecAU;
+			TSharedPtr<IElectraDecoderBitstreamInfo, ESPMode::ThreadSafe> BSI;
+			FMP4SamplePtr Sample;
+			TArray<uint8> DataCopy;
+			int32 SequenceIndex = 0;
+			int32 LoopIndex = 0;
+		};
+
+		struct FPendingBufferChange
+		{
+			FCriticalSection Lock;
+			TSharedPtr<FMP4TrackSampleBuffer, ESPMode::ThreadSafe> NewTrackSampleBuffer;
+			FGetSampleDlg NewGetSampleDelegate;
+			bool bIsSet = false;
+		};
+
+		struct FPendingSeek
+		{
+			FCriticalSection Lock;
+			FTimespan NewTime;
+			int32 NewSeqIdx = 0;
+			int32 NewLoopIdx = 0;
+			bool bIsSet = false;
+		};
+
+		const FConfig& Config;
+		int32 DecoderTypeIndex;
+		FRunnableThread* Thread = nullptr;
+		FMediaEvent WorkSignal;
+		bool bTerminate = false;
+		FString LastErrorMessage;
+
+		FOpenParam Params;
+		TSharedPtr<FSharedPlayParams, ESPMode::ThreadSafe> SharedPlayParams;
+		TSharedPtr<FElectraProtronPlayer::FImpl::FMP4TrackSampleBuffer, ESPMode::ThreadSafe> TrackSampleBuffer;
+		FGetSampleDlg GetSampleDelegate;
+		FTrackIterator TrackIterator;
+		FTimespan FirstSampleEffectiveStartTime;
+		FTimespan LatestSampleEffectiveEndTime;
+
+		FCriticalSection TimeLock;
+		FTimespan CurrentTime;
+		float CurrentRate = 0.0f;
+		float IntendedRate = 0.0f;
+		float PlaybackDirection = 0.0f;
+		bool bShouldLoop = false;
+		bool bReachedEnd = false;
+		volatile bool bIsPaused = true;
+		volatile bool bPausedForSeek = false;
+
+		TSharedPtr<FMediaEvent, ESPMode::ThreadSafe> FlushedSignal;
+		volatile bool bFlushPending = false;
+		volatile bool bIsDrainingAtEOS = false;
+
+		FPendingBufferChange PendingBufferChange;
+
+		FPendingSeek PendingSeek;
+		TOptional<FTimespan> SeekTimeToHandleTo;
+		TOptional<FTimespan> SeekTimeToDecodeTo;
+		int32 SeekTimeNumFramesDecoded = 0;
+		int32 SeekTimeNumFramesSkipped = 0;
+
+		TSharedPtr<IElectraDecoder, ESPMode::ThreadSafe> DecoderInstance;
+		TSharedPtr<IElectraDecoderBitstreamProcessor, ESPMode::ThreadSafe> DecoderBitstreamProcessor;
+
+		TMap<FString, FVariant> CurrentCodecSpecificData;
+		TSharedPtr<Electra::IVideoDecoderResourceDelegate, ESPMode::ThreadSafe> VideoResourceDelegate;
+		IElectraDecoderResourceDelegateBase::IDecoderPlatformResource* PlatformResource = nullptr;
+		TUniquePtr<FInDecoder> CurrentInputSample;
+		TSharedPtr<IElectraDecoderOutput, ESPMode::ThreadSafe> CurrentDecoderOutput;
+		TUniquePtr<FInDecoder> InputForCurrentDecoderOutput;
+		TMap<uint64, TUniquePtr<FInDecoder>> InDecoderInput;
+		TOptional<Electra::MPEG::FColorimetryHelper> CurrentColorimetry;
+		TOptional<Electra::MPEG::FHDRHelper> CurrentHDR;
+		Electra::FAudioChannelMapper AudioChannelMapper;
+
+		uint64 NextUserValue = 0;
+		int32 SequenceIndex = 0;
+		int32 LoopIndex = 0;
+		bool bWaitForSyncSample = true;
+		bool bWarnedMissingSyncSample = false;
+
+		using FVideoPool = TDecoderOutputObjectPool<FVideoDecoderOutput, FElectraPlayerDecoderResourceManager::FVideo>;
+		TSharedPtr<FVideoPool, ESPMode::ThreadSafe> VideoDecoderOutputPool;
+	};
+
+	FLoaderThread VideoLoaderThread {Config, CodecTypeIndex(ElectraProtronUtils::FCodecInfo::EType::Video)};
+	FLoaderThread AudioLoaderThread {Config, CodecTypeIndex(ElectraProtronUtils::FCodecInfo::EType::Audio)};
+
+	FDecoderThread VideoDecoderThread {Config, CodecTypeIndex(ElectraProtronUtils::FCodecInfo::EType::Video)};
+	FDecoderThread AudioDecoderThread {Config, CodecTypeIndex(ElectraProtronUtils::FCodecInfo::EType::Audio)};
+};
diff --git a/Engine/Plugins/Media/ElectraPlayer/Source/ElectraProtron/Private/Player/ElectraProtronPlayerLoader.cpp b/Engine/Plugins/Media/ElectraPlayer/Source/ElectraProtron/Private/Player/ElectraProtronPlayerLoader.cpp
new file mode 100644
--- /dev/null
+++ b/Engine/Plugins/Media/ElectraPlayer/Source/ElectraProtron/Private/Player/ElectraProtronPlayerLoader.cpp
@@ -0,0 +1,523 @@
+// Copyright Epic Games, Inc. All Rights Reserved.
+
+#include "ElectraProtronPlayerImpl.h"
+#include "ElectraProtronPrivate.h"
+#include "IElectraCodecFactoryModule.h"
+#include "IElectraCodecFactory.h"
+#include "IElectraDecoderFeaturesAndOptions.h"
+#include "Utils/MPEG/ElectraUtilsMPEGVideo_H264.h"
+#include "Utils/MPEG/ElectraUtilsMPEGVideo_H265.h"
+#include "Utils/MPEG/ElectraUtilsMPEGAudio.h"
+#include "TrackFormatInfo.h"
+
+
+
+void FElectraProtronPlayer::FImpl::FLoaderThread::StartThread(const FString& InFilename, const TSharedPtr<FSharedPlayParams, ESPMode::ThreadSafe>& InSharedPlayParams)
+{
+	if (!Thread)
+	{
+		FMessage Msg;
+		FMessage::FParamOpen open { InFilename, InSharedPlayParams };
+		Msg.Param.Emplace<FMessage::FParamOpen>(MoveTemp(open));
+		Msg.Type = FMessage::EType::Open;
+		SendThreadMessage(MoveTemp(Msg));
+		Thread = FRunnableThread::Create(this, TEXT("Electra Protron Loader"), 0, TPri_Normal);
+	}
+}
+
+void FElectraProtronPlayer::FImpl::FLoaderThread::StopThread()
+{
+	if (Thread)
+	{
+		FMessage Msg;
+		Msg.Param.Emplace<FMessage::FParamNull>();
+		Msg.Type = FMessage::EType::Terminate;
+		SendThreadMessage(MoveTemp(Msg));
+		Thread->WaitForCompletion();
+	}
+}
+
+uint32 FElectraProtronPlayer::FImpl::FLoaderThread::Run()
+{
+// TODO: clamp look ahead/behind values to 0 if negative and an internal limit if too large.
+	bool bDone = false;
+	while(!bDone)
+	{
+		MessageSignal.WaitTimeoutAndReset(1000 * 20);
+		FMessage Msg;
+		while(Messages.Dequeue(Msg))
+		{
+			switch(Msg.Type)
+			{
+				case FMessage::EType::Open:
+				{
+					const FMessage::FParamOpen& Param(Msg.Param.Get<FMessage::FParamOpen>());
+					SharedPlayParams = Param.SharedPlayParams;
+					check(SharedPlayParams.IsValid());
+					Reader = Electra::IFileDataReader::Create();
+					if (!Reader->Open(Param.Filename))
+					{
+						// Failure is not really an option seeing as how we already opened the file
+						// successfully in Open().
+						check(!"how could this fail?");
+						LastErrorMessage = Reader->GetLastError();
+						Reader.Reset();
+					}
+					break;
+				}
+				case FMessage::EType::Load:
+				{
+					if (!LastErrorMessage.IsEmpty())
+					{
+						break;
+					}
+					else
+					{
+						FScopeLock lock(&LoadRequestLock);
+						ActiveLoadRequest = MoveTemp(PendingLoadRequest);
+					}
+					// FALLTHROUGH
+					[[fallthrough]];
+				}
+				case FMessage::EType::Update:
+				{
+					LoadRequestLock.Lock();
+					FMP4TrackSampleBufferPtr TrackSampleBuffer = ActiveLoadRequest.TrackSampleBuffer;
+					FTrackIterator StartAtIterator = MoveTemp(ActiveLoadRequest.StartAtIterator);
+					LoadRequestDirty = -1;
+					LoadRequestLock.Unlock();
+					check(StartAtIterator.IsValid());
+					ELoadResult Result = Load(TrackSampleBuffer, StartAtIterator);
+					if (Result == ELoadResult::Error)
+					{
+						if (LastErrorMessage.IsEmpty())
+						{
+							LastErrorMessage = FString::Printf(TEXT("Error loading media samples"));
+						}
+					}
+					else if (Result == ELoadResult::Canceled)
+					{
+						// There are two ways the request can get canceled.
+						// Either because of the player seeking or changing tracks, or the
+						// decoder asking for a new sample after having consumed one.
+						FScopeLock lock(&LoadRequestLock);
+						// A new request from the player wins.
+						if (PendingLoadRequest.TrackSampleBuffer)
+						{
+							// A new request also means there is a new load message in the message queue,
+							// so we do not need to do anything else but reset the active request here.
+							ActiveLoadRequest.Empty();
+						}
+					}
+					break;
+				}
+				case FMessage::EType::Terminate:
+				{
+					bDone = true;
+					break;
+				}
+				default:
+				{
+					unimplemented();
+					break;
+				}
+			}
+		}
+
+		if (1)
+		{
+			FScopeLock lock(&LoadRequestLock);
+			if (ActiveLoadRequest.UpdateAtIterator.IsValid())
+			{
+				check(!ActiveLoadRequest.StartAtIterator.IsValid());
+				ActiveLoadRequest.StartAtIterator = MoveTemp(ActiveLoadRequest.UpdateAtIterator);
+				FMessage UpdateMsg;
+				UpdateMsg.Param.Emplace<FMessage::FParamNull>();
+				UpdateMsg.Type = FMessage::EType::Update;
+				SendThreadMessage(MoveTemp(UpdateMsg));
+			}
+		}
+
+	}
+	PendingLoadRequest.Empty();
+	ActiveLoadRequest.Empty();
+	return 0;
+}
+
+
+void FElectraProtronPlayer::FImpl::FLoaderThread::CalcRangeToLoad(FSampleRange& OutRange, const FMP4TrackSampleBufferPtr& InTrackSampleBuffer, const FTrackIterator& InSampleIt)
+{
+	OutRange.NumSamplesAfter = 0;
+	OutRange.NumSamplesBefore = 0;
+	OutRange.NumRemainingToLoadAfter = 0;
+	OutRange.NumRemainingToLoadBefore = 0;
+
+	// Given a track iterator, figure out the samples we need to load prior and following the iterator's current position.
+	check(InSampleIt->IsValid());
+	if (!InSampleIt->IsValid())
+	{
+		return;
+	}
+
+	const uint32 TrackTimescale = InSampleIt->GetTimescale();
+	check(TrackTimescale);
+
+	int64 FwdDurNeeded = Electra::FTimeFraction(Config.DurationCacheAhead).GetAsTimebase(TrackTimescale);
+	int64 RevDurNeeded = Electra::FTimeFraction(Config.DurationCacheBehind).GetAsTimebase(TrackTimescale);
+	// See if the combined cache ahead and behind duration encompasses the entire track.
+	// This happens if the track is short (or the cache durations excessively large)
+	if (FwdDurNeeded + RevDurNeeded >= InSampleIt->GetEffectiveTrackDuration().GetNumerator())
+	{
+		const uint32 NumTrackSamples = InSampleIt->GetLastSampleInfo().LastSampleNumber + 1 - InSampleIt->GetFirstSampleInfo().SyncSampleNumber;
+		OutRange.NumSamplesAfter = SharedPlayParams->PlaybackDirection >= 0.0f ? NumTrackSamples * 3 / 4 : NumTrackSamples / 4;
+		OutRange.NumSamplesBefore = NumTrackSamples - OutRange.NumSamplesAfter;
+		OutRange.TimeRanges.Add(TRange<FTimespan>(FTimespan(), InSampleIt->GetEffectiveTrackDuration().GetAsTimespan()));
+		OutRange.SampleRanges.Add(TRange<uint32>(InSampleIt->GetFirstSampleInfo().SyncSampleNumber, InSampleIt->GetLastSampleInfo().LastSampleNumber + 1));
+		return;
+	}
+
+	if (SharedPlayParams->PlaybackDirection < 0.0f)
+	{
+		Swap(FwdDurNeeded, RevDurNeeded);
+	}
+
+	uint32 StartSampleNum = InSampleIt->GetSampleNumber();
+	int64 DurHandled = 0;
+	FTrackIterator FwdTkIt(InSampleIt->Clone());
+	int64 StartDTS = FwdTkIt->GetEffectiveDTS().GetNumerator();
+	while(DurHandled < FwdDurNeeded)
+	{
+		const int64 dur = FwdTkIt->GetDuration().GetNumerator();
+		DurHandled += dur;
+		++OutRange.NumSamplesAfter;
+		if (!FwdTkIt->NextEffective())
+		{
+			OutRange.TimeRanges.Add(TRange<FTimespan>(Electra::FTimeFraction(StartDTS, TrackTimescale).GetAsTimespan(), Electra::FTimeFraction(FwdTkIt->GetEffectiveDTS().GetNumerator() + dur, TrackTimescale).GetAsTimespan()));
+			OutRange.SampleRanges.Add(TRange<uint32>(StartSampleNum, FwdTkIt->GetSampleNumber()+1));
+			FwdTkIt = InTrackSampleBuffer->TrackAndCodecInfo->MP4Track->CreateIterator(FwdTkIt->GetFirstSampleInfo().SyncSampleNumber);
+			StartDTS = FwdTkIt->GetEffectiveDTS().GetNumerator();
+			StartSampleNum = FwdTkIt->GetSampleNumber();
+		}
+	}
+	OutRange.TimeRanges.Add(TRange<FTimespan>(Electra::FTimeFraction(StartDTS, TrackTimescale).GetAsTimespan(), Electra::FTimeFraction(FwdTkIt->GetEffectiveDTS().GetNumerator(), TrackTimescale).GetAsTimespan()));
+	OutRange.SampleRanges.Add(TRange<uint32>(StartSampleNum, FwdTkIt->GetSampleNumber()));
+
+	// Reverse scanning
+	const bool bNeedKeyframe = !InTrackSampleBuffer->TrackAndCodecInfo->bIsKeyframeOnlyFormat;
+	DurHandled = 0;
+	FTrackIterator RevTkIt(InSampleIt->Clone());
+	int64 EndDTS = RevTkIt->GetEffectiveDTS().GetNumerator();
+	uint32 EndSampleNum = RevTkIt->GetSampleNumber();
+	bool bHaveEnough = DurHandled >= RevDurNeeded;
+	while(!bHaveEnough)
+	{
+		if (!RevTkIt->PrevEffective())
+		{
+			TRange<uint32> SmpRng(RevTkIt->GetSampleNumber(), EndSampleNum);
+			if (!SmpRng.IsEmpty())
+			{
+				int64 DTS = RevTkIt->GetEffectiveDTS().GetNumerator();
+				FTimespan ts = Electra::FTimeFraction(DTS, TrackTimescale).GetAsTimespan();
+				FTimespan te = Electra::FTimeFraction(EndDTS, TrackTimescale).GetAsTimespan();
+				OutRange.TimeRanges.Add(TRange<FTimespan>(ts, te));
+				OutRange.SampleRanges.Add(SmpRng);
+			}
+			RevTkIt = InTrackSampleBuffer->TrackAndCodecInfo->MP4Track->CreateIterator(RevTkIt->GetLastSampleInfo().LastSampleNumber);
+			EndDTS = RevTkIt->GetEffectiveDTS().GetNumerator() + RevTkIt->GetDuration().GetNumerator();
+			EndSampleNum = RevTkIt->GetSampleNumber() + 1;
+		}
+		const int64 dur = RevTkIt->GetDuration().GetNumerator();
+		DurHandled += dur;
+		++OutRange.NumSamplesBefore;
+		bHaveEnough = DurHandled >= RevDurNeeded;
+		if (bHaveEnough && bNeedKeyframe)
+		{
+			bHaveEnough = RevTkIt->IsSyncOrRAPSample();
+		}
+	}
+	TRange<uint32> SmpRng(RevTkIt->GetSampleNumber(), EndSampleNum);
+	if (!SmpRng.IsEmpty())
+	{
+		int64 DTS = RevTkIt->GetEffectiveDTS().GetNumerator();
+		FTimespan ts = Electra::FTimeFraction(DTS, TrackTimescale).GetAsTimespan();
+		FTimespan te = Electra::FTimeFraction(EndDTS, TrackTimescale).GetAsTimespan();
+		OutRange.TimeRanges.Add(TRange<FTimespan>(ts, te));
+		OutRange.SampleRanges.Add(TRange<uint32>(RevTkIt->GetSampleNumber(), EndSampleNum));
+	}
+}
+
+
+void FElectraProtronPlayer::FImpl::FLoaderThread::GetUnreferencedFrames(TArray<uint32>& OutFramesToRemove, const FMP4TrackSampleBufferPtr& InTrackSampleBuffer, const FSampleRange& InActiveSampleRange)
+{
+	TArray<uint32> FramesInMap;
+	InTrackSampleBuffer->Lock.Lock();
+	FramesInMap.Reserve(InTrackSampleBuffer->SampleMap.Num());
+	InTrackSampleBuffer->SampleMap.GenerateKeyArray(FramesInMap);
+	InTrackSampleBuffer->Lock.Unlock();
+	for(int32 i=0, iMax=FramesInMap.Num(); i<iMax; ++i)
+	{
+		if (!InActiveSampleRange.SampleRanges.Contains(FramesInMap[i]))
+		{
+			OutFramesToRemove.Add(FramesInMap[i]);
+		}
+	}
+}
+
+
+FElectraProtronPlayer::FImpl::FMP4SamplePtr FElectraProtronPlayer::FImpl::FLoaderThread::GetSample(FMP4TrackSampleBufferPtr InFromBuffer, const FTrackIterator& InAtIterator, int32 InWaitMicroseconds)
+{
+	InFromBuffer->Lock.Lock();
+	uint32 SampleNum = InAtIterator->GetSampleNumber();
+	FMP4SamplePtr Sample = InFromBuffer->SampleMap.FindRef(SampleNum);
+	InFromBuffer->Lock.Unlock();
+
+	// If the buffer is still the one that is active we need to trigger a fetch of new samples
+	// regardless of whether we got (and thus consumed) the sample asked for.
+	// We should have the sample ready. If not then triggering the load request is especially important.
+	LoadRequestLock.Lock();
+	if (ActiveLoadRequest.TrackSampleBuffer == InFromBuffer)
+	{
+		ActiveLoadRequest.UpdateAtIterator = InAtIterator;
+		LoadRequestDirty = (int32) SampleNum;
+	}
+	LoadRequestLock.Unlock();
+
+	return Sample;
+}
+
+TRangeSet<FTimespan> FElectraProtronPlayer::FImpl::FLoaderThread::GetTimeRangesToLoad()
+{
+	FScopeLock lock(&TimeRangeLock);
+	return TimeRangesToLoad;
+}
+
+void FElectraProtronPlayer::FImpl::FLoaderThread::RequestLoad(FMP4TrackSampleBufferPtr InTrackSampleBuffer, FTimespan InTime)
+{
+	if (!Reader.IsValid() || !InTrackSampleBuffer.IsValid())
+	{
+		return;
+	}
+
+	FTrackIterator TkIt =
+		InTrackSampleBuffer->TrackAndCodecInfo->MP4Track->CreateIteratorAtKeyframe(Electra::FTimeValue().SetFromTimespan(InTime),
+			Electra::FTimeValue().SetFromMilliseconds(InTrackSampleBuffer->TrackAndCodecInfo->bIsKeyframeOnlyFormat ? 0 : Config.NextKeyframeThresholdMillis));
+	if (!TkIt.IsValid())
+	{
+		LastErrorMessage = InTrackSampleBuffer->TrackAndCodecInfo->MP4Track->GetLastError();
+		UE_LOG(LogElectraProtron, Error, TEXT("%s"), *LastErrorMessage);
+		return;
+	}
+
+	LoadRequestLock.Lock();
+	PendingLoadRequest.Empty();
+	PendingLoadRequest.TrackSampleBuffer = MoveTemp(InTrackSampleBuffer);
+	PendingLoadRequest.StartAtIterator = MoveTemp(TkIt);
+	LoadRequestDirty = -2;
+	LoadRequestLock.Unlock();
+
+	FMessage Msg;
+	Msg.Param.Emplace<FMessage::FParamNull>();
+	Msg.Type = FMessage::EType::Load;
+	SendThreadMessage(MoveTemp(Msg));
+}
+
+FElectraProtronPlayer::FImpl::FLoaderThread::ELoadResult FElectraProtronPlayer::FImpl::FLoaderThread::Load(FMP4TrackSampleBufferPtr InTrackSampleBuffer, FTrackIterator InAtIterator)
+{
+	// Determine the range of samples to load.
+	check(InAtIterator.IsValid());
+	check(InTrackSampleBuffer.IsValid());
+	FSampleRange RangeToLoad;
+	CalcRangeToLoad(RangeToLoad, InTrackSampleBuffer, InAtIterator);
+	// Need to load something...
+	check(RangeToLoad.NumSamplesAfter || RangeToLoad.NumSamplesBefore);
+	if (!(RangeToLoad.NumSamplesAfter || RangeToLoad.NumSamplesBefore))
+	{
+		return ELoadResult::Ok;
+	}
+
+	TimeRangeLock.Lock();
+	TimeRangesToLoad = RangeToLoad.TimeRanges;
+	TimeRangeLock.Unlock();
+
+	// Get the frames that are not referenced now we have to evict from the data map.
+	TArray<uint32> FramesToRemove;
+	GetUnreferencedFrames(FramesToRemove, InTrackSampleBuffer, RangeToLoad);
+
+	// Are timecodes from another track referenced?
+	FTrackIterator FwdTkItTC, RevTkItTC;
+	ElectraProtronUtils::FCodecInfo::FTMCDTimecode TimecodeInfo;
+	if (Config.bReadSampleTimecode && LoaderTypeIndex == CodecTypeIndex(ElectraProtronUtils::FCodecInfo::EType::Video))
+	{
+		if (auto TimecodeTrack = InTrackSampleBuffer->TrackAndCodecInfo->ReferencedTimecodeTrack.Pin())
+		{
+			// The timecode track needs to have as many samples as this track, otherwise there would be
+			// a mismatch somewhere and the timecode couldn't be used.
+			if (TimecodeTrack->MP4Track->GetNumberOfSamples() == InTrackSampleBuffer->TrackAndCodecInfo->MP4Track->GetNumberOfSamples())
+			{
+				// Get the timecode description from the codec info.
+				TimecodeInfo = TimecodeTrack->CodecInfo.Properties.Get<ElectraProtronUtils::FCodecInfo::FTMCDTimecode>();
+				FwdTkItTC = TimecodeTrack->MP4Track->CreateIterator(InAtIterator->GetSampleNumber());
+				if (FwdTkItTC.IsValid())
+				{
+					RevTkItTC = FwdTkItTC->Clone();
+				}
+			}
+		}
+	}
+
+	// Calculate the ratio of samples to fetch ahead vs. fetch behind.
+	// We want to fetch more samples in the direction we're going than from
+	// where we came.
+	// For this we do not need to look at the playback direction, we just take
+	// the number of samples to load in either direction since that is determined
+	// by the current play direction. Whichever is the one with more samples is the
+	// direction we are going in.
+	double fwd = RangeToLoad.NumSamplesAfter;
+	double rev = RangeToLoad.NumSamplesBefore;
+	int32 RatioF = (fwd > rev && rev > 0.0) ? FMath::CeilToInt(fwd / rev) : 1;
+	int32 RatioR = (rev > fwd && fwd > 0.0) ? FMath::CeilToInt(rev / fwd) : 1;
+
+	FTrackIterator FwdTkIt(InAtIterator->Clone());
+	FTrackIterator RevTkIt(InAtIterator->Clone());
+
+	RangeToLoad.NumRemainingToLoadAfter = RangeToLoad.NumSamplesAfter;
+	RangeToLoad.NumRemainingToLoadBefore = RangeToLoad.NumSamplesBefore;
+	const int32 kMinFramesToLoad = 2;
+	while(RangeToLoad.NumRemainingToLoadAfter || RangeToLoad.NumRemainingToLoadBefore)
+	{
+		// Remove one old frame now.
+		if (FramesToRemove.Num())
+		{
+			InTrackSampleBuffer->Lock.Lock();
+			InTrackSampleBuffer->SampleMap.Remove(FramesToRemove.Last());
+			InTrackSampleBuffer->Lock.Unlock();
+			FramesToRemove.Pop(EAllowShrinking::No);
+		}
+
+		for(int32 i=RatioF; i>0 && RangeToLoad.NumRemainingToLoadAfter; --i)
+		{
+			FMP4SamplePtr Sample = RetrieveSample(InTrackSampleBuffer, FwdTkIt, FwdTkItTC, TimecodeInfo);
+			// Abort loading immediately or when we have loaded at least the minimum number of required frames?
+			if (LoadRequestDirty < -1 || (LoadRequestDirty >= 0 && (-RangeToLoad.NumRemainingToLoadAfter + RangeToLoad.NumSamplesAfter > kMinFramesToLoad)))
+			{
+				//UE_LOG(LogTemp, Log, TEXT("Load canceled with %d/%d fwd and %d/%d rev remaining"), RangeToLoad.NumRemainingToLoadAfter, RangeToLoad.NumSamplesAfter, RangeToLoad.NumRemainingToLoadBefore, RangeToLoad.NumSamplesBefore);
+				return ELoadResult::Canceled;
+			}
+			if (!Sample.IsValid())
+			{
+				return ELoadResult::Error;
+			}
+
+			if (FwdTkItTC.IsValid())
+			{
+				FwdTkItTC->Next();
+			}
+			if (!FwdTkIt->NextEffective())
+			{
+				FwdTkIt = InTrackSampleBuffer->TrackAndCodecInfo->MP4Track->CreateIterator(FwdTkIt->GetFirstSampleInfo().SyncSampleNumber);
+				FwdTkItTC.Reset();
+				if (auto TimecodeTrack = InTrackSampleBuffer->TrackAndCodecInfo->ReferencedTimecodeTrack.Pin())
+				{
+					FwdTkItTC = TimecodeTrack->MP4Track->CreateIterator(FwdTkIt->GetSampleNumber());
+				}
+			}
+			--RangeToLoad.NumRemainingToLoadAfter;
+		}
+
+		for(int32 i=RatioR; i>0 && RangeToLoad.NumRemainingToLoadBefore; --i)
+		{
+			if (RevTkItTC.IsValid())
+			{
+				RevTkItTC->Prev();
+			}
+			if (!RevTkIt->PrevEffective())
+			{
+				RevTkIt = InTrackSampleBuffer->TrackAndCodecInfo->MP4Track->CreateIterator(RevTkIt->GetLastSampleInfo().LastSampleNumber);
+				RevTkItTC.Reset();
+				if (auto TimecodeTrack = InTrackSampleBuffer->TrackAndCodecInfo->ReferencedTimecodeTrack.Pin())
+				{
+					RevTkItTC = TimecodeTrack->MP4Track->CreateIterator(FwdTkIt->GetSampleNumber());
+				}
+			}
+
+			FMP4SamplePtr Sample = RetrieveSample(InTrackSampleBuffer, RevTkIt, RevTkItTC, TimecodeInfo);
+			// Abort loading immediately or when we have loaded at least the minimum number of required frames?
+			if (LoadRequestDirty < -1 || (LoadRequestDirty >= 0 && (-RangeToLoad.NumRemainingToLoadBefore + RangeToLoad.NumSamplesBefore > kMinFramesToLoad)))
+			{
+				//UE_LOG(LogTemp, Log, TEXT("Load canceled with %d/%d fwd and %d/%d rev remaining"), RangeToLoad.NumRemainingToLoadAfter, RangeToLoad.NumSamplesAfter, RangeToLoad.NumRemainingToLoadBefore, RangeToLoad.NumSamplesBefore);
+				return ELoadResult::Canceled;
+			}
+			if (!Sample.IsValid())
+			{
+				return ELoadResult::Error;
+			}
+			--RangeToLoad.NumRemainingToLoadBefore;
+		}
+	}
+	// Remove any remaining old frames now.
+	InTrackSampleBuffer->Lock.Lock();
+	for(int32 i=0,iMax=FramesToRemove.Num(); i<iMax; ++i)
+	{
+		InTrackSampleBuffer->SampleMap.Remove(FramesToRemove[i]);
+	}
+	InTrackSampleBuffer->Lock.Unlock();
+
+	return ELoadResult::Ok;
+}
+
+TSharedPtr<FElectraProtronPlayer::FImpl::FMP4Sample, ESPMode::ThreadSafe> FElectraProtronPlayer::FImpl::FLoaderThread::RetrieveSample(const FMP4TrackSampleBufferPtr& InTrackSampleBuffer, const FTrackIterator& InSampleIt, const FTrackIterator& InOptionalTimecodeIt, const ElectraProtronUtils::FCodecInfo::FTMCDTimecode& InTimecodeInfo)
+{
+	// Note: As long as we are not adding or removing samples in another thread, a lock should not be necessary here.
+	FMP4SamplePtr* CurrentSample = InTrackSampleBuffer->SampleMap.Find(InSampleIt->GetSampleNumber());
+	if (CurrentSample)
+	{
+		return *CurrentSample;
+	}
+
+	auto CheckAbort = Electra::IBaseDataReader::FCancellationCheckDelegate::CreateLambda([&]()
+	{
+		// We do NOT abort loading of a frame!
+		return false;
+	});
+
+	FMP4SamplePtr Sample = MakeShared<FMP4Sample, ESPMode::ThreadSafe>();
+	Sample->DTS = InSampleIt->GetDTS().GetAsTimespan();
+	Sample->PTS = InSampleIt->GetPTS().GetAsTimespan();
+	Sample->EffectiveDTS = InSampleIt->GetEffectiveDTS().GetAsTimespan();
+	Sample->EffectivePTS = InSampleIt->GetEffectivePTS().GetAsTimespan();
+	Sample->Duration = InSampleIt->GetDuration().GetAsTimespan();
+	Sample->SizeInBytes = InSampleIt->GetSampleSize();
+	Sample->OffsetInFile = InSampleIt->GetSampleFileOffset();
+	Sample->TrackID = InSampleIt->GetTrackID();
+	Sample->SampleNumber = InSampleIt->GetSampleNumber();
+	Sample->bIsSyncOrRap = InSampleIt->IsSyncOrRAPSample();
+	Sample->Data.SetNumUninitialized(Sample->SizeInBytes);
+	int64 NumRead = Reader->ReadData(Sample->Data.GetData(), Sample->SizeInBytes, Sample->OffsetInFile, CheckAbort);
+	if (NumRead != Sample->SizeInBytes)
+	{
+		return nullptr;
+	}
+
+	// Optionally read the timecode sample from the associated track
+	if (InOptionalTimecodeIt.IsValid())
+	{
+		int64 tcSampleSize = InOptionalTimecodeIt->GetSampleSize();
+		if (tcSampleSize == 4)
+		{
+			TArray<uint32> TimecodeBuffer;
+			TimecodeBuffer.SetNumUninitialized(Align(tcSampleSize, 4));
+			NumRead = Reader->ReadData(TimecodeBuffer.GetData(), tcSampleSize, InOptionalTimecodeIt->GetSampleFileOffset(), CheckAbort);
+			if (NumRead == tcSampleSize)
+			{
+				Sample->AssociatedTimecode = InTimecodeInfo.ConvertToTimecode(Electra::GetFromBigEndian(TimecodeBuffer[0]));
+				Sample->AssociatedTimecodeFramerate = InTimecodeInfo.GetFrameRate();
+			}
+		}
+	}
+
+	uint32 SampleNumber = Sample->SampleNumber;
+	FScopeLock lock(&InTrackSampleBuffer->Lock);
+	InTrackSampleBuffer->SampleRanges.Add(TRange<uint32>(SampleNumber, SampleNumber+1));
+	InTrackSampleBuffer->SampleMap.Emplace(SampleNumber, Sample);
+	return Sample;
+}
diff --git a/Engine/Plugins/Media/ElectraPlayer/Source/ElectraProtron/Private/Player/TrackFormatInfo.h b/Engine/Plugins/Media/ElectraPlayer/Source/ElectraProtron/Private/Player/TrackFormatInfo.h
new file mode 100644
--- /dev/null
+++ b/Engine/Plugins/Media/ElectraPlayer/Source/ElectraProtron/Private/Player/TrackFormatInfo.h
@@ -0,0 +1,89 @@
+// Copyright Epic Games, Inc. All Rights Reserved.
+
+#pragma once
+
+#include "CoreTypes.h"
+#include "Misc/TVariant.h"
+#include "Misc/FrameRate.h"
+#include "Misc/Timecode.h"
+
+#include "PlayerTime.h"
+
+namespace ElectraProtronUtils
+{
+
+	struct FCodecInfo
+	{
+		enum class EType
+		{
+			Video,
+			Audio,
+			Subtitle,
+			Timecode,
+			Invalid,
+			MAX = Invalid
+		};
+		struct FVideo
+		{
+			uint32 Width = 0;
+			uint32 Height = 0;
+			Electra::FTimeFraction FrameRate;
+		};
+		struct FAudio
+		{
+			uint32 NumChannels = 0;
+			uint32 ChannelConfiguration = 0;
+			uint32 SampleRate = 0;
+		};
+		struct FSubtitle
+		{
+		};
+		struct FTMCDTimecode
+		{
+			// See: https://developer.apple.com/documentation/quicktime-file-format/timecode_sample_description/flags
+			enum EFlags : uint32
+			{
+				DropFrame = 0x0001,					// Indicates whether the timecode is drop frame. Set it to 1 if the timecode is drop frame.
+				Max24Hour = 0x0002,					// Indicates whether the timecode wraps after 24 hours. Set it to 1 if the timecode wraps.
+				AllowNegativeTimes = 0x0004,		// Indicates whether negative time values are allowed. Set it to 1 if the timecode supports negative values.
+				Counter = 0x0008					// Indicates whether the time value corresponds to a tape counter value. Set it to 1 if the timecode values are tape counter values.
+			};
+
+			uint32 Flags = 0;
+			uint32 Timescale = 0;
+			uint32 FrameDuration = 0;
+			uint32 NumberOfFrames = 0;
+
+			bool IsDropFrame() const
+			{ return !!(Flags & EFlags::DropFrame); }
+
+			bool WrapsAfter24Hours() const
+			{ return !!(Flags & EFlags::Max24Hour); }
+
+			bool SupportsNegativeTime() const
+			{ return !!(Flags & EFlags::AllowNegativeTimes); }
+
+			FFrameRate GetFrameRate() const
+			{ return IsDropFrame() ? FFrameRate(Timescale, FrameDuration) : FFrameRate(NumberOfFrames, 1); }
+
+			FTimecode ConvertToTimecode(uint32 InSampleTimecode) const
+			{
+				// Needs to be an int32 for use with the following methods.
+				check(InSampleTimecode <= 0x7fffffffU);
+				const FFrameRate FrameRate = GetFrameRate();
+				// Convert to time code (apply roll over, etc). via conversion to seconds first.
+				return FTimecode(FrameRate.AsSeconds(FFrameTime(FFrameNumber(static_cast<int32>(InSampleTimecode)))), FrameRate, IsDropFrame(), WrapsAfter24Hours());
+			}
+
+		};
+		EType Type = EType::Invalid;
+		FString HumanReadableFormatInfo;
+		FString RFC6381;
+		uint32 FourCC = 0;
+		TVariant<FVideo,FAudio,FSubtitle,FTMCDTimecode> Properties;
+		TArray<uint8> DCR;
+		TArray<uint8> CSD;
+		TMap<uint32, TArray<uint8>> ExtraBoxes;
+	};
+
+} // ElectraProtronUtils
diff --git a/Engine/Plugins/Media/ElectraPlayer/Source/ElectraProtron/Public/IElectraProtronModule.h b/Engine/Plugins/Media/ElectraPlayer/Source/ElectraProtron/Public/IElectraProtronModule.h
new file mode 100644
--- /dev/null
+++ b/Engine/Plugins/Media/ElectraPlayer/Source/ElectraProtron/Public/IElectraProtronModule.h
@@ -0,0 +1,31 @@
+// Copyright Epic Games, Inc. All Rights Reserved.
+
+#pragma once
+
+#include "Templates/SharedPointer.h"
+#include "Modules/ModuleInterface.h"
+
+class IMediaEventSink;
+class IMediaPlayer;
+
+/**
+ * Interface for the AndroidMedia module.
+ */
+class IElectraProtronModule
+	: public IModuleInterface
+{
+public:
+
+	/**
+	 * Create a media player.
+	 *
+	 * @param EventSink The object that receives media events from the player.
+	 * @return A new media player, or nullptr if a player couldn't be created.
+	 */
+	virtual TSharedPtr<IMediaPlayer, ESPMode::ThreadSafe> CreatePlayer(IMediaEventSink& EventSink) = 0;
+
+public:
+
+	/** Virtual destructor. */
+	virtual ~IElectraProtronModule() { }
+};
diff --git a/Engine/Plugins/Media/ElectraPlayer/Source/ElectraProtronFactory/ElectraProtronFactory.Build.cs b/Engine/Plugins/Media/ElectraPlayer/Source/ElectraProtronFactory/ElectraProtronFactory.Build.cs
new file mode 100644
--- /dev/null
+++ b/Engine/Plugins/Media/ElectraPlayer/Source/ElectraProtronFactory/ElectraProtronFactory.Build.cs
@@ -0,0 +1,48 @@
+// Copyright Epic Games, Inc. All Rights Reserved.
+
+namespace UnrealBuildTool.Rules
+{
+    public class ElectraProtronFactory : ModuleRules
+    {
+        public ElectraProtronFactory(ReadOnlyTargetRules Target) : base(Target)
+        {
+		    bLegalToDistributeObjectCode = true;
+
+			DynamicallyLoadedModuleNames.AddRange(
+                new string[] {
+                    "Media",
+                    "ElectraProtron"
+                });
+
+            PrivateDependencyModuleNames.AddRange(
+                new string[] {
+                    "MediaAssets",
+                });
+
+			PrivateIncludePathModuleNames.AddRange(
+				new string[] {
+					"Media",
+					"ElectraProtron"
+				});
+
+            PrivateDependencyModuleNames.Add("ElectraBase");
+
+            PublicDependencyModuleNames.AddRange(
+                new string[] {
+                    "Core",
+                    "CoreUObject",
+                });
+
+            if (Target.Type == TargetType.Editor)
+            {
+                DynamicallyLoadedModuleNames.Add("Settings");
+                PrivateIncludePathModuleNames.Add("Settings");
+            }
+
+            if (Target.Platform.IsInGroup(UnrealPlatformGroup.Windows))
+            {
+                PublicDefinitions.Add("_CRT_SECURE_NO_WARNINGS");
+            }
+        }
+    }
+}
diff --git a/Engine/Plugins/Media/ElectraPlayer/Source/ElectraProtronFactory/Private/ElectraProtronFactoryModule.cpp b/Engine/Plugins/Media/ElectraPlayer/Source/ElectraProtronFactory/Private/ElectraProtronFactoryModule.cpp
new file mode 100644
--- /dev/null
+++ b/Engine/Plugins/Media/ElectraPlayer/Source/ElectraProtronFactory/Private/ElectraProtronFactoryModule.cpp
@@ -0,0 +1,186 @@
+// Copyright Epic Games, Inc. All Rights Reserved.
+
+#include "ElectraProtronFactoryPrivate.h"
+
+#include "Containers/Array.h"
+#include "Containers/UnrealString.h"
+#include "IMediaModule.h"
+#include "IMediaOptions.h"
+#include "IMediaPlayerFactory.h"
+#include "Modules/ModuleManager.h"
+#include "Misc/Guid.h"
+#include "Misc/Paths.h"
+#include "Internationalization/Internationalization.h"
+#include "UObject/NameTypes.h"
+
+#include "Utilities/URLParser.h"
+
+#include "IElectraProtronModule.h"
+
+DEFINE_LOG_CATEGORY(LogElectraProtronFactory);
+
+#define LOCTEXT_NAMESPACE "ElectraProtronFactoryModule"
+
+class ElectraProtronFactoryModule
+	: public IMediaPlayerFactory
+	, public IModuleInterface
+{
+public:
+	ElectraProtronFactoryModule() { }
+
+public:
+	// IMediaPlayerFactory interface
+
+	bool CanPlayUrl(const FString& Url, const IMediaOptions* Options, TArray<FText>* OutWarnings, TArray<FText>* OutErrors) const override
+	{
+		return GetPlayabilityConfidenceScore(Url, Options, OutWarnings, OutErrors) > 0 ? true : false;
+	}
+
+	int32 GetPlayabilityConfidenceScore(const FString& Url, const IMediaOptions* Options, TArray<FText>* OutWarnings, TArray<FText>* OutErrors) const override
+	{
+		// Split the URL apart.
+		Electra::FURL_RFC3986 UrlParser;
+		if (!UrlParser.Parse(Url))
+		{
+			if (OutErrors != nullptr)
+			{
+				OutErrors->Add(FText::Format(LOCTEXT("MalformedURI", "The URI '{0}' could not be parsed"), FText::FromString(Url)));
+			}
+		}
+		// Check scheme
+		FString Scheme = UrlParser.GetScheme();
+		if (Scheme.IsEmpty())
+		{
+			if (OutErrors != nullptr)
+			{
+				OutErrors->Add(LOCTEXT("NoSchemeFound", "No URI scheme found"));
+			}
+			return 0;
+		}
+		if (!SupportedUriSchemes.Contains(Scheme))
+		{
+			if (OutErrors != nullptr)
+			{
+				OutErrors->Add(FText::Format(LOCTEXT("SchemeNotSupported", "The URI scheme '{0}' is not supported"), FText::FromString(Scheme)));
+			}
+			return 0;
+		}
+
+		// Check for known extensions
+		TArray<FString> PathComponents;
+		UrlParser.GetPathComponents(PathComponents);
+		FString LowerCaseExtension = FPaths::GetExtension(PathComponents.Last().ToLower());
+		// If the extension is known, we are confident that we can play this.
+		// At this point there is no information provided on the codecs used in the media,
+		// so we cannot check for this.
+		if (SupportedFileExtensions.Contains(LowerCaseExtension))
+		{
+	// For the time being Protron *MUST* be selected explicitly.
+			return 0;
+		}
+		return 0;
+	}
+
+	TSharedPtr<IMediaPlayer, ESPMode::ThreadSafe> CreatePlayer(IMediaEventSink& EventSink) override
+	{
+		auto PlayerModule = FModuleManager::LoadModulePtr<IElectraProtronModule>("ElectraProtron");
+		return (PlayerModule != nullptr) ? PlayerModule->CreatePlayer(EventSink) : nullptr;
+	}
+
+	FText GetDisplayName() const override
+	{
+		return LOCTEXT("MediaPlayerDisplayName", "Protron mp4 playback");
+	}
+
+	FName GetPlayerName() const override
+	{
+		static FName PlayerName(TEXT("Protron"));
+		return PlayerName;
+	}
+
+	FGuid GetPlayerPluginGUID() const override
+	{
+		static FGuid PlayerPluginGUID(0x2899727b, 0xfc934ccb, 0x94119db7, 0x185741d8);
+		return PlayerPluginGUID;
+	}
+
+	const TArray<FString>& GetSupportedPlatforms() const override
+	{
+		return SupportedPlatforms;
+	}
+
+	bool SupportsFeature(EMediaFeature Feature) const override
+	{
+		return ((Feature == EMediaFeature::AudioSamples) ||
+				(Feature == EMediaFeature::AudioTracks) ||
+				(Feature == EMediaFeature::VideoSamples) ||
+				(Feature == EMediaFeature::VideoTracks));
+	}
+
+public:
+
+	// IModuleInterface interface
+
+	void StartupModule() override
+	{
+		// supported platforms
+		MediaModule = FModuleManager::GetModulePtr<IMediaModule>("Media");
+		check(MediaModule);
+
+		AddSupportedPlatform(FGuid(0xd1d5f296, 0xff834a87, 0xb20faaa9, 0xd6b8e9a6));
+		AddSupportedPlatform(FGuid(0x003be296, 0x17004f0c, 0x8e1f7860, 0x81efbb1f));
+		AddSupportedPlatform(FGuid(0x115de4fe, 0x241b465b, 0x970a872f, 0x3167492a));
+		AddSupportedPlatform(FGuid(0xc0b45a33, 0x9de340c7, 0xbce24c47, 0x15c3babf));
+
+		// supported schemes
+		SupportedUriSchemes.Add(TEXT("file"));
+
+		// supported file extensions
+		SupportedFileExtensions.Add(TEXT("mp4"));
+		SupportedFileExtensions.Add(TEXT("m4v"));
+		SupportedFileExtensions.Add(TEXT("m4a"));
+		SupportedFileExtensions.Add(TEXT("mov"));
+
+		// register player factory
+		MediaModule->RegisterPlayerFactory(*this);
+	}
+
+	void ShutdownModule() override
+	{
+		// Get media module once more to be sure it is still there
+		MediaModule = FModuleManager::GetModulePtr<IMediaModule>("Media");
+
+		// unregister player factory
+		if (MediaModule != nullptr)
+		{
+			MediaModule->UnregisterPlayerFactory(*this);
+		}
+	}
+
+private:
+	void AddSupportedPlatform(const FGuid& PlatformGuid)
+	{
+		FName PlatformName = MediaModule->GetPlatformName(PlatformGuid);
+		if (!PlatformName.IsNone())
+		{
+			SupportedPlatforms.Add(PlatformName.ToString());
+		}
+	}
+
+	/** Media module */
+	IMediaModule* MediaModule;
+
+	/** List of platforms that the media player support. */
+	TArray<FString> SupportedPlatforms;
+
+	/** List of supported URI schemes. */
+	TArray<FString> SupportedUriSchemes;
+
+	/** List of supported media file types. */
+	TArray<FString> SupportedFileExtensions;
+};
+
+
+#undef LOCTEXT_NAMESPACE
+
+IMPLEMENT_MODULE(ElectraProtronFactoryModule, ElectraProtronFactory);
diff --git a/Engine/Plugins/Media/ElectraPlayer/Source/ElectraProtronFactory/Private/ElectraProtronFactoryPrivate.h b/Engine/Plugins/Media/ElectraPlayer/Source/ElectraProtronFactory/Private/ElectraProtronFactoryPrivate.h
new file mode 100644
--- /dev/null
+++ b/Engine/Plugins/Media/ElectraPlayer/Source/ElectraProtronFactory/Private/ElectraProtronFactoryPrivate.h
@@ -0,0 +1,8 @@
+// Copyright Epic Games, Inc. All Rights Reserved.
+
+#pragma once
+
+#include "Logging/LogMacros.h"
+
+/** Log category */
+DECLARE_LOG_CATEGORY_EXTERN(LogElectraProtronFactory, Log, All);
