ElectraProtron: Added a missing Close() call

#rnx

diff --git a/Engine/Plugins/Media/ElectraPlayer/Source/ElectraProtron/Private/Player/ElectraProtronPlayer.cpp b/Engine/Plugins/Media/ElectraPlayer/Source/ElectraProtron/Private/Player/ElectraProtronPlayer.cpp
--- a/Engine/Plugins/Media/ElectraPlayer/Source/ElectraProtron/Private/Player/ElectraProtronPlayer.cpp
+++ b/Engine/Plugins/Media/ElectraPlayer/Source/ElectraProtron/Private/Player/ElectraProtronPlayer.cpp
@@ -22,6 +22,7 @@ FElectraProtronPlayer::FElectraProtronPlayer(IMediaEventSink& InEventSink)
 
 FElectraProtronPlayer::~FElectraProtronPlayer()
 {
+	Close();
 }
 
 void FElectraProtronPlayer::Close()
