ElectraProtron: Renamed the internal player name

#rnx

diff --git a/Engine/Plugins/Media/ElectraPlayer/Source/ElectraProtronFactory/Private/ElectraProtronFactoryModule.cpp b/Engine/Plugins/Media/ElectraPlayer/Source/ElectraProtronFactory/Private/ElectraProtronFactoryModule.cpp
--- a/Engine/Plugins/Media/ElectraPlayer/Source/ElectraProtronFactory/Private/ElectraProtronFactoryModule.cpp
+++ b/Engine/Plugins/Media/ElectraPlayer/Source/ElectraProtronFactory/Private/ElectraProtronFactoryModule.cpp
@@ -89,12 +89,12 @@ public:
 
 	FText GetDisplayName() const override
 	{
-		return LOCTEXT("MediaPlayerDisplayName", "Protron mp4 playback");
+		return LOCTEXT("MediaPlayerDisplayName", "Electra Protron mp4 playback");
 	}
 
 	FName GetPlayerName() const override
 	{
-		static FName PlayerName(TEXT("Protron"));
+		static FName PlayerName(TEXT("ElectraProtron"));
 		return PlayerName;
 	}
 
