ElectraProtron: removed unsupported platform guids

#rnx

diff --git a/Engine/Plugins/Media/ElectraPlayer/Source/ElectraProtronFactory/Private/ElectraProtronFactoryModule.cpp b/Engine/Plugins/Media/ElectraPlayer/Source/ElectraProtronFactory/Private/ElectraProtronFactoryModule.cpp
--- a/Engine/Plugins/Media/ElectraPlayer/Source/ElectraProtronFactory/Private/ElectraProtronFactoryModule.cpp
+++ b/Engine/Plugins/Media/ElectraPlayer/Source/ElectraProtronFactory/Private/ElectraProtronFactoryModule.cpp
@@ -129,8 +129,6 @@ public:
 
 		AddSupportedPlatform(FGuid(0xd1d5f296, 0xff834a87, 0xb20faaa9, 0xd6b8e9a6));
 		AddSupportedPlatform(FGuid(0x003be296, 0x17004f0c, 0x8e1f7860, 0x81efbb1f));
-		AddSupportedPlatform(FGuid(0x115de4fe, 0x241b465b, 0x970a872f, 0x3167492a));
-		AddSupportedPlatform(FGuid(0xc0b45a33, 0x9de340c7, 0xbce24c47, 0x15c3babf));
 
 		// supported schemes
 		SupportedUriSchemes.Add(TEXT("file"));
