ElectraProtron: Less strict checking for valid avcC color depth, which gets set incorrectly by some tools

#rnx

diff --git a/Engine/Plugins/Media/ElectraPlayer/Source/ElectraProtron/Private/Player/ElectraProtronPlayerImpl.cpp b/Engine/Plugins/Media/ElectraPlayer/Source/ElectraProtron/Private/Player/ElectraProtronPlayerImpl.cpp
--- a/Engine/Plugins/Media/ElectraPlayer/Source/ElectraProtron/Private/Player/ElectraProtronPlayerImpl.cpp
+++ b/Engine/Plugins/Media/ElectraPlayer/Source/ElectraProtron/Private/Player/ElectraProtronPlayerImpl.cpp
@@ -691,7 +691,7 @@ void FElectraProtronPlayer::FImpl::GetTrackCodecInfo(ElectraProtronUtils::FCodec
 			case Electra::UtilitiesMP4::MakeBoxAtom('a','v','c','1'):
 			case Electra::UtilitiesMP4::MakeBoxAtom('a','v','c','3'):
 			{
-				if (Visuals[0]->GetFrameCount() == 1 && Visuals[0]->GetDepth() == 0x18)
+				if (Visuals[0]->GetFrameCount() == 1)
 				{
 					auto Avcc = Visuals[0]->FindBoxRecursive<Electra::UtilitiesMP4::FMP4BoxAVCC>(Electra::UtilitiesMP4::MakeBoxAtom('a','v','c','C'), 0);
 					if (Avcc.IsValid())
@@ -725,12 +725,16 @@ void FElectraProtronPlayer::FImpl::GetTrackCodecInfo(ElectraProtronUtils::FCodec
 						}
 					}
 				}
+				else
+				{
+					UE_LOG(LogElectraProtron, Warning, TEXT("Track #%u has a frame count other than 1 in the VisualSampleEntry, ignoring this track."), InTrackID);
+				}
 				break;
 			}
 			case Electra::UtilitiesMP4::MakeBoxAtom('h','v','c','1'):
 			case Electra::UtilitiesMP4::MakeBoxAtom('h','e','v','1'):
 			{
-				if (Visuals[0]->GetFrameCount() == 1 && Visuals[0]->GetDepth() == 0x18)
+				if (Visuals[0]->GetFrameCount() == 1)
 				{
 					auto Hvcc = Visuals[0]->FindBoxRecursive<Electra::UtilitiesMP4::FMP4BoxHVCC>(Electra::UtilitiesMP4::MakeBoxAtom('h','v','c','C'), 0);
 					if (Hvcc.IsValid())
@@ -764,6 +768,10 @@ void FElectraProtronPlayer::FImpl::GetTrackCodecInfo(ElectraProtronUtils::FCodec
 						}
 					}
 				}
+				else
+				{
+					UE_LOG(LogElectraProtron, Warning, TEXT("Track #%u has a frame count other than 1 in the VisualSampleEntry, ignoring this track."), InTrackID);
+				}
 				break;
 			}
 			default:
