Electra: Changed log messages to display an id instead of the 'this' pointer value.

#rnx

diff --git a/Engine/Plugins/Media/ElectraPlayer/Source/ElectraPlayerPlugin/Private/ElectraPlayerPlugin.cpp b/Engine/Plugins/Media/ElectraPlayer/Source/ElectraPlayerPlugin/Private/ElectraPlayerPlugin.cpp
--- a/Engine/Plugins/Media/ElectraPlayer/Source/ElectraPlayerPlugin/Private/ElectraPlayerPlugin.cpp
+++ b/Engine/Plugins/Media/ElectraPlayer/Source/ElectraPlayerPlugin/Private/ElectraPlayerPlugin.cpp
@@ -21,6 +21,10 @@ using namespace Electra;
 
 //-----------------------------------------------------------------------------
 
+std::atomic<uint32> FElectraPlayerPlugin::NextPlayerUniqueID { 0 };
+
+//-----------------------------------------------------------------------------
+
 FElectraPlayerPlugin::FElectraPlayerPlugin()
 {
 	// Make sure a few assumptions are correct...
@@ -526,9 +530,11 @@ bool FElectraPlayerPlugin::Open(const FString& Url, const IMediaOptions* Options
 
 bool FElectraPlayerPlugin::Open(const FString& Url, const IMediaOptions* Options, const FMediaPlayerOptions* InPlayerOptions)
 {
+	PlayerUniqueID = ++NextPlayerUniqueID;
+
 	if (Options == nullptr)
 	{
-		UE_LOG(LogElectraPlayerPlugin, Error, TEXT("[%p] IMediaPlayer::Open: Options == nullptr"), this);
+		UE_LOG(LogElectraPlayerPlugin, Error, TEXT("[%u] IMediaPlayer::Open: Options == nullptr"), PlayerUniqueID);
 
 		FScopeLock lock(&CallbackPointerLock);
 		if (EventSink)
@@ -543,7 +549,7 @@ bool FElectraPlayerPlugin::Open(const FString& Url, const IMediaOptions* Options
 	CallbackPointerLock.Lock();
 	OptionInterface = StaticCastSharedPtr<IElectraSafeMediaOptionInterface>(Options->GetMediaOption(TEXT("GetSafeMediaOptions"), TSharedPtr<IElectraSafeMediaOptionInterface, ESPMode::ThreadSafe>()));
 	CallbackPointerLock.Unlock();
-	UE_LOG(LogElectraPlayerPlugin, Log, TEXT("[%p] IMediaPlayer::Open: Options@%p"), this, Options);
+	UE_LOG(LogElectraPlayerPlugin, Log, TEXT("[%u] IMediaPlayer::Open"), PlayerUniqueID);
 
 	IElectraPlayerInterface::FPlaystartOptions LocalPlaystartOptions;
 
@@ -567,24 +573,26 @@ bool FElectraPlayerPlugin::Open(const FString& Url, const IMediaOptions* Options
 	if (InitialAudioLanguage.Len())
 	{
 		LocalPlaystartOptions.InitialAudioTrackAttributes.Language_ISO639 = InitialAudioLanguage;
-		UE_LOG(LogElectraPlayerPlugin, Verbose, TEXT("[%p] IMediaPlayer::Open: Asking for initial audio language \"%s\""), this, *InitialAudioLanguage);
+		UE_LOG(LogElectraPlayerPlugin, Verbose, TEXT("[%u] IMediaPlayer::Open: Asking for initial audio language \"%s\""), PlayerUniqueID, *InitialAudioLanguage);
 	}
 	FString InitialSubtitleLanguage = Options->GetMediaOption(TEXT("InitialSubtitleLanguage"), FString());
 	if (InitialSubtitleLanguage.Len())
 	{
 		LocalPlaystartOptions.InitialSubtitleTrackAttributes.Language_ISO639 = InitialSubtitleLanguage;
-		UE_LOG(LogElectraPlayerPlugin, Verbose, TEXT("[%p] IMediaPlayer::Open: Asking for initial subtitle language \"%s\""), this, *InitialSubtitleLanguage);
+		UE_LOG(LogElectraPlayerPlugin, Verbose, TEXT("[%u] IMediaPlayer::Open: Asking for initial subtitle language \"%s\""), PlayerUniqueID, *InitialSubtitleLanguage);
 	}
 	bool bNoPreloading = Options->GetMediaOption(TEXT("ElectraNoPreloading"), (bool)false);
 	if (bNoPreloading)
 	{
 		LocalPlaystartOptions.bDoNotPreload = true;
-		UE_LOG(LogElectraPlayerPlugin, Verbose, TEXT("[%p] IMediaPlayer::Open: No preloading after opening media"), this);
+		UE_LOG(LogElectraPlayerPlugin, Verbose, TEXT("[%u] IMediaPlayer::Open: No preloading after opening media"), PlayerUniqueID);
 	}
 
 
 	// Set up options to initialize the internal player with.
 	Electra::FParamDict PlayerOptions;
+	static const FName KeyUniquePlayerID(TEXT("UniquePlayerID"));
+	PlayerOptions.Set(KeyUniquePlayerID, Electra::FVariantValue((int64)PlayerUniqueID));
 	const FName OptionsByString[] = { TEXT("excluded_codecs_video"), TEXT("excluded_codecs_audio"), TEXT("excluded_codecs_subtitles"), TEXT("preferred_codecs_video"), TEXT("preferred_codecs_audio"),TEXT("preferred_codecs_subtitles") };
 	for(auto &StringOption : OptionsByString)
 	{
@@ -609,49 +617,49 @@ bool FElectraPlayerPlugin::Open(const FString& Url, const IMediaOptions* Options
 	if (InitialStreamBitrate > 0)
 	{
 		PlayerOptions.Set(TEXT("initial_bitrate"), Electra::FVariantValue(InitialStreamBitrate));
-		UE_LOG(LogElectraPlayerPlugin, Verbose, TEXT("[%p] IMediaPlayer::Open: Using initial bitrate of %d bits/second"), this, (int32)InitialStreamBitrate);
+		UE_LOG(LogElectraPlayerPlugin, Verbose, TEXT("[%u] IMediaPlayer::Open: Using initial bitrate of %d bits/second"), PlayerUniqueID, (int32)InitialStreamBitrate);
 	}
 	FString MediaMimeType = Options->GetMediaOption(TEXT("mimetype"), FString());
 	if (MediaMimeType.Len())
 	{
 		PlayerOptions.Set(TEXT("mime_type"), Electra::FVariantValue(MediaMimeType));
-		UE_LOG(LogElectraPlayerPlugin, Verbose, TEXT("[%p] IMediaPlayer::Open: Setting media mime type to \"%s\""), this, *MediaMimeType);
+		UE_LOG(LogElectraPlayerPlugin, Verbose, TEXT("[%u] IMediaPlayer::Open: Setting media mime type to \"%s\""), PlayerUniqueID, *MediaMimeType);
 	}
 	int64 MaxVerticalHeight = Options->GetMediaOption(TEXT("MaxElectraVerticalResolution"), (int64)-1);
 	if (MaxVerticalHeight > 0)
 	{
 		PlayerOptions.Set(TEXT("max_resoY"), Electra::FVariantValue(MaxVerticalHeight));
-		UE_LOG(LogElectraPlayerPlugin, Verbose, TEXT("[%p] IMediaPlayer::Open: Limiting vertical resolution to %d for all streams"), this, (int32)MaxVerticalHeight);
+		UE_LOG(LogElectraPlayerPlugin, Verbose, TEXT("[%u] IMediaPlayer::Open: Limiting vertical resolution to %d for all streams"), PlayerUniqueID, (int32)MaxVerticalHeight);
 	}
 	int64 MaxVerticalHeightAt60 = Options->GetMediaOption(TEXT("MaxElectraVerticalResolutionOf60fpsVideos"), (int64)-1);
 	if (MaxVerticalHeightAt60 > 0)
 	{
 		PlayerOptions.Set(TEXT("max_resoY_above_30fps"), Electra::FVariantValue(MaxVerticalHeightAt60));
-		UE_LOG(LogElectraPlayerPlugin, Verbose, TEXT("[%p] IMediaPlayer::Open: Limiting vertical resolution to %d for streams >30fps"), this, (int32)MaxVerticalHeightAt60);
+		UE_LOG(LogElectraPlayerPlugin, Verbose, TEXT("[%u] IMediaPlayer::Open: Limiting vertical resolution to %d for streams >30fps"), PlayerUniqueID, (int32)MaxVerticalHeightAt60);
 	}
 	double LiveEdgeDistanceForNormalPresentation = Options->GetMediaOption(TEXT("ElectraLivePresentationOffset"), (double)-1.0);
 	if (LiveEdgeDistanceForNormalPresentation > 0.0)
 	{
 		PlayerOptions.Set(TEXT("seekable_range_live_end_offset"), Electra::FVariantValue(Electra::FTimeValue().SetFromSeconds(LiveEdgeDistanceForNormalPresentation)));
-		UE_LOG(LogElectraPlayerPlugin, Verbose, TEXT("[%p] IMediaPlayer::Open: Setting distance to live edge for normal presentations to %.3f seconds"), this, LiveEdgeDistanceForNormalPresentation);
+		UE_LOG(LogElectraPlayerPlugin, Verbose, TEXT("[%u] IMediaPlayer::Open: Setting distance to live edge for normal presentations to %.3f seconds"), PlayerUniqueID, LiveEdgeDistanceForNormalPresentation);
 	}
 	double LiveEdgeDistanceForAudioOnlyPresentation = Options->GetMediaOption(TEXT("ElectraLiveAudioPresentationOffset"), (double)-1.0);
 	if (LiveEdgeDistanceForAudioOnlyPresentation > 0.0)
 	{
 		PlayerOptions.Set(TEXT("seekable_range_live_end_offset_audioonly"), Electra::FVariantValue(Electra::FTimeValue().SetFromSeconds(LiveEdgeDistanceForAudioOnlyPresentation)));
-		UE_LOG(LogElectraPlayerPlugin, Verbose, TEXT("[%p] IMediaPlayer::Open: Setting distance to live edge for audio-only presentation to %.3f seconds"), this, LiveEdgeDistanceForAudioOnlyPresentation);
+		UE_LOG(LogElectraPlayerPlugin, Verbose, TEXT("[%u] IMediaPlayer::Open: Setting distance to live edge for audio-only presentation to %.3f seconds"), PlayerUniqueID, LiveEdgeDistanceForAudioOnlyPresentation);
 	}
 	bool bUseConservativeLiveEdgeDistance = Options->GetMediaOption(TEXT("ElectraLiveUseConservativePresentationOffset"), (bool)false);
 	if (bUseConservativeLiveEdgeDistance)
 	{
 		PlayerOptions.Set(TEXT("seekable_range_live_end_offset_conservative"), Electra::FVariantValue(bUseConservativeLiveEdgeDistance));
-		UE_LOG(LogElectraPlayerPlugin, Verbose, TEXT("[%p] IMediaPlayer::Open: Using conservative live edge for distance calculation"), this);
+		UE_LOG(LogElectraPlayerPlugin, Verbose, TEXT("[%u] IMediaPlayer::Open: Using conservative live edge for distance calculation"), PlayerUniqueID);
 	}
 	bool bThrowErrorWhenRebuffering = Options->GetMediaOption(TEXT("ElectraThrowErrorWhenRebuffering"), (bool)false);
 	if (bThrowErrorWhenRebuffering)
 	{
 		PlayerOptions.Set(TEXT("throw_error_when_rebuffering"), Electra::FVariantValue(bThrowErrorWhenRebuffering));
-		UE_LOG(LogElectraPlayerPlugin, Verbose, TEXT("[%p] IMediaPlayer::Open: Throw playback error when rebuffering"), this);
+		UE_LOG(LogElectraPlayerPlugin, Verbose, TEXT("[%u] IMediaPlayer::Open: Throw playback error when rebuffering"), PlayerUniqueID);
 	}
 	FString CDNHTTPStatusDenyStream = Options->GetMediaOption(TEXT("ElectraGetDenyStreamCode"), FString());
 	if (CDNHTTPStatusDenyStream.Len())
@@ -661,7 +669,7 @@ bool FElectraPlayerPlugin::Open(const FString& Url, const IMediaOptions* Options
 		if (HTTPStatus > 0 && HTTPStatus < 1000)
 		{
 			PlayerOptions.Set(TEXT("abr:cdn_deny_httpstatus"), Electra::FVariantValue((int64)HTTPStatus));
-			UE_LOG(LogElectraPlayerPlugin, Verbose, TEXT("[%p] IMediaPlayer::Open: CDN HTTP status %d will deny a stream permanently"), this, HTTPStatus);
+			UE_LOG(LogElectraPlayerPlugin, Verbose, TEXT("[%u] IMediaPlayer::Open: CDN HTTP status %d will deny a stream permanently"), PlayerUniqueID, HTTPStatus);
 		}
 	}
 
@@ -702,14 +710,14 @@ bool FElectraPlayerPlugin::Open(const FString& Url, const IMediaOptions* Options
 	int64 MaxVerticalStreamResolution = Options->GetMediaOption(TEXT("MaxResolutionForMediaStreaming"), DefaultValue);
 	if (MaxVerticalStreamResolution != 0)
 	{
-		UE_LOG(LogElectraPlayerPlugin, Log, TEXT("[%p] IMediaPlayer::Open: Limiting max resolution to %d"), this, (int32)MaxVerticalStreamResolution);
+		UE_LOG(LogElectraPlayerPlugin, Log, TEXT("[%u] IMediaPlayer::Open: Limiting max resolution to %d"), PlayerUniqueID, (int32)MaxVerticalStreamResolution);
 		LocalPlaystartOptions.MaxVerticalStreamResolution = (int32)MaxVerticalStreamResolution;
 	}
 
 	int64 MaxBandwidthForStreaming = Options->GetMediaOption(TEXT("ElectraMaxStreamingBandwidth"), (int64)0);
 	if (MaxBandwidthForStreaming > 0)
 	{
-		UE_LOG(LogElectraPlayerPlugin, Log, TEXT("[%p] Limiting max streaming bandwidth to %d bps"), this, (int32)MaxBandwidthForStreaming);
+		UE_LOG(LogElectraPlayerPlugin, Log, TEXT("[%u] Limiting max streaming bandwidth to %d bps"), PlayerUniqueID, (int32)MaxBandwidthForStreaming);
 		LocalPlaystartOptions.MaxBandwidthForStreaming = (int32)MaxBandwidthForStreaming;
 	}
 
@@ -744,7 +752,7 @@ bool FElectraPlayerPlugin::Open(const FString& Url, const IMediaOptions* Options
 bool FElectraPlayerPlugin::Open(const TSharedRef<FArchive, ESPMode::ThreadSafe>& Archive, const FString& OriginalUrl, const IMediaOptions* /*Options*/)
 {
 	// we support playback only from an external file, not from a "resource" (e.g. a packaged asset)
-	UE_LOG(LogElectraPlayerPlugin, Error, TEXT("[%p] IMediaPlayer::Archive"), this);
+	UE_LOG(LogElectraPlayerPlugin, Error, TEXT("[%u] IMediaPlayer::Archive"), PlayerUniqueID);
 	unimplemented();
 	return false;
 }
@@ -959,7 +967,7 @@ FTimespan FElectraPlayerPlugin::GetDuration() const
 
 bool FElectraPlayerPlugin::SetRate(float Rate)
 {
-	UE_LOG(LogElectraPlayerPlugin, Log, TEXT("[%p] IMediaControls::SetRate(%f)"), this, Rate);
+	UE_LOG(LogElectraPlayerPlugin, Log, TEXT("[%u] IMediaControls::SetRate(%f)"), PlayerUniqueID, Rate);
 	CSV_EVENT(ElectraPlayer, TEXT("Setting Rate"));
 
 	return Player->SetRate(Rate);
@@ -968,7 +976,7 @@ bool FElectraPlayerPlugin::SetRate(float Rate)
 
 bool FElectraPlayerPlugin::Seek(const FTimespan& Time)
 {
-	UE_LOG(LogElectraPlayerPlugin, Log, TEXT("[%p] IMediaControls::Seek() to %s"), this, *Time.ToString(TEXT("%h:%m:%s.%f")));
+	UE_LOG(LogElectraPlayerPlugin, Log, TEXT("[%u] IMediaControls::Seek() to %s"), PlayerUniqueID, *Time.ToString(TEXT("%h:%m:%s.%f")));
 	CSV_EVENT(ElectraPlayer, TEXT("Seeking"));
 
 	const int32 NextSequenceIndex = CurrentSequenceIndex + 1;
diff --git a/Engine/Plugins/Media/ElectraPlayer/Source/ElectraPlayerPlugin/Private/ElectraPlayerPlugin.h b/Engine/Plugins/Media/ElectraPlayer/Source/ElectraPlayerPlugin/Private/ElectraPlayerPlugin.h
--- a/Engine/Plugins/Media/ElectraPlayer/Source/ElectraPlayerPlugin/Private/ElectraPlayerPlugin.h
+++ b/Engine/Plugins/Media/ElectraPlayer/Source/ElectraPlayerPlugin/Private/ElectraPlayerPlugin.h
@@ -19,6 +19,7 @@
 #include "Logging/LogMacros.h"
 #include "ElectraTextureSample.h"
 #include "ElectraPlayerAudioSample.h"
+#include <atomic>
 
 class IElectraPlayerRuntimeModule;
 class IElectraSafeMediaOptionInterface;
@@ -246,6 +247,8 @@ private:
 
 	/** The actual player */
 	TSharedPtr<IElectraPlayerInterface, ESPMode::ThreadSafe> Player;
+	static std::atomic<uint32> NextPlayerUniqueID;
+	uint32 PlayerUniqueID = 0;
 
 	/** Current player stream metadata */
 	mutable bool bMetadataChanged = false;
diff --git a/Engine/Plugins/Media/ElectraPlayer/Source/ElectraPlayerRuntime/Private/ElectraPlayer.cpp b/Engine/Plugins/Media/ElectraPlayer/Source/ElectraPlayerRuntime/Private/ElectraPlayer.cpp
--- a/Engine/Plugins/Media/ElectraPlayer/Source/ElectraPlayerRuntime/Private/ElectraPlayer.cpp
+++ b/Engine/Plugins/Media/ElectraPlayer/Source/ElectraPlayerRuntime/Private/ElectraPlayer.cpp
@@ -164,8 +164,6 @@ FElectraPlayer::FElectraPlayer(const TSharedPtr<IElectraPlayerAdapterDelegate, E
 	VideoDecoderResourceDelegate = PlatformCreateVideoDecoderResourceDelegate(AdapterDelegate);
 
 	ClearToDefaultState();
-
-	UE_LOG(LogElectraPlayer, Verbose, TEXT("[%p] FElectraPlayer() created."), this);
 }
 
 //-----------------------------------------------------------------------------
@@ -186,7 +184,7 @@ FElectraPlayer::~FElectraPlayer()
 	SendAnalyticMetricsPerMinuteDelegate.RemoveAll(this);
 	ReportVideoStreamingErrorDelegate.RemoveAll(this);
 	ReportSubtitlesMetricsDelegate.RemoveAll(this);
-	UE_LOG(LogElectraPlayer, Verbose, TEXT("[%p] ~FElectraPlayer() finished."), this);
+	UE_LOG(LogElectraPlayer, Verbose, TEXT("[%u] ~FElectraPlayer() finished."), InstanceID);
 
 	if (AsyncResourceReleaseNotification.IsValid())
 	{
@@ -229,9 +227,13 @@ void FElectraPlayer::ClearToDefaultState()
  */
 bool FElectraPlayer::OpenInternal(const FString& Url, const FParamDict& InPlayerOptions, const FPlaystartOptions& InPlaystartOptions, EOpenType InOpenType)
 {
+	static const FName KeyUniquePlayerID(TEXT("UniquePlayerID"));
+
 	LLM_SCOPE(ELLMTag::ElectraPlayer);
 	CSV_EVENT(ElectraPlayer, TEXT("Open"));
 
+	InstanceID = (uint32)InPlayerOptions.GetValue(KeyUniquePlayerID).SafeGetInt64(0);
+
 	// Open the provided URL as a media or a blob?
 	FString BlobParams;
 	bool bPreviousOpenLoadedBlob = PendingBlobRequest.IsValid();
@@ -361,7 +363,7 @@ bool FElectraPlayer::OpenInternal(const FString& Url, const FParamDict& InPlayer
 		}
 
 		// Issue load of the playlist.
-		UE_LOG(LogElectraPlayer, Log, TEXT("[%p][%p] IMediaPlayer::Open(%s)"), this, CurrentPlayer.Get(), *SanitizeMessage(MediaUrl));
+		UE_LOG(LogElectraPlayer, Log, TEXT("[%u] IMediaPlayer::Open(%s)"), InstanceID, *SanitizeMessage(MediaUrl));
 		CurrentPlayer->AdaptivePlayer->LoadManifest(MediaUrl);
 	}
 	else
@@ -369,12 +371,12 @@ bool FElectraPlayer::OpenInternal(const FString& Url, const FParamDict& InPlayer
 		PendingBlobRequest = MakeShared<FBlobRequest, ESPMode::ThreadSafe>();
 		if (!PendingBlobRequest->Request->SetFromJSON(BlobParams))
 		{
-			UE_LOG(LogElectraPlayer, Error, TEXT("[%p] IMediaPlayer::OpenBlob(%s) has bad JSON parameters"), this, *SanitizeMessage(MediaUrl));
+			UE_LOG(LogElectraPlayer, Error, TEXT("[%u] IMediaPlayer::OpenBlob(%s) has bad JSON parameters"), InstanceID, *SanitizeMessage(MediaUrl));
 			PendingBlobRequest.Reset();
 			return false;
 		}
 		CurrentPlayer = MoveTemp(NewPlayer);
-		UE_LOG(LogElectraPlayer, Verbose, TEXT("[%p][%p] IMediaPlayer::OpenBlob(%s)"), this, CurrentPlayer.Get(), *SanitizeMessage(MediaUrl));
+		UE_LOG(LogElectraPlayer, Verbose, TEXT("[%u] IMediaPlayer::OpenBlob(%s)"), InstanceID, *SanitizeMessage(MediaUrl));
 		PendingBlobRequest->Request->URL(MediaUrl).Callback().BindThreadSafeSP(PendingBlobRequest.ToSharedRef(), &FBlobRequest::OnBlobRequestComplete);
 		CurrentPlayer->AdaptivePlayer->LoadBlob(PendingBlobRequest->Request);
 	}
@@ -400,7 +402,7 @@ void FElectraPlayer::CloseInternal(bool bKillAfterClose)
 	PlayerLock.Unlock();
 
 
-	UE_LOG(LogElectraPlayer, Log, TEXT("[%p][%p] IMediaPlayer::Close()"), this, CurrentPlayer.Get());
+	UE_LOG(LogElectraPlayer, Log, TEXT("[%u] IMediaPlayer::Close()"), InstanceID);
 	CSV_EVENT(ElectraPlayer, TEXT("Close"));
 
 	/*
@@ -408,7 +410,7 @@ void FElectraPlayer::CloseInternal(bool bKillAfterClose)
 	 * that we need to make sure will not report back to us or deliver any pending data while we
 	 * are cleaning everything up.
 	 */
-	TSharedPtr<FInternalPlayerImpl, ESPMode::ThreadSafe>	Player = CurrentPlayer;
+	TSharedPtr<FInternalPlayerImpl, ESPMode::ThreadSafe> Player = CurrentPlayer;
 
 	// For all intents and purposes the player can be considered closed here now already.
 	PlayerState.State = EPlayerState::Closed;
@@ -482,18 +484,18 @@ void FElectraPlayer::CloseInternal(bool bKillAfterClose)
 	CurrentPlayer.Reset();
 
 	// Kick off asynchronous closing now.
-	FInternalPlayerImpl::DoCloseAsync(MoveTemp(Player), AsyncResourceReleaseNotification);
+	FInternalPlayerImpl::DoCloseAsync(MoveTemp(Player), InstanceID, AsyncResourceReleaseNotification);
 
 	bAllowKillAfterCloseEvent = bKillAfterClose;
 
 	WaitForPlayerDestroyedEvent->Trigger();
 }
 
-void FElectraPlayer::FInternalPlayerImpl::DoCloseAsync(TSharedPtr<FInternalPlayerImpl, ESPMode::ThreadSafe> && Player, TSharedPtr<FElectraPlayer::IAsyncResourceReleaseNotifyContainer, ESPMode::ThreadSafe> AsyncResourceReleaseNotification)
+void FElectraPlayer::FInternalPlayerImpl::DoCloseAsync(TSharedPtr<FInternalPlayerImpl, ESPMode::ThreadSafe>&& Player, uint32 PlayerID, TSharedPtr<FElectraPlayer::IAsyncResourceReleaseNotifyContainer, ESPMode::ThreadSafe> AsyncResourceReleaseNotification)
 {
 	TSharedPtr<volatile bool> bClosedSig = MakeShared<volatile bool>(false);
 
-	TFunction<void()> CloseTask = [Player, bClosedSig, AsyncResourceReleaseNotification]()
+	TFunction<void()> CloseTask = [Player, PlayerID, bClosedSig, AsyncResourceReleaseNotification]()
 	{
 		double TimeCloseBegan = FPlatformTime::Seconds();
 		Player->AdaptivePlayer->Stop();
@@ -502,7 +504,7 @@ void FElectraPlayer::FInternalPlayerImpl::DoCloseAsync(TSharedPtr<FInternalPlaye
 		Player->RendererAudio.Reset();
 		*bClosedSig = true;
 		double TimeCloseEnded = FPlatformTime::Seconds();
-		UE_LOG(LogElectraPlayer, Verbose, TEXT("[%p] DoCloseAsync() finished after %.3f msec!"), Player.Get(), (TimeCloseEnded - TimeCloseBegan) * 1000.0);
+		UE_LOG(LogElectraPlayer, Verbose, TEXT("[%u] DoCloseAsync() finished after %.3f msec!"), PlayerID, (TimeCloseEnded - TimeCloseBegan) * 1000.0);
 
 		if (AsyncResourceReleaseNotification.IsValid())
 		{
@@ -662,7 +664,7 @@ void FElectraPlayer::Tick(FTimespan DeltaTime, FTimespan Timecode)
 				if (NewVerticalStreamResolution != PlaystartOptions.MaxVerticalStreamResolution.Get(0))
 				{
 					PlaystartOptions.MaxVerticalStreamResolution = NewVerticalStreamResolution;
-					UE_LOG(LogElectraPlayer, Log, TEXT("[%p][%p] Limiting max vertical resolution to %d"), this, Player.Get(), (int32)NewVerticalStreamResolution);
+					UE_LOG(LogElectraPlayer, Log, TEXT("[%u] Limiting max vertical resolution to %d"), InstanceID, (int32)NewVerticalStreamResolution);
 					Player->AdaptivePlayer->SetMaxResolution(0, (int32)NewVerticalStreamResolution);
 				}
 			}
@@ -673,7 +675,7 @@ void FElectraPlayer::Tick(FTimespan DeltaTime, FTimespan Timecode)
 				if (NewBandwidthForStreaming != PlaystartOptions.MaxBandwidthForStreaming.Get(0))
 				{
 					PlaystartOptions.MaxBandwidthForStreaming = NewBandwidthForStreaming;
-					UE_LOG(LogElectraPlayer, Log, TEXT("[%p][%p] Limiting max streaming bandwidth to %d bps"), this, Player.Get(), (int32)NewBandwidthForStreaming);
+					UE_LOG(LogElectraPlayer, Log, TEXT("[%u] Limiting max streaming bandwidth to %d bps"), InstanceID, (int32)NewBandwidthForStreaming);
 					Player->AdaptivePlayer->SetBitrateCeiling((int32)NewBandwidthForStreaming);
 				}
 			}
@@ -994,7 +996,7 @@ bool FElectraPlayer::SetLooping(bool bLooping)
 		IAdaptiveStreamingPlayer::FLoopParam loop;
 		loop.bEnableLooping = bLooping;
 		CurrentPlayer->AdaptivePlayer->SetLooping(loop);
-		UE_LOG(LogElectraPlayer, Verbose, TEXT("[%p][%p] IMediaPlayer::SetLooping(%s)"), this, CurrentPlayer.Get(), bLooping?TEXT("true"):TEXT("false"));
+		UE_LOG(LogElectraPlayer, Verbose, TEXT("[%u] IMediaPlayer::SetLooping(%s)"), InstanceID, bLooping?TEXT("true"):TEXT("false"));
 		return true;
 	}
 	return false;
@@ -1096,7 +1098,7 @@ TRangeSet<float> FElectraPlayer::GetSupportedRates(EPlayRateType InPlayRateType)
 
 bool FElectraPlayer::SetRate(float Rate)
 {
-	//UE_LOG(LogElectraPlayer, Verbose, TEXT("[%p][%p] IMediaPlayer::SetRate(%.3f)"), this, CurrentPlayer.Get(), Rate);
+	//UE_LOG(LogElectraPlayer, Verbose, TEXT("[%u] IMediaPlayer::SetRate(%.3f)"), InstanceID, Rate);
 	if (CurrentPlayer.Get())
 	{
 		// Set the intended rate, which *may* be set negative. This is not supported and we put the adaptive player into pause
@@ -1769,7 +1771,7 @@ void FElectraPlayer::OnMediaPlayerEventReceived(TSharedPtrTS<IAdaptiveStreamingP
 {
 #if !UE_BUILD_SHIPPING
 	const TCHAR* const Origins[] = { TEXT("Playlist"), TEXT("Inband"), TEXT("TimedMetadata"), TEXT("???") };
-	UE_LOG(LogElectraPlayer, Verbose, TEXT("[%p][%p] %s event %s with \"%s\", \"%s\", \"%s\" PTS @ %.3f for %.3fs"), this, CurrentPlayer.Get(),
+	UE_LOG(LogElectraPlayer, Verbose, TEXT("[%u] %s event %s with \"%s\", \"%s\", \"%s\" PTS @ %.3f for %.3fs"), InstanceID,
 		Origins[Electra::Utils::Min((int32)InEvent->GetOrigin(), (int32)UE_ARRAY_COUNT(Origins)-1)],
 		InDispatchMode==IAdaptiveStreamingPlayerAEMSReceiver::EDispatchMode::OnReceive?TEXT("received"):TEXT("started"),
 		*InEvent->GetSchemeIdUri(), *InEvent->GetValue(), *InEvent->GetID(),
@@ -2085,7 +2087,7 @@ void FElectraPlayer::HandlePlayerEventOpenSource(const FString& URL)
 	PlayerState.State = EPlayerState::Preparing;
 	DeferredEvents.Enqueue(IElectraPlayerAdapterDelegate::EPlayerEvent::MediaConnecting);
 
-	UE_LOG(LogElectraPlayer, Verbose, TEXT("[%p][%p] Opening stream at \"%s\""), this, CurrentPlayer.Get(), *SanitizeMessage(URL));
+	UE_LOG(LogElectraPlayer, Verbose, TEXT("[%u] Opening stream at \"%s\""), InstanceID, *SanitizeMessage(URL));
 
 	// Update statistics
 	FScopeLock Lock(&StatisticsLock);
@@ -2106,7 +2108,7 @@ void FElectraPlayer::HandlePlayerEventOpenSource(const FString& URL)
 
 void FElectraPlayer::HandlePlayerEventReceivedMasterPlaylist(const FString& EffectiveURL)
 {
-	UE_LOG(LogElectraPlayer, Verbose, TEXT("[%p][%p] Received master playlist from \"%s\""), this, CurrentPlayer.Get(), *SanitizeMessage(EffectiveURL));
+	UE_LOG(LogElectraPlayer, Verbose, TEXT("[%u] Received master playlist from \"%s\""), InstanceID, *SanitizeMessage(EffectiveURL));
 
 	// Update statistics
 	FScopeLock Lock(&StatisticsLock);
@@ -2137,7 +2139,7 @@ void FElectraPlayer::HandlePlayerEventReceivedPlaylists()
 
 	MediaStateOnPreparingFinished();
 
-	UE_LOG(LogElectraPlayer, Verbose, TEXT("[%p][%p] Received initial stream playlists"), this, CurrentPlayer.Get());
+	UE_LOG(LogElectraPlayer, Verbose, TEXT("[%u] Received initial stream playlists"), InstanceID);
 
 	Electra::FTimeRange MediaTimeline;
 	Electra::FTimeValue MediaDuration;
@@ -2163,7 +2165,7 @@ void FElectraPlayer::HandlePlayerEventReceivedPlaylists()
 	{
 		for(int32 i=0; i<VideoStreamMetaData[0].StreamDetails.Num(); ++i)
 		{
-			UE_LOG(LogElectraPlayer, Verbose, TEXT("[%p][%p] Found %d * %d video stream at bitrate %d"), this, CurrentPlayer.Get(),
+			UE_LOG(LogElectraPlayer, Verbose, TEXT("[%u] Found %d * %d video stream at bitrate %d"), InstanceID,
 												VideoStreamMetaData[0].StreamDetails[i].CodecInformation.GetResolution().Width,
 												VideoStreamMetaData[0].StreamDetails[i].CodecInformation.GetResolution().Height,
 												VideoStreamMetaData[0].StreamDetails[i].Bandwidth);
@@ -2188,7 +2190,7 @@ void FElectraPlayer::HandlePlayerEventReceivedPlaylists()
 	{
 		for(int32 i=0; i<AudioStreamMetaData[0].StreamDetails.Num(); ++i)
 		{
-			UE_LOG(LogElectraPlayer, Verbose, TEXT("[%p][%p] Found audio stream at bitrate %d"), this, CurrentPlayer.Get(),
+			UE_LOG(LogElectraPlayer, Verbose, TEXT("[%u] Found audio stream at bitrate %d"), InstanceID,
 												AudioStreamMetaData[0].StreamDetails[i].Bandwidth);
 		}
 	}
@@ -2227,7 +2229,7 @@ void FElectraPlayer::HandlePlayerEventTracksChanged()
 	{
 		for(int32 i=0; i<VideoStreamMetaData[0].StreamDetails.Num(); ++i)
 		{
-			UE_LOG(LogElectraPlayer, Verbose, TEXT("[%p][%p] Found %d * %d video stream at bitrate %d"), this, CurrentPlayer.Get(),
+			UE_LOG(LogElectraPlayer, Verbose, TEXT("[%u] Found %d * %d video stream at bitrate %d"), InstanceID,
 												VideoStreamMetaData[0].StreamDetails[i].CodecInformation.GetResolution().Width,
 												VideoStreamMetaData[0].StreamDetails[i].CodecInformation.GetResolution().Height,
 												VideoStreamMetaData[0].StreamDetails[i].Bandwidth);
@@ -2241,7 +2243,7 @@ void FElectraPlayer::HandlePlayerEventTracksChanged()
 	{
 		for(int32 i=0; i<AudioStreamMetaData[0].StreamDetails.Num(); ++i)
 		{
-			UE_LOG(LogElectraPlayer, Verbose, TEXT("[%p][%p] Found audio stream at bitrate %d"), this, CurrentPlayer.Get(),
+			UE_LOG(LogElectraPlayer, Verbose, TEXT("[%u] Found audio stream at bitrate %d"), InstanceID,
 												AudioStreamMetaData[0].StreamDetails[i].Bandwidth);
 		}
 	}
@@ -2290,13 +2292,13 @@ void FElectraPlayer::HandlePlayerEventLicenseKey(const Electra::Metrics::FLicens
 	// TBD
 	if (LicenseKeyStats.bWasSuccessful)
 	{
-		UE_LOG(LogElectraPlayer, Verbose, TEXT("[%p][%p] License key obtained"), this, CurrentPlayer.Get());
+		UE_LOG(LogElectraPlayer, Verbose, TEXT("[%u] License key obtained"), InstanceID);
 		FScopeLock Lock(&StatisticsLock);
 		Statistics.AddMessageToHistory(TEXT("Obtained license key"));
 	}
 	else
 	{
-		UE_LOG(LogElectraPlayer, Log, TEXT("[%p][%p] License key error \"%s\""), this, CurrentPlayer.Get(), *LicenseKeyStats.FailureReason);
+		UE_LOG(LogElectraPlayer, Log, TEXT("[%u] License key error \"%s\""), InstanceID, *LicenseKeyStats.FailureReason);
 		FScopeLock Lock(&StatisticsLock);
 		Statistics.AddMessageToHistory(TEXT("License key error"));
 	}
@@ -2380,7 +2382,7 @@ void FElectraPlayer::HandlePlayerEventBufferingStart(Electra::Metrics::EBufferin
 	FString Msg = FString::Printf(TEXT("%s buffering starts"), Electra::Metrics::GetBufferingReasonString(BufferingReason));
 	Statistics.AddMessageToHistory(Msg);
 
-	UE_LOG(LogElectraPlayer, Verbose, TEXT("[%p][%p] %s"), this, CurrentPlayer.Get(), *Msg);
+	UE_LOG(LogElectraPlayer, Verbose, TEXT("[%u] %s"), InstanceID, *Msg);
 	CSV_EVENT(ElectraPlayer, TEXT("Buffering starts"));
 }
 
@@ -2426,7 +2428,7 @@ void FElectraPlayer::HandlePlayerEventBufferingEnd(Electra::Metrics::EBufferingR
 		AnalyticEvent->ParamArray.Add(FAnalyticsEventAttribute(TEXT("Type"), Electra::Metrics::GetBufferingReasonString(BufferingReason)));
 		EnqueueAnalyticsEvent(AnalyticEvent);
 	}
-	UE_LOG(LogElectraPlayer, Verbose, TEXT("[%p][%p] %s buffering ended after %.3fs"), this, CurrentPlayer.Get(), Electra::Metrics::GetBufferingReasonString(BufferingReason), BufferingDuration);
+	UE_LOG(LogElectraPlayer, Verbose, TEXT("[%u] %s buffering ended after %.3fs"), InstanceID, Electra::Metrics::GetBufferingReasonString(BufferingReason), BufferingDuration);
 	Statistics.AddMessageToHistory(TEXT("Buffering ended"));
 // Should we set the state (back?) to something or wait for the following play/pause event to set a new one?
 	Statistics.LastState = "Ready";
@@ -2437,7 +2439,7 @@ void FElectraPlayer::HandlePlayerEventBufferingEnd(Electra::Metrics::EBufferingR
 void FElectraPlayer::HandlePlayerEventBandwidth(int64 EffectiveBps, int64 ThroughputBps, double LatencyInSeconds)
 {
 //	FScopeLock Lock(&StatisticsLock);
-	UE_LOG(LogElectraPlayer, VeryVerbose, TEXT("[%p][%p] Observed bandwidth of %lld Kbps; throughput = %lld Kbps; latency = %.3fs"), this, CurrentPlayer.Get(), EffectiveBps/1000, ThroughputBps/1000, LatencyInSeconds);
+	UE_LOG(LogElectraPlayer, VeryVerbose, TEXT("[%u] Observed bandwidth of %lld Kbps; throughput = %lld Kbps; latency = %.3fs"), InstanceID, EffectiveBps/1000, ThroughputBps/1000, LatencyInSeconds);
 }
 
 void FElectraPlayer::HandlePlayerEventBufferUtilization(const Electra::Metrics::FBufferStats& BufferStats)
@@ -2478,7 +2480,7 @@ void FElectraPlayer::HandlePlayerEventSegmentDownload(const Electra::Metrics::FS
 	}
 	if (SegmentDownloadStats.bWasSuccessful)
 	{
-		UE_LOG(LogElectraPlayer, VeryVerbose, TEXT("[%p][%p] Downloaded %s segment at bitrate %d: Playback time = %.3fs, duration = %.3fs, download time = %.3fs, URL=%s \"%s\""), this, CurrentPlayer.Get(), Electra::GetStreamTypeName(SegmentDownloadStats.StreamType), SegmentDownloadStats.Bitrate, SegmentDownloadStats.PresentationTime, SegmentDownloadStats.Duration, SegmentDownloadStats.TimeToDownload, *SegmentDownloadStats.Range, *SanitizeMessage(SegmentDownloadStats.URL));
+		UE_LOG(LogElectraPlayer, VeryVerbose, TEXT("[%u] Downloaded %s segment at bitrate %d: Playback time = %.3fs, duration = %.3fs, download time = %.3fs, URL=%s \"%s\""), InstanceID, Electra::GetStreamTypeName(SegmentDownloadStats.StreamType), SegmentDownloadStats.Bitrate, SegmentDownloadStats.PresentationTime, SegmentDownloadStats.Duration, SegmentDownloadStats.TimeToDownload, *SegmentDownloadStats.Range, *SanitizeMessage(SegmentDownloadStats.URL));
 	}
 	else if (SegmentDownloadStats.bWasAborted)
 	{
@@ -2486,7 +2488,7 @@ void FElectraPlayer::HandlePlayerEventSegmentDownload(const Electra::Metrics::FS
 	}
 	if (!SegmentDownloadStats.bWasSuccessful || SegmentDownloadStats.RetryNumber)
 	{
-		UE_LOG(LogElectraPlayer, Verbose, TEXT("[%p][%p] %s segment download issue (%s): retry:%d, success:%d, aborted:%d, filler:%d"), this, CurrentPlayer.Get(), Electra::Metrics::GetSegmentTypeString(SegmentDownloadStats.SegmentType), *SegmentDownloadStats.FailureReason, SegmentDownloadStats.RetryNumber, SegmentDownloadStats.bWasSuccessful, SegmentDownloadStats.bWasAborted, SegmentDownloadStats.bInsertedFillerData);
+		UE_LOG(LogElectraPlayer, Verbose, TEXT("[%u] %s segment download issue (%s): retry:%d, success:%d, aborted:%d, filler:%d"), InstanceID, Electra::Metrics::GetSegmentTypeString(SegmentDownloadStats.SegmentType), *SegmentDownloadStats.FailureReason, SegmentDownloadStats.RetryNumber, SegmentDownloadStats.bWasSuccessful, SegmentDownloadStats.bWasAborted, SegmentDownloadStats.bInsertedFillerData);
 
 		if (SegmentDownloadStats.FailureReason.Len())
 		{
@@ -2552,11 +2554,11 @@ void FElectraPlayer::HandlePlayerEventVideoQualityChange(int32 NewBitrate, int32
 	}
 	if (bIsDrasticDownswitch)
 	{
-		UE_LOG(LogElectraPlayer, Log, TEXT("[%p][%p] Player switched video quality drastically down to %d bps from %d bps. %d upswitches, %d downswitches (%d drastic ones)"), this, CurrentPlayer.Get(), NewBitrate, PreviousBitrate, Statistics.NumVideoQualityUpswitches, Statistics.NumVideoQualityDownswitches, Statistics.NumVideoQualityDrasticDownswitches);
+		UE_LOG(LogElectraPlayer, Log, TEXT("[%u] Player switched video quality drastically down to %d bps from %d bps. %d upswitches, %d downswitches (%d drastic ones)"), InstanceID, NewBitrate, PreviousBitrate, Statistics.NumVideoQualityUpswitches, Statistics.NumVideoQualityDownswitches, Statistics.NumVideoQualityDrasticDownswitches);
 	}
 	else
 	{
-		UE_LOG(LogElectraPlayer, Log, TEXT("[%p][%p] Player switched video quality to %d bps from %d bps. %d upswitches, %d downswitches (%d drastic ones)"), this, CurrentPlayer.Get(), NewBitrate, PreviousBitrate, Statistics.NumVideoQualityUpswitches, Statistics.NumVideoQualityDownswitches, Statistics.NumVideoQualityDrasticDownswitches);
+		UE_LOG(LogElectraPlayer, Log, TEXT("[%u] Player switched video quality to %d bps from %d bps. %d upswitches, %d downswitches (%d drastic ones)"), InstanceID, NewBitrate, PreviousBitrate, Statistics.NumVideoQualityUpswitches, Statistics.NumVideoQualityDownswitches, Statistics.NumVideoQualityDrasticDownswitches);
 	}
 
 	int32 prvWidth  = Statistics.CurrentlyActiveResolutionWidth;
@@ -2622,11 +2624,11 @@ void FElectraPlayer::HandlePlayerEventAudioQualityChange(int32 NewBitrate, int32
 	}
 	if (bIsDrasticDownswitch)
 	{
-		UE_LOG(LogElectraPlayer, Log, TEXT("[%p][%p] Player switched audio quality drastically down to %d bps from %d bps. %d upswitches, %d downswitches (%d drastic ones)"), this, CurrentPlayer.Get(), NewBitrate, PreviousBitrate, Statistics.NumAudioQualityUpswitches, Statistics.NumAudioQualityDownswitches, Statistics.NumAudioQualityDrasticDownswitches);
+		UE_LOG(LogElectraPlayer, Log, TEXT("[%u] Player switched audio quality drastically down to %d bps from %d bps. %d upswitches, %d downswitches (%d drastic ones)"), InstanceID, NewBitrate, PreviousBitrate, Statistics.NumAudioQualityUpswitches, Statistics.NumAudioQualityDownswitches, Statistics.NumAudioQualityDrasticDownswitches);
 	}
 	else
 	{
-		UE_LOG(LogElectraPlayer, Log, TEXT("[%p][%p] Player switched audio quality to %d bps from %d bps. %d upswitches, %d downswitches (%d drastic ones)"), this, CurrentPlayer.Get(), NewBitrate, PreviousBitrate, Statistics.NumAudioQualityUpswitches, Statistics.NumAudioQualityDownswitches, Statistics.NumAudioQualityDrasticDownswitches);
+		UE_LOG(LogElectraPlayer, Log, TEXT("[%u] Player switched audio quality to %d bps from %d bps. %d upswitches, %d downswitches (%d drastic ones)"), InstanceID, NewBitrate, PreviousBitrate, Statistics.NumAudioQualityUpswitches, Statistics.NumAudioQualityDownswitches, Statistics.NumAudioQualityDrasticDownswitches);
 	}
 
 	// Enqueue a "AudioQualityChange" event.
@@ -2668,7 +2670,7 @@ void FElectraPlayer::HandlePlayerEventPrerollStart()
 	// Update statistics
 	FScopeLock Lock(&StatisticsLock);
 	Statistics.TimeAtPrerollBegin = FPlatformTime::Seconds();
-	UE_LOG(LogElectraPlayer, Verbose, TEXT("[%p][%p] Player starts prerolling to warm decoders and renderers"), this, CurrentPlayer.Get());
+	UE_LOG(LogElectraPlayer, Verbose, TEXT("[%u] Player starts prerolling to warm decoders and renderers"), InstanceID);
 //	Statistics.LastState = "Ready";
 //	FString					CurrentState;	// "Empty", "Opening", "Preparing", "Buffering", "Idle", "Ready", "Playing", "Paused", "Seeking", "Rebuffering", "Ended", "Closed"
 
@@ -2693,7 +2695,7 @@ void FElectraPlayer::HandlePlayerEventPrerollEnd()
 	{
 		Statistics.TimeForInitialPreroll = FPlatformTime::Seconds() - Statistics.TimeAtPrerollBegin;
 	}
-	UE_LOG(LogElectraPlayer, Verbose, TEXT("[%p][%p] Player prerolling complete"), this, CurrentPlayer.Get());
+	UE_LOG(LogElectraPlayer, Verbose, TEXT("[%u] Player prerolling complete"), InstanceID);
 	Statistics.LastState = "Ready";
 
 	// Enqueue a "PrerollEnd" event.
@@ -2718,7 +2720,7 @@ void FElectraPlayer::HandlePlayerEventPlaybackStart()
 		Statistics.PlayPosAtStart = PlayPos;
 	}
 	Statistics.LastState = "Playing";
-	UE_LOG(LogElectraPlayer, Log, TEXT("[%p][%p] Playback started at play position %.3f"), this, CurrentPlayer.Get(), PlayPos);
+	UE_LOG(LogElectraPlayer, Log, TEXT("[%u] Playback started at play position %.3f"), InstanceID, PlayPos);
 	Statistics.AddMessageToHistory(TEXT("Playback started"));
 
 	// Enqueue a "Start" event.
@@ -2737,7 +2739,7 @@ void FElectraPlayer::HandlePlayerEventPlaybackPaused()
 	FScopeLock Lock(&StatisticsLock);
 	Statistics.LastState = "Paused";
 	double PlayPos = CurrentPlayer->AdaptivePlayer->GetPlayPosition().GetAsSeconds();
-	UE_LOG(LogElectraPlayer, Log, TEXT("[%p][%p] Playback paused at play position %.3f"), this, CurrentPlayer.Get(), PlayPos);
+	UE_LOG(LogElectraPlayer, Log, TEXT("[%u] Playback paused at play position %.3f"), InstanceID, PlayPos);
 	Statistics.AddMessageToHistory(TEXT("Playback paused"));
 
 	// Enqueue a "Pause" event.
@@ -2756,7 +2758,7 @@ void FElectraPlayer::HandlePlayerEventPlaybackResumed()
 	FScopeLock Lock(&StatisticsLock);
 	Statistics.LastState = "Playing";
 	double PlayPos = CurrentPlayer->AdaptivePlayer->GetPlayPosition().GetAsSeconds();
-	UE_LOG(LogElectraPlayer, Log, TEXT("[%p][%p] Playback resumed at play position %.3f"), this, CurrentPlayer.Get(), PlayPos);
+	UE_LOG(LogElectraPlayer, Log, TEXT("[%u] Playback resumed at play position %.3f"), InstanceID, PlayPos);
 	Statistics.AddMessageToHistory(TEXT("Playback resumed"));
 
 	// Enqueue a "Resume" event.
@@ -2779,7 +2781,7 @@ void FElectraPlayer::HandlePlayerEventPlaybackEnded()
 	FScopeLock Lock(&StatisticsLock);
 	Statistics.LastState = "Ended";
 	Statistics.bDidPlaybackEnd = true;
-	UE_LOG(LogElectraPlayer, Log, TEXT("[%p][%p] Playback reached end at play position %.3f"), this, CurrentPlayer.Get(), PlayPos);
+	UE_LOG(LogElectraPlayer, Log, TEXT("[%u] Playback reached end at play position %.3f"), InstanceID, PlayPos);
 	Statistics.AddMessageToHistory(TEXT("Playback ended"));
 
 	// Enqueue an "End" event.
@@ -2812,14 +2814,14 @@ void FElectraPlayer::HandlePlayerEventJumpInPlayPosition(const Electra::FTimeVal
 		{
 			++Statistics.NumTimesRewound;
 		}
-		UE_LOG(LogElectraPlayer, Log, TEXT("[%p][%p] Jump in play position from %.3f to %.3f"), this, CurrentPlayer.Get(), FromTime.GetAsSeconds(), ToNewTime.GetAsSeconds());
+		UE_LOG(LogElectraPlayer, Log, TEXT("[%u] Jump in play position from %.3f to %.3f"), InstanceID, FromTime.GetAsSeconds(), ToNewTime.GetAsSeconds());
 	}
 	else if (TimejumpReason == Electra::Metrics::ETimeJumpReason::Looping)
 	{
 		++Statistics.NumTimesLooped;
 		Electra::IAdaptiveStreamingPlayer::FLoopState loopState;
 		CurrentPlayer->AdaptivePlayer->GetLoopState(loopState);
-		UE_LOG(LogElectraPlayer, Log, TEXT("[%p][%p] Looping (%d) from %.3f to %.3f"), this, CurrentPlayer.Get(), loopState.Count, FromTime.GetAsSeconds(), ToNewTime.GetAsSeconds());
+		UE_LOG(LogElectraPlayer, Log, TEXT("[%u] Looping (%d) from %.3f to %.3f"), InstanceID, loopState.Count, FromTime.GetAsSeconds(), ToNewTime.GetAsSeconds());
 		Statistics.AddMessageToHistory(TEXT("Looped"));
 	}
 
@@ -2846,7 +2848,7 @@ void FElectraPlayer::HandlePlayerEventPlaybackStopped()
 	FScopeLock Lock(&StatisticsLock);
 	Statistics.bDidPlaybackEnd = true;
 	// Note: we do not change Statistics.LastState since we want to keep the state the player was in when it got closed.
-	UE_LOG(LogElectraPlayer, Log, TEXT("[%p][%p] Playback stopped. Last play position %.3f"), this, CurrentPlayer.Get(), PlayPos);
+	UE_LOG(LogElectraPlayer, Log, TEXT("[%u] Playback stopped. Last play position %.3f"), InstanceID, PlayPos);
 	Statistics.AddMessageToHistory(TEXT("Stopped"));
 
 	// Enqueue a "Stop" event.
@@ -2861,7 +2863,7 @@ void FElectraPlayer::HandlePlayerEventPlaybackStopped()
 
 void FElectraPlayer::HandlePlayerEventSeekCompleted()
 {
-	UE_LOG(LogElectraPlayer, Verbose, TEXT("[%p][%p] Seek completed"), this, CurrentPlayer.Get());
+	UE_LOG(LogElectraPlayer, Verbose, TEXT("[%u] Seek completed"), InstanceID);
 	bDiscardOutputUntilCleanStart = false;
 	MediaStateOnSeekFinished();
 }
@@ -2896,7 +2898,7 @@ void FElectraPlayer::HandlePlayerEventError(const FString& ErrorReason)
 		Statistics.LastError = ErrorReason;
 	}
 	// Note: we do not change Statistics.LastState to something like 'error' because we want to know the state the player was in when it errored.
-	UE_LOG(LogElectraPlayer, Error, TEXT("[%p][%p] ReportError: \"%s\""), this, CurrentPlayer.Get(), *SanitizeMessage(ErrorReason));
+	UE_LOG(LogElectraPlayer, Error, TEXT("[%u] ReportError: \"%s\""), InstanceID, *SanitizeMessage(ErrorReason));
 	Statistics.AddMessageToHistory(FString::Printf(TEXT("Error: %s"), *SanitizeMessage(ErrorReason)));
 	FString MessageHistory;
 	for(auto &msg : Statistics.MessageHistoryBuffer)
@@ -2924,26 +2926,26 @@ void FElectraPlayer::HandlePlayerEventLogMessage(Electra::IInfoLog::ELevel InLog
 	{
 		case Electra::IInfoLog::ELevel::Error:
 		{
-			UE_LOG(LogElectraPlayer, Error, TEXT("[%p][%p] %s"), this, CurrentPlayer.Get(), *m);
+			UE_LOG(LogElectraPlayer, Error, TEXT("[%u] %s"), InstanceID, *m);
 			FScopeLock Lock(&StatisticsLock);
 			Statistics.AddMessageToHistory(m);
 			break;
 		}
 		case Electra::IInfoLog::ELevel::Warning:
 		{
-			UE_LOG(LogElectraPlayer, Warning, TEXT("[%p][%p] %s"), this, CurrentPlayer.Get(), *m);
+			UE_LOG(LogElectraPlayer, Warning, TEXT("[%u] %s"), InstanceID, *m);
 			FScopeLock Lock(&StatisticsLock);
 			Statistics.AddMessageToHistory(m);
 			break;
 		}
 		case Electra::IInfoLog::ELevel::Info:
 		{
-			UE_LOG(LogElectraPlayer, Log, TEXT("[%p][%p] %s"), this, CurrentPlayer.Get(), *m);
+			UE_LOG(LogElectraPlayer, Log, TEXT("[%u] %s"), InstanceID, *m);
 			break;
 		}
 		case Electra::IInfoLog::ELevel::Verbose:
 		{
-			UE_LOG(LogElectraPlayer, Verbose, TEXT("[%p][%p] %s"), this, CurrentPlayer.Get(), *m);
+			UE_LOG(LogElectraPlayer, Verbose, TEXT("[%u] %s"), InstanceID, *m);
 			break;
 		}
 	}
@@ -2975,7 +2977,7 @@ void FElectraPlayer::FDroppedFrameStats::AddNewDrop(const FTimespan& InFrameTime
 					   FrameType == EFrameType::AudioFrame ? TEXT("audio") :
 					   TEXT("<unknown>");
 
-		UE_LOG(LogElectraPlayer, Log, TEXT("[%p][%p] dropped %s %s @ %s. %d drops so far. Worst delta %.3f ms"), InPtrElectraPlayer, InPtrCurrentPlayer, *Type, *InFrameTime.ToString(TEXT("%h:%m:%s.%f")), *InPlayerTime.ToString(TEXT("%h:%m:%s.%f")), NumTotalDropped, WorstDeltaTime.GetTotalMilliseconds());
+		//UE_LOG(LogElectraPlayer, Log, TEXT("[%u] dropped %s %s @ %s. %d drops so far. Worst delta %.3f ms"), InPtrElectraPlayer, *Type, *InFrameTime.ToString(TEXT("%h:%m:%s.%f")), *InPlayerTime.ToString(TEXT("%h:%m:%s.%f")), NumTotalDropped, WorstDeltaTime.GetTotalMilliseconds());
 		SystemTimeAtLastReport = Now;
 	}
 
@@ -3037,7 +3039,7 @@ void FElectraPlayer::LogStatistics()
 	if (PinnedAdapterDelegate.IsValid())
 	{
 		UE_LOG(LogElectraPlayer, Verbose, TEXT(
-			"[%p][%p] Electra player statistics:\n"\
+			"[%u] Electra player statistics:\n"\
 			"OS: %s\n"\
 			"GPU Adapter: %s\n"
 			"URL: %s\n"\
@@ -3076,7 +3078,7 @@ void FElectraPlayer::LogStatistics()
 			"Current state: %s\n" \
 			"Last issue: %s\n"
 		),
-			this, CurrentPlayer.Get(),
+			InstanceID,
 			*FString::Printf(TEXT("%s"), *AnalyticsOSVersion),
 			*PinnedAdapterDelegate->GetVideoAdapterName().TrimStartAndEnd(),
 			*SanitizeMessage(Statistics.InitialURL),
@@ -3144,12 +3146,12 @@ void FElectraPlayer::SendAnalyticMetrics(const TSharedPtr<IAnalyticsProviderET>&
 
 	if (!Statistics.bDidPlaybackEnd)
 	{
-		UE_LOG(LogElectraPlayer, Verbose, TEXT("[%p][%p] Submitting analytics during playback, some data may be incomplete"), this, CurrentPlayer.Get());
+		UE_LOG(LogElectraPlayer, Verbose, TEXT("[%u] Submitting analytics during playback, some data may be incomplete"), InstanceID);
 		// Try to fill in some of the blanks.
 		UpdatePlayEndStatistics();
 	}
 
-	UE_LOG(LogElectraPlayer, Verbose, TEXT("[%p][%p] Submitting analytics"), this, CurrentPlayer.Get());
+	UE_LOG(LogElectraPlayer, Verbose, TEXT("[%u] Submitting analytics"), InstanceID);
 
 
 	// First emit all enqueued events before sending the final one.
@@ -3398,7 +3400,7 @@ void FElectraPlayer::FAdaptiveStreamingPlayerResourceProvider::ProcessPendingSta
 					if (!PlaylistData.IsEmpty() && PlaylistData != InOutRequest->GetResourceURL())
 					{
 						/*
-						UE_LOG(LogElectraPlayer, Log, TEXT("[%p][%p] IMediaPlayer::Open: Source provided playlist data for '%s'"), this, CurrentPlayer.Get(), *SanitizeMessage(InOutRequest->GetResourceURL()));
+						UE_LOG(LogElectraPlayer, Log, TEXT("[%u] IMediaPlayer::Open: Source provided playlist data for '%s'"), InstanceID, *SanitizeMessage(InOutRequest->GetResourceURL()));
 						*/
 
 						// FString is Unicode but the HTTP response for a playlist is expected to be a UTF-8 string.
diff --git a/Engine/Plugins/Media/ElectraPlayer/Source/ElectraPlayerRuntime/Private/ElectraPlayer.h b/Engine/Plugins/Media/ElectraPlayer/Source/ElectraPlayerRuntime/Private/ElectraPlayer.h
--- a/Engine/Plugins/Media/ElectraPlayer/Source/ElectraPlayerRuntime/Private/ElectraPlayer.h
+++ b/Engine/Plugins/Media/ElectraPlayer/Source/ElectraPlayerRuntime/Private/ElectraPlayer.h
@@ -74,12 +74,12 @@ public:
 	bool CanPresentVideoFrames(uint64 NumFrames);
 	bool CanPresentAudioFrames(uint64 NumFrames);
 
-	FString GetUrl() const
+	FString GetUrl() const override
 	{
 		return MediaUrl;
 	}
 
-	void SetGuid(const FGuid& Guid)
+	void SetGuid(const FGuid& Guid) override
 	{
 		PlayerGuid = Guid;
 	}
@@ -500,7 +500,7 @@ private:
 		TSharedPtr<FElectraRendererAudio, ESPMode::ThreadSafe>		RendererAudio;
 
 		/** */
-		static void DoCloseAsync(TSharedPtr<FInternalPlayerImpl, ESPMode::ThreadSafe> && Player, TSharedPtr<IAsyncResourceReleaseNotifyContainer, ESPMode::ThreadSafe> AsyncDestructNotification);
+		static void DoCloseAsync(TSharedPtr<FInternalPlayerImpl, ESPMode::ThreadSafe>&& Player, uint32 PlayerID, TSharedPtr<IAsyncResourceReleaseNotifyContainer, ESPMode::ThreadSafe> AsyncDestructNotification);
 	};
 
 	mutable FCriticalSection										PlayerLock;
@@ -811,6 +811,10 @@ private:
 	/** Custom analytics constants. **/
 	FString									AnalyticsCustomValues[8];
 
+
+	/** Unique instance ID */
+	uint32 InstanceID = 0;
+
 	void HandleBlobDownload();
 
 	void HandleDeferredPlayerEvents();
