ElectraPlayer: Renamed master playlist to main playlist

#rnx

diff --git a/Engine/Plugins/Media/ElectraPlayer/Source/ElectraPlayerRuntime/Private/ElectraPlayer.cpp b/Engine/Plugins/Media/ElectraPlayer/Source/ElectraPlayerRuntime/Private/ElectraPlayer.cpp
--- a/Engine/Plugins/Media/ElectraPlayer/Source/ElectraPlayerRuntime/Private/ElectraPlayer.cpp
+++ b/Engine/Plugins/Media/ElectraPlayer/Source/ElectraPlayerRuntime/Private/ElectraPlayer.cpp
@@ -1911,10 +1911,10 @@ void FElectraPlayer::HandleDeferredPlayerEvents()
 				HandlePlayerEventOpenSource(Ev->URL);
 				break;
 			}
-			case FPlayerMetricEventBase::EType::ReceivedMasterPlaylist:
+			case FPlayerMetricEventBase::EType::ReceivedMainPlaylist:
 			{
-				FPlayerMetricEvent_ReceivedMasterPlaylist* Ev = static_cast<FPlayerMetricEvent_ReceivedMasterPlaylist*>(Event.Get());
-				HandlePlayerEventReceivedMasterPlaylist(Ev->EffectiveURL);
+				FPlayerMetricEvent_ReceivedMainPlaylist* Ev = static_cast<FPlayerMetricEvent_ReceivedMainPlaylist*>(Event.Get());
+				HandlePlayerEventReceivedMainPlaylist(Ev->EffectiveURL);
 				break;
 			}
 			case FPlayerMetricEventBase::EType::ReceivedPlaylists:
@@ -2106,15 +2106,15 @@ void FElectraPlayer::HandlePlayerEventOpenSource(const FString& URL)
 	}
 }
 
-void FElectraPlayer::HandlePlayerEventReceivedMasterPlaylist(const FString& EffectiveURL)
+void FElectraPlayer::HandlePlayerEventReceivedMainPlaylist(const FString& EffectiveURL)
 {
-	UE_LOG(LogElectraPlayer, Verbose, TEXT("[%u] Received master playlist from \"%s\""), InstanceID, *SanitizeMessage(EffectiveURL));
+	UE_LOG(LogElectraPlayer, Verbose, TEXT("[%u] Received main playlist from \"%s\""), InstanceID, *SanitizeMessage(EffectiveURL));
 
 	// Update statistics
 	FScopeLock Lock(&StatisticsLock);
-	Statistics.AddMessageToHistory(TEXT("Got master playlist"));
-	// Note the time it took to get the master playlist
-	Statistics.TimeToLoadMasterPlaylist = FPlatformTime::Seconds() - Statistics.TimeAtOpen;
+	Statistics.AddMessageToHistory(TEXT("Got main playlist"));
+	// Note the time it took to get the main playlist
+	Statistics.TimeToLoadMainPlaylist = FPlatformTime::Seconds() - Statistics.TimeAtOpen;
 	Statistics.LastState = "Preparing";
 
 
@@ -3043,7 +3043,7 @@ void FElectraPlayer::LogStatistics()
 			"OS: %s\n"\
 			"GPU Adapter: %s\n"
 			"URL: %s\n"\
-			"Time after master playlist loaded: %.3fs\n"\
+			"Time after main playlist loaded: %.3fs\n"\
 			"Time after stream playlists loaded: %.3fs\n"\
 			"Time for initial buffering: %.3fs\n"\
 			"Initial video stream bitrate: %d bps\n"\
@@ -3082,7 +3082,7 @@ void FElectraPlayer::LogStatistics()
 			*FString::Printf(TEXT("%s"), *AnalyticsOSVersion),
 			*PinnedAdapterDelegate->GetVideoAdapterName().TrimStartAndEnd(),
 			*SanitizeMessage(Statistics.InitialURL),
-			Statistics.TimeToLoadMasterPlaylist,
+			Statistics.TimeToLoadMainPlaylist,
 			Statistics.TimeToLoadStreamPlaylists,
 			Statistics.InitialBufferingDuration,
 			Statistics.InitialVideoStreamBitrate,
@@ -3173,7 +3173,7 @@ void FElectraPlayer::SendAnalyticMetrics(const TSharedPtr<IAnalyticsProviderET>&
 	ParamArray.Add(FAnalyticsEventAttribute(TEXT("MessageHistory"), MessageHistory));
 	ParamArray.Add(FAnalyticsEventAttribute(TEXT("LastError"), Statistics.LastError));
 	ParamArray.Add(FAnalyticsEventAttribute(TEXT("FinalVideoResolution"), FString::Printf(TEXT("%d*%d"), Statistics.CurrentlyActiveResolutionWidth, Statistics.CurrentlyActiveResolutionHeight)));
-	ParamArray.Add(FAnalyticsEventAttribute(TEXT("TimeElapsedToMasterPlaylist"), Statistics.TimeToLoadMasterPlaylist));
+	ParamArray.Add(FAnalyticsEventAttribute(TEXT("TimeElapsedToMasterPlaylist"), Statistics.TimeToLoadMainPlaylist));
 	ParamArray.Add(FAnalyticsEventAttribute(TEXT("TimeElapsedToPlaylists"), Statistics.TimeToLoadStreamPlaylists));
 	ParamArray.Add(FAnalyticsEventAttribute(TEXT("InitialAvgBufferingBandwidth"), Statistics.InitialBufferingBandwidth.GetAverageBandwidth()));
 	ParamArray.Add(FAnalyticsEventAttribute(TEXT("InitialAvgBufferingLatency"), Statistics.InitialBufferingBandwidth.GetAverageLatency()));
diff --git a/Engine/Plugins/Media/ElectraPlayer/Source/ElectraPlayerRuntime/Private/ElectraPlayer.h b/Engine/Plugins/Media/ElectraPlayer/Source/ElectraPlayerRuntime/Private/ElectraPlayer.h
--- a/Engine/Plugins/Media/ElectraPlayer/Source/ElectraPlayerRuntime/Private/ElectraPlayer.h
+++ b/Engine/Plugins/Media/ElectraPlayer/Source/ElectraPlayerRuntime/Private/ElectraPlayer.h
@@ -185,7 +185,7 @@ private:
 		enum class EType
 		{
 			OpenSource,
-			ReceivedMasterPlaylist,
+			ReceivedMainPlaylist,
 			ReceivedPlaylists,
 			TracksChanged,
 			PlaylistDownload,
@@ -224,9 +224,9 @@ private:
 		FPlayerMetricEvent_OpenSource(const FString& InURL) : FPlayerMetricEventBase(EType::OpenSource), URL(InURL) {}
 		FString URL;
 	};
-	struct FPlayerMetricEvent_ReceivedMasterPlaylist : public FPlayerMetricEventBase
+	struct FPlayerMetricEvent_ReceivedMainPlaylist : public FPlayerMetricEventBase
 	{
-		FPlayerMetricEvent_ReceivedMasterPlaylist(const FString& InEffectiveURL) : FPlayerMetricEventBase(EType::ReceivedMasterPlaylist), EffectiveURL(InEffectiveURL) {}
+		FPlayerMetricEvent_ReceivedMainPlaylist(const FString& InEffectiveURL) : FPlayerMetricEventBase(EType::ReceivedMainPlaylist), EffectiveURL(InEffectiveURL) {}
 		FString EffectiveURL;
 	};
 	struct FPlayerMetricEvent_PlaylistDownload : public FPlayerMetricEventBase
@@ -330,8 +330,8 @@ private:
 	// Methods from IAdaptiveStreamingPlayerMetrics
 	virtual void ReportOpenSource(const FString& URL) override
 	{ DeferredPlayerEvents.Enqueue(MakeSharedTS<FPlayerMetricEvent_OpenSource>(URL)); }
-	virtual void ReportReceivedMasterPlaylist(const FString& EffectiveURL) override
-	{ DeferredPlayerEvents.Enqueue(MakeSharedTS<FPlayerMetricEvent_ReceivedMasterPlaylist>(EffectiveURL)); }
+	virtual void ReportReceivedMainPlaylist(const FString& EffectiveURL) override
+	{ DeferredPlayerEvents.Enqueue(MakeSharedTS<FPlayerMetricEvent_ReceivedMainPlaylist>(EffectiveURL)); }
 	virtual void ReportReceivedPlaylists() override
 	{ DeferredPlayerEvents.Enqueue(MakeSharedTS<FPlayerMetricEventBase>(FPlayerMetricEventBase::EType::ReceivedPlaylists)); }
 	virtual void ReportTracksChanged() override
@@ -692,7 +692,7 @@ private:
 			LastError.Empty();
 			LastState = "Empty";
 			TimeAtOpen = -1.0;
-			TimeToLoadMasterPlaylist = -1.0;
+			TimeToLoadMainPlaylist = -1.0;
 			TimeToLoadStreamPlaylists = -1.0;
 			InitialBufferingDuration = -1.0;
 			InitialVideoStreamBitrate = 0;
@@ -740,7 +740,7 @@ private:
 		FString					LastError;
 		FString					LastState;		// "Empty", "Opening", "Preparing", "Buffering", "Idle", "Ready", "Playing", "Paused", "Seeking", "Rebuffering", "Ended"
 		double					TimeAtOpen;
-		double					TimeToLoadMasterPlaylist;
+		double					TimeToLoadMainPlaylist;
 		double					TimeToLoadStreamPlaylists;
 		double					InitialBufferingDuration;
 		int32					InitialVideoStreamBitrate;
@@ -819,7 +819,7 @@ private:
 
 	void HandleDeferredPlayerEvents();
 	void HandlePlayerEventOpenSource(const FString& URL);
-	void HandlePlayerEventReceivedMasterPlaylist(const FString& EffectiveURL);
+	void HandlePlayerEventReceivedMainPlaylist(const FString& EffectiveURL);
 	void HandlePlayerEventReceivedPlaylists();
 	void HandlePlayerEventTracksChanged();
 	void HandlePlayerEventPlaylistDownload(const Metrics::FPlaylistDownloadStats& PlaylistDownloadStats);
diff --git a/Engine/Plugins/Media/ElectraPlayer/Source/ElectraPlayerRuntime/Private/Runtime/Player/AdaptiveStreamingPlayer.cpp b/Engine/Plugins/Media/ElectraPlayer/Source/ElectraPlayerRuntime/Private/Runtime/Player/AdaptiveStreamingPlayer.cpp
--- a/Engine/Plugins/Media/ElectraPlayer/Source/ElectraPlayerRuntime/Private/Runtime/Player/AdaptiveStreamingPlayer.cpp
+++ b/Engine/Plugins/Media/ElectraPlayer/Source/ElectraPlayerRuntime/Private/Runtime/Player/AdaptiveStreamingPlayer.cpp
@@ -396,8 +396,8 @@ void FAdaptiveStreamingPlayer::FireSyncEvent(TSharedPtrTS<FMetricEvent> Event)
 			case FMetricEvent::EType::OpenSource:
 				Listeners[i]->ReportOpenSource(Event->Param.URL);
 				break;
-			case FMetricEvent::EType::ReceivedMasterPlaylist:
-				Listeners[i]->ReportReceivedMasterPlaylist(Event->Param.URL);
+			case FMetricEvent::EType::ReceivedMainPlaylist:
+				Listeners[i]->ReportReceivedMainPlaylist(Event->Param.URL);
 				break;
 			case FMetricEvent::EType::ReceivedPlaylists:
 				Listeners[i]->ReportReceivedPlaylists();
@@ -2063,7 +2063,7 @@ void FAdaptiveStreamingPlayer::HandleSessionMessage(TSharedPtrTS<IPlayerMessage>
 				{
 					if (pMsg->GetListType() == Playlist::EListType::Main && pMsg->GetLoadType() == Playlist::ELoadType::Initial)
 					{
-						DispatchEvent(FMetricEvent::ReportReceivedMasterPlaylist(pMsg->GetConnectionInfo().EffectiveURL));
+						DispatchEvent(FMetricEvent::ReportReceivedMainPlaylist(pMsg->GetConnectionInfo().EffectiveURL));
 					}
 					else if (pMsg->GetLoadType() == Playlist::ELoadType::Initial)
 					{
diff --git a/Engine/Plugins/Media/ElectraPlayer/Source/ElectraPlayerRuntime/Private/Runtime/Player/AdaptiveStreamingPlayerABR.cpp b/Engine/Plugins/Media/ElectraPlayer/Source/ElectraPlayerRuntime/Private/Runtime/Player/AdaptiveStreamingPlayerABR.cpp
--- a/Engine/Plugins/Media/ElectraPlayer/Source/ElectraPlayerRuntime/Private/Runtime/Player/AdaptiveStreamingPlayerABR.cpp
+++ b/Engine/Plugins/Media/ElectraPlayer/Source/ElectraPlayerRuntime/Private/Runtime/Player/AdaptiveStreamingPlayerABR.cpp
@@ -49,7 +49,7 @@ namespace Electra
 		void ReportBufferingStart(Metrics::EBufferingReason BufferingReason) override;
 		void ReportBufferingEnd(Metrics::EBufferingReason BufferingReason) override;
 		void ReportOpenSource(const FString& URL) override {}
-		void ReportReceivedMasterPlaylist(const FString& EffectiveURL) override {}
+		void ReportReceivedMainPlaylist(const FString& EffectiveURL) override {}
 		void ReportReceivedPlaylists() override {}
 		void ReportTracksChanged() override {}
 		void ReportCleanStart() override {};
diff --git a/Engine/Plugins/Media/ElectraPlayer/Source/ElectraPlayerRuntime/Private/Runtime/Player/AdaptiveStreamingPlayerInternal.h b/Engine/Plugins/Media/ElectraPlayer/Source/ElectraPlayerRuntime/Private/Runtime/Player/AdaptiveStreamingPlayerInternal.h
--- a/Engine/Plugins/Media/ElectraPlayer/Source/ElectraPlayerRuntime/Private/Runtime/Player/AdaptiveStreamingPlayerInternal.h
+++ b/Engine/Plugins/Media/ElectraPlayer/Source/ElectraPlayerRuntime/Private/Runtime/Player/AdaptiveStreamingPlayerInternal.h
@@ -638,7 +638,7 @@ struct FMetricEvent
 	enum class EType
 	{
 		OpenSource,
-		ReceivedMasterPlaylist,
+		ReceivedMainPlaylist,
 		ReceivedPlaylists,
 		TracksChanged,
 		CleanStart,
@@ -737,10 +737,10 @@ struct FMetricEvent
 		Evt->Param.URL = URL;
 		return Evt;
 	}
-	static TSharedPtrTS<FMetricEvent> ReportReceivedMasterPlaylist(const FString& EffectiveURL)
+	static TSharedPtrTS<FMetricEvent> ReportReceivedMainPlaylist(const FString& EffectiveURL)
 	{
 		TSharedPtrTS<FMetricEvent> Evt = MakeSharedTS<FMetricEvent>();
-		Evt->Type = EType::ReceivedMasterPlaylist;
+		Evt->Type = EType::ReceivedMainPlaylist;
 		Evt->Param.URL = EffectiveURL;
 		return Evt;
 	}
diff --git a/Engine/Plugins/Media/ElectraPlayer/Source/ElectraPlayerRuntime/Private/Runtime/Player/AdaptiveStreamingPlayerMetrics.h b/Engine/Plugins/Media/ElectraPlayer/Source/ElectraPlayerRuntime/Private/Runtime/Player/AdaptiveStreamingPlayerMetrics.h
--- a/Engine/Plugins/Media/ElectraPlayer/Source/ElectraPlayerRuntime/Private/Runtime/Player/AdaptiveStreamingPlayerMetrics.h
+++ b/Engine/Plugins/Media/ElectraPlayer/Source/ElectraPlayerRuntime/Private/Runtime/Player/AdaptiveStreamingPlayerMetrics.h
@@ -97,7 +97,7 @@ namespace Metrics
 	{
 		FPlaylistDownloadStats()
 		{
-			ListType	   = Playlist::EListType::Master;
+			ListType	   = Playlist::EListType::Main;
 			LoadType	   = Playlist::ELoadType::Initial;
 			HTTPStatusCode = 0;
 			RetryNumber    = 0;
@@ -261,9 +261,9 @@ public:
 	virtual void ReportOpenSource(const FString& URL) = 0;
 
 	/**
-	 * Called when the source's master playlist has beed loaded successfully.
+	 * Called when the source's main playlist has beed loaded successfully.
 	 */
-	virtual void ReportReceivedMasterPlaylist(const FString& EffectiveURL) = 0;
+	virtual void ReportReceivedMainPlaylist(const FString& EffectiveURL) = 0;
 
 	/**
 	 * Called when the dependent child playlists have been loaded successfully.
diff --git a/Engine/Plugins/Media/ElectraPlayer/Source/ElectraPlayerRuntime/Private/Runtime/Player/DASH/PlaylistReaderDASH.cpp b/Engine/Plugins/Media/ElectraPlayer/Source/ElectraPlayerRuntime/Private/Runtime/Player/DASH/PlaylistReaderDASH.cpp
--- a/Engine/Plugins/Media/ElectraPlayer/Source/ElectraPlayerRuntime/Private/Runtime/Player/DASH/PlaylistReaderDASH.cpp
+++ b/Engine/Plugins/Media/ElectraPlayer/Source/ElectraPlayerRuntime/Private/Runtime/Player/DASH/PlaylistReaderDASH.cpp
@@ -468,7 +468,7 @@ void FPlaylistReaderDASH::CheckForMPDUpdate()
 						bRequestNow = false;
 						UpdateRequested.Empty();
 						// Send a fake update message to the player to have it evaluate the media timeline again.
-						PlayerSessionServices->SendMessageToPlayer(IPlaylistReader::PlaylistLoadedMessage::Create(FErrorDetail(), nullptr, Playlist::EListType::Master, Playlist::ELoadType::Update));
+						PlayerSessionServices->SendMessageToPlayer(IPlaylistReader::PlaylistLoadedMessage::Create(FErrorDetail(), nullptr, Playlist::EListType::Main, Playlist::ELoadType::Update));
 					}
 				}
 				else
@@ -538,7 +538,7 @@ void FPlaylistReaderDASH::CheckForMPDUpdate()
 				}
 #endif
 				PlaylistLoadRequest->CompleteCallback.BindThreadSafeSP(AsShared(), &FPlaylistReaderDASH::ManifestUpdateDownloadCompleted);
-				
+
 				bIsMPDUpdateInProgress = true;
 				EnqueueResourceRequest(MoveTemp(PlaylistLoadRequest));
 			}
@@ -640,11 +640,11 @@ void FPlaylistReaderDASH::SetStreamInbandEventUsage(EStreamType InStreamType, bo
 {
 	switch(InStreamType)
 	{
-		case EStreamType::Video: 
-			bInbandEventByStreamType[0] = bUsesInbandDASHEvents; 
+		case EStreamType::Video:
+			bInbandEventByStreamType[0] = bUsesInbandDASHEvents;
 			break;
-		case EStreamType::Audio: 
-			bInbandEventByStreamType[1] = bUsesInbandDASHEvents; 
+		case EStreamType::Audio:
+			bInbandEventByStreamType[1] = bUsesInbandDASHEvents;
 			break;
 		default:
 			break;
@@ -808,7 +808,7 @@ void FPlaylistReaderDASH::InternalSetup()
 	// Get the minimum MPD update time limit.
 	MinTimeBetweenUpdates = PlayerSessionServices->GetOptionValue(DASH::OptionKey_MinTimeBetweenMPDUpdates).SafeGetTimeValue(FTimeValue().SetFromMilliseconds(1000));
 
-	// Setup the playlist load request for the master playlist.
+	// Setup the playlist load request for the main playlist.
 	TSharedPtrTS<FMPDLoadRequestDASH> PlaylistLoadRequest = MakeSharedTS<FMPDLoadRequestDASH>();
 	PlaylistLoadRequest->LoadType = FMPDLoadRequestDASH::ELoadType::MPD;
 	FURL_RFC3986 UrlParser;
@@ -1141,7 +1141,7 @@ void FPlaylistReaderDASH::ManifestDownloadCompleted(FResourceLoadRequestPtr Requ
 		{
 			return;
 		}
-		
+
 		const HTTP::FConnectionInfo* ConnInfo = Request->GetConnectionInfo();
 		if (!ConnInfo)
 		{
@@ -1165,7 +1165,7 @@ void FPlaylistReaderDASH::ManifestDownloadCompleted(FResourceLoadRequestPtr Requ
 			}
 		}
 
-		PlayerSessionServices->SendMessageToPlayer(IPlaylistReader::PlaylistDownloadMessage::Create(ConnInfo, Playlist::EListType::Master, Playlist::ELoadType::Initial));
+		PlayerSessionServices->SendMessageToPlayer(IPlaylistReader::PlaylistDownloadMessage::Create(ConnInfo, Playlist::EListType::Main, Playlist::ELoadType::Initial));
 
 		FString XML;
 		LastErrorDetail = GetXMLResponseString(XML, Request);
@@ -1316,8 +1316,8 @@ void FPlaylistReaderDASH::FinishManifestBuildingAfterTimesync(bool bGotTheTime)
 		RequestsLock.Unlock();
 	}
 
-	// Notify that the "master playlist" has been parsed, successfully or not. If we still need to resolve remote entities this is not an error!
-	PlayerSessionServices->SendMessageToPlayer(IPlaylistReader::PlaylistLoadedMessage::Create(!LastErrorDetail.IsTryAgain() ? LastErrorDetail : FErrorDetail(), InitialTimesync.ConnInfo.Get(), Playlist::EListType::Master, Playlist::ELoadType::Initial));
+	// Notify that the "main playlist" has been parsed, successfully or not. If we still need to resolve remote entities this is not an error!
+	PlayerSessionServices->SendMessageToPlayer(IPlaylistReader::PlaylistLoadedMessage::Create(!LastErrorDetail.IsTryAgain() ? LastErrorDetail : FErrorDetail(), InitialTimesync.ConnInfo.Get(), Playlist::EListType::Main, Playlist::ELoadType::Initial));
 	if (LastErrorDetail.IsOK())
 	{
 		// Notify that the "variant playlists" are ready. There are no variants in DASH, but this is the trigger that the playlists are all set up and are good to go now.
@@ -1349,7 +1349,7 @@ void FPlaylistReaderDASH::ManifestUpdateDownloadCompleted(FResourceLoadRequestPt
 		{
 			return;
 		}
-		
+
 		const HTTP::FConnectionInfo* ConnInfo = Request->GetConnectionInfo();
 		if (!ConnInfo)
 		{
@@ -1395,7 +1395,7 @@ void FPlaylistReaderDASH::ManifestUpdateDownloadCompleted(FResourceLoadRequestPt
 		}
 		else
 		{
-			PlayerSessionServices->SendMessageToPlayer(IPlaylistReader::PlaylistDownloadMessage::Create(ConnInfo, Playlist::EListType::Master, Playlist::ELoadType::Update));
+			PlayerSessionServices->SendMessageToPlayer(IPlaylistReader::PlaylistDownloadMessage::Create(ConnInfo, Playlist::EListType::Main, Playlist::ELoadType::Update));
 
 			FString XML;
 			LastErrorDetail = GetXMLResponseString(XML, Request);
@@ -1441,7 +1441,7 @@ void FPlaylistReaderDASH::ManifestUpdateDownloadCompleted(FResourceLoadRequestPt
 						RequestsLock.Unlock();
 					}
 
-					PlayerSessionServices->SendMessageToPlayer(IPlaylistReader::PlaylistLoadedMessage::Create(!LastErrorDetail.IsTryAgain() ? LastErrorDetail : FErrorDetail(), ConnInfo, Playlist::EListType::Master, Playlist::ELoadType::Update));
+					PlayerSessionServices->SendMessageToPlayer(IPlaylistReader::PlaylistLoadedMessage::Create(!LastErrorDetail.IsTryAgain() ? LastErrorDetail : FErrorDetail(), ConnInfo, Playlist::EListType::Main, Playlist::ELoadType::Update));
 					if (LastErrorDetail.IsOK())
 					{
 					}
diff --git a/Engine/Plugins/Media/ElectraPlayer/Source/ElectraPlayerRuntime/Private/Runtime/Player/Playlist.h b/Engine/Plugins/Media/ElectraPlayer/Source/ElectraPlayerRuntime/Private/Runtime/Player/Playlist.h
--- a/Engine/Plugins/Media/ElectraPlayer/Source/ElectraPlayerRuntime/Private/Runtime/Player/Playlist.h
+++ b/Engine/Plugins/Media/ElectraPlayer/Source/ElectraPlayerRuntime/Private/Runtime/Player/Playlist.h
@@ -11,8 +11,7 @@ namespace Electra
 		enum class EListType
 		{
 			Main,
-			Variant,
-			Master = Main
+			Variant
 		};
 
 		static const TCHAR* const GetPlaylistTypeString(EListType InListType)
diff --git a/Engine/Plugins/Media/ElectraPlayer/Source/ElectraPlayerRuntime/Private/Runtime/Player/SimpleElectraAudioPlayer.cpp b/Engine/Plugins/Media/ElectraPlayer/Source/ElectraPlayerRuntime/Private/Runtime/Player/SimpleElectraAudioPlayer.cpp
--- a/Engine/Plugins/Media/ElectraPlayer/Source/ElectraPlayerRuntime/Private/Runtime/Player/SimpleElectraAudioPlayer.cpp
+++ b/Engine/Plugins/Media/ElectraPlayer/Source/ElectraPlayerRuntime/Private/Runtime/Player/SimpleElectraAudioPlayer.cpp
@@ -351,7 +351,7 @@ public:
 			UE_LOG(LogSimpleElectraPlayer, Verbose, TEXT("%s: Opening"), *AssetName);
 		}
 	}
-	void ReportReceivedMasterPlaylist(const FString& InEffectiveURL) override
+	void ReportReceivedMainPlaylist(const FString& InEffectiveURL) override
 	{ }
 	void ReportReceivedPlaylists() override
 	{
diff --git a/Engine/Plugins/Media/ElectraPlayer/Source/ElectraPlayerRuntime/Private/Runtime/Player/mkv/PlaylistReaderMKV.cpp b/Engine/Plugins/Media/ElectraPlayer/Source/ElectraPlayerRuntime/Private/Runtime/Player/mkv/PlaylistReaderMKV.cpp
--- a/Engine/Plugins/Media/ElectraPlayer/Source/ElectraPlayerRuntime/Private/Runtime/Player/mkv/PlaylistReaderMKV.cpp
+++ b/Engine/Plugins/Media/ElectraPlayer/Source/ElectraPlayerRuntime/Private/Runtime/Player/mkv/PlaylistReaderMKV.cpp
@@ -338,10 +338,10 @@ void FPlaylistReaderMKV::WorkerThread()
 
 	if (!HasReadBeenAborted())
 	{
-		// Notify the download of the "master playlist". This indicates the download only, not the parsing thereof.
-		PlayerSessionServices->SendMessageToPlayer(IPlaylistReader::PlaylistDownloadMessage::Create(&ConnectionInfo, Playlist::EListType::Master, Playlist::ELoadType::Initial));
-		// Notify that the "master playlist" has been parsed, successfully or not.
-		PlayerSessionServices->SendMessageToPlayer(IPlaylistReader::PlaylistLoadedMessage::Create(LastErrorDetail, &ConnectionInfo, Playlist::EListType::Master, Playlist::ELoadType::Initial));
+		// Notify the download of the "main playlist". This indicates the download only, not the parsing thereof.
+		PlayerSessionServices->SendMessageToPlayer(IPlaylistReader::PlaylistDownloadMessage::Create(&ConnectionInfo, Playlist::EListType::Main, Playlist::ELoadType::Initial));
+		// Notify that the "main playlist" has been parsed, successfully or not.
+		PlayerSessionServices->SendMessageToPlayer(IPlaylistReader::PlaylistLoadedMessage::Create(LastErrorDetail, &ConnectionInfo, Playlist::EListType::Main, Playlist::ELoadType::Initial));
 
 		if (LastErrorDetail.IsOK())
 		{
diff --git a/Engine/Plugins/Media/ElectraPlayer/Source/ElectraPlayerRuntime/Private/Runtime/Player/mp4/PlaylistReaderMP4.cpp b/Engine/Plugins/Media/ElectraPlayer/Source/ElectraPlayerRuntime/Private/Runtime/Player/mp4/PlaylistReaderMP4.cpp
--- a/Engine/Plugins/Media/ElectraPlayer/Source/ElectraPlayerRuntime/Private/Runtime/Player/mp4/PlaylistReaderMP4.cpp
+++ b/Engine/Plugins/Media/ElectraPlayer/Source/ElectraPlayerRuntime/Private/Runtime/Player/mp4/PlaylistReaderMP4.cpp
@@ -94,7 +94,7 @@ private:
 	void LogMessage(IInfoLog::ELevel Level, const FString& Message);
 
 	IPlayerSessionServices* PlayerSessionServices = nullptr;
-	FString MasterPlaylistURL;
+	FString MainPlaylistURL;
 	FString URLFragment;
 	FMediaEvent WorkerThreadQuitSignal;
 	bool bIsWorkerThreadStarted = false;
@@ -144,7 +144,7 @@ FPlaylistReaderMP4::~FPlaylistReaderMP4()
 
 FString FPlaylistReaderMP4::GetURL() const
 {
-	return MasterPlaylistURL;
+	return MainPlaylistURL;
 }
 
 TSharedPtrTS<IManifest> FPlaylistReaderMP4::GetManifest()
@@ -213,7 +213,7 @@ void FPlaylistReaderMP4::LoadAndParse(const FString& URL)
 {
 	FURL_RFC3986 UrlParser;
 	UrlParser.Parse(URL);
-	MasterPlaylistURL = UrlParser.Get(true, false);
+	MainPlaylistURL = UrlParser.Get(true, false);
 	URLFragment = UrlParser.GetFragment();
 
 	StartWorkerThread();
@@ -232,7 +232,7 @@ void FPlaylistReaderMP4::WorkerThread()
 	const TArray<uint32> FirstBoxes { UtilsMP4::Make4CC('f','t','y','p'), UtilsMP4::Make4CC('s','t','y','p'), UtilsMP4::Make4CC('s','i','d','x'), UtilsMP4::Make4CC('f','r','e','e'), UtilsMP4::Make4CC('s','k','i','p') };
 	UtilsMP4::FMP4RootBoxLocator BoxLocator;
 	TArray<UtilsMP4::FMP4RootBoxLocator::FBoxInfo> BoxInfos;
-	bool bGotBoxes = BoxLocator.LocateRootBoxes(BoxInfos, PlayerSessionServices->GetHTTPManager(), MasterPlaylistURL, FirstBoxes, StopAfterBoxes, ReadBoxes, UtilsMP4::FMP4RootBoxLocator::FCancellationCheckDelegate::CreateLambda([&]()
+	bool bGotBoxes = BoxLocator.LocateRootBoxes(BoxInfos, PlayerSessionServices->GetHTTPManager(), MainPlaylistURL, FirstBoxes, StopAfterBoxes, ReadBoxes, UtilsMP4::FMP4RootBoxLocator::FCancellationCheckDelegate::CreateLambda([&]()
 	{
 		return bAbort;
 	}));
@@ -243,10 +243,10 @@ void FPlaylistReaderMP4::WorkerThread()
 		PostError(BoxLocator.GetErrorMessage(), ERRCODE_MP4_INVALID_FILE, UEMEDIA_ERROR_FORMAT_ERROR);
 	}
 
-	// Notify the download of the "master playlist". This indicates the download only, not the parsing thereof.
-	PlayerSessionServices->SendMessageToPlayer(IPlaylistReader::PlaylistDownloadMessage::Create(&ConnectionInfo, Playlist::EListType::Master, Playlist::ELoadType::Initial));
-	// Notify that the "master playlist" has been parsed, successfully or not.
-	PlayerSessionServices->SendMessageToPlayer(IPlaylistReader::PlaylistLoadedMessage::Create(LastErrorDetail, &ConnectionInfo, Playlist::EListType::Master, Playlist::ELoadType::Initial));
+	// Notify the download of the "main playlist". This indicates the download only, not the parsing thereof.
+	PlayerSessionServices->SendMessageToPlayer(IPlaylistReader::PlaylistDownloadMessage::Create(&ConnectionInfo, Playlist::EListType::Main, Playlist::ELoadType::Initial));
+	// Notify that the "main playlist" has been parsed, successfully or not.
+	PlayerSessionServices->SendMessageToPlayer(IPlaylistReader::PlaylistLoadedMessage::Create(LastErrorDetail, &ConnectionInfo, Playlist::EListType::Main, Playlist::ELoadType::Initial));
 	// Failed to get the boxes but was not aborted?
 	if (!bAbort && (!bGotBoxes || bHasErrored))
 	{
@@ -290,7 +290,7 @@ void FPlaylistReaderMP4::WorkerThread()
 					FURL_RFC3986::GetQueryParams(URLFragmentComponents, URLFragment, false);	// The fragment is already URL escaped, so no need to do it again.
 					Manifest->SetURLFragmentComponents(MoveTemp(URLFragmentComponents));
 
-					LastErrorDetail = Manifest->Build(MP4Parser, MasterPlaylistURL, ConnectionInfo);
+					LastErrorDetail = Manifest->Build(MP4Parser, MainPlaylistURL, ConnectionInfo);
 
 					// Notify that the "variant playlists" are ready. There are no variants in an mp4, but this is the trigger that the playlists are all set up and are good to go now.
 					PlayerSessionServices->SendMessageToPlayer(IPlaylistReader::PlaylistLoadedMessage::Create(LastErrorDetail, &ConnectionInfo, Playlist::EListType::Variant, Playlist::ELoadType::Initial));
diff --git a/Engine/Plugins/Media/ElectraPlayer/Source/ElectraPlayerRuntime/Private/Runtime/Player/mpegaudio/PlaylistReaderMPEGAudio.cpp b/Engine/Plugins/Media/ElectraPlayer/Source/ElectraPlayerRuntime/Private/Runtime/Player/mpegaudio/PlaylistReaderMPEGAudio.cpp
--- a/Engine/Plugins/Media/ElectraPlayer/Source/ElectraPlayerRuntime/Private/Runtime/Player/mpegaudio/PlaylistReaderMPEGAudio.cpp
+++ b/Engine/Plugins/Media/ElectraPlayer/Source/ElectraPlayerRuntime/Private/Runtime/Player/mpegaudio/PlaylistReaderMPEGAudio.cpp
@@ -914,10 +914,10 @@ void FPlaylistReaderMPEGAudio::WorkerThread()
 	ConnectionInfo.StatusInfo.HTTPStatus = Req->StatusCode;
 	ConnectionInfo.StatusInfo.bReadError = Req->Result == FRequest::EResult::Failed;
 
-	// Notify the download of the "master playlist". This indicates the download only, not the parsing thereof.
-	PlayerSessionServices->SendMessageToPlayer(IPlaylistReader::PlaylistDownloadMessage::Create(&ConnectionInfo, Playlist::EListType::Master, Playlist::ELoadType::Initial));
-	// Notify that the "master playlist" has been parsed, successfully or not.
-	PlayerSessionServices->SendMessageToPlayer(IPlaylistReader::PlaylistLoadedMessage::Create(LastErrorDetail, &ConnectionInfo, Playlist::EListType::Master, Playlist::ELoadType::Initial));
+	// Notify the download of the "main playlist". This indicates the download only, not the parsing thereof.
+	PlayerSessionServices->SendMessageToPlayer(IPlaylistReader::PlaylistDownloadMessage::Create(&ConnectionInfo, Playlist::EListType::Main, Playlist::ELoadType::Initial));
+	// Notify that the "main playlist" has been parsed, successfully or not.
+	PlayerSessionServices->SendMessageToPlayer(IPlaylistReader::PlaylistLoadedMessage::Create(LastErrorDetail, &ConnectionInfo, Playlist::EListType::Main, Playlist::ELoadType::Initial));
 	if (!bAbort)
 	{
 		// Error?
