ElectraPlayer: Changed DASH playback over to use the common media segment reader

#rnx

diff --git a/Engine/Plugins/Media/ElectraPlayer/Source/ElectraPlayerRuntime/Private/Runtime/Player/DASH/ManifestDASH.cpp b/Engine/Plugins/Media/ElectraPlayer/Source/ElectraPlayerRuntime/Private/Runtime/Player/DASH/ManifestDASH.cpp
--- a/Engine/Plugins/Media/ElectraPlayer/Source/ElectraPlayerRuntime/Private/Runtime/Player/DASH/ManifestDASH.cpp
+++ b/Engine/Plugins/Media/ElectraPlayer/Source/ElectraPlayerRuntime/Private/Runtime/Player/DASH/ManifestDASH.cpp
@@ -6,7 +6,7 @@
 #include "ManifestBuilderDASH.h"
 #include "PlaylistReaderDASH.h"
 #include "Stats/Stats.h"
-#include "StreamReaderDASH.h"
+#include "Player/StreamSegmentReaderCommon.h"
 #include "Demuxer/ParserISO14496-12.h"
 #include "Demuxer/ParserISO14496-12_Utils.h"
 #include "Demuxer/ParserMKV.h"
@@ -226,7 +226,7 @@ private:
 	{
 		FString AdaptationSetID;
 		FString RepresentationSetID;
-		FManifestDASHInternal::FSegmentInformation InitSegmentInfo;
+		FSegmentInformationCommon InitSegmentInfo;
 		bool bRequested = false;
 		TSharedPtrTS<FMPDLoadRequestDASH> LoadRequest;
 
@@ -263,7 +263,7 @@ private:
 
 	bool PrepareDRM(const TArray<FManifestDASHInternal::FAdaptationSet::FContentProtection>& InContentProtections);
 
-	void SetupCommonSegmentRequestInfos(TSharedPtrTS<FStreamSegmentRequestDASH>& InOutSegmentRequest);
+	void SetupCommonSegmentRequestInfos(TSharedPtrTS<FStreamSegmentRequestCommon>& InOutSegmentRequest);
 
 	void PrioritizeSelection(TArray<FPrioritizedSelection>& Selection, EStreamType StreamType, bool bAdaptationSetLevel, bool bSortByBitrateDescending);
 
@@ -545,7 +545,7 @@ void FManifestDASH::TriggerPlaylistRefresh()
 
 IStreamReader* FManifestDASH::CreateStreamReaderHandler()
 {
-	return new FStreamReaderDASH;
+	return new FStreamSegmentReaderCommon;
 }
 
 IManifest::FResult FManifestDASH::FindPlayPeriod(TSharedPtrTS<IPlayPeriod>& OutPlayPeriod, const FPlayStartPosition& StartPosition, ESearchType SearchType)
@@ -708,7 +708,7 @@ IManifest::FResult FManifestDASH::FindPlayPeriod(TSharedPtrTS<IPlayPeriod>& OutP
 
 IManifest::FResult FManifestDASH::FindNextPlayPeriod(TSharedPtrTS<IPlayPeriod>& OutPlayPeriod, TSharedPtrTS<const IStreamSegment> InCurrentSegment)
 {
-	const FStreamSegmentRequestDASH* CurrentRequest = static_cast<const FStreamSegmentRequestDASH*>(InCurrentSegment.Get());
+	const FStreamSegmentRequestCommon* CurrentRequest = static_cast<const FStreamSegmentRequestCommon*>(InCurrentSegment.Get());
 	if (CurrentRequest)
 	{
 		FPlayStartPosition SearchTime;
@@ -1250,18 +1250,23 @@ void FDASHPlayPeriod::TriggerInitSegmentPreload(const TArray<FInitSegmentPreload
 }
 
 
-void FDASHPlayPeriod::SetupCommonSegmentRequestInfos(TSharedPtrTS<FStreamSegmentRequestDASH>& InOutSegmentRequest)
+void FDASHPlayPeriod::SetupCommonSegmentRequestInfos(TSharedPtrTS<FStreamSegmentRequestCommon>& InOutSegmentRequest)
 {
 	FManifestDASHInternal::FRepresentation* Repr = static_cast<FManifestDASHInternal::FRepresentation*>(InOutSegmentRequest->Representation.Get());
 
 	// Source buffer info
-	InOutSegmentRequest->SourceBufferInfo = InOutSegmentRequest->StreamType == EStreamType::Video ? SourceBufferInfoVideo :
-											InOutSegmentRequest->StreamType == EStreamType::Audio ? SourceBufferInfoAudio :
-											InOutSegmentRequest->StreamType == EStreamType::Subtitle ? SourceBufferInfoSubtitles : nullptr;
+	InOutSegmentRequest->SourceBufferInfo[StreamTypeToArrayIndex(InOutSegmentRequest->StreamType)] =
+			InOutSegmentRequest->StreamType == EStreamType::Video ? SourceBufferInfoVideo :
+			InOutSegmentRequest->StreamType == EStreamType::Audio ? SourceBufferInfoAudio :
+			InOutSegmentRequest->StreamType == EStreamType::Subtitle ? SourceBufferInfoSubtitles : nullptr;
 
 	// Encryption stuff
-	InOutSegmentRequest->DrmClient = DrmClient;
-	InOutSegmentRequest->DrmMimeType = Repr->GetCodecInformation().GetMimeTypeWithCodecAndFeatures();
+	InOutSegmentRequest->DrmMedia.DrmClient = DrmClient;
+	InOutSegmentRequest->DrmMedia.DrmMimeType = Repr->GetCodecInformation().GetMimeTypeWithCodecAndFeatures();
+
+	// Protocol and container
+	InOutSegmentRequest->StreamingProtocol = FStreamSegmentRequestCommon::EStreamingProtocol::DASH;
+	InOutSegmentRequest->ExpectedContainerFormat = InOutSegmentRequest->Segment.ContainerType == FSegmentInformationCommon::EContainerType::ISO14496_12 ? FStreamSegmentRequestCommon::EContainerFormat::ISO_14496_12 : FStreamSegmentRequestCommon::EContainerFormat::Matroska_WebM;
 }
 
 
@@ -1316,7 +1321,7 @@ IManifest::FResult FDASHPlayPeriod::GetStartingSegment(TSharedPtrTS<IStreamSegme
 
 	// Create a segment request to which the individual stream segment requests will add themselves as
 	// dependent streams. This is a special case for playback start.
-	TSharedPtrTS<FStreamSegmentRequestDASH> StartSegmentRequest = MakeSharedTS<FStreamSegmentRequestDASH>();
+	TSharedPtrTS<FStreamSegmentRequestCommon> StartSegmentRequest = MakeSharedTS<FStreamSegmentRequestCommon>();
 	StartSegmentRequest->bIsInitialStartRequest = true;
 	StartSegmentRequest->TimestampSequenceIndex = InSequenceState.GetSequenceIndex();
 
@@ -1358,7 +1363,7 @@ IManifest::FResult FDASHPlayPeriod::GetStartingSegment(TSharedPtrTS<IStreamSegme
 				return IManifest::FResult(IManifest::FResult::EType::NotFound).SetErrorDetail(FErrorDetail().SetMessage("Period no longer has the selected AdaptationSet or Representation."));
 			}
 
-			FManifestDASHInternal::FSegmentInformation SegmentInfo;
+			FSegmentInformationCommon SegmentInfo;
 			FManifestDASHInternal::FSegmentSearchOption SearchOpt;
 			TArray<TWeakPtrTS<FMPDLoadRequestDASH>> RemoteElementLoadRequests;
 			SearchOpt.PeriodLocalTime = StartTime;
@@ -1388,9 +1393,9 @@ IManifest::FResult FDASHPlayPeriod::GetStartingSegment(TSharedPtrTS<IStreamSegme
 			}
 			else if (SearchResult == FManifestDASHInternal::FRepresentation::ESearchResult::PastEOS)
 			{
-				TSharedPtrTS<FStreamSegmentRequestDASH> SegmentRequest = MakeSharedTS<FStreamSegmentRequestDASH>();
+				TSharedPtrTS<FStreamSegmentRequestCommon> SegmentRequest = MakeSharedTS<FStreamSegmentRequestCommon>();
 				SegmentRequest->StreamType = ActiveSelection[i].StreamType;
-				SegmentRequest->CodecInfo = Repr->GetCodecInformation();
+				SegmentRequest->CodecInfo[StreamTypeToArrayIndex(SegmentRequest->StreamType)] = Repr->GetCodecInformation();
 				SegmentRequest->Representation = Repr;
 				SegmentRequest->AdaptationSet = Adapt;
 				SegmentRequest->Period = Period;
@@ -1443,11 +1448,11 @@ IManifest::FResult FDASHPlayPeriod::GetStartingSegment(TSharedPtrTS<IStreamSegme
 					}
 				}
 
-				TSharedPtrTS<FStreamSegmentRequestDASH> SegmentRequest = MakeSharedTS<FStreamSegmentRequestDASH>();
+				TSharedPtrTS<FStreamSegmentRequestCommon> SegmentRequest = MakeSharedTS<FStreamSegmentRequestCommon>();
 				SegmentRequest->StreamType = ActiveSelection[i].StreamType;
 				SegmentRequest->QualityIndex = ActiveSelection[i].QualityIndex.Index;
 				SegmentRequest->MaxQualityIndex = ActiveSelection[i].QualityIndex.MaxIndex;
-				SegmentRequest->CodecInfo = Repr->GetCodecInformation();
+				SegmentRequest->CodecInfo[StreamTypeToArrayIndex(SegmentRequest->StreamType)] = Repr->GetCodecInformation();
 				SegmentRequest->Representation = Repr;
 				SegmentRequest->AdaptationSet = Adapt;
 				SegmentRequest->Period = Period;
@@ -1526,7 +1531,7 @@ IManifest::FResult FDASHPlayPeriod::GetContinuationSegment(TSharedPtrTS<IStreamS
 	// Create a dummy request we can use to pass into GetNextOrRetrySegment().
 	// Only set the values that that method requires.
 	bool bNeedRemoteElement = false;
-	auto DummyReq = MakeSharedTS<FStreamSegmentRequestDASH>();
+	auto DummyReq = MakeSharedTS<FStreamSegmentRequestCommon>();
 	DummyReq->StreamType = StreamType;
 	DummyReq->PeriodStart = StartPosition.Time;
 	DummyReq->TimestampSequenceIndex = SequenceState.GetSequenceIndex();
@@ -1538,7 +1543,7 @@ IManifest::FResult FDASHPlayPeriod::GetContinuationSegment(TSharedPtrTS<IStreamS
 IManifest::FResult FDASHPlayPeriod::GetNextOrRetrySegment(TSharedPtrTS<IStreamSegment>& OutSegment, bool& OutWaitForRemoteElement, TSharedPtrTS<const IStreamSegment> InCurrentSegment, ENextSegType InNextType, const FPlayStartOptions& Options)
 {
 	OutWaitForRemoteElement = false;
-	TSharedPtrTS<const FStreamSegmentRequestDASH> Current = StaticCastSharedPtr<const FStreamSegmentRequestDASH>(InCurrentSegment);
+	TSharedPtrTS<const FStreamSegmentRequestCommon> Current = StaticCastSharedPtr<const FStreamSegmentRequestCommon>(InCurrentSegment);
 	if (Current->bIsInitialStartRequest)
 	{
 		return IManifest::FResult(IManifest::FResult::EType::NotFound).SetErrorDetail(FErrorDetail().SetMessage("The next segment cannot be located for the initial start request, only for an actual media request!"));
@@ -1594,7 +1599,7 @@ IManifest::FResult FDASHPlayPeriod::GetNextOrRetrySegment(TSharedPtrTS<IStreamSe
 	PlayRangeEnd -= AST;
 	PlayRangeEnd -= Period->GetStart();
 
-	FManifestDASHInternal::FSegmentInformation SegmentInfo;
+	FSegmentInformationCommon SegmentInfo;
 	FManifestDASHInternal::FSegmentSearchOption SearchOpt;
 	TArray<TWeakPtrTS<FMPDLoadRequestDASH>> RemoteElementLoadRequests;
 	if (InNextType == ENextSegType::SamePeriodNext)
@@ -1704,12 +1709,12 @@ IManifest::FResult FDASHPlayPeriod::GetNextOrRetrySegment(TSharedPtrTS<IStreamSe
 	}
 	else if (SearchResult == FManifestDASHInternal::FRepresentation::ESearchResult::Found)
 	{
-		TSharedPtrTS<FStreamSegmentRequestDASH> SegmentRequest = MakeSharedTS<FStreamSegmentRequestDASH>();
+		TSharedPtrTS<FStreamSegmentRequestCommon> SegmentRequest = MakeSharedTS<FStreamSegmentRequestCommon>();
 		SegmentRequest->TimestampSequenceIndex = Current->TimestampSequenceIndex;
 		SegmentRequest->StreamType = Current->GetType();
 		SegmentRequest->QualityIndex = ActiveQualityIndex.Index;
 		SegmentRequest->MaxQualityIndex = ActiveQualityIndex.MaxIndex;
-		SegmentRequest->CodecInfo = Repr->GetCodecInformation();
+		SegmentRequest->CodecInfo[StreamTypeToArrayIndex(SegmentRequest->StreamType)] = Repr->GetCodecInformation();
 		SegmentRequest->Representation = Repr;
 		SegmentRequest->AdaptationSet = Adapt;
 		SegmentRequest->Period = Period;
@@ -1769,7 +1774,7 @@ IManifest::FResult FDASHPlayPeriod::GetNextSegment(TSharedPtrTS<IStreamSegment>&
 	}
 	// Did the stream reader see a 'lmsg' brand on this segment?
 	// If so then this stream has ended and there will not be a next segment.
-	const FStreamSegmentRequestDASH* CurrentRequest = static_cast<const FStreamSegmentRequestDASH*>(InCurrentSegment.Get());
+	const FStreamSegmentRequestCommon* CurrentRequest = static_cast<const FStreamSegmentRequestCommon*>(InCurrentSegment.Get());
 
 	// Check if we moved across a period.
 	bool bNeedRemoteElement = false;
@@ -1797,8 +1802,8 @@ IManifest::FResult FDASHPlayPeriod::GetRetrySegment(TSharedPtrTS<IStreamSegment>
 	// To insert filler data we can use the current request over again.
 	if (bReplaceWithFillerData)
 	{
-		const FStreamSegmentRequestDASH* CurrentRequest = static_cast<const FStreamSegmentRequestDASH*>(InCurrentSegment.Get());
-		TSharedPtrTS<FStreamSegmentRequestDASH> NewRequest = MakeSharedTS<FStreamSegmentRequestDASH>();
+		const FStreamSegmentRequestCommon* CurrentRequest = static_cast<const FStreamSegmentRequestCommon*>(InCurrentSegment.Get());
+		TSharedPtrTS<FStreamSegmentRequestCommon> NewRequest = MakeSharedTS<FStreamSegmentRequestCommon>();
 		*NewRequest = *CurrentRequest;
 		NewRequest->bInsertFillerData = true;
 		// We treat replacing the segment with filler data as a retry.
@@ -1809,7 +1814,7 @@ IManifest::FResult FDASHPlayPeriod::GetRetrySegment(TSharedPtrTS<IStreamSegment>
 
 	// Pass the download stats bWaitingForRemoteRetryElement to convey if the retry segment needs to wait for a remote element,
 	// which is either some xlink or an index segment.
-	FStreamSegmentRequestDASH* CurrentRequest = const_cast<FStreamSegmentRequestDASH*>(static_cast<const FStreamSegmentRequestDASH*>(InCurrentSegment.Get()));
+	FStreamSegmentRequestCommon* CurrentRequest = const_cast<FStreamSegmentRequestCommon*>(static_cast<const FStreamSegmentRequestCommon*>(InCurrentSegment.Get()));
 	IManifest::FResult Result = GetNextOrRetrySegment(OutSegment, CurrentRequest->DownloadStats.bWaitingForRemoteRetryElement, InCurrentSegment, ENextSegType::SamePeriodRetry, Options);
 	return Result;
 }
@@ -2032,11 +2037,11 @@ void FDASHPlayPeriod::HandleRepresentationInitSegmentLoading(const TArray<FInitS
 					is.LoadRequest->Headers.Emplace(HTTP::FHTTPHeader({DASH::HTTPHeaderOptionName, is.InitSegmentInfo.InitializationURL.CustomHeader}));
 				}
 				is.LoadRequest->PlayerSessionServices = PlayerSessionServices;
-				if (is.InitSegmentInfo.ContainerType == FManifestDASHInternal::FSegmentInformation::EContainerType::ISO14496_12)
+				if (is.InitSegmentInfo.ContainerType == FSegmentInformationCommon::EContainerType::ISO14496_12)
 				{
 					is.LoadRequest->CompleteCallback.BindThreadSafeSP(AsShared(), &FDASHPlayPeriod::InitSegmentMP4DownloadComplete);
 				}
-				else if (is.InitSegmentInfo.ContainerType == FManifestDASHInternal::FSegmentInformation::EContainerType::Matroska)
+				else if (is.InitSegmentInfo.ContainerType == FSegmentInformationCommon::EContainerType::Matroska)
 				{
 					is.LoadRequest->CompleteCallback.BindThreadSafeSP(AsShared(), &FDASHPlayPeriod::InitSegmentMKVDownloadComplete);
 				}
@@ -2310,7 +2315,7 @@ FManifestDASHInternal::FRepresentation::ESearchResult FManifestDASHInternal::FRe
 				else if (GetStreamContainerType() == EStreamContainerType::Matroska)
 				{
 					TSharedPtrTS<FMPDLoadRequestDASH> LoadReq = MakeSharedTS<FMPDLoadRequestDASH>();
-					FSegmentInformation InitSegInfo;
+					FSegmentInformationCommon InitSegInfo;
 					if (PrepareDownloadURLs(InPlayerSessionServices, InitSegInfo, SegmentBase))
 					{
 						LoadReq->Range2 = InitSegInfo.InitializationURL.Range;
@@ -2345,7 +2350,7 @@ FManifestDASHInternal::FRepresentation::ESearchResult FManifestDASHInternal::FRe
 }
 
 
-bool FManifestDASHInternal::FRepresentation::PrepareDownloadURLs(IPlayerSessionServices* InPlayerSessionServices, FSegmentInformation& InOutSegmentInfo, const TArray<TSharedPtrTS<FDashMPD_SegmentBaseType>>& SegmentBase)
+bool FManifestDASHInternal::FRepresentation::PrepareDownloadURLs(IPlayerSessionServices* InPlayerSessionServices, FSegmentInformationCommon& InOutSegmentInfo, const TArray<TSharedPtrTS<FDashMPD_SegmentBaseType>>& SegmentBase)
 {
 	TSharedPtrTS<FDashMPD_RepresentationType> MPDRepresentation = Representation.Pin();
 
@@ -2416,11 +2421,11 @@ bool FManifestDASHInternal::FRepresentation::PrepareDownloadURLs(IPlayerSessionS
 		InOutSegmentInfo.bAvailabilityTimeComplete = bATOComplete;
 	}
 	InOutSegmentInfo.bLowLatencyChunkedEncodingExpected = bAvailableAsLowLatency.GetWithDefault(false);
-	InOutSegmentInfo.ContainerType = StreamContainerType == EStreamContainerType::ISO14496_12 ? FSegmentInformation::EContainerType::ISO14496_12 : FSegmentInformation::EContainerType::Matroska;
+	InOutSegmentInfo.ContainerType = StreamContainerType == EStreamContainerType::ISO14496_12 ? FSegmentInformationCommon::EContainerType::ISO14496_12 : FSegmentInformationCommon::EContainerType::Matroska;
 	return true;
 }
 
-bool FManifestDASHInternal::FRepresentation::PrepareDownloadURLs(IPlayerSessionServices* InPlayerSessionServices, FSegmentInformation& InOutSegmentInfo, const TArray<TSharedPtrTS<FDashMPD_SegmentTemplateType>>& SegmentTemplate)
+bool FManifestDASHInternal::FRepresentation::PrepareDownloadURLs(IPlayerSessionServices* InPlayerSessionServices, FSegmentInformationCommon& InOutSegmentInfo, const TArray<TSharedPtrTS<FDashMPD_SegmentTemplateType>>& SegmentTemplate)
 {
 	TSharedPtrTS<FDashMPD_RepresentationType> MPDRepresentation = Representation.Pin();
 
@@ -2521,11 +2526,11 @@ bool FManifestDASHInternal::FRepresentation::PrepareDownloadURLs(IPlayerSessionS
 	}
 
 	InOutSegmentInfo.bLowLatencyChunkedEncodingExpected = bAvailableAsLowLatency.GetWithDefault(false);
-	InOutSegmentInfo.ContainerType = StreamContainerType == EStreamContainerType::ISO14496_12 ? FSegmentInformation::EContainerType::ISO14496_12 : FSegmentInformation::EContainerType::Matroska;
+	InOutSegmentInfo.ContainerType = StreamContainerType == EStreamContainerType::ISO14496_12 ? FSegmentInformationCommon::EContainerType::ISO14496_12 : FSegmentInformationCommon::EContainerType::Matroska;
 	return true;
 }
 
-FString FManifestDASHInternal::FRepresentation::ApplyTemplateStrings(FString TemplateURL, const FSegmentInformation& InSegmentInfo)
+FString FManifestDASHInternal::FRepresentation::ApplyTemplateStrings(FString TemplateURL, const FSegmentInformationCommon& InSegmentInfo)
 {
 	auto PrintWithWidth = [](int64 Value, int32 Width) -> FString
 	{
@@ -2617,7 +2622,7 @@ FString FManifestDASHInternal::FRepresentation::ApplyTemplateStrings(FString Tem
 	return NewURL;
 }
 
-void FManifestDASHInternal::FRepresentation::CollectInbandEventStreams(IPlayerSessionServices* InPlayerSessionServices, FSegmentInformation& InOutSegmentInfo)
+void FManifestDASHInternal::FRepresentation::CollectInbandEventStreams(IPlayerSessionServices* InPlayerSessionServices, FSegmentInformationCommon& InOutSegmentInfo)
 {
 	TSharedPtrTS<FDashMPD_RepresentationType> MPDRepresentation = Representation.Pin();
 
@@ -2632,7 +2637,7 @@ void FManifestDASHInternal::FRepresentation::CollectInbandEventStreams(IPlayerSe
 	}
 	for(int32 i=0; i<EvS.Num(); ++i)
 	{
-		FSegmentInformation::FInbandEventStream ibs;
+		FSegmentInformationCommon::FInbandEventStream ibs;
 		ibs.SchemeIdUri = EvS[i]->GetSchemeIdUri();
 		ibs.Value = EvS[i]->GetValue();
 		ibs.PTO = (int64) EvS[i]->GetPresentationTimeOffset().GetWithDefault(0);
@@ -2641,7 +2646,7 @@ void FManifestDASHInternal::FRepresentation::CollectInbandEventStreams(IPlayerSe
 	}
 }
 
-void FManifestDASHInternal::FRepresentation::SetupProducerReferenceTimeInfo(IPlayerSessionServices* InPlayerSessionServices, FSegmentInformation& InOutSegmentInfo)
+void FManifestDASHInternal::FRepresentation::SetupProducerReferenceTimeInfo(IPlayerSessionServices* InPlayerSessionServices, FSegmentInformationCommon& InOutSegmentInfo)
 {
 	InOutSegmentInfo.ProducerReferenceTimeInfos = ProducerReferenceTimeInfos;
 	DashUtils::GetPlayerOption(InPlayerSessionServices, InOutSegmentInfo.MeasureLatencyViaReferenceTimeInfoID, DASH::OptionKey_LatencyReferenceId, (int64)-1);
@@ -2649,7 +2654,7 @@ void FManifestDASHInternal::FRepresentation::SetupProducerReferenceTimeInfo(IPla
 
 
 
-FManifestDASHInternal::FRepresentation::ESearchResult FManifestDASHInternal::FRepresentation::FindSegment(IPlayerSessionServices* InPlayerSessionServices, FSegmentInformation& OutSegmentInfo, TArray<TWeakPtrTS<FMPDLoadRequestDASH>>& OutRemoteElementLoadRequests, const FSegmentSearchOption& InSearchOptions)
+FManifestDASHInternal::FRepresentation::ESearchResult FManifestDASHInternal::FRepresentation::FindSegment(IPlayerSessionServices* InPlayerSessionServices, FSegmentInformationCommon& OutSegmentInfo, TArray<TWeakPtrTS<FMPDLoadRequestDASH>>& OutRemoteElementLoadRequests, const FSegmentSearchOption& InSearchOptions)
 {
 	SCOPE_CYCLE_COUNTER(STAT_ElectraPlayer_DASH_FindSegment);
 	CSV_SCOPED_TIMING_STAT(ElectraPlayer, DASH_FindSegment);
@@ -2789,7 +2794,7 @@ FManifestDASHInternal::FRepresentation::ESearchResult FManifestDASHInternal::FRe
 }
 
 
-FManifestDASHInternal::FRepresentation::ESearchResult FManifestDASHInternal::FRepresentation::FindSegment_Base_MP4(IPlayerSessionServices* InPlayerSessionServices, FSegmentInformation& OutSegmentInfo, TArray<TWeakPtrTS<FMPDLoadRequestDASH>>& OutRemoteElementLoadRequests, const FSegmentSearchOption& InSearchOptions, const TSharedPtrTS<FDashMPD_RepresentationType>& MPDRepresentation, const TArray<TSharedPtrTS<FDashMPD_SegmentBaseType>>& SegmentBase)
+FManifestDASHInternal::FRepresentation::ESearchResult FManifestDASHInternal::FRepresentation::FindSegment_Base_MP4(IPlayerSessionServices* InPlayerSessionServices, FSegmentInformationCommon& OutSegmentInfo, TArray<TWeakPtrTS<FMPDLoadRequestDASH>>& OutRemoteElementLoadRequests, const FSegmentSearchOption& InSearchOptions, const TSharedPtrTS<FDashMPD_RepresentationType>& MPDRepresentation, const TArray<TSharedPtrTS<FDashMPD_SegmentBaseType>>& SegmentBase)
 {
 	ESearchResult SegIndexResult = PrepareSegmentIndex(InPlayerSessionServices, SegmentBase, OutRemoteElementLoadRequests, InSearchOptions);
 	if (SegIndexResult != ESearchResult::Found)
@@ -3014,7 +3019,7 @@ FManifestDASHInternal::FRepresentation::ESearchResult FManifestDASHInternal::FRe
 }
 
 
-FManifestDASHInternal::FRepresentation::ESearchResult FManifestDASHInternal::FRepresentation::FindSegment_Base_MKV(IPlayerSessionServices* InPlayerSessionServices, FSegmentInformation& OutSegmentInfo, TArray<TWeakPtrTS<FMPDLoadRequestDASH>>& OutRemoteElementLoadRequests, const FSegmentSearchOption& InSearchOptions, const TSharedPtrTS<FDashMPD_RepresentationType>& MPDRepresentation, const TArray<TSharedPtrTS<FDashMPD_SegmentBaseType>>& SegmentBase)
+FManifestDASHInternal::FRepresentation::ESearchResult FManifestDASHInternal::FRepresentation::FindSegment_Base_MKV(IPlayerSessionServices* InPlayerSessionServices, FSegmentInformationCommon& OutSegmentInfo, TArray<TWeakPtrTS<FMPDLoadRequestDASH>>& OutRemoteElementLoadRequests, const FSegmentSearchOption& InSearchOptions, const TSharedPtrTS<FDashMPD_RepresentationType>& MPDRepresentation, const TArray<TSharedPtrTS<FDashMPD_SegmentBaseType>>& SegmentBase)
 {
 	ESearchResult SegIndexResult = PrepareSegmentIndex(InPlayerSessionServices, SegmentBase, OutRemoteElementLoadRequests, InSearchOptions);
 	if (SegIndexResult != ESearchResult::Found)
@@ -3115,7 +3120,7 @@ FManifestDASHInternal::FRepresentation::ESearchResult FManifestDASHInternal::FRe
 	}
 }
 
-FManifestDASHInternal::FRepresentation::ESearchResult FManifestDASHInternal::FRepresentation::FindSegment_Template(IPlayerSessionServices* InPlayerSessionServices, FSegmentInformation& OutSegmentInfo, TArray<TWeakPtrTS<FMPDLoadRequestDASH>>& OutRemoteElementLoadRequests, const FSegmentSearchOption& InSearchOptions, const TSharedPtrTS<FDashMPD_RepresentationType>& MPDRepresentation, const TArray<TSharedPtrTS<FDashMPD_SegmentTemplateType>>& SegmentTemplate)
+FManifestDASHInternal::FRepresentation::ESearchResult FManifestDASHInternal::FRepresentation::FindSegment_Template(IPlayerSessionServices* InPlayerSessionServices, FSegmentInformationCommon& OutSegmentInfo, TArray<TWeakPtrTS<FMPDLoadRequestDASH>>& OutRemoteElementLoadRequests, const FSegmentSearchOption& InSearchOptions, const TSharedPtrTS<FDashMPD_RepresentationType>& MPDRepresentation, const TArray<TSharedPtrTS<FDashMPD_SegmentTemplateType>>& SegmentTemplate)
 {
 	uint64 PTO = GET_ATTR(SegmentTemplate, GetPresentationTimeOffset(), IsSet(), TMediaOptionalValue<uint64>(0)).Value();
 	uint32 MPDTimescale = GET_ATTR(SegmentTemplate, GetTimescale(), IsSet(), TMediaOptionalValue<uint32>(1)).Value();
@@ -3261,7 +3266,7 @@ FManifestDASHInternal::FRepresentation::ESearchResult FManifestDASHInternal::FRe
 	}
 }
 
-FManifestDASHInternal::FRepresentation::ESearchResult FManifestDASHInternal::FRepresentation::FindSegment_Timeline(IPlayerSessionServices* InPlayerSessionServices, FSegmentInformation& OutSegmentInfo, TArray<TWeakPtrTS<FMPDLoadRequestDASH>>& OutRemoteElementLoadRequests, const FSegmentSearchOption& InSearchOptions, const TSharedPtrTS<FDashMPD_RepresentationType>& MPDRepresentation, const TArray<TSharedPtrTS<FDashMPD_SegmentTemplateType>>& SegmentTemplate, const TSharedPtrTS<FDashMPD_SegmentTimelineType>& SegmentTimeline)
+FManifestDASHInternal::FRepresentation::ESearchResult FManifestDASHInternal::FRepresentation::FindSegment_Timeline(IPlayerSessionServices* InPlayerSessionServices, FSegmentInformationCommon& OutSegmentInfo, TArray<TWeakPtrTS<FMPDLoadRequestDASH>>& OutRemoteElementLoadRequests, const FSegmentSearchOption& InSearchOptions, const TSharedPtrTS<FDashMPD_RepresentationType>& MPDRepresentation, const TArray<TSharedPtrTS<FDashMPD_SegmentTemplateType>>& SegmentTemplate, const TSharedPtrTS<FDashMPD_SegmentTimelineType>& SegmentTimeline)
 {
 	// Segment timeline must not be empty.
 	auto Selements = SegmentTimeline->GetS_Elements();
@@ -3601,7 +3606,7 @@ FManifestDASHInternal::FRepresentation::ESearchResult FManifestDASHInternal::FRe
 }
 
 
-FManifestDASHInternal::FRepresentation::ESearchResult FManifestDASHInternal::FRepresentation::SetupSideloadedFile(IPlayerSessionServices* InPlayerSessionServices, FSegmentInformation& OutSegmentInfo, const FSegmentSearchOption& InSearchOptions, const TSharedPtrTS<FDashMPD_RepresentationType>& MPDRepresentation)
+FManifestDASHInternal::FRepresentation::ESearchResult FManifestDASHInternal::FRepresentation::SetupSideloadedFile(IPlayerSessionServices* InPlayerSessionServices, FSegmentInformationCommon& OutSegmentInfo, const FSegmentSearchOption& InSearchOptions, const TSharedPtrTS<FDashMPD_RepresentationType>& MPDRepresentation)
 {
 	// As per DASH-IF-IOP 4.3, Section 6.4.5. Guidelines for side-loaded TTML and WebVTT files
 	// side loaded subtitles must not use @presentationTimeOffset, which makes sense since this attribute is not defined on <Representation> elements.
@@ -3669,7 +3674,7 @@ void FManifestDASHInternal::FRepresentation::GetSegmentInformation(TArray<IManif
 	FTimeValue TimeToGo(LookAheadTime);
 	FTimeValue FixedSegmentDuration = FTimeValue(FTimeValue::MillisecondsToHNS(2000));
 	const int32 FixedBitrate = GetBitrate();
-	const FStreamSegmentRequestDASH* CurrentRequest = static_cast<const FStreamSegmentRequestDASH*>(CurrentSegment.Get());
+	const FStreamSegmentRequestCommon* CurrentRequest = static_cast<const FStreamSegmentRequestCommon*>(CurrentSegment.Get());
 
 	// This is the same as in FindSegment(), only with no error checking since this method here is not critical.
 	TSharedPtrTS<FDashMPD_RepresentationType> MPDRepresentation = Representation.Pin();
diff --git a/Engine/Plugins/Media/ElectraPlayer/Source/ElectraPlayerRuntime/Private/Runtime/Player/DASH/PlaylistReaderDASH_Internal.h b/Engine/Plugins/Media/ElectraPlayer/Source/ElectraPlayerRuntime/Private/Runtime/Player/DASH/PlaylistReaderDASH_Internal.h
--- a/Engine/Plugins/Media/ElectraPlayer/Source/ElectraPlayerRuntime/Private/Runtime/Player/DASH/PlaylistReaderDASH_Internal.h
+++ b/Engine/Plugins/Media/ElectraPlayer/Source/ElectraPlayerRuntime/Private/Runtime/Player/DASH/PlaylistReaderDASH_Internal.h
@@ -8,10 +8,10 @@
 #include "MPDElementsDASH.h"
 #include "Player/Manifest.h"
 #include "Player/AdaptiveStreamingPlayerResourceRequest.h"
+#include "Player/StreamSegmentReaderCommon.h"
 #include "Utilities/URLParser.h"
 
 
-
 namespace Electra
 {
 class FManifestDASHInternal;
@@ -134,114 +134,6 @@ public:
 		Dynamic
 	};
 
-	struct FProducerReferenceTimeInfo : public IProducerReferenceTimeInfo
-	{
-		FTimeValue WallclockTime;
-		uint64 PresentationTime = 0;
-		uint32 ID = 0;
-		EType Type = EType::Encoder;
-		bool bInband = false;
-		FTimeValue GetWallclockTime() const override
-		{ return WallclockTime; }
-		uint64 GetPresentationTime() const override
-		{ return PresentationTime; }
-		uint32 GetID() const override
-		{ return ID; }
-		EType GetType() const override
-		{ return Type; }
-		bool GetIsInband() const override
-		{ return bInband; }
-	};
-
-	struct FSegmentInformation
-	{
-		struct FURL
-		{
-			FString URL;
-			FString Range;
-			FString CDN;
-			FString CustomHeader;
-		};
-		struct FInbandEventStream
-		{
-			FString SchemeIdUri;
-			FString Value;
-			int64 PTO = 0;
-			uint32 Timescale = 0;
-		};
-		enum class EContainerType
-		{
-			ISO14496_12,
-			Matroska
-		};
-		EContainerType ContainerType = EContainerType::ISO14496_12;
-		FURL InitializationURL;
-		FURL MediaURL;
-		FTimeValue ATO;
-		FTimeValue PeriodLocalSegmentStartTime;
-		int64 Time = 0;								//!< Time value T in timescale units
-		int64 PTO = 0;								//!< PresentationTimeOffset
-		int64 EPTdelta = 0;
-		int64 Duration = 0;							//!< Duration of the segment. Not necessarily exact if <SegmentTemplate> is used).
-		int64 Number = 0;							//!< Index of the segment.
-		int64 SubIndex = 0;							//!< Subsegment index
-		int64 NumberOfBytes = 0;
-		int64 FirstByteOffset = 0;
-		int64 MediaLocalFirstAUTime = 0;			//!< Time of the first AU to use in this segment in media local time
-		int64 MediaLocalLastAUTime = 0;				//!< Time at which the last AU to use in thie segment ends in media local time
-		int64 MediaLocalFirstPTS = 0;
-		uint32 Timescale = 0;						//!< Local media timescale
-		TMediaOptionalValue<bool> bAvailabilityTimeComplete;
-		bool bLowLatencyChunkedEncodingExpected = false;
-		bool bFrameAccuracyRequired = false;		//!< true if the segment was located for frame accurate seeking.
-		bool bIsSideload = false;					//!< true if this is a side-loaded resource to be fetched and cached.
-		bool bIsLastInPeriod = false;				//!< true if known to be the last segment in the period.
-		bool bMayBeMissing = false;					//!< true if the last segment in <SegmentTemplate> that might not exist.
-		bool bIsMissing = false;					//!< Set to true if known to be missing.
-		bool bSawLMSG = false;						//!< Will be set to true by the stream reader if the 'lmsg' brand was found.
-		TArray<FInbandEventStream> InbandEventStreams;
-		TArray<FProducerReferenceTimeInfo> ProducerReferenceTimeInfos;
-		int64 MeasureLatencyViaReferenceTimeInfoID = -1;
-
-		FTimeValue CalculateASAST(const FTimeValue& AST, const FTimeValue& PeriodStart, bool bIsStatic)
-		{
-			if (bIsStatic)
-			{
-				return AST;
-			}
-			else
-			{
-				if (ATO < FTimeValue::GetPositiveInfinity())
-				{
-					return AST + PeriodStart + FTimeValue(Time - PTO - EPTdelta + Duration, Timescale) - ATO;
-				}
-				// ATO of infinity means the segment is always available, so we return zero time as earliest UTC time.
-				return FTimeValue::GetZero();
-			}
-		}
-		FTimeValue CalculateSAET(const FTimeValue& AST, const FTimeValue& PeriodStart, const FTimeValue& MPDAET, const FTimeValue& TSB, bool bIsStatic)
-		{
-			if (bIsStatic)
-			{
-				// If the MPD has a global availabilityEndTime then that is the end time of the segment as well.
-				if (MPDAET.IsValid())
-				{
-					return MPDAET;
-				}
-				return FTimeValue::GetPositiveInfinity();
-			}
-			else
-			{
-				// If the MPD has a global availabilityEndTime then that is the end time of the segment as well.
-				if (MPDAET.IsValid())
-				{
-					return MPDAET;
-				}
-				return AST + PeriodStart + FTimeValue(Time - PTO - EPTdelta + Duration * 2, Timescale) + (TSB.IsValid() ? TSB : FTimeValue::GetPositiveInfinity());
-			}
-		}
-	};
-
 	struct FSegmentSearchOption
 	{
 		FTimeValue PeriodLocalTime;					//!< Time local in the period to search a segment for.
@@ -272,7 +164,7 @@ public:
 			Gone									//!< Underlying MPD Representation (held by a weak pointer) has gone and the representation is no longer accessible.
 		};
 
-		ESearchResult FindSegment(IPlayerSessionServices* PlayerSessionServices, FSegmentInformation& OutSegmentInfo, TArray<TWeakPtrTS<FMPDLoadRequestDASH>>& OutRemoteElementLoadRequests, const FSegmentSearchOption& InSearchOptions);
+		ESearchResult FindSegment(IPlayerSessionServices* PlayerSessionServices, FSegmentInformationCommon& OutSegmentInfo, TArray<TWeakPtrTS<FMPDLoadRequestDASH>>& OutRemoteElementLoadRequests, const FSegmentSearchOption& InSearchOptions);
 
 		void GetSegmentInformation(TArray<IManifest::IPlayPeriod::FSegmentInformation>& OutSegmentInformation, FTimeValue& OutAverageSegmentDuration, TSharedPtrTS<const IStreamSegment> CurrentSegment, const FTimeValue& LookAheadTime, const TSharedPtrTS<IPlaybackAssetAdaptationSet>& AdaptationSet);
 
@@ -318,17 +210,17 @@ public:
 	private:
 		ESearchResult PrepareSegmentIndex(IPlayerSessionServices* PlayerSessionServices, const TArray<TSharedPtrTS<FDashMPD_SegmentBaseType>>& SegmentBase, TArray<TWeakPtrTS<FMPDLoadRequestDASH>>& OutRemoteElementLoadRequests, const FSegmentSearchOption& InSearchOptions);
 
-		ESearchResult FindSegment_Base_MP4(IPlayerSessionServices* PlayerSessionServices, FSegmentInformation& OutSegmentInfo, TArray<TWeakPtrTS<FMPDLoadRequestDASH>>& OutRemoteElementLoadRequests, const FSegmentSearchOption& InSearchOptions, const TSharedPtrTS<FDashMPD_RepresentationType>& MPDRepresentation, const TArray<TSharedPtrTS<FDashMPD_SegmentBaseType>>& SegmentBase);
-		ESearchResult FindSegment_Base_MKV(IPlayerSessionServices* PlayerSessionServices, FSegmentInformation& OutSegmentInfo, TArray<TWeakPtrTS<FMPDLoadRequestDASH>>& OutRemoteElementLoadRequests, const FSegmentSearchOption& InSearchOptions, const TSharedPtrTS<FDashMPD_RepresentationType>& MPDRepresentation, const TArray<TSharedPtrTS<FDashMPD_SegmentBaseType>>& SegmentBase);
-		ESearchResult FindSegment_Template(IPlayerSessionServices* PlayerSessionServices, FSegmentInformation& OutSegmentInfo, TArray<TWeakPtrTS<FMPDLoadRequestDASH>>& OutRemoteElementLoadRequests, const FSegmentSearchOption& InSearchOptions, const TSharedPtrTS<FDashMPD_RepresentationType>& MPDRepresentation, const TArray<TSharedPtrTS<FDashMPD_SegmentTemplateType>>& SegmentTemplate);
-		ESearchResult FindSegment_Timeline(IPlayerSessionServices* PlayerSessionServices, FSegmentInformation& OutSegmentInfo, TArray<TWeakPtrTS<FMPDLoadRequestDASH>>& OutRemoteElementLoadRequests, const FSegmentSearchOption& InSearchOptions, const TSharedPtrTS<FDashMPD_RepresentationType>& MPDRepresentation, const TArray<TSharedPtrTS<FDashMPD_SegmentTemplateType>>& SegmentTemplate, const TSharedPtrTS<FDashMPD_SegmentTimelineType>& SegmentTimeline);
-		ESearchResult SetupSideloadedFile(IPlayerSessionServices* PlayerSessionServices, FSegmentInformation& OutSegmentInfo, const FSegmentSearchOption& InSearchOptions, const TSharedPtrTS<FDashMPD_RepresentationType>& MPDRepresentation);
-
-		bool PrepareDownloadURLs(IPlayerSessionServices* PlayerSessionServices, FSegmentInformation& InOutSegmentInfo, const TArray<TSharedPtrTS<FDashMPD_SegmentBaseType>>& SegmentBase);
-		bool PrepareDownloadURLs(IPlayerSessionServices* PlayerSessionServices, FSegmentInformation& InOutSegmentInfo, const TArray<TSharedPtrTS<FDashMPD_SegmentTemplateType>>& SegmentTemplate);
-		FString ApplyTemplateStrings(FString TemplateURL, const FSegmentInformation& InSegmentInfo);
-		void CollectInbandEventStreams(IPlayerSessionServices* PlayerSessionServices, FSegmentInformation& InOutSegmentInfo);
-		void SetupProducerReferenceTimeInfo(IPlayerSessionServices* PlayerSessionServices, FSegmentInformation& InOutSegmentInfo);
+		ESearchResult FindSegment_Base_MP4(IPlayerSessionServices* PlayerSessionServices, FSegmentInformationCommon& OutSegmentInfo, TArray<TWeakPtrTS<FMPDLoadRequestDASH>>& OutRemoteElementLoadRequests, const FSegmentSearchOption& InSearchOptions, const TSharedPtrTS<FDashMPD_RepresentationType>& MPDRepresentation, const TArray<TSharedPtrTS<FDashMPD_SegmentBaseType>>& SegmentBase);
+		ESearchResult FindSegment_Base_MKV(IPlayerSessionServices* PlayerSessionServices, FSegmentInformationCommon& OutSegmentInfo, TArray<TWeakPtrTS<FMPDLoadRequestDASH>>& OutRemoteElementLoadRequests, const FSegmentSearchOption& InSearchOptions, const TSharedPtrTS<FDashMPD_RepresentationType>& MPDRepresentation, const TArray<TSharedPtrTS<FDashMPD_SegmentBaseType>>& SegmentBase);
+		ESearchResult FindSegment_Template(IPlayerSessionServices* PlayerSessionServices, FSegmentInformationCommon& OutSegmentInfo, TArray<TWeakPtrTS<FMPDLoadRequestDASH>>& OutRemoteElementLoadRequests, const FSegmentSearchOption& InSearchOptions, const TSharedPtrTS<FDashMPD_RepresentationType>& MPDRepresentation, const TArray<TSharedPtrTS<FDashMPD_SegmentTemplateType>>& SegmentTemplate);
+		ESearchResult FindSegment_Timeline(IPlayerSessionServices* PlayerSessionServices, FSegmentInformationCommon& OutSegmentInfo, TArray<TWeakPtrTS<FMPDLoadRequestDASH>>& OutRemoteElementLoadRequests, const FSegmentSearchOption& InSearchOptions, const TSharedPtrTS<FDashMPD_RepresentationType>& MPDRepresentation, const TArray<TSharedPtrTS<FDashMPD_SegmentTemplateType>>& SegmentTemplate, const TSharedPtrTS<FDashMPD_SegmentTimelineType>& SegmentTimeline);
+		ESearchResult SetupSideloadedFile(IPlayerSessionServices* PlayerSessionServices, FSegmentInformationCommon& OutSegmentInfo, const FSegmentSearchOption& InSearchOptions, const TSharedPtrTS<FDashMPD_RepresentationType>& MPDRepresentation);
+
+		bool PrepareDownloadURLs(IPlayerSessionServices* PlayerSessionServices, FSegmentInformationCommon& InOutSegmentInfo, const TArray<TSharedPtrTS<FDashMPD_SegmentBaseType>>& SegmentBase);
+		bool PrepareDownloadURLs(IPlayerSessionServices* PlayerSessionServices, FSegmentInformationCommon& InOutSegmentInfo, const TArray<TSharedPtrTS<FDashMPD_SegmentTemplateType>>& SegmentTemplate);
+		FString ApplyTemplateStrings(FString TemplateURL, const FSegmentInformationCommon& InSegmentInfo);
+		void CollectInbandEventStreams(IPlayerSessionServices* PlayerSessionServices, FSegmentInformationCommon& InOutSegmentInfo);
+		void SetupProducerReferenceTimeInfo(IPlayerSessionServices* PlayerSessionServices, FSegmentInformationCommon& InOutSegmentInfo);
 
 		void SegmentIndexDownloadComplete(TSharedPtrTS<FMPDLoadRequestDASH> Request, bool bSuccess);
 		friend class FManifestDASHInternal;
diff --git a/Engine/Plugins/Media/ElectraPlayer/Source/ElectraPlayerRuntime/Private/Runtime/Player/DASH/StreamReaderDASH.cpp b/Engine/Plugins/Media/ElectraPlayer/Source/ElectraPlayerRuntime/Private/Runtime/Player/DASH/StreamReaderDASH.cpp
deleted file mode 100644
diff --git a/Engine/Plugins/Media/ElectraPlayer/Source/ElectraPlayerRuntime/Private/Runtime/Player/DASH/StreamReaderDASH.h b/Engine/Plugins/Media/ElectraPlayer/Source/ElectraPlayerRuntime/Private/Runtime/Player/DASH/StreamReaderDASH.h
deleted file mode 100644
