ElectraPlayer: Fixing up some incorrect metrics

#rnx

diff --git a/Engine/Plugins/Media/ElectraPlayer/Source/ElectraPlayerRuntime/Private/ElectraPlayer.cpp b/Engine/Plugins/Media/ElectraPlayer/Source/ElectraPlayerRuntime/Private/ElectraPlayer.cpp
--- a/Engine/Plugins/Media/ElectraPlayer/Source/ElectraPlayerRuntime/Private/ElectraPlayer.cpp
+++ b/Engine/Plugins/Media/ElectraPlayer/Source/ElectraPlayerRuntime/Private/ElectraPlayer.cpp
@@ -1858,6 +1858,7 @@ void FElectraPlayer::AddCommonAnalyticsAttributes(TArray<FAnalyticsEventAttribut
 	InOutParamArray.Add(FAnalyticsEventAttribute(TEXT("Utc"), static_cast<double>(FDateTime::UtcNow().ToUnixTimestamp())));
 	InOutParamArray.Add(FAnalyticsEventAttribute(TEXT("OS"), FString::Printf(TEXT("%s"), *AnalyticsOSVersion)));
 	InOutParamArray.Add(FAnalyticsEventAttribute(TEXT("GPUAdapter"), AnalyticsGPUType));
+	++AnalyticsInstanceEventCount;
 	StatisticsLock.Lock();
 	for(int32 nI=0, nIMax=UE_ARRAY_COUNT(AnalyticsCustomValues); nI<nIMax; ++nI)
 	{
@@ -2118,11 +2119,11 @@ void FElectraPlayer::HandlePlayerEventReceivedMainPlaylist(const FString& Effect
 	Statistics.LastState = "Preparing";
 
 
-	// Enqueue a "MasterPlaylist" event.
-	static const FString kEventNameElectraMasterPlaylist(TEXT("Electra.MasterPlaylist"));
-	if (Electra::IsAnalyticsEventEnabled(kEventNameElectraMasterPlaylist))
+	// Enqueue a "MainPlaylist" event.
+	static const FString kEventNameElectraMainPlaylist(TEXT("Electra.MainPlaylist"));
+	if (Electra::IsAnalyticsEventEnabled(kEventNameElectraMainPlaylist))
 	{
-		TSharedPtr<FAnalyticsEvent> AnalyticEvent = CreateAnalyticsEvent(kEventNameElectraMasterPlaylist);
+		TSharedPtr<FAnalyticsEvent> AnalyticEvent = CreateAnalyticsEvent(kEventNameElectraMainPlaylist);
 		AnalyticEvent->ParamArray.Add(FAnalyticsEventAttribute(TEXT("URL"), *EffectiveURL));
 		EnqueueAnalyticsEvent(AnalyticEvent);
 	}
@@ -2171,17 +2172,8 @@ void FElectraPlayer::HandlePlayerEventReceivedPlaylists()
 												VideoStreamMetaData[0].StreamDetails[i].Bandwidth);
 		}
 	}
-
 	SelectedVideoTrackIndex = NumTracksVideo ? 0 : -1;
 
-	// Enqueue a "PlaylistsLoaded" event.
-	static const FString kEventNameElectraPlaylistLoaded(TEXT("Electra.PlaylistsLoaded"));
-	if (Electra::IsAnalyticsEventEnabled(kEventNameElectraPlaylistLoaded))
-	{
-		TSharedPtr<FAnalyticsEvent> AnalyticEvent = CreateAnalyticsEvent(kEventNameElectraPlaylistLoaded);
-		EnqueueAnalyticsEvent(AnalyticEvent);
-	}
-
 	// Get the audio bitrates and populate our number of segments per bitrate map.
 	TArray<Electra::FTrackMetadata> AudioStreamMetaData;
 	CurrentPlayer->AdaptivePlayer->GetTrackMetadata(AudioStreamMetaData, Electra::EStreamType::Audio);
@@ -2194,6 +2186,7 @@ void FElectraPlayer::HandlePlayerEventReceivedPlaylists()
 												AudioStreamMetaData[0].StreamDetails[i].Bandwidth);
 		}
 	}
+	SelectedAudioTrackIndex = NumTracksAudio ? 0 : -1;
 
 	TArray<Electra::FTrackMetadata> SubtitleStreamMetaData;
 	CurrentPlayer->AdaptivePlayer->GetTrackMetadata(SubtitleStreamMetaData, Electra::EStreamType::Subtitle);
@@ -2213,6 +2206,14 @@ void FElectraPlayer::HandlePlayerEventReceivedPlaylists()
 	InitialSubtitleAttributes.OverrideIndex = PlaystartOptions.InitialSubtitleTrackAttributes.TrackIndexOverride;
 	CurrentPlayer->AdaptivePlayer->SetInitialStreamAttributes(Electra::EStreamType::Subtitle, InitialSubtitleAttributes);
 
+	// Enqueue a "PlaylistsLoaded" event.
+	static const FString kEventNameElectraPlaylistLoaded(TEXT("Electra.PlaylistsLoaded"));
+	if (Electra::IsAnalyticsEventEnabled(kEventNameElectraPlaylistLoaded))
+	{
+		TSharedPtr<FAnalyticsEvent> AnalyticEvent = CreateAnalyticsEvent(kEventNameElectraPlaylistLoaded);
+		EnqueueAnalyticsEvent(AnalyticEvent);
+	}
+
 	// Trigger preloading unless forbidden.
 	if (PlaystartOptions.bDoNotPreload == false)
 	{
@@ -2477,6 +2478,20 @@ void FElectraPlayer::HandlePlayerEventSegmentDownload(const Electra::Metrics::FS
 	else if (SegmentDownloadStats.StreamType == Electra::EStreamType::Audio)
 	{
 		Statistics.NumAudioDatabytesStreamed += SegmentDownloadStats.NumBytesDownloaded;
+		if (!Statistics.AudioSegmentBitratesStreamed.Contains(SegmentDownloadStats.Bitrate))
+		{
+			Statistics.AudioSegmentBitratesStreamed.Add(SegmentDownloadStats.Bitrate, 0);
+		}
+		++Statistics.AudioSegmentBitratesStreamed[SegmentDownloadStats.Bitrate];
+
+		if (Statistics.bIsInitiallyDownloading && NumTracksVideo == 0)	// Do this just for audio-only presentations.
+		{
+			Statistics.InitialBufferingBandwidth.AddSample(8*SegmentDownloadStats.NumBytesDownloaded / (SegmentDownloadStats.TimeToDownload > 0.0 ? SegmentDownloadStats.TimeToDownload : 1.0), SegmentDownloadStats.TimeToFirstByte);
+			if (Statistics.InitialBufferingDuration > 0.0)
+			{
+				Statistics.bIsInitiallyDownloading = false;
+			}
+		}
 	}
 	if (SegmentDownloadStats.bWasSuccessful)
 	{
@@ -2960,30 +2975,6 @@ void FElectraPlayer::HandlePlayerEventDroppedAudioFrame()
 }
 
 
-void FElectraPlayer::FDroppedFrameStats::AddNewDrop(const FTimespan& InFrameTime, const FTimespan& InPlayerTime, void* InPtrElectraPlayer, void* InPtrCurrentPlayer)
-{
-	double Now = FPlatformTime::Seconds();
-	++NumTotalDropped;
-
-	FTimespan Diff = InPlayerTime - InFrameTime;
-	if (Diff > WorstDeltaTime)
-	{
-		WorstDeltaTime = Diff;
-	}
-
-	if (Now - SystemTimeAtLastReport >= LogWarningAfterSeconds)
-	{
-		FString Type = FrameType == EFrameType::VideoFrame ? TEXT("video") :
-					   FrameType == EFrameType::AudioFrame ? TEXT("audio") :
-					   TEXT("<unknown>");
-
-		//UE_LOG(LogElectraPlayer, Log, TEXT("[%u] dropped %s %s @ %s. %d drops so far. Worst delta %.3f ms"), InPtrElectraPlayer, *Type, *InFrameTime.ToString(TEXT("%h:%m:%s.%f")), *InPlayerTime.ToString(TEXT("%h:%m:%s.%f")), NumTotalDropped, WorstDeltaTime.GetTotalMilliseconds());
-		SystemTimeAtLastReport = Now;
-	}
-
-	PlayerTimeAtLastReport = InPlayerTime;
-}
-
 void FElectraPlayer::FStatistics::AddMessageToHistory(FString InMessage)
 {
 	if (MessageHistoryBuffer.Num() >= 20)
@@ -3026,13 +3017,18 @@ void FElectraPlayer::UpdatePlayEndStatistics()
 
 void FElectraPlayer::LogStatistics()
 {
-	FString SegsPerStream;
+	FString VideoSegsPerStream;
+	FString AudioSegsPerStream;
 
 	FScopeLock Lock(&StatisticsLock);
 
 	for(const TPair<int32, uint32>& pair : Statistics.VideoSegmentBitratesStreamed)
 	{
-		SegsPerStream += FString::Printf(TEXT("%d : %u\n"), pair.Key, pair.Value);
+		VideoSegsPerStream += FString::Printf(TEXT("%d : %u\n"), pair.Key, pair.Value);
+	}
+	for(const TPair<int32, uint32>& pair : Statistics.AudioSegmentBitratesStreamed)
+	{
+		AudioSegsPerStream += FString::Printf(TEXT("%d : %u\n"), pair.Key, pair.Value);
 	}
 
 	TSharedPtr<IElectraPlayerAdapterDelegate, ESPMode::ThreadSafe> PinnedAdapterDelegate = AdapterDelegate.Pin();
@@ -3072,7 +3068,8 @@ void FElectraPlayer::LogStatistics()
 			"Number of audio drastic downswitches: %d\n"\
 			"Bytes of video data streamed: %lld\n"\
 			"Bytes of audio data streamed: %lld\n"\
-			"Number of segments fetched across all quality levels:\n%s"\
+			"Number of video segments fetched across all quality levels:\n%s"\
+			"Number of audio segments fetched across all quality levels:\n%s"\
 			"Currently active playlist URL: %s\n"\
 			"Currently active resolution: %d * %d\n" \
 			"Current state: %s\n" \
@@ -3111,7 +3108,8 @@ void FElectraPlayer::LogStatistics()
 			Statistics.NumAudioQualityDrasticDownswitches,
 			(long long int)Statistics.NumVideoDatabytesStreamed,
 			(long long int)Statistics.NumAudioDatabytesStreamed,
-			*SegsPerStream,
+			*VideoSegsPerStream,
+			*AudioSegsPerStream,
 			*SanitizeMessage(Statistics.CurrentlyActivePlaylistURL),
 			Statistics.CurrentlyActiveResolutionWidth,
 			Statistics.CurrentlyActiveResolutionHeight,
@@ -3173,7 +3171,7 @@ void FElectraPlayer::SendAnalyticMetrics(const TSharedPtr<IAnalyticsProviderET>&
 	ParamArray.Add(FAnalyticsEventAttribute(TEXT("MessageHistory"), MessageHistory));
 	ParamArray.Add(FAnalyticsEventAttribute(TEXT("LastError"), Statistics.LastError));
 	ParamArray.Add(FAnalyticsEventAttribute(TEXT("FinalVideoResolution"), FString::Printf(TEXT("%d*%d"), Statistics.CurrentlyActiveResolutionWidth, Statistics.CurrentlyActiveResolutionHeight)));
-	ParamArray.Add(FAnalyticsEventAttribute(TEXT("TimeElapsedToMasterPlaylist"), Statistics.TimeToLoadMainPlaylist));
+	ParamArray.Add(FAnalyticsEventAttribute(TEXT("TimeElapsedToMainPlaylist"), Statistics.TimeToLoadMainPlaylist));
 	ParamArray.Add(FAnalyticsEventAttribute(TEXT("TimeElapsedToPlaylists"), Statistics.TimeToLoadStreamPlaylists));
 	ParamArray.Add(FAnalyticsEventAttribute(TEXT("InitialAvgBufferingBandwidth"), Statistics.InitialBufferingBandwidth.GetAverageBandwidth()));
 	ParamArray.Add(FAnalyticsEventAttribute(TEXT("InitialAvgBufferingLatency"), Statistics.InitialBufferingBandwidth.GetAverageLatency()));
@@ -3206,20 +3204,19 @@ void FElectraPlayer::SendAnalyticMetrics(const TSharedPtr<IAnalyticsProviderET>&
 	ParamArray.Add(FAnalyticsEventAttribute(TEXT("Rebuffering.MaxDuration"), Statistics.LongestRebufferingDuration));
 	ParamArray.Add(FAnalyticsEventAttribute(TEXT("NumBytesStreamedAudio"), (double) Statistics.NumAudioDatabytesStreamed));
 	ParamArray.Add(FAnalyticsEventAttribute(TEXT("NumBytesStreamedVideo"), (double) Statistics.NumVideoDatabytesStreamed));
-	ParamArray.Add(FAnalyticsEventAttribute(TEXT("NumVideoFramesDropped"), Statistics.DroppedVideoFrames.NumTotalDropped));
-	ParamArray.Add(FAnalyticsEventAttribute(TEXT("WorstVideoTimeDeltaMillis"), Statistics.DroppedVideoFrames.WorstDeltaTime.GetTotalMilliseconds()));
-	ParamArray.Add(FAnalyticsEventAttribute(TEXT("NumAudioFramesDropped"), Statistics.DroppedAudioFrames.NumTotalDropped));
-	ParamArray.Add(FAnalyticsEventAttribute(TEXT("WorstAudioTimeDeltaMillis"), Statistics.DroppedAudioFrames.WorstDeltaTime.GetTotalMilliseconds()));
 	FString SegsPerStream;
 	for(const TPair<int32, uint32>& pair : Statistics.VideoSegmentBitratesStreamed)
 	{
 		SegsPerStream += FString::Printf(TEXT("%d:%u;"), pair.Key, pair.Value);
 	}
 	ParamArray.Add(FAnalyticsEventAttribute(TEXT("VideoSegmentFetchStats"), *SegsPerStream));
+	SegsPerStream.Empty();
+	for(const TPair<int32, uint32>& pair : Statistics.AudioSegmentBitratesStreamed)
+	{
+		SegsPerStream += FString::Printf(TEXT("%d:%u;"), pair.Key, pair.Value);
+	}
+	ParamArray.Add(FAnalyticsEventAttribute(TEXT("AudioSegmentFetchStats"), *SegsPerStream));
 
-	ParamArray.Add(FAnalyticsEventAttribute(TEXT("SubtitlesURL"), Statistics.SubtitlesURL));
-	ParamArray.Add(FAnalyticsEventAttribute(TEXT("SubtitlesResponseTime"), Statistics.SubtitlesResponseTime));
-	ParamArray.Add(FAnalyticsEventAttribute(TEXT("SubtitlesLastError"), Statistics.SubtitlesLastError));
 	StatisticsLock.Unlock();
 
 	AnalyticsProvider->RecordEvent(TEXT("Electra.FinalMetrics"), MoveTemp(ParamArray));
@@ -3273,14 +3270,12 @@ void FElectraPlayer::ReportVideoStreamingError(const FGuid& InPlayerGuid, const
 
 void FElectraPlayer::ReportSubtitlesMetrics(const FGuid& InPlayerGuid, const FString& URL, double ResponseTime, const FString& LastError)
 {
+/*
 	if (PlayerGuid != InPlayerGuid)
 	{
 		return;
 	}
-
-	Statistics.SubtitlesURL = URL;
-	Statistics.SubtitlesResponseTime = ResponseTime;
-	Statistics.SubtitlesLastError = LastError;
+*/
 }
 
 void FElectraPlayer::MediaStateOnPreparingFinished()
@@ -3490,4 +3485,4 @@ IElectraPlayerInterface* FElectraPlayerRuntimeFactory::CreatePlayer(const TShare
 	FElectraPlayerReportSubtitlesMetricsDelegate& InReportSubtitlesFileMetricsDelegate)
 {
 	return new FElectraPlayer(AdapterDelegate, InSendAnalyticMetricsDelegate, InSendAnalyticMetricsPerMinuteDelegate, InReportVideoStreamingErrorDelegate, InReportSubtitlesFileMetricsDelegate);
-};
+}
diff --git a/Engine/Plugins/Media/ElectraPlayer/Source/ElectraPlayerRuntime/Private/ElectraPlayer.h b/Engine/Plugins/Media/ElectraPlayer/Source/ElectraPlayerRuntime/Private/ElectraPlayer.h
--- a/Engine/Plugins/Media/ElectraPlayer/Source/ElectraPlayerRuntime/Private/ElectraPlayer.h
+++ b/Engine/Plugins/Media/ElectraPlayer/Source/ElectraPlayerRuntime/Private/ElectraPlayer.h
@@ -604,45 +604,6 @@ private:
 		int32	MaxSamples;
 	};
 
-	struct FDroppedFrameStats
-	{
-		enum class EFrameType
-		{
-			Undefined,
-			VideoFrame,
-			AudioFrame
-		};
-		FDroppedFrameStats()
-		{
-			FrameType = EFrameType::Undefined;
-			Reset();
-		}
-		void SetFrameType(EFrameType InFrameType)
-		{
-			FrameType = InFrameType;
-		}
-		void Reset()
-		{
-			NumTotalDropped = 0;
-			WorstDeltaTime = FTimespan::Zero();
-			PlayerTimeAtLastReport = FTimespan::MinValue();
-			SystemTimeAtLastReport = 0.0;
-			LogWarningAfterSeconds = 10.0;
-		}
-		void SetLogWarningInterval(double InSecondsBetweenWarnings)
-		{
-			LogWarningAfterSeconds = InSecondsBetweenWarnings;
-		}
-		void AddNewDrop(const FTimespan& InFrameTime, const FTimespan& InPlayerTime, void* InPtrElectraPlayer, void* InPtrCurrentPlayer);
-
-		EFrameType	FrameType;
-		uint32		NumTotalDropped;
-		FTimespan	WorstDeltaTime;
-		FTimespan	PlayerTimeAtLastReport;
-		double		SystemTimeAtLastReport;
-		double		LogWarningAfterSeconds;
-	};
-
 	struct FStatistics
 	{
 		struct FBandwidth
@@ -681,8 +642,6 @@ private:
 		};
 		FStatistics()
 		{
-			DroppedVideoFrames.SetFrameType(FDroppedFrameStats::EFrameType::VideoFrame);
-			DroppedAudioFrames.SetFrameType(FDroppedFrameStats::EFrameType::AudioFrame);
 			Reset();
 		}
 		void Reset()
@@ -720,14 +679,10 @@ private:
 			CurrentlyActiveResolutionWidth = 0;
 			CurrentlyActiveResolutionHeight = 0;
 			VideoSegmentBitratesStreamed.Empty();
+			AudioSegmentBitratesStreamed.Empty();
 			InitialBufferingBandwidth.Reset();
 			bIsInitiallyDownloading = false;
 			bDidPlaybackEnd = false;
-			SubtitlesURL.Empty();
-			SubtitlesResponseTime = 0.0;
-			SubtitlesLastError.Empty();
-			DroppedVideoFrames.Reset();
-			DroppedAudioFrames.Reset();
 			MediaTimelineAtStart.Reset();
 			MediaTimelineAtEnd.Reset();
 			MediaDuration = 0.0;
@@ -768,14 +723,10 @@ private:
 		int32					CurrentlyActiveResolutionWidth;
 		int32					CurrentlyActiveResolutionHeight;
 		TMap<int32, uint32>		VideoSegmentBitratesStreamed;		// key=video stream bitrate, value=number of segments loaded at this rate
+		TMap<int32, uint32>		AudioSegmentBitratesStreamed;		// key=audioo stream bitrate, value=number of segments loaded at this rate
 		FBandwidth				InitialBufferingBandwidth;
 		bool					bIsInitiallyDownloading;
 		bool					bDidPlaybackEnd;
-		FString					SubtitlesURL;
-		double					SubtitlesResponseTime;
-		FString					SubtitlesLastError;
-		FDroppedFrameStats		DroppedVideoFrames;
-		FDroppedFrameStats		DroppedAudioFrames;
 		FTimeRange				MediaTimelineAtStart;
 		FTimeRange				MediaTimelineAtEnd;
 		double					MediaDuration;
diff --git a/Engine/Plugins/Media/ElectraPlayer/Source/ElectraPlayerRuntime/Private/Runtime/HTTP/HTTPManager.h b/Engine/Plugins/Media/ElectraPlayer/Source/ElectraPlayerRuntime/Private/Runtime/HTTP/HTTPManager.h
--- a/Engine/Plugins/Media/ElectraPlayer/Source/ElectraPlayerRuntime/Private/Runtime/HTTP/HTTPManager.h
+++ b/Engine/Plugins/Media/ElectraPlayer/Source/ElectraPlayerRuntime/Private/Runtime/HTTP/HTTPManager.h
@@ -48,19 +48,6 @@ namespace Electra
 			bool			bReadError = false;
 		};
 
-		struct FRetryInfo
-		{
-			FRetryInfo()
-			{
-				AttemptNumber = 0;
-				MaxAttempts = 0;
-			}
-			int32						AttemptNumber;
-			int32						MaxAttempts;
-			TArray<FStatusInfo>			PreviousFailureStates;
-		};
-
-
 		struct FHTTPHeader : public FElectraHTTPStreamHeader
 		{
 			FHTTPHeader() = default;
@@ -108,7 +95,6 @@ namespace Electra
 			bool								bResponseNotRanged = false;			//!< true if the response is not a range as was requested.
 			bool								bIsCachedResponse = false;			//!< true if the response came from the cache.
 			FStatusInfo							StatusInfo;
-			TSharedPtrTS<FRetryInfo>			RetryInfo;
 
 			mutable FCriticalSection			Lock;
 			TArray<IElectraHTTPStreamResponse::FTimingTrace> TimingTraces;
@@ -162,10 +148,6 @@ namespace Electra
 					bResponseNotRanged = rhs.bResponseNotRanged;
 					bIsCachedResponse = rhs.bIsCachedResponse;
 					StatusInfo = rhs.StatusInfo;
-					if (rhs.RetryInfo.IsValid())
-					{
-						RetryInfo = MakeSharedTS<FRetryInfo>(*rhs.RetryInfo);
-					}
 					TimingTraces = rhs.TimingTraces;
 					return *this;
 				}
diff --git a/Engine/Plugins/Media/ElectraPlayer/Source/ElectraPlayerRuntime/Private/Runtime/Player/AdaptiveStreamingPlayer.cpp b/Engine/Plugins/Media/ElectraPlayer/Source/ElectraPlayerRuntime/Private/Runtime/Player/AdaptiveStreamingPlayer.cpp
--- a/Engine/Plugins/Media/ElectraPlayer/Source/ElectraPlayerRuntime/Private/Runtime/Player/AdaptiveStreamingPlayer.cpp
+++ b/Engine/Plugins/Media/ElectraPlayer/Source/ElectraPlayerRuntime/Private/Runtime/Player/AdaptiveStreamingPlayer.cpp
@@ -2041,12 +2041,12 @@ void FAdaptiveStreamingPlayer::HandleSessionMessage(TSharedPtrTS<IPlayerMessage>
 		IPlaylistReader::PlaylistDownloadMessage* pMsg = static_cast<IPlaylistReader::PlaylistDownloadMessage *>(SessionMessage.Get());
 		Metrics::FPlaylistDownloadStats stats;
 		stats.bWasSuccessful = !pMsg->GetConnectionInfo().StatusInfo.ErrorDetail.IsSet();
-		stats.ListType  	 = pMsg->GetListType();
-		stats.LoadType  	 = pMsg->GetLoadType();
-		stats.URL   		 = pMsg->GetConnectionInfo().EffectiveURL;
-		stats.FailureReason  = pMsg->GetConnectionInfo().StatusInfo.ErrorDetail.GetMessage();
+		stats.ListType = pMsg->GetListType();
+		stats.LoadType = pMsg->GetLoadType();
+		stats.URL = pMsg->GetConnectionInfo().EffectiveURL;
+		stats.FailureReason = pMsg->GetConnectionInfo().StatusInfo.ErrorDetail.GetMessage();
 		stats.HTTPStatusCode = pMsg->GetConnectionInfo().StatusInfo.HTTPStatus;
-		stats.RetryNumber    = pMsg->GetConnectionInfo().RetryInfo.IsValid() ? pMsg->GetConnectionInfo().RetryInfo->AttemptNumber : 0;
+		stats.RetryNumber = Utils::Max(pMsg->GetAttempts() - 1, 0);
 		DispatchEvent(FMetricEvent::ReportPlaylistDownload(stats));
 	}
 	// Playlist fetched & parsed?
diff --git a/Engine/Plugins/Media/ElectraPlayer/Source/ElectraPlayerRuntime/Private/Runtime/Player/DASH/PlaylistReaderDASH.cpp b/Engine/Plugins/Media/ElectraPlayer/Source/ElectraPlayerRuntime/Private/Runtime/Player/DASH/PlaylistReaderDASH.cpp
--- a/Engine/Plugins/Media/ElectraPlayer/Source/ElectraPlayerRuntime/Private/Runtime/Player/DASH/PlaylistReaderDASH.cpp
+++ b/Engine/Plugins/Media/ElectraPlayer/Source/ElectraPlayerRuntime/Private/Runtime/Player/DASH/PlaylistReaderDASH.cpp
@@ -468,7 +468,7 @@ void FPlaylistReaderDASH::CheckForMPDUpdate()
 						bRequestNow = false;
 						UpdateRequested.Empty();
 						// Send a fake update message to the player to have it evaluate the media timeline again.
-						PlayerSessionServices->SendMessageToPlayer(IPlaylistReader::PlaylistLoadedMessage::Create(FErrorDetail(), nullptr, Playlist::EListType::Main, Playlist::ELoadType::Update));
+						PlayerSessionServices->SendMessageToPlayer(IPlaylistReader::PlaylistLoadedMessage::Create(FErrorDetail(), nullptr, Playlist::EListType::Main, Playlist::ELoadType::Update, 0));
 					}
 				}
 				else
@@ -1165,7 +1165,7 @@ void FPlaylistReaderDASH::ManifestDownloadCompleted(FResourceLoadRequestPtr Requ
 			}
 		}
 
-		PlayerSessionServices->SendMessageToPlayer(IPlaylistReader::PlaylistDownloadMessage::Create(ConnInfo, Playlist::EListType::Main, Playlist::ELoadType::Initial));
+		PlayerSessionServices->SendMessageToPlayer(IPlaylistReader::PlaylistDownloadMessage::Create(ConnInfo, Playlist::EListType::Main, Playlist::ELoadType::Initial, Request->Attempt));
 
 		FString XML;
 		LastErrorDetail = GetXMLResponseString(XML, Request);
@@ -1317,11 +1317,12 @@ void FPlaylistReaderDASH::FinishManifestBuildingAfterTimesync(bool bGotTheTime)
 	}
 
 	// Notify that the "main playlist" has been parsed, successfully or not. If we still need to resolve remote entities this is not an error!
-	PlayerSessionServices->SendMessageToPlayer(IPlaylistReader::PlaylistLoadedMessage::Create(!LastErrorDetail.IsTryAgain() ? LastErrorDetail : FErrorDetail(), InitialTimesync.ConnInfo.Get(), Playlist::EListType::Main, Playlist::ELoadType::Initial));
+	const int32 Attempts = 0; // There are no actual attempts made here, so send '0'.
+	PlayerSessionServices->SendMessageToPlayer(IPlaylistReader::PlaylistLoadedMessage::Create(!LastErrorDetail.IsTryAgain() ? LastErrorDetail : FErrorDetail(), InitialTimesync.ConnInfo.Get(), Playlist::EListType::Main, Playlist::ELoadType::Initial, Attempts));
 	if (LastErrorDetail.IsOK())
 	{
 		// Notify that the "variant playlists" are ready. There are no variants in DASH, but this is the trigger that the playlists are all set up and are good to go now.
-		PlayerSessionServices->SendMessageToPlayer(IPlaylistReader::PlaylistLoadedMessage::Create(LastErrorDetail, InitialTimesync.ConnInfo.Get(), Playlist::EListType::Variant, Playlist::ELoadType::Initial));
+		PlayerSessionServices->SendMessageToPlayer(IPlaylistReader::PlaylistLoadedMessage::Create(LastErrorDetail, InitialTimesync.ConnInfo.Get(), Playlist::EListType::Variant, Playlist::ELoadType::Initial, Attempts));
 		PrepareNextTimeSync();
 	}
 	// "try again" is returned when there are remote entities that need to be resolved first.
@@ -1395,7 +1396,8 @@ void FPlaylistReaderDASH::ManifestUpdateDownloadCompleted(FResourceLoadRequestPt
 		}
 		else
 		{
-			PlayerSessionServices->SendMessageToPlayer(IPlaylistReader::PlaylistDownloadMessage::Create(ConnInfo, Playlist::EListType::Main, Playlist::ELoadType::Update));
+			const int32 Attempt = Request->Attempt;
+			PlayerSessionServices->SendMessageToPlayer(IPlaylistReader::PlaylistDownloadMessage::Create(ConnInfo, Playlist::EListType::Main, Playlist::ELoadType::Update, Attempt));
 
 			FString XML;
 			LastErrorDetail = GetXMLResponseString(XML, Request);
@@ -1441,7 +1443,7 @@ void FPlaylistReaderDASH::ManifestUpdateDownloadCompleted(FResourceLoadRequestPt
 						RequestsLock.Unlock();
 					}
 
-					PlayerSessionServices->SendMessageToPlayer(IPlaylistReader::PlaylistLoadedMessage::Create(!LastErrorDetail.IsTryAgain() ? LastErrorDetail : FErrorDetail(), ConnInfo, Playlist::EListType::Main, Playlist::ELoadType::Update));
+					PlayerSessionServices->SendMessageToPlayer(IPlaylistReader::PlaylistLoadedMessage::Create(!LastErrorDetail.IsTryAgain() ? LastErrorDetail : FErrorDetail(), ConnInfo, Playlist::EListType::Main, Playlist::ELoadType::Update, Attempt));
 					if (LastErrorDetail.IsOK())
 					{
 					}
@@ -1488,7 +1490,7 @@ void FPlaylistReaderDASH::InitialMPDXLinkElementDownloadCompleted(FResourceLoadR
 		{
 			LastErrorDetail = ForManifest->BuildAfterInitialRemoteElementDownload();
 			// Notify that the "variant playlists" are ready. There are no variants in DASH, but this is the trigger that the playlists are all set up and are good to go now.
-			PlayerSessionServices->SendMessageToPlayer(IPlaylistReader::PlaylistLoadedMessage::Create(LastErrorDetail, nullptr, Playlist::EListType::Variant, Playlist::ELoadType::Initial));
+			PlayerSessionServices->SendMessageToPlayer(IPlaylistReader::PlaylistLoadedMessage::Create(LastErrorDetail, nullptr, Playlist::EListType::Variant, Playlist::ELoadType::Initial, Request->Attempt));
 			PrepareNextTimeSync();
 		}
 		else if (LastErrorDetail.IsTryAgain())
diff --git a/Engine/Plugins/Media/ElectraPlayer/Source/ElectraPlayerRuntime/Private/Runtime/Player/HLS/PlaylistHLS.h b/Engine/Plugins/Media/ElectraPlayer/Source/ElectraPlayerRuntime/Private/Runtime/Player/HLS/PlaylistHLS.h
--- a/Engine/Plugins/Media/ElectraPlayer/Source/ElectraPlayerRuntime/Private/Runtime/Player/HLS/PlaylistHLS.h
+++ b/Engine/Plugins/Media/ElectraPlayer/Source/ElectraPlayerRuntime/Private/Runtime/Player/HLS/PlaylistHLS.h
@@ -359,7 +359,7 @@ struct FLoadRequestHLSPlaylist
 	ELoadType LoadType = ELoadType::Undefined;
 	TSharedPtr<FHTTPResourceRequest, ESPMode::ThreadSafe> ResourceRequest;
 	FTimeValue ExecuteAtUTC;
-	int32 Attempt = 0;
+	int32 Attempt = 1;
 	bool bIsPrimaryPlaylist = false;
 
 	FMediaPlaylistInformationHLS PlaylistInfo;
diff --git a/Engine/Plugins/Media/ElectraPlayer/Source/ElectraPlayerRuntime/Private/Runtime/Player/HLS/PlaylistHandlerHLS.cpp b/Engine/Plugins/Media/ElectraPlayer/Source/ElectraPlayerRuntime/Private/Runtime/Player/HLS/PlaylistHandlerHLS.cpp
--- a/Engine/Plugins/Media/ElectraPlayer/Source/ElectraPlayerRuntime/Private/Runtime/Player/HLS/PlaylistHandlerHLS.cpp
+++ b/Engine/Plugins/Media/ElectraPlayer/Source/ElectraPlayerRuntime/Private/Runtime/Player/HLS/PlaylistHandlerHLS.cpp
@@ -254,7 +254,7 @@ void FPlaylistHandlerHLS::HandleOnce()
 					MediaPlaylist->SetPlaylist(NewMediaPlaylist, Now);
 					if (bGotUpdate)
 					{
-						PlayerSessionServices->SendMessageToPlayer(IPlaylistReader::PlaylistLoadedMessage::Create(LastError, FinishedReq->ResourceRequest->GetConnectionInfo(), Playlist::EListType::Variant, Playlist::ELoadType::Update));
+						PlayerSessionServices->SendMessageToPlayer(IPlaylistReader::PlaylistLoadedMessage::Create(LastError, FinishedReq->ResourceRequest->GetConnectionInfo(), Playlist::EListType::Variant, Playlist::ELoadType::Update, FinishedReq->Attempt));
 					}
 					ActivePlaylist->UpdateWithMediaPlaylist(MediaPlaylist, FinishedReq->bIsPrimaryPlaylist, bGotUpdate);
 				}
@@ -297,7 +297,7 @@ void FPlaylistHandlerHLS::HandleOnce()
 				// When retrying there is no LastError, but also no MainPlaylist either!
 				else if (MainPlaylist.IsValid())
 				{
-					PlayerSessionServices->SendMessageToPlayer(IPlaylistReader::PlaylistLoadedMessage::Create(LastError, FinishedLoadRequests[i]->ResourceRequest->GetConnectionInfo(), Playlist::EListType::Main, Playlist::ELoadType::Initial));
+					PlayerSessionServices->SendMessageToPlayer(IPlaylistReader::PlaylistLoadedMessage::Create(LastError, FinishedLoadRequests[i]->ResourceRequest->GetConnectionInfo(), Playlist::EListType::Main, Playlist::ELoadType::Initial, FinishedLoadRequests[i]->Attempt));
 
 					if (BuildPlaylist(MainPlaylist, true))
 					{
@@ -344,7 +344,7 @@ void FPlaylistHandlerHLS::HandleOnce()
 							ABR->MarkStreamAsUnavailable(ds);
 						}
 					}
-					PlayerSessionServices->SendMessageToPlayer(IPlaylistReader::PlaylistLoadedMessage::Create(LastError, FinishedLoadRequests[i]->ResourceRequest->GetConnectionInfo(), Playlist::EListType::Variant, Playlist::ELoadType::Initial));
+					PlayerSessionServices->SendMessageToPlayer(IPlaylistReader::PlaylistLoadedMessage::Create(LastError, FinishedLoadRequests[i]->ResourceRequest->GetConnectionInfo(), Playlist::EListType::Variant, Playlist::ELoadType::Initial, FinishedLoadRequests[i]->Attempt));
 				}
 				if (!LastError.IsOK())
 				{
@@ -446,7 +446,6 @@ void FPlaylistHandlerHLS::LoadMainPlaylist(const FString& InUrl)
 		.ConnectionTimeout(GetTimeoutValue(HLS::OptionKeyPlaylistLoadConnectTimeout, 5000))
 		.NoDataTimeout(GetTimeoutValue(HLS::OptionKeyPlaylistLoadNoDataTimeout, 2000))
 		.AllowStaticQuery(IAdaptiveStreamingPlayerResourceRequest::EPlaybackResourceType::Playlist);
-	lr->Attempt = 1;
 	PendingLoadRequests.Emplace(MoveTemp(lr));
 }
 
@@ -468,7 +467,6 @@ TSharedPtr<FLoadRequestHLSPlaylist, ESPMode::ThreadSafe> FPlaylistHandlerHLS::Cr
 		.ConnectionTimeout(GetTimeoutValue(HLS::OptionKeyPlaylistLoadConnectTimeout, 5000))
 		.NoDataTimeout(GetTimeoutValue(HLS::OptionKeyPlaylistLoadNoDataTimeout, 2000))
 		.AllowStaticQuery(IAdaptiveStreamingPlayerResourceRequest::EPlaybackResourceType::Playlist);
-	lr->Attempt = 1;
 	return lr;
 }
 
@@ -522,7 +520,7 @@ FTimeValue FPlaylistHandlerHLS::GetTimeoutValue(const FName& InOptionName, int32
 TSharedPtr<FPlaylistParserHLS, ESPMode::ThreadSafe> FPlaylistHandlerHLS::HandleMainPlaylist(const FTimeValue& Now, TSharedPtr<FLoadRequestHLSPlaylist, ESPMode::ThreadSafe> InRequest)
 {
 	// Notify download completion.
-	PlayerSessionServices->SendMessageToPlayer(IPlaylistReader::PlaylistDownloadMessage::Create(InRequest->ResourceRequest->GetConnectionInfo(), Playlist::EListType::Main, Playlist::ELoadType::Initial));
+	PlayerSessionServices->SendMessageToPlayer(IPlaylistReader::PlaylistDownloadMessage::Create(InRequest->ResourceRequest->GetConnectionInfo(), Playlist::EListType::Main, Playlist::ELoadType::Initial, InRequest->Attempt));
 
 	FString PlaylistURL = InRequest->ResourceRequest->GetURL();
 	// Failure?
@@ -577,7 +575,7 @@ TSharedPtr<FPlaylistParserHLS, ESPMode::ThreadSafe> FPlaylistHandlerHLS::HandleM
 TSharedPtr<FPlaylistParserHLS, ESPMode::ThreadSafe> FPlaylistHandlerHLS::HandleVariantPlaylist(const FTimeValue& Now, TSharedPtr<FLoadRequestHLSPlaylist, ESPMode::ThreadSafe> InRequest)
 {
 	// Notify download completion.
-	PlayerSessionServices->SendMessageToPlayer(IPlaylistReader::PlaylistDownloadMessage::Create(InRequest->ResourceRequest->GetConnectionInfo(), Playlist::EListType::Variant, InRequest->UpdateRequestFor.IsValid() ? Playlist::ELoadType::Update : Playlist::ELoadType::Initial));
+	PlayerSessionServices->SendMessageToPlayer(IPlaylistReader::PlaylistDownloadMessage::Create(InRequest->ResourceRequest->GetConnectionInfo(), Playlist::EListType::Variant, InRequest->UpdateRequestFor.IsValid() ? Playlist::ELoadType::Update : Playlist::ELoadType::Initial, InRequest->Attempt));
 
 	FString PlaylistURL = InRequest->ResourceRequest->GetURL();
 	// Failure?
diff --git a/Engine/Plugins/Media/ElectraPlayer/Source/ElectraPlayerRuntime/Private/Runtime/Player/PlaylistReader.h b/Engine/Plugins/Media/ElectraPlayer/Source/ElectraPlayerRuntime/Private/Runtime/Player/PlaylistReader.h
--- a/Engine/Plugins/Media/ElectraPlayer/Source/ElectraPlayerRuntime/Private/Runtime/Player/PlaylistReader.h
+++ b/Engine/Plugins/Media/ElectraPlayer/Source/ElectraPlayerRuntime/Private/Runtime/Player/PlaylistReader.h
@@ -63,116 +63,104 @@ namespace Electra
 		class PlaylistDownloadMessage : public IPlayerMessage
 		{
 		public:
-			static TSharedPtrTS<IPlayerMessage> Create(const HTTP::FConnectionInfo* ConnectionInfo, Playlist::EListType InListType, Playlist::ELoadType InLoadType)
+			static TSharedPtrTS<IPlayerMessage> Create(const HTTP::FConnectionInfo* ConnectionInfo, Playlist::EListType InListType, Playlist::ELoadType InLoadType, int32 InAttempts)
 			{
-				TSharedPtrTS<PlaylistDownloadMessage> p(new PlaylistDownloadMessage(ConnectionInfo, InListType, InLoadType));
+				TSharedPtrTS<PlaylistDownloadMessage> p(new PlaylistDownloadMessage(ConnectionInfo, InListType, InLoadType, InAttempts));
 				return p;
 			}
 
 			static const FString& Type()
 			{
-				static FString TypeName("PlaylistDownload");
+				static FString TypeName(TEXT("PlaylistDownload"));
 				return TypeName;
 			}
 
 			virtual const FString& GetType() const
-			{
-				return Type();
-			}
+			{ return Type(); }
 
 			Playlist::EListType GetListType() const
-			{
-				return ListType;
-			}
+			{ return ListType; }
 
 			Playlist::ELoadType GetLoadType() const
-			{
-				return LoadType;
-			}
+			{ return LoadType; }
+
+			int32 GetAttempts() const
+			{ return Attempts; }
 
 			const HTTP::FConnectionInfo& GetConnectionInfo() const
-			{
-				return ConnectionInfo;
-			}
+			{ return ConnectionInfo; }
 
 		private:
-			PlaylistDownloadMessage(const HTTP::FConnectionInfo* InConnectionInfo, Playlist::EListType InListType, Playlist::ELoadType InLoadType)
+			PlaylistDownloadMessage(const HTTP::FConnectionInfo* InConnectionInfo, Playlist::EListType InListType, Playlist::ELoadType InLoadType, int32 InAttempts)
 				: ListType(InListType)
 				, LoadType(InLoadType)
+				, Attempts(InAttempts)
 			{
 				if (InConnectionInfo)
 				{
-					// Have to make a dedicated copy of the connection info in order to get a copy of the retry info at this point in time.
 					ConnectionInfo = *InConnectionInfo;
 				}
 			}
-			HTTP::FConnectionInfo	ConnectionInfo;
-			Playlist::EListType		ListType;
-			Playlist::ELoadType		LoadType;
+			HTTP::FConnectionInfo ConnectionInfo;
+			Playlist::EListType ListType;
+			Playlist::ELoadType LoadType;
+			int32 Attempts;
 		};
 
 
 		class PlaylistLoadedMessage : public IPlayerMessage
 		{
 		public:
-			static TSharedPtrTS<IPlayerMessage> Create(const FErrorDetail& PlayerResult, const HTTP::FConnectionInfo* ConnectionInfo, Playlist::EListType InListType, Playlist::ELoadType InLoadType)
+			static TSharedPtrTS<IPlayerMessage> Create(const FErrorDetail& PlayerResult, const HTTP::FConnectionInfo* ConnectionInfo, Playlist::EListType InListType, Playlist::ELoadType InLoadType, int32 InAttempts)
 			{
-				TSharedPtrTS<PlaylistLoadedMessage> p(new PlaylistLoadedMessage(PlayerResult, ConnectionInfo, InListType, InLoadType));
+				TSharedPtrTS<PlaylistLoadedMessage> p(new PlaylistLoadedMessage(PlayerResult, ConnectionInfo, InListType, InLoadType, InAttempts));
 				return p;
 			}
 
 			static const FString& Type()
 			{
-				static FString TypeName("PlaylistLoaded");
+				static FString TypeName(TEXT("PlaylistLoaded"));
 				return TypeName;
 			}
 
 			virtual const FString& GetType() const
-			{
-				return Type();
-			}
+			{ return Type(); }
 
 			const FErrorDetail& GetResult() const
-			{
-				return Result;
-			}
+			{ return Result; }
 
 			Playlist::EListType GetListType() const
-			{
-				return ListType;
-			}
+			{ return ListType; }
 
 			Playlist::ELoadType GetLoadType() const
-			{
-				return LoadType;
-			}
+			{ return LoadType; }
+
+			int32 GetAttempts() const
+			{ return Attempts; }
 
 			const HTTP::FConnectionInfo& GetConnectionInfo() const
-			{
-				return ConnectionInfo;
-			}
+			{ return ConnectionInfo; }
 
 		private:
-			PlaylistLoadedMessage(const FErrorDetail& PlayerResult, const HTTP::FConnectionInfo* InConnectionInfo, Playlist::EListType InListType, Playlist::ELoadType InLoadType)
+			PlaylistLoadedMessage(const FErrorDetail& PlayerResult, const HTTP::FConnectionInfo* InConnectionInfo, Playlist::EListType InListType, Playlist::ELoadType InLoadType, int32 InAttempts)
 				: Result(PlayerResult)
 				, ListType(InListType)
 				, LoadType(InLoadType)
+				, Attempts(InAttempts)
 			{
 				if (InConnectionInfo)
 				{
-					// Have to make a dedicated copy of the connection info in order to get a copy of the retry info at this point in time.
 					ConnectionInfo = *InConnectionInfo;
 				}
 			}
-			HTTP::FConnectionInfo	ConnectionInfo;
-			FErrorDetail			Result;
-			Playlist::EListType		ListType;
-			Playlist::ELoadType		LoadType;
+			HTTP::FConnectionInfo ConnectionInfo;
+			FErrorDetail Result;
+			Playlist::EListType ListType;
+			Playlist::ELoadType LoadType;
+			int32 Attempts;
 		};
 
 	};
 
 
 } // namespace Electra
-
-
diff --git a/Engine/Plugins/Media/ElectraPlayer/Source/ElectraPlayerRuntime/Private/Runtime/Player/mkv/PlaylistReaderMKV.cpp b/Engine/Plugins/Media/ElectraPlayer/Source/ElectraPlayerRuntime/Private/Runtime/Player/mkv/PlaylistReaderMKV.cpp
--- a/Engine/Plugins/Media/ElectraPlayer/Source/ElectraPlayerRuntime/Private/Runtime/Player/mkv/PlaylistReaderMKV.cpp
+++ b/Engine/Plugins/Media/ElectraPlayer/Source/ElectraPlayerRuntime/Private/Runtime/Player/mkv/PlaylistReaderMKV.cpp
@@ -338,10 +338,12 @@ void FPlaylistReaderMKV::WorkerThread()
 
 	if (!HasReadBeenAborted())
 	{
+		// There are currently no retries, so this is the first (and only) attempt we make.
+		const int32 Attempt = 1;
 		// Notify the download of the "main playlist". This indicates the download only, not the parsing thereof.
-		PlayerSessionServices->SendMessageToPlayer(IPlaylistReader::PlaylistDownloadMessage::Create(&ConnectionInfo, Playlist::EListType::Main, Playlist::ELoadType::Initial));
+		PlayerSessionServices->SendMessageToPlayer(IPlaylistReader::PlaylistDownloadMessage::Create(&ConnectionInfo, Playlist::EListType::Main, Playlist::ELoadType::Initial, Attempt));
 		// Notify that the "main playlist" has been parsed, successfully or not.
-		PlayerSessionServices->SendMessageToPlayer(IPlaylistReader::PlaylistLoadedMessage::Create(LastErrorDetail, &ConnectionInfo, Playlist::EListType::Main, Playlist::ELoadType::Initial));
+		PlayerSessionServices->SendMessageToPlayer(IPlaylistReader::PlaylistLoadedMessage::Create(LastErrorDetail, &ConnectionInfo, Playlist::EListType::Main, Playlist::ELoadType::Initial, Attempt));
 
 		if (LastErrorDetail.IsOK())
 		{
@@ -358,7 +360,7 @@ void FPlaylistReaderMKV::WorkerThread()
 				LastErrorDetail = Manifest->Build(MKVParser, PlaylistURL, ConnectionInfo);
 
 				// Notify that the "variant playlists" are ready. There are no variants in an mkv, but this is the trigger that the playlists are all set up and are good to go now.
-				PlayerSessionServices->SendMessageToPlayer(IPlaylistReader::PlaylistLoadedMessage::Create(LastErrorDetail, &ConnectionInfo, Playlist::EListType::Variant, Playlist::ELoadType::Initial));
+				PlayerSessionServices->SendMessageToPlayer(IPlaylistReader::PlaylistLoadedMessage::Create(LastErrorDetail, &ConnectionInfo, Playlist::EListType::Variant, Playlist::ELoadType::Initial, Attempt));
 			}
 			else
 			{
diff --git a/Engine/Plugins/Media/ElectraPlayer/Source/ElectraPlayerRuntime/Private/Runtime/Player/mp4/PlaylistReaderMP4.cpp b/Engine/Plugins/Media/ElectraPlayer/Source/ElectraPlayerRuntime/Private/Runtime/Player/mp4/PlaylistReaderMP4.cpp
--- a/Engine/Plugins/Media/ElectraPlayer/Source/ElectraPlayerRuntime/Private/Runtime/Player/mp4/PlaylistReaderMP4.cpp
+++ b/Engine/Plugins/Media/ElectraPlayer/Source/ElectraPlayerRuntime/Private/Runtime/Player/mp4/PlaylistReaderMP4.cpp
@@ -243,10 +243,12 @@ void FPlaylistReaderMP4::WorkerThread()
 		PostError(BoxLocator.GetErrorMessage(), ERRCODE_MP4_INVALID_FILE, UEMEDIA_ERROR_FORMAT_ERROR);
 	}
 
+	// There are currently no retries, so this is the first (and only) attempt we make.
+	const int32 Attempt = 1;
 	// Notify the download of the "main playlist". This indicates the download only, not the parsing thereof.
-	PlayerSessionServices->SendMessageToPlayer(IPlaylistReader::PlaylistDownloadMessage::Create(&ConnectionInfo, Playlist::EListType::Main, Playlist::ELoadType::Initial));
+	PlayerSessionServices->SendMessageToPlayer(IPlaylistReader::PlaylistDownloadMessage::Create(&ConnectionInfo, Playlist::EListType::Main, Playlist::ELoadType::Initial, Attempt));
 	// Notify that the "main playlist" has been parsed, successfully or not.
-	PlayerSessionServices->SendMessageToPlayer(IPlaylistReader::PlaylistLoadedMessage::Create(LastErrorDetail, &ConnectionInfo, Playlist::EListType::Main, Playlist::ELoadType::Initial));
+	PlayerSessionServices->SendMessageToPlayer(IPlaylistReader::PlaylistLoadedMessage::Create(LastErrorDetail, &ConnectionInfo, Playlist::EListType::Main, Playlist::ELoadType::Initial, Attempt));
 	// Failed to get the boxes but was not aborted?
 	if (!bAbort && (!bGotBoxes || bHasErrored))
 	{
@@ -293,7 +295,7 @@ void FPlaylistReaderMP4::WorkerThread()
 					LastErrorDetail = Manifest->Build(MP4Parser, MainPlaylistURL, ConnectionInfo);
 
 					// Notify that the "variant playlists" are ready. There are no variants in an mp4, but this is the trigger that the playlists are all set up and are good to go now.
-					PlayerSessionServices->SendMessageToPlayer(IPlaylistReader::PlaylistLoadedMessage::Create(LastErrorDetail, &ConnectionInfo, Playlist::EListType::Variant, Playlist::ELoadType::Initial));
+					PlayerSessionServices->SendMessageToPlayer(IPlaylistReader::PlaylistLoadedMessage::Create(LastErrorDetail, &ConnectionInfo, Playlist::EListType::Variant, Playlist::ELoadType::Initial, Attempt));
 				}
 				else
 				{
diff --git a/Engine/Plugins/Media/ElectraPlayer/Source/ElectraPlayerRuntime/Private/Runtime/Player/mpegaudio/PlaylistReaderMPEGAudio.cpp b/Engine/Plugins/Media/ElectraPlayer/Source/ElectraPlayerRuntime/Private/Runtime/Player/mpegaudio/PlaylistReaderMPEGAudio.cpp
--- a/Engine/Plugins/Media/ElectraPlayer/Source/ElectraPlayerRuntime/Private/Runtime/Player/mpegaudio/PlaylistReaderMPEGAudio.cpp
+++ b/Engine/Plugins/Media/ElectraPlayer/Source/ElectraPlayerRuntime/Private/Runtime/Player/mpegaudio/PlaylistReaderMPEGAudio.cpp
@@ -914,10 +914,12 @@ void FPlaylistReaderMPEGAudio::WorkerThread()
 	ConnectionInfo.StatusInfo.HTTPStatus = Req->StatusCode;
 	ConnectionInfo.StatusInfo.bReadError = Req->Result == FRequest::EResult::Failed;
 
+	// There are currently no retries, so this is the first (and only) attempt we make.
+	const int32 Attempt = 1;
 	// Notify the download of the "main playlist". This indicates the download only, not the parsing thereof.
-	PlayerSessionServices->SendMessageToPlayer(IPlaylistReader::PlaylistDownloadMessage::Create(&ConnectionInfo, Playlist::EListType::Main, Playlist::ELoadType::Initial));
+	PlayerSessionServices->SendMessageToPlayer(IPlaylistReader::PlaylistDownloadMessage::Create(&ConnectionInfo, Playlist::EListType::Main, Playlist::ELoadType::Initial, Attempt));
 	// Notify that the "main playlist" has been parsed, successfully or not.
-	PlayerSessionServices->SendMessageToPlayer(IPlaylistReader::PlaylistLoadedMessage::Create(LastErrorDetail, &ConnectionInfo, Playlist::EListType::Main, Playlist::ELoadType::Initial));
+	PlayerSessionServices->SendMessageToPlayer(IPlaylistReader::PlaylistLoadedMessage::Create(LastErrorDetail, &ConnectionInfo, Playlist::EListType::Main, Playlist::ELoadType::Initial, Attempt));
 	if (!bAbort)
 	{
 		// Error?
@@ -948,7 +950,7 @@ void FPlaylistReaderMPEGAudio::WorkerThread()
 				LastErrorDetail = Manifest->Build(InfoHeader, PlaylistURL);
 
 				// Notify that the "variant playlists" are ready. There are no variants in an mp4, but this is the trigger that the playlists are all set up and are good to go now.
-				PlayerSessionServices->SendMessageToPlayer(IPlaylistReader::PlaylistLoadedMessage::Create(LastErrorDetail, &ConnectionInfo, Playlist::EListType::Variant, Playlist::ELoadType::Initial));
+				PlayerSessionServices->SendMessageToPlayer(IPlaylistReader::PlaylistLoadedMessage::Create(LastErrorDetail, &ConnectionInfo, Playlist::EListType::Variant, Playlist::ELoadType::Initial, Attempt));
 			}
 			else
 			{
