ElectraPlayer: Removal of scrub optimizations that turned out not to be useful.

#rnx

diff --git a/Engine/Plugins/Media/ElectraPlayer/Source/ElectraPlayerPlugin/Private/ElectraPlayerPlugin.cpp b/Engine/Plugins/Media/ElectraPlayer/Source/ElectraPlayerPlugin/Private/ElectraPlayerPlugin.cpp
--- a/Engine/Plugins/Media/ElectraPlayer/Source/ElectraPlayerPlugin/Private/ElectraPlayerPlugin.cpp
+++ b/Engine/Plugins/Media/ElectraPlayer/Source/ElectraPlayerPlugin/Private/ElectraPlayerPlugin.cpp
@@ -415,7 +415,6 @@ void FElectraPlayerPlugin::PresentVideoFrame(const FVideoDecoderOutputPtr& InVid
 	TSharedPtr<FElectraTextureSamplePool, ESPMode::ThreadSafe> TexturePool = OutputTexturePool;
 	if (VideoFrame.IsValid() && TexturePool.IsValid())
 	{
-		SequenceIndexMapperVideo.Remap(VideoFrame);
 		FElectraTextureSampleRef TextureSample = TexturePool->AcquireShared();
 		TextureSample->Initialize(VideoFrame.Get());
 		MediaSamples->AddVideo(TextureSample);
@@ -430,7 +429,6 @@ void FElectraPlayerPlugin::PresentAudioFrame(const IAudioDecoderOutputPtr& InAud
 	IAudioDecoderOutputPtr AudioFrame = InAudioFrame;
 	if (AudioFrame.IsValid())
 	{
-		SequenceIndexMapperAudio.Remap(AudioFrame);
 		TSharedRef<FElectraPlayerAudioSample, ESPMode::ThreadSafe> AudioSample = OutputAudioPool.AcquireShared();
 		AudioSample->Initialize(InAudioFrame);
 		MediaSamples->AddAudio(AudioSample);
@@ -445,7 +443,6 @@ void FElectraPlayerPlugin::PresentSubtitleSample(const ISubtitleDecoderOutputPtr
 	ISubtitleDecoderOutputPtr Subtitle = InSubtitleSample;
 	if (Subtitle.IsValid())
 	{
-		SequenceIndexMapperSubtitle.Remap(Subtitle);
 		TSharedRef<FElectraSubtitleSample, ESPMode::ThreadSafe> SubtitleSample = MakeShared<FElectraSubtitleSample, ESPMode::ThreadSafe>();
 		SubtitleSample->Subtitle = InSubtitleSample;
 		MediaSamples->AddSubtitle(SubtitleSample);
@@ -460,7 +457,6 @@ void FElectraPlayerPlugin::PresentMetadataSample(const IMetaDataDecoderOutputPtr
 	IMetaDataDecoderOutputPtr MetadataFrame = InMetadataFrame;
 	if (MetadataFrame.IsValid())
 	{
-		SequenceIndexMapperMetadata.Remap(MetadataFrame);
 		TSharedRef<FElectraBinarySample, ESPMode::ThreadSafe> MetaDataSample = MakeShared<FElectraBinarySample, ESPMode::ThreadSafe>();
 		MetaDataSample->Metadata = InMetadataFrame;
 		MediaSamples->AddMetadata(MetaDataSample);
@@ -677,32 +673,8 @@ bool FElectraPlayerPlugin::Open(const FString& Url, const IMediaOptions* Options
 	// Certain optimization settings apply for dedicated environments.
 	if (Environment == MediaPlayerOptionValues::Environment_Preview() || Environment == MediaPlayerOptionValues::Environment_Sequencer())
 	{
-		PlayerOptions.Set(TEXT("optimize_seek_for_scrubbing"), Electra::FVariantValue(true));
-		//PlayerOptions.Set(TEXT("do_not_hold_back_first_frame"), Electra::FVariantValue(true));
-		//PlayerOptions.Set(TEXT("always_emit_samples_when_paused"), Electra::FVariantValue(true));
 		PlayerOptions.Set(TEXT("worker_threads"), Electra::FVariantValue(FString(TEXT("worker"))));
 	}
-	else
-	{
-		// Check if there are options to tweak for scrubbing
-		bool bUseScrubOptimizations = Options->GetMediaOption(TEXT("ElectraScrubOptimization"), (bool)false);
-		if (bUseScrubOptimizations)
-		{
-			int64 ScrubSeekBitrate = Options->GetMediaOption(TEXT("ElectraScrubBitrate"), (int64)-1);
-			int64 ScrubCacheSizeKiB = Options->GetMediaOption(TEXT("ElectraScrubCacheSizeKiB"), (int64)-1);
-			if (ScrubSeekBitrate >= 0)
-			{
-				PlayerOptions.Set(TEXT("seekstart_bitrate"), Electra::FVariantValue(ScrubSeekBitrate));
-			}
-			if (ScrubCacheSizeKiB > 0)
-			{
-				PlayerOptions.Set(TEXT("httpcache_max_bytesize"), Electra::FVariantValue(ScrubCacheSizeKiB << 10));
-				PlayerOptions.Set(TEXT("httpcache_max_entries"), Electra::FVariantValue((int64)10000));
-			}
-			PlayerOptions.Set(TEXT("optimize_seek_for_scrubbing"), Electra::FVariantValue(true));
-			PlayerOptions.Set(TEXT("do_not_hold_back_first_frame"), Electra::FVariantValue(true));
-		}
-	}
 
 	// Check for options that can be changed during playback and apply them at startup already.
 	// If a media source supports the MaxResolutionForMediaStreaming option then we can override the max resolution.
@@ -724,11 +696,6 @@ bool FElectraPlayerPlugin::Open(const FString& Url, const IMediaOptions* Options
 	bMetadataChanged = false;
 	CurrentMetadata.Reset();
 
-	SequenceIndexMapperVideo.Reset();
-	SequenceIndexMapperAudio.Reset();
-	SequenceIndexMapperSubtitle.Reset();
-	SequenceIndexMapperMetadata.Reset();
-
 	// Check if we can get a segment cache interface for this playback request...
 	TSharedPtr<FElectraPlayerDataCacheContainer, ESPMode::ThreadSafe> ElectraPlayerDataCacheContainer;
 	TSharedPtr<IMediaOptions::FDataContainer, ESPMode::ThreadSafe> ElectraPlayerDataCacheDefaultValue;
@@ -978,112 +945,7 @@ bool FElectraPlayerPlugin::Seek(const FTimespan& Time)
 {
 	UE_LOG(LogElectraPlayerPlugin, Log, TEXT("[%u] IMediaControls::Seek() to %s"), PlayerUniqueID, *Time.ToString(TEXT("%h:%m:%s.%f")));
 	CSV_EVENT(ElectraPlayer, TEXT("Seeking"));
-
-	const int32 NextSequenceIndex = CurrentSequenceIndex + 1;
-
-	// Check if the target time is already available.
-	FScopeLock SampleLock(&MediaSamplesLock);
-	if (!MediaSamples.Get())
-	{
-		return false;
-	}
-	FMediaSamples& Samples = *MediaSamples;
-	const float Rate = Player->GetRate();
-	if (Rate >= 0.0f)
-	{
-		TArray<TRange<FMediaTimeStamp>> QueuedRange;
-		// Note: When a Seek() is performed, the loop counter contractually restarts at 0 !!
-		FMediaTimeStamp TargetTime = FMediaTimeStamp(Time, FMediaTimeStamp::MakeSequenceIndex(CurrentSequenceIndex, 0));
-		if (Samples.PeekVideoSampleTimeRanges(QueuedRange) && QueuedRange.Num())
-		{
-			FMediaTimeStamp MinTime(FTimespan::MinValue(), FMediaTimeStamp::MakeSequenceIndex(-1, 0));
-			FMediaTimeStamp DiscardTo = MinTime;
-			// Loop over all entries in case there are several candidates in a (more or less)
-			// degenerate case of a single looping frame.
-			for(int32 i=0; i<QueuedRange.Num(); ++i)
-			{
-				// For comparison reasons we need to assume the loop counter of the sample range we're looking at.
-				TargetTime.SetSecondaryIndex(QueuedRange[i].GetLowerBoundValue().GetSecondaryIndex());
-				if (QueuedRange[i].Contains(TargetTime))
-				{
-					DiscardTo = QueuedRange[i].GetLowerBoundValue() - FTimespan(1);
-				}
-			}
-			if (DiscardTo > MinTime)
-			{
-				UE_LOG(LogElectraPlayerPlugin, Verbose, TEXT("Seek time %lld (%d,%d) already available, rolling forward"), (long long int)TargetTime.Time.GetTicks(), TargetTime.GetPrimaryIndex(), TargetTime.GetSecondaryIndex());
-				TRange<FMediaTimeStamp> DiscardRange(MinTime, DiscardTo);
-
-				Samples.DiscardVideoSamples(DiscardRange, false);
-				Samples.DiscardAudioSamples(DiscardRange, false);
-				Samples.DiscardCaptionSamples(DiscardRange, false);
-				Samples.DiscardSubtitleSamples(DiscardRange, false);
-				Samples.DiscardMetadataSamples(DiscardRange, false);
-
-				TOptional<int32> LastVideoLoopIndex;
-				TOptional<int32> LastAudioLoopIndex;
-
-				// We now need to "renumber" the (primary) sequence index of what is left in the sample queues.
-				// The secondary index (the loop index) restarts at 0. See above.
-				FMediaSamples TempSamples;
-				TRange<FTimespan> AllTime(FTimespan::MinValue(), FTimespan::MaxValue());
-				// Video
-				{
-					TSharedPtr<IMediaTextureSample, ESPMode::ThreadSafe> TempSample;
-					while(Samples.FetchVideo(AllTime, TempSample))
-					{
-						TSharedPtr<FElectraTextureSample, ESPMode::ThreadSafe> vs = StaticCastSharedPtr<FElectraTextureSample>(TempSample);
-						FMediaTimeStamp ts = vs->GetTime();
-						LastVideoLoopIndex = ts.GetSecondaryIndex();
-						vs->SetTime(ts.SetPrimaryIndex(NextSequenceIndex).SetSecondaryIndex(0));
-						TempSamples.AddVideo(TempSample.ToSharedRef());
-					}
-					while(TempSamples.FetchVideo(AllTime, TempSample))
-					{
-						Samples.AddVideo(TempSample.ToSharedRef());
-					}
-				}
-				// Audio
-				{
-					TSharedPtr<IMediaAudioSample, ESPMode::ThreadSafe> TempSample;
-					while(Samples.FetchAudio(AllTime, TempSample))
-					{
-						TSharedPtr<FElectraPlayerAudioSample, ESPMode::ThreadSafe> as = StaticCastSharedPtr<FElectraPlayerAudioSample>(TempSample);
-						FMediaTimeStamp ts = as->GetTime();
-						LastAudioLoopIndex = ts.GetSecondaryIndex();
-						as->SetTime(ts.SetPrimaryIndex(NextSequenceIndex).SetSecondaryIndex(0));
-						TempSamples.AddAudio(TempSample.ToSharedRef());
-					}
-					while(TempSamples.FetchAudio(AllTime, TempSample))
-					{
-						Samples.AddAudio(TempSample.ToSharedRef());
-					}
-				}
-				// Set the index remappers. The primary index changes only with a Seek() call, and since we are not
-				// actually seeking whatever the player still has available needs to be remapped to the new index
-				// the facade is expecting.
-				SequenceIndexMapperVideo.SetRemapPrimaryIndex(NextSequenceIndex, LastVideoLoopIndex);
-				SequenceIndexMapperAudio.SetRemapPrimaryIndex(NextSequenceIndex, LastAudioLoopIndex);
-				SequenceIndexMapperSubtitle.SetRemapPrimaryIndex(NextSequenceIndex, TOptional<int32>());
-				SequenceIndexMapperMetadata.SetRemapPrimaryIndex(NextSequenceIndex, TOptional<int32>());
-				CurrentSequenceIndex = NextSequenceIndex;
-				SampleLock.Unlock();
-				SendMediaEvent(EPlayerEvent::SeekCompleted);
-				return true;
-			}
-		}
-	}
-
-	SequenceIndexMapperVideo.SetExpectedPrimaryIndex(NextSequenceIndex);
-	SequenceIndexMapperAudio.SetExpectedPrimaryIndex(NextSequenceIndex);
-	SequenceIndexMapperSubtitle.SetExpectedPrimaryIndex(NextSequenceIndex);
-	SequenceIndexMapperMetadata.SetExpectedPrimaryIndex(NextSequenceIndex);
-	CurrentSequenceIndex = NextSequenceIndex;
-	SampleLock.Unlock();
-
-	IElectraPlayerInterface::FSeekParam sp;
-	sp.SequenceIndex = CurrentSequenceIndex;
-	return Player->Seek(Time, sp);
+	return Player->Seek(Time);
 }
 
 TRange<FTimespan> FElectraPlayerPlugin::GetPlaybackTimeRange(EMediaTimeRangeType InRangeToGet) const
diff --git a/Engine/Plugins/Media/ElectraPlayer/Source/ElectraPlayerPlugin/Private/ElectraPlayerPlugin.h b/Engine/Plugins/Media/ElectraPlayer/Source/ElectraPlayerPlugin/Private/ElectraPlayerPlugin.h
--- a/Engine/Plugins/Media/ElectraPlayer/Source/ElectraPlayerPlugin/Private/ElectraPlayerPlugin.h
+++ b/Engine/Plugins/Media/ElectraPlayer/Source/ElectraPlayerPlugin/Private/ElectraPlayerPlugin.h
@@ -174,64 +174,6 @@ private:
 
 
 private:
-	template <typename T>
-	struct FSequenceIndexMapper
-	{
-		void Reset()
-		{
-			SetExpectedPrimaryIndex(0);
-		}
-		void SetExpectedPrimaryIndex(int32 InIndex)
-		{
-			RemapPrimaryIndex.Reset();
-			RemapSecondaryIndex.Reset();
-			PrimaryIndex = InIndex;
-			SecondaryIndex = 0;
-		}
-		void SetRemapPrimaryIndex(int32 InIndex, const TOptional<int32>& InLoop)
-		{
-			RemapPrimaryIndex = InIndex;
-			RemapSecondaryIndex = InLoop;
-		}
-		bool Remap(T InOutSample)
-		{
-			FDecoderTimeStamp ts = InOutSample->GetTime();
-			// Remap?
-			if (RemapPrimaryIndex.IsSet())
-			{
-				ts = RemapTime(ts);
-				InOutSample->SetTime(ts);
-			}
-			// Always return true for now.
-			// In the future we could check if the sequence index is
-			// what is expected and if not return false to indicate this.
-			return true;
-		}
-		FDecoderTimeStamp RemapTime(FDecoderTimeStamp InTime)
-		{
-			int32 InLoopCount = (int32)(InTime.SequenceIndex & 0xffffffff);
-			// Do we have an established loop count yet?
-			if (!RemapSecondaryIndex.IsSet())
-			{
-				RemapSecondaryIndex = InLoopCount;
-			}
-			uint32 NewLoopCount = (uint32)InLoopCount - (uint32)RemapSecondaryIndex.GetValue();
-
-			InTime.SequenceIndex = ((int64)RemapPrimaryIndex.Get(0) << 32) | NewLoopCount;
-			return InTime;
-		}
-		int32 PrimaryIndex = 0;
-		int32 SecondaryIndex = 0;
-		TOptional<int32> RemapPrimaryIndex;
-		TOptional<int32> RemapSecondaryIndex;
-	};
-	FSequenceIndexMapper<FVideoDecoderOutputPtr> SequenceIndexMapperVideo;
-	FSequenceIndexMapper<IAudioDecoderOutputPtr> SequenceIndexMapperAudio;
-	FSequenceIndexMapper<ISubtitleDecoderOutputPtr> SequenceIndexMapperSubtitle;
-	FSequenceIndexMapper<IMetaDataDecoderOutputPtr> SequenceIndexMapperMetadata;
-	int32 CurrentSequenceIndex = 0;
-
-
 	/** Output queues as needed by MediaFramework */
 	TUniquePtr<FMediaSamples> MediaSamples;
 	mutable FCriticalSection MediaSamplesLock;
diff --git a/Engine/Plugins/Media/ElectraPlayer/Source/ElectraPlayerRuntime/Private/ElectraPlayer.cpp b/Engine/Plugins/Media/ElectraPlayer/Source/ElectraPlayerRuntime/Private/ElectraPlayer.cpp
--- a/Engine/Plugins/Media/ElectraPlayer/Source/ElectraPlayerRuntime/Private/ElectraPlayer.cpp
+++ b/Engine/Plugins/Media/ElectraPlayer/Source/ElectraPlayerRuntime/Private/ElectraPlayer.cpp
@@ -199,6 +199,7 @@ void FElectraPlayer::ClearToDefaultState()
 {
 	FScopeLock lock(&PlayerLock);
 
+	CurrentSequenceIndex = 0;
 	PlayerState.Reset();
 	NumTracksAudio = 0;
 	NumTracksVideo = 0;
@@ -1254,6 +1255,7 @@ bool FElectraPlayer::Seek(const FTimespan& Time)
 
 bool FElectraPlayer::Seek(const FTimespan& Time, const FSeekParam& Param)
 {
+	++CurrentSequenceIndex;
 	if (CurrentPlayer.Get())
 	{
 		FTimespan Target;
@@ -1266,6 +1268,10 @@ bool FElectraPlayer::Seek(const FTimespan& Time, const FSeekParam& Param)
 		seek.NewSequenceIndex = Param.SequenceIndex;
 		seek.StartingBitrate = Param.StartingBitrate;
 		seek.bOptimizeForScrubbing = Param.bOptimizeForScrubbing;
+		if (!seek.NewSequenceIndex.IsSet())
+		{
+			seek.NewSequenceIndex = CurrentSequenceIndex;
+		}
 		bInitialSeekPerformed = true;
 		bDiscardOutputUntilCleanStart = true;
 		CurrentPlayer->AdaptivePlayer->SeekTo(seek);
diff --git a/Engine/Plugins/Media/ElectraPlayer/Source/ElectraPlayerRuntime/Private/ElectraPlayer.h b/Engine/Plugins/Media/ElectraPlayer/Source/ElectraPlayerRuntime/Private/ElectraPlayer.h
--- a/Engine/Plugins/Media/ElectraPlayer/Source/ElectraPlayerRuntime/Private/ElectraPlayer.h
+++ b/Engine/Plugins/Media/ElectraPlayer/Source/ElectraPlayerRuntime/Private/ElectraPlayer.h
@@ -173,7 +173,7 @@ private:
 		virtual void OnMediaPlayerSubtitleReceived(ISubtitleDecoderOutputPtr Subtitle) override
 		{ SubtitleReceivedDelegate.ExecuteIfBound(Subtitle); }
 		virtual void OnMediaPlayerFlushSubtitles() override
-		{ SubtitleFlushDelegate.ExecuteIfBound(); }		 
+		{ SubtitleFlushDelegate.ExecuteIfBound(); }
 		FOnMediaPlayerSubtitleReceivedDelegate SubtitleReceivedDelegate;
 		FOnMediaPlayerSubtitleFlushDelegate SubtitleFlushDelegate;
 	};
@@ -339,7 +339,7 @@ private:
 	virtual void ReportPlaylistDownload(const Metrics::FPlaylistDownloadStats& PlaylistDownloadStats) override
 	{ DeferredPlayerEvents.Enqueue(MakeSharedTS<FPlayerMetricEvent_PlaylistDownload>(PlaylistDownloadStats)); }
 	virtual void ReportCleanStart() override
-	{ /*DeferredPlayerEvents.Enqueue(MakeSharedTS<FPlayerMetricEventBase>(FPlayerMetricEventBase::EType::CleanStart));*/ 
+	{ /*DeferredPlayerEvents.Enqueue(MakeSharedTS<FPlayerMetricEventBase>(FPlayerMetricEventBase::EType::CleanStart));*/
 		bDiscardOutputUntilCleanStart = false;
 	}
 	virtual void ReportBufferingStart(Metrics::EBufferingReason BufferingReason) override
@@ -407,6 +407,9 @@ private:
 	// Delegate to talk back to adapter host
 	TWeakPtr<IElectraPlayerAdapterDelegate, ESPMode::ThreadSafe>	AdapterDelegate;
 
+	// The current sequence index as expected by the media player facade.
+	int32 CurrentSequenceIndex = 0;
+
 	// Contains number of audio tracks available to expose it later.
 	int32											NumTracksAudio;
 	int32											NumTracksVideo;
diff --git a/Engine/Plugins/Media/ElectraPlayer/Source/ElectraPlayerRuntime/Private/Runtime/Player/AdaptivePlayerOptionKeynames.h b/Engine/Plugins/Media/ElectraPlayer/Source/ElectraPlayerRuntime/Private/Runtime/Player/AdaptivePlayerOptionKeynames.h
--- a/Engine/Plugins/Media/ElectraPlayer/Source/ElectraPlayerRuntime/Private/Runtime/Player/AdaptivePlayerOptionKeynames.h
+++ b/Engine/Plugins/Media/ElectraPlayer/Source/ElectraPlayerRuntime/Private/Runtime/Player/AdaptivePlayerOptionKeynames.h
@@ -33,14 +33,8 @@ namespace Electra
 	//! (bool) true to perform frame accurate seeks (slow as decoding and discarding data from a preceeding keyframe is required)
 	const FName OptionKeyFrameAccurateSeek(TEXT("frame_accurate_seeking"));
 
-	//! (bool) true to optimize seeking for faster frame scrubbing, false to optimize for playback.
-	const FName OptionKeyFrameOptimizeSeekForScrubbing(TEXT("optimize_seek_for_scrubbing"));
-
-	//! (bool) true to emit the first decoded video frame while prerolling so it can be displayed while scrubbing.
-	const FName OptionKeyDoNotHoldBackFirstVideoFrame(TEXT("do_not_hold_back_first_frame"));
-
-	//! (bool) true to always emit decoded samples when the player is paused. Supersedes `OptionKeyDoNotHoldBackFirstVideoFrame`.
-	const FName OptionKeyAlwaysEmitSamplesWhenPaused(TEXT("always_emit_samples_when_paused"));
+	//! (bool) false (default) to emit the first decoded video frame while prerolling so it can be displayed while scrubbing, true to not display it right away.
+	const FName OptionKeyHoldBackFirstVideoFrame(TEXT("hold_back_first_frame"));
 
 	//! (bool) true to not truncate the media segment access units at the end of the presentation. Must only be used without a set playback range end!
 	const FName OptionKeyDoNotTruncateAtPresentationEnd(TEXT("do_not_truncate_at_presentation_end"));
@@ -48,8 +42,6 @@ namespace Electra
 	//! (bool) true to have every request to read data break out to an external data reader.
 	const FName OptionKeyUseExternalDataReader(TEXT("use_external_data_reader"));
 
-	const FName OptionKeyHoldCurrentFrame(TEXT("hold_current_frame"));
-
 	const FName OptionKeyCurrentAvgStartingVideoBitrate(TEXT("current:avg_video_bitrate"));
 
 	const FName OptionKeyExcludedCodecsVideo(TEXT("excluded_codecs_video"));
diff --git a/Engine/Plugins/Media/ElectraPlayer/Source/ElectraPlayerRuntime/Private/Runtime/Player/AdaptiveStreamingPlayer.cpp b/Engine/Plugins/Media/ElectraPlayer/Source/ElectraPlayerRuntime/Private/Runtime/Player/AdaptiveStreamingPlayer.cpp
--- a/Engine/Plugins/Media/ElectraPlayer/Source/ElectraPlayerRuntime/Private/Runtime/Player/AdaptiveStreamingPlayer.cpp
+++ b/Engine/Plugins/Media/ElectraPlayer/Source/ElectraPlayerRuntime/Private/Runtime/Player/AdaptiveStreamingPlayer.cpp
@@ -1103,7 +1103,6 @@ void FAdaptiveStreamingPlayer::StopRendering()
  */
 int32 FAdaptiveStreamingPlayer::CreateRenderers()
 {
-	bool bAlwaysEmitSamplesWhenPaused = PlayerOptions.GetValue(OptionKeyAlwaysEmitSamplesWhenPaused).SafeGetBool(false);
 	double CurrentRate = PlaybackState.GetCurrentPlayRate();
 	double IntendedRate = PlaybackState.GetDesiredPlayRate();
 	const bool bCurrentlyPaused = PlaybackState.GetIsPaused();
@@ -1113,19 +1112,14 @@ int32 FAdaptiveStreamingPlayer::CreateRenderers()
 	{
 		VideoRender.Renderer->SetRenderClock(RenderClock);
 		VideoRender.Renderer->SetPlaybackRate(CurrentRate, IntendedRate, bCurrentlyPaused);
-		VideoRender.Renderer->AlwaysEmitSamplesWhenPaused(bAlwaysEmitSamplesWhenPaused);
 		// Hold back all video frames during preroll or emit the first frame for scrubbing?
-		if (PlayerOptions.HaveKey(OptionKeyDoNotHoldBackFirstVideoFrame))
-		{
-			bool bDoNotHoldBackFirstVideoFrame = PlayerOptions.GetValue(OptionKeyDoNotHoldBackFirstVideoFrame).SafeGetBool(false);
-			VideoRender.Renderer->DisableHoldbackOfFirstRenderableVideoFrame(bDoNotHoldBackFirstVideoFrame);
-		}
+		bool bHoldBackFirstVideoFrame = PlayerOptions.GetValue(OptionKeyHoldBackFirstVideoFrame).SafeGetBool(false);
+		VideoRender.Renderer->EnableHoldbackOfFirstRenderableVideoFrame(bHoldBackFirstVideoFrame);
 	}
 	if (AudioRender.Renderer)
 	{
 		AudioRender.Renderer->SetRenderClock(RenderClock);
 		AudioRender.Renderer->SetPlaybackRate(CurrentRate, IntendedRate, bCurrentlyPaused);
-		AudioRender.Renderer->AlwaysEmitSamplesWhenPaused(bAlwaysEmitSamplesWhenPaused);
 	}
 
 	return 0;
@@ -1875,9 +1869,6 @@ void FAdaptiveStreamingPlayer::HandleSeeking()
 			CurrentState = EPlayerState::eState_Seeking;
 			SeekVars.Lock.Lock();
 
-			bool bIsForScrubbing = SeekParam.bOptimizeForScrubbing.Get(PlayerOptions.GetValue(OptionKeyFrameOptimizeSeekForScrubbing).SafeGetBool(false));
-			SeekVars.bForScrubbing = bIsForScrubbing;
-
 			// When seeking (other than the initial playstart seek) any default end time that
 			// was set via the URL #t,e parameter no longer applies.
 			if (Manifest.IsValid())
@@ -2297,71 +2288,31 @@ void FAdaptiveStreamingPlayer::HandleNewBufferedData()
 	// Are we in buffering state? (this includes rebuffering)
 	if (bStreamingHasStarted && CurrentState == EPlayerState::eState_Buffering && DecoderState == EDecoderState::eDecoder_Paused)
 	{
-		if (!SeekVars.bForScrubbing)
+		// Seek for scrubbing is not possible at playstart so we can skip a few things here.
+		if (!SeekVars.bPrerollDone)
 		{
-			check(LastBufferingState == EPlayerState::eState_Buffering || LastBufferingState == EPlayerState::eState_Rebuffering || LastBufferingState == EPlayerState::eState_Seeking);
-			// Check if we have enough data buffered up to begin handing off data to the decoders.
-			bool bHaveEnoughData = HaveEnoughBufferedDataToStartPlayback();
-			if (bHaveEnoughData)
+			DataBuffersCriticalSection.Lock();
+			if (!ActiveDataOutputBuffers.IsValid())
 			{
-				PlaybackState.SetIsBuffering(false);
-				PlaybackState.SetIsSeeking(false);
-
-				// Activate the stream buffers from which to feed the decoders.
-				DataBuffersCriticalSection.Lock();
-				if (!ActiveDataOutputBuffers.IsValid())
+				TSharedPtrTS<FStreamDataBuffers> NewOutputBuffers;
+				if (NextDataBuffers.Num())
 				{
-					TSharedPtrTS<FStreamDataBuffers> NewOutputBuffers;
-					check(NextDataBuffers.Num());
-					if (NextDataBuffers.Num())
-					{
-						NewOutputBuffers = NextDataBuffers[0];
-						NextDataBuffers.RemoveAt(0);
-					}
-					ActiveDataOutputBuffers = MoveTemp(NewOutputBuffers);
+					NewOutputBuffers = NextDataBuffers[0];
+					NextDataBuffers.RemoveAt(0);
 				}
-				DataBuffersCriticalSection.Unlock();
-
+				ActiveDataOutputBuffers = MoveTemp(NewOutputBuffers);
+			}
+			bool bHaveActiveOutputBuffer = ActiveDataOutputBuffers.IsValid();
+			DataBuffersCriticalSection.Unlock();
+			if (bHaveActiveOutputBuffer)
+			{
 				PrepareForPrerolling();
 				PipelineState = EPipelineState::ePipeline_Prerolling;
 				DecoderState = EDecoderState::eDecoder_Running;
-
 				PostrollVars.Clear();
-
-				// Send buffering end event
-				DispatchBufferingEvent(false, LastBufferingState);
-				// Send pre-roll begin event.
 				DispatchEvent(FMetricEvent::ReportPrerollStart());
 			}
 		}
-		else
-		{
-			// Optimized seek for scrubbing is not possible at playstart so we can skip a few things here.
-			if (!SeekVars.bScrubPrerollDone)
-			{
-				DataBuffersCriticalSection.Lock();
-				if (!ActiveDataOutputBuffers.IsValid())
-				{
-					TSharedPtrTS<FStreamDataBuffers> NewOutputBuffers;
-					if (NextDataBuffers.Num())
-					{
-						NewOutputBuffers = NextDataBuffers[0];
-						NextDataBuffers.RemoveAt(0);
-					}
-					ActiveDataOutputBuffers = MoveTemp(NewOutputBuffers);
-				}
-				bool bHaveActiveOutputBuffer = ActiveDataOutputBuffers.IsValid();
-				DataBuffersCriticalSection.Unlock();
-				if (bHaveActiveOutputBuffer)
-				{
-					PrepareForPrerolling();
-					PipelineState = EPipelineState::ePipeline_Prerolling;
-					DecoderState = EDecoderState::eDecoder_Running;
-					PostrollVars.Clear();
-					DispatchEvent(FMetricEvent::ReportPrerollStart());
-				}
-			}
-		}
 	}
 }
 
@@ -2501,55 +2452,30 @@ void FAdaptiveStreamingPlayer::HandleNewOutputData()
 		// Ready to go?
 		if (bHaveEnoughVideo && bHaveEnoughAudio)
 		{
-			if (!SeekVars.bForScrubbing)
+			// The first time around we send the seek complete event and stop the feeding of the decoders
+			// to prevent them consuming and sending output to the renderers while we await the buffers
+			// to fill to the amount we need.
+			if (!SeekVars.bPrerollDone)
 			{
-				if (!PrerollVars.bIsMidSequencePreroll)
-				{
-					PipelineState = EPipelineState::ePipeline_Stopped;
-					CurrentState = EPlayerState::eState_Paused;
-					PrerollVars.Clear();
-					SeekVars.ClearWorkVars();
-					SeekVars.SetFinished();
-					if (LastBufferingState == EPlayerState::eState_Seeking)
-					{
-						DispatchEvent(FMetricEvent::ReportSeekCompleted());
-					}
-					DispatchEvent(FMetricEvent::ReportPrerollEnd());
-				}
-				else
-				{
-					PrerollVars.Clear();
-					StartRendering();
-					PipelineState = EPipelineState::ePipeline_Running;
-				}
+				SeekVars.bPrerollDone = true;
+				SeekVars.SetFinished();
+				DispatchEvent(FMetricEvent::ReportSeekCompleted());
+				// Pause feeding the decoder to stop draining buffers we want to fill.
+				DecoderState = EDecoderState::eDecoder_Paused;
 			}
-			else
+
+			if (HaveEnoughBufferedDataToStartPlayback())
 			{
-				// The first time around we send the seek complete event and stop the feeding of the decoders
-				// to prevent them consuming and sending output to the renderers while we await the buffers
-				// to fill to the amount we need.
-				if (!SeekVars.bScrubPrerollDone)
-				{
-					SeekVars.bScrubPrerollDone = true;
-					SeekVars.SetFinished();
-					DispatchEvent(FMetricEvent::ReportSeekCompleted());
-					// Pause feeding the decoder to stop draining buffers we want to fill.
-					DecoderState = EDecoderState::eDecoder_Paused;
-				}
+				PlaybackState.SetIsBuffering(false);
+				PlaybackState.SetIsSeeking(false);
 
-				if (HaveEnoughBufferedDataToStartPlayback())
-				{
-					PlaybackState.SetIsBuffering(false);
-					PlaybackState.SetIsSeeking(false);
-
-					SeekVars.ClearWorkVars();
-					PipelineState = EPipelineState::ePipeline_Stopped;
-					CurrentState = EPlayerState::eState_Paused;
-					DecoderState = EDecoderState::eDecoder_Running;
-					PrerollVars.Clear();
-					DispatchBufferingEvent(false, LastBufferingState);
-					DispatchEvent(FMetricEvent::ReportPrerollEnd());
-				}
+				SeekVars.ClearWorkVars();
+				PipelineState = EPipelineState::ePipeline_Stopped;
+				CurrentState = EPlayerState::eState_Paused;
+				DecoderState = EDecoderState::eDecoder_Running;
+				PrerollVars.Clear();
+				DispatchBufferingEvent(false, LastBufferingState);
+				DispatchEvent(FMetricEvent::ReportPrerollEnd());
 			}
 		}
 	}
diff --git a/Engine/Plugins/Media/ElectraPlayer/Source/ElectraPlayerRuntime/Private/Runtime/Player/AdaptiveStreamingPlayerInternal.h b/Engine/Plugins/Media/ElectraPlayer/Source/ElectraPlayerRuntime/Private/Runtime/Player/AdaptiveStreamingPlayerInternal.h
--- a/Engine/Plugins/Media/ElectraPlayer/Source/ElectraPlayerRuntime/Private/Runtime/Player/AdaptiveStreamingPlayerInternal.h
+++ b/Engine/Plugins/Media/ElectraPlayer/Source/ElectraPlayerRuntime/Private/Runtime/Player/AdaptiveStreamingPlayerInternal.h
@@ -71,7 +71,7 @@ public:
 
 	virtual void AlwaysEmitSamplesWhenPaused(bool bEmitAlways) = 0;
 	virtual void SetPlaybackRate(double InCurrentPlaybackRate, double InIntendedPlaybackRate, bool bInCurrentlyPaused) = 0;
-	virtual void DisableHoldbackOfFirstRenderableVideoFrame(bool bDisableHoldback) = 0;
+	virtual void EnableHoldbackOfFirstRenderableVideoFrame(bool bEnableHoldback) = 0;
 
 	virtual FTimeRange GetSupportedRenderRateScale() = 0;
 	virtual void SetPlayRateScale(double InNewScale) = 0;
@@ -1249,9 +1249,8 @@ private:
 		{
 			if (Renderer.IsValid())
 			{
-				FParamDict options;
-				options.Set(Electra::OptionKeyHoldCurrentFrame, FVariantValue(bHoldCurrentFrame));
-				Renderer->Flush(options);
+				FParamDict NoOptions;
+				Renderer->Flush(NoOptions);
 			}
 		}
 		TSharedPtr<IAdaptiveStreamingWrappedRenderer, ESPMode::ThreadSafe>	Renderer;
@@ -1882,8 +1881,7 @@ private:
 	{
 		void ClearWorkVars()
 		{
-			bForScrubbing = false;
-			bScrubPrerollDone = false;
+			bPrerollDone = false;
 		}
 
 		void Reset()
@@ -1894,8 +1892,7 @@ private:
 			//PendingRequest.Reset();
 			ActiveRequest.Reset();
 			PlayrangeOnRequest.Reset();
-			bForScrubbing = false;
-			bScrubPrerollDone = false;
+			bPrerollDone = false;
 		}
 
 		void SetFinished()
@@ -1903,8 +1900,7 @@ private:
 			ActiveRequest.Reset();
 		}
 
-		bool bForScrubbing = false;
-		bool bScrubPrerollDone = false;
+		bool bPrerollDone = false;
 		bool bIsPlayStart = true;
 
 		TOptional<FSeekParam> ActiveRequest;
diff --git a/Engine/Plugins/Media/ElectraPlayer/Source/ElectraPlayerRuntime/Private/Runtime/Player/AdaptiveStreamingPlayerRenderers.cpp b/Engine/Plugins/Media/ElectraPlayer/Source/ElectraPlayerRuntime/Private/Runtime/Player/AdaptiveStreamingPlayerRenderers.cpp
--- a/Engine/Plugins/Media/ElectraPlayer/Source/ElectraPlayerRuntime/Private/Runtime/Player/AdaptiveStreamingPlayerRenderers.cpp
+++ b/Engine/Plugins/Media/ElectraPlayer/Source/ElectraPlayerRuntime/Private/Runtime/Player/AdaptiveStreamingPlayerRenderers.cpp
@@ -97,7 +97,7 @@ namespace Electra
 
 		void AlwaysEmitSamplesWhenPaused(bool bEmitAlways) override;
 		void SetPlaybackRate(double InCurrentPlaybackRate, double InIntendedPlaybackRate, bool bInCurrentlyPaused) override;
-		void DisableHoldbackOfFirstRenderableVideoFrame(bool bDisableHoldback) override;
+		void EnableHoldbackOfFirstRenderableVideoFrame(bool bEnableHoldback) override;
 
 		FTimeRange GetSupportedRenderRateScale() override;
 		void SetPlayRateScale(double InNewScale) override;
@@ -135,7 +135,7 @@ namespace Electra
 		double CurrentPlaybackRate = 0.0;
 		double IntendedPlaybackRate = 0.0;
 		bool bAlwaysEmitSamplesWhenPaused = false;
-		bool bDoNotHoldBackFirstVideoFrame = true;
+		bool bHoldBackFirstVideoFrame = false;
 		uint32 NumBuffersNotHeldBack = 0;
 
 		TQueue<FPendingReturnBuffer> PendingReturnBuffers;
@@ -468,7 +468,7 @@ UEMediaError FAdaptiveStreamingWrappedRenderer::ReturnBufferCommon(IBuffer* Buff
 		}
 		// If the video renderer shall not hold back the first frame (used for scrubbing video)
 		// then we pass it out. The count is reset in Flush().
-		else if (Type == EStreamType::Video && bDoNotHoldBackFirstVideoFrame)
+		else if (Type == EStreamType::Video && !bHoldBackFirstVideoFrame)
 		{
 			if (NumBuffersNotHeldBack == 0)
 			{
@@ -634,10 +634,10 @@ void FAdaptiveStreamingWrappedRenderer::SetPlaybackRate(double InCurrentPlayback
 	IntendedPlaybackRate = bInCurrentlyPaused ? 0.0 : InIntendedPlaybackRate;
 }
 
-void FAdaptiveStreamingWrappedRenderer::DisableHoldbackOfFirstRenderableVideoFrame(bool bInDisableHoldback)
+void FAdaptiveStreamingWrappedRenderer::EnableHoldbackOfFirstRenderableVideoFrame(bool bEnableHoldback)
 {
 	FScopeLock lock(&Lock);
-	bDoNotHoldBackFirstVideoFrame = bInDisableHoldback;
+	bHoldBackFirstVideoFrame = bEnableHoldback;
 }
 
 FTimeRange FAdaptiveStreamingWrappedRenderer::GetSupportedRenderRateScale()
