ElectraPlayer: Added DASH <Label> element to AdaptationSet handling for exposure through the track's metadata

#rnx

diff --git a/Engine/Plugins/Media/ElectraPlayer/Source/ElectraPlayerRuntime/Private/Runtime/Player/DASH/ManifestBuilderDASH.cpp b/Engine/Plugins/Media/ElectraPlayer/Source/ElectraPlayerRuntime/Private/Runtime/Player/DASH/ManifestBuilderDASH.cpp
--- a/Engine/Plugins/Media/ElectraPlayer/Source/ElectraPlayerRuntime/Private/Runtime/Player/DASH/ManifestBuilderDASH.cpp
+++ b/Engine/Plugins/Media/ElectraPlayer/Source/ElectraPlayerRuntime/Private/Runtime/Player/DASH/ManifestBuilderDASH.cpp
@@ -1366,6 +1366,14 @@ void FManifestDASHInternal::PreparePeriodAdaptationSets(TSharedPtrTS<FPeriod> Pe
 			AdaptationSet->PAR = MPDAdaptationSet->GetPAR();
 			AdaptationSet->Language = ISO639::RFC5646To639_1(MPDAdaptationSet->GetLanguage());
 			AdaptationSet->SelectionPriority = (int32) MPDAdaptationSet->GetSelectionPriority();
+			// Set up labels. For now we ignore GroupLabel elements.
+			for(auto& itLabel : MPDAdaptationSet->GetLabels())
+			{
+				FLabelElement& Label = AdaptationSet->Labels.Emplace_GetRef();
+				Label.ID = itLabel->GetID();
+				Label.Language = ISO639::RFC5646To639_1(itLabel->GetLanguage());
+				Label.LabelText = itLabel->GetData();
+			}
 
 			// Content components are not supported.
 			if (MPDAdaptationSet->GetContentComponents().Num())
diff --git a/Engine/Plugins/Media/ElectraPlayer/Source/ElectraPlayerRuntime/Private/Runtime/Player/DASH/PlaylistReaderDASH_Internal.h b/Engine/Plugins/Media/ElectraPlayer/Source/ElectraPlayerRuntime/Private/Runtime/Player/DASH/PlaylistReaderDASH_Internal.h
--- a/Engine/Plugins/Media/ElectraPlayer/Source/ElectraPlayerRuntime/Private/Runtime/Player/DASH/PlaylistReaderDASH_Internal.h
+++ b/Engine/Plugins/Media/ElectraPlayer/Source/ElectraPlayerRuntime/Private/Runtime/Player/DASH/PlaylistReaderDASH_Internal.h
@@ -149,6 +149,13 @@ public:
 		int32 MaxQualityIndex = 0;
 	};
 
+	struct FLabelElement
+	{
+		FString LabelText;
+		FString Language;
+		uint64 ID = 0;
+	};
+
 	class FRepresentation : public IPlaybackAssetRepresentation, public TSharedFromThis<FRepresentation, ESPMode::ThreadSafe>
 	{
 	public:
@@ -272,6 +279,7 @@ public:
 		bool GetIsUsable() const													{ return bIsUsable; }
 		bool GetIsInSwitchGroup() const												{ return bIsInSwitchGroup; }
 		bool GetIsSwitchGroup() const												{ return bIsSwitchGroup; }
+		const TArray<FLabelElement>& GetLabels() const								{ return Labels; }
 
 		TSharedPtrTS<FRepresentation> GetRepresentationByUniqueID(const FString& UniqueIdentifier) const
 		{
@@ -401,6 +409,11 @@ public:
 			OutMetadata.Language = GetLanguage();
 			OutMetadata.HighestBandwidth = GetMaxBandwidth();
 			OutMetadata.HighestBandwidthCodec = GetCodec();
+			if (Labels.Num())
+			{
+				// Return the first label, regardless of its language or related group label ID.
+				OutMetadata.Label = Labels[0].LabelText;
+			}
 			// Map role and accessibility. Do this last since this is allowed to overwrite ID and Language
 			MapRoleAccessibilityToHTML5(OutMetadata, StreamType);
 			const TArray<TSharedPtrTS<FRepresentation>>& Reprs = GetRepresentations();
@@ -486,6 +499,7 @@ public:
 		TArray<TSharedPtrTS<FRepresentation>> Representations;
 		TArray<FString> Roles;
 		TArray<FString> Accessibilities;
+		TArray<FLabelElement> Labels;
 		FTimeFraction PAR;
 		FString Language;
 		int32 MaxBandwidth = 0;
