ElectraPlayer: returning what the renderer believes are in the sample queue in the cache query

#rnx

diff --git a/Engine/Plugins/Media/ElectraPlayer/Source/ElectraPlayerRuntime/Private/Runtime/Player/AdaptiveStreamingPlayer.cpp b/Engine/Plugins/Media/ElectraPlayer/Source/ElectraPlayerRuntime/Private/Runtime/Player/AdaptiveStreamingPlayer.cpp
--- a/Engine/Plugins/Media/ElectraPlayer/Source/ElectraPlayerRuntime/Private/Runtime/Player/AdaptiveStreamingPlayer.cpp
+++ b/Engine/Plugins/Media/ElectraPlayer/Source/ElectraPlayerRuntime/Private/Runtime/Player/AdaptiveStreamingPlayer.cpp
@@ -994,7 +994,6 @@ void FAdaptiveStreamingPlayer::QueryStreamBufferInfo(FStreamBufferInfo& OutStrea
 					}
 				};
 
-#if 0
 				// Ask the renderers what they know they have released into the media sample queue.
 				TArray<IAdaptiveStreamingWrappedRenderer::FEnqueuedSampleInfo> Enqueued;
 				if (InStreamType == EStreamType::Video && VideoRender.Renderer.IsValid())
@@ -1013,7 +1012,6 @@ void FAdaptiveStreamingPlayer::QueryStreamBufferInfo(FStreamBufferInfo& OutStrea
 					tr.End = Enqueued.Last().PTS + Enqueued.Last().Duration;
 					OutStreamBufferInfo.TimeEnqueued.Emplace(MoveTemp(tr));
 				}
-#endif
 
 				// Get the time range of the access units the respective decoder has in its internal pipeline
 				// These AUs are not in the player buffer any more.
