ElectraPlayer: lowering DASH log verbosity; allowing DASH full profile even though not supported because some sources incorrectly say they require it when they do not

#rnx

diff --git a/Engine/Plugins/Media/ElectraPlayer/Source/ElectraPlayerRuntime/Private/Runtime/Player/DASH/ManifestBuilderDASH.cpp b/Engine/Plugins/Media/ElectraPlayer/Source/ElectraPlayerRuntime/Private/Runtime/Player/DASH/ManifestBuilderDASH.cpp
--- a/Engine/Plugins/Media/ElectraPlayer/Source/ElectraPlayerRuntime/Private/Runtime/Player/DASH/ManifestBuilderDASH.cpp
+++ b/Engine/Plugins/Media/ElectraPlayer/Source/ElectraPlayerRuntime/Private/Runtime/Player/DASH/ManifestBuilderDASH.cpp
@@ -77,6 +77,7 @@ namespace
 //		TEXT("urn:mpeg:dash:profile:isoff-ext-live:2014"),				// full xlink required
 //		TEXT("urn:mpeg:dash:profile:isoff-common:2014"),				// may contain both ext-on-demand and ext-live additions
 //		TEXT("urn:mpeg:dash:profile:isoff-broadcast:2015"),				// RandomAccess and Switching support required
+		TEXT("urn:mpeg:dash:profile:full:2011"),						// NOTE: Full profile is not supported, but some sources that are a lesser profile are incorrectly saying they require full, so we hope for the best.
 		// Possibly add DASH-IF-IOP profiles?
 //		TEXT(""),
 		// DVB-DASH profiles
diff --git a/Engine/Plugins/Media/ElectraPlayer/Source/ElectraPlayerRuntime/Private/Runtime/Player/DASH/ManifestDASH.cpp b/Engine/Plugins/Media/ElectraPlayer/Source/ElectraPlayerRuntime/Private/Runtime/Player/DASH/ManifestDASH.cpp
--- a/Engine/Plugins/Media/ElectraPlayer/Source/ElectraPlayerRuntime/Private/Runtime/Player/DASH/ManifestDASH.cpp
+++ b/Engine/Plugins/Media/ElectraPlayer/Source/ElectraPlayerRuntime/Private/Runtime/Player/DASH/ManifestDASH.cpp
@@ -3142,7 +3142,7 @@ FManifestDASHInternal::FRepresentation::ESearchResult FManifestDASHInternal::FRe
 	else if (MPDTimescale == 1 && !bWarnedAboutTimescale)
 	{
 		bWarnedAboutTimescale = true;
-		LogMessage(InPlayerSessionServices, IInfoLog::ELevel::Warning, FString::Printf(TEXT("Timescale for Representation \"%s\" is given as 1. Is this intended?"), *MPDRepresentation->GetID()));
+		LogMessage(InPlayerSessionServices, IInfoLog::ELevel::Verbose, FString::Printf(TEXT("Timescale for Representation \"%s\" is given as 1. Is this intended?"), *MPDRepresentation->GetID()));
 	}
 
 	// There needs to be a segment duration here.
@@ -3302,7 +3302,7 @@ FManifestDASHInternal::FRepresentation::ESearchResult FManifestDASHInternal::FRe
 	else if (MPDTimescale == 1 && !bWarnedAboutTimescale)
 	{
 		bWarnedAboutTimescale = true;
-		LogMessage(InPlayerSessionServices, IInfoLog::ELevel::Warning, FString::Printf(TEXT("Timescale for Representation \"%s\" is given as 1. Is this intended?"), *MPDRepresentation->GetID()));
+		LogMessage(InPlayerSessionServices, IInfoLog::ELevel::Verbose, FString::Printf(TEXT("Timescale for Representation \"%s\" is given as 1. Is this intended?"), *MPDRepresentation->GetID()));
 	}
 
 	// Get the period local time into media local timescale and add the PTO.
