Electra: Extracting the DCR and CSD from an avc3 H.264 stream

#rnx

diff --git a/Engine/Plugins/Media/D3D12VideoDecodersElectra/Source/Private/VideoDecoder_D3D12_H264.cpp b/Engine/Plugins/Media/D3D12VideoDecodersElectra/Source/Private/VideoDecoder_D3D12_H264.cpp
--- a/Engine/Plugins/Media/D3D12VideoDecodersElectra/Source/Private/VideoDecoder_D3D12_H264.cpp
+++ b/Engine/Plugins/Media/D3D12VideoDecodersElectra/Source/Private/VideoDecoder_D3D12_H264.cpp
@@ -115,12 +115,10 @@ IElectraDecoder::EDecoderError FD3D12VideoDecoder_H264::DecodeAccessUnit(const F
 		return IElectraDecoder::EDecoderError::None;
 	}
 
+	bool bGotCSD = true;
 	if ((InInputAccessUnit.Flags & EElectraDecoderFlags::IsSyncSample) != EElectraDecoderFlags::None)
 	{
-		if (GetCodecSpecificDataH264(BitstreamParamsH264, InAdditionalOptions, true) != IElectraDecoder::EDecoderError::None)
-		{
-			return IElectraDecoder::EDecoderError::Error;
-		}
+		bGotCSD = GetCodecSpecificDataH264(BitstreamParamsH264, InAdditionalOptions, false) == IElectraDecoder::EDecoderError::None;
 	}
 
 	// We need to isolate the slices that make up this frame.
@@ -129,6 +127,8 @@ IElectraDecoder::EDecoderError FD3D12VideoDecoder_H264::DecodeAccessUnit(const F
 	// Go over each of the NALUs in the bitstream.
 	const uint8* Start = (const uint8*)InInputAccessUnit.Data;
 	const uint8* End = Start + InInputAccessUnit.DataSize;
+	bool bGotSPS = false;
+	bool bGotPPS = false;
 	while(Start < End)
 	{
 		uint32 NaluLength = ((uint32)Start[0] << 24) | ((uint32)Start[1] << 16) | ((uint32)Start[2] << 8) | ((uint32)Start[3]);
@@ -141,6 +141,13 @@ IElectraDecoder::EDecoderError FD3D12VideoDecoder_H264::DecodeAccessUnit(const F
 				bIsIDR = true;
 			}
 
+			bGotCSD |= (bGotSPS && bGotPPS);
+			if (!bGotCSD)
+			{
+				PostError(0, TEXT("No SPS and PPS found in CSD or inband, cannot decode slice"), ERRCODE_INTERNAL_FAILED_TO_DECODE);
+				return IElectraDecoder::EDecoderError::Error;
+			}
+
 			FBitstreamParamsH264::FSliceDecodeInfo& SliceInfo = SliceInfos.Emplace_GetRef();
 			SliceInfo.NalUnitType = (uint8)nut;
 			SliceInfo.NalRefIdc = (uint8)refIdc;
@@ -176,6 +183,7 @@ IElectraDecoder::EDecoderError FD3D12VideoDecoder_H264::DecodeAccessUnit(const F
 				PostError(0, TEXT("Failed to parse bitstream inband SPS"), ERRCODE_INTERNAL_FAILED_TO_DECODE);
 				return IElectraDecoder::EDecoderError::Error;
 			}
+			bGotSPS = true;
 		}
 		// Inband PPS
 		else if (nut == 8)
@@ -185,6 +193,7 @@ IElectraDecoder::EDecoderError FD3D12VideoDecoder_H264::DecodeAccessUnit(const F
 				PostError(0, TEXT("Failed to parse bitstream inband PPS"), ERRCODE_INTERNAL_FAILED_TO_DECODE);
 				return IElectraDecoder::EDecoderError::Error;
 			}
+			bGotPPS = true;
 		}
 		// SVC / AVC 3D extension?
 		else if (nut == 14 || nut == 20 || nut == 21)
diff --git a/Engine/Plugins/Media/ElectraCodecs/Source/ElectraDecoders/Private/Utils/MPEG/ElectraUtilsMPEGVideo_H264.cpp b/Engine/Plugins/Media/ElectraCodecs/Source/ElectraDecoders/Private/Utils/MPEG/ElectraUtilsMPEGVideo_H264.cpp
--- a/Engine/Plugins/Media/ElectraCodecs/Source/ElectraDecoders/Private/Utils/MPEG/ElectraUtilsMPEGVideo_H264.cpp
+++ b/Engine/Plugins/Media/ElectraCodecs/Source/ElectraDecoders/Private/Utils/MPEG/ElectraUtilsMPEGVideo_H264.cpp
@@ -80,7 +80,11 @@ namespace ElectraDecodersUtil
 
 			bool ParseBitstreamForNALUs(TArray<FNaluInfo>& OutNALUs, const uint8* InBitstream, uint64 InBitstreamLength)
 			{
-				if (!InBitstream || InBitstreamLength < 3)
+				if (!InBitstream)
+				{
+					return false;
+				}
+				else if (InBitstreamLength < 3)
 				{
 					UE_LOG(LogElectraDecoders, Error, TEXT("Bitstream too short to parse for NALUs"));
 					return false;
diff --git a/Engine/Plugins/Media/ElectraCodecs/Source/ElectraDecoders/Private/Utils/MPEG/ElectraUtilsMPEGVideo_H265.cpp b/Engine/Plugins/Media/ElectraCodecs/Source/ElectraDecoders/Private/Utils/MPEG/ElectraUtilsMPEGVideo_H265.cpp
--- a/Engine/Plugins/Media/ElectraCodecs/Source/ElectraDecoders/Private/Utils/MPEG/ElectraUtilsMPEGVideo_H265.cpp
+++ b/Engine/Plugins/Media/ElectraCodecs/Source/ElectraDecoders/Private/Utils/MPEG/ElectraUtilsMPEGVideo_H265.cpp
@@ -128,7 +128,11 @@ namespace ElectraDecodersUtil
 
 			bool ParseBitstreamForNALUs(TArray<FNaluInfo>& OutNALUs, const uint8* InBitstream, uint64 InBitstreamLength)
 			{
-				if (!InBitstream || InBitstreamLength < 3)
+				if (!InBitstream)
+				{
+					return false;
+				}
+				else if (InBitstreamLength < 3)
 				{
 					UE_LOG(LogElectraDecoders, Error, TEXT("Bitstream too short to parse for NALUs"));
 					return false;
diff --git a/Engine/Plugins/Media/ElectraPlayer/Source/ElectraPlayerRuntime/Private/Runtime/Decoder/BitstreamProcessors/VideoDecoderInputBitstreamProcessorH264.cpp b/Engine/Plugins/Media/ElectraPlayer/Source/ElectraPlayerRuntime/Private/Runtime/Decoder/BitstreamProcessors/VideoDecoderInputBitstreamProcessorH264.cpp
--- a/Engine/Plugins/Media/ElectraPlayer/Source/ElectraPlayerRuntime/Private/Runtime/Decoder/BitstreamProcessors/VideoDecoderInputBitstreamProcessorH264.cpp
+++ b/Engine/Plugins/Media/ElectraPlayer/Source/ElectraPlayerRuntime/Private/Runtime/Decoder/BitstreamProcessors/VideoDecoderInputBitstreamProcessorH264.cpp
@@ -35,8 +35,7 @@ private:
 
 	bool bReplaceLengthWithStartcode = true;
 
-	TSharedPtr<const FAccessUnit::CodecData, ESPMode::ThreadSafe> PreviousCodecData;
-	TSharedPtr<const FAccessUnit::CodecData, ESPMode::ThreadSafe> CurrentCodecData;
+	TSharedPtr<FAccessUnit::CodecData, ESPMode::ThreadSafe> CurrentCodecData;
 	TMap<uint32, ElectraDecodersUtil::MPEG::H264::FSequenceParameterSet> SPSs;
 
 	MPEG::FColorimetryHelper Colorimetry;
@@ -65,7 +64,6 @@ FVideoDecoderInputBitstreamProcessorH264::FVideoDecoderInputBitstreamProcessorH2
 
 void FVideoDecoderInputBitstreamProcessorH264::Clear()
 {
-	PreviousCodecData.Reset();
 	CurrentCodecData.Reset();
 	SPSs.Empty();
 	Colorimetry.Reset();
@@ -73,6 +71,7 @@ void FVideoDecoderInputBitstreamProcessorH264::Clear()
 
 IVideoDecoderInputBitstreamProcessor::EProcessResult FVideoDecoderInputBitstreamProcessorH264::ProcessAccessUnitForDecoding(FBitstreamInfo& OutBSI, FAccessUnit* InOutAccessUnit)
 {
+	static const uint8 StartCode[4] = {0,0,0,1};
 	IVideoDecoderInputBitstreamProcessor::EProcessResult Result = IVideoDecoderInputBitstreamProcessor::EProcessResult::None;
 	if (!InOutAccessUnit)
 	{
@@ -87,41 +86,10 @@ IVideoDecoderInputBitstreamProcessor::EProcessResult FVideoDecoderInputBitstream
 		OutBSI.CodecSpecificMessages = Msgs;
 	}
 
-	//
-	// Extract sequence parameter sets from the codec specific data.
-	//
-	if (InOutAccessUnit->AUCodecData.IsValid() && InOutAccessUnit->AUCodecData.Get() != CurrentCodecData.Get())
-	{
-		// Pointers are different. Is the content too?
-		bool bDifferent = !CurrentCodecData.IsValid() || (CurrentCodecData.IsValid() && InOutAccessUnit->AUCodecData->CodecSpecificData != CurrentCodecData->CodecSpecificData);
-		if (bDifferent)
-		{
-			SPSs.Empty();
-			TArray<ElectraDecodersUtil::MPEG::H264::FNaluInfo> NALUs;
-			const uint8* pD = InOutAccessUnit->AUCodecData->CodecSpecificData.GetData();
-			ElectraDecodersUtil::MPEG::H264::ParseBitstreamForNALUs(NALUs, pD, InOutAccessUnit->AUCodecData->CodecSpecificData.Num());
-			for(int32 i=0; i<NALUs.Num(); ++i)
-			{
-				const uint8* NALU = (const uint8*)Electra::AdvancePointer(pD, NALUs[i].Offset + NALUs[i].UnitLength);
-				// SPS?
-				if (NALUs[i].Type == 7)
-				{
-					bool bOk = ElectraDecodersUtil::MPEG::H264::ParseSequenceParameterSet(SPSs, NALU, NALUs[i].Size);
-					check(bOk); (void)bOk;
-				}
-			}
-			Result = IVideoDecoderInputBitstreamProcessor::EProcessResult::CSDChanged;
-		}
-		PreviousCodecData = CurrentCodecData.IsValid() ? CurrentCodecData : InOutAccessUnit->AUCodecData;
-		CurrentCodecData = InOutAccessUnit->AUCodecData;
-	}
-
-	// Set the messages from the CSD on the access unit.
-	Msgs->SPSs = SPSs;
-
-	// NOTE: In a second phase we should probably scan the input access unit for inband SPS if the codec is avc3.
-
-	// Now go over the NALUs in the access unit and see what is there.
+	TArray<uint8> InbandSPSPPS;
+	bool bHaveInbandSPS = false;
+	bool bHaveInbandPPS = false;
+	// Go over the NALUs in the access unit and see what is there.
 	OutBSI.bIsDiscardable = true;
 	OutBSI.bIsSyncFrame = InOutAccessUnit->bIsSyncSample;
 	uint32* NALU = (uint32*)InOutAccessUnit->AUData;
@@ -138,15 +106,30 @@ IVideoDecoderInputBitstreamProcessor::EProcessResult FVideoDecoderInputBitstream
 			OutBSI.bIsDiscardable = false;
 		}
 		// IDR frame?
-		if ((nal & 0x1f) == 5)
+		const uint8 nut = nal & 0x1f;
+		if (nut == 5)
 		{
 			OutBSI.bIsSyncFrame = true;
 		}
 		// SEI message(s)?
-		if ((nal & 0x1f) == 6)
+		else if (nut == 6)
 		{
 			ElectraDecodersUtil::MPEG::ExtractSEIMessages(Msgs->SEIMessages, Electra::AdvancePointer(NALU, 5), naluLen-1, ElectraDecodersUtil::MPEG::ESEIStreamType::H264, false);
 		}
+		// SPS?
+		else if (nut == 7)
+		{
+			bHaveInbandSPS = true;
+			InbandSPSPPS.Append(StartCode, sizeof(StartCode));
+			InbandSPSPPS.Append(reinterpret_cast<const uint8*>(Electra::AdvancePointer(NALU, 4)), naluLen);
+		}
+		// PPS?
+		else if (nut == 8)
+		{
+			bHaveInbandPPS = true;
+			InbandSPSPPS.Append(StartCode, sizeof(StartCode));
+			InbandSPSPPS.Append(reinterpret_cast<const uint8*>(Electra::AdvancePointer(NALU, 4)), naluLen);
+		}
 
 		if (bReplaceLengthWithStartcode)
 		{
@@ -155,6 +138,52 @@ IVideoDecoderInputBitstreamProcessor::EProcessResult FVideoDecoderInputBitstream
 		NALU = Electra::AdvancePointer(NALU, naluLen + 4);
 	}
 
+	// Must have either both SPS and PPS or neither.
+	check(bHaveInbandSPS == bHaveInbandPPS);
+	if (bHaveInbandSPS != bHaveInbandPPS)
+	{
+		InbandSPSPPS.Empty();
+	}
+
+	// Compare against the inband SPS/PPS or the ones in the AU's CSD? If the AU has no CSD we always use the inband, even if that is empty.
+	const TConstArrayView<uint8>& CSDtoCompare(InbandSPSPPS.Num() ? InbandSPSPPS : InOutAccessUnit->AUCodecData.IsValid() ? InOutAccessUnit->AUCodecData->CodecSpecificData : InbandSPSPPS);
+
+	// Check for a change in SPS/PPS only on sync samples.
+	bool bIsDifferent = OutBSI.bIsSyncFrame && (!CurrentCodecData.IsValid() || CurrentCodecData->CodecSpecificData != CSDtoCompare);
+	if (bIsDifferent)
+	{
+		Result = IVideoDecoderInputBitstreamProcessor::EProcessResult::CSDChanged;
+		ElectraDecodersUtil::MPEG::H264::FAVCDecoderConfigurationRecord dcr;
+		if (!dcr.CreateFromCodecSpecificData(CSDtoCompare))
+		{
+			check(!"Failed to create DCR from CSD");
+			return Result;
+		}
+		TSharedPtr<FAccessUnit::CodecData, ESPMode::ThreadSafe> NewCSD = MakeShared<FAccessUnit::CodecData, ESPMode::ThreadSafe>();
+		NewCSD->RawCSD = dcr.GetRawData();
+		NewCSD->CodecSpecificData = dcr.GetCodecSpecificData();
+		// Copy any existing parsed info across, even though it may not be correct.
+		if (InOutAccessUnit->AUCodecData.IsValid())
+		{
+			NewCSD->ParsedInfo = InOutAccessUnit->AUCodecData->ParsedInfo;
+		}
+		// But set the new CSD and DCR on it.
+		NewCSD->ParsedInfo.SetDecoderConfigRecord(NewCSD->RawCSD);
+		NewCSD->ParsedInfo.SetCodecSpecificData(NewCSD->CodecSpecificData);
+		CurrentCodecData = NewCSD;
+
+		// We need the SPSs for later.
+		SPSs.Empty();
+		for(int32 i=0; i<dcr.GetSequenceParameterSets().Num(); ++i)
+		{
+			bool bOk = ElectraDecodersUtil::MPEG::H264::ParseSequenceParameterSet(SPSs, dcr.GetSequenceParameterSets()[i].GetData(), dcr.GetSequenceParameterSets()[i].Num());
+			check(bOk); (void)bOk;
+		}
+	}
+	InOutAccessUnit->AUCodecData = CurrentCodecData;
+
+	// Set the messages from the CSD on the access unit.
+	Msgs->SPSs = SPSs;
 	return Result;
 }
 
