ElectraPlayer: Ignoring SAP_delta_time on DASH segment indices as it is merely informational

#rnx

diff --git a/Engine/Plugins/Media/ElectraPlayer/Source/ElectraPlayerRuntime/Private/Runtime/Player/DASH/ManifestDASH.cpp b/Engine/Plugins/Media/ElectraPlayer/Source/ElectraPlayerRuntime/Private/Runtime/Player/DASH/ManifestDASH.cpp
--- a/Engine/Plugins/Media/ElectraPlayer/Source/ElectraPlayerRuntime/Private/Runtime/Player/DASH/ManifestDASH.cpp
+++ b/Engine/Plugins/Media/ElectraPlayer/Source/ElectraPlayerRuntime/Private/Runtime/Player/DASH/ManifestDASH.cpp
@@ -2885,13 +2885,17 @@ FManifestDASHInternal::FRepresentation::ESearchResult FManifestDASHInternal::FRe
 			bIsUsable = false;
 			return FManifestDASHInternal::FRepresentation::ESearchResult::BadType;
 		}
+		/*
+
+			The SAP_delta_time is merely an indication of the PTS of the SAP, it does not enter any time calculations.
+
 		else if (SegmentInfo.SAPDeltaTime != 0)
 		{
 			LogMessage(InPlayerSessionServices, IInfoLog::ELevel::Warning, FString::Printf(TEXT("Segment index for Representation \"%s\" uses SAP_delta_time. This may result in incorrect decoding."), *MPDRepresentation->GetID()));
 			bIsUsable = false;
 			return FManifestDASHInternal::FRepresentation::ESearchResult::BadType;
 		}
-
+		*/
 		CurrentD = SegmentInfo.SubSegmentDuration;
 		// Does the segment start on or after the time we're looking for?
 		if (CurrentT >= MediaLocalSearchTime)
