Verbose log for enabled ElectraPlayer analytic events on startup

#rnx

diff --git a/Engine/Plugins/Media/ElectraPlayer/Source/ElectraPlayerRuntime/Private/ElectraPlayerRuntimeModule.cpp b/Engine/Plugins/Media/ElectraPlayer/Source/ElectraPlayerRuntime/Private/ElectraPlayerRuntimeModule.cpp
--- a/Engine/Plugins/Media/ElectraPlayer/Source/ElectraPlayerRuntime/Private/ElectraPlayerRuntimeModule.cpp
+++ b/Engine/Plugins/Media/ElectraPlayer/Source/ElectraPlayerRuntime/Private/ElectraPlayerRuntimeModule.cpp
@@ -58,6 +58,9 @@ public:
 				ElectraConfig.EnabledAnalyticsEvents.Add(EnabledEvent, true);
 			}
 		}
+
+		UE_LOG(LogElectraPlayer, Verbose, TEXT("Found %d enabled ElectraPlayer analytic events."), ElectraConfig.EnabledAnalyticsEvents.Num());
+
 		Electra::Startup(ElectraConfig);
 
 		// Video renderer
