ElectraPlayer: Removed a check()

#rnx

diff --git a/Engine/Plugins/Media/ElectraPlayer/Source/ElectraPlayerRuntime/Private/Runtime/Decoder/BitstreamProcessors/VideoDecoderInputBitstreamProcessorH264.cpp b/Engine/Plugins/Media/ElectraPlayer/Source/ElectraPlayerRuntime/Private/Runtime/Decoder/BitstreamProcessors/VideoDecoderInputBitstreamProcessorH264.cpp
--- a/Engine/Plugins/Media/ElectraPlayer/Source/ElectraPlayerRuntime/Private/Runtime/Decoder/BitstreamProcessors/VideoDecoderInputBitstreamProcessorH264.cpp
+++ b/Engine/Plugins/Media/ElectraPlayer/Source/ElectraPlayerRuntime/Private/Runtime/Decoder/BitstreamProcessors/VideoDecoderInputBitstreamProcessorH264.cpp
@@ -139,7 +139,6 @@ IVideoDecoderInputBitstreamProcessor::EProcessResult FVideoDecoderInputBitstream
 	}
 
 	// Must have either both SPS and PPS or neither.
-	check(bHaveInbandSPS == bHaveInbandPPS);
 	if (bHaveInbandSPS != bHaveInbandPPS)
 	{
 		InbandSPSPPS.Empty();
