ElectraPlayer: Removed GetSegmentInformation() as the only value needed from that is the average segment duration, which can be obtained with less effort.

#rnx

diff --git a/Engine/Plugins/Media/ElectraPlayer/Source/ElectraPlayerRuntime/Private/Runtime/Player/ABRRules/ABRLiveStream.cpp b/Engine/Plugins/Media/ElectraPlayer/Source/ElectraPlayerRuntime/Private/Runtime/Player/ABRRules/ABRLiveStream.cpp
--- a/Engine/Plugins/Media/ElectraPlayer/Source/ElectraPlayerRuntime/Private/Runtime/Player/ABRRules/ABRLiveStream.cpp
+++ b/Engine/Plugins/Media/ElectraPlayer/Source/ElectraPlayerRuntime/Private/Runtime/Player/ABRRules/ABRLiveStream.cpp
@@ -847,8 +847,7 @@ void FABRLiveStream::RepresentationsChanged(EStreamType InStreamType, TSharedPtr
 
 		if (Representations.Num())
 		{
-			TArray<IManifest::IPlayPeriod::FSegmentInformation> NextSegments;
-			InCurrentPlayPeriod->GetSegmentInformation(NextSegments, WorkVars->AverageSegmentDuration, nullptr, FTimeValue(FTimeValue::MillisecondsToHNS(1000)), Representations[0]->AdaptationSetUniqueID, Representations[0]->RepresentationUniqueID);
+			InCurrentPlayPeriod->GetAverageSegmentDuration(WorkVars->AverageSegmentDuration, Representations[0]->AdaptationSetUniqueID, Representations[0]->RepresentationUniqueID);
 			for(auto &Rep : Representations)
 			{
 				if (Rep->Bitrate > WorkVars->HighestQualityStreamBitrate)
diff --git a/Engine/Plugins/Media/ElectraPlayer/Source/ElectraPlayerRuntime/Private/Runtime/Player/ABRRules/ABRVoDPlus.cpp b/Engine/Plugins/Media/ElectraPlayer/Source/ElectraPlayerRuntime/Private/Runtime/Player/ABRRules/ABRVoDPlus.cpp
--- a/Engine/Plugins/Media/ElectraPlayer/Source/ElectraPlayerRuntime/Private/Runtime/Player/ABRRules/ABRVoDPlus.cpp
+++ b/Engine/Plugins/Media/ElectraPlayer/Source/ElectraPlayerRuntime/Private/Runtime/Player/ABRRules/ABRVoDPlus.cpp
@@ -573,8 +573,7 @@ void FABROnDemandPlus::RepresentationsChanged(EStreamType InStreamType, TSharedP
 		WorkVars->NumQualities = Representations.Num();
 		if (Representations.Num())
 		{
-			TArray<IManifest::IPlayPeriod::FSegmentInformation> NextSegments;
-			InCurrentPlayPeriod->GetSegmentInformation(NextSegments, WorkVars->AverageSegmentDuration, nullptr, FTimeValue(FTimeValue::MillisecondsToHNS(1000)), Representations[0]->AdaptationSetUniqueID, Representations[0]->RepresentationUniqueID);
+			InCurrentPlayPeriod->GetAverageSegmentDuration(WorkVars->AverageSegmentDuration, Representations[0]->AdaptationSetUniqueID, Representations[0]->RepresentationUniqueID);
 			for(auto &Rep : Representations)
 			{
 				if (Rep->Bitrate > WorkVars->HighestQualityStreamBitrate)
diff --git a/Engine/Plugins/Media/ElectraPlayer/Source/ElectraPlayerRuntime/Private/Runtime/Player/DASH/ManifestDASH.cpp b/Engine/Plugins/Media/ElectraPlayer/Source/ElectraPlayerRuntime/Private/Runtime/Player/DASH/ManifestDASH.cpp
--- a/Engine/Plugins/Media/ElectraPlayer/Source/ElectraPlayerRuntime/Private/Runtime/Player/DASH/ManifestDASH.cpp
+++ b/Engine/Plugins/Media/ElectraPlayer/Source/ElectraPlayerRuntime/Private/Runtime/Player/DASH/ManifestDASH.cpp
@@ -24,7 +24,6 @@
 
 
 DECLARE_CYCLE_STAT(TEXT("FRepresentation::FindSegment"), STAT_ElectraPlayer_DASH_FindSegment, STATGROUP_ElectraPlayer);
-DECLARE_CYCLE_STAT(TEXT("FDASHPlayPeriod::GetSegmentInformation"), STAT_ElectraPlayer_DASH_GetSegmentInformation, STATGROUP_ElectraPlayer);
 
 namespace Electra
 {
@@ -209,7 +208,7 @@ public:
 	IManifest::FResult GetNextSegment(TSharedPtrTS<IStreamSegment>& OutSegment, TSharedPtrTS<const IStreamSegment> CurrentSegment, const FPlayStartOptions& Options) override;
 	IManifest::FResult GetRetrySegment(TSharedPtrTS<IStreamSegment>& OutSegment, TSharedPtrTS<const IStreamSegment> CurrentSegment, const FPlayStartOptions& Options, bool bReplaceWithFillerData) override;
 	void IncreaseSegmentFetchDelay(const FTimeValue& IncreaseAmount) override;
-	void GetSegmentInformation(TArray<FSegmentInformation>& OutSegmentInformation, FTimeValue& OutAverageSegmentDuration, TSharedPtrTS<const IStreamSegment> CurrentSegment, const FTimeValue& LookAheadTime, const FString& AdaptationSetID, const FString& RepresentationID) override;
+	void GetAverageSegmentDuration(FTimeValue& OutAverageSegmentDuration, const FString& AdaptationSetID, const FString& RepresentationID) override;
 
 private:
 	ELECTRA_IMPL_DEFAULT_ERROR_METHODS(DASHManifest);
@@ -1844,12 +1843,8 @@ void FDASHPlayPeriod::IncreaseSegmentFetchDelay(const FTimeValue& IncreaseAmount
 }
 
 
-void FDASHPlayPeriod::GetSegmentInformation(TArray<FSegmentInformation>& OutSegmentInformation, FTimeValue& OutAverageSegmentDuration, TSharedPtrTS<const IStreamSegment> CurrentSegment, const FTimeValue& LookAheadTime, const FString& AdaptationSetID, const FString& RepresentationID)
+void FDASHPlayPeriod::GetAverageSegmentDuration(FTimeValue& OutAverageSegmentDuration, const FString& AdaptationSetID, const FString& RepresentationID)
 {
-	SCOPE_CYCLE_COUNTER(STAT_ElectraPlayer_DASH_GetSegmentInformation);
-	CSV_SCOPED_TIMING_STAT(ElectraPlayer, DASH_GetSegmentInformation);
-
-	OutSegmentInformation.Empty();
 	OutAverageSegmentDuration.SetToInvalid();
 
 	TSharedPtrTS<FManifestDASHInternal> Manifest = GetCurrentManifest();
@@ -1865,7 +1860,7 @@ void FDASHPlayPeriod::GetSegmentInformation(TArray<FSegmentInformation>& OutSegm
 				TSharedPtrTS<FManifestDASHInternal::FRepresentation> Repr = AdaptationSet->GetRepresentationByUniqueID(RepresentationID);
 				if (Repr.IsValid())
 				{
-					Repr->GetSegmentInformation(OutSegmentInformation, OutAverageSegmentDuration, CurrentSegment, LookAheadTime, AdaptationSet);
+					Repr->GetAverageSegmentDuration(OutAverageSegmentDuration, AdaptationSet);
 				}
 			}
 		}
@@ -3673,13 +3668,8 @@ FManifestDASHInternal::FRepresentation::ESearchResult FManifestDASHInternal::FRe
 
 
 
-void FManifestDASHInternal::FRepresentation::GetSegmentInformation(TArray<IManifest::IPlayPeriod::FSegmentInformation>& OutSegmentInformation, FTimeValue& OutAverageSegmentDuration, TSharedPtrTS<const IStreamSegment> CurrentSegment, const FTimeValue& LookAheadTime, const TSharedPtrTS<IPlaybackAssetAdaptationSet>& AdaptationSet)
+void FManifestDASHInternal::FRepresentation::GetAverageSegmentDuration(FTimeValue& OutAverageSegmentDuration, const TSharedPtrTS<IPlaybackAssetAdaptationSet>& AdaptationSet)
 {
-	FTimeValue TimeToGo(LookAheadTime);
-	FTimeValue FixedSegmentDuration = FTimeValue(FTimeValue::MillisecondsToHNS(2000));
-	const int32 FixedBitrate = GetBitrate();
-	const FStreamSegmentRequestCommon* CurrentRequest = static_cast<const FStreamSegmentRequestCommon*>(CurrentSegment.Get());
-
 	// This is the same as in FindSegment(), only with no error checking since this method here is not critical.
 	TSharedPtrTS<FDashMPD_RepresentationType> MPDRepresentation = Representation.Pin();
 	if (MPDRepresentation.IsValid())
@@ -3722,20 +3712,13 @@ void FManifestDASHInternal::FRepresentation::GetSegmentInformation(TArray<IManif
 							uint32 SidxTimescale = Sidx->GetTimescale();
 							if (SidxTimescale)
 							{
-								int64 LocalSearchTime = CurrentRequest ? FTimeValue(CurrentRequest->Segment.Time + CurrentRequest->Segment.Duration, CurrentRequest->Segment.Timescale).GetAsTimebase(SidxTimescale) : 0;
-								int64 CurrentT = 0;
-								for(int32 nI=0, nIMax=Sidx->GetNumEntries(); nI<nIMax && TimeToGo>FTimeValue::GetZero(); ++nI)
+								int64 DurationSum = 0;
+								for(int32 nI=0, nIMax=Sidx->GetNumEntries(); nI<nIMax; ++nI)
 								{
-									const IParserISO14496_12::ISegmentIndex::FEntry& SegmentInfo = Sidx->GetEntry(nI);
-									if (CurrentT >= LocalSearchTime)
-									{
-										FTimeValue sd(SegmentInfo.SubSegmentDuration, SidxTimescale);
-										int64 ss = bExact ? SegmentInfo.Size : (int64)(FixedBitrate * sd.GetAsSeconds() / 8);
-										OutSegmentInformation.Emplace(IManifest::IPlayPeriod::FSegmentInformation({sd, ss}));
-										TimeToGo -= sd;
-									}
-									CurrentT += SegmentInfo.SubSegmentDuration;
+									DurationSum += Sidx->GetEntry(nI).SubSegmentDuration;
 								}
+								OutAverageSegmentDuration.SetFromND(DurationSum / Sidx->GetNumEntries(), SidxTimescale);
+								return;
 							}
 						}
 					}
@@ -3749,30 +3732,24 @@ void FManifestDASHInternal::FRepresentation::GetSegmentInformation(TArray<IManif
 						auto Selements = SegmentTimeline->GetS_Elements();
 						if (MPDTimescale)
 						{
-							int64 LocalSearchTime = CurrentRequest ? FTimeValue(CurrentRequest->Segment.Time + CurrentRequest->Segment.Duration, CurrentRequest->Segment.Timescale).GetAsTimebase(MPDTimescale) : 0;
-							int64 CurrentT = 0;
-							for(int32 nI=0, nIMax=Selements.Num(); nI<nIMax && TimeToGo>FTimeValue::GetZero(); ++nI)
+							int64 DurationSum = 0;
+							int64 NumElements = 0;
+							for(int32 nI=0, nIMax=Selements.Num(); nI<nIMax; ++nI)
 							{
-								FTimeValue sd(Selements[nI].D, MPDTimescale);
-								int64 ss = (int64)(FixedBitrate * sd.GetAsSeconds() / 8);
-								int32 R = Selements[nI].R;
-								// Note: If the repeat count is negative here then we don't care. We just repeat the same element until we are done.
-								//       There is no real need to be that precise to complicate the logic here.
-								//       We also do not care if there are gaps or overlaps in the timeline.
-								if (R < 0)
-								{
-									R = MAX_int32;
-								}
-								for(int32 nJ=R; nJ>=0  && TimeToGo>FTimeValue::GetZero(); --nJ)
+								// How many times does this entry repeat?
+								int32 RCount = Selements[nI].R + 1;
+								// If until the end of the period, for which we don't know how long that is exactly,
+								// we use a repeat count guess of 10, since a single repeat might skew the average if it is
+								// differing too much from the other duration values.
+								if (RCount <= 0)
 								{
-									if (CurrentT >= LocalSearchTime)
-									{
-										OutSegmentInformation.Emplace(IManifest::IPlayPeriod::FSegmentInformation({sd, ss}));
-										TimeToGo -= sd;
-									}
-									CurrentT += Selements[nI].D;
+									RCount = 10;
 								}
+								DurationSum += Selements[nI].D * RCount;
+								NumElements += RCount;
 							}
+							OutAverageSegmentDuration.SetFromND(DurationSum / NumElements, MPDTimescale);
+							return;
 						}
 					}
 					else
@@ -3782,25 +3759,17 @@ void FManifestDASHInternal::FRepresentation::GetSegmentInformation(TArray<IManif
 						uint32 MPDTimescale = GET_ATTR(SegmentTemplate, GetTimescale(), IsSet(), TMediaOptionalValue<uint32>(1)).Value();
 						if (Duration.IsSet() && Duration.Value() && MPDTimescale)
 						{
-							FixedSegmentDuration.SetFromND(Duration.Value(), MPDTimescale);
+							OutAverageSegmentDuration.SetFromND(Duration.Value(), MPDTimescale);
+							return;
 						}
 					}
 				}
 			}
 		}
 	}
-	// Nothing we could get from the actual representation. Cook up some default values.
-	int64 FixedSegmentSize = static_cast<int64>(FixedBitrate * (FixedSegmentDuration.GetAsSeconds() / 8));
-	while(TimeToGo > FTimeValue::GetZero())
-	{
-		OutSegmentInformation.Emplace(IManifest::IPlayPeriod::FSegmentInformation({FixedSegmentDuration, FixedSegmentSize}));
-		TimeToGo -= FixedSegmentDuration;
-	}
-	// Set up average duration.
-	if (OutSegmentInformation.Num())
-	{
-		OutAverageSegmentDuration = (LookAheadTime - TimeToGo) / OutSegmentInformation.Num();
-	}
+	// Nothing we could get from the actual representation.
+	// Set some commonly used duration.
+	OutAverageSegmentDuration.SetFromMilliseconds(4000);
 }
 
 
diff --git a/Engine/Plugins/Media/ElectraPlayer/Source/ElectraPlayerRuntime/Private/Runtime/Player/DASH/PlaylistReaderDASH_Internal.h b/Engine/Plugins/Media/ElectraPlayer/Source/ElectraPlayerRuntime/Private/Runtime/Player/DASH/PlaylistReaderDASH_Internal.h
--- a/Engine/Plugins/Media/ElectraPlayer/Source/ElectraPlayerRuntime/Private/Runtime/Player/DASH/PlaylistReaderDASH_Internal.h
+++ b/Engine/Plugins/Media/ElectraPlayer/Source/ElectraPlayerRuntime/Private/Runtime/Player/DASH/PlaylistReaderDASH_Internal.h
@@ -173,7 +173,7 @@ public:
 
 		ESearchResult FindSegment(IPlayerSessionServices* PlayerSessionServices, FSegmentInformationCommon& OutSegmentInfo, TArray<TWeakPtrTS<FMPDLoadRequestDASH>>& OutRemoteElementLoadRequests, const FSegmentSearchOption& InSearchOptions);
 
-		void GetSegmentInformation(TArray<IManifest::IPlayPeriod::FSegmentInformation>& OutSegmentInformation, FTimeValue& OutAverageSegmentDuration, TSharedPtrTS<const IStreamSegment> CurrentSegment, const FTimeValue& LookAheadTime, const TSharedPtrTS<IPlaybackAssetAdaptationSet>& AdaptationSet);
+		void GetAverageSegmentDuration(FTimeValue& OutAverageSegmentDuration, const TSharedPtrTS<IPlaybackAssetAdaptationSet>& AdaptationSet);
 
 		int32 GetSelectionPriority() const { return SelectionPriority; }
 
diff --git a/Engine/Plugins/Media/ElectraPlayer/Source/ElectraPlayerRuntime/Private/Runtime/Player/HLS/PlaylistHLS.cpp b/Engine/Plugins/Media/ElectraPlayer/Source/ElectraPlayerRuntime/Private/Runtime/Player/HLS/PlaylistHLS.cpp
--- a/Engine/Plugins/Media/ElectraPlayer/Source/ElectraPlayerRuntime/Private/Runtime/Player/HLS/PlaylistHLS.cpp
+++ b/Engine/Plugins/Media/ElectraPlayer/Source/ElectraPlayerRuntime/Private/Runtime/Player/HLS/PlaylistHLS.cpp
@@ -101,7 +101,7 @@ public:
 	FResult GetNextSegment(TSharedPtrTS<IStreamSegment>& OutSegment, TSharedPtrTS<const IStreamSegment> InCurrentSegment, const FPlayStartOptions& InOptions) override;
 	FResult GetRetrySegment(TSharedPtrTS<IStreamSegment>& OutSegment, TSharedPtrTS<const IStreamSegment> InCurrentSegment, const FPlayStartOptions& InOptions, bool bInReplaceWithFillerData) override;
 	void IncreaseSegmentFetchDelay(const FTimeValue& InIncreaseAmount) override;
-	void GetSegmentInformation(TArray<FSegmentInformation>& OutSegmentInformation, FTimeValue& OutAverageSegmentDuration, TSharedPtrTS<const IStreamSegment> InCurrentSegment, const FTimeValue& InLookAheadTime, const FString& InAdaptationSetID, const FString& InRepresentationID) override;
+	void GetAverageSegmentDuration(FTimeValue& OutAverageSegmentDuration, const FString& InAdaptationSetID, const FString& InRepresentationID) override;
 
 	enum class ENextSegType
 	{
@@ -2611,7 +2611,7 @@ void FActiveHLSPlaylist::FPlayPeriod::GetActiveMediaPlaylists(TArray<TSharedPtr<
 	}
 }
 
-void FActiveHLSPlaylist::FPlayPeriod::GetSegmentInformation(TArray<FSegmentInformation>& /*OutSegmentInformation*/, FTimeValue& OutAverageSegmentDuration, TSharedPtrTS<const IStreamSegment> /*InCurrentSegment*/, const FTimeValue& /*InLookAheadTime*/, const FString& InAdaptationSetID, const FString& /*InRepresentationID*/)
+void FActiveHLSPlaylist::FPlayPeriod::GetAverageSegmentDuration(FTimeValue& OutAverageSegmentDuration, const FString& InAdaptationSetID, const FString& /*InRepresentationID*/)
 {
 	// The segments of a variant should have equal durations across all variants.
 	// For renditions they should be similar (although they can't really be due to probably different codecs (AAC has different block sizes than AC3)).
diff --git a/Engine/Plugins/Media/ElectraPlayer/Source/ElectraPlayerRuntime/Private/Runtime/Player/Manifest.h b/Engine/Plugins/Media/ElectraPlayer/Source/ElectraPlayerRuntime/Private/Runtime/Player/Manifest.h
--- a/Engine/Plugins/Media/ElectraPlayer/Source/ElectraPlayerRuntime/Private/Runtime/Player/Manifest.h
+++ b/Engine/Plugins/Media/ElectraPlayer/Source/ElectraPlayerRuntime/Private/Runtime/Player/Manifest.h
@@ -553,23 +553,10 @@ namespace Electra
 			 */
 			virtual void IncreaseSegmentFetchDelay(const FTimeValue& IncreaseAmount) = 0;
 
-			struct FSegmentInformation
-			{
-				FTimeValue	Duration;
-				int64		ByteSize = 0;
-			};
-
 			/**
-			 * Obtains information on the stream segmentation of a particular stream starting at a given current reference segment (optional, if not given returns suitable default values).
-			 *
-			 * @param OutSegmentInformation
-			 * @param OutAverageSegmentDuration
-			 * @param CurrentSegment
-			 * @param LookAheadTime
-			 * @param AdaptationSetID
-			 * @param RepresentationID
+			 * Returns the average segment duration.
 			 */
-			virtual void GetSegmentInformation(TArray<FSegmentInformation>& OutSegmentInformation, FTimeValue& OutAverageSegmentDuration, TSharedPtrTS<const IStreamSegment> CurrentSegment, const FTimeValue& LookAheadTime, const FString& AdaptationSetID, const FString& RepresentationID) = 0;
+			virtual void GetAverageSegmentDuration(FTimeValue& OutAverageSegmentDuration, const FString& AdaptationSetID, const FString& RepresentationID) = 0;
 		};
 
 		//! Finds the playback period the specified start time falls into.
diff --git a/Engine/Plugins/Media/ElectraPlayer/Source/ElectraPlayerRuntime/Private/Runtime/Player/mkv/ManifestMKV.cpp b/Engine/Plugins/Media/ElectraPlayer/Source/ElectraPlayerRuntime/Private/Runtime/Player/mkv/ManifestMKV.cpp
--- a/Engine/Plugins/Media/ElectraPlayer/Source/ElectraPlayerRuntime/Private/Runtime/Player/mkv/ManifestMKV.cpp
+++ b/Engine/Plugins/Media/ElectraPlayer/Source/ElectraPlayerRuntime/Private/Runtime/Player/mkv/ManifestMKV.cpp
@@ -22,7 +22,6 @@
 
 
 DECLARE_CYCLE_STAT(TEXT("FPlayPeriodMKV::FindSegment"), STAT_ElectraPlayer_MKV_FindSegment, STATGROUP_ElectraPlayer);
-DECLARE_CYCLE_STAT(TEXT("FPlayPeriodMKV::GetSegmentInformation"), STAT_ElectraPlayer_MKV_GetSegmentInformation, STATGROUP_ElectraPlayer);
 
 namespace Electra
 {
@@ -603,21 +602,14 @@ IManifest::FResult FManifestMKVInternal::FPlayPeriodMKV::GetRetrySegment(TShared
 
 //-----------------------------------------------------------------------------
 /**
- * Returns segment information for the next n segments.
- *
- * @param OutSegmentInformation
- * @param OutAverageSegmentDuration
- * @param CurrentSegment
- * @param LookAheadTime
- * @param AdaptationSet
- * @param Representation
+ * Returns the average segment duration.
  */
-void FManifestMKVInternal::FPlayPeriodMKV::GetSegmentInformation(TArray<FSegmentInformation>& OutSegmentInformation, FTimeValue& OutAverageSegmentDuration, TSharedPtrTS<const IStreamSegment> CurrentSegment, const FTimeValue& LookAheadTime, const FString& AdaptationSetID, const FString& RepresentationID)
+void FManifestMKVInternal::FPlayPeriodMKV::GetAverageSegmentDuration(FTimeValue& OutAverageSegmentDuration, const FString& AdaptationSetID, const FString& RepresentationID)
 {
 	TSharedPtrTS<FTimelineAssetMKV> ma = MediaAsset.Pin();
 	if (ma.IsValid())
 	{
-		ma->GetSegmentInformation(OutSegmentInformation, OutAverageSegmentDuration, CurrentSegment, LookAheadTime, AdaptationSetID, RepresentationID);
+		ma->GetAverageSegmentDuration(OutAverageSegmentDuration, AdaptationSetID, RepresentationID);
 	}
 }
 
@@ -930,17 +922,10 @@ IManifest::FResult FManifestMKVInternal::FTimelineAssetMKV::GetLoopingSegment(TS
 	return GetStartingSegment(OutSegment, SequenceState, SelectedStreamID, StartPosition, SearchType, -1);
 }
 
-void FManifestMKVInternal::FTimelineAssetMKV::GetSegmentInformation(TArray<IManifest::IPlayPeriod::FSegmentInformation>& OutSegmentInformation, FTimeValue& OutAverageSegmentDuration, TSharedPtrTS<const IStreamSegment> CurrentSegment, const FTimeValue& LookAheadTime, const FString& AdaptationSetID, const FString& RepresentationID)
+void FManifestMKVInternal::FTimelineAssetMKV::GetAverageSegmentDuration(FTimeValue& OutAverageSegmentDuration, const FString& AdaptationSetID, const FString& RepresentationID)
 {
-	SCOPE_CYCLE_COUNTER(STAT_ElectraPlayer_MKV_GetSegmentInformation);
-	CSV_SCOPED_TIMING_STAT(ElectraPlayer, MKV_GetSegmentInformation);
-
 	// This is not expected to be called. And if it does we return a dummy entry.
-	OutAverageSegmentDuration.SetFromSeconds(60.0);
-	OutSegmentInformation.Empty();
-	IManifest::IPlayPeriod::FSegmentInformation& si = OutSegmentInformation.AddDefaulted_GetRef();
-	si.ByteSize = 1024 * 1024 * 1024;
-	si.Duration.SetFromSeconds(60.0);
+	OutAverageSegmentDuration.SetFromSeconds(5.0);
 }
 
 TSharedPtrTS<IParserMKV> FManifestMKVInternal::FTimelineAssetMKV::GetMKVParser()
diff --git a/Engine/Plugins/Media/ElectraPlayer/Source/ElectraPlayerRuntime/Private/Runtime/Player/mkv/ManifestMKV.h b/Engine/Plugins/Media/ElectraPlayer/Source/ElectraPlayerRuntime/Private/Runtime/Player/mkv/ManifestMKV.h
--- a/Engine/Plugins/Media/ElectraPlayer/Source/ElectraPlayerRuntime/Private/Runtime/Player/mkv/ManifestMKV.h
+++ b/Engine/Plugins/Media/ElectraPlayer/Source/ElectraPlayerRuntime/Private/Runtime/Player/mkv/ManifestMKV.h
@@ -252,7 +252,7 @@ public:
 		FResult GetRetrySegment(TSharedPtrTS<IStreamSegment>& OutSegment, TSharedPtrTS<const IStreamSegment> CurrentSegment, const FPlayStartOptions& Options, bool bReplaceWithFillerData);
 		FResult GetLoopingSegment(TSharedPtrTS<IStreamSegment>& OutSegment, const FPlayerSequenceState& InSequenceState, FString SelectedStreamID, const FPlayStartPosition& StartPosition, ESearchType SearchType);
 
-		void GetSegmentInformation(TArray<IManifest::IPlayPeriod::FSegmentInformation>& OutSegmentInformation, FTimeValue& OutAverageSegmentDuration, TSharedPtrTS<const IStreamSegment> CurrentSegment, const FTimeValue& LookAheadTime, const FString& AdaptationSetID, const FString& RepresentationID);
+		void GetAverageSegmentDuration(FTimeValue& OutAverageSegmentDuration, const FString& AdaptationSetID, const FString& RepresentationID);
 		TSharedPtrTS<IParserMKV> GetMKVParser();
 		const FString& GetMediaURL() const
 		{
@@ -293,7 +293,7 @@ public:
 		FResult GetRetrySegment(TSharedPtrTS<IStreamSegment>& OutSegment, TSharedPtrTS<const IStreamSegment> CurrentSegment, const FPlayStartOptions& Options, bool bReplaceWithFillerData) override;
 		FResult GetLoopingSegment(TSharedPtrTS<IStreamSegment>& OutSegment, const FPlayerSequenceState& InSequenceState, const FPlayStartPosition& StartPosition, ESearchType SearchType) override;
 		void IncreaseSegmentFetchDelay(const FTimeValue& IncreaseAmount) override;
-		void GetSegmentInformation(TArray<FSegmentInformation>& OutSegmentInformation, FTimeValue& OutAverageSegmentDuration, TSharedPtrTS<const IStreamSegment> CurrentSegment, const FTimeValue& LookAheadTime, const FString& AdaptationSetID, const FString& RepresentationID) override;
+		void GetAverageSegmentDuration(FTimeValue& OutAverageSegmentDuration, const FString& AdaptationSetID, const FString& RepresentationID) override;
 
 	private:
 		void SelectInitialStream(EStreamType StreamType);
diff --git a/Engine/Plugins/Media/ElectraPlayer/Source/ElectraPlayerRuntime/Private/Runtime/Player/mp4/ManifestMP4.cpp b/Engine/Plugins/Media/ElectraPlayer/Source/ElectraPlayerRuntime/Private/Runtime/Player/mp4/ManifestMP4.cpp
--- a/Engine/Plugins/Media/ElectraPlayer/Source/ElectraPlayerRuntime/Private/Runtime/Player/mp4/ManifestMP4.cpp
+++ b/Engine/Plugins/Media/ElectraPlayer/Source/ElectraPlayerRuntime/Private/Runtime/Player/mp4/ManifestMP4.cpp
@@ -27,7 +27,6 @@
 
 
 DECLARE_CYCLE_STAT(TEXT("FPlayPeriodMP4::FindSegment"), STAT_ElectraPlayer_MP4_FindSegment, STATGROUP_ElectraPlayer);
-DECLARE_CYCLE_STAT(TEXT("FPlayPeriodMP4::GetSegmentInformation"), STAT_ElectraPlayer_MP4_GetSegmentInformation, STATGROUP_ElectraPlayer);
 
 namespace Electra
 {
@@ -632,21 +631,14 @@ IManifest::FResult FManifestMP4Internal::FPlayPeriodMP4::GetRetrySegment(TShared
 
 //-----------------------------------------------------------------------------
 /**
- * Returns segment information for the next n segments.
- *
- * @param OutSegmentInformation
- * @param OutAverageSegmentDuration
- * @param CurrentSegment
- * @param LookAheadTime
- * @param AdaptationSet
- * @param Representation
+ * Returns the average segment duration.
  */
-void FManifestMP4Internal::FPlayPeriodMP4::GetSegmentInformation(TArray<FSegmentInformation>& OutSegmentInformation, FTimeValue& OutAverageSegmentDuration, TSharedPtrTS<const IStreamSegment> CurrentSegment, const FTimeValue& LookAheadTime, const FString& AdaptationSetID, const FString& RepresentationID)
+void FManifestMP4Internal::FPlayPeriodMP4::GetAverageSegmentDuration(FTimeValue& OutAverageSegmentDuration, const FString& AdaptationSetID, const FString& RepresentationID)
 {
 	TSharedPtrTS<FTimelineAssetMP4> ma = MediaAsset.Pin();
 	if (ma.IsValid())
 	{
-		ma->GetSegmentInformation(OutSegmentInformation, OutAverageSegmentDuration, CurrentSegment, LookAheadTime, AdaptationSetID, RepresentationID);
+		ma->GetAverageSegmentDuration(OutAverageSegmentDuration, AdaptationSetID, RepresentationID);
 	}
 }
 
@@ -1333,17 +1325,10 @@ IManifest::FResult FManifestMP4Internal::FTimelineAssetMP4::GetLoopingSegment(TS
 }
 
 
-void FManifestMP4Internal::FTimelineAssetMP4::GetSegmentInformation(TArray<IManifest::IPlayPeriod::FSegmentInformation>& OutSegmentInformation, FTimeValue& OutAverageSegmentDuration, TSharedPtrTS<const IStreamSegment> CurrentSegment, const FTimeValue& LookAheadTime, const FString& AdaptationSetID, const FString& RepresentationID)
+void FManifestMP4Internal::FTimelineAssetMP4::GetAverageSegmentDuration(FTimeValue& OutAverageSegmentDuration, const FString& AdaptationSetID, const FString& RepresentationID)
 {
-	SCOPE_CYCLE_COUNTER(STAT_ElectraPlayer_MP4_GetSegmentInformation);
-	CSV_SCOPED_TIMING_STAT(ElectraPlayer, MP4_GetSegmentInformation);
-
 	// This is not expected to be called. And if it does we return a dummy entry.
-	OutAverageSegmentDuration.SetFromSeconds(60.0);
-	OutSegmentInformation.Empty();
-	IManifest::IPlayPeriod::FSegmentInformation& si = OutSegmentInformation.AddDefaulted_GetRef();
-	si.ByteSize = 1024 * 1024 * 1024;
-	si.Duration.SetFromSeconds(60.0);
+	OutAverageSegmentDuration.SetFromSeconds(5.0);
 }
 
 TSharedPtrTS<IParserISO14496_12> FManifestMP4Internal::FTimelineAssetMP4::GetMoovBoxParser()
diff --git a/Engine/Plugins/Media/ElectraPlayer/Source/ElectraPlayerRuntime/Private/Runtime/Player/mp4/ManifestMP4.h b/Engine/Plugins/Media/ElectraPlayer/Source/ElectraPlayerRuntime/Private/Runtime/Player/mp4/ManifestMP4.h
--- a/Engine/Plugins/Media/ElectraPlayer/Source/ElectraPlayerRuntime/Private/Runtime/Player/mp4/ManifestMP4.h
+++ b/Engine/Plugins/Media/ElectraPlayer/Source/ElectraPlayerRuntime/Private/Runtime/Player/mp4/ManifestMP4.h
@@ -237,7 +237,7 @@ public:
 		FResult GetRetrySegment(TSharedPtrTS<IStreamSegment>& OutSegment, TSharedPtrTS<const IStreamSegment> CurrentSegment, const FPlayStartOptions& Options, bool bReplaceWithFillerData);
 		FResult GetLoopingSegment(TSharedPtrTS<IStreamSegment>& OutSegment, const FPlayerSequenceState& InSequenceState, const FPlayStartPosition& StartPosition, ESearchType SearchType);
 
-		void GetSegmentInformation(TArray<IManifest::IPlayPeriod::FSegmentInformation>& OutSegmentInformation, FTimeValue& OutAverageSegmentDuration, TSharedPtrTS<const IStreamSegment> CurrentSegment, const FTimeValue& LookAheadTime, const FString& AdaptationSetID, const FString& RepresentationID);
+		void GetAverageSegmentDuration(FTimeValue& OutAverageSegmentDuration, const FString& AdaptationSetID, const FString& RepresentationID);
 		TSharedPtrTS<IParserISO14496_12>	GetMoovBoxParser();
 		const FString& GetMediaURL() const
 		{
@@ -297,7 +297,7 @@ public:
 		FResult GetRetrySegment(TSharedPtrTS<IStreamSegment>& OutSegment, TSharedPtrTS<const IStreamSegment> CurrentSegment, const FPlayStartOptions& Options, bool bReplaceWithFillerData) override;
 		FResult GetLoopingSegment(TSharedPtrTS<IStreamSegment>& OutSegment, const FPlayerSequenceState& InSequenceState, const FPlayStartPosition& StartPosition, ESearchType SearchType) override;
 		void IncreaseSegmentFetchDelay(const FTimeValue& IncreaseAmount) override;
-		void GetSegmentInformation(TArray<FSegmentInformation>& OutSegmentInformation, FTimeValue& OutAverageSegmentDuration, TSharedPtrTS<const IStreamSegment> CurrentSegment, const FTimeValue& LookAheadTime, const FString& AdaptationSetID, const FString& RepresentationID) override;
+		void GetAverageSegmentDuration(FTimeValue& OutAverageSegmentDuration, const FString& AdaptationSetID, const FString& RepresentationID) override;
 
 	private:
 		void SelectInitialStream(EStreamType StreamType);
diff --git a/Engine/Plugins/Media/ElectraPlayer/Source/ElectraPlayerRuntime/Private/Runtime/Player/mpegaudio/ManifestMPEGAudio.cpp b/Engine/Plugins/Media/ElectraPlayer/Source/ElectraPlayerRuntime/Private/Runtime/Player/mpegaudio/ManifestMPEGAudio.cpp
--- a/Engine/Plugins/Media/ElectraPlayer/Source/ElectraPlayerRuntime/Private/Runtime/Player/mpegaudio/ManifestMPEGAudio.cpp
+++ b/Engine/Plugins/Media/ElectraPlayer/Source/ElectraPlayerRuntime/Private/Runtime/Player/mpegaudio/ManifestMPEGAudio.cpp
@@ -27,7 +27,6 @@
 #define USE_LOW_LATENCY_DESCRIPTOR 0
 
 DECLARE_CYCLE_STAT(TEXT("FPlayPeriodMPEGAudio::FindSegment"), STAT_ElectraPlayer_MPEGAudio_FindSegment, STATGROUP_ElectraPlayer);
-DECLARE_CYCLE_STAT(TEXT("FPlayPeriodMPEGAudio::GetSegmentInformation"), STAT_ElectraPlayer_MPEGAudio_GetSegmentInformation, STATGROUP_ElectraPlayer);
 
 namespace Electra
 {
@@ -657,21 +656,14 @@ IManifest::FResult FManifestMPEGAudioInternal::FPlayPeriodMPEGAudio::GetRetrySeg
 
 //-----------------------------------------------------------------------------
 /**
- * Returns segment information for the next n segments.
- *
- * @param OutSegmentInformation
- * @param OutAverageSegmentDuration
- * @param CurrentSegment
- * @param LookAheadTime
- * @param AdaptationSet
- * @param Representation
+ * Returns the average segment duration.
  */
-void FManifestMPEGAudioInternal::FPlayPeriodMPEGAudio::GetSegmentInformation(TArray<FSegmentInformation>& OutSegmentInformation, FTimeValue& OutAverageSegmentDuration, TSharedPtrTS<const IStreamSegment> CurrentSegment, const FTimeValue& LookAheadTime, const FString& AdaptationSetID, const FString& RepresentationID)
+void FManifestMPEGAudioInternal::FPlayPeriodMPEGAudio::GetAverageSegmentDuration(FTimeValue& OutAverageSegmentDuration, const FString& AdaptationSetID, const FString& RepresentationID)
 {
 	TSharedPtrTS<FTimelineAssetMPEGAudio> ma = MediaAsset.Pin();
 	if (ma.IsValid())
 	{
-		ma->GetSegmentInformation(OutSegmentInformation, OutAverageSegmentDuration, CurrentSegment, LookAheadTime, AdaptationSetID, RepresentationID);
+		ma->GetAverageSegmentDuration(OutAverageSegmentDuration, AdaptationSetID, RepresentationID);
 	}
 }
 
@@ -1059,17 +1051,10 @@ IManifest::FResult FManifestMPEGAudioInternal::FTimelineAssetMPEGAudio::GetLoopi
 }
 
 
-void FManifestMPEGAudioInternal::FTimelineAssetMPEGAudio::GetSegmentInformation(TArray<IManifest::IPlayPeriod::FSegmentInformation>& OutSegmentInformation, FTimeValue& OutAverageSegmentDuration, TSharedPtrTS<const IStreamSegment> CurrentSegment, const FTimeValue& LookAheadTime, const FString& AdaptationSetID, const FString& RepresentationID)
+void FManifestMPEGAudioInternal::FTimelineAssetMPEGAudio::GetAverageSegmentDuration(FTimeValue& OutAverageSegmentDuration, const FString& AdaptationSetID, const FString& RepresentationID)
 {
-	SCOPE_CYCLE_COUNTER(STAT_ElectraPlayer_MPEGAudio_GetSegmentInformation);
-	CSV_SCOPED_TIMING_STAT(ElectraPlayer, MPEGAudio_GetSegmentInformation);
-
 	// This is not expected to be called. And if it does we return a dummy entry.
-	OutAverageSegmentDuration.SetFromSeconds(60.0);
-	OutSegmentInformation.Empty();
-	IManifest::IPlayPeriod::FSegmentInformation& si = OutSegmentInformation.AddDefaulted_GetRef();
-	si.ByteSize = 16384;
-	si.Duration.SetFromSeconds(1.0);
+	OutAverageSegmentDuration.SetFromSeconds(5.0);
 }
 
 FErrorDetail FManifestMPEGAudioInternal::FAdaptationSetMPEGAudio::CreateFrom(const FMPEGAudioInfoHeader& InHeader, const FString& URL)
diff --git a/Engine/Plugins/Media/ElectraPlayer/Source/ElectraPlayerRuntime/Private/Runtime/Player/mpegaudio/ManifestMPEGAudio.h b/Engine/Plugins/Media/ElectraPlayer/Source/ElectraPlayerRuntime/Private/Runtime/Player/mpegaudio/ManifestMPEGAudio.h
--- a/Engine/Plugins/Media/ElectraPlayer/Source/ElectraPlayerRuntime/Private/Runtime/Player/mpegaudio/ManifestMPEGAudio.h
+++ b/Engine/Plugins/Media/ElectraPlayer/Source/ElectraPlayerRuntime/Private/Runtime/Player/mpegaudio/ManifestMPEGAudio.h
@@ -267,7 +267,7 @@ public:
 		FResult GetRetrySegment(TSharedPtrTS<IStreamSegment>& OutSegment, TSharedPtrTS<const IStreamSegment> CurrentSegment, const FPlayStartOptions& Options, bool bReplaceWithFillerData);
 		FResult GetLoopingSegment(TSharedPtrTS<IStreamSegment>& OutSegment, const FPlayerSequenceState& InSequenceState, const FPlayStartPosition& StartPosition, ESearchType SearchType);
 
-		void GetSegmentInformation(TArray<IManifest::IPlayPeriod::FSegmentInformation>& OutSegmentInformation, FTimeValue& OutAverageSegmentDuration, TSharedPtrTS<const IStreamSegment> CurrentSegment, const FTimeValue& LookAheadTime, const FString& AdaptationSetID, const FString& RepresentationID);
+		void GetAverageSegmentDuration(FTimeValue& OutAverageSegmentDuration, const FString& AdaptationSetID, const FString& RepresentationID);
 		const FString& GetMediaURL() const
 		{ return MediaURL; }
 
@@ -302,7 +302,7 @@ public:
 		FResult GetRetrySegment(TSharedPtrTS<IStreamSegment>& OutSegment, TSharedPtrTS<const IStreamSegment> CurrentSegment, const FPlayStartOptions& Options, bool bReplaceWithFillerData) override;
 		FResult GetLoopingSegment(TSharedPtrTS<IStreamSegment>& OutSegment, const FPlayerSequenceState& InSequenceState, const FPlayStartPosition& StartPosition, ESearchType SearchType) override;
 		void IncreaseSegmentFetchDelay(const FTimeValue& IncreaseAmount) override;
-		void GetSegmentInformation(TArray<FSegmentInformation>& OutSegmentInformation, FTimeValue& OutAverageSegmentDuration, TSharedPtrTS<const IStreamSegment> CurrentSegment, const FTimeValue& LookAheadTime, const FString& AdaptationSetID, const FString& RepresentationID) override;
+		void GetAverageSegmentDuration(FTimeValue& OutAverageSegmentDuration, const FString& AdaptationSetID, const FString& RepresentationID) override;
 
 	private:
 		void SelectInitialStream(EStreamType StreamType);
