Electra: Replacing media critical section with common critical section

#rnx

diff --git a/Engine/Plugins/Media/ElectraPlayer/Source/ElectraPlayerRuntime/Private/Runtime/HTTP/HTTPManager.h b/Engine/Plugins/Media/ElectraPlayer/Source/ElectraPlayerRuntime/Private/Runtime/HTTP/HTTPManager.h
--- a/Engine/Plugins/Media/ElectraPlayer/Source/ElectraPlayerRuntime/Private/Runtime/HTTP/HTTPManager.h
+++ b/Engine/Plugins/Media/ElectraPlayer/Source/ElectraPlayerRuntime/Private/Runtime/HTTP/HTTPManager.h
@@ -178,13 +178,13 @@ namespace Electra
 			}
 			void Clear()
 			{
-				FMediaCriticalSection::ScopedLock lock(Lock);
+				FScopeLock lock(&Lock);
 				ProgressDelegate.Unbind();
 				CompletionDelegate.Unbind();
 			}
 			int32 CallProgressDelegate(const FRequest* Request)
 			{
-				FMediaCriticalSection::ScopedLock lock(Lock);
+				FScopeLock lock(&Lock);
 				if (ProgressDelegate.IsBound())
 				{
 					return ProgressDelegate.Execute(Request);
@@ -193,15 +193,15 @@ namespace Electra
 			}
 			void CallCompletionDelegate(const FRequest* Request)
 			{
-				FMediaCriticalSection::ScopedLock lock(Lock);
+				FScopeLock lock(&Lock);
 				if (CompletionDelegate.IsBound())
 				{
 					CompletionDelegate.Execute(Request);
 				}
 			}
-			FMediaCriticalSection	Lock;
-			FProgressDelegate		ProgressDelegate;
-			FCompletionDelegate		CompletionDelegate;
+			FCriticalSection Lock;
+			FProgressDelegate ProgressDelegate;
+			FCompletionDelegate CompletionDelegate;
 		};
 
 		struct FParams
diff --git a/Engine/Plugins/Media/ElectraPlayer/Source/ElectraPlayerRuntime/Private/Runtime/Player/mkv/StreamReaderMKV.cpp b/Engine/Plugins/Media/ElectraPlayer/Source/ElectraPlayerRuntime/Private/Runtime/Player/mkv/StreamReaderMKV.cpp
--- a/Engine/Plugins/Media/ElectraPlayer/Source/ElectraPlayerRuntime/Private/Runtime/Player/mkv/StreamReaderMKV.cpp
+++ b/Engine/Plugins/Media/ElectraPlayer/Source/ElectraPlayerRuntime/Private/Runtime/Player/mkv/StreamReaderMKV.cpp
@@ -235,7 +235,7 @@ void FStreamReaderMKV::HTTPUpdateStats(const FTimeValue& CurrentTime, const IEle
 	TSharedPtrTS<FStreamSegmentRequestMKV> SegmentRequest = CurrentRequest;
 	if (SegmentRequest.IsValid())
 	{
-		FMediaCriticalSection::ScopedLock lock(MetricUpdateLock);
+		FScopeLock lock(&MetricUpdateLock);
 		SegmentRequest->ConnectionInfo = Request->ConnectionInfo;
 		// Update the current download stats which we report periodically to the ABR.
 		Metrics::FSegmentDownloadStats& ds = SegmentRequest->DownloadStats;
diff --git a/Engine/Plugins/Media/ElectraPlayer/Source/ElectraPlayerRuntime/Private/Runtime/Player/mkv/StreamReaderMKV.h b/Engine/Plugins/Media/ElectraPlayer/Source/ElectraPlayerRuntime/Private/Runtime/Player/mkv/StreamReaderMKV.h
--- a/Engine/Plugins/Media/ElectraPlayer/Source/ElectraPlayerRuntime/Private/Runtime/Player/mkv/StreamReaderMKV.h
+++ b/Engine/Plugins/Media/ElectraPlayer/Source/ElectraPlayerRuntime/Private/Runtime/Player/mkv/StreamReaderMKV.h
@@ -47,7 +47,7 @@ public:
 
 
 	// Which tracks are enabled and will be extracted.
-	TArray<uint64> EnabledTrackIDs;			
+	TArray<uint64> EnabledTrackIDs;
 
 	// The generated MKV asset
 	TSharedPtrTS<ITimelineMediaAsset> MediaAsset;
@@ -66,7 +66,7 @@ public:
 	int64 TimestampSequenceIndex = 0;
 
 	// Type of the stream for which the request was generated
-	EStreamType PrimaryStreamType = EStreamType::Unsupported;			
+	EStreamType PrimaryStreamType = EStreamType::Unsupported;
 	// Where to start in the file
 	int64 FileStartOffset = -1;
 	// Where to end in the file (for HTTP range GET requests)
@@ -271,7 +271,7 @@ private:
 	TSharedPtrTS<IElectraHttpManager::FRequest> HTTPRequest;
 	FMediaEvent WorkSignal;
 	FReadBuffer ReadBuffer;
-	FMediaCriticalSection MetricUpdateLock;
+	FCriticalSection MetricUpdateLock;
 	TMap<uint64, FSelectedTrackData> ActiveTrackMap;
 	static uint32 UniqueDownloadID;
 };
diff --git a/Engine/Plugins/Media/ElectraPlayer/Source/ElectraPlayerRuntime/Private/Runtime/PlayerCore.h b/Engine/Plugins/Media/ElectraPlayer/Source/ElectraPlayerRuntime/Private/Runtime/PlayerCore.h
--- a/Engine/Plugins/Media/ElectraPlayer/Source/ElectraPlayerRuntime/Private/Runtime/PlayerCore.h
+++ b/Engine/Plugins/Media/ElectraPlayer/Source/ElectraPlayerRuntime/Private/Runtime/PlayerCore.h
@@ -11,5 +11,4 @@
 #include "Core/MediaSemaphore.h"
 #include "Core/MediaEventSignal.h"
 #include "Core/MediaInterlocked.h"
-#include "Core/MediaCriticalSection.h"
 #include "Core/MediaMessageQueue.h"
diff --git a/Engine/Plugins/Media/ElectraPlayer/Source/ElectraPlayerRuntime/Private/Runtime/SynchronizedClock.cpp b/Engine/Plugins/Media/ElectraPlayer/Source/ElectraPlayerRuntime/Private/Runtime/SynchronizedClock.cpp
--- a/Engine/Plugins/Media/ElectraPlayer/Source/ElectraPlayerRuntime/Private/Runtime/SynchronizedClock.cpp
+++ b/Engine/Plugins/Media/ElectraPlayer/Source/ElectraPlayerRuntime/Private/Runtime/SynchronizedClock.cpp
@@ -19,9 +19,9 @@ namespace Electra
 		virtual FTimeValue MapToSyncTime(const FTimeValue& InTimeToMap) override;
 
 	private:
-		FMediaCriticalSection	Lock;
-		FTimeValue				BaseUTCTime;
-		int64					BaseSystemTime;
+		FCriticalSection Lock;
+		FTimeValue BaseUTCTime;
+		int64 BaseSystemTime;
 	};
 
 
@@ -81,5 +81,3 @@ namespace Electra
 
 
 } // namespace Electra
-
-
diff --git a/Engine/Plugins/Media/ElectraUtil/Source/ElectraBase/Public/Core/MediaCriticalSection.h b/Engine/Plugins/Media/ElectraUtil/Source/ElectraBase/Public/Core/MediaCriticalSection.h
deleted file mode 100644
