Electra: Removed drop state and presentation time offset from access unit and made sure all earliest and latest timestamps are always set

#rnx

diff --git a/Engine/Plugins/Media/ElectraPlayer/Source/ElectraPlayerRuntime/Private/Runtime/Decoder/SubtitleDecoder.cpp b/Engine/Plugins/Media/ElectraPlayer/Source/ElectraPlayerRuntime/Private/Runtime/Decoder/SubtitleDecoder.cpp
--- a/Engine/Plugins/Media/ElectraPlayer/Source/ElectraPlayerRuntime/Private/Runtime/Decoder/SubtitleDecoder.cpp
+++ b/Engine/Plugins/Media/ElectraPlayer/Source/ElectraPlayerRuntime/Private/Runtime/Decoder/SubtitleDecoder.cpp
@@ -12,7 +12,6 @@ namespace SubtitleOptionKeys
 {
 	static const FName SendEmptySubtitleDuringGaps(TEXT("sendEmptySubtitleDuringGaps"));
 	static const FName SideloadedID(TEXT("SideloadedID"));
-	static const FName PTO(TEXT("PTO"));
 	static const FName Width(TEXT("width"));
 	static const FName Height(TEXT("height"));
 	static const FName OffsetX(TEXT("offset_x"));
@@ -367,7 +366,6 @@ void FSubtitleDecoder::AUdataPushAU(FAccessUnit* AccessUnit)
 			{
 				Addtl.Set(SubtitleOptionKeys::SideloadedID, FVariantValue(AccessUnit->BufferSourceInfo->PeriodAdaptationSetID));
 			}
-			Addtl.Set(SubtitleOptionKeys::PTO, FVariantValue(AccessUnit->PTO));
 			PluginDecoder->AddStreamedSubtitleData(TArray<uint8>((const uint8*)AccessUnit->AUData, (int32)AccessUnit->AUSize), AccessUnit->PTS, AccessUnit->Duration, Addtl);
 		}
 	}
diff --git a/Engine/Plugins/Media/ElectraPlayer/Source/ElectraPlayerRuntime/Private/Runtime/Player/AdaptiveStreamingPlayer.cpp b/Engine/Plugins/Media/ElectraPlayer/Source/ElectraPlayerRuntime/Private/Runtime/Player/AdaptiveStreamingPlayer.cpp
--- a/Engine/Plugins/Media/ElectraPlayer/Source/ElectraPlayerRuntime/Private/Runtime/Player/AdaptiveStreamingPlayer.cpp
+++ b/Engine/Plugins/Media/ElectraPlayer/Source/ElectraPlayerRuntime/Private/Runtime/Player/AdaptiveStreamingPlayer.cpp
@@ -202,7 +202,6 @@ void FAdaptiveStreamingPlayer::GetStreamBufferUtilization(FAccessUnitBufferInfo&
 			{
 				FAccessUnitBufferInfo info;
 				Buffer->GetStats(info);
-				OutInfo.PushedDuration += info.PushedDuration;
 				OutInfo.PlayableDuration += info.PlayableDuration;
 				OutInfo.CurrentMemInUse += info.CurrentMemInUse;
 				OutInfo.NumCurrentAccessUnits += info.NumCurrentAccessUnits;
diff --git a/Engine/Plugins/Media/ElectraPlayer/Source/ElectraPlayerRuntime/Private/Runtime/Player/AdaptiveStreamingPlayerDecoders.cpp b/Engine/Plugins/Media/ElectraPlayer/Source/ElectraPlayerRuntime/Private/Runtime/Player/AdaptiveStreamingPlayerDecoders.cpp
--- a/Engine/Plugins/Media/ElectraPlayer/Source/ElectraPlayerRuntime/Private/Runtime/Player/AdaptiveStreamingPlayerDecoders.cpp
+++ b/Engine/Plugins/Media/ElectraPlayer/Source/ElectraPlayerRuntime/Private/Runtime/Player/AdaptiveStreamingPlayerDecoders.cpp
@@ -815,7 +815,7 @@ bool FAdaptiveStreamingPlayer::IsSeamlessBufferSwitchPossible(EStreamType InStre
 		Buffer->PeekAndAddRef(PeekedAU);
 		if (PeekedAU)
 		{
-			bool bFirstIsRenderedKeyframe = PeekedAU->bIsSyncSample && ((!PeekedAU->EarliestPTS.IsValid() && PeekedAU->DropState == 0) || (PeekedAU->EarliestPTS.IsValid() && PeekedAU->PTS >= PeekedAU->EarliestPTS));
+			bool bFirstIsRenderedKeyframe = PeekedAU->bIsSyncSample && (!PeekedAU->EarliestPTS.IsValid() || (PeekedAU->EarliestPTS.IsValid() && PeekedAU->PTS >= PeekedAU->EarliestPTS));
 			FAccessUnit::Release(PeekedAU);
 			return bFirstIsRenderedKeyframe;
 		}
diff --git a/Engine/Plugins/Media/ElectraPlayer/Source/ElectraPlayerRuntime/Private/Runtime/Player/AdaptiveStreamingPlayerEvents.cpp b/Engine/Plugins/Media/ElectraPlayer/Source/ElectraPlayerRuntime/Private/Runtime/Player/AdaptiveStreamingPlayerEvents.cpp
--- a/Engine/Plugins/Media/ElectraPlayer/Source/ElectraPlayerRuntime/Private/Runtime/Player/AdaptiveStreamingPlayerEvents.cpp
+++ b/Engine/Plugins/Media/ElectraPlayer/Source/ElectraPlayerRuntime/Private/Runtime/Player/AdaptiveStreamingPlayerEvents.cpp
@@ -71,8 +71,8 @@ void FAdaptiveStreamingPlayer::DispatchBufferUtilizationEvent(EStreamType Buffer
 		Metrics::FBufferStats stats;
 		stats.BufferType = BufferType;
 		stats.MaxDurationInSeconds = BufferConfig ? BufferConfig->MaxDuration.GetAsSeconds() : 0.0;
-		stats.DurationInUse 	   = BufferStats->PushedDuration.GetAsSeconds();
-		stats.BytesInUse		   = BufferStats->CurrentMemInUse;
+		stats.DurationInUse = BufferStats->PlayableDuration.GetAsSeconds();
+		stats.BytesInUse = BufferStats->CurrentMemInUse;
 		DispatchEvent(FMetricEvent::ReportBufferUtilization(stats));
 	}
 }
diff --git a/Engine/Plugins/Media/ElectraPlayer/Source/ElectraPlayerRuntime/Private/Runtime/Player/StreamSegmentReaderCommon.cpp b/Engine/Plugins/Media/ElectraPlayer/Source/ElectraPlayerRuntime/Private/Runtime/Player/StreamSegmentReaderCommon.cpp
--- a/Engine/Plugins/Media/ElectraPlayer/Source/ElectraPlayerRuntime/Private/Runtime/Player/StreamSegmentReaderCommon.cpp
+++ b/Engine/Plugins/Media/ElectraPlayer/Source/ElectraPlayerRuntime/Private/Runtime/Player/StreamSegmentReaderCommon.cpp
@@ -1845,7 +1845,6 @@ FStreamSegmentReaderCommon::FStreamHandler::EHandleResult FStreamSegmentReaderCo
 					AccessUnit->bIsDummyData = false;
 					AccessUnit->AUCodecData = CurrentlyActiveTrackData->CSD;
 					AccessUnit->BufferSourceInfo = CurrentlyActiveTrackData->BufferSourceInfo;
-					AccessUnit->DropState = FAccessUnit::EDropState::None;
 
 					FTimeValue Duration(AUDuration, TrackTimescale);
 					AccessUnit->Duration = Duration;
@@ -1855,7 +1854,6 @@ FStreamSegmentReaderCommon::FStreamHandler::EHandleResult FStreamSegmentReaderCo
 					AccessUnit->DTS += TimeOffset + CurrentlyActiveTrackData->TimeMappingOffset;
 					AccessUnit->PTS.SetFromND(AUPTS - CurrentlyActiveTrackData->PTO, TrackTimescale);
 					AccessUnit->PTS += TimeOffset + CurrentlyActiveTrackData->TimeMappingOffset;
-					AccessUnit->PTO.SetFromND(CurrentlyActiveTrackData->PTO, TrackTimescale);
 
 					// Remember the first AU's DTS for this stream type in the segment
 					if (!InRequest->TimestampVars.Local.First[stIdx].IsValid())
@@ -1904,19 +1902,13 @@ FStreamSegmentReaderCommon::FStreamHandler::EHandleResult FStreamSegmentReaderCo
 					AccessUnit->DTS.SetSequenceIndex(InRequest->TimestampSequenceIndex);
 					AccessUnit->PTS.SetSequenceIndex(InRequest->TimestampSequenceIndex);
 
-					if (InRequest->Segment.bFrameAccuracyRequired)
+					AccessUnit->EarliestPTS.SetFromND(CurrentlyActiveTrackData->MediaLocalFirstAUTime - CurrentlyActiveTrackData->PTO, TrackTimescale);
+					AccessUnit->EarliestPTS += TimeOffset;
+					if (InRequest->Segment.bFrameAccuracyRequired && InRequest->FrameAccurateStartTime.IsValid())
 					{
-						if (InRequest->FrameAccurateStartTime.IsValid())
-						{
-							AccessUnit->EarliestPTS = InRequest->FrameAccurateStartTime;
-						}
-						else
-						{
-							AccessUnit->EarliestPTS.SetFromND(CurrentlyActiveTrackData->MediaLocalFirstAUTime - CurrentlyActiveTrackData->PTO, TrackTimescale);
-							AccessUnit->EarliestPTS += TimeOffset;
-						}
-						AccessUnit->EarliestPTS.SetSequenceIndex(InRequest->TimestampSequenceIndex);
+						AccessUnit->EarliestPTS = InRequest->FrameAccurateStartTime;
 					}
+					AccessUnit->EarliestPTS.SetSequenceIndex(InRequest->TimestampSequenceIndex);
 					if (CurrentlyActiveTrackData->MediaLocalLastAUTime != TNumericLimits<int64>::Max())
 					{
 						AccessUnit->LatestPTS.SetFromND(CurrentlyActiveTrackData->MediaLocalLastAUTime - CurrentlyActiveTrackData->PTO, TrackTimescale);
@@ -2224,6 +2216,11 @@ FStreamSegmentReaderCommon::FStreamHandler::EHandleResult FStreamSegmentReaderCo
 	FTimeValue PTO(InRequest->Segment.PTO, InRequest->Segment.Timescale);
 	FTimeValue EarliestPTS(InRequest->Segment.MediaLocalFirstAUTime, InRequest->Segment.Timescale, InRequest->TimestampSequenceIndex);
 	EarliestPTS += TimeOffset - PTO;
+	if (InRequest->Segment.bFrameAccuracyRequired && InRequest->FrameAccurateStartTime.IsValid())
+	{
+		EarliestPTS = InRequest->FrameAccurateStartTime;
+	}
+	EarliestPTS.SetSequenceIndex(InRequest->TimestampSequenceIndex);
 	FTimeValue LastPTS;
 	if (!bDoNotTruncateAtPresentationEnd && InRequest->Segment.MediaLocalLastAUTime != TNumericLimits<int64>::Max())
 	{
@@ -2234,6 +2231,7 @@ FStreamSegmentReaderCommon::FStreamHandler::EHandleResult FStreamSegmentReaderCo
 	{
 		LastPTS.SetToPositiveInfinity();
 	}
+	LastPTS.SetSequenceIndex(InRequest->TimestampSequenceIndex);
 	bool bIsFirstFrame = true;
 	FAccessUnit *AccessUnit = nullptr;
 	auto PrepareAccessUnit = [this, &AccessUnit](int64 NumToRead) -> void*
@@ -2317,10 +2315,7 @@ FStreamSegmentReaderCommon::FStreamHandler::EHandleResult FStreamSegmentReaderCo
 					AccessUnit->DTS += TimeOffset;
 					AccessUnit->DTS.SetSequenceIndex(InRequest->TimestampSequenceIndex);
 					AccessUnit->Duration = Action->GetDuration();
-					if (InRequest->Segment.bFrameAccuracyRequired)
-					{
-						AccessUnit->EarliestPTS = EarliestPTS;
-					}
+					AccessUnit->EarliestPTS = EarliestPTS;
 					AccessUnit->LatestPTS = LastPTS;
 
 					AccessUnit->bIsFirstInSequence = CurrentlyActiveTrackData->bIsFirstInSequence;
@@ -2328,7 +2323,6 @@ FStreamSegmentReaderCommon::FStreamHandler::EHandleResult FStreamSegmentReaderCo
 					AccessUnit->bIsDummyData = false;
 					AccessUnit->AUCodecData = CurrentlyActiveTrackData->CSD;
 					AccessUnit->BufferSourceInfo = CurrentlyActiveTrackData->BufferSourceInfo;
-					AccessUnit->DropState = FAccessUnit::EDropState::None;
 					AccessUnit->SequenceIndex = InRequest->TimestampSequenceIndex;
 
 					// VP9 codec?
@@ -2500,6 +2494,24 @@ FStreamSegmentReaderCommon::FStreamHandler::EHandleResult FStreamSegmentReaderCo
 
 	bool bTimeCheckPassed = InRequest->TimestampVars.Next.bCheck == false;
 
+	FTimeValue EarliestPTS;
+	EarliestPTS.SetFromHNS(InRequest->Segment.MediaLocalFirstAUTime);
+	if (InRequest->Segment.bFrameAccuracyRequired && InRequest->FrameAccurateStartTime.IsValid())
+	{
+		EarliestPTS = InRequest->FrameAccurateStartTime;
+	}
+	EarliestPTS.SetSequenceIndex(InRequest->TimestampSequenceIndex);
+	FTimeValue LatestPTS;
+	if (InRequest->Segment.MediaLocalLastAUTime != TNumericLimits<int64>::Max())
+	{
+		LatestPTS.SetFromHNS(InRequest->Segment.MediaLocalLastAUTime);
+	}
+	else
+	{
+		LatestPTS.SetToPositiveInfinity();
+	}
+	LatestPTS.SetSequenceIndex(InRequest->TimestampSequenceIndex);
+
 	while(!bDone && !HasErrored() && !HasReadBeenAborted())
 	{
 		TSParserState = TSParser->Parse(PlayerSessionService, this);
@@ -2824,37 +2836,17 @@ FStreamSegmentReaderCommon::FStreamHandler::EHandleResult FStreamSegmentReaderCo
 				AccessUnit->DTS += TimeOffset;
 				AccessUnit->DTS.SetSequenceIndex(InRequest->TimestampSequenceIndex);
 				AccessUnit->Duration = CurrentlyActiveTrackData->DefaultDurationFromCSD;
-
-				PTS += AccessUnit->Duration;
-				DTS += AccessUnit->Duration;
-
-				if (InRequest->Segment.bFrameAccuracyRequired)
-				{
-					if (InRequest->FrameAccurateStartTime.IsValid())
-					{
-						AccessUnit->EarliestPTS = InRequest->FrameAccurateStartTime;
-						AccessUnit->EarliestPTS.SetSequenceIndex(InRequest->TimestampSequenceIndex);
-					}
-				}
-
-				if (InRequest->Segment.MediaLocalLastAUTime != TNumericLimits<int64>::Max())
-				{
-					AccessUnit->LatestPTS.SetFromHNS(InRequest->Segment.MediaLocalLastAUTime);
-				}
-				else
-				{
-					AccessUnit->LatestPTS.SetToPositiveInfinity();
-				}
-				AccessUnit->LatestPTS.SetSequenceIndex(InRequest->TimestampSequenceIndex);
-
+				AccessUnit->EarliestPTS = EarliestPTS;
+				AccessUnit->LatestPTS = LatestPTS;
 				AccessUnit->bIsFirstInSequence = CurrentlyActiveTrackData->bIsFirstInSequence;
 				AccessUnit->bIsSyncSample = pp.bIsSyncFrame;
 				AccessUnit->bIsDummyData = false;
 				AccessUnit->AUCodecData = CurrentlyActiveTrackData->CSD;
 				AccessUnit->BufferSourceInfo = CurrentlyActiveTrackData->BufferSourceInfo;
-				AccessUnit->DropState = FAccessUnit::EDropState::None;
 				AccessUnit->SequenceIndex = InRequest->TimestampSequenceIndex;
 
+				PTS += AccessUnit->Duration;
+				DTS += AccessUnit->Duration;
 				CurrentlyActiveTrackData->bIsFirstInSequence = false;
 
 				// Add to the track AU FIFO unless we already reached the last sample of the time range.
@@ -3025,6 +3017,25 @@ FStreamSegmentReaderCommon::FStreamHandler::EHandleResult FStreamSegmentReaderCo
 		check(TimeMappingOffset.IsValid());
 		const int32 stIdx = StreamTypeToArrayIndex(CurrentlyActiveTrackData->StreamType);
 		bool bTimeCheckPassed = InRequest->TimestampVars.Next.bCheck == false;
+
+		FTimeValue EarliestPTS;
+		EarliestPTS.SetFromHNS(InRequest->Segment.MediaLocalFirstAUTime);
+		if (InRequest->Segment.bFrameAccuracyRequired && InRequest->FrameAccurateStartTime.IsValid())
+		{
+			EarliestPTS = InRequest->FrameAccurateStartTime;
+		}
+		EarliestPTS.SetSequenceIndex(InRequest->TimestampSequenceIndex);
+		FTimeValue LatestPTS;
+		if (InRequest->Segment.MediaLocalLastAUTime != TNumericLimits<int64>::Max())
+		{
+			LatestPTS.SetFromHNS(InRequest->Segment.MediaLocalLastAUTime);
+		}
+		else
+		{
+			LatestPTS.SetToPositiveInfinity();
+		}
+		LatestPTS.SetSequenceIndex(InRequest->TimestampSequenceIndex);
+
 		while(!bDone && !HasErrored() && !HasReadBeenAborted())
 		{
 			// We now wait for the next few bytes to arrive so we can get the packet size
@@ -3263,36 +3274,16 @@ FStreamSegmentReaderCommon::FStreamHandler::EHandleResult FStreamSegmentReaderCo
 				AccessUnit->PTS.SetSequenceIndex(InRequest->TimestampSequenceIndex);
 				AccessUnit->DTS = AccessUnit->PTS;
 				AccessUnit->Duration = CurrentlyActiveTrackData->DefaultDurationFromCSD;
-
-				PTS += AccessUnit->Duration;
-
-				if (InRequest->Segment.bFrameAccuracyRequired)
-				{
-					if (InRequest->FrameAccurateStartTime.IsValid())
-					{
-						AccessUnit->EarliestPTS = InRequest->FrameAccurateStartTime;
-						AccessUnit->EarliestPTS.SetSequenceIndex(InRequest->TimestampSequenceIndex);
-					}
-				}
-
-				if (InRequest->Segment.MediaLocalLastAUTime != TNumericLimits<int64>::Max())
-				{
-					AccessUnit->LatestPTS.SetFromHNS(InRequest->Segment.MediaLocalLastAUTime);
-				}
-				else
-				{
-					AccessUnit->LatestPTS.SetToPositiveInfinity();
-				}
-				AccessUnit->LatestPTS.SetSequenceIndex(InRequest->TimestampSequenceIndex);
-
+				AccessUnit->EarliestPTS = EarliestPTS;
+				AccessUnit->LatestPTS = LatestPTS;
 				AccessUnit->bIsFirstInSequence = CurrentlyActiveTrackData->bIsFirstInSequence;
 				AccessUnit->bIsSyncSample = true;
 				AccessUnit->bIsDummyData = false;
 				AccessUnit->AUCodecData = CurrentlyActiveTrackData->CSD;
 				AccessUnit->BufferSourceInfo = CurrentlyActiveTrackData->BufferSourceInfo;
-				AccessUnit->DropState = FAccessUnit::EDropState::None;
 				AccessUnit->SequenceIndex = InRequest->TimestampSequenceIndex;
 
+				PTS += AccessUnit->Duration;
 				CurrentlyActiveTrackData->bIsFirstInSequence = false;
 
 				// Add to the track AU FIFO unless we already reached the last sample of the time range.
@@ -3434,6 +3425,28 @@ FStreamSegmentReaderCommon::FStreamHandler::EHandleResult FStreamSegmentReaderCo
 
 		uint32 TrackTimescale = InRequest->Segment.Timescale;
 
+		FTimeValue TimeOffset = InRequest->PeriodStart + InRequest->AST + InRequest->AdditionalAdjustmentTime;
+		const bool bDoNotTruncateAtPresentationEnd = PlayerSessionService->GetOptionValue(OptionKeyDoNotTruncateAtPresentationEnd).SafeGetBool(false);
+		FTimeValue PTO(InRequest->Segment.PTO, InRequest->Segment.Timescale);
+		FTimeValue EarliestPTS(InRequest->Segment.MediaLocalFirstAUTime, InRequest->Segment.Timescale, InRequest->TimestampSequenceIndex);
+		EarliestPTS += TimeOffset - PTO;
+		if (InRequest->Segment.bFrameAccuracyRequired && InRequest->FrameAccurateStartTime.IsValid())
+		{
+			EarliestPTS = InRequest->FrameAccurateStartTime;
+		}
+		EarliestPTS.SetSequenceIndex(InRequest->TimestampSequenceIndex);
+		FTimeValue LatestPTS;
+		if (!bDoNotTruncateAtPresentationEnd && InRequest->Segment.MediaLocalLastAUTime != TNumericLimits<int64>::Max())
+		{
+			LatestPTS.SetFromND(InRequest->Segment.MediaLocalLastAUTime, InRequest->Segment.Timescale, InRequest->TimestampSequenceIndex);
+			LatestPTS += TimeOffset - PTO;
+		}
+		else
+		{
+			LatestPTS.SetToPositiveInfinity();
+		}
+		LatestPTS.SetSequenceIndex(InRequest->TimestampSequenceIndex);
+
 		// Create an access unit
 		FAccessUnit *AccessUnit = FAccessUnit::Create(Parameters.MemoryProvider);
 		AccessUnit->ESType = InRequest->GetType();
@@ -3447,15 +3460,14 @@ FStreamSegmentReaderCommon::FStreamHandler::EHandleResult FStreamSegmentReaderCo
 		AccessUnit->bIsSideloaded = true;
 		AccessUnit->BufferSourceInfo = CurrentlyActiveTrackData->BufferSourceInfo;
 		AccessUnit->SequenceIndex = InRequest->TimestampSequenceIndex;
-		AccessUnit->DropState = FAccessUnit::EDropState::None;
 
 		// Sideloaded files coincide with the period start
-		FTimeValue TimeOffset = InRequest->PeriodStart + InRequest->AST + InRequest->AdditionalAdjustmentTime;
 		AccessUnit->DTS = TimeOffset;
 		AccessUnit->PTS = TimeOffset;
-		AccessUnit->PTO.SetFromND(InRequest->Segment.PTO, TrackTimescale);
 		AccessUnit->DTS.SetSequenceIndex(InRequest->TimestampSequenceIndex);
 		AccessUnit->PTS.SetSequenceIndex(InRequest->TimestampSequenceIndex);
+		AccessUnit->EarliestPTS = EarliestPTS;
+		AccessUnit->LatestPTS = LatestPTS;
 
 		FTimeValue Duration(InRequest->Segment.Duration, TrackTimescale);
 		AccessUnit->Duration = Duration;
@@ -3521,19 +3533,16 @@ void FStreamSegmentReaderCommon::FStreamHandler::InsertFillerData(TSharedPtrTS<F
 	}
 	FTimeValue PTO(InRequest->Segment.PTO, InRequest->Segment.Timescale, 0);
 	FTimeValue Earliest;
+	Earliest.SetFromND(InRequest->Segment.MediaLocalFirstAUTime - InRequest->Segment.PTO, InRequest->Segment.Timescale);
+	Earliest += TimeOffset;
 	if (InRequest->Segment.bFrameAccuracyRequired)
 	{
 		if (InRequest->FrameAccurateStartTime.IsValid())
 		{
 			Earliest = InRequest->FrameAccurateStartTime;
 		}
-		else
-		{
-			Earliest.SetFromND(InRequest->Segment.MediaLocalFirstAUTime - InRequest->Segment.PTO, InRequest->Segment.Timescale);
-			Earliest += TimeOffset;
-		}
-		Earliest.SetSequenceIndex(InRequest->TimestampSequenceIndex);
 	}
+	Earliest.SetSequenceIndex(InRequest->TimestampSequenceIndex);
 
 	bool bDoNotTruncateAtPresentationEnd = PlayerSessionService->GetOptionValue(OptionKeyDoNotTruncateAtPresentationEnd).SafeGetBool(false);
 	FTimeValue Latest;
@@ -3576,20 +3585,12 @@ void FStreamSegmentReaderCommon::FStreamHandler::InsertFillerData(TSharedPtrTS<F
 		AccessUnit->SequenceIndex = InRequest->TimestampSequenceIndex;
 		AccessUnit->DTS = NextExpectedDTS;
 		AccessUnit->PTS = NextExpectedDTS;
-		AccessUnit->PTO = PTO;
 		AccessUnit->EarliestPTS = Earliest;
 		AccessUnit->LatestPTS = Latest;
 
-		// Calculate the drop on the fragment local NextExpectedDTS/PTS.
-		AccessUnit->DropState = FAccessUnit::EDropState::None;
-		if (NextExpectedDTS + AccessUnit->Duration < Earliest)
-		{
-			AccessUnit->DropState |= FAccessUnit::EDropState::TooEarly;
-		}
 		bool bIsLast = false;
 		if (NextExpectedDTS > Latest)
 		{
-			AccessUnit->DropState |= FAccessUnit::EDropState::TooLate;
 			bIsLast = AccessUnit->bIsLastInPeriod = true;
 		}
 
@@ -3783,7 +3784,6 @@ FStreamSegmentReaderCommon::FStreamHandler::EEmitResult FStreamSegmentReaderComm
 					if (i < CurrentlyActiveTrackData->SortedAccessUnitFIFO.Num()-1)
 					{
 						CurrentlyActiveTrackData->AccessUnitFIFO[0].AU->Duration = CurrentlyActiveTrackData->SortedAccessUnitFIFO[i+1].PTS - CurrentlyActiveTrackData->SortedAccessUnitFIFO[i].PTS;
-						UpdateAUDropState(CurrentlyActiveTrackData->AccessUnitFIFO[0].AU);
 						if (!CurrentlyActiveTrackData->AverageDuration.IsValid() || CurrentlyActiveTrackData->AverageDuration.IsZero())
 						{
 							CurrentlyActiveTrackData->AverageDuration = CurrentlyActiveTrackData->AccessUnitFIFO[0].AU->Duration;
@@ -3806,10 +3806,6 @@ FStreamSegmentReaderCommon::FStreamHandler::EEmitResult FStreamSegmentReaderComm
 				}
 			}
 		}
-		else
-		{
-			UpdateAUDropState(CurrentlyActiveTrackData->AccessUnitFIFO[0].AU);
-		}
 
 		FAccessUnit* pNext = CurrentlyActiveTrackData->AccessUnitFIFO[0].AU;
 		// Check if this is the last access unit in the requested time range.
@@ -3889,19 +3885,6 @@ FStreamSegmentReaderCommon::FStreamHandler::EEmitResult FStreamSegmentReaderComm
 	return Result;
 }
 
-void FStreamSegmentReaderCommon::FStreamHandler::UpdateAUDropState(FAccessUnit* InAU)
-{
-	if (InAU->PTS + InAU->Duration < InAU->EarliestPTS)
-	{
-		InAU->DropState |= FAccessUnit::EDropState::TooEarly;
-	}
-	if (InAU->PTS >= InAU->LatestPTS)
-	{
-		InAU->DropState |= FAccessUnit::EDropState::TooLate;
-	}
-}
-
-
 
 
 /**
diff --git a/Engine/Plugins/Media/ElectraPlayer/Source/ElectraPlayerRuntime/Private/Runtime/Player/StreamSegmentReaderCommon.h b/Engine/Plugins/Media/ElectraPlayer/Source/ElectraPlayerRuntime/Private/Runtime/Player/StreamSegmentReaderCommon.h
--- a/Engine/Plugins/Media/ElectraPlayer/Source/ElectraPlayerRuntime/Private/Runtime/Player/StreamSegmentReaderCommon.h
+++ b/Engine/Plugins/Media/ElectraPlayer/Source/ElectraPlayerRuntime/Private/Runtime/Player/StreamSegmentReaderCommon.h
@@ -568,7 +568,6 @@ private:
 
 		void SelectPrimaryTrackData(const TSharedPtrTS<FStreamSegmentRequestCommon>& InRequest);
 		EEmitResult EmitSamples(EEmitType InEmitType);
-		void UpdateAUDropState(FAccessUnit* InAU);
 
 		// Methods from IGenericDataReader
 		int64 ReadData(void* IntoBuffer, int64 NumBytesToRead, int64 InFromOffset) override;
diff --git a/Engine/Plugins/Media/ElectraPlayer/Source/ElectraPlayerRuntime/Private/Runtime/Player/mkv/ManifestMKV.cpp b/Engine/Plugins/Media/ElectraPlayer/Source/ElectraPlayerRuntime/Private/Runtime/Player/mkv/ManifestMKV.cpp
--- a/Engine/Plugins/Media/ElectraPlayer/Source/ElectraPlayerRuntime/Private/Runtime/Player/mkv/ManifestMKV.cpp
+++ b/Engine/Plugins/Media/ElectraPlayer/Source/ElectraPlayerRuntime/Private/Runtime/Player/mkv/ManifestMKV.cpp
@@ -760,16 +760,13 @@ IManifest::FResult FManifestMKVInternal::FTimelineAssetMKV::GetStartingSegment(T
 			FTimeValue StreamDuration = MKVParser->GetDuration();
 			if (StreamDuration.IsValid() && StreamDuration <= StartPosition.Time && (SearchType == ESearchType::After || SearchType == ESearchType::StrictlyAfter || SearchType == ESearchType::Closest))
 			{
-				firstTimestamp.SetFromTimeFraction(StreamDuration);
+				firstTimestamp = StreamDuration;
 				Error = UEMEDIA_ERROR_END_OF_STREAM;
 			}
 			else
 			{
 				Error = TrackIt->StartAtTime(StartPosition.Time, SearchMode);
-				if (Error == UEMEDIA_ERROR_OK)
-				{
-					firstTimestamp = TrackIt->GetTimestamp();
-				}
+				firstTimestamp = Error == UEMEDIA_ERROR_OK ? TrackIt->GetTimestamp() : StreamDuration;
 			}
 			if (Error == UEMEDIA_ERROR_OK || Error == UEMEDIA_ERROR_END_OF_STREAM)
 			{
@@ -783,9 +780,9 @@ IManifest::FResult FManifestMKVInternal::FTimelineAssetMKV::GetStartingSegment(T
 				req->FirstPTS = firstTimestamp;
 				req->PrimaryStreamType = ReprType;
 				req->Bitrate = Repr->GetBitrate();
-				req->bIsContinuationSegment = false;
 				// Set the PTS range for which to present samples.
 				req->EarliestPTS = bFrameAccurateSearch ? StartPosition.Time : firstTimestamp;
+				check(req->EarliestPTS.IsValid());
 				req->LastPTS = PlayRangeEnd;
 				req->TimestampSequenceIndex = InSequenceState.GetSequenceIndex();
 				// Set info from the track iterator if found
@@ -859,7 +856,6 @@ IManifest::FResult FManifestMKVInternal::FTimelineAssetMKV::GetNextSegment(TShar
 					req->SegmentDuration = TrackIt->GetClusterDuration();
 					req->Bitrate = Request->Bitrate;
 					req->PrimaryStreamType = Request->PrimaryStreamType;
-					req->bIsContinuationSegment = true;
 					req->bIsFirstSegment = false;
 					req->bIsLastSegment = TrackIt->IsLastCluster() || TrackIt->GetTimestamp() + TrackIt->GetClusterDuration() >= Request->LastPTS;
 					req->EarliestPTS = Request->EarliestPTS;
@@ -897,7 +893,6 @@ IManifest::FResult FManifestMKVInternal::FTimelineAssetMKV::GetRetrySegment(TSha
 		req->NextCueUniqueID = Request->NextCueUniqueID;
 		req->PlaybackSequenceID = Request->PlaybackSequenceID;
 		req->Bitrate = Request->Bitrate;
-		req->bIsContinuationSegment = Request->bIsContinuationSegment;
 		req->bIsFirstSegment = Request->bIsFirstSegment;
 		req->bIsLastSegment = Request->bIsLastSegment;
 		req->bIsEOSSegment = Request->bIsEOSSegment;
diff --git a/Engine/Plugins/Media/ElectraPlayer/Source/ElectraPlayerRuntime/Private/Runtime/Player/mkv/StreamReaderMKV.cpp b/Engine/Plugins/Media/ElectraPlayer/Source/ElectraPlayerRuntime/Private/Runtime/Player/mkv/StreamReaderMKV.cpp
--- a/Engine/Plugins/Media/ElectraPlayer/Source/ElectraPlayerRuntime/Private/Runtime/Player/mkv/StreamReaderMKV.cpp
+++ b/Engine/Plugins/Media/ElectraPlayer/Source/ElectraPlayerRuntime/Private/Runtime/Player/mkv/StreamReaderMKV.cpp
@@ -491,11 +491,6 @@ void FStreamReaderMKV::HandleRequest()
 					AccessUnit->LatestPTS = Request->LastPTS;
 
 					AccessUnit->AUCodecData = SelectedTrack.CSD;
-					AccessUnit->DropState = FAccessUnit::EDropState::None;
-					if (!SelectedTrack.bNeedToRecalculateDurations)
-					{
-						UpdateAUDropState(AccessUnit, Request);
-					}
 
 					// Set the sequence index member and update all timestamps with it as well.
 					AccessUnit->SequenceIndex = Request->TimestampSequenceIndex;
@@ -711,23 +706,6 @@ void FStreamReaderMKV::HandleRequest()
 	Parameters.EventListener->OnFragmentClose(Request);
 }
 
-void FStreamReaderMKV::UpdateAUDropState(FAccessUnit* InAU, const TSharedPtrTS<FStreamSegmentRequestMKV>& InRequest)
-{
-	// If this is a continuation then we must not tag samples as being too early.
-	if (!InRequest->bIsContinuationSegment)
-	{
-		if (InAU->PTS + InAU->Duration < InRequest->FirstPTS)
-		{
-			InAU->DropState |= FAccessUnit::EDropState::TooEarly;
-		}
-	}
-	if (InAU->PTS >= InAU->LatestPTS)
-	{
-		InAU->DropState |= FAccessUnit::EDropState::TooLate;
-	}
-}
-
-
 FStreamReaderMKV::EEmitResult FStreamReaderMKV::EmitSamples(EEmitType InEmitType, const TSharedPtrTS<FStreamSegmentRequestMKV>& InRequest)
 {
 	EEmitResult Result = EEmitResult::SentNothing;
@@ -765,7 +743,6 @@ FStreamReaderMKV::EEmitResult FStreamReaderMKV::EmitSamples(EEmitType InEmitType
 							if (i < td.SortedAccessUnitFIFO.Num()-1)
 							{
 								td.AccessUnitFIFO[0].AU->Duration = td.SortedAccessUnitFIFO[i+1].PTS - td.SortedAccessUnitFIFO[i].PTS;
-								UpdateAUDropState(td.AccessUnitFIFO[0].AU, InRequest);
 							}
 							else if (InEmitType == EEmitType::KnownDurationOnly)
 							{
diff --git a/Engine/Plugins/Media/ElectraPlayer/Source/ElectraPlayerRuntime/Private/Runtime/Player/mkv/StreamReaderMKV.h b/Engine/Plugins/Media/ElectraPlayer/Source/ElectraPlayerRuntime/Private/Runtime/Player/mkv/StreamReaderMKV.h
--- a/Engine/Plugins/Media/ElectraPlayer/Source/ElectraPlayerRuntime/Private/Runtime/Player/mkv/StreamReaderMKV.h
+++ b/Engine/Plugins/Media/ElectraPlayer/Source/ElectraPlayerRuntime/Private/Runtime/Player/mkv/StreamReaderMKV.h
@@ -86,8 +86,6 @@ public:
 	int32 NumOverallRetries = 0;
 
 	int32 Bitrate = 0;
-	// true if this segment continues where the previous left off and no sync samples should be expected.
-	bool bIsContinuationSegment = false;
 	// true if this segment is the first to start with or the first after a seek.
 	bool bIsFirstSegment = false;
 	// true if this segment is the last.
@@ -141,7 +139,6 @@ private:
 		AllReachedEOS,
 	};
 	EEmitResult EmitSamples(EEmitType InEmitType, const TSharedPtrTS<FStreamSegmentRequestMKV>& InRequest);
-	void UpdateAUDropState(FAccessUnit* InAU, const TSharedPtrTS<FStreamSegmentRequestMKV>& InRequest);
 
 	bool HasBeenAborted() const;
 	bool HasErrored() const;
diff --git a/Engine/Plugins/Media/ElectraPlayer/Source/ElectraPlayerRuntime/Private/Runtime/Player/mp4/ManifestMP4.cpp b/Engine/Plugins/Media/ElectraPlayer/Source/ElectraPlayerRuntime/Private/Runtime/Player/mp4/ManifestMP4.cpp
--- a/Engine/Plugins/Media/ElectraPlayer/Source/ElectraPlayerRuntime/Private/Runtime/Player/mp4/ManifestMP4.cpp
+++ b/Engine/Plugins/Media/ElectraPlayer/Source/ElectraPlayerRuntime/Private/Runtime/Player/mp4/ManifestMP4.cpp
@@ -996,9 +996,9 @@ IManifest::FResult FManifestMP4Internal::FTimelineAssetMP4::GetStartingSegment(T
 					// Get the track duration. If the start time is outside the track this could be a deliberate seek past the end of the video track
 					// into a longer audio track. We do not want to snap to the last video IDR frame then.
 					FTimeFraction TrackDuration = Track->GetDuration();
+					firstTimestamp.SetFromTimeFraction(TrackDuration);
 					if (TrackDuration.IsValid() && FTimeValue().SetFromTimeFraction(TrackDuration) <= StartPosition.Time && (SearchType == ESearchType::After || SearchType == ESearchType::StrictlyAfter || SearchType == ESearchType::Closest))
 					{
-						firstTimestamp.SetFromTimeFraction(TrackDuration);
 						err = UEMEDIA_ERROR_END_OF_STREAM;
 					}
 					else
@@ -1032,21 +1032,21 @@ IManifest::FResult FManifestMP4Internal::FTimelineAssetMP4::GetStartingSegment(T
 					// Time found. Set up the fragment request.
 					TSharedPtrTS<FStreamSegmentRequestMP4> req(new FStreamSegmentRequestMP4);
 					OutSegment = req;
-
-					req->MediaAsset 			= SharedThis(this);
-					req->PrimaryTrackIterator   = TrackIt;
-					req->FirstPTS   			= firstTimestamp;
-					req->PrimaryStreamType  	= EStreamType::Video;
-					req->FileStartOffset		= firstByteOffset;
-					req->FileEndOffset  		= -1;
-					req->Bitrate				= Repr->GetBitrate();
-					req->bStartingOnMOOF		= false;
-					req->bIsContinuationSegment = false;
-					req->bIsFirstSegment		= true;
-					req->bIsLastSegment			= true;
-					req->SegmentDuration		= GetDuration() - req->FirstPTS;
+					req->MediaAsset = SharedThis(this);
+					req->PrimaryTrackIterator = TrackIt;
+					check(firstTimestamp.IsValid());
+					req->FirstPTS = firstTimestamp;
+					req->PrimaryStreamType = EStreamType::Video;
+					req->FileStartOffset = firstByteOffset;
+					req->FileEndOffset = -1;
+					req->Bitrate = Repr->GetBitrate();
+					req->bStartingOnMOOF = false;
+					req->bIsFirstSegment = true;
+					req->bIsLastSegment = true;
+					req->SegmentDuration = GetDuration() - req->FirstPTS;
 					// Set the PTS range for which to present samples.
 					req->EarliestPTS = bFrameAccurateSearch ? StartPosition.Time : firstTimestamp;
+					check(PlayRangeEnd.IsValid());
 					req->LastPTS = PlayRangeEnd;
 					req->TimestampSequenceIndex = InSequenceState.GetSequenceIndex();
 
@@ -1079,15 +1079,14 @@ IManifest::FResult FManifestMP4Internal::FTimelineAssetMP4::GetStartingSegment(T
 					{
 						TSharedPtrTS<FStreamSegmentRequestMP4> req(new FStreamSegmentRequestMP4);
 						OutSegment = req;
-						req->MediaAsset 			= SharedThis(this);
-						req->FirstPTS   			= firstTimestamp;
-						req->PrimaryStreamType  	= EStreamType::Video;
-						req->Bitrate				= Repr->GetBitrate();
-						req->bStartingOnMOOF		= false;
-						req->bIsContinuationSegment = false;
-						req->bIsFirstSegment		= false;
-						req->bIsLastSegment			= true;
-						req->bAllTracksAtEOS		= true;
+						req->MediaAsset = SharedThis(this);
+						req->FirstPTS = firstTimestamp;
+						req->PrimaryStreamType = EStreamType::Video;
+						req->Bitrate = Repr->GetBitrate();
+						req->bStartingOnMOOF = false;
+						req->bIsFirstSegment = false;
+						req->bIsLastSegment = true;
+						req->bAllTracksAtEOS = true;
 						// Set the PTS range for which to present samples.
 						req->EarliestPTS = bFrameAccurateSearch ? StartPosition.Time : firstTimestamp;
 						req->LastPTS = PlayRangeEnd;
@@ -1140,9 +1139,9 @@ IManifest::FResult FManifestMP4Internal::FTimelineAssetMP4::GetStartingSegment(T
 					// Get the track duration. If the start time is outside the track this could be a deliberate seek past the end of the audio track
 					// for which we return EOS.
 					FTimeFraction TrackDuration = Track->GetDuration();
+					firstTimestamp.SetFromTimeFraction(TrackDuration);
 					if (TrackDuration.IsValid() && FTimeValue().SetFromTimeFraction(TrackDuration) <= StartPosition.Time && (SearchType == ESearchType::After || SearchType == ESearchType::StrictlyAfter || SearchType == ESearchType::Closest))
 					{
-						firstTimestamp.SetFromTimeFraction(TrackDuration);
 						err = UEMEDIA_ERROR_END_OF_STREAM;
 					}
 					else
@@ -1176,21 +1175,21 @@ IManifest::FResult FManifestMP4Internal::FTimelineAssetMP4::GetStartingSegment(T
 					// Time found. Set up the fragment request.
 					TSharedPtrTS<FStreamSegmentRequestMP4> req(new FStreamSegmentRequestMP4);
 					OutSegment = req;
-
-					req->MediaAsset 			= SharedThis(this);
-					req->PrimaryTrackIterator   = TrackIt;
-					req->FirstPTS   			= firstTimestamp;
-					req->PrimaryStreamType  	= EStreamType::Audio;
-					req->FileStartOffset		= firstByteOffset;
-					req->FileEndOffset  		= -1;
-					req->Bitrate				= Repr->GetBitrate();
-					req->bStartingOnMOOF		= false;
-					req->bIsContinuationSegment = false;
-					req->bIsFirstSegment		= true;
-					req->bIsLastSegment			= true;
-					req->SegmentDuration		= GetDuration() - req->FirstPTS;
+					req->MediaAsset = SharedThis(this);
+					req->PrimaryTrackIterator = TrackIt;
+					check(firstTimestamp.IsValid());
+					req->FirstPTS = firstTimestamp;
+					req->PrimaryStreamType = EStreamType::Audio;
+					req->FileStartOffset = firstByteOffset;
+					req->FileEndOffset = -1;
+					req->Bitrate = Repr->GetBitrate();
+					req->bStartingOnMOOF = false;
+					req->bIsFirstSegment = true;
+					req->bIsLastSegment = true;
+					req->SegmentDuration = GetDuration() - req->FirstPTS;
 					// Set the PTS range for which to present samples.
 					req->EarliestPTS = bFrameAccurateSearch ? StartPosition.Time : firstTimestamp;
+					check(PlayRangeEnd.IsValid());
 					req->LastPTS = PlayRangeEnd;
 					req->TimestampSequenceIndex = InSequenceState.GetSequenceIndex();
 
@@ -1222,15 +1221,14 @@ IManifest::FResult FManifestMP4Internal::FTimelineAssetMP4::GetStartingSegment(T
 					// Regardless of there being a video stream or not we return an EOS request.
 					TSharedPtrTS<FStreamSegmentRequestMP4> req(new FStreamSegmentRequestMP4);
 					OutSegment = req;
-					req->MediaAsset 			= SharedThis(this);
-					req->FirstPTS   			= firstTimestamp;
-					req->PrimaryStreamType  	= EStreamType::Audio;
-					req->Bitrate				= Repr->GetBitrate();
-					req->bStartingOnMOOF		= false;
-					req->bIsContinuationSegment = false;
-					req->bIsFirstSegment		= false;
-					req->bIsLastSegment			= true;
-					req->bAllTracksAtEOS		= true;
+					req->MediaAsset = SharedThis(this);
+					req->FirstPTS = firstTimestamp;
+					req->PrimaryStreamType = EStreamType::Audio;
+					req->Bitrate = Repr->GetBitrate();
+					req->bStartingOnMOOF = false;
+					req->bIsFirstSegment = false;
+					req->bIsLastSegment = true;
+					req->bAllTracksAtEOS = true;
 					// Set the PTS range for which to present samples.
 					req->EarliestPTS = bFrameAccurateSearch ? StartPosition.Time : firstTimestamp;
 					req->LastPTS = PlayRangeEnd;
@@ -1281,10 +1279,9 @@ IManifest::FResult FManifestMP4Internal::FTimelineAssetMP4::GetNextSegment(TShar
 			if (res.GetType() == IManifest::FResult::EType::Found)
 			{
 				FStreamSegmentRequestMP4* NextRequest = static_cast<FStreamSegmentRequestMP4*>(OutSegment.Get());
-				NextRequest->bIsContinuationSegment = true;
-				NextRequest->bIsFirstSegment		= false;
-				NextRequest->EarliestPTS			= Request->EarliestPTS;
-				NextRequest->LastPTS				= Request->LastPTS;
+				NextRequest->bIsFirstSegment = false;
+				NextRequest->EarliestPTS = Request->EarliestPTS;
+				NextRequest->LastPTS = Request->LastPTS;
 				if (NextRequest->SegmentInternalSize > 0)
 				{
 					return res;
@@ -1308,10 +1305,9 @@ IManifest::FResult FManifestMP4Internal::FTimelineAssetMP4::GetRetrySegment(TSha
 		if (res.GetType() == IManifest::FResult::EType::Found)
 		{
 			FStreamSegmentRequestMP4* RetryRequest = static_cast<FStreamSegmentRequestMP4*>(OutSegment.Get());
-			RetryRequest->bIsContinuationSegment = true;
-			RetryRequest->NumOverallRetries 	 = Request->NumOverallRetries + 1;
-			RetryRequest->EarliestPTS			 = Request->EarliestPTS;
-			RetryRequest->LastPTS				 = Request->LastPTS;
+			RetryRequest->NumOverallRetries = Request->NumOverallRetries + 1;
+			RetryRequest->EarliestPTS = Request->EarliestPTS;
+			RetryRequest->LastPTS = Request->LastPTS;
 			return res;
 		}
 	}
diff --git a/Engine/Plugins/Media/ElectraPlayer/Source/ElectraPlayerRuntime/Private/Runtime/Player/mp4/StreamReaderMP4.cpp b/Engine/Plugins/Media/ElectraPlayer/Source/ElectraPlayerRuntime/Private/Runtime/Player/mp4/StreamReaderMP4.cpp
--- a/Engine/Plugins/Media/ElectraPlayer/Source/ElectraPlayerRuntime/Private/Runtime/Player/mp4/StreamReaderMP4.cpp
+++ b/Engine/Plugins/Media/ElectraPlayer/Source/ElectraPlayerRuntime/Private/Runtime/Player/mp4/StreamReaderMP4.cpp
@@ -34,7 +34,6 @@ FStreamSegmentRequestMP4::FStreamSegmentRequestMP4()
 	Bitrate 			   = 0;
 	PlaybackSequenceID     = ~0U;
 	bStartingOnMOOF 	   = false;
-	bIsContinuationSegment = false;
 	bIsFirstSegment		   = false;
 	bIsLastSegment		   = false;
 	bAllTracksAtEOS 	   = false;
@@ -536,20 +535,6 @@ void FStreamReaderMP4::HandleRequest()
 
 					AccessUnit->AUSize = (uint32) SampleSize;
 					AccessUnit->AUCodecData = SelectedTrack.CSD;
-					AccessUnit->DropState = FAccessUnit::EDropState::None;
-					// If this is a continuation then we must not tag samples as being too early.
-					if (!Request->bIsContinuationSegment)
-					{
-						if (AccessUnit->PTS + AccessUnit->Duration < Request->FirstPTS)
-						{
-							AccessUnit->DropState |= FAccessUnit::EDropState::TooEarly;
-						}
-					}
-					if (AccessUnit->PTS >= AccessUnit->LatestPTS)
-					{
-						AccessUnit->DropState |= FAccessUnit::EDropState::TooLate;
-					}
-
 					// Set the sequence index member and update all timestamps with it as well.
 					AccessUnit->SequenceIndex = Request->TimestampSequenceIndex;
 					AccessUnit->DTS.SetSequenceIndex(Request->TimestampSequenceIndex);
diff --git a/Engine/Plugins/Media/ElectraPlayer/Source/ElectraPlayerRuntime/Private/Runtime/Player/mp4/StreamReaderMP4.h b/Engine/Plugins/Media/ElectraPlayer/Source/ElectraPlayerRuntime/Private/Runtime/Player/mp4/StreamReaderMP4.h
--- a/Engine/Plugins/Media/ElectraPlayer/Source/ElectraPlayerRuntime/Private/Runtime/Player/mp4/StreamReaderMP4.h
+++ b/Engine/Plugins/Media/ElectraPlayer/Source/ElectraPlayerRuntime/Private/Runtime/Player/mp4/StreamReaderMP4.h
@@ -70,7 +70,6 @@ public:
 	uint32												PlaybackSequenceID;
 	int32												Bitrate;
 	bool												bStartingOnMOOF;			//!< when using the 'sidx' a segment is expected to start on a 'moof', otherwise inside an 'mdat'.
-	bool												bIsContinuationSegment;		//!< true if this segment continues where the previous left off and no sync samples should be expected.
 	bool												bIsFirstSegment;			//!< true if this segment is the first to start with or the first after a seek.
 	bool												bIsLastSegment;				//!< true if this segment is the last.
 
diff --git a/Engine/Plugins/Media/ElectraPlayer/Source/ElectraPlayerRuntime/Private/Runtime/Player/mpegaudio/ManifestMPEGAudio.cpp b/Engine/Plugins/Media/ElectraPlayer/Source/ElectraPlayerRuntime/Private/Runtime/Player/mpegaudio/ManifestMPEGAudio.cpp
--- a/Engine/Plugins/Media/ElectraPlayer/Source/ElectraPlayerRuntime/Private/Runtime/Player/mpegaudio/ManifestMPEGAudio.cpp
+++ b/Engine/Plugins/Media/ElectraPlayer/Source/ElectraPlayerRuntime/Private/Runtime/Player/mpegaudio/ManifestMPEGAudio.cpp
@@ -929,6 +929,7 @@ IManifest::FResult FManifestMPEGAudioInternal::FTimelineAssetMPEGAudio::GetStart
 				TSharedPtrTS<FStreamSegmentRequestMPEGAudio> req(new FStreamSegmentRequestMPEGAudio);
 				OutSegment = req;
 				req->MediaAsset = SharedThis(this);
+				check(StartTime.IsValid() && PlayRangeEnd.IsValid());
 				req->FirstPTS = StartTime;
 				req->EarliestPTS = StartTime;
 				req->LastPTS = PlayRangeEnd;
diff --git a/Engine/Plugins/Media/ElectraPlayer/Source/ElectraPlayerRuntime/Private/Runtime/Player/mpegaudio/StreamReaderMPEGAudio.cpp b/Engine/Plugins/Media/ElectraPlayer/Source/ElectraPlayerRuntime/Private/Runtime/Player/mpegaudio/StreamReaderMPEGAudio.cpp
--- a/Engine/Plugins/Media/ElectraPlayer/Source/ElectraPlayerRuntime/Private/Runtime/Player/mpegaudio/StreamReaderMPEGAudio.cpp
+++ b/Engine/Plugins/Media/ElectraPlayer/Source/ElectraPlayerRuntime/Private/Runtime/Player/mpegaudio/StreamReaderMPEGAudio.cpp
@@ -427,7 +427,7 @@ void FStreamReaderMPEGAudio::HandleRequest()
 		{
 			LiveRequest->Handle->SetHeader(TEXT("Icy-Metadata"), TEXT("1"));
 		}
-		LiveRequest->MaxDataBytes = (int32) (PlayerSessionServices->GetOptionValue(MPEGAudio::OptionKeyMPEGAudioMaxPreloadBufferDuration).SafeGetTimeValue(FTimeValue().SetFromMilliseconds(1000 * 90)).GetAsSeconds() * Request->GetBitrate() / 8);
+		LiveRequest->MaxDataBytes = Utils::Min((int32) (PlayerSessionServices->GetOptionValue(MPEGAudio::OptionKeyMPEGAudioMaxPreloadBufferDuration).SafeGetTimeValue(FTimeValue().SetFromMilliseconds(1000 * 180)).GetAsSeconds() * Request->GetBitrate() / 8), (int32)(8<<20));
 		LiveRequest->Handle->SetActivityTimeout(PlayerSessionServices->GetOptionValue(MPEGAudio::OptionKeyMPEGAudioLoadNoDataTimeout).SafeGetTimeValue(FTimeValue().SetFromMilliseconds(1000 * 5)).GetAsSeconds());
 		LiveRequest->Handle->ProcessRequest();
 	}
@@ -670,12 +670,6 @@ void FStreamReaderMPEGAudio::HandleRequest()
 
 					AccessUnit->AUSize = (uint32) NextAUFrameSize;
 					AccessUnit->AUCodecData = AUCodecData;
-					AccessUnit->DropState = FAccessUnit::EDropState::None;
-					if (AccessUnit->PTS >= AccessUnit->LatestPTS)
-					{
-						AccessUnit->DropState |= FAccessUnit::EDropState::TooLate;
-					}
-
 					// Set the sequence index member and update all timestamps with it as well.
 					AccessUnit->SequenceIndex = Request->TimestampSequenceIndex;
 					AccessUnit->DTS.SetSequenceIndex(Request->TimestampSequenceIndex);
diff --git a/Engine/Plugins/Media/ElectraPlayer/Source/ElectraPlayerRuntime/Private/Runtime/StreamAccessUnitBuffer.cpp b/Engine/Plugins/Media/ElectraPlayer/Source/ElectraPlayerRuntime/Private/Runtime/StreamAccessUnitBuffer.cpp
--- a/Engine/Plugins/Media/ElectraPlayer/Source/ElectraPlayerRuntime/Private/Runtime/StreamAccessUnitBuffer.cpp
+++ b/Engine/Plugins/Media/ElectraPlayer/Source/ElectraPlayerRuntime/Private/Runtime/StreamAccessUnitBuffer.cpp
@@ -91,7 +91,10 @@ namespace Electra
 			bEndOfTrack = false;
 			PlayableDurationPushedSinceEOT.SetToZero();
 		}
-		if (AU->DropState == FAccessUnit::EDropState::None)
+		check(AU->EarliestPTS.IsValid());
+		check(AU->LatestPTS.IsValid());
+		const bool bIsPlayableAU = AU->PTS >= AU->EarliestPTS && (AU->PTS + AU->Duration) < AU->LatestPTS;
+		if (bIsPlayableAU)
 		{
 			PlayableDurationPushedSinceEOT += AU->Duration;
 		}
diff --git a/Engine/Plugins/Media/ElectraPlayer/Source/ElectraPlayerRuntime/Private/Runtime/StreamAccessUnitBuffer.h b/Engine/Plugins/Media/ElectraPlayer/Source/ElectraPlayerRuntime/Private/Runtime/StreamAccessUnitBuffer.h
--- a/Engine/Plugins/Media/ElectraPlayer/Source/ElectraPlayerRuntime/Private/Runtime/StreamAccessUnitBuffer.h
+++ b/Engine/Plugins/Media/ElectraPlayer/Source/ElectraPlayerRuntime/Private/Runtime/StreamAccessUnitBuffer.h
@@ -61,28 +61,19 @@ namespace Electra
 			FStreamCodecInformation			ParsedInfo;
 		};
 
-
-		enum EDropState
-		{
-			None = 0,
-			TooEarly = 1,
-			TooLate = 2,
-		};
-
+		TSharedPtrTS<CodecData>		AUCodecData;					//!< If set, points to constant sideband data for this access unit.
+		TSharedPtrTS<const FBufferSourceInfo>	BufferSourceInfo;
+		TUniquePtr<TMap<uint32, TArray<uint8>>> DynamicSidebandData;	//!< If set, contains a map of dynamically changing sideband data for just this access unit.
 		EStreamType					ESType;							//!< Type of elementary stream this is an access unit of.
 		FTimeValue					PTS;							//!< PTS
 		FTimeValue					DTS;							//!< DTS
 		FTimeValue					Duration;						//!< Duration
-		FTimeValue					PTO;							//!< Media local presentation time offset.
 		FTimeValue					EarliestPTS;					//!< Earliest PTS at which to present samples. If this is larger than PTS the sample is not to be presented.
 		FTimeValue					LatestPTS;						//!< Latest PTS at which to present samples. If this is less than PTS the sample is not to be presented.
 		FTimeValue					ProducerReferenceTime;			//!< If set, the wallclock time of the producer when this AU was encoded or captured
 		int64						SequenceIndex;
-		uint32						AUSize;							//!< Size of this access unit
 		void*						AUData;							//!< Access unit data
-		TSharedPtrTS<CodecData>		AUCodecData;					//!< If set, points to constant sideband data for this access unit.
-		TUniquePtr<TMap<uint32, TArray<uint8>>> DynamicSidebandData;	//!< If set, contains a map of dynamically changing sideband data for just this access unit.
-		uint8						DropState;						//!< If set this access unit is not to be rendered. If possible also not to be decoded.
+		uint32						AUSize;							//!< Size of this access unit
 		bool						bIsFirstInSequence;				//!< true for the first AU in a segment
 		bool						bIsLastInPeriod;				//!< true if this is the last AU in the playing period.
 		bool						bIsSyncSample;					//!< true if this is a sync sample (keyframe)
@@ -90,7 +81,6 @@ namespace Electra
 		bool						bTrackChangeDiscontinuity;		//!< True if this is the first AU after a track change.
 		bool						bIsSideloaded;					//!< True if the payload is not streamed but loaded from a sidecar file.
 
-		TSharedPtrTS<const FBufferSourceInfo>	BufferSourceInfo;
 
 		static FAccessUnit* Create(IAccessUnitMemoryProvider* MemProvider)
 		{
@@ -161,7 +151,6 @@ namespace Electra
 			PTS.SetToInvalid();
 			DTS.SetToInvalid();
 			Duration.SetToInvalid();
-			PTO.SetToZero();
 			EarliestPTS.SetToInvalid();
 			LatestPTS.SetToInvalid();
 			SequenceIndex = 0;
@@ -169,7 +158,6 @@ namespace Electra
 			AUSize = 0;
 			AUData = nullptr;
 			AUMemoryProvider = nullptr;
-			DropState = EDropState::None;
 			bIsFirstInSequence = false;
 			bIsLastInPeriod = false;
 			bIsSyncSample = false;
@@ -205,7 +193,6 @@ namespace Electra
 			FrontDTS.SetToInvalid();
 			SmallestPTS.SetToInvalid();
 			LargestPTSPlusDur.SetToInvalid();
-			PushedDuration.SetToZero();
 			PlayableDuration.SetToZero();
 			CurrentMemInUse = 0;
 			NumCurrentAccessUnits = 0;
@@ -214,16 +201,15 @@ namespace Electra
 			bLastPushWasBlocked = false;
 		}
 
-		FTimeValue			FrontDTS;
-		FTimeValue			SmallestPTS;
-		FTimeValue			LargestPTSPlusDur;
-		FTimeValue			PushedDuration;
-		FTimeValue			PlayableDuration;
-		int64				CurrentMemInUse;
-		int64				NumCurrentAccessUnits;
-		bool				bEndOfData;
-		bool				bEndOfTrack;
-		bool				bLastPushWasBlocked;
+		FTimeValue FrontDTS;
+		FTimeValue SmallestPTS;
+		FTimeValue LargestPTSPlusDur;
+		FTimeValue PlayableDuration;
+		int64 CurrentMemInUse;
+		int64 NumCurrentAccessUnits;
+		bool bEndOfData;
+		bool bEndOfTrack;
+		bool bLastPushWasBlocked;
 	};
 
 
@@ -251,7 +237,6 @@ namespace Electra
 			: FrontDTS(FTimeValue::GetInvalid())
 			, SmallestPTS(FTimeValue::GetInvalid())
 			, LargestPTSPlusDur(FTimeValue::GetInvalid())
-			, PushedDuration(FTimeValue::GetZero())
 			, PlayableDuration(FTimeValue::GetZero())
 			, CurrentMemInUse(0)
 			, bEndOfData(false)
@@ -282,12 +267,6 @@ namespace Electra
 			return CurrentMemInUse;
 		}
 
-		//! Returns the amount of playable duration.
-		FTimeValue GetPlayableDuration() const
-		{
-			return PlayableDuration;
-		}
-
 		//! Returns all vital statistics.
 		void GetStats(FAccessUnitBufferInfo& OutStats) const
 		{
@@ -295,7 +274,6 @@ namespace Electra
 			OutStats.FrontDTS = FrontDTS;
 			OutStats.SmallestPTS = SmallestPTS;
 			OutStats.LargestPTSPlusDur = LargestPTSPlusDur;
-			OutStats.PushedDuration = PushedDuration;
 			OutStats.PlayableDuration = PlayableDuration;
 			OutStats.CurrentMemInUse = CurrentMemInUse;
 			OutStats.NumCurrentAccessUnits = AccessUnits.Num();
@@ -313,14 +291,17 @@ namespace Electra
 			bEndOfData = false;
 			bEndOfTrack = false;
 			ExternalInfo = ExternalInfo ? ExternalInfo : &ZeroExternalInfo;
-			if (CanPush(AU, Limit, ExternalInfo))
+			check(AU->EarliestPTS.IsValid());
+			check(AU->LatestPTS.IsValid());
+			const bool bIsPlayableAU = AU->PTS >= AU->EarliestPTS && (AU->PTS + AU->Duration) < AU->LatestPTS;
+			if (!bIsPlayableAU || CanPush(AU, Limit, ExternalInfo))
 			{
 				bLastPushWasBlocked = false;
 				AccessUnits.Push(AU);
 				int64 memSize = AU->TotalMemSize();
 				CurrentMemInUse += memSize;
 
-				if (AU->DropState == FAccessUnit::EDropState::None)
+				if (bIsPlayableAU)
 				{
 					if (!FrontDTS.IsValid())
 					{
@@ -338,8 +319,6 @@ namespace Electra
 					}
 					PlayableDuration += AU->Duration;
 				}
-
-				PushedDuration += AU->Duration;
 				NumInSemaphore.Release();
 				return true;
 			}
@@ -382,7 +361,8 @@ namespace Electra
 
 				if (!AccessUnits.IsEmpty())
 				{
-					if (AccessUnits.FrontRef()->DropState == FAccessUnit::EDropState::None)
+					const bool bIsPlayableAU = AccessUnits.FrontRef()->PTS >= AccessUnits.FrontRef()->EarliestPTS && (AccessUnits.FrontRef()->PTS + AccessUnits.FrontRef()->Duration) < AccessUnits.FrontRef()->LatestPTS;
+					if (bIsPlayableAU)
 					{
 						FrontDTS = AccessUnits.FrontRef()->DTS;
 					}
@@ -391,17 +371,20 @@ namespace Electra
 						FrontDTS.SetToInvalid();
 						for(int32 i=1; i<AccessUnits.Num(); ++i)
 						{
-							if (AccessUnits[i]->DropState == FAccessUnit::EDropState::None)
+							const bool bNextIsPlayableAU = AccessUnits[i]->PTS >= AccessUnits[i]->EarliestPTS && (AccessUnits[i]->PTS + AccessUnits[i]->Duration) < AccessUnits[i]->LatestPTS;
+							if (bNextIsPlayableAU)
 							{
 								FrontDTS = AccessUnits[i]->DTS;
 								break;
 							}
 						}
 					}
+
 					SmallestPTS.SetToPositiveInfinity();
 					for(int32 i=0,iMax=AccessUnits.Num(),j=0; i<iMax; ++i)
 					{
-						if (AccessUnits[i]->DropState == FAccessUnit::EDropState::None)
+						const bool bNextIsPlayableAU = AccessUnits[i]->PTS >= AccessUnits[i]->EarliestPTS && (AccessUnits[i]->PTS + AccessUnits[i]->Duration) < AccessUnits[i]->LatestPTS;
+						if (bNextIsPlayableAU)
 						{
 							if (AccessUnits[i]->PTS < SmallestPTS)
 							{
@@ -415,11 +398,11 @@ namespace Electra
 							}
 						}
 					}
-					if (OutAU->DropState == FAccessUnit::EDropState::None)
+					const bool bWasPlayableAU = OutAU->PTS >= OutAU->EarliestPTS && (OutAU->PTS + OutAU->Duration) < OutAU->LatestPTS;
+					if (bWasPlayableAU)
 					{
 						PlayableDuration -= OutAU->Duration;
 					}
-					PushedDuration -= OutAU->Duration;
 				}
 				else
 				{
@@ -427,7 +410,6 @@ namespace Electra
 					SmallestPTS.SetToInvalid();
 					LargestPTSPlusDur.SetToInvalid();
 					PlayableDuration.SetToZero();
-					PushedDuration.SetToZero();
 				}
 				return true;
 			}
@@ -544,7 +526,6 @@ namespace Electra
 			bEndOfTrack = false;
 			bLastPushWasBlocked = false;
 			FrontDTS.SetToInvalid();
-			PushedDuration.SetToZero();
 			PlayableDuration.SetToZero();
 		}
 
@@ -599,13 +580,6 @@ namespace Electra
 			}
 			check(Limit->MaxDuration > FTimeValue::GetZero());
 			check(AU->Duration.IsValid() && !AU->Duration.IsInfinity());
-			// If tagged with a drop state we unconditionally accept the AU. These are early ones that won't get decoded.
-			// They should appear at the beginning and end of a stream only and since we need the following data that
-			// will be decoded we have to accept those first.
-			if (AU->DropState != 0)
-			{
-				return true;
-			}
 			if (ExternalInfo == nullptr)
 			{
 				// Max allowed duration ok?
@@ -618,19 +592,18 @@ namespace Electra
 			}
 		}
 
-		mutable FCriticalSection					AccessLock;
-		FExternalBufferInfo							ZeroExternalInfo;
-		TMediaQueueDynamicNoLock<FAccessUnit*>		AccessUnits;
-		FMediaSemaphore								NumInSemaphore;
-		FTimeValue									FrontDTS;					//!< DTS of first AU in buffer
-		FTimeValue									SmallestPTS;				//!< Smallest PTS in buffer
-		FTimeValue									LargestPTSPlusDur;			//!< Largest PTS plus the AU duration in buffer
-		FTimeValue									PushedDuration;
-		FTimeValue									PlayableDuration;
-		int64										CurrentMemInUse = 0;
-		bool										bEndOfData = false;
-		bool										bEndOfTrack = false;
-		bool										bLastPushWasBlocked = false;
+		mutable FCriticalSection AccessLock;
+		FExternalBufferInfo ZeroExternalInfo;
+		TMediaQueueDynamicNoLock<FAccessUnit*> AccessUnits;
+		FMediaSemaphore NumInSemaphore;
+		FTimeValue FrontDTS;
+		FTimeValue SmallestPTS;
+		FTimeValue LargestPTSPlusDur;
+		FTimeValue PlayableDuration;
+		int64 CurrentMemInUse = 0;
+		bool bEndOfData = false;
+		bool bEndOfTrack = false;
+		bool bLastPushWasBlocked = false;
 	};
 
 
