ElectraPlayer: Skipping over location in malformed `url ` boxes

#rnx

diff --git a/Engine/Plugins/Media/ElectraPlayer/Source/ElectraPlayerRuntime/Private/Runtime/Demuxer/ParserISO14496-12.cpp b/Engine/Plugins/Media/ElectraPlayer/Source/ElectraPlayerRuntime/Private/Runtime/Demuxer/ParserISO14496-12.cpp
--- a/Engine/Plugins/Media/ElectraPlayer/Source/ElectraPlayerRuntime/Private/Runtime/Demuxer/ParserISO14496-12.cpp
+++ b/Engine/Plugins/Media/ElectraPlayer/Source/ElectraPlayerRuntime/Private/Runtime/Demuxer/ParserISO14496-12.cpp
@@ -1938,11 +1938,22 @@ namespace Electra
 			RETURN_IF_ERROR(FMP4BoxFull::ReadAndParseAttributes(ParseInfo));
 			if ((Flags & 1) == 0)
 			{
-				uint32 BytesRemaining = BoxSize - (ParseInfo->Reader()->GetCurrentReadOffset() - StartOffset);
-				RETURN_IF_ERROR(ParseInfo->Reader()->ReadString(LocationUTF8, BytesRemaining));
+				uint32 LocationLen = BoxSize - (ParseInfo->Reader()->GetCurrentReadOffset() - StartOffset);
+				RETURN_IF_ERROR(ParseInfo->Reader()->ReadString(LocationUTF8, LocationLen));
+			}
+			/*
+				Ideally this has now all bytes of this box consumed. That's not a requirement but should be the case.
+				There are malformed `url ` boxes out there that violate the standard by providing a location despite
+				the flag being 1 as per 8.7.2.3:
+						one flag is defined (x000001) which means that the media data
+						is in the same file as the Box containing this data reference. If this flag is set, the DataEntryUrlBox
+						shall be used and no string is present; the box terminates with the entry-flags field.
+			*/
+			uint32 BytesRemaining = BoxSize - (ParseInfo->Reader()->GetCurrentReadOffset() - StartOffset);
+			if (BytesRemaining)
+			{
+				RETURN_IF_ERROR(ParseInfo->Reader()->ReadBytes(nullptr, BytesRemaining));
 			}
-			// Ideally this has now all bytes of this box consumed. That's not a requirement but should be the case.
-			check(0 == BoxSize - (ParseInfo->Reader()->GetCurrentReadOffset() - StartOffset));
 			return Error;
 		}
 
