ElectraPlayer: Checking for HLS spec violation of target duration being less than the actual media segment duration.

#rnx

diff --git a/Engine/Plugins/Media/ElectraPlayer/Source/ElectraPlayerRuntime/Private/Runtime/Player/HLS/PlaylistHandlerHLS.cpp b/Engine/Plugins/Media/ElectraPlayer/Source/ElectraPlayerRuntime/Private/Runtime/Player/HLS/PlaylistHandlerHLS.cpp
--- a/Engine/Plugins/Media/ElectraPlayer/Source/ElectraPlayerRuntime/Private/Runtime/Player/HLS/PlaylistHandlerHLS.cpp
+++ b/Engine/Plugins/Media/ElectraPlayer/Source/ElectraPlayerRuntime/Private/Runtime/Player/HLS/PlaylistHandlerHLS.cpp
@@ -963,6 +963,8 @@ bool FPlaylistHandlerHLS::BuildMediaPlaylist(TSharedPtr<FMediaPlaylistHLS, ESPMo
 		}
 	};
 
+	FTimeValue LongestMediaSegmentDuration;
+	LongestMediaSegmentDuration.SetToZero();
 	// Then process all other elements.
 	const TArray<TUniquePtr<FPlaylistParserHLS::FElement>>& Elements = InPlaylist->GetElements();
 	for(int32 nElem=0; nElem<Elements.Num(); ++nElem)
@@ -1019,6 +1021,12 @@ bool FPlaylistHandlerHLS::BuildMediaPlaylist(TSharedPtr<FMediaPlaylistHLS, ESPMo
 			}
 			Segment.Duration.SetFromTimeFraction(FTimeFraction().SetFromFloatString(Elements[nElem]->ElementValue.GetValue()));
 
+			// Track the longest segment duration to see if it aligns with the #EXT-X-TARGET-DURATION
+			if (Segment.Duration > LongestMediaSegmentDuration)
+			{
+				LongestMediaSegmentDuration = Segment.Duration;
+			}
+
 			if (NewMediaPlaylist->bHasProgramDateTime)
 			{
 				/*
@@ -1304,6 +1312,15 @@ bool FPlaylistHandlerHLS::BuildMediaPlaylist(TSharedPtr<FMediaPlaylistHLS, ESPMo
 		}
 	}
 
+	// Check that the segments are within the target duration
+	int64 MaxSegDurationInt = FMath::CeilToInt(LongestMediaSegmentDuration.GetAsSeconds());
+	int64 TargetDuration = FMath::CeilToInt(NewMediaPlaylist->TargetDuration.GetAsSeconds());
+	if (MaxSegDurationInt > TargetDuration)
+	{
+		LogMessage(IInfoLog::ELevel::Warning, FString::Printf(TEXT("HLS (RFC-8216) violation: Longest playlist #EXTINF segment duration of %.4f is longer than the specified #EXT-TARGET-DURATION of %i, adjusting."), LongestMediaSegmentDuration.GetAsSeconds(), (int32)TargetDuration));
+		NewMediaPlaylist->TargetDuration = LongestMediaSegmentDuration;
+	}
+
 	// If PDT is given in the playlist, but none of its values are valid, treat it as absent PDTs.
 	if (NewMediaPlaylist->bHasProgramDateTime && !bHasValidDateTime)
 	{
