Electra: fixed incorrect integer rounding

#rnx

diff --git a/Engine/Plugins/Media/ElectraPlayer/Source/ElectraPlayerRuntime/Private/Runtime/Player/HLS/PlaylistHandlerHLS.cpp b/Engine/Plugins/Media/ElectraPlayer/Source/ElectraPlayerRuntime/Private/Runtime/Player/HLS/PlaylistHandlerHLS.cpp
--- a/Engine/Plugins/Media/ElectraPlayer/Source/ElectraPlayerRuntime/Private/Runtime/Player/HLS/PlaylistHandlerHLS.cpp
+++ b/Engine/Plugins/Media/ElectraPlayer/Source/ElectraPlayerRuntime/Private/Runtime/Player/HLS/PlaylistHandlerHLS.cpp
@@ -1313,7 +1313,7 @@ bool FPlaylistHandlerHLS::BuildMediaPlaylist(TSharedPtr<FMediaPlaylistHLS, ESPMo
 	}
 
 	// Check that the segments are within the target duration
-	int64 MaxSegDurationInt = FMath::CeilToInt(LongestMediaSegmentDuration.GetAsSeconds());
+	int64 MaxSegDurationInt = FMath::RoundToInt(LongestMediaSegmentDuration.GetAsSeconds());
 	int64 TargetDuration = FMath::CeilToInt(NewMediaPlaylist->TargetDuration.GetAsSeconds());
 	if (MaxSegDurationInt > TargetDuration)
 	{
