Electra: Added option to drain decoder on backwards timejump; keeping stale AUs longer in case of timejumps that are less than the check period; added log messages when timejump in either direction occurs (if enabled to check for)

#rnx

diff --git a/Engine/Plugins/Media/ElectraPlayer/Source/ElectraPlayerRuntime/Private/Runtime/Decoder/AudioDecoder.cpp b/Engine/Plugins/Media/ElectraPlayer/Source/ElectraPlayerRuntime/Private/Runtime/Decoder/AudioDecoder.cpp
--- a/Engine/Plugins/Media/ElectraPlayer/Source/ElectraPlayerRuntime/Private/Runtime/Decoder/AudioDecoder.cpp
+++ b/Engine/Plugins/Media/ElectraPlayer/Source/ElectraPlayerRuntime/Private/Runtime/Decoder/AudioDecoder.cpp
@@ -155,6 +155,7 @@ private:
 	TArray<TSharedPtrTS<FDecoderInput>>										InDecoderInput;
 	TSharedPtrTS<FDecoderInput>												CurrentAccessUnit;
 	TOptional<int64>														CurrentSequenceIndex;
+	TOptional<int64>														NextExpectedDTSHNS;
 	EDecodingState															CurrentDecodingState = EDecodingState::NormalDecoding;
 	EDecodingState															NextDecodingStateAfterDrain = EDecodingState::NormalDecoding;
 	bool																	bIsDecoderClean = true;
@@ -641,7 +642,24 @@ FAudioDecoderImpl::EAUChangeFlags FAudioDecoderImpl::GetAndPrepareInputAU()
 				{
 					NewAUFlags |= EAUChangeFlags::Discontinuity;
 				}
+				else if (DecoderConfig::bCheckForDTSTimejump && NextExpectedDTSHNS.IsSet() && NextExpectedDTSHNS.GetValue() > CurrentAccessUnit->AccessUnit->DTS.GetAsHNS() + DecoderConfig::BackwardsTimejumpThresholdHNS)
+				{
+					if (DecoderConfig::bDrainDecoderOnDetectedBackwardsTimejump)
+					{
+						LogMessage(IInfoLog::ELevel::Warning, FString::Printf(TEXT("Audio timestamp jumped back unexpectedly by %.4fs. Draining the decoder before continuing."), (NextExpectedDTSHNS.GetValue() - CurrentAccessUnit->AccessUnit->DTS.GetAsHNS()) / 10000000.0));
+						NewAUFlags |= EAUChangeFlags::Discontinuity;
+					}
+					else
+					{
+						LogMessage(IInfoLog::ELevel::Warning, FString::Printf(TEXT("Audio timestamp jumped back unexpectedly by %.4fs."), (NextExpectedDTSHNS.GetValue() - CurrentAccessUnit->AccessUnit->DTS.GetAsHNS()) / 10000000.0));
+					}
+				}
+				else if (DecoderConfig::bCheckForDTSTimejump && NextExpectedDTSHNS.IsSet() && NextExpectedDTSHNS.GetValue() < CurrentAccessUnit->AccessUnit->DTS.GetAsHNS() - DecoderConfig::ForwardTimejumpThresholdHNS)
+				{
+					LogMessage(IInfoLog::ELevel::Warning, FString::Printf(TEXT("Audio timestamp jumped forward unexpectedly by %.4fs."), (CurrentAccessUnit->AccessUnit->DTS.GetAsHNS() - NextExpectedDTSHNS.GetValue()) / 10000000.0));
+				}
 				CurrentSequenceIndex = CurrentAccessUnit->AccessUnit->PTS.GetSequenceIndex();
+				NextExpectedDTSHNS = (CurrentAccessUnit->AccessUnit->DTS + CurrentAccessUnit->AccessUnit->Duration).GetAsHNS();
 
 				// Did the codec specific data change?
 				if (CurrentCodecData.IsValid() && CurrentAccessUnit->AccessUnit->AUCodecData.IsValid() && CurrentCodecData->CodecSpecificData != CurrentAccessUnit->AccessUnit->AUCodecData->CodecSpecificData)
@@ -757,7 +775,6 @@ IElectraDecoder::EOutputStatus FAudioDecoderImpl::HandleOutput()
 						/*
 							Not the first element. This can happen if the audio decoder aggregates output from multiple inputs.
 						*/
-						const int64 kTooOldThresholdHNS = 5000000;	// 0.5 seconds
 						for (int32 i = 0; i < InDecoderInput.Num(); ++i)
 						{
 							if (InDecoderInput[i]->PTS == Output->GetUserValue())
@@ -766,7 +783,7 @@ IElectraDecoder::EOutputStatus FAudioDecoderImpl::HandleOutput()
 								InDecoderInput.RemoveAt(i);
 								break;
 							}
-							else if (InDecoderInput[i]->PTS + kTooOldThresholdHNS < (int64)Output->GetUserValue())
+							else if (InDecoderInput[i]->PTS + DecoderConfig::RemovalOfOldDecoderInputThresholdHNS < (int64)Output->GetUserValue())
 							{
 								InDecoderInput.RemoveAt(i);
 								--i;
@@ -917,7 +934,7 @@ bool FAudioDecoderImpl::HandleDecoding()
 			if (DecErr == IElectraDecoder::EDecoderError::None)
 			{
 				InDecoderInput.Emplace(CurrentAccessUnit);
-				InDecoderInput.Sort([](const TSharedPtr<FDecoderInput, ESPMode::ThreadSafe>& a, const TSharedPtr<FDecoderInput, ESPMode::ThreadSafe>& b)
+				InDecoderInput.StableSort([](const TSharedPtr<FDecoderInput, ESPMode::ThreadSafe>& a, const TSharedPtr<FDecoderInput, ESPMode::ThreadSafe>& b)
 				{
 					return a->PTS < b->PTS;
 				});
@@ -1065,6 +1082,7 @@ bool FAudioDecoderImpl::CheckForFlush()
 		NextAccessUnits.Empty();
 		InDecoderInput.Empty();
 		CurrentSequenceIndex.Reset();
+		NextExpectedDTSHNS.Reset();
 		CurrentAccessUnit.Reset();
 		CurrentOutputFormat.Reset();
 		bIsDecoderClean = true;
@@ -1266,6 +1284,7 @@ void FAudioDecoderImpl::WorkerThread()
 	NextAccessUnits.Empty();
 	InDecoderInput.Empty();
 	CurrentSequenceIndex.Reset();
+	NextExpectedDTSHNS.Reset();
 	CurrentOutputFormat.Reset();
 	CurrentCodecData.Reset();
 	ChannelMapper.Reset();
diff --git a/Engine/Plugins/Media/ElectraPlayer/Source/ElectraPlayerRuntime/Private/Runtime/Decoder/DecoderBase.h b/Engine/Plugins/Media/ElectraPlayer/Source/ElectraPlayerRuntime/Private/Runtime/Decoder/DecoderBase.h
--- a/Engine/Plugins/Media/ElectraPlayer/Source/ElectraPlayerRuntime/Private/Runtime/Decoder/DecoderBase.h
+++ b/Engine/Plugins/Media/ElectraPlayer/Source/ElectraPlayerRuntime/Private/Runtime/Decoder/DecoderBase.h
@@ -12,6 +12,22 @@ namespace Electra
 	class IAccessUnitBufferListener;
 	class IDecoderOutputBufferListener;
 
+	namespace DecoderConfig
+	{
+		// Should we check if the DTS of the input AU jumps?
+		static const bool bCheckForDTSTimejump = true;
+		// When checking the AU DTS, how much does it need to jump back before triggering a decoder drain?
+		static const int64 BackwardsTimejumpThresholdHNS = 10 * 1000 * 500;						// 500ms
+		// Drain the decoder when a timejump is being detected?
+		static const bool bDrainDecoderOnDetectedBackwardsTimejump = false;
+		// When checking the AU DTS, how much does it need to jump forward before logging a warning?
+		static const int64 ForwardTimejumpThresholdHNS = 10 * 1000 * 500;						// 500ms
+
+		// After what time decoder inputs are to be flushed in case some decoder implementation drops
+		// output which we would otherwise keep the source AU around for.
+		static const int64 RemovalOfOldDecoderInputThresholdHNS = 10 * 1000 * 1000 * 10;		// 10 seconds
+	};
+
 
 	/**
 	 *
@@ -83,12 +99,12 @@ namespace Electra
 		{
 			return DecoderInstance;
 		}
-		
+
 		EStreamType GetStreamType() const
 		{
 			return StreamType;
 		}
-		
+
 		FStreamCodecInformation::ECodec GetCodec() const
 		{
 			return Codec;
@@ -96,16 +112,16 @@ namespace Electra
 
 	private:
 		FDecoderMessage(EReason InReason, IDecoderBase* InDecoderInstance, EStreamType InStreamType, FStreamCodecInformation::ECodec InCodec)
-			: Reason(InReason)
+			: Codec(InCodec)
 			, DecoderInstance(InDecoderInstance)
 			, StreamType(InStreamType)
-			, Codec(InCodec)
+			, Reason(InReason)
 		{
 		}
-		EReason							Reason;
-		IDecoderBase*					DecoderInstance;
-		EStreamType						StreamType;
-		FStreamCodecInformation::ECodec	Codec;
+		FStreamCodecInformation::ECodec Codec;
+		IDecoderBase* DecoderInstance;
+		EStreamType StreamType;
+		EReason Reason;
 	};
 
 
diff --git a/Engine/Plugins/Media/ElectraPlayer/Source/ElectraPlayerRuntime/Private/Runtime/Decoder/VideoDecoder.cpp b/Engine/Plugins/Media/ElectraPlayer/Source/ElectraPlayerRuntime/Private/Runtime/Decoder/VideoDecoder.cpp
--- a/Engine/Plugins/Media/ElectraPlayer/Source/ElectraPlayerRuntime/Private/Runtime/Decoder/VideoDecoder.cpp
+++ b/Engine/Plugins/Media/ElectraPlayer/Source/ElectraPlayerRuntime/Private/Runtime/Decoder/VideoDecoder.cpp
@@ -160,6 +160,7 @@ private:
 	TArray<TSharedPtrTS<FDecoderInput>>										InDecoderInput;
 	TSharedPtrTS<FDecoderInput>												CurrentAccessUnit;
 	TOptional<int64>														CurrentSequenceIndex;
+	TOptional<int64>														NextExpectedDTSHNS;
 	EDecodingState															CurrentDecodingState = EDecodingState::NormalDecoding;
 	EDecodingState															NextDecodingStateAfterDrain = EDecodingState::NormalDecoding;
 	bool																	bIsDecoderClean = true;
@@ -751,7 +752,24 @@ FVideoDecoderImpl::EAUChangeFlags FVideoDecoderImpl::GetAndPrepareInputAU()
 				{
 					NewAUFlags |= EAUChangeFlags::Discontinuity;
 				}
+				else if (DecoderConfig::bCheckForDTSTimejump && NextExpectedDTSHNS.IsSet() && NextExpectedDTSHNS.GetValue() > CurrentAccessUnit->AccessUnit->DTS.GetAsHNS() + DecoderConfig::BackwardsTimejumpThresholdHNS)
+				{
+					if (DecoderConfig::bDrainDecoderOnDetectedBackwardsTimejump)
+					{
+						LogMessage(IInfoLog::ELevel::Warning, FString::Printf(TEXT("Video timestamp jumped back unexpectedly by %.4fs. Draining the decoder before continuing."), (NextExpectedDTSHNS.GetValue() - CurrentAccessUnit->AccessUnit->DTS.GetAsHNS()) / 10000000.0));
+						NewAUFlags |= EAUChangeFlags::Discontinuity;
+					}
+					else
+					{
+						LogMessage(IInfoLog::ELevel::Warning, FString::Printf(TEXT("Video timestamp jumped back unexpectedly by %.4fs."), (NextExpectedDTSHNS.GetValue() - CurrentAccessUnit->AccessUnit->DTS.GetAsHNS()) / 10000000.0));
+					}
+				}
+				else if (DecoderConfig::bCheckForDTSTimejump && NextExpectedDTSHNS.IsSet() && NextExpectedDTSHNS.GetValue() < CurrentAccessUnit->AccessUnit->DTS.GetAsHNS() - DecoderConfig::ForwardTimejumpThresholdHNS)
+				{
+					LogMessage(IInfoLog::ELevel::Warning, FString::Printf(TEXT("Video timestamp jumped forward unexpectedly by %.4fs."), (CurrentAccessUnit->AccessUnit->DTS.GetAsHNS() - NextExpectedDTSHNS.GetValue()) / 10000000.0));
+				}
 				CurrentSequenceIndex = CurrentAccessUnit->AccessUnit->PTS.GetSequenceIndex();
+				NextExpectedDTSHNS = (CurrentAccessUnit->AccessUnit->DTS + CurrentAccessUnit->AccessUnit->Duration).GetAsHNS();
 
 				// The very first access unit can't have differences to the one before so we clear the flags.
 				if (bIsFirstAccessUnit)
@@ -900,7 +918,6 @@ IElectraDecoder::EOutputStatus FVideoDecoderImpl::HandleOutput()
 							Any elements in the list that are far too old we remove since it is not likely for the decoder to
 							emit those frames at all and we don't want our list to grow too long.
 						*/
-						const int64 kTooOldThresholdHNS = 20000000;	// 2 seconds
 						for(int32 i=0; i<InDecoderInput.Num(); ++i)
 						{
 							if (InDecoderInput[i]->PTS == CurrentDecoderOutput->GetUserValue())
@@ -909,7 +926,7 @@ IElectraDecoder::EOutputStatus FVideoDecoderImpl::HandleOutput()
 								InDecoderInput.RemoveAt(i);
 								break;
 							}
-							else if (InDecoderInput[i]->PTS + kTooOldThresholdHNS < (int64)CurrentDecoderOutput->GetUserValue())
+							else if (InDecoderInput[i]->PTS + DecoderConfig::RemovalOfOldDecoderInputThresholdHNS < (int64)CurrentDecoderOutput->GetUserValue())
 							{
 								InDecoderInput.RemoveAt(i);
 								--i;
@@ -1049,7 +1066,7 @@ FVideoDecoderImpl::ENextDecodingState FVideoDecoderImpl::HandleDecoding()
 				if ((DecAU.Flags & EElectraDecoderFlags::DoNotOutput) == EElectraDecoderFlags::None)
 				{
 					InDecoderInput.Emplace(CurrentAccessUnit);
-					InDecoderInput.Sort([](const TSharedPtr<FDecoderInput, ESPMode::ThreadSafe>& a, const TSharedPtr<FDecoderInput, ESPMode::ThreadSafe>& b)
+					InDecoderInput.StableSort([](const TSharedPtr<FDecoderInput, ESPMode::ThreadSafe>& a, const TSharedPtr<FDecoderInput, ESPMode::ThreadSafe>& b)
 					{
 						return a->PTS < b->PTS;
 					});
@@ -1307,6 +1324,7 @@ bool FVideoDecoderImpl::CheckForFlush()
 		ReplayAccessUnit.Reset();
 		InDecoderInput.Empty();
 		CurrentSequenceIndex.Reset();
+		NextExpectedDTSHNS.Reset();
 		CurrentAccessUnit.Reset();
 		bIsDecoderClean = true;
 		bInDummyDecodeMode = false;
@@ -1582,6 +1600,7 @@ void FVideoDecoderImpl::WorkerThread()
 	NextAccessUnits.Empty();
 	InDecoderInput.Empty();
 	CurrentSequenceIndex.Reset();
+	NextExpectedDTSHNS.Reset();
 	ReplayAccessUnit.Reset();
 	ReplayAccessUnits.Empty();
 	ReplayingAccessUnits.Empty();
