ElectraPlayer: Added missing setting of static resource request completed flag which prevents HLS playlist loads being served internally from Blurl or EP

#rnx

diff --git a/Engine/Plugins/Media/ElectraPlayer/Source/ElectraPlayerRuntime/Private/Runtime/Player/AdaptiveStreamingPlayerResourceRequest.cpp b/Engine/Plugins/Media/ElectraPlayer/Source/ElectraPlayerRuntime/Private/Runtime/Player/AdaptiveStreamingPlayerResourceRequest.cpp
--- a/Engine/Plugins/Media/ElectraPlayer/Source/ElectraPlayerRuntime/Private/Runtime/Player/AdaptiveStreamingPlayerResourceRequest.cpp
+++ b/Engine/Plugins/Media/ElectraPlayer/Source/ElectraPlayerRuntime/Private/Runtime/Player/AdaptiveStreamingPlayerResourceRequest.cpp
@@ -186,6 +186,7 @@ void FHTTPResourceRequest::StaticDataReady()
 			bInCallback = true;
 			CompletedCallback.ExecuteIfBound(AsShared());
 			bInCallback = false;
+			bHasFinished = true;
 		}
 	}
 }
