Electra: Removed old bitstream processors in favor of the updated ones

#rnx

diff --git a/Engine/Plugins/Media/ElectraPlayer/Source/ElectraPlayerRuntime/Private/Runtime/Decoder/BitstreamProcessors/VideoDecoderInputBitstreamProcessorH264.cpp b/Engine/Plugins/Media/ElectraPlayer/Source/ElectraPlayerRuntime/Private/Runtime/Decoder/BitstreamProcessors/VideoDecoderInputBitstreamProcessorH264.cpp
deleted file mode 100644
diff --git a/Engine/Plugins/Media/ElectraPlayer/Source/ElectraPlayerRuntime/Private/Runtime/Decoder/BitstreamProcessors/VideoDecoderInputBitstreamProcessorH264.h b/Engine/Plugins/Media/ElectraPlayer/Source/ElectraPlayerRuntime/Private/Runtime/Decoder/BitstreamProcessors/VideoDecoderInputBitstreamProcessorH264.h
deleted file mode 100644
diff --git a/Engine/Plugins/Media/ElectraPlayer/Source/ElectraPlayerRuntime/Private/Runtime/Decoder/BitstreamProcessors/VideoDecoderInputBitstreamProcessorH265.cpp b/Engine/Plugins/Media/ElectraPlayer/Source/ElectraPlayerRuntime/Private/Runtime/Decoder/BitstreamProcessors/VideoDecoderInputBitstreamProcessorH265.cpp
deleted file mode 100644
diff --git a/Engine/Plugins/Media/ElectraPlayer/Source/ElectraPlayerRuntime/Private/Runtime/Decoder/BitstreamProcessors/VideoDecoderInputBitstreamProcessorH265.h b/Engine/Plugins/Media/ElectraPlayer/Source/ElectraPlayerRuntime/Private/Runtime/Decoder/BitstreamProcessors/VideoDecoderInputBitstreamProcessorH265.h
deleted file mode 100644
diff --git a/Engine/Plugins/Media/ElectraPlayer/Source/ElectraPlayerRuntime/Private/Runtime/Decoder/BitstreamProcessors/VideoDecoderInputBitstreamProcessorVP8.cpp b/Engine/Plugins/Media/ElectraPlayer/Source/ElectraPlayerRuntime/Private/Runtime/Decoder/BitstreamProcessors/VideoDecoderInputBitstreamProcessorVP8.cpp
deleted file mode 100644
diff --git a/Engine/Plugins/Media/ElectraPlayer/Source/ElectraPlayerRuntime/Private/Runtime/Decoder/BitstreamProcessors/VideoDecoderInputBitstreamProcessorVP8.h b/Engine/Plugins/Media/ElectraPlayer/Source/ElectraPlayerRuntime/Private/Runtime/Decoder/BitstreamProcessors/VideoDecoderInputBitstreamProcessorVP8.h
deleted file mode 100644
diff --git a/Engine/Plugins/Media/ElectraPlayer/Source/ElectraPlayerRuntime/Private/Runtime/Decoder/BitstreamProcessors/VideoDecoderInputBitstreamProcessorVP9.cpp b/Engine/Plugins/Media/ElectraPlayer/Source/ElectraPlayerRuntime/Private/Runtime/Decoder/BitstreamProcessors/VideoDecoderInputBitstreamProcessorVP9.cpp
deleted file mode 100644
diff --git a/Engine/Plugins/Media/ElectraPlayer/Source/ElectraPlayerRuntime/Private/Runtime/Decoder/BitstreamProcessors/VideoDecoderInputBitstreamProcessorVP9.h b/Engine/Plugins/Media/ElectraPlayer/Source/ElectraPlayerRuntime/Private/Runtime/Decoder/BitstreamProcessors/VideoDecoderInputBitstreamProcessorVP9.h
deleted file mode 100644
diff --git a/Engine/Plugins/Media/ElectraPlayer/Source/ElectraPlayerRuntime/Private/Runtime/Decoder/VideoDecoder.cpp b/Engine/Plugins/Media/ElectraPlayer/Source/ElectraPlayerRuntime/Private/Runtime/Decoder/VideoDecoder.cpp
--- a/Engine/Plugins/Media/ElectraPlayer/Source/ElectraPlayerRuntime/Private/Runtime/Decoder/VideoDecoder.cpp
+++ b/Engine/Plugins/Media/ElectraPlayer/Source/ElectraPlayerRuntime/Private/Runtime/Decoder/VideoDecoder.cpp
@@ -10,14 +10,20 @@
 
 #include "StreamAccessUnitBuffer.h"
 #include "Decoder/VideoDecoder.h"
+#include "Decoder/VideoDecoderHelpers.h"
 #include "Renderer/RendererBase.h"
 #include "Player/PlayerSessionServices.h"
 #include "Utilities/Utilities.h"
 #include "Utilities/StringHelpers.h"
-
 #include "SynchronizedClock.h"
 
-#include "VideoDecoderInputBitstreamProcessor.h"
+#include "IElectraCodecFactoryModule.h"
+#include "IElectraCodecFactory.h"
+#include "IElectraDecoder.h"
+#include "IElectraDecoderOutputVideo.h"
+#include "IElectraDecoderFeaturesAndOptions.h"
+#include "ElectraDecodersUtils.h"
+#include "ElectraDecoderResourceManager.h"
 
 // Error codes must be in the 1000-1999 range. 1-999 is reserved for the decoder implementation.
 #define ERRCODE_VIDEO_INTERNAL_COULD_NOT_CREATE_DECODER			1001
@@ -25,16 +31,9 @@
 #define ERRCODE_VIDEO_INTERNAL_COULD_NOT_GET_SAMPLE_BUFFER		1003
 #define ERRCODE_VIDEO_INTERNAL_UNSUPPORTED_OUTPUT_FORMAT		1004
 #define ERRCODE_VIDEO_INTERNAL_FAILED_TO_CONVERT_OUTPUT			1005
+#define ERRCODE_VIDEO_INTERNAL_FAILED_TO_PARSE_BITSTREAM		1006
 
 
-#include "IElectraCodecFactoryModule.h"
-#include "IElectraCodecFactory.h"
-#include "IElectraDecoder.h"
-#include "IElectraDecoderOutputVideo.h"
-#include "IElectraDecoderFeaturesAndOptions.h"
-#include "ElectraDecodersUtils.h"
-#include "ElectraDecoderResourceManager.h"
-
 /***************************************************************************************************************************************************/
 
 DECLARE_CYCLE_STAT(TEXT("FVideoDecoderImpl::Decode()"), STAT_ElectraPlayer_VideoDecode, STATGROUP_ElectraPlayer);
@@ -68,6 +67,11 @@ public:
 	void SetReadyBufferListener(IDecoderOutputBufferListener* Listener) override;
 
 private:
+	enum
+	{
+		CfgMakeAUCopy = 0
+	};
+
 	struct FDecoderInput
 	{
 		~FDecoderInput()
@@ -75,17 +79,16 @@ private:
 			FAccessUnit::Release(AccessUnit);
 		}
 
-		FTimeValue		AdjustedPTS;
-		FTimeValue		AdjustedDuration;
-		FTimeValue		StartOverlapDuration;
-		FTimeValue		EndOverlapDuration;
-		FAccessUnit*	AccessUnit = nullptr;
-		int64			PTS = 0;
-		int64			EndPTS = 0;
-		bool			bHasBeenPrepared = false;
-		bool			bMaySkipDecoding = false;
-
-		IVideoDecoderInputBitstreamProcessor::FBitstreamInfo BitstreamInfo;
+		IElectraDecoder::FInputAccessUnit DecAU;
+		TMap<FString, FVariant> CSDOptions;
+		TArray<uint8> DataCopy;
+		TSharedPtr<IElectraDecoderBitstreamInfo, ESPMode::ThreadSafe> BSI;
+		FTimeValue AdjustedPTS;
+		FTimeValue AdjustedDuration;
+		FAccessUnit* AccessUnit = nullptr;
+		int64 PTS = 0;
+		bool bHasBeenPrepared = false;
+		bool bMaySkipDecoding = false;
 	};
 
 	enum class EDecodingState
@@ -158,6 +161,7 @@ private:
 	TSharedPtrTS<FDecoderInput>												ReplayAccessUnit;
 
 	TArray<TSharedPtrTS<FDecoderInput>>										InDecoderInput;
+	TMap<FString, FVariant>													CurrentActiveCSD;
 	TSharedPtrTS<FDecoderInput>												CurrentAccessUnit;
 	TOptional<int64>														CurrentSequenceIndex;
 	TOptional<int64>														NextExpectedDTSHNS;
@@ -205,15 +209,18 @@ private:
 
 	IElectraDecoderResourceDelegateBase::IDecoderPlatformResource*			PlatformResource = nullptr;
 
-	TSharedPtr<IVideoDecoderInputBitstreamProcessor, ESPMode::ThreadSafe>	BitstreamProcessor;
 	TMap<FString, FVariant>													DecoderConfigOptions;
 	TSharedPtr<IElectraDecoder, ESPMode::ThreadSafe>						DecoderInstance;
+	TSharedPtr<IElectraDecoderBitstreamProcessor, ESPMode::ThreadSafe>		DecoderBitstreamProcessor;
+
 	bool																	bIsAdaptiveDecoder = false;
 	bool																	bSupportsDroppingOutput = false;
 	bool																	bNeedsReplayData = true;
 	bool																	bMustBeSuspendedInBackground = false;
 
 	TSharedPtr<IElectraDecoderVideoOutput, ESPMode::ThreadSafe>				CurrentDecoderOutput;
+	TOptional<MPEG::FColorimetryHelper>										CurrentColorimetry;
+	TOptional<MPEG::FHDRHelper>												CurrentHDR;
 
 	IMediaRenderer::IBuffer*												CurrentOutputBuffer = nullptr;
 	FParamDict																EmptyOptions;
@@ -606,7 +613,8 @@ bool FVideoDecoderImpl::InternalDecoderCreate()
 		}
 	}
 
-
+	// Get the bitstream processor for this decoder, if it requires one.
+	DecoderBitstreamProcessor = DecoderInstance->CreateBitstreamProcessor();
 	return true;
 }
 
@@ -617,6 +625,11 @@ void FVideoDecoderImpl::InternalDecoderDestroy()
 		RemoveBGFGNotificationHandler(FGBGHandlers);
 		FGBGHandlers.Reset();
 	}
+	if (DecoderBitstreamProcessor.IsValid())
+	{
+		DecoderBitstreamProcessor->Clear();
+		DecoderBitstreamProcessor.Reset();
+	}
 	if (DecoderInstance.IsValid())
 	{
 		DecoderInstance->Close();
@@ -630,6 +643,9 @@ void FVideoDecoderImpl::InternalDecoderDestroy()
 	bIsAdaptiveDecoder = false;
 	bSupportsDroppingOutput = false;
 	bNeedsReplayData = true;
+	CurrentActiveCSD.Empty();
+	CurrentColorimetry.Reset();
+	CurrentHDR.Reset();
 }
 
 void FVideoDecoderImpl::ReturnUnusedOutputBuffer()
@@ -652,9 +668,6 @@ FVideoDecoderImpl::EAUChangeFlags FVideoDecoderImpl::PrepareAU(TSharedPtrTS<FDec
 		FTimeValue StartTime = AU->AccessUnit->PTS;
 		FTimeValue EndTime = AU->AccessUnit->PTS + AU->AccessUnit->Duration;
 		AU->PTS = StartTime.GetAsHNS();		// The PTS we give the decoder no matter any adjustment.
-		AU->EndPTS = EndTime.GetAsHNS();
-		AU->StartOverlapDuration.SetToZero();
-		AU->EndOverlapDuration.SetToZero();
 		if (AU->AccessUnit->EarliestPTS.IsValid())
 		{
 			// If the end time of the AU is before the earliest render PTS we do not need to decode it.
@@ -665,7 +678,6 @@ FVideoDecoderImpl::EAUChangeFlags FVideoDecoderImpl::PrepareAU(TSharedPtrTS<FDec
 			}
 			else if (StartTime < AU->AccessUnit->EarliestPTS)
 			{
-				AU->StartOverlapDuration = AU->AccessUnit->EarliestPTS - StartTime;
 				StartTime = AU->AccessUnit->EarliestPTS;
 			}
 		}
@@ -683,7 +695,6 @@ FVideoDecoderImpl::EAUChangeFlags FVideoDecoderImpl::PrepareAU(TSharedPtrTS<FDec
 			}
 			else if (EndTime >= AU->AccessUnit->LatestPTS)
 			{
-				AU->EndOverlapDuration = EndTime - AU->AccessUnit->LatestPTS;
 				EndTime = AU->AccessUnit->LatestPTS;
 			}
 		}
@@ -694,10 +705,52 @@ FVideoDecoderImpl::EAUChangeFlags FVideoDecoderImpl::PrepareAU(TSharedPtrTS<FDec
 			AU->AdjustedPTS.SetToInvalid();
 		}
 
-		// Extract codec specific messages and frame properties.
-		if (BitstreamProcessor->ProcessAccessUnitForDecoding(AU->BitstreamInfo, AU->AccessUnit) == IVideoDecoderInputBitstreamProcessor::EProcessResult::CSDChanged)
+		// Get the codec specific data
+		if (AU->AccessUnit->AUCodecData.IsValid())
 		{
-			NewAUFlags = EAUChangeFlags::CSDChanged;
+			AU->CSDOptions.Emplace(TEXT("csd"), FVariant(CurrentAccessUnit->AccessUnit->AUCodecData->CodecSpecificData));
+			AU->CSDOptions.Emplace(TEXT("dcr"), FVariant(CurrentAccessUnit->AccessUnit->AUCodecData->RawCSD));
+		}
+		else
+		{
+			AU->CSDOptions = CurrentActiveCSD;
+		}
+
+		// Set the timestamps in the decoder input description.
+		AU->DecAU.DTS = AU->AccessUnit->DTS.GetAsTimespan();
+		AU->DecAU.PTS = AU->AccessUnit->PTS.GetAsTimespan();
+		AU->DecAU.Duration = AU->AccessUnit->Duration.GetAsTimespan();
+		AU->DecAU.UserValue = AU->PTS;
+		AU->DecAU.Flags |= AU->AccessUnit->bIsSyncSample ? EElectraDecoderFlags::IsSyncSample : EElectraDecoderFlags::None;
+
+		// Set the bitstream data and size in the decoder input description.
+		// If the bitstream processor will modify the data in place we need to make a copy.
+		// NOTE: In-place modification only means changing values in place, but not removing or inserting new data.
+		AU->DecAU.Data = AU->AccessUnit->AUData;
+		AU->DecAU.DataSize = AU->AccessUnit->AUSize;
+		if (DecoderBitstreamProcessor.IsValid() && DecoderBitstreamProcessor->WillModifyBitstreamInPlace() && CfgMakeAUCopy)
+		{
+			AU->DataCopy = MakeConstArrayView<const uint8>(reinterpret_cast<const uint8*>(AU->AccessUnit->AUData), AU->AccessUnit->AUSize);
+			AU->DecAU.Data = AU->DataCopy.GetData();
+		}
+		// We do not need to set up more than the bitstream data and the codec specific info to call the bitstream processor.
+		IElectraDecoderBitstreamProcessor::EProcessResult BSResult = DecoderBitstreamProcessor->ProcessInputForDecoding(AU->BSI, AU->DecAU, AU->CSDOptions);
+		if (BSResult == IElectraDecoderBitstreamProcessor::EProcessResult::Error)
+		{
+			PostError(-2, DecoderBitstreamProcessor->GetLastError(), ERRCODE_VIDEO_INTERNAL_FAILED_TO_PARSE_BITSTREAM);
+			return NewAUFlags;
+		}
+		else if (BSResult == IElectraDecoderBitstreamProcessor::EProcessResult::CSDChanged)
+		{
+			CurrentActiveCSD = AU->CSDOptions;
+			NewAUFlags |= EAUChangeFlags::CSDChanged;
+		}
+
+		if (bSupportsDroppingOutput && !AU->AdjustedPTS.IsValid())
+		{
+			NumInitialSkippedFrames += bIsStartOfSequence ? 1 : 0;
+			NumInitialSkippedDecodingFrames += (AU->DecAU.Flags & EElectraDecoderFlags::IsDiscardable) == EElectraDecoderFlags::IsDiscardable ? 1 : 0;
+			AU->DecAU.Flags |= EElectraDecoderFlags::DoNotOutput;
 		}
 	}
 	return NewAUFlags;
@@ -779,12 +832,12 @@ FVideoDecoderImpl::EAUChangeFlags FVideoDecoderImpl::GetAndPrepareInputAU()
 				}
 
 				// If this is a sync frame then we can dump all replay data we have and start from here.
-				if (CurrentAccessUnit->BitstreamInfo.bIsSyncFrame)
+				if ((CurrentAccessUnit->DecAU.Flags & EElectraDecoderFlags::IsSyncSample) != EElectraDecoderFlags::None)
 				{
 					ReplayAccessUnits.Empty();
 				}
 				// If the decoder needs to be replayed when lost we need to hold on to the data.
-				if (bNeedsReplayData && !CurrentAccessUnit->AccessUnit->bIsDummyData && !CurrentAccessUnit->BitstreamInfo.bIsDiscardable)
+				if (bNeedsReplayData && !CurrentAccessUnit->AccessUnit->bIsDummyData && (CurrentAccessUnit->DecAU.Flags & EElectraDecoderFlags::IsDiscardable) == EElectraDecoderFlags::None)
 				{
 					ReplayAccessUnits.Enqueue(CurrentAccessUnit);
 				}
@@ -949,7 +1002,75 @@ IElectraDecoder::EOutputStatus FVideoDecoderImpl::HandleOutput()
 					BufferProperties->Set(RenderOptionKeys::Duration, FVariantValue(MatchingInput->AdjustedDuration));
 
 					// Set properties from the bitstream messages.
-					BitstreamProcessor->SetPropertiesOnOutput(CurrentDecoderOutput, BufferProperties.Get(), MatchingInput->BitstreamInfo);
+					if (DecoderBitstreamProcessor.IsValid())
+					{
+						TMap<FString, FVariant> BSIProperties;
+						DecoderBitstreamProcessor->SetPropertiesOnOutput(BSIProperties, MatchingInput->BSI);
+						if (BSIProperties.Num())
+						{
+							// Colorimetry?
+							TArray<uint8> CommonColorimetry(ElectraDecodersUtil::GetVariantValueUInt8Array(BSIProperties, IElectraDecoderBitstreamProcessorInfo::CommonColorimetry));
+							if (CommonColorimetry.Num() == sizeof(ElectraDecodersUtil::MPEG::FCommonColorimetry))
+							{
+								const ElectraDecodersUtil::MPEG::FCommonColorimetry& Colorimetry(*reinterpret_cast<const ElectraDecodersUtil::MPEG::FCommonColorimetry*>(CommonColorimetry.GetData()));
+								if (!CurrentColorimetry.IsSet())
+								{
+									CurrentColorimetry = Electra::MPEG::FColorimetryHelper();
+								}
+								CurrentColorimetry.GetValue().Update(Colorimetry.colour_primaries, Colorimetry.transfer_characteristics, Colorimetry.matrix_coeffs, Colorimetry.video_full_range_flag, Colorimetry.video_format);
+							}
+
+							// HDR parameters?
+							TArray<uint8> Mdcv(ElectraDecodersUtil::GetVariantValueUInt8Array(BSIProperties, IElectraDecoderBitstreamProcessorInfo::SeiMasteringDisplayColorVolume));
+							if (Mdcv.Num() == sizeof(ElectraDecodersUtil::MPEG::FSEImastering_display_colour_volume))
+							{
+								if (!CurrentHDR.IsSet())
+								{
+									CurrentHDR = Electra::MPEG::FHDRHelper();
+								}
+								CurrentHDR.GetValue().UpdateWith(*reinterpret_cast<const ElectraDecodersUtil::MPEG::FSEImastering_display_colour_volume*>(Mdcv.GetData()));
+							}
+							TArray<uint8> Clli(ElectraDecodersUtil::GetVariantValueUInt8Array(BSIProperties, IElectraDecoderBitstreamProcessorInfo::SeiContentLightLeveInfo));
+							if (Clli.Num() == sizeof(ElectraDecodersUtil::MPEG::FSEIcontent_light_level_info))
+							{
+								if (!CurrentHDR.IsSet())
+								{
+									CurrentHDR = Electra::MPEG::FHDRHelper();
+								}
+								CurrentHDR.GetValue().UpdateWith(*reinterpret_cast<const ElectraDecodersUtil::MPEG::FSEIcontent_light_level_info*>(Clli.GetData()));
+							}
+							TArray<uint8> Altc(ElectraDecodersUtil::GetVariantValueUInt8Array(BSIProperties, IElectraDecoderBitstreamProcessorInfo::SeiAlternateTransferCharacteristics));
+							if (Altc.Num() == sizeof(ElectraDecodersUtil::MPEG::FSEIalternative_transfer_characteristics))
+							{
+								if (!CurrentHDR.IsSet())
+								{
+									CurrentHDR = Electra::MPEG::FHDRHelper();
+								}
+								CurrentHDR.GetValue().UpdateWith(*reinterpret_cast<const ElectraDecodersUtil::MPEG::FSEIalternative_transfer_characteristics*>(Altc.GetData()));
+							}
+
+							// Timecode?
+							TArray<uint8> PicTiming(ElectraDecodersUtil::GetVariantValueUInt8Array(BSIProperties, IElectraDecoderBitstreamProcessorInfo::CommonPictureTiming));
+							if (PicTiming.Num() == sizeof(ElectraDecodersUtil::MPEG::FCommonPictureTiming))
+							{
+								TSharedPtr<MPEG::FVideoDecoderTimecode, ESPMode::ThreadSafe> NewTimecode = MakeShared<MPEG::FVideoDecoderTimecode, ESPMode::ThreadSafe>();
+								NewTimecode->UpdateWith(*reinterpret_cast<const ElectraDecodersUtil::MPEG::FCommonPictureTiming*>(PicTiming.GetData()));
+								BufferProperties->Set(IDecoderOutputOptionNames::Timecode, FVariantValue(NewTimecode));
+							}
+						}
+
+						// Set the colorimetry, if available, on the output properties.
+						if (CurrentColorimetry.IsSet())
+						{
+							CurrentColorimetry.GetValue().UpdateParamDict(*BufferProperties);
+							// Also HDR information (which requires colorimetry!) if available.
+							if (CurrentHDR.IsSet())
+							{
+								CurrentHDR.GetValue().SetHDRType(CurrentDecoderOutput->GetNumberOfBits(), CurrentColorimetry.GetValue());
+								CurrentHDR.GetValue().UpdateParamDict(*BufferProperties);
+							}
+						}
+					}
 
 					if (!FPlatformElectraDecoderResourceManager::SetupRenderBufferFromDecoderOutput(CurrentOutputBuffer, BufferProperties, CurrentDecoderOutput, PlatformResource))
 					{
@@ -984,7 +1105,7 @@ FVideoDecoderImpl::ENextDecodingState FVideoDecoderImpl::HandleDecoding()
 	{
 		// If this AU falls outside the range where it is to be rendered and it is also discardable
 		// we do not need to process it.
-		if (CurrentAccessUnit->BitstreamInfo.bIsDiscardable && CurrentAccessUnit->bMaySkipDecoding)
+		if ((CurrentAccessUnit->DecAU.Flags & EElectraDecoderFlags::IsDiscardable) != EElectraDecoderFlags::None && CurrentAccessUnit->bMaySkipDecoding)
 		{
 			// Even if this access unit won't be decoded, if it is the last in the period and we are
 			// not decoding dummy data the decoder must be drained to get the last decoded data out.
@@ -1005,51 +1126,18 @@ FVideoDecoderImpl::ENextDecodingState FVideoDecoderImpl::HandleDecoding()
 			return bOk ? ENextDecodingState::NormalDecoding : ENextDecodingState::Error;
 		}
 
-		if (!DecoderInstance.IsValid())
-		{
-			if (!InternalDecoderCreate())
-			{
-				return ENextDecodingState::Error;
-			}
-		}
-
 		if (DecoderInstance.IsValid())
 		{
 			SCOPE_CYCLE_COUNTER(STAT_ElectraPlayer_VideoDecode);
 			CSV_SCOPED_TIMING_STAT(ElectraPlayer, VideoDecode);
 
-			IElectraDecoder::FInputAccessUnit DecAU;
-			DecAU.Data = CurrentAccessUnit->AccessUnit->AUData;
-			DecAU.DataSize = (int32) CurrentAccessUnit->AccessUnit->AUSize;
-			DecAU.DTS = CurrentAccessUnit->AccessUnit->DTS.GetAsTimespan();
-			DecAU.PTS = CurrentAccessUnit->AccessUnit->PTS.GetAsTimespan();
-			DecAU.Duration = CurrentAccessUnit->AccessUnit->Duration.GetAsTimespan();
-			DecAU.UserValue = CurrentAccessUnit->PTS;
-			DecAU.Flags = CurrentAccessUnit->BitstreamInfo.bIsSyncFrame ? EElectraDecoderFlags::IsSyncSample : EElectraDecoderFlags::None;
-			if (CurrentAccessUnit->BitstreamInfo.bIsDiscardable)
-			{
-				DecAU.Flags |= EElectraDecoderFlags::IsDiscardable;
-			}
-			if (bSupportsDroppingOutput && !CurrentAccessUnit->AdjustedPTS.IsValid())
-			{
-				NumInitialSkippedFrames += bIsStartOfSequence ? 1 : 0;
-				NumInitialSkippedDecodingFrames += CurrentAccessUnit->BitstreamInfo.bIsDiscardable ? 1 : 0;
-				DecAU.Flags |= EElectraDecoderFlags::DoNotOutput;
-			}
-			TMap<FString, FVariant> CSDOptions;
-			if (CurrentAccessUnit->BitstreamInfo.bIsSyncFrame && CurrentAccessUnit->AccessUnit->AUCodecData.IsValid())
-			{
-				CSDOptions.Emplace(TEXT("csd"), FVariant(CurrentAccessUnit->AccessUnit->AUCodecData->CodecSpecificData));
-				CSDOptions.Emplace(TEXT("dcr"), FVariant(CurrentAccessUnit->AccessUnit->AUCodecData->RawCSD));
-			}
-
 			// Need to wait for a sync sample?
-			if (bWaitForSyncSample && !CurrentAccessUnit->BitstreamInfo.bIsSyncFrame)
+			if (bWaitForSyncSample && (CurrentAccessUnit->DecAU.Flags & EElectraDecoderFlags::IsSyncSample) == EElectraDecoderFlags::None)
 			{
 				if (!bWarnedMissingSyncSample)
 				{
 					bWarnedMissingSyncSample = true;
-					UE_LOG(LogElectraPlayer, Warning, TEXT("Expected a video sync sample at PTS %lld, but did not get one. The stream may be packaged incorrectly. Dropping frames until one arrives, which may take a while. Please wait!"), (long long int)DecAU.PTS.GetTicks());
+					UE_LOG(LogElectraPlayer, Warning, TEXT("Expected a video sync sample at PTS %lld, but did not get one. The stream may be packaged incorrectly. Dropping frames until one arrives, which may take a while. Please wait!"), (long long int)CurrentAccessUnit->DecAU.PTS.GetTicks());
 				}
 				bDrainAfterDecode = CurrentAccessUnit->AccessUnit->bIsLastInPeriod;
 				CurrentAccessUnit.Reset();
@@ -1060,10 +1148,10 @@ FVideoDecoderImpl::ENextDecodingState FVideoDecoderImpl::HandleDecoding()
 				return ENextDecodingState::NormalDecoding;
 			}
 
-			IElectraDecoder::EDecoderError DecErr = DecoderInstance->DecodeAccessUnit(DecAU, CSDOptions);
+			IElectraDecoder::EDecoderError DecErr = DecoderInstance->DecodeAccessUnit(CurrentAccessUnit->DecAU, CurrentAccessUnit->CSDOptions);
 			if (DecErr == IElectraDecoder::EDecoderError::None)
 			{
-				if ((DecAU.Flags & EElectraDecoderFlags::DoNotOutput) == EElectraDecoderFlags::None)
+				if ((CurrentAccessUnit->DecAU.Flags & EElectraDecoderFlags::DoNotOutput) == EElectraDecoderFlags::None)
 				{
 					InDecoderInput.Emplace(CurrentAccessUnit);
 					InDecoderInput.StableSort([](const TSharedPtr<FDecoderInput, ESPMode::ThreadSafe>& a, const TSharedPtr<FDecoderInput, ESPMode::ThreadSafe>& b)
@@ -1141,37 +1229,14 @@ FVideoDecoderImpl::ENextDecodingState FVideoDecoderImpl::HandleReplaying()
 	}
 	bool bIsLastReplayAU = ReplayingAccessUnits.IsEmpty();
 
-	if (!DecoderInstance.IsValid())
-	{
-		if (!InternalDecoderCreate())
-		{
-			return ENextDecodingState::Error;
-		}
-	}
 	if (DecoderInstance.IsValid())
 	{
-		IElectraDecoder::FInputAccessUnit DecAU;
-		DecAU.Data = ReplayAccessUnit->AccessUnit->AUData;
-		DecAU.DataSize = (int32) ReplayAccessUnit->AccessUnit->AUSize;
-		DecAU.DTS = ReplayAccessUnit->AccessUnit->DTS.GetAsTimespan();
-		DecAU.PTS = ReplayAccessUnit->AccessUnit->PTS.GetAsTimespan();
-		DecAU.Duration = ReplayAccessUnit->AccessUnit->Duration.GetAsTimespan();
-		DecAU.UserValue = ReplayAccessUnit->PTS;
-		DecAU.Flags = ReplayAccessUnit->BitstreamInfo.bIsSyncFrame ? EElectraDecoderFlags::IsSyncSample : EElectraDecoderFlags::None;
-		DecAU.Flags |= EElectraDecoderFlags::IsReplaySample;
-		if (bIsLastReplayAU)
-		{
-			DecAU.Flags |= EElectraDecoderFlags::IsLastReplaySample;
-		}
-
-		TMap<FString, FVariant> CSDOptions;
-		if (ReplayAccessUnit->BitstreamInfo.bIsSyncFrame && ReplayAccessUnit->AccessUnit->AUCodecData.IsValid())
-		{
-			CSDOptions.Emplace(TEXT("csd"), FVariant(ReplayAccessUnit->AccessUnit->AUCodecData->CodecSpecificData));
-			CSDOptions.Emplace(TEXT("dcr"), FVariant(ReplayAccessUnit->AccessUnit->AUCodecData->RawCSD));
-		}
+		// Set replay flags for this decode call
+		ReplayAccessUnit->DecAU.Flags |= EElectraDecoderFlags::IsReplaySample | (bIsLastReplayAU ? EElectraDecoderFlags::IsLastReplaySample : EElectraDecoderFlags::None);
+		IElectraDecoder::EDecoderError DecErr = DecoderInstance->DecodeAccessUnit(ReplayAccessUnit->DecAU, ReplayAccessUnit->CSDOptions);
+		// Clear them again.
+		ReplayAccessUnit->DecAU.Flags &= ~(EElectraDecoderFlags::IsReplaySample | EElectraDecoderFlags::IsLastReplaySample);
 
-		IElectraDecoder::EDecoderError DecErr = DecoderInstance->DecodeAccessUnit(DecAU, CSDOptions);
 		if (DecErr == IElectraDecoder::EDecoderError::None)
 		{
 			// The decoder must not deliver output from replays, so we must not keep track of the input.
@@ -1326,6 +1391,9 @@ bool FVideoDecoderImpl::CheckForFlush()
 		CurrentSequenceIndex.Reset();
 		NextExpectedDTSHNS.Reset();
 		CurrentAccessUnit.Reset();
+		CurrentActiveCSD.Empty();
+		CurrentColorimetry.Reset();
+		CurrentHDR.Reset();
 		bIsDecoderClean = true;
 		bInDummyDecodeMode = false;
 		bWaitForSyncSample = true;
@@ -1334,7 +1402,10 @@ bool FVideoDecoderImpl::CheckForFlush()
 		NumInitialSkippedFrames = 0;
 		NumInitialSkippedDecodingFrames = 0;
 		bIsStartOfSequence = true;
-		BitstreamProcessor->Clear();
+		if (DecoderBitstreamProcessor.IsValid())
+		{
+			DecoderBitstreamProcessor->Clear();
+		}
 		FlushDecoderSignal.Reset();
 		DecoderFlushedSignal.Signal();
 		return true;
@@ -1420,10 +1491,6 @@ void FVideoDecoderImpl::WorkerThread()
 		{
 			DecoderFactory->GetConfigurationOptions(DecoderConfigOptions);
 		}
-
-		TMap<FString, FVariant> AllOptions(DecoderFactoryAddtlCfg);
-		AllOptions.Append(DecoderConfigOptions);
-		BitstreamProcessor = IVideoDecoderInputBitstreamProcessor::Create(DecoderFormat, AllOptions);
 	}
 
 	CreateDecoderOutputPool();
@@ -1458,12 +1525,20 @@ void FVideoDecoderImpl::WorkerThread()
 		TimeLast = TimeNow;
 		MinLoopSleepTimeMsec = kDefaultMinLoopSleepTimeMS;
 
-
-		// Get the next access unit to decode.
-		EAUChangeFlags NewAUFlags = GetAndPrepareInputAU();
+		// Create decoder if necessary.
+		if (!DecoderInstance.IsValid())
+		{
+			if (!InternalDecoderCreate())
+			{
+				bError = true;
+			}
+		}
 
 		if (!bError)
 		{
+			// Get the next access unit to decode.
+			EAUChangeFlags NewAUFlags = GetAndPrepareInputAU();
+
 			// Did the codec specific data change?
 			if ((NewAUFlags & EAUChangeFlags::CSDChanged) != EAUChangeFlags::None)
 			{
@@ -1592,7 +1667,6 @@ void FVideoDecoderImpl::WorkerThread()
 	InternalDecoderDestroy();
 	DestroyDecoderOutputPool();
 
-	BitstreamProcessor.Reset();
 	DecoderFactory.Reset();
 	DecoderFactoryAddtlCfg.Empty();
 
@@ -1601,6 +1675,9 @@ void FVideoDecoderImpl::WorkerThread()
 	InDecoderInput.Empty();
 	CurrentSequenceIndex.Reset();
 	NextExpectedDTSHNS.Reset();
+	CurrentActiveCSD.Empty();
+	CurrentColorimetry.Reset();
+	CurrentHDR.Reset();
 	ReplayAccessUnit.Reset();
 	ReplayAccessUnits.Empty();
 	ReplayingAccessUnits.Empty();
diff --git a/Engine/Plugins/Media/ElectraPlayer/Source/ElectraPlayerRuntime/Private/Runtime/Decoder/VideoDecoderInputBitstreamProcessor.cpp b/Engine/Plugins/Media/ElectraPlayer/Source/ElectraPlayerRuntime/Private/Runtime/Decoder/VideoDecoderInputBitstreamProcessor.cpp
deleted file mode 100644
diff --git a/Engine/Plugins/Media/ElectraPlayer/Source/ElectraPlayerRuntime/Private/Runtime/Decoder/VideoDecoderInputBitstreamProcessor.h b/Engine/Plugins/Media/ElectraPlayer/Source/ElectraPlayerRuntime/Private/Runtime/Decoder/VideoDecoderInputBitstreamProcessor.h
deleted file mode 100644
diff --git a/Engine/Plugins/Media/ElectraPlayer/Source/ElectraPlayerRuntime/Public/Decoder/VideoDecoderHelpers.h b/Engine/Plugins/Media/ElectraPlayer/Source/ElectraPlayerRuntime/Public/Decoder/VideoDecoderHelpers.h
--- a/Engine/Plugins/Media/ElectraPlayer/Source/ElectraPlayerRuntime/Public/Decoder/VideoDecoderHelpers.h
+++ b/Engine/Plugins/Media/ElectraPlayer/Source/ElectraPlayerRuntime/Public/Decoder/VideoDecoderHelpers.h
@@ -23,6 +23,26 @@ public:
 	{ return &DecoderTimecode; }
 	void Update(const IVideoDecoderTimecode::FMPEGDefinition& InTimecode)
 	{ DecoderTimecode = InTimecode; }
+	void UpdateWith(const ElectraDecodersUtil::MPEG::FCommonPictureTiming& InTiming)
+	{
+		DecoderTimecode.clockTimestamp = InTiming.clockTimestamp;
+		DecoderTimecode.num_units_in_tick = InTiming.num_units_in_tick;
+		DecoderTimecode.time_scale = InTiming.time_scale;
+		DecoderTimecode.time_offset = InTiming.time_offset;
+		DecoderTimecode.n_frames = InTiming.n_frames;
+		DecoderTimecode.timing_info_present_flag = InTiming.timing_info_present_flag;
+		DecoderTimecode.clock_timestamp_flag = InTiming.clock_timestamp_flag;
+		DecoderTimecode.ct_type = InTiming.ct_type;
+		DecoderTimecode.nuit_field_based_flag = InTiming.nuit_field_based_flag;
+		DecoderTimecode.counting_type = InTiming.counting_type;
+		DecoderTimecode.full_timestamp_flag = InTiming.full_timestamp_flag;
+		DecoderTimecode.discontinuity_flag = InTiming.discontinuity_flag;
+		DecoderTimecode.cnt_dropped_flag = InTiming.cnt_dropped_flag;
+		DecoderTimecode.seconds_value = InTiming.seconds_value;
+		DecoderTimecode.minutes_value = InTiming.minutes_value;
+		DecoderTimecode.hours_value = InTiming.hours_value;
+		DecoderTimecode.FromH26x = InTiming.FromH26x;
+	}
 private:
 	IVideoDecoderTimecode::FMPEGDefinition DecoderTimecode;
 };
