Electra: Extending the period during which a HLS playlist does not need to update from 1.5 to 3 target durations, adding log messages to that effect; added validation that downloaded segment duration is withing what the variant playlist says it is.

#rnx

diff --git a/Engine/Plugins/Media/ElectraPlayer/Source/ElectraPlayerRuntime/Private/Runtime/Player/HLS/PlaylistHLS.cpp b/Engine/Plugins/Media/ElectraPlayer/Source/ElectraPlayerRuntime/Private/Runtime/Player/HLS/PlaylistHLS.cpp
--- a/Engine/Plugins/Media/ElectraPlayer/Source/ElectraPlayerRuntime/Private/Runtime/Player/HLS/PlaylistHLS.cpp
+++ b/Engine/Plugins/Media/ElectraPlayer/Source/ElectraPlayerRuntime/Private/Runtime/Player/HLS/PlaylistHLS.cpp
@@ -5,6 +5,7 @@
 #include "Player/AdaptivePlayerOptionKeynames.h"
 #include "Player/AdaptiveStreamingPlayerResourceRequest.h"
 #include "Player/StreamSegmentReaderCommon.h"
+#include "Utilities/Utilities.h"
 #include "Utilities/TimeUtilities.h"
 #include "Utilities/ISO639-Map.h"
 #include "Crypto/StreamCryptoAES128.h"
@@ -113,6 +114,7 @@ public:
 	void GetActiveMediaPlaylists(TArray<TSharedPtr<FMediaPlaylistAndStateHLS, ESPMode::ThreadSafe>>& OutActivePlaylists);
 private:
 	void SetTimestampAdjustIfNecessary(const TSharedPtrTS<FStreamSegmentRequestCommon>& InSegment);
+	void ValidateDownloadedSegmentDuration(const FStreamSegmentRequestCommon* InRequest);
 
 	struct FSelectedTrackStream
 	{
@@ -3175,9 +3177,28 @@ IManifest::FResult FActiveHLSPlaylist::FPlayPeriod::GetLoopingSegment(TSharedPtr
 	return GetStartingSegment(OutSegment, InSequenceState, InStartPosition, InSearchType);
 }
 
+void FActiveHLSPlaylist::FPlayPeriod::ValidateDownloadedSegmentDuration(const FStreamSegmentRequestCommon* InRequest)
+{
+	// For successful downloads check if the segment duration was as specified.
+	if (InRequest && InRequest->DownloadStats.bWasSuccessful)
+	{
+		// The #EXTINF duration needs to be as precise to the actual media duration as possible.
+		// Otherwise playback errors could occur and switching across renditions may not work correctly.
+		double AbsDelta = Utils::AbsoluteValue(InRequest->DownloadStats.Duration - InRequest->DownloadStats.DurationDownloaded);
+		if (AbsDelta > 1.0)
+		{
+			FURL_RFC3986 MediaURL, PlaylistURL;
+			MediaURL.Parse(InRequest->DownloadStats.URL);
+			PlaylistURL.Parse(InRequest->HLS.Playlist.IsValid() ? InRequest->HLS.Playlist->URL : TEXT(""));
+			PlayerSessionServices->PostLog(Facility::EFacility::HLSPlaylistHandler, IInfoLog::ELevel::Warning, FString::Printf(TEXT("Media segment duration for segment %s in variant playlist %s was given as %.3f seconds but really was %.3f seconds!"), *MediaURL.GetLastPathComponent(), *PlaylistURL.GetLastPathComponent(), InRequest->DownloadStats.Duration, InRequest->DownloadStats.DurationDownloaded));
+		}
+	}
+}
+
 IManifest::FResult FActiveHLSPlaylist::FPlayPeriod::GetNextSegment(TSharedPtrTS<IStreamSegment>& OutSegment, TSharedPtrTS<const IStreamSegment> InCurrentSegment, const FPlayStartOptions& InOptions)
 {
 	FStreamSegmentRequestCommon* CurrentRequest = const_cast<FStreamSegmentRequestCommon*>(static_cast<const FStreamSegmentRequestCommon*>(InCurrentSegment.Get()));
+	ValidateDownloadedSegmentDuration(CurrentRequest);
 	IManifest::FResult Result = GetSegment(OutSegment, CurrentRequest, InOptions, ENextSegType::Next);
 	if (Result.GetRetryAgainAtTime().IsValid())
 	{
@@ -3189,6 +3210,7 @@ IManifest::FResult FActiveHLSPlaylist::FPlayPeriod::GetNextSegment(TSharedPtrTS<
 IManifest::FResult FActiveHLSPlaylist::FPlayPeriod::GetRetrySegment(TSharedPtrTS<IStreamSegment>& OutSegment, TSharedPtrTS<const IStreamSegment> InCurrentSegment, const FPlayStartOptions& InOptions, bool bInReplaceWithFillerData)
 {
 	FStreamSegmentRequestCommon* CurrentRequest = const_cast<FStreamSegmentRequestCommon*>(static_cast<const FStreamSegmentRequestCommon*>(InCurrentSegment.Get()));
+	ValidateDownloadedSegmentDuration(CurrentRequest);
 	// To insert filler data we can use the current request over again.
 	if (bInReplaceWithFillerData)
 	{
@@ -3213,7 +3235,7 @@ IManifest::FResult FActiveHLSPlaylist::FPlayPeriod::GetRetrySegment(TSharedPtrTS
 
 
 
-void FMediaPlaylistAndStateHLS::SetPlaylist(TSharedPtr<FMediaPlaylistHLS, ESPMode::ThreadSafe> InPlaylist, FTimeValue InNow)
+void FMediaPlaylistAndStateHLS::SetPlaylist(IPlayerSessionServices* InPlayerSessionServices, TSharedPtr<FMediaPlaylistHLS, ESPMode::ThreadSafe> InPlaylist, FTimeValue InNow)
 {
 	check(InPlaylist.IsValid());
 
@@ -3283,18 +3305,30 @@ void FMediaPlaylistAndStateHLS::SetPlaylist(TSharedPtr<FMediaPlaylistHLS, ESPMod
 					switch to a different Variant Stream or trigger a playback error.
 			*/
 			++ReloadCount;
-			if (InNow < TimeWhenLoaded + ((Playlist->TargetDuration * 3) / 2))
+			const FTimeValue RequiredUpdateInterval = (Playlist->TargetDuration * 3) / 2;
+			if (InNow < TimeWhenLoaded + RequiredUpdateInterval)
 			{
 				// Try again after half a target duration.
 				TimeAtWhichToReload = InNow + (Playlist->TargetDuration / 2);
 			}
 			else
 			{
-				// Still no change. Stop reloading and mark it as no longer updating.
-				TimeWhenLoaded = InNow;
-				TimeAtWhichToReload.SetToInvalid();
-				ReloadCount = 0;
-				LiveUpdateState = ELiveUpdateState::NotUpdating;
+				// If three target durations have passed we call it quits.
+				if (InNow > TimeWhenLoaded + Playlist->TargetDuration * 3)
+				{
+					InPlayerSessionServices->PostLog(Facility::EFacility::HLSPlaylistBuilder, IInfoLog::ELevel::Warning, FString::Printf(TEXT("HLS Live variant playlist still did not update after %.3f seconds but had to after %.3f seconds, giving up!"), (InNow - TimeWhenLoaded).GetAsSeconds(), RequiredUpdateInterval.GetAsSeconds()));
+					// Stop reloading and mark it as no longer updating.
+					TimeWhenLoaded = InNow;
+					TimeAtWhichToReload.SetToInvalid();
+					ReloadCount = 0;
+					LiveUpdateState = ELiveUpdateState::NotUpdating;
+				}
+				else
+				{
+					InPlayerSessionServices->PostLog(Facility::EFacility::HLSPlaylistBuilder, IInfoLog::ELevel::Warning, FString::Printf(TEXT("HLS Live variant playlist did not update after %.3f seconds, but must every %.3f seconds!"), (InNow - TimeWhenLoaded).GetAsSeconds(), RequiredUpdateInterval.GetAsSeconds()));
+					// Try again after another half a target duration.
+					TimeAtWhichToReload = InNow + (Playlist->TargetDuration / 2);
+				}
 			}
 		}
 	}
diff --git a/Engine/Plugins/Media/ElectraPlayer/Source/ElectraPlayerRuntime/Private/Runtime/Player/HLS/PlaylistHLS.h b/Engine/Plugins/Media/ElectraPlayer/Source/ElectraPlayerRuntime/Private/Runtime/Player/HLS/PlaylistHLS.h
--- a/Engine/Plugins/Media/ElectraPlayer/Source/ElectraPlayerRuntime/Private/Runtime/Player/HLS/PlaylistHLS.h
+++ b/Engine/Plugins/Media/ElectraPlayer/Source/ElectraPlayerRuntime/Private/Runtime/Player/HLS/PlaylistHLS.h
@@ -300,7 +300,7 @@ struct FMediaPlaylistAndStateHLS
 		TSharedPtr<FMediaPlaylistHLS, ESPMode::ThreadSafe> pl(Playlist);
 		return pl;
 	}
-	void SetPlaylist(TSharedPtr<FMediaPlaylistHLS, ESPMode::ThreadSafe> InPlaylist, FTimeValue InNow);
+	void SetPlaylist(IPlayerSessionServices* InPlayerSessionServices, TSharedPtr<FMediaPlaylistHLS, ESPMode::ThreadSafe> InPlaylist, FTimeValue InNow);
 	void ClearPlaylist()
 	{
 		TimeAtWhichToReload.SetToInvalid();
diff --git a/Engine/Plugins/Media/ElectraPlayer/Source/ElectraPlayerRuntime/Private/Runtime/Player/HLS/PlaylistHandlerHLS.cpp b/Engine/Plugins/Media/ElectraPlayer/Source/ElectraPlayerRuntime/Private/Runtime/Player/HLS/PlaylistHandlerHLS.cpp
--- a/Engine/Plugins/Media/ElectraPlayer/Source/ElectraPlayerRuntime/Private/Runtime/Player/HLS/PlaylistHandlerHLS.cpp
+++ b/Engine/Plugins/Media/ElectraPlayer/Source/ElectraPlayerRuntime/Private/Runtime/Player/HLS/PlaylistHandlerHLS.cpp
@@ -33,7 +33,7 @@ public:
 	TSharedPtrTS<IManifest> GetManifest() override;
 
 private:
-	ELECTRA_CLASS_DEFAULT_ERROR_METHODS(PlayerSessionServices, HLSPlaylistReader);
+	ELECTRA_CLASS_DEFAULT_ERROR_METHODS(PlayerSessionServices, HLSPlaylistHandler);
 
 	static bool ValidateNumbersOnly(const FString& In)
 	{
@@ -251,7 +251,7 @@ void FPlaylistHandlerHLS::HandleOnce()
 				TSharedPtr<FMediaPlaylistHLS, ESPMode::ThreadSafe> NewMediaPlaylist;
 				if (BuildMediaPlaylist(NewMediaPlaylist, vp, MultiVariantPlaylist))
 				{
-					MediaPlaylist->SetPlaylist(NewMediaPlaylist, Now);
+					MediaPlaylist->SetPlaylist(PlayerSessionServices, NewMediaPlaylist, Now);
 					if (bGotUpdate)
 					{
 						PlayerSessionServices->SendMessageToPlayer(IPlaylistReader::PlaylistLoadedMessage::Create(LastError, FinishedReq->ResourceRequest->GetConnectionInfo(), Playlist::EListType::Variant, Playlist::ELoadType::Update, FinishedReq->Attempt));
diff --git a/Engine/Plugins/Media/ElectraPlayer/Source/ElectraPlayerRuntime/Private/Runtime/PlayerFacility.h b/Engine/Plugins/Media/ElectraPlayer/Source/ElectraPlayerRuntime/Private/Runtime/PlayerFacility.h
--- a/Engine/Plugins/Media/ElectraPlayer/Source/ElectraPlayerRuntime/Private/Runtime/PlayerFacility.h
+++ b/Engine/Plugins/Media/ElectraPlayer/Source/ElectraPlayerRuntime/Private/Runtime/PlayerFacility.h
@@ -19,11 +19,9 @@ namespace Electra
 			ABR,
 			HTTPReader,
 			BufferedDataReader,
-			HLSPlaylistReader,
+			HLSPlaylistHandler,
 			HLSParser,
 			HLSPlaylistBuilder,
-			HLSManifest,
-			HLSFMP4Reader,
 			MP4Parser,
 			VideoRenderer,
 			AudioRenderer,
@@ -72,16 +70,12 @@ namespace Electra
 					return TEXT("HTTP reader");
 				case EFacility::BufferedDataReader:
 					return TEXT("Buffered data reader");
-				case EFacility::HLSPlaylistReader:
-					return TEXT("HLS playlist reader");
+				case EFacility::HLSPlaylistHandler:
+					return TEXT("HLS playlist handler");
 				case EFacility::HLSParser:
 					return TEXT("HLS parser");
 				case EFacility::HLSPlaylistBuilder:
 					return TEXT("HLS playlist builder");
-				case EFacility::HLSManifest:
-					return TEXT("HLS manifest");
-				case EFacility::HLSFMP4Reader:
-					return TEXT("HLS fmp4 reader");
 				case EFacility::MP4Parser:
 					return TEXT("MP4 parser");
 				case EFacility::VideoRenderer:
