[Electra] Changing verbosity from "Log" to "Verbose" for the seek related logs. This is spamming the console when scrubbing media tracks normally with sequencer.

diff --git a/Engine/Plugins/Media/ElectraPlayer/Source/ElectraPlayerPlugin/Private/ElectraPlayerPlugin.cpp b/Engine/Plugins/Media/ElectraPlayer/Source/ElectraPlayerPlugin/Private/ElectraPlayerPlugin.cpp
--- a/Engine/Plugins/Media/ElectraPlayer/Source/ElectraPlayerPlugin/Private/ElectraPlayerPlugin.cpp
+++ b/Engine/Plugins/Media/ElectraPlayer/Source/ElectraPlayerPlugin/Private/ElectraPlayerPlugin.cpp
@@ -943,7 +943,7 @@ bool FElectraPlayerPlugin::SetRate(float Rate)
 
 bool FElectraPlayerPlugin::Seek(const FTimespan& Time)
 {
-	UE_LOG(LogElectraPlayerPlugin, Log, TEXT("[%u] IMediaControls::Seek() to %s"), PlayerUniqueID, *Time.ToString(TEXT("%h:%m:%s.%f")));
+	UE_LOG(LogElectraPlayerPlugin, Verbose, TEXT("[%u] IMediaControls::Seek() to %s"), PlayerUniqueID, *Time.ToString(TEXT("%h:%m:%s.%f")));
 	CSV_EVENT(ElectraPlayer, TEXT("Seeking"));
 	return Player->Seek(Time);
 }
diff --git a/Engine/Plugins/Media/ElectraPlayer/Source/ElectraPlayerRuntime/Private/ElectraPlayer.cpp b/Engine/Plugins/Media/ElectraPlayer/Source/ElectraPlayerRuntime/Private/ElectraPlayer.cpp
--- a/Engine/Plugins/Media/ElectraPlayer/Source/ElectraPlayerRuntime/Private/ElectraPlayer.cpp
+++ b/Engine/Plugins/Media/ElectraPlayer/Source/ElectraPlayerRuntime/Private/ElectraPlayer.cpp
@@ -2835,7 +2835,7 @@ void FElectraPlayer::HandlePlayerEventJumpInPlayPosition(const Electra::FTimeVal
 		{
 			++Statistics.NumTimesRewound;
 		}
-		UE_LOG(LogElectraPlayer, Log, TEXT("[%u] Jump in play position from %.3f to %.3f"), InstanceID, FromTime.GetAsSeconds(), ToNewTime.GetAsSeconds());
+		UE_LOG(LogElectraPlayer, Verbose, TEXT("[%u] Jump in play position from %.3f to %.3f"), InstanceID, FromTime.GetAsSeconds(), ToNewTime.GetAsSeconds());
 	}
 	else if (TimejumpReason == Electra::Metrics::ETimeJumpReason::Looping)
 	{
diff --git a/Engine/Plugins/Media/ElectraPlayer/Source/ElectraProtron/Private/Player/ElectraProtronPlayerDecoder.cpp b/Engine/Plugins/Media/ElectraPlayer/Source/ElectraProtron/Private/Player/ElectraProtronPlayerDecoder.cpp
--- a/Engine/Plugins/Media/ElectraPlayer/Source/ElectraProtron/Private/Player/ElectraProtronPlayerDecoder.cpp
+++ b/Engine/Plugins/Media/ElectraPlayer/Source/ElectraProtron/Private/Player/ElectraProtronPlayerDecoder.cpp
@@ -693,7 +693,7 @@ void FElectraProtronPlayer::FImpl::FDecoderThread::HandleOutputFrame()
 				{
 					AdjustToTime = X;
 					SeekTimeToHandleTo.Reset();
-					if (DecoderTypeIndex == CodecTypeIndex(ElectraProtronUtils::FCodecInfo::EType::Video) && (SeekTimeNumFramesDecoded + SeekTimeNumFramesSkipped) > 0)
+					if (DecoderTypeIndex == CodecTypeIndex(ElectraProtronUtils::FCodecInfo::EType::Video) && (SeekTimeNumFramesDecoded + SeekTimeNumFramesSkipped) > 1)
 					{
 						UE_LOG(LogElectraProtron, Log, TEXT("Processed %d frames to reach seek point, of which %d could be skipped"), SeekTimeNumFramesDecoded + SeekTimeNumFramesSkipped, SeekTimeNumFramesSkipped);
 					}
