ElectraProtron: Considering a small timestamp threshold to avoid rounding issues between FTimespan and the internally used timescale of the mp4 track

#rnx

diff --git a/Engine/Plugins/Media/ElectraPlayer/Source/ElectraProtron/Private/Player/ElectraProtronPlayerDecoder.cpp b/Engine/Plugins/Media/ElectraPlayer/Source/ElectraProtron/Private/Player/ElectraProtronPlayerDecoder.cpp
--- a/Engine/Plugins/Media/ElectraPlayer/Source/ElectraProtron/Private/Player/ElectraProtronPlayerDecoder.cpp
+++ b/Engine/Plugins/Media/ElectraPlayer/Source/ElectraProtron/Private/Player/ElectraProtronPlayerDecoder.cpp
@@ -119,7 +119,10 @@ uint32 FElectraProtronPlayer::FImpl::FDecoderThread::Run()
 				DestroyDecoder();
 				TrackSampleBuffer = MoveTemp(NewTrackSampleBuffer);
 				GetSampleDelegate = MoveTemp(NewGetSampleDelegate);
-				UpdateTrackIterator();
+				TimeLock.Lock();
+				FTimespan Time = CurrentTime;
+				TimeLock.Unlock();
+				UpdateTrackIterator(Time);
 			}
 			else
 			{
@@ -172,7 +175,7 @@ uint32 FElectraProtronPlayer::FImpl::FDecoderThread::Run()
 		}
 
 		// Are we to do something?
-		if (IsPausedForSeek() || (IsPaused() && !SeekTimeToHandleTo.IsSet() && !PendingSeek.bIsSet))
+		if (!TrackSampleBuffer.IsValid() || IsPausedForSeek() || (IsPaused() && !SeekTimeToHandleTo.IsSet() && !PendingSeek.bIsSet))
 		{
 			continue;
 		}
@@ -193,10 +196,10 @@ uint32 FElectraProtronPlayer::FImpl::FDecoderThread::Run()
 			LoopIndex = NewLoopIndex;
 			TimeLock.Unlock();
 
-			UpdateTrackIterator();
+			UpdateTrackIterator(NewTime);
 			// Handle decoding to this time now.
-			SeekTimeToHandleTo = CurrentTime;
-			SeekTimeToDecodeTo = CurrentTime;
+			SeekTimeToHandleTo = NewTime;
+			SeekTimeToDecodeTo = NewTime;
 			SeekTimeNumFramesDecoded = 0;
 			SeekTimeNumFramesSkipped = 0;
 		}
@@ -319,18 +322,15 @@ void FElectraProtronPlayer::FImpl::FDecoderThread::Flush(const TSharedPtr<FMedia
 	bFlushPending = true;
 }
 
-void FElectraProtronPlayer::FImpl::FDecoderThread::UpdateTrackIterator()
+void FElectraProtronPlayer::FImpl::FDecoderThread::UpdateTrackIterator(const FTimespan& InForTime)
 {
 	if (!TrackSampleBuffer.IsValid())
 	{
 		return;
 	}
 
-	TimeLock.Lock();
-	FTimespan Time = CurrentTime;
-	TimeLock.Unlock();
 	TSharedPtr<Electra::UtilitiesMP4::FMP4Track::FIterator> TkIt =
-		TrackSampleBuffer->TrackAndCodecInfo->MP4Track->CreateIteratorAtKeyframe(Electra::FTimeValue().SetFromTimespan(Time),
+		TrackSampleBuffer->TrackAndCodecInfo->MP4Track->CreateIteratorAtKeyframe(Electra::FTimeValue().SetFromTimespan(InForTime),
 			Electra::FTimeValue().SetFromMilliseconds(TrackSampleBuffer->TrackAndCodecInfo->bIsKeyframeOnlyFormat ? 0 : Config.NextKeyframeThresholdMillis));
 	if (!TkIt.IsValid())
 	{
@@ -479,8 +479,9 @@ void FElectraProtronPlayer::FImpl::FDecoderThread::DecodeOneFrame()
 				// If handling to a specific time see if this sample is discardable and does not require decoding.
 				if (SeekTimeToDecodeTo.IsSet())
 				{
-					FTimespan Start = CurrentInputSample->Sample->EffectivePTS;
-					FTimespan End = CurrentInputSample->Sample->EffectivePTS + CurrentInputSample->Sample->Duration;
+					const FTimespan& kThreshold = DecoderTypeIndex == CodecTypeIndex(ElectraProtronUtils::FCodecInfo::EType::Video) ? Config.TimeAccuracyThresholdVideo : Config.TimeAccuracyThresholdAudio;
+					FTimespan Start = CurrentInputSample->Sample->EffectivePTS - kThreshold;
+					FTimespan End = CurrentInputSample->Sample->EffectivePTS + CurrentInputSample->Sample->Duration + kThreshold;
 					FTimespan X = SeekTimeToDecodeTo.GetValue();
 					if (X >= Start && X < End)
 					{
@@ -681,13 +682,13 @@ void FElectraProtronPlayer::FImpl::FDecoderThread::HandleOutputFrame()
 		TOptional<FTimespan> AdjustToTime;
 		bool bSendToOutput = true;
 
-		auto AdjustForSeek = [&]() -> void
+		auto AdjustForSeek = [&](const FTimespan& InThreshold) -> void
 		{
 			// Is this the frame we were supposed to decode up to?
 			if (SeekTimeToHandleTo.IsSet())
 			{
-				FTimespan Start = InputForCurrentDecoderOutput->Sample->EffectivePTS;
-				FTimespan End = InputForCurrentDecoderOutput->Sample->EffectivePTS + InputForCurrentDecoderOutput->Sample->Duration;
+				FTimespan Start = InputForCurrentDecoderOutput->Sample->EffectivePTS - InThreshold;
+				FTimespan End = InputForCurrentDecoderOutput->Sample->EffectivePTS + InputForCurrentDecoderOutput->Sample->Duration + InThreshold;
 				FTimespan X = SeekTimeToHandleTo.GetValue();
 				if (X >= Start && X < End)
 				{
@@ -708,7 +709,7 @@ void FElectraProtronPlayer::FImpl::FDecoderThread::HandleOutputFrame()
 
 		if (DecoderTypeIndex == CodecTypeIndex(ElectraProtronUtils::FCodecInfo::EType::Video))
 		{
-			AdjustForSeek();
+			AdjustForSeek(Config.TimeAccuracyThresholdVideo);
 
 			// When sending the output we need to wait until there is room to receive it.
 			// If the output is not to be sent we can handle it quickly without waiting.
@@ -844,7 +845,7 @@ void FElectraProtronPlayer::FImpl::FDecoderThread::HandleOutputFrame()
 		}
 		else if (DecoderTypeIndex == CodecTypeIndex(ElectraProtronUtils::FCodecInfo::EType::Audio))
 		{
-			AdjustForSeek();
+			AdjustForSeek(Config.TimeAccuracyThresholdAudio);
 
 			// When sending the output we need to wait until there is room to receive it.
 			// If the output is not to be sent we can handle it quickly without waiting.
diff --git a/Engine/Plugins/Media/ElectraPlayer/Source/ElectraProtron/Private/Player/ElectraProtronPlayerImpl.h b/Engine/Plugins/Media/ElectraPlayer/Source/ElectraProtron/Private/Player/ElectraProtronPlayerImpl.h
--- a/Engine/Plugins/Media/ElectraPlayer/Source/ElectraProtron/Private/Player/ElectraProtronPlayerImpl.h
+++ b/Engine/Plugins/Media/ElectraPlayer/Source/ElectraProtron/Private/Player/ElectraProtronPlayerImpl.h
@@ -90,6 +90,8 @@ private:
 	{
 		FTimespan DurationCacheAhead { ETimespan::TicksPerSecond * 2 };
 		FTimespan DurationCacheBehind { ETimespan::TicksPerSecond * 1 };
+		FTimespan TimeAccuracyThresholdVideo { ETimespan::TicksPerMillisecond * 2 };
+		FTimespan TimeAccuracyThresholdAudio { ETimespan::TicksPerMillisecond * 2 };
 		int32 NextKeyframeThresholdMillis { 2 };
 		bool bReadFirstTimecode = true;
 		bool bReadSampleTimecode = true;
@@ -384,7 +386,7 @@ private:
 		void SetSampleQueue(const TSharedPtr<FMediaSamples, ESPMode::ThreadSafe>& InNewSampleQueue);
 	private:
 
-		void UpdateTrackIterator();
+		void UpdateTrackIterator(const FTimespan& InForTime);
 		void StepTrackIterator(bool bPerformFlush);
 		bool CreateDecoder();
 		void DestroyDecoder();
