Protron: Fixed media sample duration inaccuracies during timescale conversions

#rnx

diff --git a/Engine/Plugins/Media/ElectraPlayer/Source/ElectraProtron/Private/Player/ElectraProtronPlayerDecoder.cpp b/Engine/Plugins/Media/ElectraPlayer/Source/ElectraProtron/Private/Player/ElectraProtronPlayerDecoder.cpp
--- a/Engine/Plugins/Media/ElectraPlayer/Source/ElectraProtron/Private/Player/ElectraProtronPlayerDecoder.cpp
+++ b/Engine/Plugins/Media/ElectraPlayer/Source/ElectraProtron/Private/Player/ElectraProtronPlayerDecoder.cpp
@@ -479,9 +479,8 @@ void FElectraProtronPlayer::FImpl::FDecoderThread::DecodeOneFrame()
 				// If handling to a specific time see if this sample is discardable and does not require decoding.
 				if (SeekTimeToDecodeTo.IsSet())
 				{
-					const FTimespan& kThreshold = DecoderTypeIndex == CodecTypeIndex(ElectraProtronUtils::FCodecInfo::EType::Video) ? Config.TimeAccuracyThresholdVideo : Config.TimeAccuracyThresholdAudio;
-					FTimespan Start = CurrentInputSample->Sample->EffectivePTS - kThreshold;
-					FTimespan End = CurrentInputSample->Sample->EffectivePTS + CurrentInputSample->Sample->Duration + kThreshold;
+					FTimespan Start = CurrentInputSample->Sample->EffectivePTS;
+					FTimespan End = CurrentInputSample->Sample->EffectivePTS + CurrentInputSample->Sample->Duration;
 					FTimespan X = SeekTimeToDecodeTo.GetValue();
 					if (X >= Start && X < End)
 					{
@@ -682,13 +681,13 @@ void FElectraProtronPlayer::FImpl::FDecoderThread::HandleOutputFrame()
 		TOptional<FTimespan> AdjustToTime;
 		bool bSendToOutput = true;
 
-		auto AdjustForSeek = [&](const FTimespan& InThreshold) -> void
+		auto AdjustForSeek = [&]() -> void
 		{
 			// Is this the frame we were supposed to decode up to?
 			if (SeekTimeToHandleTo.IsSet())
 			{
-				FTimespan Start = InputForCurrentDecoderOutput->Sample->EffectivePTS - InThreshold;
-				FTimespan End = InputForCurrentDecoderOutput->Sample->EffectivePTS + InputForCurrentDecoderOutput->Sample->Duration + InThreshold;
+				FTimespan Start = InputForCurrentDecoderOutput->Sample->EffectivePTS;
+				FTimespan End = InputForCurrentDecoderOutput->Sample->EffectivePTS + InputForCurrentDecoderOutput->Sample->Duration;
 				FTimespan X = SeekTimeToHandleTo.GetValue();
 				if (X >= Start && X < End)
 				{
@@ -709,7 +708,7 @@ void FElectraProtronPlayer::FImpl::FDecoderThread::HandleOutputFrame()
 
 		if (DecoderTypeIndex == CodecTypeIndex(ElectraProtronUtils::FCodecInfo::EType::Video))
 		{
-			AdjustForSeek(Config.TimeAccuracyThresholdVideo);
+			AdjustForSeek();
 
 			// When sending the output we need to wait until there is room to receive it.
 			// If the output is not to be sent we can handle it quickly without waiting.
@@ -845,7 +844,7 @@ void FElectraProtronPlayer::FImpl::FDecoderThread::HandleOutputFrame()
 		}
 		else if (DecoderTypeIndex == CodecTypeIndex(ElectraProtronUtils::FCodecInfo::EType::Audio))
 		{
-			AdjustForSeek(Config.TimeAccuracyThresholdAudio);
+			AdjustForSeek();
 
 			// When sending the output we need to wait until there is room to receive it.
 			// If the output is not to be sent we can handle it quickly without waiting.
diff --git a/Engine/Plugins/Media/ElectraPlayer/Source/ElectraProtron/Private/Player/ElectraProtronPlayerImpl.h b/Engine/Plugins/Media/ElectraPlayer/Source/ElectraProtron/Private/Player/ElectraProtronPlayerImpl.h
--- a/Engine/Plugins/Media/ElectraPlayer/Source/ElectraProtron/Private/Player/ElectraProtronPlayerImpl.h
+++ b/Engine/Plugins/Media/ElectraPlayer/Source/ElectraProtron/Private/Player/ElectraProtronPlayerImpl.h
@@ -90,8 +90,6 @@ private:
 	{
 		FTimespan DurationCacheAhead { ETimespan::TicksPerSecond * 2 };
 		FTimespan DurationCacheBehind { ETimespan::TicksPerSecond * 1 };
-		FTimespan TimeAccuracyThresholdVideo { ETimespan::TicksPerMillisecond * 2 };
-		FTimespan TimeAccuracyThresholdAudio { ETimespan::TicksPerMillisecond * 2 };
 		int32 NextKeyframeThresholdMillis { 2 };
 		bool bReadFirstTimecode = true;
 		bool bReadSampleTimecode = true;
diff --git a/Engine/Plugins/Media/ElectraPlayer/Source/ElectraProtron/Private/Player/ElectraProtronPlayerLoader.cpp b/Engine/Plugins/Media/ElectraPlayer/Source/ElectraProtron/Private/Player/ElectraProtronPlayerLoader.cpp
--- a/Engine/Plugins/Media/ElectraPlayer/Source/ElectraProtron/Private/Player/ElectraProtronPlayerLoader.cpp
+++ b/Engine/Plugins/Media/ElectraPlayer/Source/ElectraProtron/Private/Player/ElectraProtronPlayerLoader.cpp
@@ -485,7 +485,7 @@ TSharedPtr<FElectraProtronPlayer::FImpl::FMP4Sample, ESPMode::ThreadSafe> FElect
 	Sample->PTS = InSampleIt->GetPTS().GetAsTimespan();
 	Sample->EffectiveDTS = InSampleIt->GetEffectiveDTS().GetAsTimespan();
 	Sample->EffectivePTS = InSampleIt->GetEffectivePTS().GetAsTimespan();
-	Sample->Duration = InSampleIt->GetDuration().GetAsTimespan();
+	Sample->Duration = InSampleIt->GetDurationAsTimespan();
 	Sample->SizeInBytes = InSampleIt->GetSampleSize();
 	Sample->OffsetInFile = InSampleIt->GetSampleFileOffset();
 	Sample->TrackID = InSampleIt->GetTrackID();
diff --git a/Engine/Plugins/Media/ElectraUtil/Source/ElectraBase/Private/Utilities/MP4Boxes/MP4Track.cpp b/Engine/Plugins/Media/ElectraUtil/Source/ElectraBase/Private/Utilities/MP4Boxes/MP4Track.cpp
--- a/Engine/Plugins/Media/ElectraUtil/Source/ElectraBase/Private/Utilities/MP4Boxes/MP4Track.cpp
+++ b/Engine/Plugins/Media/ElectraUtil/Source/ElectraBase/Private/Utilities/MP4Boxes/MP4Track.cpp
@@ -504,7 +504,12 @@ namespace Electra
 			CurrentPTS.SetFromND(PTS, Timescale);
 			CurrentEffectiveDTS.SetFromND(DTS - Convs.DTSShiftAtZeroPoint, Timescale);
 			CurrentEffectivePTS.SetFromND(PTS - Convs.CompositionTimeAtZeroPoint, Timescale);
+			// Set the duration as the fraction of the duration and the timescale.
 			CurrentDuration.SetFromND(Duration, Timescale);
+			// Also set the duration as the delta of the DTS of this sample and the next
+			// in timespan units. This is to avoid transformation issues from media local
+			// time into the timescale used in engine.
+			CurrentDurationTS = FTimeFraction(DTS + Duration, Timescale).GetAsTimespan() - CurrentDTS.GetAsTimespan();
 			CurrentSampleSize = stszIt.GetCurrentSampleSize();
 			bCurrentIsSyncOrRAP = stssIt.IsSyncSample() || rapIt.GetCurrentGroupDescriptionIndex() != 0;
 
diff --git a/Engine/Plugins/Media/ElectraUtil/Source/ElectraBase/Public/Utilities/MP4Boxes/MP4Track.h b/Engine/Plugins/Media/ElectraUtil/Source/ElectraBase/Public/Utilities/MP4Boxes/MP4Track.h
--- a/Engine/Plugins/Media/ElectraUtil/Source/ElectraBase/Public/Utilities/MP4Boxes/MP4Track.h
+++ b/Engine/Plugins/Media/ElectraUtil/Source/ElectraBase/Public/Utilities/MP4Boxes/MP4Track.h
@@ -82,6 +82,9 @@ namespace Electra
 				{ return CurrentEffectivePTS; }
 				FTimeFraction GetDuration() const
 				{ return CurrentDuration; }
+				// Returns the duration as an FTimespan, which may be slightly more accurate than as a fraction.
+				FTimespan GetDurationAsTimespan() const
+				{ return CurrentDurationTS; }
 				bool IsSyncOrRAPSample() const
 				{ return bCurrentIsSyncOrRAP; }
 				int64 GetSampleSize() const
@@ -156,6 +159,7 @@ namespace Electra
 				FTimeFraction CurrentEffectiveDTS;
 				FTimeFraction CurrentEffectivePTS;
 				FTimeFraction CurrentDuration;
+				FTimespan CurrentDurationTS;
 				int64 CurrentSampleFileOffset = 0;
 				int64 CurrentSampleSize = 0;
 				bool bCurrentIsSyncOrRAP = false;
