Protron: Corrected sequence index adjustment in Protron on direction changes.

#rnx

diff --git a/Engine/Plugins/Media/ElectraPlayer/Source/ElectraProtron/Private/Player/ElectraProtronPlayer.cpp b/Engine/Plugins/Media/ElectraPlayer/Source/ElectraProtron/Private/Player/ElectraProtronPlayer.cpp
--- a/Engine/Plugins/Media/ElectraPlayer/Source/ElectraProtron/Private/Player/ElectraProtronPlayer.cpp
+++ b/Engine/Plugins/Media/ElectraPlayer/Source/ElectraProtron/Private/Player/ElectraProtronPlayer.cpp
@@ -358,8 +358,14 @@ bool FElectraProtronPlayer::Seek(const FTimespan& InTime)
 	SampleQueueLock.Unlock();
 
 	TSharedPtr<FImpl, ESPMode::ThreadSafe> Player(GetCurrentPlayer());
-	++CurrentSequenceIndex;
-	bool bOk = Player.IsValid() ? Player->Seek(InTime, CurrentSampleQueue, CurrentSequenceIndex) : false;
+	if (!Player.IsValid())
+	{
+		return false;
+	}
+	const int32 SeqChange = Player->GetRate() >= 0.0f ? 1 : -1;
+	int32 SequenceIndexNow = CurrentSequenceIndex;
+	CurrentSequenceIndex += SeqChange;
+	bool bOk = Player->Seek(InTime, CurrentSampleQueue, CurrentSequenceIndex);
 	if (bOk)
 	{
 		// Send a seek complete event even if that is not really true and the seek is in progress or
@@ -370,7 +376,7 @@ bool FElectraProtronPlayer::Seek(const FTimespan& InTime)
 	}
 	else
 	{
-		--CurrentSequenceIndex;
+		CurrentSequenceIndex = SequenceIndexNow;
 	}
 	return bOk;
 }
