ElectraPlayer: Not emitting accumulated AUs on ABR abort without early emitting enabled; segment index download errors now just disable the representation; adjusted time on segment retry by a margin to work around timescale conversion precision issues

#rnx

diff --git a/Engine/Plugins/Media/ElectraPlayer/Source/ElectraPlayerRuntime/Private/Runtime/Player/ABRRules/ABRVoDPlus.cpp b/Engine/Plugins/Media/ElectraPlayer/Source/ElectraPlayerRuntime/Private/Runtime/Player/ABRRules/ABRVoDPlus.cpp
--- a/Engine/Plugins/Media/ElectraPlayer/Source/ElectraPlayerRuntime/Private/Runtime/Player/ABRRules/ABRVoDPlus.cpp
+++ b/Engine/Plugins/Media/ElectraPlayer/Source/ElectraPlayerRuntime/Private/Runtime/Player/ABRRules/ABRVoDPlus.cpp
@@ -653,6 +653,10 @@ IAdaptiveStreamSelector::ESegmentAction FABROnDemandPlus::EvaluateForError(TArra
 			{
 				if (Stats.bWaitingForRemoteRetryElement)
 				{
+					// Switching to another representation that needs to be made ready resets the consecutive errors so we get
+					// a chance to switch to a different representation if several are being blocked by the CDN to force
+					// multiple downswitches.
+					WorkVars->NumConsecutiveSegmentErrors = 0;
 					return IAdaptiveStreamSelector::ESegmentAction::Retry;
 				}
 
diff --git a/Engine/Plugins/Media/ElectraPlayer/Source/ElectraPlayerRuntime/Private/Runtime/Player/AdaptiveStreamingPlayerABR.cpp b/Engine/Plugins/Media/ElectraPlayer/Source/ElectraPlayerRuntime/Private/Runtime/Player/AdaptiveStreamingPlayerABR.cpp
--- a/Engine/Plugins/Media/ElectraPlayer/Source/ElectraPlayerRuntime/Private/Runtime/Player/AdaptiveStreamingPlayerABR.cpp
+++ b/Engine/Plugins/Media/ElectraPlayer/Source/ElectraPlayerRuntime/Private/Runtime/Player/AdaptiveStreamingPlayerABR.cpp
@@ -650,6 +650,34 @@ namespace Electra
 			return NextSegmentAction;
 		}
 
+		// Check if the state of a stream has become unusable.
+		TSharedPtrTS<ITimelineMediaAsset> Asset = CurrentPlayPeriod->GetMediaAsset();
+		if (Asset.IsValid())
+		{
+			FString AdaptationSetId = CurrentPlayPeriod->GetSelectedAdaptationSetID(StreamType);
+			if (!AdaptationSetId.IsEmpty())
+			{
+				for(int32 nA=0, maxA=Asset->GetNumberOfAdaptationSets(StreamType); nA<maxA; ++nA)
+				{
+					TSharedPtrTS<IPlaybackAssetAdaptationSet> Adapt = Asset->GetAdaptationSetByTypeAndIndex(StreamType, nA);
+					if (!AdaptationSetId.Equals(Adapt->GetUniqueIdentifier()))
+					{
+						continue;
+					}
+					for(int32 i=0; i<CandidateRepresentations.Num(); ++i)
+					{
+						TSharedPtrTS<IPlaybackAssetRepresentation> Repr = Adapt->GetRepresentationByUniqueIdentifier(CandidateRepresentations[i]->RepresentationUniqueID);
+						if (Repr.IsValid() && !Repr->CanBePlayed())
+						{
+							CandidateRepresentations.RemoveAt(i);
+							--i;
+						}
+					}
+					break;
+				}
+			}
+		}
+
 		// Evaluate which stream qualities could potentially be used and select one.
 		NextSegmentAction = ABRMethod->EvaluateForQuality(CandidateRepresentations, StreamType, OutDelay, CurrentPlayPeriod, CurrentSegment, TimeNow);
 		if (NextSegmentAction != IAdaptiveStreamSelector::ESegmentAction::FetchNext)
diff --git a/Engine/Plugins/Media/ElectraPlayer/Source/ElectraPlayerRuntime/Private/Runtime/Player/DASH/ManifestDASH.cpp b/Engine/Plugins/Media/ElectraPlayer/Source/ElectraPlayerRuntime/Private/Runtime/Player/DASH/ManifestDASH.cpp
--- a/Engine/Plugins/Media/ElectraPlayer/Source/ElectraPlayerRuntime/Private/Runtime/Player/DASH/ManifestDASH.cpp
+++ b/Engine/Plugins/Media/ElectraPlayer/Source/ElectraPlayerRuntime/Private/Runtime/Player/DASH/ManifestDASH.cpp
@@ -1617,8 +1617,9 @@ IManifest::FResult FDASHPlayPeriod::GetNextOrRetrySegment(TSharedPtrTS<IStreamSe
 	}
 	else if (InNextType == ENextSegType::SamePeriodRetry)
 	{
-		// Use the same period local time for the retry representation as was used to locate the current segment.
-		SearchOpt.PeriodLocalTime.SetFromND(Current->Segment.Time - Current->Segment.PTO, Current->Segment.Timescale);
+		// Set up the search time as was used to locate the current segment with a tenth of the duration into it.
+		// This is to make sure the time is sufficiently large that it won't be affected by rounding errors in timescale conversions.
+		SearchOpt.PeriodLocalTime.SetFromND(Current->Segment.Time - Current->Segment.PTO + Current->Segment.Duration*1/10, Current->Segment.Timescale);
 		SearchOpt.SearchType = bFrameAccurateSearch ? IManifest::ESearchType::Before : IManifest::ESearchType::Closest;
 	}
 	else if (InNextType == ENextSegType::SamePeriodStartOver)
@@ -2791,6 +2792,10 @@ FManifestDASHInternal::FRepresentation::ESearchResult FManifestDASHInternal::FRe
 
 FManifestDASHInternal::FRepresentation::ESearchResult FManifestDASHInternal::FRepresentation::FindSegment_Base_MP4(IPlayerSessionServices* InPlayerSessionServices, FSegmentInformationCommon& OutSegmentInfo, TArray<TWeakPtrTS<FMPDLoadRequestDASH>>& OutRemoteElementLoadRequests, const FSegmentSearchOption& InSearchOptions, const TSharedPtrTS<FDashMPD_RepresentationType>& MPDRepresentation, const TArray<TSharedPtrTS<FDashMPD_SegmentBaseType>>& SegmentBase)
 {
+	if (!bIsUsable)
+	{
+		return FManifestDASHInternal::FRepresentation::ESearchResult::BadType;
+	}
 	ESearchResult SegIndexResult = PrepareSegmentIndex(InPlayerSessionServices, SegmentBase, OutRemoteElementLoadRequests, InSearchOptions);
 	if (SegIndexResult != ESearchResult::Found)
 	{
diff --git a/Engine/Plugins/Media/ElectraPlayer/Source/ElectraPlayerRuntime/Private/Runtime/Player/StreamSegmentReaderCommon.cpp b/Engine/Plugins/Media/ElectraPlayer/Source/ElectraPlayerRuntime/Private/Runtime/Player/StreamSegmentReaderCommon.cpp
--- a/Engine/Plugins/Media/ElectraPlayer/Source/ElectraPlayerRuntime/Private/Runtime/Player/StreamSegmentReaderCommon.cpp
+++ b/Engine/Plugins/Media/ElectraPlayer/Source/ElectraPlayerRuntime/Private/Runtime/Player/StreamSegmentReaderCommon.cpp
@@ -1390,29 +1390,38 @@ void FStreamSegmentReaderCommon::FStreamHandler::HandleCommonMediaEnd(EHandleRes
 	// stalled on feeding the other buffer.
 	FTimeValue DurationDownloaded(FTimeValue::GetZero());
 	FTimeValue DurationDelivered(FTimeValue::GetZero());
-	while(1)
+	if (bAbortedByABR && !bAllowEarlyEmitting)
 	{
-		bool bSentSomething = false;
-		for(auto& ItActiveTrackData : TrackDataMap)
+		// When asked to abort and we were not allowed to send any AUs yet then we must also
+		// not send what we have accumulated now since the segment will be retried on another
+		// quality level.
+	}
+	else
+	{
+		while(1)
 		{
-			CurrentlyActiveTrackData = ItActiveTrackData.Value;
-			CurrentlyActiveTrackData->bGotAllSamples = true;
-			EEmitResult er = EmitSamples(EEmitType::UntilBlocked);
-			if (CurrentlyActiveTrackData->DurationSuccessfullyRead > DurationDownloaded)
+			bool bSentSomething = false;
+			for(auto& ItActiveTrackData : TrackDataMap)
 			{
-				DurationDownloaded = CurrentlyActiveTrackData->DurationSuccessfullyRead;
+				CurrentlyActiveTrackData = ItActiveTrackData.Value;
+				CurrentlyActiveTrackData->bGotAllSamples = true;
+				EEmitResult er = EmitSamples(EEmitType::UntilBlocked);
+				if (CurrentlyActiveTrackData->DurationSuccessfullyRead > DurationDownloaded)
+				{
+					DurationDownloaded = CurrentlyActiveTrackData->DurationSuccessfullyRead;
+				}
+				if (CurrentlyActiveTrackData->DurationSuccessfullyDelivered > DurationDelivered)
+				{
+					DurationDelivered = CurrentlyActiveTrackData->DurationSuccessfullyDelivered;
+				}
+				bSentSomething |= er != EEmitResult::SentEverything;
 			}
-			if (CurrentlyActiveTrackData->DurationSuccessfullyDelivered > DurationDelivered)
+			if (!bSentSomething)
 			{
-				DurationDelivered = CurrentlyActiveTrackData->DurationSuccessfullyDelivered;
+				break;
 			}
-			bSentSomething |= er != EEmitResult::SentEverything;
-		}
-		if (!bSentSomething)
-		{
-			break;
+			FMediaRunnable::SleepMilliseconds(100);
 		}
-		FMediaRunnable::SleepMilliseconds(100);
 	}
 
 	// Dump the active track map.
