Protron: removed frame loader message queue in favor of atomically updating structures

#rnx

diff --git a/Engine/Plugins/Media/ElectraPlayer/Source/ElectraProtron/Private/Player/ElectraProtronPlayerImpl.h b/Engine/Plugins/Media/ElectraPlayer/Source/ElectraProtron/Private/Player/ElectraProtronPlayerImpl.h
--- a/Engine/Plugins/Media/ElectraPlayer/Source/ElectraProtron/Private/Player/ElectraProtronPlayerImpl.h
+++ b/Engine/Plugins/Media/ElectraPlayer/Source/ElectraProtron/Private/Player/ElectraProtronPlayerImpl.h
@@ -247,33 +247,8 @@ private:
 
 		// Inherited from FRunnable
 		uint32 Run() override;
-		struct FMessage
-		{
-			enum class EType
-			{
-				Nop,
-				Open,
-				Load,
-				Update,
-				Terminate
-			};
-			struct FParamOpen
-			{
-				FString Filename;
-				TSharedPtr<FSharedPlayParams, ESPMode::ThreadSafe> SharedPlayParams;
-			};
-			struct FParamNull
-			{
-			};
 
-			EType Type = EType::Nop;
-			TVariant<FParamOpen, FParamNull> Param;
-		};
-		void SendThreadMessage(FMessage&& InMessage)
-		{
-			Messages.Enqueue(MoveTemp(InMessage));
-			MessageSignal.Signal();
-		}
+		// Thread start/stop
 		void StartThread(const FString& InFilename, const TSharedPtr<FSharedPlayParams, ESPMode::ThreadSafe>& InSharedPlayParams);
 		void StopThread();
 
@@ -284,17 +259,11 @@ private:
 		// Called by the decoder
 		FMP4SamplePtr GetSample(FMP4TrackSampleBufferPtr InFromBuffer, const FTrackIterator& InAtIterator, int32 InWaitMicroseconds);
 	private:
-		const FConfig& Config;
-		int32 LoaderTypeIndex;
-
-		FRunnableThread* Thread = nullptr;
-		FMediaEvent MessageSignal;
-		TQueue<FMessage, EQueueMode::Mpsc> Messages;
-
-		TSharedPtr<FSharedPlayParams, ESPMode::ThreadSafe> SharedPlayParams;
-		TSharedPtr<Electra::IFileDataReader, ESPMode::ThreadSafe> Reader;
-		FString LastErrorMessage;
-
+		struct FOpenRequest
+		{
+			FString Filename;
+			TSharedPtr<FSharedPlayParams, ESPMode::ThreadSafe> SharedPlayParams;
+		};
 		struct FLoadRequest
 		{
 			void Empty()
@@ -308,7 +277,19 @@ private:
 			FTrackIterator UpdateAtIterator;
 		};
 
+		const FConfig& Config;
+		int32 LoaderTypeIndex;
+
+		FRunnableThread* Thread = nullptr;
+		FMediaEvent WorkSignal;
+		volatile bool bTerminateThread = false;
+
+		TSharedPtr<FSharedPlayParams, ESPMode::ThreadSafe> SharedPlayParams;
+		TSharedPtr<Electra::IFileDataReader, ESPMode::ThreadSafe> Reader;
+		FString LastErrorMessage;
+
 		FCriticalSection LoadRequestLock;
+		FOpenRequest OpenRequest;
 		FLoadRequest PendingLoadRequest;
 		FLoadRequest ActiveLoadRequest;
 		FCriticalSection TimeRangeLock;
diff --git a/Engine/Plugins/Media/ElectraPlayer/Source/ElectraProtron/Private/Player/ElectraProtronPlayerLoader.cpp b/Engine/Plugins/Media/ElectraPlayer/Source/ElectraProtron/Private/Player/ElectraProtronPlayerLoader.cpp
--- a/Engine/Plugins/Media/ElectraPlayer/Source/ElectraProtron/Private/Player/ElectraProtronPlayerLoader.cpp
+++ b/Engine/Plugins/Media/ElectraPlayer/Source/ElectraProtron/Private/Player/ElectraProtronPlayerLoader.cpp
@@ -16,11 +16,10 @@ void FElectraProtronPlayer::FImpl::FLoaderThread::StartThread(const FString& InF
 {
 	if (!Thread)
 	{
-		FMessage Msg;
-		FMessage::FParamOpen open { InFilename, InSharedPlayParams };
-		Msg.Param.Emplace<FMessage::FParamOpen>(MoveTemp(open));
-		Msg.Type = FMessage::EType::Open;
-		SendThreadMessage(MoveTemp(Msg));
+		bTerminateThread = false;
+		OpenRequest.Filename = InFilename;
+		OpenRequest.SharedPlayParams = InSharedPlayParams;
+		WorkSignal.Signal();
 		Thread = FRunnableThread::Create(this, TEXT("Electra Protron Loader"), 0, TPri_Normal);
 	}
 }
@@ -29,10 +28,7 @@ void FElectraProtronPlayer::FImpl::FLoaderThread::StopThread()
 {
 	if (Thread)
 	{
-		FMessage Msg;
-		Msg.Param.Emplace<FMessage::FParamNull>();
-		Msg.Type = FMessage::EType::Terminate;
-		SendThreadMessage(MoveTemp(Msg));
+		bTerminateThread = true;
 		Thread->WaitForCompletion();
 	}
 }
@@ -40,104 +36,63 @@ void FElectraProtronPlayer::FImpl::FLoaderThread::StopThread()
 uint32 FElectraProtronPlayer::FImpl::FLoaderThread::Run()
 {
 // TODO: clamp look ahead/behind values to 0 if negative and an internal limit if too large.
-	bool bDone = false;
-	while(!bDone)
+	while(!bTerminateThread)
 	{
-		MessageSignal.WaitTimeoutAndReset(1000 * 20);
-		FMessage Msg;
-		while(Messages.Dequeue(Msg))
+		WorkSignal.WaitTimeoutAndReset(1000 * 20);
+
+		bool bNeedToLoad = false;
+
+		FScopeLock lock(&LoadRequestLock);
+		// Open a file?
+		if (!OpenRequest.Filename.IsEmpty())
 		{
-			switch(Msg.Type)
+			SharedPlayParams = MoveTemp(OpenRequest.SharedPlayParams);
+			check(SharedPlayParams.IsValid());
+			Reader = Electra::IFileDataReader::Create();
+			FString Filename = MoveTemp(OpenRequest.Filename);
+			if (!Reader->Open(Filename))
 			{
-				case FMessage::EType::Open:
-				{
-					const FMessage::FParamOpen& Param(Msg.Param.Get<FMessage::FParamOpen>());
-					SharedPlayParams = Param.SharedPlayParams;
-					check(SharedPlayParams.IsValid());
-					Reader = Electra::IFileDataReader::Create();
-					if (!Reader->Open(Param.Filename))
-					{
-						// Failure is not really an option seeing as how we already opened the file
-						// successfully in Open().
-						check(!"how could this fail?");
-						LastErrorMessage = Reader->GetLastError();
-						Reader.Reset();
-					}
-					break;
-				}
-				case FMessage::EType::Load:
-				{
-					if (!LastErrorMessage.IsEmpty())
-					{
-						break;
-					}
-					else
-					{
-						FScopeLock lock(&LoadRequestLock);
-						ActiveLoadRequest = MoveTemp(PendingLoadRequest);
-					}
-					// FALLTHROUGH
-					[[fallthrough]];
-				}
-				case FMessage::EType::Update:
-				{
-					LoadRequestLock.Lock();
-					FMP4TrackSampleBufferPtr TrackSampleBuffer = ActiveLoadRequest.TrackSampleBuffer;
-					FTrackIterator StartAtIterator = MoveTemp(ActiveLoadRequest.StartAtIterator);
-					LoadRequestDirty = -1;
-					LoadRequestLock.Unlock();
-					check(StartAtIterator.IsValid());
-					ELoadResult Result = Load(TrackSampleBuffer, StartAtIterator);
-					if (Result == ELoadResult::Error)
-					{
-						if (LastErrorMessage.IsEmpty())
-						{
-							LastErrorMessage = FString::Printf(TEXT("Error loading media samples"));
-						}
-					}
-					else if (Result == ELoadResult::Canceled)
-					{
-						// There are two ways the request can get canceled.
-						// Either because of the player seeking or changing tracks, or the
-						// decoder asking for a new sample after having consumed one.
-						FScopeLock lock(&LoadRequestLock);
-						// A new request from the player wins.
-						if (PendingLoadRequest.TrackSampleBuffer)
-						{
-							// A new request also means there is a new load message in the message queue,
-							// so we do not need to do anything else but reset the active request here.
-							ActiveLoadRequest.Empty();
-						}
-					}
-					break;
-				}
-				case FMessage::EType::Terminate:
-				{
-					bDone = true;
-					break;
-				}
-				default:
-				{
-					unimplemented();
-					break;
-				}
+				// Failure is not really an option seeing as how we already opened the file
+				// successfully in Open().
+				check(!"how could this fail?");
+				LastErrorMessage = Reader->GetLastError();
+				Reader.Reset();
 			}
 		}
-
-		if (1)
+		// New load request trigger by user interaction?
+		else if (PendingLoadRequest.StartAtIterator.IsValid())
 		{
-			FScopeLock lock(&LoadRequestLock);
-			if (ActiveLoadRequest.UpdateAtIterator.IsValid())
+			// Only when there hasn't been an error yet.
+			if (LastErrorMessage.IsEmpty())
 			{
-				check(!ActiveLoadRequest.StartAtIterator.IsValid());
-				ActiveLoadRequest.StartAtIterator = MoveTemp(ActiveLoadRequest.UpdateAtIterator);
-				FMessage UpdateMsg;
-				UpdateMsg.Param.Emplace<FMessage::FParamNull>();
-				UpdateMsg.Type = FMessage::EType::Update;
-				SendThreadMessage(MoveTemp(UpdateMsg));
+				ActiveLoadRequest = MoveTemp(PendingLoadRequest);
+				bNeedToLoad = true;
 			}
 		}
+		// New update request?
+		else if (ActiveLoadRequest.UpdateAtIterator.IsValid())
+		{
+			check(!ActiveLoadRequest.StartAtIterator.IsValid());
+			ActiveLoadRequest.StartAtIterator = MoveTemp(ActiveLoadRequest.UpdateAtIterator);
+			bNeedToLoad = true;
+		}
 
+		if (bNeedToLoad)
+		{
+			FMP4TrackSampleBufferPtr TrackSampleBuffer = ActiveLoadRequest.TrackSampleBuffer;
+			FTrackIterator StartAtIterator = MoveTemp(ActiveLoadRequest.StartAtIterator);
+			LoadRequestDirty = -1;
+			check(StartAtIterator.IsValid());
+			lock.Unlock();
+			ELoadResult Result = Load(TrackSampleBuffer, StartAtIterator);
+			if (Result == ELoadResult::Error)
+			{
+				if (LastErrorMessage.IsEmpty())
+				{
+					LastErrorMessage = FString::Printf(TEXT("Error loading media samples"));
+				}
+			}
+		}
 	}
 	PendingLoadRequest.Empty();
 	ActiveLoadRequest.Empty();
@@ -279,6 +234,7 @@ FElectraProtronPlayer::FImpl::FMP4SamplePtr FElectraProtronPlayer::FImpl::FLoade
 	{
 		ActiveLoadRequest.UpdateAtIterator = InAtIterator->Clone();
 		LoadRequestDirty = (int32) SampleNum;
+		WorkSignal.Signal();
 	}
 	LoadRequestLock.Unlock();
 
@@ -313,12 +269,8 @@ void FElectraProtronPlayer::FImpl::FLoaderThread::RequestLoad(FMP4TrackSampleBuf
 	PendingLoadRequest.TrackSampleBuffer = MoveTemp(InTrackSampleBuffer);
 	PendingLoadRequest.StartAtIterator = MoveTemp(TkIt);
 	LoadRequestDirty = -2;
+	WorkSignal.Signal();
 	LoadRequestLock.Unlock();
-
-	FMessage Msg;
-	Msg.Param.Emplace<FMessage::FParamNull>();
-	Msg.Type = FMessage::EType::Load;
-	SendThreadMessage(MoveTemp(Msg));
 }
 
 FElectraProtronPlayer::FImpl::FLoaderThread::ELoadResult FElectraProtronPlayer::FImpl::FLoaderThread::Load(FMP4TrackSampleBufferPtr InTrackSampleBuffer, FTrackIterator InAtIterator)
