ElectraPlayer: Extended player metrics

#rnx

diff --git a/Engine/Plugins/Media/ElectraPlayer/Source/ElectraPlayerRuntime/Private/ElectraPlayer.cpp b/Engine/Plugins/Media/ElectraPlayer/Source/ElectraPlayerRuntime/Private/ElectraPlayer.cpp
--- a/Engine/Plugins/Media/ElectraPlayer/Source/ElectraPlayerRuntime/Private/ElectraPlayer.cpp
+++ b/Engine/Plugins/Media/ElectraPlayer/Source/ElectraPlayerRuntime/Private/ElectraPlayer.cpp
@@ -2151,6 +2151,12 @@ void FElectraPlayer::HandlePlayerEventReceivedPlaylists()
 	Statistics.MediaTimelineAtStart = MediaTimeline;
 	Statistics.MediaTimelineAtEnd = MediaTimeline;
 	Statistics.MediaDuration = MediaDuration.IsInfinity() ? -1.0 : MediaDuration.GetAsSeconds();
+	Statistics.VideoQualityPercentages.Empty();
+	Statistics.AudioQualityPercentages.Empty();
+	Statistics.VideoSegmentBitratesStreamed.Empty();
+	Statistics.AudioSegmentBitratesStreamed.Empty();
+	Statistics.NumVideoSegmentsStreamed = 0;
+	Statistics.NumAudioSegmentsStreamed = 0;
 	StatisticsLock.Unlock();
 	// Get the video bitrates and populate our number of segments per bitrate map.
 	TArray<Electra::FTrackMetadata> VideoStreamMetaData;
@@ -2160,6 +2166,10 @@ void FElectraPlayer::HandlePlayerEventReceivedPlaylists()
 	{
 		for(int32 i=0; i<VideoStreamMetaData[0].StreamDetails.Num(); ++i)
 		{
+			StatisticsLock.Lock();
+			Statistics.VideoSegmentBitratesStreamed.Add(VideoStreamMetaData[0].StreamDetails[i].Bandwidth, 0);
+			Statistics.VideoQualityPercentages.Add(VideoStreamMetaData[0].StreamDetails[i].Bandwidth, 0);
+			StatisticsLock.Unlock();
 			UE_LOG(LogElectraPlayer, Verbose, TEXT("[%u] Found %d * %d video stream at bitrate %d"), InstanceID,
 												VideoStreamMetaData[0].StreamDetails[i].CodecInformation.GetResolution().Width,
 												VideoStreamMetaData[0].StreamDetails[i].CodecInformation.GetResolution().Height,
@@ -2176,6 +2186,10 @@ void FElectraPlayer::HandlePlayerEventReceivedPlaylists()
 	{
 		for(int32 i=0; i<AudioStreamMetaData[0].StreamDetails.Num(); ++i)
 		{
+			StatisticsLock.Lock();
+			Statistics.AudioSegmentBitratesStreamed.Add(AudioStreamMetaData[0].StreamDetails[i].Bandwidth, 0);
+			Statistics.AudioQualityPercentages.Add(AudioStreamMetaData[0].StreamDetails[i].Bandwidth, 0);
+			StatisticsLock.Unlock();
 			UE_LOG(LogElectraPlayer, Verbose, TEXT("[%u] Found audio stream at bitrate %d"), InstanceID,
 												AudioStreamMetaData[0].StreamDetails[i].Bandwidth);
 		}
@@ -2262,7 +2276,7 @@ void FElectraPlayer::HandlePlayerEventPlaylistDownload(const Electra::Metrics::F
 {
 	// To reduce the number of playlist events during a Live presentation we will only report the initial playlist load
 	// and later on only failed loads but not successful ones.
-	bool bReport =  PlaylistDownloadStats.LoadType == Electra::Playlist::ELoadType::Initial || !PlaylistDownloadStats.bWasSuccessful;
+	bool bReport = PlaylistDownloadStats.LoadType == Electra::Playlist::ELoadType::Initial || !PlaylistDownloadStats.bWasSuccessful;
 	if (bReport)
 	{
 		static const FString kEventNameElectraPlaylistDownload(TEXT("Electra.PlaylistDownload"));
@@ -2280,6 +2294,31 @@ void FElectraPlayer::HandlePlayerEventPlaylistDownload(const Electra::Metrics::F
 			EnqueueAnalyticsEvent(AnalyticEvent);
 		}
 	}
+	// If unsuccessful keep track of the type of error.
+	if (!PlaylistDownloadStats.bWasSuccessful && !PlaylistDownloadStats.bWasAborted)
+	{
+		FScopeLock Lock(&StatisticsLock);
+		if (PlaylistDownloadStats.HTTPStatusCode == 404)
+		{
+			++Statistics.NumErr404;
+		}
+		else if (PlaylistDownloadStats.HTTPStatusCode >= 400 && PlaylistDownloadStats.HTTPStatusCode < 500)
+		{
+			++Statistics.NumErr4xx;
+		}
+		else if (PlaylistDownloadStats.HTTPStatusCode >= 500 && PlaylistDownloadStats.HTTPStatusCode < 600)
+		{
+			++Statistics.NumErr5xx;
+		}
+		else if (PlaylistDownloadStats.bDidTimeout)
+		{
+			++Statistics.NumErrTimeouts;
+		}
+		else
+		{
+			++Statistics.NumErrConnDrops;
+		}
+	}
 }
 
 void FElectraPlayer::HandlePlayerEventLicenseKey(const Electra::Metrics::FLicenseKeyStats& LicenseKeyStats)
@@ -2460,6 +2499,18 @@ void FElectraPlayer::HandlePlayerEventSegmentDownload(const Electra::Metrics::FS
 		}
 		++Statistics.VideoSegmentBitratesStreamed[SegmentDownloadStats.Bitrate];
 
+		++Statistics.NumVideoSegmentsStreamed;
+		if (!Statistics.VideoQualityPercentages.Contains(SegmentDownloadStats.Bitrate))
+		{
+			Statistics.VideoQualityPercentages.Add(SegmentDownloadStats.Bitrate, 0);
+		}
+		for(auto& It : Statistics.VideoQualityPercentages)
+		{
+			const uint32 NumAt = Statistics.VideoSegmentBitratesStreamed[It.Key];
+			const int32 AsPercentage = FMath::RoundToInt(100.0 * (double)NumAt / (double)Statistics.NumVideoSegmentsStreamed);
+			It.Value = AsPercentage;
+		}
+
 		if (Statistics.bIsInitiallyDownloading)
 		{
 			Statistics.InitialBufferingBandwidth.AddSample(8*SegmentDownloadStats.NumBytesDownloaded / (SegmentDownloadStats.TimeToDownload > 0.0 ? SegmentDownloadStats.TimeToDownload : 1.0), SegmentDownloadStats.TimeToFirstByte);
@@ -2478,6 +2529,18 @@ void FElectraPlayer::HandlePlayerEventSegmentDownload(const Electra::Metrics::FS
 		}
 		++Statistics.AudioSegmentBitratesStreamed[SegmentDownloadStats.Bitrate];
 
+		++Statistics.NumAudioSegmentsStreamed;
+		if (!Statistics.AudioQualityPercentages.Contains(SegmentDownloadStats.Bitrate))
+		{
+			Statistics.AudioQualityPercentages.Add(SegmentDownloadStats.Bitrate, 0);
+		}
+		for(auto& It : Statistics.AudioQualityPercentages)
+		{
+			const uint32 NumAt = Statistics.AudioSegmentBitratesStreamed[It.Key];
+			const int32 AsPercentage = FMath::RoundToInt(100.0 * (double)NumAt / (double)Statistics.NumAudioSegmentsStreamed);
+			It.Value = AsPercentage;
+		}
+
 		if (Statistics.bIsInitiallyDownloading && NumTracksVideo == 0)	// Do this just for audio-only presentations.
 		{
 			Statistics.InitialBufferingBandwidth.AddSample(8*SegmentDownloadStats.NumBytesDownloaded / (SegmentDownloadStats.TimeToDownload > 0.0 ? SegmentDownloadStats.TimeToDownload : 1.0), SegmentDownloadStats.TimeToFirstByte);
@@ -2536,6 +2599,33 @@ void FElectraPlayer::HandlePlayerEventSegmentDownload(const Electra::Metrics::FS
 			EnqueueAnalyticsEvent(AnalyticEvent);
 		}
 	}
+	if (!SegmentDownloadStats.bWasSuccessful && !SegmentDownloadStats.bWasAborted)
+	{
+		if (SegmentDownloadStats.HTTPStatusCode == 404)
+		{
+			++Statistics.NumErr404;
+		}
+		else if (SegmentDownloadStats.HTTPStatusCode >= 400 && SegmentDownloadStats.HTTPStatusCode < 500)
+		{
+			++Statistics.NumErr4xx;
+		}
+		else if (SegmentDownloadStats.HTTPStatusCode >= 500 && SegmentDownloadStats.HTTPStatusCode < 600)
+		{
+			++Statistics.NumErr5xx;
+		}
+		else if (SegmentDownloadStats.bDidTimeout)
+		{
+			++Statistics.NumErrTimeouts;
+		}
+		else if (SegmentDownloadStats.bParseFailure)
+		{
+			++Statistics.NumErrOther;
+		}
+		else
+		{
+			++Statistics.NumErrConnDrops;
+		}
+	}
 }
 
 void FElectraPlayer::HandlePlayerEventVideoQualityChange(int32 NewBitrate, int32 PreviousBitrate, bool bIsDrasticDownswitch)
@@ -3011,18 +3101,20 @@ void FElectraPlayer::UpdatePlayEndStatistics()
 
 void FElectraPlayer::LogStatistics()
 {
-	FString VideoSegsPerStream;
-	FString AudioSegsPerStream;
+	FString VideoSegsPercentage;
+	FString AudioSegsPercentage;
 
 	FScopeLock Lock(&StatisticsLock);
 
-	for(const TPair<int32, uint32>& pair : Statistics.VideoSegmentBitratesStreamed)
+	int32 Idx=0;
+	for(auto& It : Statistics.VideoQualityPercentages)
 	{
-		VideoSegsPerStream += FString::Printf(TEXT("%d : %u\n"), pair.Key, pair.Value);
+		VideoSegsPercentage += FString::Printf(TEXT("%d/%d: %d%%\n"), Idx++, It.Key, It.Value);
 	}
-	for(const TPair<int32, uint32>& pair : Statistics.AudioSegmentBitratesStreamed)
+	Idx=0;
+	for(auto& It : Statistics.AudioQualityPercentages)
 	{
-		AudioSegsPerStream += FString::Printf(TEXT("%d : %u\n"), pair.Key, pair.Value);
+		AudioSegsPercentage += FString::Printf(TEXT("%d/%d: %d%%\n"), Idx++, It.Key, It.Value);
 	}
 
 	TSharedPtr<IElectraPlayerAdapterDelegate, ESPMode::ThreadSafe> PinnedAdapterDelegate = AdapterDelegate.Pin();
@@ -3062,11 +3154,17 @@ void FElectraPlayer::LogStatistics()
 			"Number of audio drastic downswitches: %d\n"\
 			"Bytes of video data streamed: %lld\n"\
 			"Bytes of audio data streamed: %lld\n"\
-			"Number of video segments fetched across all quality levels:\n%s"\
-			"Number of audio segments fetched across all quality levels:\n%s"\
+			"Video quality percentage:\n%s"\
+			"Audio quality percentage:\n%s"\
 			"Currently active playlist URL: %s\n"\
 			"Currently active resolution: %d * %d\n" \
 			"Current state: %s\n" \
+			"404 errors: %u\n" \
+			"4xx errors: %u\n" \
+			"5xx errors: %u\n" \
+			"Timeouts: %u\n" \
+			"Connection failures: %u\n" \
+			"Other failures: %u\n" \
 			"Last issue: %s\n"
 		),
 			InstanceID,
@@ -3102,12 +3200,18 @@ void FElectraPlayer::LogStatistics()
 			Statistics.NumAudioQualityDrasticDownswitches,
 			(long long int)Statistics.NumVideoDatabytesStreamed,
 			(long long int)Statistics.NumAudioDatabytesStreamed,
-			*VideoSegsPerStream,
-			*AudioSegsPerStream,
+			*VideoSegsPercentage,
+			*AudioSegsPercentage,
 			*SanitizeMessage(Statistics.CurrentlyActivePlaylistURL),
 			Statistics.CurrentlyActiveResolutionWidth,
 			Statistics.CurrentlyActiveResolutionHeight,
 			*Statistics.LastState,
+			Statistics.NumErr404,
+			Statistics.NumErr4xx,
+			Statistics.NumErr5xx,
+			Statistics.NumErrTimeouts,
+			Statistics.NumErrConnDrops,
+			Statistics.NumErrOther,
 			*SanitizeMessage(Statistics.LastError)
 		);
 
@@ -3211,6 +3315,19 @@ void FElectraPlayer::SendAnalyticMetrics(const TSharedPtr<IAnalyticsProviderET>&
 	}
 	ParamArray.Add(FAnalyticsEventAttribute(TEXT("AudioSegmentFetchStats"), *SegsPerStream));
 
+	// Quality buckets by percentage
+	int32 qbIdx = 0;
+	for(auto &qbIt : (Statistics.NumVideoSegmentsStreamed ? Statistics.VideoQualityPercentages : Statistics.AudioQualityPercentages))
+	{
+		ParamArray.Add(FAnalyticsEventAttribute(FString::Printf(TEXT("qp%d"), qbIdx++), (int32) qbIt.Value));
+	}
+	ParamArray.Add(FAnalyticsEventAttribute(TEXT("Num404"), (uint32) Statistics.NumErr404));
+	ParamArray.Add(FAnalyticsEventAttribute(TEXT("Num4xx"), (uint32) Statistics.NumErr4xx));
+	ParamArray.Add(FAnalyticsEventAttribute(TEXT("Num5xx"), (uint32) Statistics.NumErr5xx));
+	ParamArray.Add(FAnalyticsEventAttribute(TEXT("NumTimeouts"), (uint32) Statistics.NumErrTimeouts));
+	ParamArray.Add(FAnalyticsEventAttribute(TEXT("NumConnDrops"), (uint32) Statistics.NumErrConnDrops));
+	ParamArray.Add(FAnalyticsEventAttribute(TEXT("NumErrOther"), (uint32) Statistics.NumErrOther));
+
 	StatisticsLock.Unlock();
 
 	AnalyticsProvider->RecordEvent(TEXT("Electra.FinalMetrics"), MoveTemp(ParamArray));
diff --git a/Engine/Plugins/Media/ElectraPlayer/Source/ElectraPlayerRuntime/Private/ElectraPlayer.h b/Engine/Plugins/Media/ElectraPlayer/Source/ElectraPlayerRuntime/Private/ElectraPlayer.h
--- a/Engine/Plugins/Media/ElectraPlayer/Source/ElectraPlayerRuntime/Private/ElectraPlayer.h
+++ b/Engine/Plugins/Media/ElectraPlayer/Source/ElectraPlayerRuntime/Private/ElectraPlayer.h
@@ -5,8 +5,10 @@
 
 #include "Containers/UnrealString.h"
 #include "Containers/Queue.h"
+#include "Containers/SortedMap.h"
 #include "Misc/Guid.h"
 #include "Templates/SharedPointer.h"
+#include "Templates/Greater.h"
 #include "IAnalyticsProviderET.h"
 
 #include "IElectraPlayerInterface.h"
@@ -679,6 +681,10 @@ private:
 			CurrentlyActiveResolutionHeight = 0;
 			VideoSegmentBitratesStreamed.Empty();
 			AudioSegmentBitratesStreamed.Empty();
+			VideoQualityPercentages.Empty();
+			AudioQualityPercentages.Empty();
+			NumVideoSegmentsStreamed = 0;
+			NumAudioSegmentsStreamed = 0;
 			InitialBufferingBandwidth.Reset();
 			bIsInitiallyDownloading = false;
 			bDidPlaybackEnd = false;
@@ -686,9 +692,18 @@ private:
 			MediaTimelineAtEnd.Reset();
 			MediaDuration = 0.0;
 			MessageHistoryBuffer.Empty();
+			NumErr404 = 0;
+			NumErr4xx = 0;
+			NumErr5xx = 0;
+			NumErrTimeouts = 0;
+			NumErrConnDrops = 0;
+			NumErrOther = 0;
 		}
 		void AddMessageToHistory(FString InMessage);
 
+		// n elements by descending bitrate holding a percentage of this bitrate being used
+		using FQualityPercentages = TSortedMap<int32, int32, FDefaultAllocator, TGreater<int32>>;
+
 		FString					InitialURL;
 		FString					CurrentlyActivePlaylistURL;
 		FString					LastError;
@@ -722,7 +737,11 @@ private:
 		int32					CurrentlyActiveResolutionWidth;
 		int32					CurrentlyActiveResolutionHeight;
 		TMap<int32, uint32>		VideoSegmentBitratesStreamed;		// key=video stream bitrate, value=number of segments loaded at this rate
-		TMap<int32, uint32>		AudioSegmentBitratesStreamed;		// key=audioo stream bitrate, value=number of segments loaded at this rate
+		TMap<int32, uint32>		AudioSegmentBitratesStreamed;		// key=audio stream bitrate, value=number of segments loaded at this rate
+		FQualityPercentages		VideoQualityPercentages;
+		FQualityPercentages		AudioQualityPercentages;
+		uint32					NumVideoSegmentsStreamed;
+		uint32					NumAudioSegmentsStreamed;
 		FBandwidth				InitialBufferingBandwidth;
 		bool					bIsInitiallyDownloading;
 		bool					bDidPlaybackEnd;
@@ -730,6 +749,12 @@ private:
 		FTimeRange				MediaTimelineAtEnd;
 		double					MediaDuration;
 		TArray<FHistoryEntry>	MessageHistoryBuffer;
+		uint32					NumErr404;
+		uint32					NumErr4xx;
+		uint32					NumErr5xx;
+		uint32					NumErrTimeouts;
+		uint32					NumErrConnDrops;
+		uint32					NumErrOther;
 	};
 
 	struct FAnalyticsEvent
diff --git a/Engine/Plugins/Media/ElectraPlayer/Source/ElectraPlayerRuntime/Private/Runtime/HTTP/HTTPManager.cpp b/Engine/Plugins/Media/ElectraPlayer/Source/ElectraPlayerRuntime/Private/Runtime/HTTP/HTTPManager.cpp
--- a/Engine/Plugins/Media/ElectraPlayer/Source/ElectraPlayerRuntime/Private/Runtime/HTTP/HTTPManager.cpp
+++ b/Engine/Plugins/Media/ElectraPlayer/Source/ElectraPlayerRuntime/Private/Runtime/HTTP/HTTPManager.cpp
@@ -26,15 +26,6 @@
 #include "ElectraHTTPStream.h"
 
 
-#define ERRCODE_OK								0
-#define ERRCODE_HTTP_CONNECTION_TIMEOUT			1
-#define ERRCODE_HTTP_RETURNED_ERROR		 		2
-#define ERRCODE_WRITE_ERROR						3
-#define ERRCODE_OPERATION_TIMEDOUT				4
-#define ERRCODE_HTTP_RANGE_ERROR				5
-#define ERRCODE_HTTP_FILE_COULDNT_READ_FILE		6
-#define ERRCODE_HTTPMODULE_FAILURE				100
-
 DECLARE_CYCLE_STAT(TEXT("FElectraHttpManager::WorkerThread"), STAT_ElectraPlayer_FElectraHttpManager_Worker, STATGROUP_ElectraPlayer);
 
 DEFINE_LOG_CATEGORY(LogElectraHTTPManager);
@@ -354,7 +345,7 @@ namespace Electra
 			void Clear()
 			{
 				Message.Empty();
-				ErrorCode = ERRCODE_OK;
+				ErrorCode = HTTP::EStatusErrorCode::ERRCODE_OK;
 			}
 			void Set(int32 InCode, const char* InMessage)
 			{
@@ -491,11 +482,11 @@ namespace Electra
 		if (!HTTPStreamHandler.IsValid())
 		{
 			Request->ConnectionInfo.EffectiveURL = Request->Parameters.URL;
-			Request->ConnectionInfo.StatusInfo.ErrorCode = ERRCODE_HTTPMODULE_FAILURE;
+			Request->ConnectionInfo.StatusInfo.ErrorCode = HTTP::EStatusErrorCode::ERRCODE_HTTPMODULE_FAILURE;
 			Request->ConnectionInfo.StatusInfo.ErrorDetail.
 				SetError(UEMEDIA_ERROR_NOT_SUPPORTED).
 				SetFacility(Facility::EFacility::HTTPReader).
-				SetCode(ERRCODE_HTTPMODULE_FAILURE).
+				SetCode(HTTP::EStatusErrorCode::ERRCODE_HTTPMODULE_FAILURE).
 				SetMessage(FString::Printf(TEXT("FElectraHttpManager did not initialize")));
 			Request->ConnectionInfo.RequestEndTime = MEDIAutcTime::Current();
 			Request->ConnectionInfo.bHasFinished = true;
@@ -566,7 +557,7 @@ namespace Electra
 			Handle->LocalByteStream = MakeShareable<FLocalByteStream>(DataUrl);
 			if (!DataUrl->SetData(Request->Parameters.URL))
 			{
-				OutError.Set(ERRCODE_HTTP_FILE_COULDNT_READ_FILE, FString::Printf(TEXT("Failed to use data URL \"%s\""), *Request->Parameters.URL));
+				OutError.Set(HTTP::EStatusErrorCode::ERRCODE_HTTP_FILE_COULDNT_READ_FILE, FString::Printf(TEXT("Failed to use data URL \"%s\""), *Request->Parameters.URL));
 				return nullptr;
 			}
 		}
@@ -583,7 +574,7 @@ namespace Electra
 			}
 			else
 			{
-				OutError.Set(ERRCODE_HTTP_FILE_COULDNT_READ_FILE, FString::Printf(TEXT("Failed to parse file name \"%s\""), *Request->Parameters.URL));
+				OutError.Set(HTTP::EStatusErrorCode::ERRCODE_HTTP_FILE_COULDNT_READ_FILE, FString::Printf(TEXT("Failed to parse file name \"%s\""), *Request->Parameters.URL));
 				return nullptr;
 			}
 		}
@@ -605,7 +596,7 @@ namespace Electra
 		}
 		else
 		{
-			OutError.Set(ERRCODE_HTTP_FILE_COULDNT_READ_FILE, FString::Printf(TEXT("External data reader is not valid to read from \"%s\""), *Request->Parameters.URL));
+			OutError.Set(HTTP::EStatusErrorCode::ERRCODE_HTTP_FILE_COULDNT_READ_FILE, FString::Printf(TEXT("External data reader is not valid to read from \"%s\""), *Request->Parameters.URL));
 			return nullptr;
 		}
 	}
@@ -966,11 +957,11 @@ namespace Electra
 				if (!Handle->bIsConnected || !Handle->bHaveResponseHeaders)
 				{
 					Request->ConnectionInfo.StatusInfo.ConnectionTimeoutAfterMilliseconds = (Now - Handle->RequestStartTime).GetAsMilliseconds();
-					Request->ConnectionInfo.StatusInfo.ErrorCode = ERRCODE_OPERATION_TIMEDOUT;
+					Request->ConnectionInfo.StatusInfo.ErrorCode = HTTP::EStatusErrorCode::ERRCODE_HTTP_CONNECTION_TIMEOUT;
 					Request->ConnectionInfo.StatusInfo.ErrorDetail.
 						SetError(UEMEDIA_ERROR_READ_ERROR).
 						SetFacility(Facility::EFacility::HTTPReader).
-						SetCode(ERRCODE_HTTP_CONNECTION_TIMEOUT).
+						SetCode(HTTP::EStatusErrorCode::ERRCODE_HTTP_CONNECTION_TIMEOUT).
 						SetMessage(FString::Printf(TEXT("Connection timeout after %d milliseconds, limit was %d"), Request->ConnectionInfo.StatusInfo.ConnectionTimeoutAfterMilliseconds, (int32)Request->Parameters.ConnectTimeout.GetAsMilliseconds()));
 					Request->ConnectionInfo.RequestEndTime = Now;
 					if (Handle->HttpsRequestCallbackWrapper.IsValid())
@@ -989,11 +980,11 @@ namespace Electra
 				if (DeltaTime >= Request->Parameters.NoDataTimeout)
 				{
 					Request->ConnectionInfo.StatusInfo.NoDataTimeoutAfterMilliseconds = DeltaTime.GetAsMilliseconds();
-					Request->ConnectionInfo.StatusInfo.ErrorCode = ERRCODE_OPERATION_TIMEDOUT;
+					Request->ConnectionInfo.StatusInfo.ErrorCode = HTTP::EStatusErrorCode::ERRCODE_HTTP_CONNECTION_TIMEOUT;
 					Request->ConnectionInfo.StatusInfo.ErrorDetail.
 						SetError(UEMEDIA_ERROR_READ_ERROR).
 						SetFacility(Facility::EFacility::HTTPReader).
-						SetCode(ERRCODE_HTTP_CONNECTION_TIMEOUT).
+						SetCode(HTTP::EStatusErrorCode::ERRCODE_HTTP_CONNECTION_TIMEOUT).
 						SetMessage(FString::Printf(TEXT("No data timeout after %d milliseconds, limit was %d. Received %lld of %lld bytes"), Request->ConnectionInfo.StatusInfo.NoDataTimeoutAfterMilliseconds, (int32)Request->Parameters.NoDataTimeout.GetAsMilliseconds()
 							, (long long int)Request->ConnectionInfo.BytesReadSoFar, (long long int)Request->ConnectionInfo.ContentLength));
 					Request->ConnectionInfo.RequestEndTime = Now;
@@ -1341,29 +1332,29 @@ namespace Electra
 						bool bFailed = false;
 						if (!bHTTPResponseOk || ci.bResponseNotRanged)
 						{
-							ci.StatusInfo.ErrorDetail.SetError(UEMEDIA_ERROR_READ_ERROR).SetFacility(Facility::EFacility::HTTPReader).SetCode(ERRCODE_HTTPMODULE_FAILURE);
+							ci.StatusInfo.ErrorDetail.SetError(UEMEDIA_ERROR_READ_ERROR).SetFacility(Facility::EFacility::HTTPReader).SetCode(HTTP::EStatusErrorCode::ERRCODE_HTTPMODULE_FAILURE);
 							if (ci.StatusInfo.HTTPStatus >= 400)
 							{
 								ci.StatusInfo.ErrorDetail.SetMessage(FString::Printf(TEXT("HTTP returned status %d"), ci.StatusInfo.HTTPStatus));
-								ci.StatusInfo.ErrorCode = ERRCODE_HTTP_RETURNED_ERROR;
+								ci.StatusInfo.ErrorCode = HTTP::EStatusErrorCode::ERRCODE_HTTP_RETURNED_ERROR;
 								bFailed = true;
 							}
 							else if (ci.StatusInfo.HTTPStatus >= 300)
 							{
 								ci.StatusInfo.ErrorDetail.SetMessage(FString::Printf(TEXT("Reached HTTP redirection limit with returned status %d"), ci.StatusInfo.HTTPStatus));
-								ci.StatusInfo.ErrorCode = ERRCODE_HTTP_RETURNED_ERROR;
+								ci.StatusInfo.ErrorCode = HTTP::EStatusErrorCode::ERRCODE_HTTP_RETURNED_ERROR;
 								bFailed = true;
 							}
 							else if (ci.StatusInfo.HTTPStatus == 0)
 							{
 								ci.StatusInfo.ErrorDetail.SetMessage(FString::Printf(TEXT("Failed to get response from server")));
-								ci.StatusInfo.ErrorCode = ERRCODE_HTTP_RETURNED_ERROR;
+								ci.StatusInfo.ErrorCode = HTTP::EStatusErrorCode::ERRCODE_HTTP_RETURNED_ERROR;
 								bFailed = true;
 							}
 							else if (ci.bResponseNotRanged)
 							{
 								ci.StatusInfo.ErrorDetail.SetMessage(FString::Printf(TEXT("Did not receive HTTP 206 for range request")));
-								ci.StatusInfo.ErrorCode = ERRCODE_HTTP_RANGE_ERROR;
+								ci.StatusInfo.ErrorCode = HTTP::EStatusErrorCode::ERRCODE_HTTP_RANGE_ERROR;
 								bFailed = true;
 							}
 						}
@@ -1425,7 +1416,7 @@ namespace Electra
 					// Has it failed?
 					if (Handle->HttpRequest->HasFailed() ) //|| bCompletedWithInsufficientData)
 					{
-						ci.StatusInfo.ErrorCode = ERRCODE_WRITE_ERROR;
+						ci.StatusInfo.ErrorCode = HTTP::EStatusErrorCode::ERRCODE_WRITE_ERROR;
 						ci.StatusInfo.bReadError = true;
 						ci.StatusInfo.ErrorDetail.SetMessage(FString::Printf(TEXT("Read error after receiving %lld of %lld bytes"), (long long int)ci.BytesReadSoFar, (long long int)ci.ContentLength));
 						ci.bHasFinished = true;
@@ -1555,7 +1546,7 @@ namespace Electra
 								else
 								{
 									//
-									ci.StatusInfo.ErrorCode = ERRCODE_WRITE_ERROR;
+									ci.StatusInfo.ErrorCode = HTTP::EStatusErrorCode::ERRCODE_WRITE_ERROR;
 									ci.StatusInfo.bReadError = true;
 									ci.StatusInfo.ErrorDetail.SetMessage(FString::Printf(TEXT("Error setting up the next sub range request")));
 								}
@@ -1657,7 +1648,7 @@ namespace Electra
 			{
 				Request->ConnectionInfo.StatusInfo.HTTPStatus = 404;	// File not found
 				Request->ConnectionInfo.StatusInfo.ErrorDetail.SetMessage(FString::Printf(TEXT("HTTP returned status 404")));
-				Request->ConnectionInfo.StatusInfo.ErrorCode = ERRCODE_HTTP_RETURNED_ERROR;
+				Request->ConnectionInfo.StatusInfo.ErrorCode = HTTP::EStatusErrorCode::ERRCODE_HTTP_RETURNED_ERROR;
 			}
 		}
 	}
@@ -1749,9 +1740,9 @@ namespace Electra
 			{
 				HTTP::FConnectionInfo& ci = Request->ConnectionInfo;
 				ci.StatusInfo.HTTPStatus = 404;
-				ci.StatusInfo.ErrorDetail.SetError(UEMEDIA_ERROR_READ_ERROR).SetFacility(Facility::EFacility::HTTPReader).SetCode(ERRCODE_HTTPMODULE_FAILURE);
+				ci.StatusInfo.ErrorDetail.SetError(UEMEDIA_ERROR_READ_ERROR).SetFacility(Facility::EFacility::HTTPReader).SetCode(HTTP::EStatusErrorCode::ERRCODE_HTTPMODULE_FAILURE);
 				ci.StatusInfo.ErrorDetail.SetMessage(FString::Printf(TEXT("External reader returned -1 for resource size, indicating file not found (HTTP status %d)"), ci.StatusInfo.HTTPStatus));
-				ci.StatusInfo.ErrorCode = ERRCODE_HTTP_RETURNED_ERROR;
+				ci.StatusInfo.ErrorCode = HTTP::EStatusErrorCode::ERRCODE_HTTP_RETURNED_ERROR;
 			}
 		}
 		return -1;
diff --git a/Engine/Plugins/Media/ElectraPlayer/Source/ElectraPlayerRuntime/Private/Runtime/HTTP/HTTPManager.h b/Engine/Plugins/Media/ElectraPlayer/Source/ElectraPlayerRuntime/Private/Runtime/HTTP/HTTPManager.h
--- a/Engine/Plugins/Media/ElectraPlayer/Source/ElectraPlayerRuntime/Private/Runtime/HTTP/HTTPManager.h
+++ b/Engine/Plugins/Media/ElectraPlayer/Source/ElectraPlayerRuntime/Private/Runtime/HTTP/HTTPManager.h
@@ -24,6 +24,16 @@ namespace Electra
 
 	namespace HTTP
 	{
+		enum EStatusErrorCode
+		{
+			ERRCODE_OK = 0,
+			ERRCODE_HTTP_CONNECTION_TIMEOUT = 1,
+			ERRCODE_HTTP_RETURNED_ERROR = 2,
+			ERRCODE_WRITE_ERROR = 3,
+			ERRCODE_HTTP_RANGE_ERROR = 4,
+			ERRCODE_HTTP_FILE_COULDNT_READ_FILE = 5,
+			ERRCODE_HTTPMODULE_FAILURE = 100
+		};
 
 		struct FStatusInfo
 		{
@@ -121,36 +131,36 @@ namespace Electra
 				return *this;
 			}
 
-			private:
-				FConnectionInfo& CopyFrom(const FConnectionInfo& rhs)
-				{
-					FScopeLock l1(&Lock);
-					FScopeLock l2(&rhs.Lock);
-					ResponseHeaders = rhs.ResponseHeaders;
-					EffectiveURL = rhs.EffectiveURL;
-					ContentType = rhs.ContentType;
-					ContentRangeHeader = rhs.ContentRangeHeader;
-					ContentLengthHeader = rhs.ContentLengthHeader;
-					RequestStartTime = rhs.RequestStartTime;
-					RequestEndTime = rhs.RequestEndTime;
-					TimeForDNSResolve = rhs.TimeForDNSResolve;
-					TimeUntilConnected = rhs.TimeUntilConnected;
-					TimeUntilFirstByte = rhs.TimeUntilFirstByte;
-					ContentLength = rhs.ContentLength;
-					BytesReadSoFar = rhs.BytesReadSoFar;
-					NumberOfRedirections = rhs.NumberOfRedirections;
-					HTTPVersionReceived = rhs.HTTPVersionReceived;
-					bIsConnected = rhs.bIsConnected;
-					bHaveResponseHeaders = rhs.bHaveResponseHeaders;
-					bIsChunked = rhs.bIsChunked;
-					bWasAborted = rhs.bWasAborted;
-					bHasFinished = rhs.bHasFinished;
-					bResponseNotRanged = rhs.bResponseNotRanged;
-					bIsCachedResponse = rhs.bIsCachedResponse;
-					StatusInfo = rhs.StatusInfo;
-					TimingTraces = rhs.TimingTraces;
-					return *this;
-				}
+		private:
+			FConnectionInfo& CopyFrom(const FConnectionInfo& rhs)
+			{
+				FScopeLock l1(&Lock);
+				FScopeLock l2(&rhs.Lock);
+				ResponseHeaders = rhs.ResponseHeaders;
+				EffectiveURL = rhs.EffectiveURL;
+				ContentType = rhs.ContentType;
+				ContentRangeHeader = rhs.ContentRangeHeader;
+				ContentLengthHeader = rhs.ContentLengthHeader;
+				RequestStartTime = rhs.RequestStartTime;
+				RequestEndTime = rhs.RequestEndTime;
+				TimeForDNSResolve = rhs.TimeForDNSResolve;
+				TimeUntilConnected = rhs.TimeUntilConnected;
+				TimeUntilFirstByte = rhs.TimeUntilFirstByte;
+				ContentLength = rhs.ContentLength;
+				BytesReadSoFar = rhs.BytesReadSoFar;
+				NumberOfRedirections = rhs.NumberOfRedirections;
+				HTTPVersionReceived = rhs.HTTPVersionReceived;
+				bIsConnected = rhs.bIsConnected;
+				bHaveResponseHeaders = rhs.bHaveResponseHeaders;
+				bIsChunked = rhs.bIsChunked;
+				bWasAborted = rhs.bWasAborted;
+				bHasFinished = rhs.bHasFinished;
+				bResponseNotRanged = rhs.bResponseNotRanged;
+				bIsCachedResponse = rhs.bIsCachedResponse;
+				StatusInfo = rhs.StatusInfo;
+				TimingTraces = rhs.TimingTraces;
+				return *this;
+			}
 		};
 
 	} // namespace HTTP
diff --git a/Engine/Plugins/Media/ElectraPlayer/Source/ElectraPlayerRuntime/Private/Runtime/Player/AdaptiveStreamingPlayer.cpp b/Engine/Plugins/Media/ElectraPlayer/Source/ElectraPlayerRuntime/Private/Runtime/Player/AdaptiveStreamingPlayer.cpp
--- a/Engine/Plugins/Media/ElectraPlayer/Source/ElectraPlayerRuntime/Private/Runtime/Player/AdaptiveStreamingPlayer.cpp
+++ b/Engine/Plugins/Media/ElectraPlayer/Source/ElectraPlayerRuntime/Private/Runtime/Player/AdaptiveStreamingPlayer.cpp
@@ -2029,6 +2029,8 @@ void FAdaptiveStreamingPlayer::HandleSessionMessage(TSharedPtrTS<IPlayerMessage>
 		IPlaylistReader::PlaylistDownloadMessage* pMsg = static_cast<IPlaylistReader::PlaylistDownloadMessage *>(SessionMessage.Get());
 		Metrics::FPlaylistDownloadStats stats;
 		stats.bWasSuccessful = !pMsg->GetConnectionInfo().StatusInfo.ErrorDetail.IsSet();
+		stats.bWasAborted = pMsg->GetConnectionInfo().bWasAborted;
+		stats.bDidTimeout = pMsg->GetConnectionInfo().StatusInfo.ErrorCode == HTTP::EStatusErrorCode::ERRCODE_HTTP_CONNECTION_TIMEOUT;
 		stats.ListType = pMsg->GetListType();
 		stats.LoadType = pMsg->GetLoadType();
 		stats.URL = pMsg->GetConnectionInfo().EffectiveURL;
diff --git a/Engine/Plugins/Media/ElectraPlayer/Source/ElectraPlayerRuntime/Private/Runtime/Player/AdaptiveStreamingPlayerMetrics.h b/Engine/Plugins/Media/ElectraPlayer/Source/ElectraPlayerRuntime/Private/Runtime/Player/AdaptiveStreamingPlayerMetrics.h
--- a/Engine/Plugins/Media/ElectraPlayer/Source/ElectraPlayerRuntime/Private/Runtime/Player/AdaptiveStreamingPlayerMetrics.h
+++ b/Engine/Plugins/Media/ElectraPlayer/Source/ElectraPlayerRuntime/Private/Runtime/Player/AdaptiveStreamingPlayerMetrics.h
@@ -95,23 +95,16 @@ namespace Metrics
 
 	struct FPlaylistDownloadStats
 	{
-		FPlaylistDownloadStats()
-		{
-			ListType	   = Playlist::EListType::Main;
-			LoadType	   = Playlist::ELoadType::Initial;
-			HTTPStatusCode = 0;
-			RetryNumber    = 0;
-			bWasSuccessful = false;
-		}
-
 		FString					URL;
-		//FString				CDN;							// do not have this yet
-		FString					FailureReason;					//!< Human readable failure reason. Only for display purposes.
-		Playlist::EListType		ListType;
-		Playlist::ELoadType		LoadType;
-		int32					HTTPStatusCode;					//!< HTTP status code (0 if not connected to server yet)
-		int32					RetryNumber;
-		bool					bWasSuccessful;
+		//FString				CDN;
+		FString					FailureReason;						// Human readable failure reason. Only for display purposes.
+		Playlist::EListType		ListType = Playlist::EListType::Main;
+		Playlist::ELoadType		LoadType = Playlist::ELoadType::Initial;
+		int32					HTTPStatusCode = 0;					// HTTP status code (0 if not connected to server yet)
+		int32					RetryNumber = 0;
+		bool					bWasSuccessful = false;
+		bool					bDidTimeout = false;
+		bool					bWasAborted = false;
 	};
 
 	struct FSegmentDownloadStats
diff --git a/Engine/Plugins/Media/ElectraPlayer/Source/ElectraPlayerRuntime/Private/Runtime/Player/HLS/PlaylistHLS.cpp b/Engine/Plugins/Media/ElectraPlayer/Source/ElectraPlayerRuntime/Private/Runtime/Player/HLS/PlaylistHLS.cpp
--- a/Engine/Plugins/Media/ElectraPlayer/Source/ElectraPlayerRuntime/Private/Runtime/Player/HLS/PlaylistHLS.cpp
+++ b/Engine/Plugins/Media/ElectraPlayer/Source/ElectraPlayerRuntime/Private/Runtime/Player/HLS/PlaylistHLS.cpp
@@ -3180,7 +3180,7 @@ IManifest::FResult FActiveHLSPlaylist::FPlayPeriod::GetLoopingSegment(TSharedPtr
 void FActiveHLSPlaylist::FPlayPeriod::ValidateDownloadedSegmentDuration(const FStreamSegmentRequestCommon* InRequest)
 {
 	// For successful downloads check if the segment duration was as specified.
-	if (InRequest && InRequest->DownloadStats.bWasSuccessful)
+	if (InRequest && InRequest->DownloadStats.bWasSuccessful && !InRequest->DownloadStats.bInsertedFillerData)
 	{
 		// The #EXTINF duration needs to be as precise to the actual media duration as possible.
 		// Otherwise playback errors could occur and switching across renditions may not work correctly.
diff --git a/Engine/Plugins/Media/ElectraPlayer/Source/ElectraPlayerRuntime/Private/Runtime/Player/StreamSegmentReaderCommon.cpp b/Engine/Plugins/Media/ElectraPlayer/Source/ElectraPlayerRuntime/Private/Runtime/Player/StreamSegmentReaderCommon.cpp
--- a/Engine/Plugins/Media/ElectraPlayer/Source/ElectraPlayerRuntime/Private/Runtime/Player/StreamSegmentReaderCommon.cpp
+++ b/Engine/Plugins/Media/ElectraPlayer/Source/ElectraPlayerRuntime/Private/Runtime/Player/StreamSegmentReaderCommon.cpp
@@ -1441,6 +1441,7 @@ void FStreamSegmentReaderCommon::FStreamHandler::HandleCommonMediaEnd(EHandleRes
 	// Set up remaining download stat fields.
 	DownloadStats.bWasAborted = bAbortedByABR;
 	DownloadStats.bWasSuccessful = !bHasErrored && !bAbortedByABR;
+	DownloadStats.bDidTimeout = CurrentConnectionInfo.StatusInfo.ErrorCode == HTTP::EStatusErrorCode::ERRCODE_HTTP_CONNECTION_TIMEOUT;
 	DownloadStats.URL = CurrentConnectionInfo.EffectiveURL;
 	DownloadStats.HTTPStatusCode = CurrentConnectionInfo.StatusInfo.HTTPStatus;
 	DownloadStats.DurationDownloaded = DurationDownloaded.GetAsSeconds();
@@ -1507,12 +1508,6 @@ void FStreamSegmentReaderCommon::FStreamHandler::HandleCommonMediaEnd(EHandleRes
 		}
 	}
 
-	#if 0
-	// TBD: Set these stream reader outputs ?
-		bool bDidTimeout;
-		bool bParseFailure;
-	#endif
-
 	InRequest->ConnectionInfo = CurrentConnectionInfo;
 	InRequest->DownloadStats = DownloadStats;
 
