ElectraPlayer: Added a custom playlist attribute extraction hook

#rnx

diff --git a/Engine/Plugins/Media/ElectraPlayer/Source/ElectraPlayerPlugin/Private/ElectraPlayerPlugin.cpp b/Engine/Plugins/Media/ElectraPlayer/Source/ElectraPlayerPlugin/Private/ElectraPlayerPlugin.cpp
--- a/Engine/Plugins/Media/ElectraPlayer/Source/ElectraPlayerPlugin/Private/ElectraPlayerPlugin.cpp
+++ b/Engine/Plugins/Media/ElectraPlayer/Source/ElectraPlayerPlugin/Private/ElectraPlayerPlugin.cpp
@@ -598,6 +598,13 @@ bool FElectraPlayerPlugin::Open(const FString& Url, const IMediaOptions* Options
 			PlayerOptions.Set(StringOption, Electra::FVariantValue(Value));
 		}
 	}
+	// Required playlist properties?
+	FString PlaylistProperties = Options->GetMediaOption(TEXT("playlist_properties"), FString());
+	if (PlaylistProperties.Len())
+	{
+		PlayerOptions.Set(TEXT("playlist_properties"), Electra::FVariantValue(PlaylistProperties));
+	}
+
 
 	if (InPlayerOptions)
 	{
diff --git a/Engine/Plugins/Media/ElectraPlayer/Source/ElectraPlayerRuntime/Private/Runtime/Player/AdaptivePlayerOptionKeynames.h b/Engine/Plugins/Media/ElectraPlayer/Source/ElectraPlayerRuntime/Private/Runtime/Player/AdaptivePlayerOptionKeynames.h
--- a/Engine/Plugins/Media/ElectraPlayer/Source/ElectraPlayerRuntime/Private/Runtime/Player/AdaptivePlayerOptionKeynames.h
+++ b/Engine/Plugins/Media/ElectraPlayer/Source/ElectraPlayerRuntime/Private/Runtime/Player/AdaptivePlayerOptionKeynames.h
@@ -56,6 +56,8 @@ namespace Electra
 	const FName OptionKeyResponseCacheMaxByteSize(TEXT("httpcache_max_bytesize"));
 
 	const FName OptionKeyParseTimecodeInfo(TEXT("parse_timecode_info"));
+
+	const FName OptionKeyPlaylistProperties(TEXT("playlist_properties"));
 } // namespace Electra
 
 
diff --git a/Engine/Plugins/Media/ElectraPlayer/Source/ElectraPlayerRuntime/Private/Runtime/Player/AdaptiveStreamingPlayer.cpp b/Engine/Plugins/Media/ElectraPlayer/Source/ElectraPlayerRuntime/Private/Runtime/Player/AdaptiveStreamingPlayer.cpp
--- a/Engine/Plugins/Media/ElectraPlayer/Source/ElectraPlayerRuntime/Private/Runtime/Player/AdaptiveStreamingPlayer.cpp
+++ b/Engine/Plugins/Media/ElectraPlayer/Source/ElectraPlayerRuntime/Private/Runtime/Player/AdaptiveStreamingPlayer.cpp
@@ -4577,6 +4577,16 @@ void FAdaptiveStreamingPlayer::InternalInitialize()
 		PlayerOptions.Set(OptionKeyFrameAccurateSeek, FVariantValue(true));
 	}
 
+	// Get the custom attributes to extract from the main playlist, if any.
+	if (PlayerOptions.HaveKey(OptionKeyPlaylistProperties))
+	{
+		if (!PlaylistPropertyHandler.ConfigureFromJSON(PlayerOptions.GetValue(OptionKeyPlaylistProperties).GetFString()))
+		{
+			PostLog(Facility::EFacility::Player, IInfoLog::ELevel::Warning, FString::Printf(TEXT("Failed to parse playlist property JSON configuration, ignoring.")));
+		}
+	}
+
+
 	// Initialize decoder limits.
 	// NOTE: This is to be deprecated as soon as possible.
 	Deprecate_InternalInitializeDecoderLimits();
diff --git a/Engine/Plugins/Media/ElectraPlayer/Source/ElectraPlayerRuntime/Private/Runtime/Player/AdaptiveStreamingPlayerInternal.h b/Engine/Plugins/Media/ElectraPlayer/Source/ElectraPlayerRuntime/Private/Runtime/Player/AdaptiveStreamingPlayerInternal.h
--- a/Engine/Plugins/Media/ElectraPlayer/Source/ElectraPlayerRuntime/Private/Runtime/Player/AdaptiveStreamingPlayerInternal.h
+++ b/Engine/Plugins/Media/ElectraPlayer/Source/ElectraPlayerRuntime/Private/Runtime/Player/AdaptiveStreamingPlayerInternal.h
@@ -34,6 +34,13 @@
 #include "InfoLog.h"
 
 #include "Misc/TVariant.h"
+#include "Dom/JsonObject.h"
+#include "Serialization/JsonReader.h"
+#include "Serialization/JsonSerializer.h"
+#include "Policies/CondensedJsonPrintPolicy.h"
+#include "Misc/Base64.h"
+
+
 
 #define INTERR_ALL_STREAMS_HAVE_FAILED			1
 #define INTERR_UNSUPPORTED_FORMAT				2
@@ -999,6 +1006,77 @@ private:
 
 
 
+class FAdaptiveStreamingPlayerPlaylistPropertyHandler
+{
+public:
+	FAdaptiveStreamingPlayerPlaylistPropertyHandler() = default;
+	~FAdaptiveStreamingPlayerPlaylistPropertyHandler() = default;
+	bool ConfigureFromJSON(const FString& InJSONConfig)
+	{
+		if (InJSONConfig.Len())
+		{
+			TSharedRef<TJsonReader<>> Reader = TJsonReaderFactory<>::Create(InJSONConfig);
+			return FJsonSerializer::Deserialize(Reader, Config);
+		}
+		return true;
+	}
+	bool WantProperty(FName& OutAsName, const FString& InProtocol, const FString& InProperty)
+	{
+		static const FString AnyProtocol(TEXT("*"));
+		if (!Config.IsValid())
+		{
+			return false;
+		}
+		// Get the protocol section
+		TArray<TArray<TSharedPtr<FJsonValue>>> Protocols;
+		const TArray<TSharedPtr<FJsonValue>>* ProtocolObj = nullptr;
+		if (Config->TryGetArrayField(InProtocol, ProtocolObj) && ProtocolObj)
+		{
+			Protocols.Emplace(*ProtocolObj);
+		}
+		ProtocolObj = nullptr;
+		if (Config->TryGetArrayField(AnyProtocol, ProtocolObj) && ProtocolObj)
+		{
+			Protocols.Emplace(*ProtocolObj);
+		}
+		if (Protocols.IsEmpty())
+		{
+			return false;
+		}
+		for(auto& Protocol : Protocols)
+		{
+			for(auto& It : Protocol)
+			{
+				const TSharedPtr<FJsonObject>* Obj = nullptr;
+				if (!It.IsValid() || !It->TryGetObject(Obj))
+				{
+					continue;
+				}
+				check(Obj);
+				FString Key;
+				if (!(*Obj)->TryGetStringField(TEXT("k"), Key))
+				{
+					continue;
+				}
+				if (Key.Equals(InProperty, ESearchCase::IgnoreCase))
+				{
+					FString As;
+					if (!(*Obj)->TryGetStringField(TEXT("as"), As))
+					{
+						As = Key;
+					}
+					OutAsName = FName(*As);
+					return true;
+				}
+			}
+		}
+		return false;
+	}
+private:
+	TSharedPtr<FJsonObject> Config;
+};
+
+
 class FAdaptiveStreamingPlayer : public IAdaptiveStreamingPlayer
 							   , public IPlayerSessionServices
 							   , public IStreamReader::StreamReaderEventListener
@@ -1129,9 +1207,9 @@ private:
 	FParamDictTS& GetMutableOptions() override;
 	bool HaveOptionValue(const FName& InOption) override;
 	const FVariantValue GetOptionValue(const FName& InOption) override;
-
 	TSharedPtrTS<FDRMManager> GetDRMManager() override;
 	void SetPlaybackEnd(const FTimeValue& InEndAtTime, IPlayerSessionServices::EPlayEndReason InEndingReason, TSharedPtrTS<IPlayerSessionServices::IPlayEndReason> InCustomManifestObject) override;
+	ECustomPropertyResult ValidateMainPlaylistCustomProperty(const FString& InProtocol, const FString& InPlaylistURL, const TArray<FElectraHTTPStreamHeader>& InPlaylistFetchResponseHeaders, const FPlaylistProperty& InProperty) override;
 
 	// Methods from IPlayerStreamFilter
 	bool CanDecodeStream(const FStreamCodecInformation& InStreamCodecInfo) const override;
@@ -2238,6 +2316,7 @@ private:
 	TSharedPtrTS<IManifest>												Manifest;
 	TSharedPtrTS<IPlaylistReader>										ManifestReader;
 	TSharedPtrTS<FHTTPResourceRequest>									ManifestMimeTypeRequest;
+	FAdaptiveStreamingPlayerPlaylistPropertyHandler						PlaylistPropertyHandler;
 
 	IStreamReader*														StreamReaderHandler;
 	bool																bStreamingHasStarted;
diff --git a/Engine/Plugins/Media/ElectraPlayer/Source/ElectraPlayerRuntime/Private/Runtime/Player/AdaptiveStreamingPlayerManifest.cpp b/Engine/Plugins/Media/ElectraPlayer/Source/ElectraPlayerRuntime/Private/Runtime/Player/AdaptiveStreamingPlayerManifest.cpp
--- a/Engine/Plugins/Media/ElectraPlayer/Source/ElectraPlayerRuntime/Private/Runtime/Player/AdaptiveStreamingPlayerManifest.cpp
+++ b/Engine/Plugins/Media/ElectraPlayer/Source/ElectraPlayerRuntime/Private/Runtime/Player/AdaptiveStreamingPlayerManifest.cpp
@@ -450,4 +450,61 @@ FAdaptiveStreamingPlayer::FMediaMetadataUpdate::EResult FAdaptiveStreamingPlayer
 }
 
 
+IPlayerSessionServices::ECustomPropertyResult FAdaptiveStreamingPlayer::ValidateMainPlaylistCustomProperty(const FString& InProtocol, const FString& InPlaylistURL, const TArray<FElectraHTTPStreamHeader>& InPlaylistFetchResponseHeaders, const FPlaylistProperty& InProperty)
+{
+	if (InProperty.Tag.Len())
+	{
+		FName AsName;
+		if (PlaylistPropertyHandler.WantProperty(AsName, InProtocol, InProperty.Tag))
+		{
+			/*
+				If the property is asked for, we create a single JSON string describing how it was retrieved, that is, from which
+				streaming protocol, the URL of the playlist and the HTTP response headers associated with the playlist fetch:
+
+				{
+					"v": "the value of the property as it appears in the playlist",
+					"proto": "the streaming protocol in use",
+					"url": "the url the playlist was fetched from",
+					"headers": [
+					            [ "responseheader1", "value1" ],
+					            [ "responseheader2", "value2" ],
+					            [ ... ... ],
+					           ]
+				}
+
+				Note that the response headers is an array of arrays, where each entry is an array
+				of two elements, [0] being the header and [1] the value. This is faster to handle than
+				an object for which the key name would be variable/unknown.
+
+				This JSON string is then set as the value for the mapped property name as the property
+				configuration says, and stored in the player options from where it can be retrieved
+				with the GetMediaInfo() method.
+			*/
+
+			TSharedRef<FJsonObject> RootObject = MakeShareable(new FJsonObject);
+			TArray<TSharedPtr<FJsonValue>> HeaderArray;
+			RootObject->SetField(TEXT("v"), MakeShareable(new FJsonValueString(InProperty.Value)));
+			RootObject->SetField(TEXT("proto"), MakeShareable(new FJsonValueString(InProtocol)));
+			RootObject->SetField(TEXT("url"), MakeShareable(new FJsonValueString(InPlaylistURL)));
+			for(auto& Hdr : InPlaylistFetchResponseHeaders)
+			{
+				TArray<TSharedPtr<FJsonValue>> h;
+				h.Add(MakeShareable(new FJsonValueString(Hdr.Header)));
+				h.Add(MakeShareable(new FJsonValueString(Hdr.Value)));
+				HeaderArray.Add(MakeShareable(new FJsonValueArray(h)));
+			}
+			RootObject->SetArrayField(TEXT("headers"), HeaderArray);
+			FString PropertyJSON;
+			TSharedRef<TJsonWriter<TCHAR, TCondensedJsonPrintPolicy<TCHAR>>> JsonWriter = TJsonWriterFactory<TCHAR, TCondensedJsonPrintPolicy<TCHAR>>::Create(&PropertyJSON, 0);
+			if (FJsonSerializer::Serialize(RootObject, JsonWriter))
+			{
+				PlayerOptions.Set(AsName, FVariantValue(PropertyJSON));
+				return IPlayerSessionServices::ECustomPropertyResult::Accept;
+			}
+		}
+	}
+	return IPlayerSessionServices::ECustomPropertyResult::Default;
+}
+
+
 } // namespace Electra
diff --git a/Engine/Plugins/Media/ElectraPlayer/Source/ElectraPlayerRuntime/Private/Runtime/Player/DASH/ManifestBuilderDASH.cpp b/Engine/Plugins/Media/ElectraPlayer/Source/ElectraPlayerRuntime/Private/Runtime/Player/DASH/ManifestBuilderDASH.cpp
--- a/Engine/Plugins/Media/ElectraPlayer/Source/ElectraPlayerRuntime/Private/Runtime/Player/DASH/ManifestBuilderDASH.cpp
+++ b/Engine/Plugins/Media/ElectraPlayer/Source/ElectraPlayerRuntime/Private/Runtime/Player/DASH/ManifestBuilderDASH.cpp
@@ -40,12 +40,13 @@
 #define ERRCODE_DASH_MPD_BUILDER_URL_FAILED_TO_RESOLVE				105
 #define ERRCODE_DASH_MPD_BUILDER_MISSING_REQUIRED_ELEMENT			106
 
-
 #define ERRCODE_DASH_MPD_BUILDER_BAD_PERIOD_START					200
 #define ERRCODE_DASH_MPD_BUILDER_EARLY_PERIODS_MUST_BE_LAST			201
 #define ERRCODE_DASH_MPD_BUILDER_MEDIAPRESENTATIONDURATION_NEEDED	202
 #define ERRCODE_DASH_MPD_BUILDER_BAD_PERIOD_DURATION				203
 
+#define ERRCODE_DASH_MPD_BUILDER_REJECTED_BY_APPLICATION			300
+
 
 DECLARE_CYCLE_STAT(TEXT("FManifestDASHInternal::Build"), STAT_ElectraPlayer_FManifestDASHInternal_Build, STATGROUP_ElectraPlayer);
 DECLARE_CYCLE_STAT(TEXT("FManifestDASHInternal::XLink"), STAT_ElectraPlayer_FManifestDASHInternal_XLink, STATGROUP_ElectraPlayer);
@@ -98,6 +99,7 @@ namespace
 
 	const TCHAR* const DASHRole = TEXT("urn:mpeg:dash:role:2011");
 
+	const TCHAR* const EssentialDashReservedPropertyPrefix = TEXT("urn:mpeg:dash:");
 	const TCHAR* const SupportedEssentialProperties[] =
 	{
 		TEXT("urn:mpeg:dash:urlparam:2014"),
@@ -735,7 +737,7 @@ public:
 	FManifestBuilderDASH(IPlayerSessionServices* InPlayerSessionServices);
 	virtual ~FManifestBuilderDASH() = default;
 
-	virtual FErrorDetail BuildFromMPD(TSharedPtrTS<FManifestDASHInternal>& OutMPD, TCHAR* InOutMPDXML, const FString& EffectiveURL, const FString& ETag) override;
+	virtual FErrorDetail BuildFromMPD(TSharedPtrTS<FManifestDASHInternal>& OutMPD, TCHAR* InOutMPDXML, const FString& EffectiveURL, const FString& ETag, const TArray<FElectraHTTPStreamHeader>& InResponseHeaders) override;
 
 private:
 	ELECTRA_IMPL_DEFAULT_ERROR_METHODS(DASHMPDBuilder);
@@ -757,7 +759,7 @@ FManifestBuilderDASH::FManifestBuilderDASH(IPlayerSessionServices* InPlayerSessi
 {
 }
 
-FErrorDetail FManifestBuilderDASH::BuildFromMPD(TSharedPtrTS<FManifestDASHInternal>& OutMPD, TCHAR* InOutMPDXML, const FString& EffectiveURL, const FString& ETag)
+FErrorDetail FManifestBuilderDASH::BuildFromMPD(TSharedPtrTS<FManifestDASHInternal>& OutMPD, TCHAR* InOutMPDXML, const FString& EffectiveURL, const FString& ETag, const TArray<FElectraHTTPStreamHeader>& InResponseHeaders)
 {
 	TSharedPtrTS<FDashMPD_MPDType> NewMPD;
 	TArray<TWeakPtrTS<IDashMPDElement>> XLinkElements;
@@ -800,19 +802,63 @@ FErrorDetail FManifestBuilderDASH::BuildFromMPD(TSharedPtrTS<FManifestDASHIntern
 	}
 
 	// Check if this MPD requires essential properties we do not understand.
+	FString DashProtocol(TEXT("dash"));
 	const TArray<TSharedPtrTS<FDashMPD_DescriptorType>>& EssentialProperties = NewMPD->GetEssentialProperties();
+	TArray<bool> IsSupported;
+	IsSupported.SetNum(EssentialProperties.Num());
+	// Not supported by default.
+	for(int32 i=0; i<IsSupported.Num(); ++i)
+	{
+		IsSupported[i] = false;
+	}
 	for(int32 nEssProp=0; nEssProp<EssentialProperties.Num(); ++nEssProp)
 	{
-		bool bIsSupported = false;
+		// If supported there is no need to check again.
+		if (IsSupported[nEssProp])
+		{
+			continue;
+		}
+		// Check against a fixed list of properties we support internally.
 		for(int32 nSupportedEssProp=0; nSupportedEssProp<UE_ARRAY_COUNT(SupportedEssentialProperties); ++nSupportedEssProp)
 		{
 			if (EssentialProperties[nEssProp]->GetSchemeIdUri().Equals(SupportedEssentialProperties[nSupportedEssProp]))
 			{
-				bIsSupported = true;
+				IsSupported[nEssProp] = true;
 				break;
 			}
 		}
-		if (!bIsSupported)
+		// If still not supported and not a DASH prefix check with the external registration
+		if (!IsSupported[nEssProp] && !EssentialProperties[nEssProp]->GetSchemeIdUri().StartsWith(EssentialDashReservedPropertyPrefix))
+		{
+			IPlayerSessionServices::FPlaylistProperty Prop;
+			Prop.Tag = EssentialProperties[nEssProp]->GetSchemeIdUri();
+			Prop.Value = EssentialProperties[nEssProp]->GetValue();
+			IPlayerSessionServices::ECustomPropertyResult Result = PlayerSessionServices->ValidateMainPlaylistCustomProperty(DashProtocol, EffectiveURL, InResponseHeaders, Prop);
+			// Essential properties need to be explicitly accepted. `Default` means the external registry doesn't know or care, which means no one knows how to handle it.
+			IsSupported[nEssProp] = Result == IPlayerSessionServices::ECustomPropertyResult::Accept;
+		}
+
+		// If supported propagate this to all elements sharing the same ID
+		if (IsSupported[nEssProp] && EssentialProperties[nEssProp]->GetID().Len())
+		{
+			for(int32 j=0; j<EssentialProperties.Num(); ++j)
+			{
+				if (j!=nEssProp && EssentialProperties[nEssProp]->GetID().Equals(EssentialProperties[j]->GetID()))
+				{
+					IsSupported[j] = true;
+				}
+			}
+		}
+	}
+	// Call external registry with an end-of-properties call.
+	if (PlayerSessionServices->ValidateMainPlaylistCustomProperty(DashProtocol, EffectiveURL, InResponseHeaders, IPlayerSessionServices::FPlaylistProperty()) == IPlayerSessionServices::ECustomPropertyResult::Reject)
+	{
+		return CreateErrorAndLog(PlayerSessionServices, FString::Printf(TEXT("MPD has been rejected by application")), ERRCODE_DASH_MPD_BUILDER_REJECTED_BY_APPLICATION);
+	}
+	// If any one essential property is not supported we cannot commence playback.
+	for(int32 nEssProp=0; nEssProp<EssentialProperties.Num(); ++nEssProp)
+	{
+		if (!IsSupported[nEssProp])
 		{
 			return CreateErrorAndLog(PlayerSessionServices, FString::Printf(TEXT("MPD requires unsupported EssentialProperty \"%s\" for playback"), *EssentialProperties[nEssProp]->GetSchemeIdUri()), ERRCODE_DASH_MPD_BUILDER_UNSUPPORTED_ESSENTIAL_PROPERTY);
 		}
diff --git a/Engine/Plugins/Media/ElectraPlayer/Source/ElectraPlayerRuntime/Private/Runtime/Player/DASH/ManifestBuilderDASH.h b/Engine/Plugins/Media/ElectraPlayer/Source/ElectraPlayerRuntime/Private/Runtime/Player/DASH/ManifestBuilderDASH.h
--- a/Engine/Plugins/Media/ElectraPlayer/Source/ElectraPlayerRuntime/Private/Runtime/Player/DASH/ManifestBuilderDASH.h
+++ b/Engine/Plugins/Media/ElectraPlayer/Source/ElectraPlayerRuntime/Private/Runtime/Player/DASH/ManifestBuilderDASH.h
@@ -7,6 +7,7 @@
 #include "Player/PlaylistReader.h"
 #include "Player/PlayerSessionServices.h"
 #include "PlaylistReaderDASH_Internal.h"
+#include "ElectraHTTPStream.h"
 #include "StreamTypes.h"
 #include "ErrorDetail.h"
 
@@ -28,10 +29,11 @@ public:
 	 * @param InOutMPDXML
 	 * @param EffectiveURL
 	 * @param ETag
+	 * @param InResponseHeaders
 	 *
 	 * @return
 	 */
-	virtual FErrorDetail BuildFromMPD(TSharedPtrTS<FManifestDASHInternal>& OutMPD, TCHAR* InOutMPDXML, const FString& EffectiveURL, const FString& ETag) = 0;
+	virtual FErrorDetail BuildFromMPD(TSharedPtrTS<FManifestDASHInternal>& OutMPD, TCHAR* InOutMPDXML, const FString& EffectiveURL, const FString& ETag, const TArray<FElectraHTTPStreamHeader>& InResponseHeaders) = 0;
 };
 
 
diff --git a/Engine/Plugins/Media/ElectraPlayer/Source/ElectraPlayerRuntime/Private/Runtime/Player/DASH/PlaylistReaderDASH.cpp b/Engine/Plugins/Media/ElectraPlayer/Source/ElectraPlayerRuntime/Private/Runtime/Player/DASH/PlaylistReaderDASH.cpp
--- a/Engine/Plugins/Media/ElectraPlayer/Source/ElectraPlayerRuntime/Private/Runtime/Player/DASH/PlaylistReaderDASH.cpp
+++ b/Engine/Plugins/Media/ElectraPlayer/Source/ElectraPlayerRuntime/Private/Runtime/Player/DASH/PlaylistReaderDASH.cpp
@@ -20,7 +20,7 @@
 #include "Player/DASH/OptionKeynamesDASH.h"
 #include "Player/DASH/PlayerEventDASH.h"
 #include "Player/DASH/PlayerEventDASH_Internal.h"
-
+#include "ElectraHTTPStream.h"
 #include "Misc/DateTime.h"
 
 #define ERRCODE_DASH_PARSER_ERROR							1
@@ -1150,11 +1150,13 @@ void FPlaylistReaderDASH::ManifestDownloadCompleted(FResourceLoadRequestPtr Requ
 
 		FString ETag;
 		FString EffectiveURL = ConnInfo->EffectiveURL;
+		TArray<FElectraHTTPStreamHeader> Headers;
 		InitialTimesync.HttpRequestStart = ConnInfo->RequestStartTime;
 		InitialTimesync.UtcTimeThen = MEDIAutcTime::Current();
 		InitialTimesync.ConnInfo = MakeUnique<HTTP::FConnectionInfo>(*ConnInfo);
 		for(int32 i=0; i<ConnInfo->ResponseHeaders.Num(); ++i)
 		{
+			Headers.Add(ConnInfo->ResponseHeaders[i]);
 			if (ConnInfo->ResponseHeaders[i].Header.Equals(TEXT("Date"), ESearchCase::IgnoreCase))
 			{
 				InitialTimesync.HttpDateHeader = ConnInfo->ResponseHeaders[i].Value;
@@ -1173,7 +1175,7 @@ void FPlaylistReaderDASH::ManifestDownloadCompleted(FResourceLoadRequestPtr Requ
 		{
 			Builder.Reset(IManifestBuilderDASH::Create(PlayerSessionServices));
 			TSharedPtrTS<FManifestDASHInternal> NewManifest;
-			LastErrorDetail = Builder->BuildFromMPD(NewManifest, XML.GetCharArray().GetData(), EffectiveURL, ETag);
+			LastErrorDetail = Builder->BuildFromMPD(NewManifest, XML.GetCharArray().GetData(), EffectiveURL, ETag, Headers);
 			if (LastErrorDetail.IsOK() || LastErrorDetail.IsTryAgain())
 			{
 				if (NewManifest.IsValid())
@@ -1358,12 +1360,14 @@ void FPlaylistReaderDASH::ManifestUpdateDownloadCompleted(FResourceLoadRequestPt
 		}
 		FString ETag;
 		FString EffectiveURL;
+		TArray<FElectraHTTPStreamHeader> Headers;
 		FTimeValue FetchTime;
 		EffectiveURL = ConnInfo->EffectiveURL;
 		//FetchTime = PlayerSessionServices->GetSynchronizedUTCTime()->MapToSyncTime(ConnInfo->RequestStartTime);
 		FetchTime = PlayerSessionServices->GetSynchronizedUTCTime()->GetTime();
 		for(int32 i=0; i<ConnInfo->ResponseHeaders.Num(); ++i)
 		{
+			Headers.Add(ConnInfo->ResponseHeaders[i]);
 			if (ConnInfo->ResponseHeaders[i].Header.Equals(TEXT("ETag"), ESearchCase::IgnoreCase))
 			{
 				ETag = ConnInfo->ResponseHeaders[i].Value;
@@ -1404,7 +1408,7 @@ void FPlaylistReaderDASH::ManifestUpdateDownloadCompleted(FResourceLoadRequestPt
 			if (LastErrorDetail.IsOK())
 			{
 				TSharedPtrTS<FManifestDASHInternal> NewManifest;
-				LastErrorDetail = Builder->BuildFromMPD(NewManifest, XML.GetCharArray().GetData(), EffectiveURL, ETag);
+				LastErrorDetail = Builder->BuildFromMPD(NewManifest, XML.GetCharArray().GetData(), EffectiveURL, ETag, Headers);
 				if (LastErrorDetail.IsOK() || LastErrorDetail.IsTryAgain())
 				{
 					// Set the availability start time of the new MPD to be that of the current MPD.
diff --git a/Engine/Plugins/Media/ElectraPlayer/Source/ElectraPlayerRuntime/Private/Runtime/Player/HLS/PlaylistHandlerHLS.cpp b/Engine/Plugins/Media/ElectraPlayer/Source/ElectraPlayerRuntime/Private/Runtime/Player/HLS/PlaylistHandlerHLS.cpp
--- a/Engine/Plugins/Media/ElectraPlayer/Source/ElectraPlayerRuntime/Private/Runtime/Player/HLS/PlaylistHandlerHLS.cpp
+++ b/Engine/Plugins/Media/ElectraPlayer/Source/ElectraPlayerRuntime/Private/Runtime/Player/HLS/PlaylistHandlerHLS.cpp
@@ -562,9 +562,14 @@ TSharedPtr<FPlaylistParserHLS, ESPMode::ThreadSafe> FPlaylistHandlerHLS::HandleM
 
 	// Remember the effective URL after redirections.
 	MultiVariantPlaylistEffectiveURL = InRequest->ResourceRequest->GetConnectionInfo()->EffectiveURL;
+	TArray<FElectraHTTPStreamHeader> Headers;
+	for(auto &hdr : InRequest->ResourceRequest->GetConnectionInfo()->ResponseHeaders)
+	{
+		Headers.Add(hdr);
+	}
 	// Create a parser and parse the response.
 	TSharedPtr<FPlaylistParserHLS, ESPMode::ThreadSafe> MainPlaylist = MakeShared<FPlaylistParserHLS, ESPMode::ThreadSafe>();
-	LastError = MainPlaylist->Parse(M3U8, MultiVariantPlaylistEffectiveURL);
+	LastError = MainPlaylist->Parse(M3U8, MultiVariantPlaylistEffectiveURL, Headers);
 	if (LastError.IsError())
 	{
 		MainPlaylist.Reset();
@@ -712,7 +717,12 @@ TSharedPtr<FPlaylistParserHLS, ESPMode::ThreadSafe> FPlaylistHandlerHLS::HandleV
 	// Create a parser and parse the response.
 	TSharedPtr<FPlaylistParserHLS, ESPMode::ThreadSafe> VariantPlaylist;
 	VariantPlaylist = MakeShared<FPlaylistParserHLS, ESPMode::ThreadSafe>();
-	LastError = VariantPlaylist->Parse(M3U8, InRequest->ResourceRequest->GetConnectionInfo()->EffectiveURL);
+	TArray<FElectraHTTPStreamHeader> Headers;
+	for(auto &hdr : InRequest->ResourceRequest->GetConnectionInfo()->ResponseHeaders)
+	{
+		Headers.Add(hdr);
+	}
+	LastError = VariantPlaylist->Parse(M3U8, InRequest->ResourceRequest->GetConnectionInfo()->EffectiveURL, Headers);
 	if (LastError.IsError())
 	{
 		VariantPlaylist.Reset();
@@ -1679,6 +1689,52 @@ bool FPlaylistHandlerHLS::BuildMultiVariantPlaylist(TSharedPtr<FMultiVariantPlay
 		{
 			ParseStartTime(NewVariantPlaylist->StartTime, Elements[nElem]);
 		}
+		else if (Elements[nElem]->Tag == FPlaylistParserHLS::EExtTag::EXT_X_SESSION_DATA)
+		{
+			if (const FPlaylistParserHLS::FAttribute* DataID = Elements[nElem]->GetAttribute(TEXT("DATA-ID")))
+			{
+				const FPlaylistParserHLS::FAttribute* SessDataValue = Elements[nElem]->GetAttribute(TEXT("VALUE"));
+				const FPlaylistParserHLS::FAttribute* SessDataURI = Elements[nElem]->GetAttribute(TEXT("URI"));
+				if (!SessDataValue && !SessDataURI)
+				{
+					LastError = PostError(FString::Printf(TEXT("EXT-X-SESSION-DATA has neither VALUE nor URI attribute")), HLS::ERRCODE_PLAYLIST_PARSING_FAILED);
+					return false;
+				}
+				if (SessDataValue && SessDataURI)
+				{
+					LastError = PostError(FString::Printf(TEXT("EXT-X-SESSION-DATA has both VALUE or URI attributes")), HLS::ERRCODE_PLAYLIST_PARSING_FAILED);
+					return false;
+				}
+				// We do not resolve any session data pointing to another URL resource.
+				if (SessDataValue)
+				{
+					IPlayerSessionServices::FPlaylistProperty Prop;
+					if (!DataID->GetValue(Prop.Tag, NewVariantPlaylist->VariableSubstitutions))
+					{
+						LastError = PostError(FString::Printf(TEXT("EXT-X-SESSION-DATA has a bad DATA-ID value")), HLS::ERRCODE_PLAYLIST_PARSING_FAILED);
+						return false;
+					}
+					if (!SessDataValue->GetValue(Prop.Value, NewVariantPlaylist->VariableSubstitutions))
+					{
+						LastError = PostError(FString::Printf(TEXT("EXT-X-SESSION-DATA has a bad VALUE value")), HLS::ERRCODE_PLAYLIST_PARSING_FAILED);
+						return false;
+					}
+
+					// Send this session data property up to registered listeners.
+					IPlayerSessionServices::ECustomPropertyResult Result = PlayerSessionServices->ValidateMainPlaylistCustomProperty(TEXT("hls"), InPlaylist->GetURL(), InPlaylist->GetResponseHeaders(), Prop);
+					if (Result == IPlayerSessionServices::ECustomPropertyResult::Reject)
+					{
+						LastError = PostError(FString::Printf(TEXT("Playlist has been rejected by application due to bad EXT-X-SESSION-DATA")), HLS::ERRCODE_PLAYLIST_PARSING_FAILED);
+						return false;
+					}
+				}
+			}
+			else
+			{
+				LastError = PostError(FString::Printf(TEXT("EXT-X-SESSION-DATA is missing required DATA-ID attribute")), HLS::ERRCODE_PLAYLIST_PARSING_FAILED);
+				return false;
+			}
+		}
 	}
 
 	// Now process all EXT-X-STREAM-INF elements. These may reference any of the rendition groups.
@@ -1958,6 +2014,13 @@ bool FPlaylistHandlerHLS::BuildMultiVariantPlaylist(TSharedPtr<FMultiVariantPlay
 		}
 	}
 
+	// Call external registry with an end-of-properties call.
+	if (PlayerSessionServices->ValidateMainPlaylistCustomProperty(TEXT("hls"), InPlaylist->GetURL(), InPlaylist->GetResponseHeaders(), IPlayerSessionServices::FPlaylistProperty()) == IPlayerSessionServices::ECustomPropertyResult::Reject)
+	{
+		LastError = PostError(FString::Printf(TEXT("Playlist has been rejected by application")), HLS::ERRCODE_PLAYLIST_PARSING_FAILED);
+		return false;
+	}
+
 	OutMultiVariantPlaylist = MoveTemp(NewVariantPlaylist);
 	return true;
 }
diff --git a/Engine/Plugins/Media/ElectraPlayer/Source/ElectraPlayerRuntime/Private/Runtime/Player/HLS/PlaylistParserHLS.cpp b/Engine/Plugins/Media/ElectraPlayer/Source/ElectraPlayerRuntime/Private/Runtime/Player/HLS/PlaylistParserHLS.cpp
--- a/Engine/Plugins/Media/ElectraPlayer/Source/ElectraPlayerRuntime/Private/Runtime/Player/HLS/PlaylistParserHLS.cpp
+++ b/Engine/Plugins/Media/ElectraPlayer/Source/ElectraPlayerRuntime/Private/Runtime/Player/HLS/PlaylistParserHLS.cpp
@@ -5,9 +5,10 @@
 namespace Electra
 {
 
-FErrorDetail FPlaylistParserHLS::Parse(const FString& InM3U8, const FString& InEffectiveURL)
+FErrorDetail FPlaylistParserHLS::Parse(const FString& InM3U8, const FString& InEffectiveURL, const TArray<FElectraHTTPStreamHeader>& InResponseHeaders)
 {
 	PlaylistURL = InEffectiveURL;
+	ResponseHeaders = InResponseHeaders;
 	FURL_RFC3986 UrlParser;
 	UrlParser.Parse(InEffectiveURL);
 	UrlParser.GetQueryParams(QueryParameters, true);
@@ -175,6 +176,7 @@ FPlaylistParserHLS::EParseState FPlaylistParserHLS::ParseEXT(FErrorDetail& OutEr
 	else if (IsTag(InOutIt, TAGC("#EXT-X-SESSION-KEY")))			return ParseEXT_X_SESSION_KEY(OutErr, InOutIt);
 	else if (IsTag(InOutIt, TAGC("#EXT-X-CONTENT-STEERING")))		return ParseEXT_X_CONTENT_STEERING(OutErr, InOutIt);
 	else if (IsTag(InOutIt, TAGC("#EXT-X-SERVER-CONTROL")))			return ParseEXT_X_SERVER_CONTROL(OutErr, InOutIt);
+	else if (IsTag(InOutIt, TAGC("#EXT-X-SESSION-DATA")))			return ParseEXT_X_SESSION_DATA(OutErr, InOutIt);
 
 	/*
 	else if (IsTag(InOutIt, TAGC("#EXT-X-I-FRAME-STREAM-INF")))		return ParseEXT_X_I_FRAME_STREAM_INF(OutErr, InOutIt);
@@ -187,7 +189,6 @@ FPlaylistParserHLS::EParseState FPlaylistParserHLS::ParseEXT(FErrorDetail& OutEr
 	else if (IsTag(InOutIt, TAGC("#EXT-X-SKIP")))					return ParseEXTOther(OutErr, InOutIt);
 	else if (IsTag(InOutIt, TAGC("#EXT-X-PRELOAD-HINT")))			return ParseEXTOther(OutErr, InOutIt);
 	else if (IsTag(InOutIt, TAGC("#EXT-X-RENDITION-REPORT")))		return ParseEXTOther(OutErr, InOutIt);
-	else if (IsTag(InOutIt, TAGC("#EXT-X-SESSION-DATA")))			return ParseEXTOther(OutErr, InOutIt);
 	else if (IsTag(InOutIt, TAGC("#EXT-X-ALLOW-CACHE")))			return ParseEXTOther(OutErr, InOutIt);	// was removed with HLS version 7
 	*/
 	else															return ParseEXTOther(OutErr, InOutIt);
@@ -371,6 +372,12 @@ FPlaylistParserHLS::EParseState FPlaylistParserHLS::ParseEXT_X_SERVER_CONTROL(FE
 	return ParseEXT_X_Common(OutErr, InOutIt, EAttrType::List, EExtTag::EXT_X_SERVER_CONTROL, EPlaylistTagType::Either);
 }
 
+FPlaylistParserHLS::EParseState FPlaylistParserHLS::ParseEXT_X_SESSION_DATA(FErrorDetail& OutErr, StringHelpers::FStringIterator& InOutIt)
+{
+	return ParseEXT_X_Common(OutErr, InOutIt, EAttrType::List, EExtTag::EXT_X_SESSION_DATA, EPlaylistTagType::MultiVariantOnly);
+}
+
+
 bool FPlaylistParserHLS::ParseOneAttribute(FErrorDetail& OutErr, FString& OutValue, bool& bOutWasQuoted, TArray<FString>* OutSubstitutions, StringHelpers::FStringIterator& InOutIt)
 {
 	if (InOutIt)
diff --git a/Engine/Plugins/Media/ElectraPlayer/Source/ElectraPlayerRuntime/Private/Runtime/Player/HLS/PlaylistParserHLS.h b/Engine/Plugins/Media/ElectraPlayer/Source/ElectraPlayerRuntime/Private/Runtime/Player/HLS/PlaylistParserHLS.h
--- a/Engine/Plugins/Media/ElectraPlayer/Source/ElectraPlayerRuntime/Private/Runtime/Player/HLS/PlaylistParserHLS.h
+++ b/Engine/Plugins/Media/ElectraPlayer/Source/ElectraPlayerRuntime/Private/Runtime/Player/HLS/PlaylistParserHLS.h
@@ -6,6 +6,7 @@
 #include "ErrorDetail.h"
 #include "Utilities/StringHelpers.h"
 #include "Utilities/URLParser.h"
+#include "ElectraHTTPStream.h"
 
 namespace Electra
 {
@@ -14,7 +15,7 @@ class FPlaylistParserHLS
 {
 public:
 	FPlaylistParserHLS() = default;
-	FErrorDetail Parse(const FString& InM3U8, const FString& InEffectiveURL);
+	FErrorDetail Parse(const FString& InM3U8, const FString& InEffectiveURL, const TArray<FElectraHTTPStreamHeader>& InResponseHeaders);
 
 	enum class EPlaylistType
 	{
@@ -26,6 +27,9 @@ public:
 	const FString& GetURL() const
 	{ return PlaylistURL; }
 
+	const TArray<FElectraHTTPStreamHeader>& GetResponseHeaders() const
+	{ return ResponseHeaders; }
+
 	bool IsMultiVariantPlaylist() const
 	{ return bHaveMultiVariantTag; }
 
@@ -228,6 +232,7 @@ private:
 	EParseState ParseEXT_X_SESSION_KEY(FErrorDetail& OutErr, StringHelpers::FStringIterator& InOutIt);
 	EParseState ParseEXT_X_CONTENT_STEERING(FErrorDetail& OutErr, StringHelpers::FStringIterator& InOutIt);
 	EParseState ParseEXT_X_SERVER_CONTROL(FErrorDetail& OutErr, StringHelpers::FStringIterator& InOutIt);
+	EParseState ParseEXT_X_SESSION_DATA(FErrorDetail& OutErr, StringHelpers::FStringIterator& InOutIt);
 	EParseState ParseEXTOther(FErrorDetail& OutErr, StringHelpers::FStringIterator& InOutIt);
 	enum class EAttrType
 	{
@@ -243,6 +248,7 @@ private:
 
 	// The query parameters from the playlist URL. Use an array over a map to maintain case!
 	TArray<Electra::FURL_RFC3986::FQueryParam> QueryParameters;
+	TArray<FElectraHTTPStreamHeader> ResponseHeaders;
 	FString PlaylistURL;
 
 	TArray<TUniquePtr<FElement>> Elements;
diff --git a/Engine/Plugins/Media/ElectraPlayer/Source/ElectraPlayerRuntime/Private/Runtime/Player/PlayerSessionServices.h b/Engine/Plugins/Media/ElectraPlayer/Source/ElectraPlayerRuntime/Private/Runtime/Player/PlayerSessionServices.h
--- a/Engine/Plugins/Media/ElectraPlayer/Source/ElectraPlayerRuntime/Private/Runtime/Player/PlayerSessionServices.h
+++ b/Engine/Plugins/Media/ElectraPlayer/Source/ElectraPlayerRuntime/Private/Runtime/Player/PlayerSessionServices.h
@@ -7,6 +7,8 @@
 #include "InfoLog.h"
 #include "StreamTypes.h"
 
+struct FElectraHTTPStreamHeader;
+
 namespace Electra
 {
 	class ISynchronizedUTCTime;
@@ -89,7 +91,7 @@ namespace Electra
 
 		/**
 		 * Returns the optional external data reader.
-		 * 
+		 *
 		 * @return Pointer to the external data reader, or nullptr if none is to be used.
 		 */
 		virtual TSharedPtrTS<IExternalDataReader> GetExternalDataReader() = 0;
@@ -167,6 +169,31 @@ namespace Electra
 			virtual ~IPlayEndReason() = default;
 		};
 		virtual void SetPlaybackEnd(const FTimeValue& InEndAtTime, EPlayEndReason InEndingReason, TSharedPtrTS<IPlayEndReason> InCustomManifestObject) = 0;
+
+
+
+		/**
+		 * Called when a non-standard property is encountered in the main playlist (not a variant playlist!).
+		 * This may be passed on to registered listeners to evaluate if the playlist meets certain custom requirements.
+		 * Even if there are no such properties in the playlist this method is called at least once with an empty
+		 * property to signal the end of the property list, which is useful if playback must be rejected when no
+		 * custom properties are present (if they have been removed by some playlist transformation service).
+		 */
+		enum class ECustomPropertyResult
+		{
+			// Default handling, depending on where the property occurs.
+			Default,
+			// Property accepted, continue.
+			Accept,
+			// Property rejected, fail!
+			Reject
+		};
+		struct FPlaylistProperty
+		{
+			FString Tag;
+			FString Value;
+		};
+		virtual ECustomPropertyResult ValidateMainPlaylistCustomProperty(const FString& InProtocol, const FString& InPlaylistURL, const TArray<FElectraHTTPStreamHeader>& InPlaylistFetchResponseHeaders, const FPlaylistProperty& InProperty) = 0;
 	};
 
 
diff --git a/Engine/Plugins/Media/ElectraPlayer/Source/ElectraPlayerRuntime/Private/Runtime/Player/mkv/PlaylistReaderMKV.cpp b/Engine/Plugins/Media/ElectraPlayer/Source/ElectraPlayerRuntime/Private/Runtime/Player/mkv/PlaylistReaderMKV.cpp
--- a/Engine/Plugins/Media/ElectraPlayer/Source/ElectraPlayerRuntime/Private/Runtime/Player/mkv/PlaylistReaderMKV.cpp
+++ b/Engine/Plugins/Media/ElectraPlayer/Source/ElectraPlayerRuntime/Private/Runtime/Player/mkv/PlaylistReaderMKV.cpp
@@ -359,6 +359,9 @@ void FPlaylistReaderMKV::WorkerThread()
 
 				LastErrorDetail = Manifest->Build(MKVParser, PlaylistURL, ConnectionInfo);
 
+				// Let the external registry know that we have no properties with an end-of-properties call.
+				PlayerSessionServices->ValidateMainPlaylistCustomProperty(GetPlaylistType(), PlaylistURL, TArray<FElectraHTTPStreamHeader>(), IPlayerSessionServices::FPlaylistProperty());
+
 				// Notify that the "variant playlists" are ready. There are no variants in an mkv, but this is the trigger that the playlists are all set up and are good to go now.
 				PlayerSessionServices->SendMessageToPlayer(IPlaylistReader::PlaylistLoadedMessage::Create(LastErrorDetail, &ConnectionInfo, Playlist::EListType::Variant, Playlist::ELoadType::Initial, Attempt));
 			}
diff --git a/Engine/Plugins/Media/ElectraPlayer/Source/ElectraPlayerRuntime/Private/Runtime/Player/mp4/PlaylistReaderMP4.cpp b/Engine/Plugins/Media/ElectraPlayer/Source/ElectraPlayerRuntime/Private/Runtime/Player/mp4/PlaylistReaderMP4.cpp
--- a/Engine/Plugins/Media/ElectraPlayer/Source/ElectraPlayerRuntime/Private/Runtime/Player/mp4/PlaylistReaderMP4.cpp
+++ b/Engine/Plugins/Media/ElectraPlayer/Source/ElectraPlayerRuntime/Private/Runtime/Player/mp4/PlaylistReaderMP4.cpp
@@ -294,6 +294,9 @@ void FPlaylistReaderMP4::WorkerThread()
 
 					LastErrorDetail = Manifest->Build(MP4Parser, MainPlaylistURL, ConnectionInfo);
 
+					// Let the external registry know that we have no properties with an end-of-properties call.
+					PlayerSessionServices->ValidateMainPlaylistCustomProperty(GetPlaylistType(), MainPlaylistURL, TArray<FElectraHTTPStreamHeader>(), IPlayerSessionServices::FPlaylistProperty());
+
 					// Notify that the "variant playlists" are ready. There are no variants in an mp4, but this is the trigger that the playlists are all set up and are good to go now.
 					PlayerSessionServices->SendMessageToPlayer(IPlaylistReader::PlaylistLoadedMessage::Create(LastErrorDetail, &ConnectionInfo, Playlist::EListType::Variant, Playlist::ELoadType::Initial, Attempt));
 				}
diff --git a/Engine/Plugins/Media/ElectraPlayer/Source/ElectraPlayerRuntime/Private/Runtime/Player/mpegaudio/PlaylistReaderMPEGAudio.cpp b/Engine/Plugins/Media/ElectraPlayer/Source/ElectraPlayerRuntime/Private/Runtime/Player/mpegaudio/PlaylistReaderMPEGAudio.cpp
--- a/Engine/Plugins/Media/ElectraPlayer/Source/ElectraPlayerRuntime/Private/Runtime/Player/mpegaudio/PlaylistReaderMPEGAudio.cpp
+++ b/Engine/Plugins/Media/ElectraPlayer/Source/ElectraPlayerRuntime/Private/Runtime/Player/mpegaudio/PlaylistReaderMPEGAudio.cpp
@@ -949,6 +949,9 @@ void FPlaylistReaderMPEGAudio::WorkerThread()
 
 				LastErrorDetail = Manifest->Build(InfoHeader, PlaylistURL);
 
+				// Let the external registry know that we have no properties with an end-of-properties call.
+				PlayerSessionServices->ValidateMainPlaylistCustomProperty(GetPlaylistType(), PlaylistURL, TArray<FElectraHTTPStreamHeader>(), IPlayerSessionServices::FPlaylistProperty());
+
 				// Notify that the "variant playlists" are ready. There are no variants in an mp4, but this is the trigger that the playlists are all set up and are good to go now.
 				PlayerSessionServices->SendMessageToPlayer(IPlaylistReader::PlaylistLoadedMessage::Create(LastErrorDetail, &ConnectionInfo, Playlist::EListType::Variant, Playlist::ELoadType::Initial, Attempt));
 			}
