ElectraPlayer: removal of some no longer used options

#rnx

diff --git a/Engine/Plugins/Media/ElectraPlayer/Source/ElectraPlayerPlugin/Private/ElectraPlayerPlugin.cpp b/Engine/Plugins/Media/ElectraPlayer/Source/ElectraPlayerPlugin/Private/ElectraPlayerPlugin.cpp
--- a/Engine/Plugins/Media/ElectraPlayer/Source/ElectraPlayerPlugin/Private/ElectraPlayerPlugin.cpp
+++ b/Engine/Plugins/Media/ElectraPlayer/Source/ElectraPlayerPlugin/Private/ElectraPlayerPlugin.cpp
@@ -25,6 +25,36 @@ std::atomic<uint32> FElectraPlayerPlugin::NextPlayerUniqueID { 0 };
 
 //-----------------------------------------------------------------------------
 
+namespace ElectraMediaOptions
+{
+	static const FName GetSafeMediaOptions(TEXT("GetSafeMediaOptions"));
+	static const FName InitialAudioLanguage(TEXT("InitialAudioLanguage"));
+	static const FName InitialSubtitleLanguage(TEXT("InitialSubtitleLanguage"));
+	static const FName ElectraNoPreloading(TEXT("ElectraNoPreloading"));
+	static const FName PlaylistProperties(TEXT("playlist_properties"));
+	static const FName ElectraInitialBitrate(TEXT("ElectraInitialBitrate"));
+	static const FName MaxElectraVerticalResolution(TEXT("MaxElectraVerticalResolution"));
+	static const FName MaxElectraVerticalResolutionOf60fpsVideos(TEXT("MaxElectraVerticalResolutionOf60fpsVideos"));
+	static const FName ElectraLivePresentationOffset(TEXT("ElectraLivePresentationOffset"));
+	static const FName ElectraLiveAudioPresentationOffset(TEXT("ElectraLiveAudioPresentationOffset"));
+	static const FName ElectraLiveUseConservativePresentationOffset(TEXT("ElectraLiveUseConservativePresentationOffset"));
+	static const FName ElectraThrowErrorWhenRebuffering(TEXT("ElectraThrowErrorWhenRebuffering"));
+	static const FName ElectraGetDenyStreamCode(TEXT("ElectraGetDenyStreamCode"));
+	static const FName MaxResolutionForMediaStreaming(TEXT("MaxResolutionForMediaStreaming"));
+	static const FName ElectraMaxStreamingBandwidth(TEXT("ElectraMaxStreamingBandwidth"));
+	static const FName ElectraPlayerDataCache(TEXT("ElectraPlayerDataCache"));
+	static const FName Mimetype(TEXT("mimetype"));
+	static const FName CodecOptions[] = { TEXT("excluded_codecs_video"), TEXT("excluded_codecs_audio"), TEXT("excluded_codecs_subtitles"), TEXT("preferred_codecs_video"), TEXT("preferred_codecs_audio"),TEXT("preferred_codecs_subtitles") };
+	static const FName ElectraGetPlaylistData(TEXT("ElectraGetPlaylistData"));
+	static const FName ElectraGetLicenseKeyData(TEXT("ElectraGetLicenseKeyData"));
+	static const FName ElectraGetPlaystartPosFromSeekPositions(TEXT("ElectraGetPlaystartPosFromSeekPositions"));
+
+	static const FName KeyUniquePlayerID(TEXT("UniquePlayerID"));
+	static const FName OptionKeyParseTimecodeInfo(TEXT("parse_timecode_info"));
+}
+
+//-----------------------------------------------------------------------------
+
 FElectraPlayerPlugin::FElectraPlayerPlugin()
 {
 	// Make sure a few assumptions are correct...
@@ -274,46 +304,30 @@ Electra::FVariantValue FElectraPlayerPlugin::QueryOptions(EOptionType Type, cons
 			{
 				case EOptionType::MaxVerticalStreamResolution:
 				{
-					static const FName MaxResolutionOptionKey = TEXT("MaxResolutionForMediaStreaming");
-					return FVariantValue((int64)SafeOptions->GetMediaOption(MaxResolutionOptionKey, (int64)0));
+					return FVariantValue((int64)SafeOptions->GetMediaOption(ElectraMediaOptions::MaxResolutionForMediaStreaming, (int64)0));
 				}
 
 				case EOptionType::MaxBandwidthForStreaming:
 				{
-					static const FName MaxBandwidthOptionKey = TEXT("ElectraMaxStreamingBandwidth");
-					return FVariantValue((int64)SafeOptions->GetMediaOption(MaxBandwidthOptionKey, (int64)0));
+					return FVariantValue((int64)SafeOptions->GetMediaOption(ElectraMediaOptions::ElectraMaxStreamingBandwidth, (int64)0));
 				}
 
 				case EOptionType::PlayListData:
 				{
-					static const FName PlaylistOptionKey = TEXT("ElectraGetPlaylistData");
-					if (SafeOptions->HasMediaOption(PlaylistOptionKey))
+					if (SafeOptions->HasMediaOption(ElectraMediaOptions::ElectraGetPlaylistData))
 					{
 						check(Param.IsType(FVariantValue::EDataType::TypeFString));
-						return FVariantValue(SafeOptions->GetMediaOption(PlaylistOptionKey, Param.GetFString()));
+						return FVariantValue(SafeOptions->GetMediaOption(ElectraMediaOptions::ElectraGetPlaylistData, Param.GetFString()));
 					}
 					break;
 				}
 
 				case EOptionType::LicenseKeyData:
 				{
-					static const FName LicenseKeyDataOptionKey = TEXT("ElectraGetLicenseKeyData");
-					if (SafeOptions->HasMediaOption(LicenseKeyDataOptionKey))
-					{
-						check(Param.IsType(FVariantValue::EDataType::TypeFString));
-						return FVariantValue(SafeOptions->GetMediaOption(LicenseKeyDataOptionKey, Param.GetFString()));
-					}
-					break;
-				}
-
-				case EOptionType::MediaMetadataUpdate:
-				{
-					static const FName MetadataUpdateOptionKey = TEXT("ElectraMetaDataUpdate");
-					if (SafeOptions->HasMediaOption(MetadataUpdateOptionKey))
+					if (SafeOptions->HasMediaOption(ElectraMediaOptions::ElectraGetLicenseKeyData))
 					{
 						check(Param.IsType(FVariantValue::EDataType::TypeFString));
-						// This only provides metadata, the return value of the Get is of no consequence.
-						SafeOptions->GetMediaOption(MetadataUpdateOptionKey, Param.GetFString());
+						return FVariantValue(SafeOptions->GetMediaOption(ElectraMediaOptions::ElectraGetLicenseKeyData, Param.GetFString()));
 					}
 					break;
 				}
@@ -334,15 +348,14 @@ Electra::FVariantValue FElectraPlayerPlugin::QueryOptions(EOptionType Type, cons
 
 				case EOptionType::PlaystartPosFromSeekPositions:
 				{
-					static const FName PlaystartOptionKey = TEXT("ElectraGetPlaystartPosFromSeekPositions");
-					if (SafeOptions->HasMediaOption(PlaystartOptionKey))
+					if (SafeOptions->HasMediaOption(ElectraMediaOptions::ElectraGetPlaystartPosFromSeekPositions))
 					{
 						check(Param.IsType(FVariantValue::EDataType::TypeSharedPointer));
 
 						TSharedPtr<TArray<FTimespan>, ESPMode::ThreadSafe> PosArray = Param.GetSharedPointer<TArray<FTimespan>>();
 						if (PosArray.IsValid())
 						{
-							TSharedPtr<FElectraSeekablePositions, ESPMode::ThreadSafe> Res = StaticCastSharedPtr<FElectraSeekablePositions, IMediaOptions::FDataContainer, ESPMode::ThreadSafe>(SafeOptions->GetMediaOption(PlaystartOptionKey, MakeShared<FElectraSeekablePositions, ESPMode::ThreadSafe>(*PosArray)));
+							TSharedPtr<FElectraSeekablePositions, ESPMode::ThreadSafe> Res = StaticCastSharedPtr<FElectraSeekablePositions, IMediaOptions::FDataContainer, ESPMode::ThreadSafe>(SafeOptions->GetMediaOption(ElectraMediaOptions::ElectraGetPlaystartPosFromSeekPositions, MakeShared<FElectraSeekablePositions, ESPMode::ThreadSafe>(*PosArray)));
 							if (Res.IsValid() && Res->Data.Num())
 							{
 								return FVariantValue(int64(Res->Data[0].GetTicks())); // return HNS
@@ -543,7 +556,7 @@ bool FElectraPlayerPlugin::Open(const FString& Url, const IMediaOptions* Options
 
 	// Get the safe option interface to poll for changes during playback.
 	CallbackPointerLock.Lock();
-	OptionInterface = StaticCastSharedPtr<IElectraSafeMediaOptionInterface>(Options->GetMediaOption(TEXT("GetSafeMediaOptions"), TSharedPtr<IElectraSafeMediaOptionInterface, ESPMode::ThreadSafe>()));
+	OptionInterface = StaticCastSharedPtr<IElectraSafeMediaOptionInterface>(Options->GetMediaOption(ElectraMediaOptions::GetSafeMediaOptions, TSharedPtr<IElectraSafeMediaOptionInterface, ESPMode::ThreadSafe>()));
 	CallbackPointerLock.Unlock();
 	UE_LOG(LogElectraPlayerPlugin, Log, TEXT("[%u] IMediaPlayer::Open"), PlayerUniqueID);
 
@@ -565,106 +578,98 @@ bool FElectraPlayerPlugin::Open(const FString& Url, const IMediaOptions* Options
 		const FVariant* Env = InPlayerOptions->InternalCustomOptions.Find(MediaPlayerOptionValues::Environment());
 		Environment = Env ? Env->GetValue<FName>() : Environment;
 	}
-	FString InitialAudioLanguage = Options->GetMediaOption(TEXT("InitialAudioLanguage"), FString());
+	FString InitialAudioLanguage = Options->GetMediaOption(ElectraMediaOptions::InitialAudioLanguage, FString());
 	if (InitialAudioLanguage.Len())
 	{
 		LocalPlaystartOptions.InitialAudioTrackAttributes.Language_ISO639 = InitialAudioLanguage;
 		UE_LOG(LogElectraPlayerPlugin, Verbose, TEXT("[%u] IMediaPlayer::Open: Asking for initial audio language \"%s\""), PlayerUniqueID, *InitialAudioLanguage);
 	}
-	FString InitialSubtitleLanguage = Options->GetMediaOption(TEXT("InitialSubtitleLanguage"), FString());
+	FString InitialSubtitleLanguage = Options->GetMediaOption(ElectraMediaOptions::InitialSubtitleLanguage, FString());
 	if (InitialSubtitleLanguage.Len())
 	{
 		LocalPlaystartOptions.InitialSubtitleTrackAttributes.Language_ISO639 = InitialSubtitleLanguage;
 		UE_LOG(LogElectraPlayerPlugin, Verbose, TEXT("[%u] IMediaPlayer::Open: Asking for initial subtitle language \"%s\""), PlayerUniqueID, *InitialSubtitleLanguage);
 	}
-	bool bNoPreloading = Options->GetMediaOption(TEXT("ElectraNoPreloading"), (bool)false);
+	bool bNoPreloading = Options->GetMediaOption(ElectraMediaOptions::ElectraNoPreloading, (bool)false);
 	if (bNoPreloading)
 	{
 		LocalPlaystartOptions.bDoNotPreload = true;
 		UE_LOG(LogElectraPlayerPlugin, Verbose, TEXT("[%u] IMediaPlayer::Open: No preloading after opening media"), PlayerUniqueID);
 	}
 
-
 	// Set up options to initialize the internal player with.
 	Electra::FParamDict PlayerOptions;
-	static const FName KeyUniquePlayerID(TEXT("UniquePlayerID"));
-	PlayerOptions.Set(KeyUniquePlayerID, Electra::FVariantValue((int64)PlayerUniqueID));
-	const FName OptionsByString[] = { TEXT("excluded_codecs_video"), TEXT("excluded_codecs_audio"), TEXT("excluded_codecs_subtitles"), TEXT("preferred_codecs_video"), TEXT("preferred_codecs_audio"),TEXT("preferred_codecs_subtitles") };
-	for(auto &StringOption : OptionsByString)
+	PlayerOptions.Set(ElectraMediaOptions::KeyUniquePlayerID, Electra::FVariantValue((int64)PlayerUniqueID));
+	for(auto &CodecOption : ElectraMediaOptions::CodecOptions)
 	{
-		FString Value = Options->GetMediaOption(StringOption, FString());
+		FString Value = Options->GetMediaOption(CodecOption, FString());
 		if (Value.Len())
 		{
-			PlayerOptions.Set(StringOption, Electra::FVariantValue(Value));
+			PlayerOptions.Set(CodecOption, Electra::FVariantValue(Value));
 		}
 	}
 	// Required playlist properties?
-	FString PlaylistProperties = Options->GetMediaOption(TEXT("playlist_properties"), FString());
+	FString PlaylistProperties = Options->GetMediaOption(ElectraMediaOptions::PlaylistProperties, FString());
 	if (PlaylistProperties.Len())
 	{
-		PlayerOptions.Set(TEXT("playlist_properties"), Electra::FVariantValue(PlaylistProperties));
+		PlayerOptions.Set(ElectraMediaOptions::PlaylistProperties, Electra::FVariantValue(PlaylistProperties));
 	}
 
-
-	if (InPlayerOptions)
+	if (InPlayerOptions && InPlayerOptions->InternalCustomOptions.Find(MediaPlayerOptionValues::ParseTimecodeInfo()))
 	{
-		if (InPlayerOptions->InternalCustomOptions.Find(MediaPlayerOptionValues::ParseTimecodeInfo()))
-		{
-			static const FName OptionKeyParseTimecodeInfo(TEXT("parse_timecode_info"));
-			PlayerOptions.Set(OptionKeyParseTimecodeInfo, FVariantValue());
-		}
+		PlayerOptions.Set(ElectraMediaOptions::OptionKeyParseTimecodeInfo, FVariantValue());
 	}
 
 	// Check for one-time initialization options that can't be changed during playback.
-	int64 InitialStreamBitrate = Options->GetMediaOption(TEXT("ElectraInitialBitrate"), (int64)-1);
+	int64 InitialStreamBitrate = Options->GetMediaOption(ElectraMediaOptions::ElectraInitialBitrate, (int64)-1);
 	if (InitialStreamBitrate > 0)
 	{
 		PlayerOptions.Set(TEXT("initial_bitrate"), Electra::FVariantValue(InitialStreamBitrate));
 		UE_LOG(LogElectraPlayerPlugin, Verbose, TEXT("[%u] IMediaPlayer::Open: Using initial bitrate of %d bits/second"), PlayerUniqueID, (int32)InitialStreamBitrate);
 	}
-	FString MediaMimeType = Options->GetMediaOption(TEXT("mimetype"), FString());
+	FString MediaMimeType = Options->GetMediaOption(ElectraMediaOptions::Mimetype, FString());
 	if (MediaMimeType.Len())
 	{
 		PlayerOptions.Set(TEXT("mime_type"), Electra::FVariantValue(MediaMimeType));
 		UE_LOG(LogElectraPlayerPlugin, Verbose, TEXT("[%u] IMediaPlayer::Open: Setting media mime type to \"%s\""), PlayerUniqueID, *MediaMimeType);
 	}
-	int64 MaxVerticalHeight = Options->GetMediaOption(TEXT("MaxElectraVerticalResolution"), (int64)-1);
+	int64 MaxVerticalHeight = Options->GetMediaOption(ElectraMediaOptions::MaxElectraVerticalResolution, (int64)-1);
 	if (MaxVerticalHeight > 0)
 	{
 		PlayerOptions.Set(TEXT("max_resoY"), Electra::FVariantValue(MaxVerticalHeight));
 		UE_LOG(LogElectraPlayerPlugin, Verbose, TEXT("[%u] IMediaPlayer::Open: Limiting vertical resolution to %d for all streams"), PlayerUniqueID, (int32)MaxVerticalHeight);
 	}
-	int64 MaxVerticalHeightAt60 = Options->GetMediaOption(TEXT("MaxElectraVerticalResolutionOf60fpsVideos"), (int64)-1);
+	int64 MaxVerticalHeightAt60 = Options->GetMediaOption(ElectraMediaOptions::MaxElectraVerticalResolutionOf60fpsVideos, (int64)-1);
 	if (MaxVerticalHeightAt60 > 0)
 	{
 		PlayerOptions.Set(TEXT("max_resoY_above_30fps"), Electra::FVariantValue(MaxVerticalHeightAt60));
 		UE_LOG(LogElectraPlayerPlugin, Verbose, TEXT("[%u] IMediaPlayer::Open: Limiting vertical resolution to %d for streams >30fps"), PlayerUniqueID, (int32)MaxVerticalHeightAt60);
 	}
-	double LiveEdgeDistanceForNormalPresentation = Options->GetMediaOption(TEXT("ElectraLivePresentationOffset"), (double)-1.0);
+	double LiveEdgeDistanceForNormalPresentation = Options->GetMediaOption(ElectraMediaOptions::ElectraLivePresentationOffset, (double)-1.0);
 	if (LiveEdgeDistanceForNormalPresentation > 0.0)
 	{
 		PlayerOptions.Set(TEXT("seekable_range_live_end_offset"), Electra::FVariantValue(Electra::FTimeValue().SetFromSeconds(LiveEdgeDistanceForNormalPresentation)));
 		UE_LOG(LogElectraPlayerPlugin, Verbose, TEXT("[%u] IMediaPlayer::Open: Setting distance to live edge for normal presentations to %.3f seconds"), PlayerUniqueID, LiveEdgeDistanceForNormalPresentation);
 	}
-	double LiveEdgeDistanceForAudioOnlyPresentation = Options->GetMediaOption(TEXT("ElectraLiveAudioPresentationOffset"), (double)-1.0);
+	double LiveEdgeDistanceForAudioOnlyPresentation = Options->GetMediaOption(ElectraMediaOptions::ElectraLiveAudioPresentationOffset, (double)-1.0);
 	if (LiveEdgeDistanceForAudioOnlyPresentation > 0.0)
 	{
 		PlayerOptions.Set(TEXT("seekable_range_live_end_offset_audioonly"), Electra::FVariantValue(Electra::FTimeValue().SetFromSeconds(LiveEdgeDistanceForAudioOnlyPresentation)));
 		UE_LOG(LogElectraPlayerPlugin, Verbose, TEXT("[%u] IMediaPlayer::Open: Setting distance to live edge for audio-only presentation to %.3f seconds"), PlayerUniqueID, LiveEdgeDistanceForAudioOnlyPresentation);
 	}
-	bool bUseConservativeLiveEdgeDistance = Options->GetMediaOption(TEXT("ElectraLiveUseConservativePresentationOffset"), (bool)false);
+	bool bUseConservativeLiveEdgeDistance = Options->GetMediaOption(ElectraMediaOptions::ElectraLiveUseConservativePresentationOffset, (bool)false);
 	if (bUseConservativeLiveEdgeDistance)
 	{
 		PlayerOptions.Set(TEXT("seekable_range_live_end_offset_conservative"), Electra::FVariantValue(bUseConservativeLiveEdgeDistance));
 		UE_LOG(LogElectraPlayerPlugin, Verbose, TEXT("[%u] IMediaPlayer::Open: Using conservative live edge for distance calculation"), PlayerUniqueID);
 	}
-	bool bThrowErrorWhenRebuffering = Options->GetMediaOption(TEXT("ElectraThrowErrorWhenRebuffering"), (bool)false);
+	bool bThrowErrorWhenRebuffering = Options->GetMediaOption(ElectraMediaOptions::ElectraThrowErrorWhenRebuffering, (bool)false);
 	if (bThrowErrorWhenRebuffering)
 	{
 		PlayerOptions.Set(TEXT("throw_error_when_rebuffering"), Electra::FVariantValue(bThrowErrorWhenRebuffering));
 		UE_LOG(LogElectraPlayerPlugin, Verbose, TEXT("[%u] IMediaPlayer::Open: Throw playback error when rebuffering"), PlayerUniqueID);
 	}
-	FString CDNHTTPStatusDenyStream = Options->GetMediaOption(TEXT("ElectraGetDenyStreamCode"), FString());
+	FString CDNHTTPStatusDenyStream = Options->GetMediaOption(ElectraMediaOptions::ElectraGetDenyStreamCode, FString());
 	if (CDNHTTPStatusDenyStream.Len())
 	{
 		int32 HTTPStatus = -1;
@@ -686,14 +691,14 @@ bool FElectraPlayerPlugin::Open(const FString& Url, const IMediaOptions* Options
 	// Check for options that can be changed during playback and apply them at startup already.
 	// If a media source supports the MaxResolutionForMediaStreaming option then we can override the max resolution.
 	int64 DefaultValue = 0;
-	int64 MaxVerticalStreamResolution = Options->GetMediaOption(TEXT("MaxResolutionForMediaStreaming"), DefaultValue);
+	int64 MaxVerticalStreamResolution = Options->GetMediaOption(ElectraMediaOptions::MaxResolutionForMediaStreaming, DefaultValue);
 	if (MaxVerticalStreamResolution != 0)
 	{
 		UE_LOG(LogElectraPlayerPlugin, Log, TEXT("[%u] IMediaPlayer::Open: Limiting max resolution to %d"), PlayerUniqueID, (int32)MaxVerticalStreamResolution);
 		LocalPlaystartOptions.MaxVerticalStreamResolution = (int32)MaxVerticalStreamResolution;
 	}
 
-	int64 MaxBandwidthForStreaming = Options->GetMediaOption(TEXT("ElectraMaxStreamingBandwidth"), (int64)0);
+	int64 MaxBandwidthForStreaming = Options->GetMediaOption(ElectraMediaOptions::ElectraMaxStreamingBandwidth, (int64)0);
 	if (MaxBandwidthForStreaming > 0)
 	{
 		UE_LOG(LogElectraPlayerPlugin, Log, TEXT("[%u] Limiting max streaming bandwidth to %d bps"), PlayerUniqueID, (int32)MaxBandwidthForStreaming);
@@ -706,7 +711,7 @@ bool FElectraPlayerPlugin::Open(const FString& Url, const IMediaOptions* Options
 	// Check if we can get a segment cache interface for this playback request...
 	TSharedPtr<FElectraPlayerDataCacheContainer, ESPMode::ThreadSafe> ElectraPlayerDataCacheContainer;
 	TSharedPtr<IMediaOptions::FDataContainer, ESPMode::ThreadSafe> ElectraPlayerDataCacheDefaultValue;
-	TSharedPtr<IMediaOptions::FDataContainer, ESPMode::ThreadSafe> DataContainer = Options->GetMediaOption(TEXT("ElectraPlayerDataCache"), ElectraPlayerDataCacheDefaultValue);
+	TSharedPtr<IMediaOptions::FDataContainer, ESPMode::ThreadSafe> DataContainer = Options->GetMediaOption(ElectraMediaOptions::ElectraPlayerDataCache, ElectraPlayerDataCacheDefaultValue);
 	if (DataContainer.IsValid())
 	{
 		ElectraPlayerDataCacheContainer = StaticCastSharedPtr<FElectraPlayerDataCacheContainer, IMediaOptions::FDataContainer, ESPMode::ThreadSafe>(DataContainer);
@@ -756,7 +761,7 @@ void FElectraPlayerPlugin::TickInput(FTimespan DeltaTime, FTimespan Timecode)
 
 FVariant FElectraPlayerPlugin::GetMediaInfo(FName InInfoName) const
 {
-	return Player.IsValid() ?  Player->GetMediaInfo(InInfoName).ToFVariant() : FVariant();
+	return Player.IsValid() ? Player->GetMediaInfo(InInfoName).ToFVariant() : FVariant();
 }
 
 //-----------------------------------------------------------------------------
@@ -1122,5 +1127,3 @@ void FElectraPlayerPlugin::SetLastAudioRenderedSampleTime(FTimespan SampleTime)
 {
 	// No-op.
 }
-
-
diff --git a/Engine/Plugins/Media/ElectraPlayer/Source/ElectraPlayerPlugin/Private/ElectraPlayerPlugin.h b/Engine/Plugins/Media/ElectraPlayer/Source/ElectraPlayerPlugin/Private/ElectraPlayerPlugin.h
--- a/Engine/Plugins/Media/ElectraPlayer/Source/ElectraPlayerPlugin/Private/ElectraPlayerPlugin.h
+++ b/Engine/Plugins/Media/ElectraPlayer/Source/ElectraPlayerPlugin/Private/ElectraPlayerPlugin.h
@@ -100,7 +100,7 @@ public:
 	void TickInput(FTimespan DeltaTime, FTimespan Timecode) override;
 	void SetLastAudioRenderedSampleTime(FTimespan SampleTime) override;
 
-	virtual FVariant GetMediaInfo(FName InInfoName) const override;
+	FVariant GetMediaInfo(FName InInfoName) const override;
 
 	TSharedPtr<TMap<FString, TArray<TUniquePtr<IMediaMetadataItem>>>, ESPMode::ThreadSafe> GetMediaMetadata() const override;
 
diff --git a/Engine/Plugins/Media/ElectraPlayer/Source/ElectraPlayerRuntime/Private/Android/ElectraPlayer_Android.cpp b/Engine/Plugins/Media/ElectraPlayer/Source/ElectraPlayerRuntime/Private/Android/ElectraPlayer_Android.cpp
--- a/Engine/Plugins/Media/ElectraPlayer/Source/ElectraPlayerRuntime/Private/Android/ElectraPlayer_Android.cpp
+++ b/Engine/Plugins/Media/ElectraPlayer/Source/ElectraPlayerRuntime/Private/Android/ElectraPlayer_Android.cpp
@@ -66,22 +66,6 @@ bool FElectraPlayerPlatform::StartupPlatformResources(const Electra::FParamDict&
 }
 
 
-void FElectraPlayer::PlatformNotifyOfOptionChange()
-{
-	if (TSharedPtr<IElectraPlayerAdapterDelegate, ESPMode::ThreadSafe> PinnedAdapterDelegate = AdapterDelegate.Pin())
-	{
-		FVariantValue Value = PinnedAdapterDelegate->QueryOptions(IElectraPlayerAdapterDelegate::EOptionType::Android_VideoSurface);
-		if (Value.IsValid())
-		{
-			if (CurrentPlayer.IsValid() && CurrentPlayer->AdaptivePlayer.IsValid())
-			{
-				check(!"This is deprecated. Do not call any more");
-				//CurrentPlayer->AdaptivePlayer->Android_UpdateSurface(Value.GetSharedPointer<IOptionPointerValueContainer>());
-			}
-		}
-	}
-}
-
 void FElectraPlayer::PlatformSuspendOrResumeDecoders(bool bSuspend, const Electra::FParamDict& InOptions)
 {
 	if (CurrentPlayer.IsValid() && CurrentPlayer->AdaptivePlayer.IsValid())
diff --git a/Engine/Plugins/Media/ElectraPlayer/Source/ElectraPlayerRuntime/Private/Apple/ElectraPlayer_Apple.cpp b/Engine/Plugins/Media/ElectraPlayer/Source/ElectraPlayerRuntime/Private/Apple/ElectraPlayer_Apple.cpp
--- a/Engine/Plugins/Media/ElectraPlayer/Source/ElectraPlayerRuntime/Private/Apple/ElectraPlayer_Apple.cpp
+++ b/Engine/Plugins/Media/ElectraPlayer/Source/ElectraPlayerRuntime/Private/Apple/ElectraPlayer_Apple.cpp
@@ -65,10 +65,6 @@ bool FElectraPlayerPlatform::StartupPlatformResources(const Electra::FParamDict&
 }
 
 
-void FElectraPlayer::PlatformNotifyOfOptionChange()
-{
-}
-
 void FElectraPlayer::PlatformSuspendOrResumeDecoders(bool /*bSuspend*/, const Electra::FParamDict& /*InOptions*/)
 {
 }
diff --git a/Engine/Plugins/Media/ElectraPlayer/Source/ElectraPlayerRuntime/Private/ElectraPlayer.cpp b/Engine/Plugins/Media/ElectraPlayer/Source/ElectraPlayerRuntime/Private/ElectraPlayer.cpp
--- a/Engine/Plugins/Media/ElectraPlayer/Source/ElectraPlayerRuntime/Private/ElectraPlayer.cpp
+++ b/Engine/Plugins/Media/ElectraPlayer/Source/ElectraPlayerRuntime/Private/ElectraPlayer.cpp
@@ -570,14 +570,6 @@ void FElectraPlayer::SuspendOrResumeDecoders(bool bSuspend, const Electra::FPara
 }
 
 
-//-----------------------------------------------------------------------------
-/**
-*/
-void FElectraPlayer::NotifyOfOptionChange()
-{
-	PlatformNotifyOfOptionChange();
-}
-
 //-----------------------------------------------------------------------------
 /**
  * Provides information about the time ranges that are currently available to the
@@ -1167,17 +1159,12 @@ void FElectraPlayer::TriggerFirstSeekIfNecessary()
 		if (PinnedAdapterDelegate.IsValid())
 		{
 			TSharedPtr<TArray<FTimespan>, ESPMode::ThreadSafe> SeekablePositions = MakeShared<TArray<FTimespan>, ESPMode::ThreadSafe>();
-			if (SeekablePositions.IsValid())
+			// Check with the delegate if it wants to start somewhere else.
+			FVariantValue Result = PinnedAdapterDelegate->QueryOptions(IElectraPlayerAdapterDelegate::EOptionType::PlaystartPosFromSeekPositions, FVariantValue(SeekablePositions));
+			if (Result.IsValid())
 			{
-				CurrentPlayer->AdaptivePlayer->GetSeekablePositions(*SeekablePositions);
-
-				// Check with the delegate if it wants to start somewhere else.
-				FVariantValue Result = PinnedAdapterDelegate->QueryOptions(IElectraPlayerAdapterDelegate::EOptionType::PlaystartPosFromSeekPositions, FVariantValue(SeekablePositions));
-				if (Result.IsValid())
-				{
-					check(Result.IsType(FVariantValue::EDataType::TypeInt64));
-					playParam.Time.SetFromHNS(Result.GetInt64());
-				}
+				check(Result.IsType(FVariantValue::EDataType::TypeInt64));
+				playParam.Time.SetFromHNS(Result.GetInt64());
 			}
 		}
 
@@ -1278,16 +1265,6 @@ void FElectraPlayer::SetFrameAccurateSeekMode(bool bEnableFrameAccuracy)
 	}
 }
 
-void FElectraPlayer::ModifyOptions(const FParamDict& InOptionsToSetOrChange, const FParamDict& InOptionsToClear)
-{
-	TSharedPtr<FInternalPlayerImpl, ESPMode::ThreadSafe> LockedPlayer = CurrentPlayer;
-	if (LockedPlayer.IsValid() && LockedPlayer->AdaptivePlayer.IsValid())
-	{
-		LockedPlayer->AdaptivePlayer->ModifyOptions(InOptionsToSetOrChange, InOptionsToClear);
-	}
-}
-
-
 void FElectraPlayer::SetPlaybackRange(const FPlaybackRange& InPlaybackRange)
 {
 	CurrentPlaybackRange = InPlaybackRange;
@@ -2974,14 +2951,6 @@ void FElectraPlayer::HandlePlayerMediaMetadataChanged(const TSharedPtrTS<Electra
 		TSharedPtr<TMap<FString, TArray<TSharedPtr<Electra::IMediaStreamMetadata::IItem, ESPMode::ThreadSafe>>>, ESPMode::ThreadSafe> NewMeta = InMetadata->GetMediaStreamMetadata();
 		CurrentStreamMetadata = MoveTemp(NewMeta);
 		DeferredEvents.Enqueue(IElectraPlayerAdapterDelegate::EPlayerEvent::MetadataChanged);
-
-		TSharedPtr<IElectraPlayerAdapterDelegate, ESPMode::ThreadSafe> PinnedAdapterDelegate = AdapterDelegate.Pin();
-		if (PinnedAdapterDelegate.IsValid())
-		{
-			// Send out the metadata through the option query interface for the time being.
-			// This should be removed in the near future.
-			/*FVariantValue Result =*/ PinnedAdapterDelegate->QueryOptions(IElectraPlayerAdapterDelegate::EOptionType::MediaMetadataUpdate, FVariantValue(InMetadata->GetAsJSON()));
-		}
 	}
 }
 
diff --git a/Engine/Plugins/Media/ElectraPlayer/Source/ElectraPlayerRuntime/Private/ElectraPlayer.h b/Engine/Plugins/Media/ElectraPlayer/Source/ElectraPlayerRuntime/Private/ElectraPlayer.h
--- a/Engine/Plugins/Media/ElectraPlayer/Source/ElectraPlayerRuntime/Private/ElectraPlayer.h
+++ b/Engine/Plugins/Media/ElectraPlayer/Source/ElectraPlayerRuntime/Private/ElectraPlayer.h
@@ -121,8 +121,6 @@ public:
 	bool Seek(const FTimespan& Time, const FSeekParam& Param) override;
 	void SetFrameAccurateSeekMode(bool bEnableFrameAccuracy) override;
 
-	void ModifyOptions(const FParamDict& InOptionsToSetOrChange, const FParamDict& InOptionsToClear) override;
-
 	bool GetAudioTrackFormat(int32 TrackIndex, int32 FormatIndex, FAudioTrackFormat& OutFormat) const override;
 	bool GetVideoTrackFormat(int32 TrackIndex, int32 FormatIndex, FVideoTrackFormat& OutFormat) const override;
 
@@ -138,10 +136,9 @@ public:
 	int32 GetNumVideoStreams(int32 TrackIndex) const override;
 	bool GetVideoStreamFormat(FVideoStreamFormat& OutFormat, int32 InTrackIndex, int32 InStreamIndex) const override;
 	bool GetActiveVideoStreamFormat(FVideoStreamFormat& OutFormat) const override;
-	virtual Electra::FVariantValue GetMediaInfo(FName InInfoName) const override;
+	Electra::FVariantValue GetMediaInfo(FName InInfoName) const override;
 	TSharedPtr<TMap<FString, TArray<TSharedPtr<Electra::IMediaStreamMetadata::IItem, ESPMode::ThreadSafe>>>, ESPMode::ThreadSafe> GetMediaMetadata() const override;
 
-	void NotifyOfOptionChange() override;
 	bool GetStreamBufferInformation(FStreamBufferInfo& OutBufferInformation, EPlayerTrackType InTrackType) const override;
 
 	void SuspendOrResumeDecoders(bool bSuspend, const Electra::FParamDict& InOptions) override;
@@ -323,7 +320,6 @@ private:
 	bool PresentAudioFrame(const IAudioDecoderOutputPtr& DecoderOutput);
 	bool PresentSubtitle(const ISubtitleDecoderOutputPtr& DecoderOutput);
 
-	void PlatformNotifyOfOptionChange();
 	void PlatformSuspendOrResumeDecoders(bool bSuspend, const Electra::FParamDict& InOptions);
 
 	void OnMediaPlayerEventReceived(TSharedPtrTS<IAdaptiveStreamingPlayerAEMSEvent> InEvent, IAdaptiveStreamingPlayerAEMSReceiver::EDispatchMode InDispatchMode);
diff --git a/Engine/Plugins/Media/ElectraPlayer/Source/ElectraPlayerRuntime/Private/Linux/ElectraPlayer_Linux.cpp b/Engine/Plugins/Media/ElectraPlayer/Source/ElectraPlayerRuntime/Private/Linux/ElectraPlayer_Linux.cpp
--- a/Engine/Plugins/Media/ElectraPlayer/Source/ElectraPlayerRuntime/Private/Linux/ElectraPlayer_Linux.cpp
+++ b/Engine/Plugins/Media/ElectraPlayer/Source/ElectraPlayerRuntime/Private/Linux/ElectraPlayer_Linux.cpp
@@ -42,10 +42,6 @@ bool FElectraPlayerPlatform::StartupPlatformResources(const Electra::FParamDict&
 }
 
 
-void FElectraPlayer::PlatformNotifyOfOptionChange()
-{
-}
-
 void FElectraPlayer::PlatformSuspendOrResumeDecoders(bool /*bSuspend*/, const Electra::FParamDict& /*InOptions*/)
 {
 }
diff --git a/Engine/Plugins/Media/ElectraPlayer/Source/ElectraPlayerRuntime/Private/Runtime/Player/AdaptiveStreamingPlayer.cpp b/Engine/Plugins/Media/ElectraPlayer/Source/ElectraPlayerRuntime/Private/Runtime/Player/AdaptiveStreamingPlayer.cpp
--- a/Engine/Plugins/Media/ElectraPlayer/Source/ElectraPlayerRuntime/Private/Runtime/Player/AdaptiveStreamingPlayer.cpp
+++ b/Engine/Plugins/Media/ElectraPlayer/Source/ElectraPlayerRuntime/Private/Runtime/Player/AdaptiveStreamingPlayer.cpp
@@ -693,17 +693,6 @@ void FAdaptiveStreamingPlayer::GetSeekableRange(FTimeRange& OutRange) const
 	PlaybackState.GetSeekableRange(OutRange);
 }
 
-//-----------------------------------------------------------------------------
-/**
- * Fills the provided array with time values that can be seeked to. These are segment start times from the video (or audio if there is no video) track.
- *
- * @param OutPositions
- */
-void FAdaptiveStreamingPlayer::GetSeekablePositions(TArray<FTimespan>& OutPositions) const
-{
-	PlaybackState.GetSeekablePositions(OutPositions);
-}
-
 //-----------------------------------------------------------------------------
 /**
  * Returns the timeline range.
diff --git a/Engine/Plugins/Media/ElectraPlayer/Source/ElectraPlayerRuntime/Private/Runtime/Player/AdaptiveStreamingPlayer.h b/Engine/Plugins/Media/ElectraPlayer/Source/ElectraPlayerRuntime/Private/Runtime/Player/AdaptiveStreamingPlayer.h
--- a/Engine/Plugins/Media/ElectraPlayer/Source/ElectraPlayerRuntime/Private/Runtime/Player/AdaptiveStreamingPlayer.h
+++ b/Engine/Plugins/Media/ElectraPlayer/Source/ElectraPlayerRuntime/Private/Runtime/Player/AdaptiveStreamingPlayer.h
@@ -87,7 +87,7 @@ public:
 	 * Returns a player option value.
 	 */
 	virtual FVariantValue GetMediaInfo(FName InKey) const = 0;
-	
+
 	/**
 	 * Sets the attributes for the stream to start buffering for and playing.
 	 * This must be set before calling SeekTo() or LoadManifest() to have an immediate effect.
@@ -285,8 +285,6 @@ public:
 	virtual void GetSeekableRange(FTimeRange& OutRange) const = 0;
 	//! Returns the timeline range.
 	virtual void GetTimelineRange(FTimeRange& OutRange) const = 0;
-	//! Fills the provided array with time values that can be seeked to. These are segment start times from the video (or audio if there is no video) track.
-	virtual void GetSeekablePositions(TArray<FTimespan>& OutPositions) const = 0;
 	//! Returns true when playback has finished.
 	virtual bool HasEnded() const = 0;
 	//! Returns true when data is being buffered/rebuffered, false otherwise.
diff --git a/Engine/Plugins/Media/ElectraPlayer/Source/ElectraPlayerRuntime/Private/Runtime/Player/AdaptiveStreamingPlayerInternal.h b/Engine/Plugins/Media/ElectraPlayer/Source/ElectraPlayerRuntime/Private/Runtime/Player/AdaptiveStreamingPlayerInternal.h
--- a/Engine/Plugins/Media/ElectraPlayer/Source/ElectraPlayerRuntime/Private/Runtime/Player/AdaptiveStreamingPlayerInternal.h
+++ b/Engine/Plugins/Media/ElectraPlayer/Source/ElectraPlayerRuntime/Private/Runtime/Player/AdaptiveStreamingPlayerInternal.h
@@ -271,7 +271,6 @@ struct FPlaybackState
 	TArray<FTrackMetadata>					VideoTracks;
 	TArray<FTrackMetadata>					AudioTracks;
 	TArray<FTrackMetadata>					SubtitleTracks;
-	TArray<FTimespan>						SeekablePositions;
 	FTimeRange								CurrentSegmentDownloadTimeRange[4];	// 0=video, 1=audio, 2=subtitiles, 3=UNSUPPORTED
 
 
@@ -286,17 +285,6 @@ struct FPlaybackState
 		OutRange = SeekableRange;
 	}
 
-	void SetSeekablePositions(const TArray<FTimespan>& InSeekablePositions)
-	{
-		FScopeLock lock(&Lock);
-		SeekablePositions = InSeekablePositions;
-	}
-	void GetSeekablePositions(TArray<FTimespan>& OutSeekablePositions) const
-	{
-		FScopeLock lock(&Lock);
-		OutSeekablePositions = SeekablePositions;
-	}
-
 	void SetTimelineRange(const FTimeRange& TimeRange)
 	{
 		FScopeLock lock(&Lock);
@@ -1160,7 +1148,6 @@ public:
 	FTimeValue GetPlayPosition() const override;
 	void GetTimelineRange(FTimeRange& OutRange) const override;
 	void GetSeekableRange(FTimeRange& OutRange) const override;
-	void GetSeekablePositions(TArray<FTimespan>& OutPositions) const override;
 	bool HasEnded() const override;
 	bool IsSeeking() const override;
 	bool IsBuffering() const override;
diff --git a/Engine/Plugins/Media/ElectraPlayer/Source/ElectraPlayerRuntime/Private/Runtime/Player/AdaptiveStreamingPlayerManifest.cpp b/Engine/Plugins/Media/ElectraPlayer/Source/ElectraPlayerRuntime/Private/Runtime/Player/AdaptiveStreamingPlayerManifest.cpp
--- a/Engine/Plugins/Media/ElectraPlayer/Source/ElectraPlayerRuntime/Private/Runtime/Player/AdaptiveStreamingPlayerManifest.cpp
+++ b/Engine/Plugins/Media/ElectraPlayer/Source/ElectraPlayerRuntime/Private/Runtime/Player/AdaptiveStreamingPlayerManifest.cpp
@@ -297,13 +297,10 @@ bool FAdaptiveStreamingPlayer::SelectManifest()
 		check(Manifest == nullptr);
 		if (ManifestType != EMediaFormatType::Unknown)
 		{
-			TArray<FTimespan> SeekablePositions;
 			TSharedPtrTS<IManifest> NewPresentation = ManifestReader->GetManifest();
 			check(NewPresentation.IsValid());
 
 			PlaybackState.SetSeekableRange(NewPresentation->GetSeekableTimeRange());
-			NewPresentation->GetSeekablePositions(SeekablePositions);
-			PlaybackState.SetSeekablePositions(SeekablePositions);
 			PlaybackState.SetTimelineRange(NewPresentation->GetTotalTimeRange());
 			PlaybackState.SetDuration(NewPresentation->GetDuration());
 
@@ -382,10 +379,7 @@ void FAdaptiveStreamingPlayer::UpdateManifest()
 {
 	if (Manifest.IsValid())
 	{
-		TArray<FTimespan> SeekablePositions;
 		PlaybackState.SetSeekableRange(Manifest->GetSeekableTimeRange());
-		Manifest->GetSeekablePositions(SeekablePositions);
-		PlaybackState.SetSeekablePositions(SeekablePositions);
 		PlaybackState.SetTimelineRange(Manifest->GetTotalTimeRange());
 		PlaybackState.SetDuration(Manifest->GetDuration());
 	}
diff --git a/Engine/Plugins/Media/ElectraPlayer/Source/ElectraPlayerRuntime/Private/Runtime/Player/DASH/ManifestBuilderDASH.cpp b/Engine/Plugins/Media/ElectraPlayer/Source/ElectraPlayerRuntime/Private/Runtime/Player/DASH/ManifestBuilderDASH.cpp
--- a/Engine/Plugins/Media/ElectraPlayer/Source/ElectraPlayerRuntime/Private/Runtime/Player/DASH/ManifestBuilderDASH.cpp
+++ b/Engine/Plugins/Media/ElectraPlayer/Source/ElectraPlayerRuntime/Private/Runtime/Player/DASH/ManifestBuilderDASH.cpp
@@ -2673,30 +2673,6 @@ FTimeRange FManifestDASHInternal::GetSeekableTimeRange() const
 	return SeekableTimeRange;
 }
 
-void FManifestDASHInternal::GetSeekablePositions(TArray<FTimespan>& OutPositions) const
-{
-	FTimeValue Anchor = GetAnchorTime();
-	for(int32 i=0; i<Periods.Num(); ++i)
-	{
-		if (!Periods[i]->GetIsEarlyPeriod())
-		{
-			if (GetPresentationType() == FManifestDASHInternal::EPresentationType::Static)
-			{
-				// The beginning of a period is a seekable position.
-				OutPositions.Emplace(FTimespan((Periods[i]->GetStart() + Anchor).GetAsHNS()));
-			}
-			else
-			{
-				/*
-					If the period is within the availability window we could add times.
-					But as with static presentations, without knowing the period contents there is nothing
-					we can really add here.
-				*/
-			}
-		}
-	}
-}
-
 FTimeValue FManifestDASHInternal::GetDuration() const
 {
 	// An MPD@mediaPresentationDuration is authoritative in both static and dynamic presentations.
diff --git a/Engine/Plugins/Media/ElectraPlayer/Source/ElectraPlayerRuntime/Private/Runtime/Player/DASH/ManifestDASH.cpp b/Engine/Plugins/Media/ElectraPlayer/Source/ElectraPlayerRuntime/Private/Runtime/Player/DASH/ManifestDASH.cpp
--- a/Engine/Plugins/Media/ElectraPlayer/Source/ElectraPlayerRuntime/Private/Runtime/Player/DASH/ManifestDASH.cpp
+++ b/Engine/Plugins/Media/ElectraPlayer/Source/ElectraPlayerRuntime/Private/Runtime/Player/DASH/ManifestDASH.cpp
@@ -402,15 +402,6 @@ FTimeRange FManifestDASH::GetPlaybackRange(EPlaybackRangeType InRangeType) const
 	return Manifest.IsValid() ? Manifest->GetPlayTimesFromURI(InRangeType) : FTimeRange();
 }
 
-void FManifestDASH::GetSeekablePositions(TArray<FTimespan>& OutPositions) const
-{
-	TSharedPtrTS<FManifestDASHInternal> Manifest(CurrentManifest);
-	if (Manifest.IsValid())
-	{
-		Manifest->GetSeekablePositions(OutPositions);
-	}
-}
-
 FTimeValue FManifestDASH::GetDuration() const
 {
 	TSharedPtrTS<FManifestDASHInternal> Manifest(CurrentManifest);
diff --git a/Engine/Plugins/Media/ElectraPlayer/Source/ElectraPlayerRuntime/Private/Runtime/Player/DASH/ManifestDASH.h b/Engine/Plugins/Media/ElectraPlayer/Source/ElectraPlayerRuntime/Private/Runtime/Player/DASH/ManifestDASH.h
--- a/Engine/Plugins/Media/ElectraPlayer/Source/ElectraPlayerRuntime/Private/Runtime/Player/DASH/ManifestDASH.h
+++ b/Engine/Plugins/Media/ElectraPlayer/Source/ElectraPlayerRuntime/Private/Runtime/Player/DASH/ManifestDASH.h
@@ -25,7 +25,6 @@ public:
 	FTimeRange GetTotalTimeRange() const override;
 	FTimeRange GetSeekableTimeRange() const override;
 	FTimeRange GetPlaybackRange(EPlaybackRangeType InRangeType) const override;
-	void GetSeekablePositions(TArray<FTimespan>& OutPositions) const override;
 	FTimeValue GetDuration() const override;
 	FTimeValue GetDefaultStartTime() const override;
 	void ClearDefaultStartTime() override;
diff --git a/Engine/Plugins/Media/ElectraPlayer/Source/ElectraPlayerRuntime/Private/Runtime/Player/DASH/PlaylistReaderDASH_Internal.h b/Engine/Plugins/Media/ElectraPlayer/Source/ElectraPlayerRuntime/Private/Runtime/Player/DASH/PlaylistReaderDASH_Internal.h
--- a/Engine/Plugins/Media/ElectraPlayer/Source/ElectraPlayerRuntime/Private/Runtime/Player/DASH/PlaylistReaderDASH_Internal.h
+++ b/Engine/Plugins/Media/ElectraPlayer/Source/ElectraPlayerRuntime/Private/Runtime/Player/DASH/PlaylistReaderDASH_Internal.h
@@ -799,7 +799,6 @@ public:
 	FTimeValue GetAnchorTime() const;
 	FTimeRange GetTotalTimeRange() const;
 	FTimeRange GetSeekableTimeRange() const;
-	void GetSeekablePositions(TArray<FTimespan>& OutPositions) const;
 	FTimeValue GetDuration() const;
 	void PrepareDefaultStartTime();
 	FTimeValue GetDefaultStartTime() const;
diff --git a/Engine/Plugins/Media/ElectraPlayer/Source/ElectraPlayerRuntime/Private/Runtime/Player/HLS/PlaylistHLS.cpp b/Engine/Plugins/Media/ElectraPlayer/Source/ElectraPlayerRuntime/Private/Runtime/Player/HLS/PlaylistHLS.cpp
--- a/Engine/Plugins/Media/ElectraPlayer/Source/ElectraPlayerRuntime/Private/Runtime/Player/HLS/PlaylistHLS.cpp
+++ b/Engine/Plugins/Media/ElectraPlayer/Source/ElectraPlayerRuntime/Private/Runtime/Player/HLS/PlaylistHLS.cpp
@@ -2136,11 +2136,6 @@ FTimeRange FActiveHLSPlaylist::GetSeekableTimeRange() const
 	return TimelineMediaAsset.IsValid() ? TimelineMediaAsset->GetSeekableTimeRange() : FTimeRange();
 }
 
-void FActiveHLSPlaylist::GetSeekablePositions(TArray<FTimespan>& OutPositions) const
-{
-	// This method is considered deprecated.
-}
-
 FTimeValue FActiveHLSPlaylist::GetDuration() const
 {
 	return TimelineMediaAsset.IsValid() ? TimelineMediaAsset->GetDuration() : FTimeValue();
diff --git a/Engine/Plugins/Media/ElectraPlayer/Source/ElectraPlayerRuntime/Private/Runtime/Player/HLS/PlaylistHLS.h b/Engine/Plugins/Media/ElectraPlayer/Source/ElectraPlayerRuntime/Private/Runtime/Player/HLS/PlaylistHLS.h
--- a/Engine/Plugins/Media/ElectraPlayer/Source/ElectraPlayerRuntime/Private/Runtime/Player/HLS/PlaylistHLS.h
+++ b/Engine/Plugins/Media/ElectraPlayer/Source/ElectraPlayerRuntime/Private/Runtime/Player/HLS/PlaylistHLS.h
@@ -414,7 +414,6 @@ public:
 	FTimeRange GetTotalTimeRange() const override;
 	FTimeRange GetSeekableTimeRange() const override;
 	FTimeRange GetPlaybackRange(EPlaybackRangeType InRangeType) const override;
-	void GetSeekablePositions(TArray<FTimespan>& OutPositions) const override;
 	FTimeValue GetDuration() const override;
 	FTimeValue GetDefaultStartTime() const override;
 	void ClearDefaultStartTime() override;
diff --git a/Engine/Plugins/Media/ElectraPlayer/Source/ElectraPlayerRuntime/Private/Runtime/Player/Manifest.h b/Engine/Plugins/Media/ElectraPlayer/Source/ElectraPlayerRuntime/Private/Runtime/Player/Manifest.h
--- a/Engine/Plugins/Media/ElectraPlayer/Source/ElectraPlayerRuntime/Private/Runtime/Player/Manifest.h
+++ b/Engine/Plugins/Media/ElectraPlayer/Source/ElectraPlayerRuntime/Private/Runtime/Player/Manifest.h
@@ -316,12 +316,6 @@ namespace Electra
 		 */
 		virtual FTimeRange GetPlaybackRange(EPlaybackRangeType InRangeType) const = 0;
 
-		/**
-		 * Returns the timestamps of the segments from the video or audio track (if no video is present).
-		 * Segments are required to start with a keyframe and can thus be used to start playback with.
-		 */
-		virtual void GetSeekablePositions(TArray<FTimespan>& OutPositions) const = 0;
-
 		/**
 		 * Returns the duration of the assets on the timeline.
 		 * Typically this is the difference of the end and start values of the total time range
diff --git a/Engine/Plugins/Media/ElectraPlayer/Source/ElectraPlayerRuntime/Private/Runtime/Player/mkv/ManifestMKV.h b/Engine/Plugins/Media/ElectraPlayer/Source/ElectraPlayerRuntime/Private/Runtime/Player/mkv/ManifestMKV.h
--- a/Engine/Plugins/Media/ElectraPlayer/Source/ElectraPlayerRuntime/Private/Runtime/Player/mkv/ManifestMKV.h
+++ b/Engine/Plugins/Media/ElectraPlayer/Source/ElectraPlayerRuntime/Private/Runtime/Player/mkv/ManifestMKV.h
@@ -44,10 +44,6 @@ public:
 		return GetTotalTimeRange();
 	}
 	FTimeRange GetPlaybackRange(EPlaybackRangeType InRangeType) const override;
-	void GetSeekablePositions(TArray<FTimespan>& OutPositions) const override
-	{
-		// For the time being we do not return anything here as that would require to iterate the tracks.
-	}
 	FTimeValue GetDuration() const override
 	{ return MediaAsset.IsValid() ? MediaAsset->GetDuration() : FTimeValue(); }
 	FTimeValue GetDefaultStartTime() const override
diff --git a/Engine/Plugins/Media/ElectraPlayer/Source/ElectraPlayerRuntime/Private/Runtime/Player/mp4/ManifestMP4.h b/Engine/Plugins/Media/ElectraPlayer/Source/ElectraPlayerRuntime/Private/Runtime/Player/mp4/ManifestMP4.h
--- a/Engine/Plugins/Media/ElectraPlayer/Source/ElectraPlayerRuntime/Private/Runtime/Player/mp4/ManifestMP4.h
+++ b/Engine/Plugins/Media/ElectraPlayer/Source/ElectraPlayerRuntime/Private/Runtime/Player/mp4/ManifestMP4.h
@@ -45,10 +45,6 @@ public:
 		return GetTotalTimeRange();
 	}
 	FTimeRange GetPlaybackRange(EPlaybackRangeType InRangeType) const override;
-	void GetSeekablePositions(TArray<FTimespan>& OutPositions) const override
-	{
-		// For the time being we do not return anything here as that would require to iterate the tracks.
-	}
 	FTimeValue GetDuration() const override
 	{ return MediaAsset.IsValid() ? MediaAsset->GetDuration() : FTimeValue(); }
 	FTimeValue GetDefaultStartTime() const override
diff --git a/Engine/Plugins/Media/ElectraPlayer/Source/ElectraPlayerRuntime/Private/Runtime/Player/mpegaudio/ManifestMPEGAudio.h b/Engine/Plugins/Media/ElectraPlayer/Source/ElectraPlayerRuntime/Private/Runtime/Player/mpegaudio/ManifestMPEGAudio.h
--- a/Engine/Plugins/Media/ElectraPlayer/Source/ElectraPlayerRuntime/Private/Runtime/Player/mpegaudio/ManifestMPEGAudio.h
+++ b/Engine/Plugins/Media/ElectraPlayer/Source/ElectraPlayerRuntime/Private/Runtime/Player/mpegaudio/ManifestMPEGAudio.h
@@ -81,8 +81,6 @@ public:
 	FTimeRange GetSeekableTimeRange() const override
 	{ return GetTotalTimeRange(); }
 	FTimeRange GetPlaybackRange(EPlaybackRangeType InRangeType) const override;
-	void GetSeekablePositions(TArray<FTimespan>& OutPositions) const override
-	{ }
 	FTimeValue GetDuration() const override
 	{ return MediaAsset.IsValid() ? MediaAsset->GetDuration() : FTimeValue(); }
 	FTimeValue GetDefaultStartTime() const override
diff --git a/Engine/Plugins/Media/ElectraPlayer/Source/ElectraPlayerRuntime/Private/Windows/ElectraPlayer_Windows.cpp b/Engine/Plugins/Media/ElectraPlayer/Source/ElectraPlayerRuntime/Private/Windows/ElectraPlayer_Windows.cpp
--- a/Engine/Plugins/Media/ElectraPlayer/Source/ElectraPlayerRuntime/Private/Windows/ElectraPlayer_Windows.cpp
+++ b/Engine/Plugins/Media/ElectraPlayer/Source/ElectraPlayerRuntime/Private/Windows/ElectraPlayer_Windows.cpp
@@ -135,10 +135,6 @@ bool FElectraPlayerPlatform::StartupPlatformResources(const Electra::FParamDict&
 	return FElectraDecoderResourceManagerWindows::Startup(Callbacks);
 }
 
-void FElectraPlayer::PlatformNotifyOfOptionChange()
-{
-}
-
 void FElectraPlayer::PlatformSuspendOrResumeDecoders(bool /*bSuspend*/, const Electra::FParamDict& /*InOptions*/)
 {
 }
diff --git a/Engine/Plugins/Media/ElectraPlayer/Source/ElectraPlayerRuntime/Public/IElectraPlayerInterface.h b/Engine/Plugins/Media/ElectraPlayer/Source/ElectraPlayerRuntime/Public/IElectraPlayerInterface.h
--- a/Engine/Plugins/Media/ElectraPlayer/Source/ElectraPlayerRuntime/Public/IElectraPlayerInterface.h
+++ b/Engine/Plugins/Media/ElectraPlayer/Source/ElectraPlayerRuntime/Public/IElectraPlayerInterface.h
@@ -42,9 +42,7 @@ public:
 		PlayListData,
 		LicenseKeyData,
 		PlaystartPosFromSeekPositions,
-		MediaMetadataUpdate,
-		CustomAnalyticsMetric,
-		Android_VideoSurface,
+		CustomAnalyticsMetric
 	};
 	virtual Electra::FVariantValue QueryOptions(EOptionType Type, const Electra::FVariantValue& Param = Electra::FVariantValue()) = 0;
 
@@ -269,8 +267,6 @@ public:
 	virtual bool Seek(const FTimespan& Time, const FSeekParam& Param) = 0;
 	virtual void SetFrameAccurateSeekMode(bool bEnableFrameAccuracy) = 0;
 
-	virtual void ModifyOptions(const Electra::FParamDict& InOptionsToSetOrChange, const Electra::FParamDict& InOptionsToClear) = 0;
-
 	struct FAudioTrackFormat
 	{
 		uint32 BitsPerSample;
@@ -320,8 +316,6 @@ public:
 	virtual Electra::FVariantValue GetMediaInfo(FName InInfoName) const = 0;
 	virtual TSharedPtr<TMap<FString, TArray<TSharedPtr<Electra::IMediaStreamMetadata::IItem, ESPMode::ThreadSafe>>>, ESPMode::ThreadSafe> GetMediaMetadata() const = 0;
 
-	virtual void NotifyOfOptionChange() = 0;
-
 	struct FStreamBufferInfo
 	{
 		struct FTimeValue
