ElectraProtron: Creating the track-by-type array in the constructor to avoid out of bound accesses if track queries are made before opening a media source

#rnx

diff --git a/Engine/Plugins/Media/ElectraPlayer/Source/ElectraProtron/Private/Player/ElectraProtronPlayerImpl.cpp b/Engine/Plugins/Media/ElectraPlayer/Source/ElectraProtron/Private/Player/ElectraProtronPlayerImpl.cpp
--- a/Engine/Plugins/Media/ElectraPlayer/Source/ElectraProtron/Private/Player/ElectraProtronPlayerImpl.cpp
+++ b/Engine/Plugins/Media/ElectraPlayer/Source/ElectraProtron/Private/Player/ElectraProtronPlayerImpl.cpp
@@ -25,6 +25,9 @@ FElectraProtronPlayer::FImpl::FImpl()
 	// Create the playback parameter structure that has members changing at any moment in time.
 	// This information is shared with the frame loader.
 	SharedPlayParams = MakeShared<FSharedPlayParams, ESPMode::ThreadSafe>();
+
+	// Create the track-by-type array upfront in case queries to tracks are made before opening a source.
+	UsableTrackArrayIndicesByType.SetNum(UE_ARRAY_COUNT(kCodecTrackIndexMap));
 }
 
 FElectraProtronPlayer::FImpl::~FImpl()
@@ -537,15 +540,7 @@ void FElectraProtronPlayer::FImpl::InternalOpen(const FString& InFilename)
 	}
 
 	// One last pass to count how many usable tracks per type we have.
-	const int32 IndexMap[(int32) ElectraProtronUtils::FCodecInfo::EType::MAX] =
-	{
-		CodecTypeIndex(ElectraProtronUtils::FCodecInfo::EType::Video),
-		CodecTypeIndex(ElectraProtronUtils::FCodecInfo::EType::Audio),
-		CodecTypeIndex(ElectraProtronUtils::FCodecInfo::EType::Subtitle),
-		CodecTypeIndex(ElectraProtronUtils::FCodecInfo::EType::Timecode)
-	};
-	UsableTrackArrayIndicesByType.SetNum(UE_ARRAY_COUNT(IndexMap));
-	for(int32 TkTypIdx=0; TkTypIdx<UE_ARRAY_COUNT(IndexMap); ++TkTypIdx)
+	for(int32 TkTypIdx=0; TkTypIdx<UE_ARRAY_COUNT(kCodecTrackIndexMap); ++TkTypIdx)
 	{
 		for(int32 TkIdx=0; TkIdx<Tracks.Num(); ++TkIdx)
 		{
diff --git a/Engine/Plugins/Media/ElectraPlayer/Source/ElectraProtron/Private/Player/ElectraProtronPlayerImpl.h b/Engine/Plugins/Media/ElectraPlayer/Source/ElectraProtron/Private/Player/ElectraProtronPlayerImpl.h
--- a/Engine/Plugins/Media/ElectraPlayer/Source/ElectraProtron/Private/Player/ElectraProtronPlayerImpl.h
+++ b/Engine/Plugins/Media/ElectraPlayer/Source/ElectraProtron/Private/Player/ElectraProtronPlayerImpl.h
@@ -200,6 +200,13 @@ private:
 	void HandleRateChanges();
 	void HandleSeekRequest(const FSeekRequest& InSeekRequest);
 
+	const int32 kCodecTrackIndexMap[(int32) ElectraProtronUtils::FCodecInfo::EType::MAX] =
+	{
+		CodecTypeIndex(ElectraProtronUtils::FCodecInfo::EType::Video),
+		CodecTypeIndex(ElectraProtronUtils::FCodecInfo::EType::Audio),
+		CodecTypeIndex(ElectraProtronUtils::FCodecInfo::EType::Subtitle),
+		CodecTypeIndex(ElectraProtronUtils::FCodecInfo::EType::Timecode)
+	};
 
 	FConfig Config;
 
