Electra: Added multiple CDN support and ContentSteering for DASH and HLS

#rnx

diff --git a/Engine/Plugins/Media/ElectraPlayer/Source/ElectraPlayerRuntime/Private/ElectraPlayer.cpp b/Engine/Plugins/Media/ElectraPlayer/Source/ElectraPlayerRuntime/Private/ElectraPlayer.cpp
--- a/Engine/Plugins/Media/ElectraPlayer/Source/ElectraPlayerRuntime/Private/ElectraPlayer.cpp
+++ b/Engine/Plugins/Media/ElectraPlayer/Source/ElectraPlayerRuntime/Private/ElectraPlayer.cpp
@@ -2261,7 +2261,7 @@ void FElectraPlayer::HandlePlayerEventPlaylistDownload(const Electra::Metrics::F
 		{
 			// Enqueue a "PlaylistDownload" event.
 			TSharedPtr<FAnalyticsEvent> AnalyticEvent = CreateAnalyticsEvent(kEventNameElectraPlaylistDownload);
-			AnalyticEvent->ParamArray.Add(FAnalyticsEventAttribute(TEXT("URL"), *PlaylistDownloadStats.URL));
+			AnalyticEvent->ParamArray.Add(FAnalyticsEventAttribute(TEXT("URL"), *PlaylistDownloadStats.Url.URL));
 			AnalyticEvent->ParamArray.Add(FAnalyticsEventAttribute(TEXT("Failure"), *PlaylistDownloadStats.FailureReason));
 			AnalyticEvent->ParamArray.Add(FAnalyticsEventAttribute(TEXT("ListType"), Electra::Playlist::GetPlaylistTypeString(PlaylistDownloadStats.ListType)));
 			AnalyticEvent->ParamArray.Add(FAnalyticsEventAttribute(TEXT("LoadType"), Electra::Playlist::GetPlaylistLoadTypeString(PlaylistDownloadStats.LoadType)));
@@ -2529,7 +2529,7 @@ void FElectraPlayer::HandlePlayerEventSegmentDownload(const Electra::Metrics::FS
 	}
 	if (SegmentDownloadStats.bWasSuccessful)
 	{
-		UE_LOG(LogElectraPlayer, VeryVerbose, TEXT("[%u] Downloaded %s segment at bitrate %d: Playback time = %.3fs, duration = %.3fs, download time = %.3fs, URL=%s \"%s\""), InstanceID, Electra::GetStreamTypeName(SegmentDownloadStats.StreamType), SegmentDownloadStats.Bitrate, SegmentDownloadStats.PresentationTime, SegmentDownloadStats.Duration, SegmentDownloadStats.TimeToDownload, *SegmentDownloadStats.Range, *SanitizeMessage(SegmentDownloadStats.URL));
+		UE_LOG(LogElectraPlayer, VeryVerbose, TEXT("[%u] Downloaded %s segment at bitrate %d: Playback time = %.3fs, duration = %.3fs, download time = %.3fs, URL=%s \"%s\""), InstanceID, Electra::GetStreamTypeName(SegmentDownloadStats.StreamType), SegmentDownloadStats.Bitrate, SegmentDownloadStats.PresentationTime, SegmentDownloadStats.Duration, SegmentDownloadStats.TimeToDownload, *SegmentDownloadStats.Range, *SanitizeMessage(SegmentDownloadStats.URL.URL));
 	}
 	else if (SegmentDownloadStats.bWasAborted)
 	{
@@ -2560,7 +2560,7 @@ void FElectraPlayer::HandlePlayerEventSegmentDownload(const Electra::Metrics::FS
 		{
 			// Enqueue a "SegmentIssue" event.
 			TSharedPtr<FAnalyticsEvent> AnalyticEvent = CreateAnalyticsEvent(kEventNameElectraSegmentIssue);
-			AnalyticEvent->ParamArray.Add(FAnalyticsEventAttribute(TEXT("URL"), *SegmentDownloadStats.URL));
+			AnalyticEvent->ParamArray.Add(FAnalyticsEventAttribute(TEXT("URL"), *SegmentDownloadStats.URL.URL));
 			AnalyticEvent->ParamArray.Add(FAnalyticsEventAttribute(TEXT("Failure"), *SegmentDownloadStats.FailureReason));
 			AnalyticEvent->ParamArray.Add(FAnalyticsEventAttribute(TEXT("SegmentType"), Electra::Metrics::GetSegmentTypeString(SegmentDownloadStats.SegmentType)));
 			AnalyticEvent->ParamArray.Add(FAnalyticsEventAttribute(TEXT("HTTPStatus"), SegmentDownloadStats.HTTPStatusCode));
diff --git a/Engine/Plugins/Media/ElectraPlayer/Source/ElectraPlayerRuntime/Private/Runtime/HTTP/HTTPManager.cpp b/Engine/Plugins/Media/ElectraPlayer/Source/ElectraPlayerRuntime/Private/Runtime/HTTP/HTTPManager.cpp
--- a/Engine/Plugins/Media/ElectraPlayer/Source/ElectraPlayerRuntime/Private/Runtime/HTTP/HTTPManager.cpp
+++ b/Engine/Plugins/Media/ElectraPlayer/Source/ElectraPlayerRuntime/Private/Runtime/HTTP/HTTPManager.cpp
@@ -30,11 +30,6 @@ DECLARE_CYCLE_STAT(TEXT("FElectraHttpManager::WorkerThread"), STAT_ElectraPlayer
 
 DEFINE_LOG_CATEGORY(LogElectraHTTPManager);
 
-
-#define ELECTRA_HTTPMANAGER_USER_AGENT					TEXT("X-UnrealEngine-Agent")
-//#define ELECTRA_HTTPMANAGER_DEBUG_PROXY_ADDRESS			TEXT("10.29.4.42:8888")
-//#define ELECTRA_HTTPMANAGER_ALLOW_UNSAFE_CONNECTIONS	1
-
 namespace Electra
 {
 
@@ -430,7 +425,7 @@ namespace Electra
 
 	FString IElectraHttpManager::GetDefaultUserAgent()
 	{
-		return FString(ELECTRA_HTTPMANAGER_USER_AGENT);
+		return FString(TEXT("X-UnrealEngine-Agent"));
 	}
 
 
@@ -717,9 +712,6 @@ namespace Electra
 				Handle->HttpRequest->SetUserAgent(GetDefaultUserAgent());
 			}
 			Handle->HttpRequest->AllowCompression(!Request->Parameters.AcceptEncoding.GetWithDefault(TEXT("")).Equals(TEXT("identity")));
-			#if defined(ELECTRA_HTTPMANAGER_ALLOW_UNSAFE_CONNECTIONS) && ELECTRA_HTTPMANAGER_ALLOW_UNSAFE_CONNECTIONS
-				Handle->HttpRequest->AllowUnsafeRequestsForDebugging();
-			#endif
 			if (Handle->ActiveResponse.Range.IsSet())
 			{
 				Handle->HttpRequest->SetRange(Handle->ActiveResponse.Range.GetString());
@@ -1026,9 +1018,6 @@ namespace Electra
 		LLM_SCOPE(ELLMTag::ElectraPlayer);
 
 		FParamDict HttpOptions;
-		#ifdef ELECTRA_HTTPMANAGER_DEBUG_PROXY_ADDRESS
-			HttpOptions.Set(FName(TEXT("proxy")), FVariantValue(FString(ELECTRA_HTTPMANAGER_DEBUG_PROXY_ADDRESS)));
-		#endif
 		HTTPStreamHandler = IElectraHTTPStream::Create(HttpOptions);
 		if (HTTPStreamHandler.IsValid())
 		{
@@ -1752,32 +1741,41 @@ namespace Electra
 	bool FElectraHttpManager::FDataUrl::SetData(const FString& InUrl)
 	{
 		// See https://en.wikipedia.org/wiki/Data_URI_scheme
-		int32 CommaPos = INDEX_NONE;
-		if (InUrl.FindChar(TCHAR(','), CommaPos))
+		FURL_RFC3986 du;
+		TArray<FString> dp;
+		if (du.Parse(InUrl) && du.GetPathComponents(dp) && dp.Num() > 0)
 		{
-			FString MediaType = InUrl.Mid(5, CommaPos-5);
-			// Base64 encoded?
-			bool bIsBase64Encoded = MediaType.EndsWith(TEXT(";base64"), ESearchCase::CaseSensitive);
-			MimeType = bIsBase64Encoded ? MediaType.LeftChop(7) : MediaType;
-			if (MimeType.IsEmpty())
-			{
-				MimeType = TEXT("text/plain;charset=US-ASCII");
+			// The last component, even if empty, is the data.
+			// The second to last may be `base64` to indicate the encoding
+			// The first, unless it is the `base64`, tends to be the mime type
+			// Anything in between we do not care about, not even the charset.
+			int32 base64Pos = INDEX_NONE;
+			if (dp.Find(TEXT("base64"), base64Pos))
+			{
+				if (!FBase64::Decode(dp.Last(), Data))
+				{
+					return false;
+				}
 			}
-			if (bIsBase64Encoded)
+			else
 			{
-				if (FBase64::Decode(InUrl.Mid(CommaPos + 1), Data))
+				FString unesc;
+				if (!FURL_RFC3986::UrlDecode(unesc, dp.Last()))
 				{
-					return true;
+					return false;
 				}
+				StringHelpers::StringToArray(Data, unesc);
+			}
+			// If `base64` is not the first then we assume the first to be the mime type
+			if (base64Pos > 0)
+			{
+				MimeType = dp[0];
 			}
 			else
 			{
-				// The data is 8 bit plain text. We need to convert it back.
-				FTCHARToUTF8 cnv(*InUrl + CommaPos + 1);
-				Data.AddUninitialized(cnv.Length());
-				FMemory::Memcpy(Data.GetData(), cnv.Get(), cnv.Length());
-				return true;
+				MimeType = TEXT("text/plain;charset=US-ASCII");
 			}
+			return true;
 		}
 		return false;
 	}
diff --git a/Engine/Plugins/Media/ElectraPlayer/Source/ElectraPlayerRuntime/Private/Runtime/Player/ABRRules/ABRLiveStream.cpp b/Engine/Plugins/Media/ElectraPlayer/Source/ElectraPlayerRuntime/Private/Runtime/Player/ABRRules/ABRLiveStream.cpp
--- a/Engine/Plugins/Media/ElectraPlayer/Source/ElectraPlayerRuntime/Private/Runtime/Player/ABRRules/ABRLiveStream.cpp
+++ b/Engine/Plugins/Media/ElectraPlayer/Source/ElectraPlayerRuntime/Private/Runtime/Player/ABRRules/ABRLiveStream.cpp
@@ -504,7 +504,7 @@ FABRDownloadProgressDecision FABRLiveStream::ReportDownloadProgress(const Metric
 						// If not already on the worst quality, abort the download.
 						if (QualityIndex > 0)
 						{
-							Info->LogMessage(IInfoLog::ELevel::Verbose, FString::Printf(TEXT("Not connected to server after %.3fs for \"%s\""), SegmentDownloadStats.TimeToDownload, *GetFilenameFromURL(SegmentDownloadStats.URL)));
+							Info->LogMessage(IInfoLog::ELevel::Verbose, FString::Printf(TEXT("Not connected to server after %.3fs for \"%s\""), SegmentDownloadStats.TimeToDownload, *GetFilenameFromURL(SegmentDownloadStats.URL.URL)));
 							Flags = FABRDownloadProgressDecision::EDecisionFlags::eABR_AbortDownload;
 							WorkVars->Decision = FStreamWorkVars::EDecision::NotConnectedAbort;
 						}
diff --git a/Engine/Plugins/Media/ElectraPlayer/Source/ElectraPlayerRuntime/Private/Runtime/Player/AdaptiveStreamingPlayer.cpp b/Engine/Plugins/Media/ElectraPlayer/Source/ElectraPlayerRuntime/Private/Runtime/Player/AdaptiveStreamingPlayer.cpp
--- a/Engine/Plugins/Media/ElectraPlayer/Source/ElectraPlayerRuntime/Private/Runtime/Player/AdaptiveStreamingPlayer.cpp
+++ b/Engine/Plugins/Media/ElectraPlayer/Source/ElectraPlayerRuntime/Private/Runtime/Player/AdaptiveStreamingPlayer.cpp
@@ -1186,6 +1186,11 @@ TSharedPtrTS<IPlaylistReader> FAdaptiveStreamingPlayer::GetManifestReader()
 	return ManifestReader;
 }
 
+TSharedPtrTS<FContentSteeringHandler> FAdaptiveStreamingPlayer::GetContentSteeringHandler()
+{
+	return ContentSteeringHandler;
+}
+
 TSharedPtrTS<IPlayerEntityCache> FAdaptiveStreamingPlayer::GetEntityCache()
 {
 	return EntityCache;
@@ -1402,6 +1407,11 @@ void FAdaptiveStreamingPlayer::InternalHandleOnce()
 				InternalStartoverAtLiveEdge();
 			}
 		}
+		// Handle content steering
+		if (ContentSteeringHandler.IsValid())
+		{
+			ContentSteeringHandler->PeriodicHandle();
+		}
 	}
 }
 
@@ -1668,8 +1678,18 @@ bool FAdaptiveStreamingPlayer::InternalHandleThreadMessages()
 				{
 					PlaybackState.ClearCurrentDownloadRequestTimeRange(reqType);
 
+					Metrics::FSegmentDownloadStats stats;
+					pRequest->GetDownloadStats(stats);
+
+					// Inform content steering
+					FContentSteeringHandler::FStreamParams streamParams;
+					streamParams.bActiveVideo = IsExpectedToStreamNow(EStreamType::Video);
+					streamParams.bActiveAudio = IsExpectedToStreamNow(EStreamType::Audio);
+					streamParams.bActiveSubtitles = IsExpectedToStreamNow(EStreamType::Subtitle);
+					GetContentSteeringHandler()->FinishedDownloadRequestOn(stats, streamParams);
+
 					// Dispatch download event
-					DispatchSegmentDownloadedEvent(pRequest);
+					DispatchEvent(FMetricEvent::ReportSegmentDownload(stats));
 
 					// Dispatch buffer utilization
 					DispatchBufferUtilizationEvent(reqType);
@@ -2022,7 +2042,7 @@ void FAdaptiveStreamingPlayer::HandleSessionMessage(TSharedPtrTS<IPlayerMessage>
 		stats.bDidTimeout = pMsg->GetConnectionInfo().StatusInfo.ErrorCode == HTTP::EStatusErrorCode::ERRCODE_HTTP_CONNECTION_TIMEOUT;
 		stats.ListType = pMsg->GetListType();
 		stats.LoadType = pMsg->GetLoadType();
-		stats.URL = pMsg->GetConnectionInfo().EffectiveURL;
+		stats.Url.URL = pMsg->GetConnectionInfo().EffectiveURL;
 		stats.FailureReason = pMsg->GetConnectionInfo().StatusInfo.ErrorDetail.GetMessage();
 		stats.HTTPStatusCode = pMsg->GetConnectionInfo().StatusInfo.HTTPStatus;
 		stats.RetryNumber = Utils::Max(pMsg->GetAttempts() - 1, 0);
@@ -2122,7 +2142,7 @@ double FAdaptiveStreamingPlayer::GetMinBufferTimeBeforePlayback()
 
 	FTimeValue mbtAbr = StreamSelector.IsValid() ? StreamSelector->GetMinBufferTimeForPlayback(LastBufferingState == EPlayerState::eState_Seeking ? IAdaptiveStreamSelector::EMinBufferType::Seeking
 																							 : LastBufferingState == EPlayerState::eState_Rebuffering ? IAdaptiveStreamSelector::EMinBufferType::Rebuffering
-																							 :	IAdaptiveStreamSelector::EMinBufferType::Initial, Manifest->GetMinBufferTime()) : FTimeValue();
+																							 : IAdaptiveStreamSelector::EMinBufferType::Initial, Manifest->GetMinBufferTime()) : FTimeValue();
 	kMinBufferBeforePlayback = mbtAbr.IsValid() ? mbtAbr.GetAsSeconds() : kMinBufferBeforePlayback;
 	return kMinBufferBeforePlayback;
 }
@@ -2464,6 +2484,10 @@ void FAdaptiveStreamingPlayer::HandleNewOutputData()
 				PrerollVars.Clear();
 				DispatchBufferingEvent(false, LastBufferingState);
 				DispatchEvent(FMetricEvent::ReportPrerollEnd());
+				if (Manifest.IsValid())
+				{
+					Manifest->ReachedStableBuffer();
+				}
 			}
 		}
 	}
@@ -4519,6 +4543,9 @@ void FAdaptiveStreamingPlayer::InternalInitialize()
 		ExternalDataReader = MakeSharedTS<FExternalDataReader>(this);
 	}
 
+	// Create the content steering handler.
+	ContentSteeringHandler = MakeSharedTS<FContentSteeringHandler>(this);
+
 	// Create the ABR stream selector.
 	StreamSelector = IAdaptiveStreamSelector::Create(this, this);
 	AddMetricsReceiver(StreamSelector.Get());
@@ -4605,6 +4632,7 @@ void FAdaptiveStreamingPlayer::InternalClose()
 
 	RemoveMetricsReceiver(StreamSelector.Get());
 	StreamSelector.Reset();
+	ContentSteeringHandler.Reset();
 
 	HttpManager.Reset();
 	EntityCache.Reset();
@@ -4943,6 +4971,11 @@ void FAdaptiveStreamingPlayer::DebugPrint(void* pPlayer, void (*pPrintFN)(void*
 	pPrintFN(pPlayer, "Pipeline state: %s", GetPipelineStateName(PipelineState));
 	pPrintFN(pPlayer, "Stream state  : %s", GetStreamStateName(StreamState));
 	pPrintFN(pPlayer, "-----------------------------------");
+	if (auto csh = ContentSteeringHandler)
+	{
+		pPrintFN(pPlayer, "CDN priority: %s", (ANSICHAR*)StringCast<ANSICHAR>(static_cast<const TCHAR*>( *ContentSteeringHandler->GetCurrentCDNListForDisplay() )).Get());
+		pPrintFN(pPlayer, "CDN access  : %s", (ANSICHAR*)StringCast<ANSICHAR>(static_cast<const TCHAR*>( *ContentSteeringHandler->GetRecentlyAccessedCDNListForDisplay() )).Get());
+	}
 	FTimeRange ActiveRange = PlaybackState.GetActivePlayRange();
 	FTimeRange seekable, timeline;
 	GetSeekableRange(seekable);
@@ -4959,7 +4992,7 @@ void FAdaptiveStreamingPlayer::DebugPrint(void* pPlayer, void (*pPrintFN)(void*
 	pPrintFN(pPlayer, "    is seeking: %s", IsSeeking() ? "Yes" : "No");
 	pPrintFN(pPlayer, "     has ended: %s", HasEnded() ? "Yes" : "No");
 #endif
-	if (StreamSelector.IsValid())
+	if (auto strstel = StreamSelector)
 	{
 		StreamSelector->DebugPrint(pPlayer, pPrintFN);
 	}
diff --git a/Engine/Plugins/Media/ElectraPlayer/Source/ElectraPlayerRuntime/Private/Runtime/Player/AdaptiveStreamingPlayerABR.cpp b/Engine/Plugins/Media/ElectraPlayer/Source/ElectraPlayerRuntime/Private/Runtime/Player/AdaptiveStreamingPlayerABR.cpp
--- a/Engine/Plugins/Media/ElectraPlayer/Source/ElectraPlayerRuntime/Private/Runtime/Player/AdaptiveStreamingPlayerABR.cpp
+++ b/Engine/Plugins/Media/ElectraPlayer/Source/ElectraPlayerRuntime/Private/Runtime/Player/AdaptiveStreamingPlayerABR.cpp
@@ -35,6 +35,7 @@ namespace Electra
 		EHandlingAction PeriodicHandle() override;
 		void MarkStreamAsUnavailable(const FDenylistedStream& DenylistedStream) override;
 		void MarkStreamAsAvailable(const FDenylistedStream& NoLongerDenylistedStream) override;
+		void PathwayChanged(const FString& InNewPathway) override;
 		int64 GetLastBandwidth() override;
 		int64 GetAverageBandwidth() override;
 		int64 GetAverageThroughput() override;
@@ -91,6 +92,7 @@ namespace Electra
 		int32 GetBandwidthCeiling() override
 		{ return BandwidthCeiling; }
 
+		void UpdateCurrentPlaybackPeriods();
 		ESegmentAction SelectSuitableStreamByType(FTimeValue& OutDelay, TSharedPtrTS<const IStreamSegment> CurrentSegment, TSharedPtrTS<IManifest::IPlayPeriod> CurrentPlayPeriod, EStreamType StreamType);
 
 		FTimeValue ABRGetPlayPosition() const override
@@ -217,118 +219,130 @@ namespace Electra
 		}
 	}
 
-	//-----------------------------------------------------------------------------
-	/**
-	 * Sets the playback period from which to select streams now.
-	 *
-	 * @param InStreamType
-	 * @param InCurrentPlayPeriod
-	 */
-	void FAdaptiveStreamSelector::SetCurrentPlaybackPeriod(EStreamType InStreamType, TSharedPtrTS<IManifest::IPlayPeriod> InCurrentPlayPeriod)
+	void FAdaptiveStreamSelector::UpdateCurrentPlaybackPeriods()
 	{
 		FScopeLock lock(&AccessMutex);
-		if (InStreamType == EStreamType::Video)
+		if (CurrentPlayPeriodVideo.IsValid())
 		{
-			CurrentPlayPeriodVideo = InCurrentPlayPeriod;
-			if (CurrentPlayPeriodVideo.IsValid())
+			TSharedPtrTS<ITimelineMediaAsset> Asset = CurrentPlayPeriodVideo->GetMediaAsset();
+			if (Asset.IsValid())
 			{
-				TSharedPtrTS<ITimelineMediaAsset> Asset = CurrentPlayPeriodVideo->GetMediaAsset();
-				if (Asset.IsValid())
+				CurrentVideoAdaptationSetID = CurrentPlayPeriodVideo->GetSelectedAdaptationSetID(EStreamType::Video);
+				StreamInformationVideo.Empty();
+				// Get video stream information
+				if (!CurrentVideoAdaptationSetID.IsEmpty())
 				{
-					CurrentVideoAdaptationSetID = InCurrentPlayPeriod->GetSelectedAdaptationSetID(InStreamType);
-					StreamInformationVideo.Empty();
-					// Get video stream information
-					if (!CurrentVideoAdaptationSetID.IsEmpty())
+					for(int32 nA=0, maxA=Asset->GetNumberOfAdaptationSets(EStreamType::Video); nA<maxA; ++nA)
 					{
-						for(int32 nA=0, maxA=Asset->GetNumberOfAdaptationSets(InStreamType); nA<maxA; ++nA)
+						TSharedPtrTS<IPlaybackAssetAdaptationSet> Adapt = Asset->GetAdaptationSetByTypeAndIndex(EStreamType::Video, nA);
+						if (CurrentVideoAdaptationSetID.Equals(Adapt->GetUniqueIdentifier()))
 						{
-							TSharedPtrTS<IPlaybackAssetAdaptationSet> Adapt = Asset->GetAdaptationSetByTypeAndIndex(InStreamType, nA);
-							if (CurrentVideoAdaptationSetID.Equals(Adapt->GetUniqueIdentifier()))
+							bool bLowLatency = Adapt->IsLowLatencyEnabled();
+							for(int32 i=0, iMax=Adapt->GetNumberOfRepresentations(); i<iMax; ++i)
 							{
-								bool bLowLatency = Adapt->IsLowLatencyEnabled();
-								for(int32 i=0, iMax=Adapt->GetNumberOfRepresentations(); i<iMax; ++i)
+								TSharedPtrTS<IPlaybackAssetRepresentation> Repr = Adapt->GetRepresentationByIndex(i);
+								TSharedPtrTS<FABRStreamInformation> si = MakeSharedTS<FABRStreamInformation>();
+								si->AdaptationSetUniqueID = Adapt->GetUniqueIdentifier();
+								si->RepresentationUniqueID = Repr->GetUniqueIdentifier();
+								si->Bitrate = Repr->GetBitrate();
+								si->QualityIndex = Repr->GetQualityIndex();
+								si->bLowLatencyEnabled = bLowLatency;
+								if (Repr->GetCodecInformation().IsVideoCodec())
 								{
-									TSharedPtrTS<IPlaybackAssetRepresentation> Repr = Adapt->GetRepresentationByIndex(i);
-									TSharedPtrTS<FABRStreamInformation> si = MakeSharedTS<FABRStreamInformation>();
-									si->AdaptationSetUniqueID = Adapt->GetUniqueIdentifier();
-									si->RepresentationUniqueID = Repr->GetUniqueIdentifier();
-									si->Bitrate = Repr->GetBitrate();
-									si->QualityIndex = Repr->GetQualityIndex();
-									si->bLowLatencyEnabled = bLowLatency;
-									if (Repr->GetCodecInformation().IsVideoCodec())
-									{
-										si->Resolution = Repr->GetCodecInformation().GetResolution();
-									}
-									StreamInformationVideo.Push(si);
+									si->Resolution = Repr->GetCodecInformation().GetResolution();
 								}
-								// Sort the representations by ascending bitrate
-								StreamInformationVideo.StableSort([](const TSharedPtrTS<FABRStreamInformation>& a, const TSharedPtrTS<FABRStreamInformation>& b)
-								{
-									return a->Bitrate < b->Bitrate;
-								});
-
-								break;
+								StreamInformationVideo.Push(si);
 							}
+							// Sort the representations by ascending bitrate
+							StreamInformationVideo.StableSort([](const TSharedPtrTS<FABRStreamInformation>& a, const TSharedPtrTS<FABRStreamInformation>& b)
+							{
+								return a->Bitrate < b->Bitrate;
+							});
+
+							break;
 						}
 					}
 				}
 			}
-			else
-			{
-				CurrentVideoAdaptationSetID.Empty();
-				StreamInformationVideo.Empty();
-			}
 		}
-		else if (InStreamType == EStreamType::Audio)
+		else
 		{
-			CurrentPlayPeriodAudio = InCurrentPlayPeriod;
-			if (CurrentPlayPeriodAudio.IsValid())
+			CurrentVideoAdaptationSetID.Empty();
+			StreamInformationVideo.Empty();
+		}
+
+		if (CurrentPlayPeriodAudio.IsValid())
+		{
+			TSharedPtrTS<ITimelineMediaAsset> Asset = CurrentPlayPeriodAudio->GetMediaAsset();
+			if (Asset.IsValid())
 			{
-				TSharedPtrTS<ITimelineMediaAsset> Asset = CurrentPlayPeriodAudio->GetMediaAsset();
-				if (Asset.IsValid())
+				CurrentAudioAdaptationSetID = CurrentPlayPeriodAudio->GetSelectedAdaptationSetID(EStreamType::Audio);
+				StreamInformationAudio.Empty();
+				// Get audio stream information
+				if (!CurrentAudioAdaptationSetID.IsEmpty())
 				{
-					CurrentAudioAdaptationSetID = InCurrentPlayPeriod->GetSelectedAdaptationSetID(InStreamType);
-					StreamInformationAudio.Empty();
-					// Get audio stream information
-					if (!CurrentAudioAdaptationSetID.IsEmpty())
+					for(int32 nA=0, maxA=Asset->GetNumberOfAdaptationSets(EStreamType::Audio); nA<maxA; ++nA)
 					{
-						for(int32 nA=0, maxA=Asset->GetNumberOfAdaptationSets(InStreamType); nA<maxA; ++nA)
+						TSharedPtrTS<IPlaybackAssetAdaptationSet> Adapt = Asset->GetAdaptationSetByTypeAndIndex(EStreamType::Audio, nA);
+						if (CurrentAudioAdaptationSetID.Equals(Adapt->GetUniqueIdentifier()))
 						{
-							TSharedPtrTS<IPlaybackAssetAdaptationSet> Adapt = Asset->GetAdaptationSetByTypeAndIndex(InStreamType, nA);
-							if (CurrentAudioAdaptationSetID.Equals(Adapt->GetUniqueIdentifier()))
+							bool bLowLatency = Adapt->IsLowLatencyEnabled();
+							for(int32 i=0, iMax=Adapt->GetNumberOfRepresentations(); i<iMax; ++i)
 							{
-								bool bLowLatency = Adapt->IsLowLatencyEnabled();
-								for(int32 i=0, iMax=Adapt->GetNumberOfRepresentations(); i<iMax; ++i)
-								{
-									TSharedPtrTS<IPlaybackAssetRepresentation> Repr = Adapt->GetRepresentationByIndex(i);
-									TSharedPtrTS<FABRStreamInformation> si = MakeSharedTS<FABRStreamInformation>();
-									si->AdaptationSetUniqueID = Adapt->GetUniqueIdentifier();
-									si->RepresentationUniqueID = Repr->GetUniqueIdentifier();
-									si->Bitrate = Repr->GetBitrate();
-									si->QualityIndex = Repr->GetQualityIndex();
-									si->bLowLatencyEnabled = bLowLatency;
-									StreamInformationAudio.Push(si);
-								}
-								// Sort the representations by ascending bitrate
-								StreamInformationAudio.StableSort([](const TSharedPtrTS<FABRStreamInformation>& a, const TSharedPtrTS<FABRStreamInformation>& b)
-								{
-									return a->Bitrate < b->Bitrate;
-								});
-
-								break;
+								TSharedPtrTS<IPlaybackAssetRepresentation> Repr = Adapt->GetRepresentationByIndex(i);
+								TSharedPtrTS<FABRStreamInformation> si = MakeSharedTS<FABRStreamInformation>();
+								si->AdaptationSetUniqueID = Adapt->GetUniqueIdentifier();
+								si->RepresentationUniqueID = Repr->GetUniqueIdentifier();
+								si->Bitrate = Repr->GetBitrate();
+								si->QualityIndex = Repr->GetQualityIndex();
+								si->bLowLatencyEnabled = bLowLatency;
+								StreamInformationAudio.Push(si);
 							}
+							// Sort the representations by ascending bitrate
+							StreamInformationAudio.StableSort([](const TSharedPtrTS<FABRStreamInformation>& a, const TSharedPtrTS<FABRStreamInformation>& b)
+							{
+								return a->Bitrate < b->Bitrate;
+							});
+
+							break;
 						}
 					}
 				}
 			}
-			else
-			{
-				CurrentAudioAdaptationSetID.Empty();
-				StreamInformationAudio.Empty();
-			}
 		}
+		else
+		{
+			CurrentAudioAdaptationSetID.Empty();
+			StreamInformationAudio.Empty();
+		}
+
 		// Notify the ABR method that there is a different set of representations now.
-		ABRMethod->RepresentationsChanged(InStreamType, InCurrentPlayPeriod);
+		if (ABRMethod.IsValid())
+		{
+			ABRMethod->RepresentationsChanged(EStreamType::Video, CurrentPlayPeriodVideo);
+			ABRMethod->RepresentationsChanged(EStreamType::Audio, CurrentPlayPeriodAudio);
+		}
+	}
+
+	//-----------------------------------------------------------------------------
+	/**
+	 * Sets the playback period from which to select streams now.
+	 *
+	 * @param InStreamType
+	 * @param InCurrentPlayPeriod
+	 */
+	void FAdaptiveStreamSelector::SetCurrentPlaybackPeriod(EStreamType InStreamType, TSharedPtrTS<IManifest::IPlayPeriod> InCurrentPlayPeriod)
+	{
+		FScopeLock lock(&AccessMutex);
+		if (InStreamType == EStreamType::Video)
+		{
+			CurrentPlayPeriodVideo = InCurrentPlayPeriod;
+		}
+		else if (InStreamType == EStreamType::Audio)
+		{
+			CurrentPlayPeriodAudio = InCurrentPlayPeriod;
+		}
+		UpdateCurrentPlaybackPeriods();
 	}
 
 	//-----------------------------------------------------------------------------
@@ -410,6 +424,19 @@ namespace Electra
 		DenylistedExternally.Remove(InNoLongerDenylistedStream);
 	}
 
+
+	//-----------------------------------------------------------------------------
+	/**
+	 * Called when content steering has changed to a different pathway (HLS)
+	 * This _could_ mean that a whole different set of streams is suddenly available
+	 * to select from, including their IDs, so we need to update our internal
+	 * knowledge about what is available.
+	 */
+	void FAdaptiveStreamSelector::PathwayChanged(const FString& InNewPathway)
+	{
+		UpdateCurrentPlaybackPeriods();
+	}
+
 	//-----------------------------------------------------------------------------
 	/**
 	 * Returns the last measured bandwidth sample in bits per second.
@@ -610,7 +637,11 @@ namespace Electra
 		}
 		else
 		{
-			// Initial request.
+			// Initial request. If there are no streams at all we fail.
+			if (StreamInformationVideo.IsEmpty() && StreamInformationAudio.IsEmpty())
+			{
+				return IAdaptiveStreamSelector::ESegmentAction::Fail;
+			}
 			ESegmentAction ActionV = ESegmentAction::FetchNext;
 			ESegmentAction ActionA = ESegmentAction::FetchNext;
 			if (StreamInformationVideo.Num())
diff --git a/Engine/Plugins/Media/ElectraPlayer/Source/ElectraPlayerRuntime/Private/Runtime/Player/AdaptiveStreamingPlayerABR.h b/Engine/Plugins/Media/ElectraPlayer/Source/ElectraPlayerRuntime/Private/Runtime/Player/AdaptiveStreamingPlayerABR.h
--- a/Engine/Plugins/Media/ElectraPlayer/Source/ElectraPlayerRuntime/Private/Runtime/Player/AdaptiveStreamingPlayerABR.h
+++ b/Engine/Plugins/Media/ElectraPlayer/Source/ElectraPlayerRuntime/Private/Runtime/Player/AdaptiveStreamingPlayerABR.h
@@ -31,13 +31,13 @@ public:
 
 		// Provides the current wallclock time (server real time).
 		virtual FTimeValue ABRGetWallclockTime() const = 0;
-		
+
 		// Returns true if this is a Live presentation, false for VoD.
 		virtual bool ABRIsLive() const = 0;
-		
+
 		// Returns true if playback should be on the Live edge. false if the user intentionally seeked away from the edge.
 		virtual bool ABRShouldPlayOnLiveEdge() const = 0;
-		
+
 		// Provides the current low latency descriptor if low latency playback is active. Returns nullptr if not.
 		virtual TSharedPtrTS<const FLowLatencyDescriptor> ABRGetLowLatencyDescriptor() const = 0;
 
@@ -139,6 +139,7 @@ public:
 	virtual void MarkStreamAsUnavailable(const FDenylistedStream& DenylistedStream) = 0;
 	virtual void MarkStreamAsAvailable(const FDenylistedStream& NoLongerDenylistedStream) = 0;
 
+	virtual void PathwayChanged(const FString& InNewPathway) = 0;
 
 	virtual FABRDownloadProgressDecision ReportDownloadProgress(const Metrics::FSegmentDownloadStats& SegmentDownloadStats) = 0;
 	virtual void ReportDownloadEnd(const Metrics::FSegmentDownloadStats& SegmentDownloadStats) = 0;
diff --git a/Engine/Plugins/Media/ElectraPlayer/Source/ElectraPlayerRuntime/Private/Runtime/Player/AdaptiveStreamingPlayerEvents.cpp b/Engine/Plugins/Media/ElectraPlayer/Source/ElectraPlayerRuntime/Private/Runtime/Player/AdaptiveStreamingPlayerEvents.cpp
--- a/Engine/Plugins/Media/ElectraPlayer/Source/ElectraPlayerRuntime/Private/Runtime/Player/AdaptiveStreamingPlayerEvents.cpp
+++ b/Engine/Plugins/Media/ElectraPlayer/Source/ElectraPlayerRuntime/Private/Runtime/Player/AdaptiveStreamingPlayerEvents.cpp
@@ -30,21 +30,6 @@ void FAdaptiveStreamingPlayer::DispatchBufferingEvent(bool bBegin, FAdaptiveStre
 	}
 }
 
-//-----------------------------------------------------------------------------
-/**
- * Dispatches a 'fragment download' event.
- *
- * @param Request
- */
-void FAdaptiveStreamingPlayer::DispatchSegmentDownloadedEvent(TSharedPtrTS<IStreamSegment> Request)
-{
-	if (Request.IsValid())
-	{
-		Metrics::FSegmentDownloadStats stats;
-		Request->GetDownloadStats(stats);
-		DispatchEvent(FMetricEvent::ReportSegmentDownload(stats));
-	}
-}
 
 
 void FAdaptiveStreamingPlayer::DispatchBufferUtilizationEvent(EStreamType BufferType)
@@ -85,12 +70,12 @@ void FAdaptiveStreamingPlayer::DispatchBufferUtilizationEvent(EStreamType Buffer
 
 namespace ReservedEventSchemeNames
 {
-	static const TArray<FString> IgnoredByWildcard { DASH::Schemes::ManifestEvents::Scheme_urn_mpeg_dash_event_2012, 
-													 DASH::Schemes::ManifestEvents::Scheme_urn_mpeg_dash_event_callback_2015, 
+	static const TArray<FString> IgnoredByWildcard { DASH::Schemes::ManifestEvents::Scheme_urn_mpeg_dash_event_2012,
+													 DASH::Schemes::ManifestEvents::Scheme_urn_mpeg_dash_event_callback_2015,
 													 DASH::Schemes::ManifestEvents::Scheme_urn_mpeg_dash_event_ttfn_2016 };
 };
 
-    
+
 
 
 class FAdaptiveStreamingPlayerAEMSHandler : public IAdaptiveStreamingPlayerAEMSHandler
@@ -98,7 +83,7 @@ class FAdaptiveStreamingPlayerAEMSHandler : public IAdaptiveStreamingPlayerAEMSH
 public:
 	FAdaptiveStreamingPlayerAEMSHandler();
 	virtual ~FAdaptiveStreamingPlayerAEMSHandler();
-	
+
 	virtual void ExecutePendingEvents(FTimeValue CurrentPlaybackTime) override;
 	virtual void FlushEverything() override;
 	virtual void FlushDynamic() override;
@@ -391,7 +376,7 @@ void FAdaptiveStreamingPlayerAEMSHandler::FireEventsInRange(FTimeRange InRange)
 		FTimeValue evS = EnqueuedEvents[i]->Event->GetPresentationTime();
 		FTimeValue evD = EnqueuedEvents[i]->Event->GetDuration();
 		FTimeValue evE = evS + (evD.IsValid() ? evD : FTimeValue::GetZero());
-		
+
 		// Event already expired?
 		if (evE < InRange.Start)
 		{
diff --git a/Engine/Plugins/Media/ElectraPlayer/Source/ElectraPlayerRuntime/Private/Runtime/Player/AdaptiveStreamingPlayerInternal.h b/Engine/Plugins/Media/ElectraPlayer/Source/ElectraPlayerRuntime/Private/Runtime/Player/AdaptiveStreamingPlayerInternal.h
--- a/Engine/Plugins/Media/ElectraPlayer/Source/ElectraPlayerRuntime/Private/Runtime/Player/AdaptiveStreamingPlayerInternal.h
+++ b/Engine/Plugins/Media/ElectraPlayer/Source/ElectraPlayerRuntime/Private/Runtime/Player/AdaptiveStreamingPlayerInternal.h
@@ -24,6 +24,7 @@
 #include "Player/PlayerEntityCache.h"
 #include "Player/AdaptiveStreamingPlayerABR.h"
 #include "Player/AdaptivePlayerOptionKeynames.h"
+#include "Player/ContentSteeringHandler.h"
 
 #include "Templates/UniqueObj.h"
 
@@ -100,7 +101,7 @@ public:
 	{
 	}
 
-	virtual void SetCurrentTime(ERendererType ForRenderer, const FTimeValue& CurrentRenderTime) override
+	void SetCurrentTime(ERendererType ForRenderer, const FTimeValue& CurrentRenderTime) override
 	{
 		FClock* Clk = GetClock(ForRenderer);
 		if (Clk)
@@ -113,7 +114,7 @@ public:
 		}
 	}
 
-	virtual FTimeValue GetInterpolatedRenderTime(ERendererType FromRenderer) override
+	FTimeValue GetInterpolatedRenderTime(ERendererType FromRenderer) override
 	{
 		FClock* Clk = GetClock(FromRenderer);
 		if (Clk)
@@ -1123,7 +1124,7 @@ public:
 
 	void Initialize(const FParamDict& Options) override;
 	void ModifyOptions(const FParamDict& InOptionsToSetOrChange, const FParamDict& InOptionsToClear) override;
-	virtual FVariantValue GetMediaInfo(FName InKey) const override;
+	FVariantValue GetMediaInfo(FName InKey) const override;
 
 	void SetInitialStreamAttributes(EStreamType StreamType, const FStreamSelectionAttributes& InitialSelection) override;
 	void EnableFrameAccurateSeeking(bool bEnabled) override;
@@ -1187,7 +1188,8 @@ private:
 	TSharedPtrTS<IAdaptiveStreamSelector> GetStreamSelector() override;
 	IPlayerStreamFilter* GetStreamFilter() override;
 	const FCodecSelectionPriorities& GetCodecSelectionPriorities(EStreamType ForStream) override;
-	virtual	TSharedPtrTS<IPlaylistReader> GetManifestReader() override;
+	TSharedPtrTS<IPlaylistReader> GetManifestReader() override;
+	TSharedPtrTS<FContentSteeringHandler> GetContentSteeringHandler() override;
 	TSharedPtrTS<IPlayerEntityCache> GetEntityCache() override;
 	TSharedPtrTS<IHTTPResponseCache> GetHTTPResponseCache() override;
 	IAdaptiveStreamingPlayerAEMSHandler* GetAEMSEventHandler() override;
@@ -2259,7 +2261,6 @@ private:
 
 	void DispatchEvent(TSharedPtrTS<FMetricEvent> Event);
 	void DispatchBufferingEvent(bool bBegin, EPlayerState Reason);
-	void DispatchSegmentDownloadedEvent(TSharedPtrTS<IStreamSegment> Request);
 	void DispatchBufferUtilizationEvent(EStreamType BufferType);
 
 	void UpdateDataAvailabilityState(Metrics::FDataAvailabilityChange& DataAvailabilityState, Metrics::FDataAvailabilityChange::EAvailability NewAvailability);
@@ -2353,6 +2354,8 @@ private:
 	ERebufferCause														RebufferCause;
 	bool																bIsClosing;
 
+	TSharedPtrTS<FContentSteeringHandler>								ContentSteeringHandler;
+
 	TSharedPtrTS<IAdaptiveStreamSelector>								StreamSelector;
 	int32																BitrateCeiling;
 	int32																VideoResolutionLimitWidth;
diff --git a/Engine/Plugins/Media/ElectraPlayer/Source/ElectraPlayerRuntime/Private/Runtime/Player/AdaptiveStreamingPlayerMetrics.h b/Engine/Plugins/Media/ElectraPlayer/Source/ElectraPlayerRuntime/Private/Runtime/Player/AdaptiveStreamingPlayerMetrics.h
--- a/Engine/Plugins/Media/ElectraPlayer/Source/ElectraPlayerRuntime/Private/Runtime/Player/AdaptiveStreamingPlayerMetrics.h
+++ b/Engine/Plugins/Media/ElectraPlayer/Source/ElectraPlayerRuntime/Private/Runtime/Player/AdaptiveStreamingPlayerMetrics.h
@@ -4,6 +4,7 @@
 
 #include "PlayerCore.h"
 #include "StreamTypes.h"
+#include "MediaURLType.h"
 #include "InfoLog.h"
 #include "Player/AdaptiveStreamingPlayerABR_State.h"
 #include "Player/Playlist.h"
@@ -19,18 +20,34 @@ namespace Metrics
 	{
 		Undefined,
 		Init,
-		Media
+		Media,
+		Sideloaded,
+		PlaylistElement
 	};
 	static const TCHAR * const GetSegmentTypeString(ESegmentType SegmentType)
 	{
 		switch(SegmentType)
 		{
 			case ESegmentType::Undefined:
+			{
 				return TEXT("Undefined");
+			}
 			case ESegmentType::Init:
+			{
 				return TEXT("Init");
+			}
 			case ESegmentType::Media:
+			{
 				return TEXT("Media");
+			}
+			case ESegmentType::Sideloaded:
+			{
+				return TEXT("Sideloaded");
+			}
+			case ESegmentType::PlaylistElement:
+			{
+				return TEXT("Playlist element");
+			}
 		}
 		return TEXT("n/a");
 	}
@@ -95,8 +112,7 @@ namespace Metrics
 
 	struct FPlaylistDownloadStats
 	{
-		FString					URL;
-		//FString				CDN;
+		FMediaURL				Url;
 		FString					FailureReason;						// Human readable failure reason. Only for display purposes.
 		Playlist::EListType		ListType = Playlist::EListType::Main;
 		Playlist::ELoadType		LoadType = Playlist::ELoadType::Initial;
@@ -123,14 +139,14 @@ namespace Metrics
 		// Inputs from stream request
 		EStreamType StreamType;					//!< Type of stream
 		ESegmentType SegmentType;				//!< Type of segment (init or media)
-		FString URL;							//!< Effective URL used to download from
+		FMediaURL URL;							//!< Effective URL used to download from
 		FString Range;							//!< Range used to download
-		FString CDN;							//!< CDN
 		FString MediaAssetID;
 		FString AdaptationSetID;
 		FString RepresentationID;
 		double PresentationTime;				//!< Presentation time on media timeline
 		double Duration;						//!< Duration of segment as specified in manifest
+		int64 SteeringID;						//!< ID from the content steering handler for this request.
 		int32 Bitrate;							//!< Stream bitrate as specified in manifest
 		int32 QualityIndex;						//!< Quality index of this segment
 		int32 HighestQualityIndex;				//!< The highest quality index that could be had.
@@ -168,12 +184,12 @@ namespace Metrics
 			SegmentType = ESegmentType::Undefined;
 			URL.Empty();
 			Range.Empty();
-			CDN.Empty();
 			MediaAssetID.Empty();
 			AdaptationSetID.Empty();
 			RepresentationID.Empty();
 			PresentationTime = 0.0;
 			Duration = 0.0;
+			SteeringID = 0;
 			Bitrate = 0;
 			QualityIndex = 0;
 			HighestQualityIndex = 0;
diff --git a/Engine/Plugins/Media/ElectraPlayer/Source/ElectraPlayerRuntime/Private/Runtime/Player/DASH/MPDElementsDASH.h b/Engine/Plugins/Media/ElectraPlayer/Source/ElectraPlayerRuntime/Private/Runtime/Player/DASH/MPDElementsDASH.h
--- a/Engine/Plugins/Media/ElectraPlayer/Source/ElectraPlayerRuntime/Private/Runtime/Player/DASH/MPDElementsDASH.h
+++ b/Engine/Plugins/Media/ElectraPlayer/Source/ElectraPlayerRuntime/Private/Runtime/Player/DASH/MPDElementsDASH.h
@@ -95,6 +95,7 @@ public:
 		Latency,
 		ServiceDescription,
 		Resync,
+		ContentSteering,
 		Other
 	};
 
@@ -142,12 +143,12 @@ public:
 	IDashMPDElement(const TCHAR* InName, const TCHAR* InData) : Name(InName), Data(InData)
 	{ }
 	virtual ~IDashMPDElement() = default;
-	virtual bool ProcessElement(FManifestParserDASH* Builder, const TCHAR* ElementName, const TCHAR* ElementData, int32 XmlFileLineNumber);
-	virtual bool ProcessAttribute(FManifestParserDASH* Builder, const TCHAR* AttributeName, const TCHAR* AttributeValue);
-	virtual bool ProcessClose(FManifestParserDASH* Builder, const TCHAR* ElementName) override;
+	bool ProcessElement(FManifestParserDASH* Builder, const TCHAR* ElementName, const TCHAR* ElementData, int32 XmlFileLineNumber) override;
+	bool ProcessAttribute(FManifestParserDASH* Builder, const TCHAR* AttributeName, const TCHAR* AttributeValue) override;
+	bool ProcessClose(FManifestParserDASH* Builder, const TCHAR* ElementName) override;
 	virtual EType GetElementType() const = 0;
 
-	virtual const FString& GetName() const override
+	const FString& GetName() const override
 	{ return Name; }
 
 	TSharedPtrTS<IDashMPDElement> GetParentElement() const						{ return ParentElement.Pin(); }
@@ -203,7 +204,7 @@ class FDashMPD_OtherType : public IDashMPDElement
 public:
 	FDashMPD_OtherType(const TCHAR* Name, const TCHAR* Data) : Super(Name, Data)
 	{ }
-	virtual EType GetElementType() const override
+	EType GetElementType() const override
 	{ return EType::Other; }
 };
 
@@ -218,7 +219,7 @@ public:
 	FDashMPD_BaseURLType(const TCHAR* Name, const TCHAR* Data) : Super(Name, Data)
 	{ }
 	virtual ~FDashMPD_BaseURLType() = default;
-	virtual EType GetElementType() const override
+	EType GetElementType() const override
 	{ return EType::BaseURL; }
 
 	// Convenience method to get the URL.
@@ -229,16 +230,20 @@ public:
 	const FTimeValue& GetAvailabilityTimeOffset() const						{ return AvailabilityTimeOffset; }
 	const TMediaOptionalValue<bool>& GetAvailabilityTimeComplete() const	{ return bAvailabilityTimeComplete; }
 	const FTimeValue& GetTimeShiftBufferDepth() const						{ return TimeShiftBufferDepth; }
+	const TMediaOptionalValue<int64> GetDVBPriority() const					{ return DVBPriority; }
+	const TMediaOptionalValue<int64> GetDVBWeight() const					{ return DVBWeight; }
 
 private:
-	virtual bool ProcessElement(FManifestParserDASH* Builder, const TCHAR* ElementName, const TCHAR* ElementData, int32 XmlFileLineNumber) override;
-	virtual bool ProcessAttribute(FManifestParserDASH* Builder, const TCHAR* AttributeName, const TCHAR* AttributeValue) override;
+	bool ProcessElement(FManifestParserDASH* Builder, const TCHAR* ElementName, const TCHAR* ElementData, int32 XmlFileLineNumber) override;
+	bool ProcessAttribute(FManifestParserDASH* Builder, const TCHAR* AttributeName, const TCHAR* AttributeValue) override;
 
 	FString ServiceLocation;
 	FString ByteRange;
 	FTimeValue AvailabilityTimeOffset = FTimeValue::GetZero();
 	TMediaOptionalValue<bool> bAvailabilityTimeComplete;
 	FTimeValue TimeShiftBufferDepth;
+	TMediaOptionalValue<int64> DVBPriority;
+	TMediaOptionalValue<int64> DVBWeight;
 };
 
 
@@ -252,14 +257,14 @@ public:
 	FDashMPD_URLType(const TCHAR* Name, const TCHAR* Data) : Super(Name, Data)
 	{ }
 	virtual ~FDashMPD_URLType() = default;
-	virtual EType GetElementType() const override
+	EType GetElementType() const override
 	{ return EType::URLType; }
 
 	const FString& GetSourceURL() const	{ return SourceURL; }
 	const FString& GetRange() const		{ return Range; }
 
 private:
-	virtual bool ProcessAttribute(FManifestParserDASH* Builder, const TCHAR* AttributeName, const TCHAR* AttributeValue) override;
+	bool ProcessAttribute(FManifestParserDASH* Builder, const TCHAR* AttributeName, const TCHAR* AttributeValue) override;
 
 	FString SourceURL;
 	FString Range;
@@ -276,7 +281,7 @@ public:
 	FDashMPD_SegmentURLType(const TCHAR* Name, const TCHAR* Data) : Super(Name, Data)
 	{ }
 	virtual ~FDashMPD_SegmentURLType() = default;
-	virtual EType GetElementType() const override
+	EType GetElementType() const override
 	{ return EType::SegmentURL; }
 
 	const FString& GetMedia() const			{ return Media; }
@@ -285,7 +290,7 @@ public:
 	const FString& GetIndexRange() const	{ return IndexRange; }
 
 private:
-	virtual bool ProcessAttribute(FManifestParserDASH* Builder, const TCHAR* AttributeName, const TCHAR* AttributeValue) override;
+	bool ProcessAttribute(FManifestParserDASH* Builder, const TCHAR* AttributeName, const TCHAR* AttributeValue) override;
 
 	FString Media;
 	FString MediaRange;
@@ -297,7 +302,7 @@ private:
 /**
  * Annex I.2 - UrlQueryInfoType
  * Annex I.3 - ExtendedUrlInfoType
- * 
+ *
  * For simplicity both these types are merged into a single type.
  */
 class FDashMPD_UrlQueryInfoType : public IDashMPDElement
@@ -317,13 +322,13 @@ public:
 		IncludeInRequests.Add(TEXT("segment"));
 	}
 	virtual ~FDashMPD_UrlQueryInfoType() = default;
-	virtual EType GetElementType() const override
+	EType GetElementType() const override
 	{ return EType::URLQueryInfo; }
 
 	const FString& GetQueryTemplate() const					{ return QueryTemplate; }
 	bool GetUseMPDUrlQuery() const							{ return bUseMPDUrlQuery; }
 	const FString& GetQueryString() const					{ return QueryString; }
-	virtual FXLink& GetXLink() override						{ return XLink; }
+	FXLink& GetXLink() override								{ return XLink; }
 
 	EExtendedUrlInfoType GetExtendedUrlInfoType() const		{ return ExtendedUrlInfoType; }
 	const TArray<FString>& GetIncludeInRequests() const		{ return IncludeInRequests; }
@@ -331,7 +336,7 @@ public:
 	bool GetSameOriginOnly() const							{ return bSameOriginOnly; }
 
 protected:
-	virtual bool ProcessAttribute(FManifestParserDASH* Builder, const TCHAR* AttributeName, const TCHAR* AttributeValue) override;
+	bool ProcessAttribute(FManifestParserDASH* Builder, const TCHAR* AttributeName, const TCHAR* AttributeValue) override;
 
 	FString QueryTemplate;
 	bool bUseMPDUrlQuery = false;
@@ -355,7 +360,7 @@ public:
 	FDashMPD_DescriptorType(const TCHAR* Name, const TCHAR* Data) : Super(Name, Data)
 	{ }
 	virtual ~FDashMPD_DescriptorType() = default;
-	virtual EType GetElementType() const override
+	EType GetElementType() const override
 	{ return EType::Descriptor; }
 
 	const FString& GetSchemeIdUri() const	{ return SchemeIdUri; }
@@ -377,9 +382,9 @@ public:
 	void SetID(const FString& InID) { ID = InID; }
 
 private:
-	virtual bool ProcessElement(FManifestParserDASH* Builder, const TCHAR* ElementName, const TCHAR* ElementData, int32 XmlFileLineNumber) override;
-	virtual bool ProcessAttribute(FManifestParserDASH* Builder, const TCHAR* AttributeName, const TCHAR* AttributeValue) override;
-	virtual int32 ReplaceChildElementWithRemoteEntities(TSharedPtrTS<IDashMPDElement> Element, const FDashMPD_RootEntities& NewRootEntities, int64 OldResolveID, int64 NewResolveID) override;
+	bool ProcessElement(FManifestParserDASH* Builder, const TCHAR* ElementName, const TCHAR* ElementData, int32 XmlFileLineNumber) override;
+	bool ProcessAttribute(FManifestParserDASH* Builder, const TCHAR* AttributeName, const TCHAR* AttributeValue) override;
+	int32 ReplaceChildElementWithRemoteEntities(TSharedPtrTS<IDashMPDElement> Element, const FDashMPD_RootEntities& NewRootEntities, int64 OldResolveID, int64 NewResolveID) override;
 
 	TArray<TSharedPtrTS<IDashMPDElement>> WellKnownDescriptors;
 
@@ -401,7 +406,7 @@ public:
 	FDashMPD_ProgramInformationType(const TCHAR* Name, const TCHAR* Data) : Super(Name, Data)
 	{ }
 	virtual ~FDashMPD_ProgramInformationType() = default;
-	virtual EType GetElementType() const override
+	EType GetElementType() const override
 	{ return EType::ProgramInformation; }
 
 	TSharedPtrTS<FDashMPD_OtherType> GetTitle() const		{ return Title; }
@@ -411,8 +416,8 @@ public:
 	const FString& GetMoreInformationURL() const			{ return MoreInformationURL; }
 
 private:
-	virtual bool ProcessElement(FManifestParserDASH* Builder, const TCHAR* ElementName, const TCHAR* ElementData, int32 XmlFileLineNumber) override;
-	virtual bool ProcessAttribute(FManifestParserDASH* Builder, const TCHAR* AttributeName, const TCHAR* AttributeValue) override;
+	bool ProcessElement(FManifestParserDASH* Builder, const TCHAR* ElementName, const TCHAR* ElementData, int32 XmlFileLineNumber) override;
+	bool ProcessAttribute(FManifestParserDASH* Builder, const TCHAR* AttributeName, const TCHAR* AttributeValue) override;
 
 	TSharedPtrTS<FDashMPD_OtherType> Title;
 	TSharedPtrTS<FDashMPD_OtherType> Source;
@@ -432,7 +437,7 @@ public:
 	FDashMPD_EventType(const TCHAR* Name, const TCHAR* Data) : Super(Name, Data)
 	{ }
 	virtual ~FDashMPD_EventType() = default;
-	virtual EType GetElementType() const override
+	EType GetElementType() const override
 	{ return EType::Event; }
 
 	uint64 GetPresentationTime() const							{ return PresentationTime; }
@@ -442,7 +447,7 @@ public:
 	const FString& GetMessageData() const						{ return MessageData; }
 
 protected:
-	virtual bool ProcessAttribute(FManifestParserDASH* Builder, const TCHAR* AttributeName, const TCHAR* AttributeValue) override;
+	bool ProcessAttribute(FManifestParserDASH* Builder, const TCHAR* AttributeName, const TCHAR* AttributeValue) override;
 
 private:
 	// Attributes
@@ -464,7 +469,7 @@ public:
 	FDashMPD_EventStreamType(const TCHAR* Name, const TCHAR* Data) : Super(Name, Data)
 	{ }
 	virtual ~FDashMPD_EventStreamType() = default;
-	virtual EType GetElementType() const override
+	EType GetElementType() const override
 	{ return EType::EventStream; }
 
 	const TArray<TSharedPtrTS<FDashMPD_EventType>>& GetEvents() const	{ return Events; }
@@ -473,11 +478,11 @@ public:
 	const FString& GetValue() const											{ return Value; }
 	const TMediaOptionalValue<uint32>& GetTimescale() const					{ return Timescale; }
 	const TMediaOptionalValue<uint64>& GetPresentationTimeOffset() const	{ return PresentationTimeOffset; }
-	virtual FXLink& GetXLink() override										{ return XLink; }
+	FXLink& GetXLink() override												{ return XLink; }
 
 protected:
-	virtual bool ProcessElement(FManifestParserDASH* Builder, const TCHAR* ElementName, const TCHAR* ElementData, int32 XmlFileLineNumber) override;
-	virtual bool ProcessAttribute(FManifestParserDASH* Builder, const TCHAR* AttributeName, const TCHAR* AttributeValue) override;
+	bool ProcessElement(FManifestParserDASH* Builder, const TCHAR* ElementName, const TCHAR* ElementData, int32 XmlFileLineNumber) override;
+	bool ProcessAttribute(FManifestParserDASH* Builder, const TCHAR* AttributeName, const TCHAR* AttributeValue) override;
 
 private:
 	TArray<TSharedPtrTS<FDashMPD_EventType>> Events;
@@ -501,7 +506,7 @@ public:
 	FDashMPD_ContentComponentType(const TCHAR* Name, const TCHAR* Data) : Super(Name, Data)
 	{ }
 	virtual ~FDashMPD_ContentComponentType() = default;
-	virtual EType GetElementType() const override
+	EType GetElementType() const override
 	{ return EType::ContentComponent; }
 
 	const TArray<TSharedPtrTS<FDashMPD_DescriptorType>>& GetAccessibilities() const	{ return Accessibilities; }
@@ -517,8 +522,8 @@ public:
 	const FString& GetTag() const			{ return Tag; }
 
 protected:
-	virtual bool ProcessElement(FManifestParserDASH* Builder, const TCHAR* ElementName, const TCHAR* ElementData, int32 XmlFileLineNumber) override;
-	virtual bool ProcessAttribute(FManifestParserDASH* Builder, const TCHAR* AttributeName, const TCHAR* AttributeValue) override;
+	bool ProcessElement(FManifestParserDASH* Builder, const TCHAR* ElementName, const TCHAR* ElementData, int32 XmlFileLineNumber) override;
+	bool ProcessAttribute(FManifestParserDASH* Builder, const TCHAR* AttributeName, const TCHAR* AttributeValue) override;
 
 private:
 	TArray<TSharedPtrTS<FDashMPD_DescriptorType>> Accessibilities;
@@ -571,15 +576,15 @@ public:
 	FDashMPD_SegmentTimelineType(const TCHAR* Name, const TCHAR* Data) : Super(Name, Data)
 	{ }
 	virtual ~FDashMPD_SegmentTimelineType() = default;
-	virtual EType GetElementType() const override
+	EType GetElementType() const override
 	{ return EType::SegmentTimeline; }
 
 	const TChunkedArray<FSElement, 20480>& GetS_Elements() const	{ return S_Elements; }
 
 protected:
-	virtual bool ProcessElement(FManifestParserDASH* Builder, const TCHAR* ElementName, const TCHAR* ElementData, int32 XmlFileLineNumber) override;
-	virtual bool ProcessAttribute(FManifestParserDASH* Builder, const TCHAR* AttributeName, const TCHAR* AttributeValue) override;
-	virtual bool ProcessClose(FManifestParserDASH* Builder, const TCHAR* ElementName) override;
+	bool ProcessElement(FManifestParserDASH* Builder, const TCHAR* ElementName, const TCHAR* ElementData, int32 XmlFileLineNumber) override;
+	bool ProcessAttribute(FManifestParserDASH* Builder, const TCHAR* AttributeName, const TCHAR* AttributeValue) override;
+	bool ProcessClose(FManifestParserDASH* Builder, const TCHAR* ElementName) override;
 
 private:
 	// At 40 bytes per 'S' entry we get 512 per chunk.
@@ -613,7 +618,7 @@ public:
 	FDashMPD_FailoverContentType(const TCHAR* Name, const TCHAR* Data) : Super(Name, Data)
 	{ }
 	virtual ~FDashMPD_FailoverContentType() = default;
-	virtual EType GetElementType() const override
+	EType GetElementType() const override
 	{ return EType::FailoverContent; }
 
 	const TChunkedArray<FFCSElement, 3072>& GetFCS_Elements() const	{ return FCS_Elements; }
@@ -621,9 +626,9 @@ public:
 	bool GetValid() const	{ return bValid; }
 
 protected:
-	virtual bool ProcessElement(FManifestParserDASH* Builder, const TCHAR* ElementName, const TCHAR* ElementData, int32 XmlFileLineNumber) override;
-	virtual bool ProcessAttribute(FManifestParserDASH* Builder, const TCHAR* AttributeName, const TCHAR* AttributeValue) override;
-	virtual bool ProcessClose(FManifestParserDASH* Builder, const TCHAR* ElementName) override;
+	bool ProcessElement(FManifestParserDASH* Builder, const TCHAR* ElementName, const TCHAR* ElementData, int32 XmlFileLineNumber) override;
+	bool ProcessAttribute(FManifestParserDASH* Builder, const TCHAR* AttributeName, const TCHAR* AttributeValue) override;
+	bool ProcessClose(FManifestParserDASH* Builder, const TCHAR* ElementName) override;
 
 private:
 	// At 24 bytes per 'FCS' entry we get 128 per chunk.
@@ -644,7 +649,7 @@ public:
 	FDashMPD_SegmentBaseType(const TCHAR* Name, const TCHAR* Data) : Super(Name, Data)
 	{ }
 	virtual ~FDashMPD_SegmentBaseType() = default;
-	virtual EType GetElementType() const override
+	EType GetElementType() const override
 	{ return EType::SegmentBase; }
 
 	TSharedPtrTS<FDashMPD_URLType> GetInitialization() const				{ return Initialization; }
@@ -663,8 +668,8 @@ public:
 	const TMediaOptionalValue<bool>& GetAvailabilityTimeComplete() const	{ return bAvailabilityTimeComplete; }
 
 protected:
-	virtual bool ProcessElement(FManifestParserDASH* Builder, const TCHAR* ElementName, const TCHAR* ElementData, int32 XmlFileLineNumber) override;
-	virtual bool ProcessAttribute(FManifestParserDASH* Builder, const TCHAR* AttributeName, const TCHAR* AttributeValue) override;
+	bool ProcessElement(FManifestParserDASH* Builder, const TCHAR* ElementName, const TCHAR* ElementData, int32 XmlFileLineNumber) override;
+	bool ProcessAttribute(FManifestParserDASH* Builder, const TCHAR* AttributeName, const TCHAR* AttributeValue) override;
 
 private:
 	TSharedPtrTS<FDashMPD_URLType> Initialization;
@@ -694,7 +699,7 @@ public:
 	FDashMPD_MultipleSegmentBaseType(const TCHAR* Name, const TCHAR* Data) : Super(Name, Data)
 	{ }
 	virtual ~FDashMPD_MultipleSegmentBaseType() = default;
-	virtual EType GetElementType() const override
+	EType GetElementType() const override
 	{ return EType::MultipleSegmentBase; }
 
 	TSharedPtrTS<FDashMPD_SegmentTimelineType> GetSegmentTimeline() const	{ return SegmentTimeline; }
@@ -705,8 +710,8 @@ public:
 	const TMediaOptionalValue<uint32>& GetEndNumber() const		{ return EndNumber; }
 
 protected:
-	virtual bool ProcessElement(FManifestParserDASH* Builder, const TCHAR* ElementName, const TCHAR* ElementData, int32 XmlFileLineNumber) override;
-	virtual bool ProcessAttribute(FManifestParserDASH* Builder, const TCHAR* AttributeName, const TCHAR* AttributeValue) override;
+	bool ProcessElement(FManifestParserDASH* Builder, const TCHAR* ElementName, const TCHAR* ElementData, int32 XmlFileLineNumber) override;
+	bool ProcessAttribute(FManifestParserDASH* Builder, const TCHAR* AttributeName, const TCHAR* AttributeValue) override;
 
 private:
 	TSharedPtrTS<FDashMPD_SegmentTimelineType> SegmentTimeline;
@@ -729,16 +734,16 @@ public:
 	FDashMPD_SegmentListType(const TCHAR* Name, const TCHAR* Data) : Super(Name, Data)
 	{ }
 	virtual ~FDashMPD_SegmentListType() = default;
-	virtual EType GetElementType() const override
+	EType GetElementType() const override
 	{ return EType::SegmentList; }
 
 	const TArray<TSharedPtrTS<FDashMPD_SegmentURLType>>& GetSegmentURLs() const	{ return SegmentURLs; }
 
-	virtual FXLink& GetXLink() override	{ return XLink; }
+	FXLink& GetXLink() override	{ return XLink; }
 
 protected:
-	virtual bool ProcessElement(FManifestParserDASH* Builder, const TCHAR* ElementName, const TCHAR* ElementData, int32 XmlFileLineNumber) override;
-	virtual bool ProcessAttribute(FManifestParserDASH* Builder, const TCHAR* AttributeName, const TCHAR* AttributeValue) override;
+	bool ProcessElement(FManifestParserDASH* Builder, const TCHAR* ElementName, const TCHAR* ElementData, int32 XmlFileLineNumber) override;
+	bool ProcessAttribute(FManifestParserDASH* Builder, const TCHAR* AttributeName, const TCHAR* AttributeValue) override;
 
 private:
 	TArray<TSharedPtrTS<FDashMPD_SegmentURLType>> SegmentURLs;
@@ -758,7 +763,7 @@ public:
 	FDashMPD_SegmentTemplateType(const TCHAR* Name, const TCHAR* Data) : Super(Name, Data)
 	{ }
 	virtual ~FDashMPD_SegmentTemplateType() = default;
-	virtual EType GetElementType() const override
+	EType GetElementType() const override
 	{ return EType::SegmentTemplate; }
 
 	// The following accessors are suffixed 'Template' to differentiate the method from the inherited methods
@@ -769,7 +774,7 @@ public:
 	const FString& GetBitstreamSwitchingTemplate() const	{ return BitstreamSwitching; }
 
 protected:
-	virtual bool ProcessAttribute(FManifestParserDASH* Builder, const TCHAR* AttributeName, const TCHAR* AttributeValue) override;
+	bool ProcessAttribute(FManifestParserDASH* Builder, const TCHAR* AttributeName, const TCHAR* AttributeValue) override;
 
 private:
 	// Attributes
@@ -790,15 +795,15 @@ public:
 	FDashMPD_SubsetType(const TCHAR* Name, const TCHAR* Data) : Super(Name, Data)
 	{ }
 	virtual ~FDashMPD_SubsetType() = default;
-	virtual EType GetElementType() const override
+	EType GetElementType() const override
 	{ return EType::Subset; }
 
 	const TArray<uint64>& GetContains() const	{ return Contains; }
 	const FString& GetID() const				{ return ID; }
 
 private:
-	virtual bool ProcessElement(FManifestParserDASH* Builder, const TCHAR* ElementName, const TCHAR* ElementData, int32 XmlFileLineNumber) override;
-	virtual bool ProcessAttribute(FManifestParserDASH* Builder, const TCHAR* AttributeName, const TCHAR* AttributeValue) override;
+	bool ProcessElement(FManifestParserDASH* Builder, const TCHAR* ElementName, const TCHAR* ElementData, int32 XmlFileLineNumber) override;
+	bool ProcessAttribute(FManifestParserDASH* Builder, const TCHAR* AttributeName, const TCHAR* AttributeValue) override;
 
 	// Attributes
 	TArray<uint64> Contains;
@@ -816,15 +821,15 @@ public:
 	FDashMPD_SwitchingType(const TCHAR* Name, const TCHAR* Data) : Super(Name, Data)
 	{ }
 	virtual ~FDashMPD_SwitchingType() = default;
-	virtual EType GetElementType() const override
+	EType GetElementType() const override
 	{ return EType::Switching; }
 
 	uint64 GetInterval() const		{ return Interval; }
 	const FString& GetType() const	{ return Type; }
 
 private:
-	virtual bool ProcessElement(FManifestParserDASH* Builder, const TCHAR* ElementName, const TCHAR* ElementData, int32 XmlFileLineNumber) override;
-	virtual bool ProcessAttribute(FManifestParserDASH* Builder, const TCHAR* AttributeName, const TCHAR* AttributeValue) override;
+	bool ProcessElement(FManifestParserDASH* Builder, const TCHAR* ElementName, const TCHAR* ElementData, int32 XmlFileLineNumber) override;
+	bool ProcessAttribute(FManifestParserDASH* Builder, const TCHAR* AttributeName, const TCHAR* AttributeValue) override;
 
 	// Attributes
 	uint64 Interval = 0;
@@ -842,7 +847,7 @@ public:
 	FDashMPD_RandomAccessType(const TCHAR* Name, const TCHAR* Data) : Super(Name, Data)
 	{ }
 	virtual ~FDashMPD_RandomAccessType() = default;
-	virtual EType GetElementType() const override
+	EType GetElementType() const override
 	{ return EType::RandomAccess; }
 
 	uint64 GetInterval() const								{ return Interval; }
@@ -851,8 +856,8 @@ public:
 	const TMediaOptionalValue<uint64>& GetBandwidth() const	{ return Bandwidth; }
 
 private:
-	virtual bool ProcessElement(FManifestParserDASH* Builder, const TCHAR* ElementName, const TCHAR* ElementData, int32 XmlFileLineNumber) override;
-	virtual bool ProcessAttribute(FManifestParserDASH* Builder, const TCHAR* AttributeName, const TCHAR* AttributeValue) override;
+	bool ProcessElement(FManifestParserDASH* Builder, const TCHAR* ElementName, const TCHAR* ElementData, int32 XmlFileLineNumber) override;
+	bool ProcessAttribute(FManifestParserDASH* Builder, const TCHAR* AttributeName, const TCHAR* AttributeValue) override;
 
 	// Attributes
 	uint64 Interval = 0;
@@ -872,15 +877,15 @@ public:
 	FDashMPD_LabelType(const TCHAR* Name, const TCHAR* Data) : Super(Name, Data)
 	{ }
 	virtual ~FDashMPD_LabelType() = default;
-	virtual EType GetElementType() const override
+	EType GetElementType() const override
 	{ return EType::Label; }
 
 	uint64 GetID() const				{ return ID; }
 	const FString& GetLanguage() const	{ return Language; }
 
 private:
-	virtual bool ProcessElement(FManifestParserDASH* Builder, const TCHAR* ElementName, const TCHAR* ElementData, int32 XmlFileLineNumber) override;
-	virtual bool ProcessAttribute(FManifestParserDASH* Builder, const TCHAR* AttributeName, const TCHAR* AttributeValue) override;
+	bool ProcessElement(FManifestParserDASH* Builder, const TCHAR* ElementName, const TCHAR* ElementData, int32 XmlFileLineNumber) override;
+	bool ProcessAttribute(FManifestParserDASH* Builder, const TCHAR* AttributeName, const TCHAR* AttributeValue) override;
 
 	// Attributes
 	uint64 ID = 0;
@@ -898,7 +903,7 @@ public:
 	FDashMPD_ProducerReferenceTimeType(const TCHAR* Name, const TCHAR* Data) : Super(Name, Data)
 	{ }
 	virtual ~FDashMPD_ProducerReferenceTimeType() = default;
-	virtual EType GetElementType() const override
+	EType GetElementType() const override
 	{ return EType::ProducerReferenceTime; }
 
 	TSharedPtrTS<FDashMPD_DescriptorType> GetUTCTiming() const	{ return UTCTiming; }
@@ -911,8 +916,8 @@ public:
 	uint64 GetPresentationTime() const			{ return PresentationTime; }
 
 private:
-	virtual bool ProcessElement(FManifestParserDASH* Builder, const TCHAR* ElementName, const TCHAR* ElementData, int32 XmlFileLineNumber) override;
-	virtual bool ProcessAttribute(FManifestParserDASH* Builder, const TCHAR* AttributeName, const TCHAR* AttributeValue) override;
+	bool ProcessElement(FManifestParserDASH* Builder, const TCHAR* ElementName, const TCHAR* ElementData, int32 XmlFileLineNumber) override;
+	bool ProcessAttribute(FManifestParserDASH* Builder, const TCHAR* AttributeName, const TCHAR* AttributeValue) override;
 
 	TSharedPtrTS<FDashMPD_DescriptorType> UTCTiming;
 
@@ -936,7 +941,7 @@ public:
 	FDashMPD_LeapSecondInformationType(const TCHAR* Name, const TCHAR* Data) : Super(Name, Data)
 	{ }
 	virtual ~FDashMPD_LeapSecondInformationType() = default;
-	virtual EType GetElementType() const override
+	EType GetElementType() const override
 	{ return EType::LeapSecondInformation; }
 
 	int64 GetAvailabilityStartLeapOffset() const									{ return AvailabilityStartLeapOffset; }
@@ -944,7 +949,7 @@ public:
 	const FTimeValue& GetNextLeapChangeTime() const									{ return NextLeapChangeTime; }
 
 private:
-	virtual bool ProcessAttribute(FManifestParserDASH* Builder, const TCHAR* AttributeName, const TCHAR* AttributeValue) override;
+	bool ProcessAttribute(FManifestParserDASH* Builder, const TCHAR* AttributeName, const TCHAR* AttributeValue) override;
 
 	// Attributes
 	int64 AvailabilityStartLeapOffset = 0;
@@ -963,7 +968,7 @@ public:
 	FDashMPD_UIntVWithIDType(const TCHAR* Name, const TCHAR* Data) : Super(Name, Data)
 	{ }
 	virtual ~FDashMPD_UIntVWithIDType() = default;
-	virtual EType GetElementType() const override
+	EType GetElementType() const override
 	{ return EType::UIntVWithID; }
 
 	const TArray<uint32>& GetElements() const	{ return Elements; }
@@ -973,9 +978,9 @@ public:
 	const FString& GetContentType() const		{ return ContentType; }
 
 private:
-	virtual bool ProcessElement(FManifestParserDASH* Builder, const TCHAR* ElementName, const TCHAR* ElementData, int32 XmlFileLineNumber) override;
-	virtual bool ProcessAttribute(FManifestParserDASH* Builder, const TCHAR* AttributeName, const TCHAR* AttributeValue) override;
-	virtual bool ProcessClose(FManifestParserDASH* Builder, const TCHAR* ElementName) override;
+	bool ProcessElement(FManifestParserDASH* Builder, const TCHAR* ElementName, const TCHAR* ElementData, int32 XmlFileLineNumber) override;
+	bool ProcessAttribute(FManifestParserDASH* Builder, const TCHAR* AttributeName, const TCHAR* AttributeValue) override;
+	bool ProcessClose(FManifestParserDASH* Builder, const TCHAR* ElementName) override;
 
 	TArray<uint32> Elements;
 
@@ -996,7 +1001,7 @@ public:
 	FDashMPD_UIntPairsWithIDType(const TCHAR* Name, const TCHAR* Data) : Super(Name, Data)
 	{ }
 	virtual ~FDashMPD_UIntPairsWithIDType() = default;
-	virtual EType GetElementType() const override
+	EType GetElementType() const override
 	{ return EType::UIntPairsWithID; }
 
 	const TArray<uint32>& GetElements() const	{ return Elements; }
@@ -1004,9 +1009,9 @@ public:
 	const FString& GetType() const	{ return Type; }
 
 private:
-	virtual bool ProcessElement(FManifestParserDASH* Builder, const TCHAR* ElementName, const TCHAR* ElementData, int32 XmlFileLineNumber) override;
-	virtual bool ProcessAttribute(FManifestParserDASH* Builder, const TCHAR* AttributeName, const TCHAR* AttributeValue) override;
-	virtual bool ProcessClose(FManifestParserDASH* Builder, const TCHAR* ElementName) override;
+	bool ProcessElement(FManifestParserDASH* Builder, const TCHAR* ElementName, const TCHAR* ElementData, int32 XmlFileLineNumber) override;
+	bool ProcessAttribute(FManifestParserDASH* Builder, const TCHAR* AttributeName, const TCHAR* AttributeValue) override;
+	bool ProcessClose(FManifestParserDASH* Builder, const TCHAR* ElementName) override;
 
 	TArray<uint32> Elements;
 
@@ -1025,7 +1030,7 @@ public:
 	FDashMPD_OperatingBandwidthType(const TCHAR* Name, const TCHAR* Data) : Super(Name, Data)
 	{ }
 	virtual ~FDashMPD_OperatingBandwidthType() = default;
-	virtual EType GetElementType() const override
+	EType GetElementType() const override
 	{ return EType::OperatingBandwidth; }
 
 	const FString& GetMediaType() const						{ return MediaType; }
@@ -1034,8 +1039,8 @@ public:
 	const TMediaOptionalValue<uint64>& GetTarget() const	{ return Target; }
 
 private:
-	virtual bool ProcessElement(FManifestParserDASH* Builder, const TCHAR* ElementName, const TCHAR* ElementData, int32 XmlFileLineNumber) override;
-	virtual bool ProcessAttribute(FManifestParserDASH* Builder, const TCHAR* AttributeName, const TCHAR* AttributeValue) override;
+	bool ProcessElement(FManifestParserDASH* Builder, const TCHAR* ElementName, const TCHAR* ElementData, int32 XmlFileLineNumber) override;
+	bool ProcessAttribute(FManifestParserDASH* Builder, const TCHAR* AttributeName, const TCHAR* AttributeValue) override;
 
 	// Attributes
 	FString MediaType;		// one of "video", "audio", "any", "all"
@@ -1055,7 +1060,7 @@ public:
 	FDashMPD_OperatingQualityType(const TCHAR* Name, const TCHAR* Data) : Super(Name, Data)
 	{ }
 	virtual ~FDashMPD_OperatingQualityType() = default;
-	virtual EType GetElementType() const override
+	EType GetElementType() const override
 	{ return EType::OperatingQuality; }
 
 	const FString& GetMediaType() const							{ return MediaType; }
@@ -1066,8 +1071,8 @@ public:
 	const TMediaOptionalValue<uint64>& GetMaxDifference() const	{ return MaxDifference; }
 
 private:
-	virtual bool ProcessElement(FManifestParserDASH* Builder, const TCHAR* ElementName, const TCHAR* ElementData, int32 XmlFileLineNumber) override;
-	virtual bool ProcessAttribute(FManifestParserDASH* Builder, const TCHAR* AttributeName, const TCHAR* AttributeValue) override;
+	bool ProcessElement(FManifestParserDASH* Builder, const TCHAR* ElementName, const TCHAR* ElementData, int32 XmlFileLineNumber) override;
+	bool ProcessAttribute(FManifestParserDASH* Builder, const TCHAR* AttributeName, const TCHAR* AttributeValue) override;
 
 	// Attributes
 	FString MediaType;		// one of "video", "audio", "any"
@@ -1089,15 +1094,15 @@ public:
 	FDashMPD_PlaybackRateType(const TCHAR* Name, const TCHAR* Data) : Super(Name, Data)
 	{ }
 	virtual ~FDashMPD_PlaybackRateType() = default;
-	virtual EType GetElementType() const override
+	EType GetElementType() const override
 	{ return EType::PlaybackRate; }
 
 	const FTimeValue& GetMin() const	{ return Min; }
 	const FTimeValue& GetMax() const	{ return Max; }
 
 private:
-	virtual bool ProcessElement(FManifestParserDASH* Builder, const TCHAR* ElementName, const TCHAR* ElementData, int32 XmlFileLineNumber) override;
-	virtual bool ProcessAttribute(FManifestParserDASH* Builder, const TCHAR* AttributeName, const TCHAR* AttributeValue) override;
+	bool ProcessElement(FManifestParserDASH* Builder, const TCHAR* ElementName, const TCHAR* ElementData, int32 XmlFileLineNumber) override;
+	bool ProcessAttribute(FManifestParserDASH* Builder, const TCHAR* AttributeName, const TCHAR* AttributeValue) override;
 
 	// Attributes
 	FTimeValue Min;
@@ -1115,7 +1120,7 @@ public:
 	FDashMPD_LatencyType(const TCHAR* Name, const TCHAR* Data) : Super(Name, Data)
 	{ }
 	virtual ~FDashMPD_LatencyType() = default;
-	virtual EType GetElementType() const override
+	EType GetElementType() const override
 	{ return EType::Latency; }
 
 	const TArray<TSharedPtrTS<FDashMPD_UIntPairsWithIDType>>& GetQualityLatencies() const	{ return QualityLatencies; }
@@ -1126,8 +1131,8 @@ public:
 	const TMediaOptionalValue<uint64>& GetMin() const			{ return Min; }
 
 private:
-	virtual bool ProcessElement(FManifestParserDASH* Builder, const TCHAR* ElementName, const TCHAR* ElementData, int32 XmlFileLineNumber) override;
-	virtual bool ProcessAttribute(FManifestParserDASH* Builder, const TCHAR* AttributeName, const TCHAR* AttributeValue) override;
+	bool ProcessElement(FManifestParserDASH* Builder, const TCHAR* ElementName, const TCHAR* ElementData, int32 XmlFileLineNumber) override;
+	bool ProcessAttribute(FManifestParserDASH* Builder, const TCHAR* AttributeName, const TCHAR* AttributeValue) override;
 
 	TArray<TSharedPtrTS<FDashMPD_UIntPairsWithIDType>> QualityLatencies;
 
@@ -1156,7 +1161,7 @@ public:
 		FPRElement(const TCHAR* Name, const TCHAR* Data) : Super(Name, Data)
 		{ }
 		virtual ~FPRElement() = default;
-		virtual EType GetElementType() const override
+		EType GetElementType() const override
 		{ return EType::ContentPopularityRate_PR; }
 
 		uint64 GetPopularityRate() const						{ return PopularityRate; }
@@ -1164,7 +1169,7 @@ public:
 		int32 GetR() const										{ return R; }
 
 	protected:
-		virtual bool ProcessAttribute(FManifestParserDASH* Builder, const TCHAR* AttributeName, const TCHAR* AttributeValue) override;
+		bool ProcessAttribute(FManifestParserDASH* Builder, const TCHAR* AttributeName, const TCHAR* AttributeValue) override;
 
 	private:
 		// Attributes
@@ -1176,7 +1181,7 @@ public:
 	FDashMPD_ContentPopularityRateType(const TCHAR* Name, const TCHAR* Data) : Super(Name, Data)
 	{ }
 	virtual ~FDashMPD_ContentPopularityRateType() = default;
-	virtual EType GetElementType() const override
+	EType GetElementType() const override
 	{ return EType::ContentPopularityRate; }
 
 	const TArray<TSharedPtrTS<FPRElement>>& GetPR_Elements() const	{ return PR_Elements; }
@@ -1185,8 +1190,8 @@ public:
 	const FString& GetSource_Description() const	{ return Source_Description; }
 
 protected:
-	virtual bool ProcessElement(FManifestParserDASH* Builder, const TCHAR* ElementName, const TCHAR* ElementData, int32 XmlFileLineNumber) override;
-	virtual bool ProcessAttribute(FManifestParserDASH* Builder, const TCHAR* AttributeName, const TCHAR* AttributeValue) override;
+	bool ProcessElement(FManifestParserDASH* Builder, const TCHAR* ElementName, const TCHAR* ElementData, int32 XmlFileLineNumber) override;
+	bool ProcessAttribute(FManifestParserDASH* Builder, const TCHAR* AttributeName, const TCHAR* AttributeValue) override;
 
 private:
 	TArray<TSharedPtrTS<FPRElement>>	PR_Elements;
@@ -1197,6 +1202,38 @@ private:
 };
 
 
+
+/**
+ * DASH-IF CTS 00XX V0.9.0 / ETSI TS 103 998 V1.1.1 - ContentSteeringType
+ */
+class FDashMPD_ContentSteeringType : public IDashMPDElement
+{
+	typedef IDashMPDElement Super;
+public:
+	FDashMPD_ContentSteeringType(const TCHAR* Name, const TCHAR* Data) : Super(Name, Data)
+	{ }
+	virtual ~FDashMPD_ContentSteeringType() = default;
+	EType GetElementType() const override
+	{ return EType::ContentSteering; }
+
+	const FString& GetContentSteeringServer() const		{ return GetData(); }
+	const FString& GetProxyServerURL() const			{ return ProxyServerURL; }
+	const FString& GetDefaultServiceLocation() const	{ return DefaultServiceLocation; }
+	bool QueryBeforeStart() const						{ return bQueryBeforeStart; }
+	bool ClientRequirement() const						{ return bClientRequirement; }
+
+private:
+	bool ProcessAttribute(FManifestParserDASH* Builder, const TCHAR* AttributeName, const TCHAR* AttributeValue) override;
+
+	// Attributes
+	FString ProxyServerURL;
+	FString DefaultServiceLocation;
+	bool bQueryBeforeStart = false;
+	bool bClientRequirement = false;
+};
+
+
+
 /**
  * K.4.2.1 - ServiceDescriptionType
  */
@@ -1207,7 +1244,7 @@ public:
 	FDashMPD_ServiceDescriptionType(const TCHAR* Name, const TCHAR* Data) : Super(Name, Data)
 	{ }
 	virtual ~FDashMPD_ServiceDescriptionType() = default;
-	virtual EType GetElementType() const override
+	EType GetElementType() const override
 	{ return EType::ServiceDescription; }
 
 	const TArray<TSharedPtrTS<FDashMPD_DescriptorType>>& GetScopes() const						{ return Scopes; }
@@ -1219,8 +1256,8 @@ public:
 	uint32 GetID() const { return (uint32) ID; }
 
 private:
-	virtual bool ProcessElement(FManifestParserDASH* Builder, const TCHAR* ElementName, const TCHAR* ElementData, int32 XmlFileLineNumber) override;
-	virtual bool ProcessAttribute(FManifestParserDASH* Builder, const TCHAR* AttributeName, const TCHAR* AttributeValue) override;
+	bool ProcessElement(FManifestParserDASH* Builder, const TCHAR* ElementName, const TCHAR* ElementData, int32 XmlFileLineNumber) override;
+	bool ProcessAttribute(FManifestParserDASH* Builder, const TCHAR* AttributeName, const TCHAR* AttributeValue) override;
 
 private:
 	TArray<TSharedPtrTS<FDashMPD_DescriptorType>> Scopes;
@@ -1244,7 +1281,7 @@ public:
 	FDashMPD_ResyncType(const TCHAR* Name, const TCHAR* Data) : Super(Name, Data)
 	{ }
 	virtual ~FDashMPD_ResyncType() = default;
-	virtual EType GetElementType() const override
+	EType GetElementType() const override
 	{ return EType::Resync; }
 
 	int32 GetType() const							{ return Type; }
@@ -1256,8 +1293,8 @@ public:
 	const FString& GetIndex() const					{ return Index; }
 
 private:
-	virtual bool ProcessElement(FManifestParserDASH* Builder, const TCHAR* ElementName, const TCHAR* ElementData, int32 XmlFileLineNumber) override;
-	virtual bool ProcessAttribute(FManifestParserDASH* Builder, const TCHAR* AttributeName, const TCHAR* AttributeValue) override;
+	bool ProcessElement(FManifestParserDASH* Builder, const TCHAR* ElementName, const TCHAR* ElementData, int32 XmlFileLineNumber) override;
+	bool ProcessAttribute(FManifestParserDASH* Builder, const TCHAR* AttributeName, const TCHAR* AttributeValue) override;
 
 	// Attributes
 	int32 Type = 0;						// SAP type
@@ -1281,19 +1318,19 @@ public:
 	{ }
 	virtual ~FDashMPD_RepresentationBaseType() = default;
 
-	const TArray<TSharedPtrTS<FDashMPD_DescriptorType>>& GetFramePackings() const							{ return FramePackings; }
-	const TArray<TSharedPtrTS<FDashMPD_DescriptorType>>& GetAudioChannelConfigurations() const				{ return AudioChannelConfigurations; }
-	const TArray<TSharedPtrTS<FDashMPD_DescriptorType>>& GetContentProtections() const						{ return ContentProtections; }
-	virtual const TArray<TSharedPtrTS<FDashMPD_DescriptorType>>& GetEssentialProperties() const	override	{ return EssentialProperties; }
-	virtual const TArray<TSharedPtrTS<FDashMPD_DescriptorType>>& GetSupplementalProperties() const override	{ return SupplementalProperties; }
-	const TArray<TSharedPtrTS<FDashMPD_EventStreamType>>& GetInbandEventStreams() const						{ return InbandEventStreams; }
-	const TArray<TSharedPtrTS<FDashMPD_SwitchingType>>& GetSwitchings() const								{ return Switchings; }
-	const TArray<TSharedPtrTS<FDashMPD_RandomAccessType>>& GetRandomAccesses() const						{ return RandomAccesses; }
-	const TArray<TSharedPtrTS<FDashMPD_LabelType>>& GetGroupLabels() const									{ return GroupLabels; }
-	const TArray<TSharedPtrTS<FDashMPD_LabelType>>& GetLabels() const										{ return Labels; }
-	const TArray<TSharedPtrTS<FDashMPD_ProducerReferenceTimeType>>& GetProducerReferenceTimes() const		{ return ProducerReferenceTimes; }
-	const TArray<TSharedPtrTS<FDashMPD_ContentPopularityRateType>>& GetContentPopularityRates() const		{ return ContentPopularityRates; }
-	const TArray<TSharedPtrTS<FDashMPD_ResyncType>>& GetResyncs() const										{ return Resyncs; }
+	const TArray<TSharedPtrTS<FDashMPD_DescriptorType>>& GetFramePackings() const						{ return FramePackings; }
+	const TArray<TSharedPtrTS<FDashMPD_DescriptorType>>& GetAudioChannelConfigurations() const			{ return AudioChannelConfigurations; }
+	const TArray<TSharedPtrTS<FDashMPD_DescriptorType>>& GetContentProtections() const					{ return ContentProtections; }
+	const TArray<TSharedPtrTS<FDashMPD_DescriptorType>>& GetEssentialProperties() const	override		{ return EssentialProperties; }
+	const TArray<TSharedPtrTS<FDashMPD_DescriptorType>>& GetSupplementalProperties() const override		{ return SupplementalProperties; }
+	const TArray<TSharedPtrTS<FDashMPD_EventStreamType>>& GetInbandEventStreams() const					{ return InbandEventStreams; }
+	const TArray<TSharedPtrTS<FDashMPD_SwitchingType>>& GetSwitchings() const							{ return Switchings; }
+	const TArray<TSharedPtrTS<FDashMPD_RandomAccessType>>& GetRandomAccesses() const					{ return RandomAccesses; }
+	const TArray<TSharedPtrTS<FDashMPD_LabelType>>& GetGroupLabels() const								{ return GroupLabels; }
+	const TArray<TSharedPtrTS<FDashMPD_LabelType>>& GetLabels() const									{ return Labels; }
+	const TArray<TSharedPtrTS<FDashMPD_ProducerReferenceTimeType>>& GetProducerReferenceTimes() const	{ return ProducerReferenceTimes; }
+	const TArray<TSharedPtrTS<FDashMPD_ContentPopularityRateType>>& GetContentPopularityRates() const	{ return ContentPopularityRates; }
+	const TArray<TSharedPtrTS<FDashMPD_ResyncType>>& GetResyncs() const									{ return Resyncs; }
 
 	const TArray<FString>& GetProfiles() const						{ return Profiles; }
 	const TMediaOptionalValue<uint64>& GetWidth() const				{ return Width; }
@@ -1313,9 +1350,9 @@ public:
 	const FString& GetTag() const									{ return Tag; }
 
 protected:
-	virtual bool ProcessElement(FManifestParserDASH* Builder, const TCHAR* ElementName, const TCHAR* ElementData, int32 XmlFileLineNumber) override;
-	virtual bool ProcessAttribute(FManifestParserDASH* Builder, const TCHAR* AttributeName, const TCHAR* AttributeValue) override;
-	virtual int32 ReplaceChildElementWithRemoteEntities(TSharedPtrTS<IDashMPDElement> Element, const FDashMPD_RootEntities& NewRootEntities, int64 OldResolveID, int64 NewResolveID) override;
+	bool ProcessElement(FManifestParserDASH* Builder, const TCHAR* ElementName, const TCHAR* ElementData, int32 XmlFileLineNumber) override;
+	bool ProcessAttribute(FManifestParserDASH* Builder, const TCHAR* AttributeName, const TCHAR* AttributeValue) override;
+	int32 ReplaceChildElementWithRemoteEntities(TSharedPtrTS<IDashMPDElement> Element, const FDashMPD_RootEntities& NewRootEntities, int64 OldResolveID, int64 NewResolveID) override;
 
 private:
 	TArray<TSharedPtrTS<FDashMPD_DescriptorType>> FramePackings;
@@ -1362,11 +1399,11 @@ public:
 	FDashMPD_SubRepresentationType(const TCHAR* Name, const TCHAR* Data) : Super(Name, Data)
 	{ }
 	virtual ~FDashMPD_SubRepresentationType() = default;
-	virtual EType GetElementType() const override
+	EType GetElementType() const override
 	{ return EType::SubRepresentation; }
 
 private:
-	virtual bool ProcessAttribute(FManifestParserDASH* Builder, const TCHAR* AttributeName, const TCHAR* AttributeValue) override;
+	bool ProcessAttribute(FManifestParserDASH* Builder, const TCHAR* AttributeName, const TCHAR* AttributeValue) override;
 
 	// Attributes
 	TMediaOptionalValue<uint64>	Level;
@@ -1386,7 +1423,7 @@ public:
 	FDashMPD_InitializationSetType(const TCHAR* Name, const TCHAR* Data) : Super(Name, Data)
 	{ }
 	virtual ~FDashMPD_InitializationSetType() = default;
-	virtual EType GetElementType() const override
+	EType GetElementType() const override
 	{ return EType::InitializationSet; }
 
 	const TArray<TSharedPtrTS<FDashMPD_DescriptorType>>& GetAccessibilities() const	{ return Accessibilities; }
@@ -1403,11 +1440,11 @@ public:
 	const TMediaOptionalValue<uint64>& GetMaxHeight() const	{ return MaxHeight; }
 	const FTimeFraction& GetMaxFrameRate() const			{ return MaxFrameRate; }
 	const FString& GetInitialization() const				{ return Initialization; }
-	virtual FXLink& GetXLink() override						{ return XLink; }
+	FXLink& GetXLink() override								{ return XLink; }
 
 private:
-	virtual bool ProcessElement(FManifestParserDASH* Builder, const TCHAR* ElementName, const TCHAR* ElementData, int32 XmlFileLineNumber) override;
-	virtual bool ProcessAttribute(FManifestParserDASH* Builder, const TCHAR* AttributeName, const TCHAR* AttributeValue) override;
+	bool ProcessElement(FManifestParserDASH* Builder, const TCHAR* ElementName, const TCHAR* ElementData, int32 XmlFileLineNumber) override;
+	bool ProcessAttribute(FManifestParserDASH* Builder, const TCHAR* AttributeName, const TCHAR* AttributeValue) override;
 
 	TArray<TSharedPtrTS<FDashMPD_DescriptorType>> Accessibilities;
 	TArray<TSharedPtrTS<FDashMPD_DescriptorType>> Roles;
@@ -1437,7 +1474,7 @@ public:
 	FDashMPD_PreselectionType(const TCHAR* Name, const TCHAR* Data) : Super(Name, Data)
 	{ }
 	virtual ~FDashMPD_PreselectionType() = default;
-	virtual EType GetElementType() const override
+	EType GetElementType() const override
 	{ return EType::Preselection; }
 
 	const TArray<TSharedPtrTS<FDashMPD_DescriptorType>>& GetAccessibilities() const	{ return Accessibilities; }
@@ -1451,8 +1488,8 @@ public:
 	const FString& GetOrder() const								{ return Order; }
 
 private:
-	virtual bool ProcessElement(FManifestParserDASH* Builder, const TCHAR* ElementName, const TCHAR* ElementData, int32 XmlFileLineNumber) override;
-	virtual bool ProcessAttribute(FManifestParserDASH* Builder, const TCHAR* AttributeName, const TCHAR* AttributeValue) override;
+	bool ProcessElement(FManifestParserDASH* Builder, const TCHAR* ElementName, const TCHAR* ElementData, int32 XmlFileLineNumber) override;
+	bool ProcessAttribute(FManifestParserDASH* Builder, const TCHAR* AttributeName, const TCHAR* AttributeValue) override;
 
 	TArray<TSharedPtrTS<FDashMPD_DescriptorType>> Accessibilities;
 	TArray<TSharedPtrTS<FDashMPD_DescriptorType>> Roles;
@@ -1477,10 +1514,10 @@ public:
 	FDashMPD_RepresentationType(const TCHAR* Name, const TCHAR* Data) : Super(Name, Data)
 	{ }
 	virtual ~FDashMPD_RepresentationType() = default;
-	virtual EType GetElementType() const override
+	EType GetElementType() const override
 	{ return EType::Representation; }
 
-	virtual const TArray<TSharedPtrTS<FDashMPD_BaseURLType>>& GetBaseURLs() const override		{ return BaseURLs; }
+	const TArray<TSharedPtrTS<FDashMPD_BaseURLType>>& GetBaseURLs() const override				{ return BaseURLs; }
 	const TArray<TSharedPtrTS<FDashMPD_SubRepresentationType>>& GetSubRepresentations() const	{ return SubRepresentations; }
 	TSharedPtrTS<FDashMPD_SegmentBaseType> GetSegmentBase() const								{ return SegmentBase; }
 	TSharedPtrTS<FDashMPD_SegmentListType> GetSegmentList() const								{ return SegmentList; }
@@ -1496,9 +1533,9 @@ public:
 	const TArray<FString>& GetMediaStreamStructureIDs() const		{ return MediaStreamStructureIDs; }
 
 private:
-	virtual bool ProcessElement(FManifestParserDASH* Builder, const TCHAR* ElementName, const TCHAR* ElementData, int32 XmlFileLineNumber) override;
-	virtual bool ProcessAttribute(FManifestParserDASH* Builder, const TCHAR* AttributeName, const TCHAR* AttributeValue) override;
-	virtual int32 ReplaceChildElementWithRemoteEntities(TSharedPtrTS<IDashMPDElement> Element, const FDashMPD_RootEntities& NewRootEntities, int64 OldResolveID, int64 NewResolveID) override;
+	bool ProcessElement(FManifestParserDASH* Builder, const TCHAR* ElementName, const TCHAR* ElementData, int32 XmlFileLineNumber) override;
+	bool ProcessAttribute(FManifestParserDASH* Builder, const TCHAR* AttributeName, const TCHAR* AttributeValue) override;
+	int32 ReplaceChildElementWithRemoteEntities(TSharedPtrTS<IDashMPDElement> Element, const FDashMPD_RootEntities& NewRootEntities, int64 OldResolveID, int64 NewResolveID) override;
 
 	TArray<TSharedPtrTS<FDashMPD_BaseURLType>> BaseURLs;
 	TArray<TSharedPtrTS<FDashMPD_SubRepresentationType>> SubRepresentations;
@@ -1527,7 +1564,7 @@ public:
 	FDashMPD_AdaptationSetType(const TCHAR* Name, const TCHAR* Data) : Super(Name, Data)
 	{ }
 	virtual ~FDashMPD_AdaptationSetType() = default;
-	virtual EType GetElementType() const override
+	EType GetElementType() const override
 	{ return EType::AdaptationSet; }
 
 
@@ -1536,7 +1573,7 @@ public:
 	const TArray<TSharedPtrTS<FDashMPD_DescriptorType>>& GetRatings() const					{ return Ratings; }
 	const TArray<TSharedPtrTS<FDashMPD_DescriptorType>>& GetViewpoints() const				{ return Viewpoints; }
 	const TArray<TSharedPtrTS<FDashMPD_ContentComponentType>>& GetContentComponents() const	{ return ContentComponents; }
-	virtual const TArray<TSharedPtrTS<FDashMPD_BaseURLType>>& GetBaseURLs() const override	{ return BaseURLs; }
+	const TArray<TSharedPtrTS<FDashMPD_BaseURLType>>& GetBaseURLs() const override			{ return BaseURLs; }
 	TSharedPtrTS<FDashMPD_SegmentBaseType> GetSegmentBase() const							{ return SegmentBase; }
 	TSharedPtrTS<FDashMPD_SegmentListType> GetSegmentList() const							{ return SegmentList; }
 	TSharedPtrTS<FDashMPD_SegmentTemplateType> GetSegmentTemplate() const					{ return SegmentTemplate; }
@@ -1561,12 +1598,12 @@ public:
 	int64 GetSubsegmentStartsWithSAP() const						{ return SubsegmentStartsWithSAP; }
 	const TMediaOptionalValue<bool>& GetBitstreamSwitching() const	{ return bBitstreamSwitching; }
 	const TArray<uint64>& GetInitializationSetRefs() const			{ return InitializationSetRefs; }
-	virtual FXLink& GetXLink() override								{ return XLink; }
+	FXLink& GetXLink() override										{ return XLink; }
 
 private:
-	virtual bool ProcessElement(FManifestParserDASH* Builder, const TCHAR* ElementName, const TCHAR* ElementData, int32 XmlFileLineNumber) override;
-	virtual bool ProcessAttribute(FManifestParserDASH* Builder, const TCHAR* AttributeName, const TCHAR* AttributeValue) override;
-	virtual int32 ReplaceChildElementWithRemoteEntities(TSharedPtrTS<IDashMPDElement> Element, const FDashMPD_RootEntities& NewRootEntities, int64 OldResolveID, int64 NewResolveID) override;
+	bool ProcessElement(FManifestParserDASH* Builder, const TCHAR* ElementName, const TCHAR* ElementData, int32 XmlFileLineNumber) override;
+	bool ProcessAttribute(FManifestParserDASH* Builder, const TCHAR* AttributeName, const TCHAR* AttributeValue) override;
+	int32 ReplaceChildElementWithRemoteEntities(TSharedPtrTS<IDashMPDElement> Element, const FDashMPD_RootEntities& NewRootEntities, int64 OldResolveID, int64 NewResolveID) override;
 
 	TArray<TSharedPtrTS<FDashMPD_DescriptorType>> Accessibilities;
 	TArray<TSharedPtrTS<FDashMPD_DescriptorType>> Roles;
@@ -1613,28 +1650,28 @@ public:
 	FDashMPD_PeriodType(const TCHAR* Name, const TCHAR* Data) : Super(Name, Data)
 	{ }
 	virtual ~FDashMPD_PeriodType() = default;
-	virtual EType GetElementType() const override
+	EType GetElementType() const override
 	{ return EType::Period; }
 
-	virtual const TArray<TSharedPtrTS<FDashMPD_BaseURLType>>& GetBaseURLs() const override					{ return BaseURLs; }
-	TSharedPtrTS<FDashMPD_SegmentBaseType> GetSegmentBase() const											{ return SegmentBase; }
-	TSharedPtrTS<FDashMPD_SegmentListType> GetSegmentList() const											{ return SegmentList; }
-	TSharedPtrTS<FDashMPD_SegmentTemplateType> GetSegmentTemplate() const									{ return SegmentTemplate; }
-	TSharedPtrTS<FDashMPD_DescriptorType> GetAssetIdentifier() const										{ return AssetIdentifier; }
-	const TArray<TSharedPtrTS<FDashMPD_EventStreamType>>& GetEventStreams() const							{ return EventStreams; }
-	const TArray<TSharedPtrTS<FDashMPD_AdaptationSetType>>& GetAdaptationSets() const						{ return AdaptationSets; }
-	const TArray<TSharedPtrTS<FDashMPD_SubsetType>>& GetSubsets() const										{ return Subsets; }
-	virtual const TArray<TSharedPtrTS<FDashMPD_DescriptorType>>& GetSupplementalProperties() const override	{ return SupplementalProperties; }
-	const TArray<TSharedPtrTS<FDashMPD_ServiceDescriptionType>>& GetServiceDescriptions() const				{ return ServiceDescriptions; }
-	const TArray<TSharedPtrTS<FDashMPD_AdaptationSetType>>& GetEmptyAdaptationSets() const					{ return EmptyAdaptationSets; }
-	const TArray<TSharedPtrTS<FDashMPD_LabelType>>& GetGroupLabels() const									{ return GroupLabels; }
-	const TArray<TSharedPtrTS<FDashMPD_PreselectionType>>& GetPreselections() const							{ return Preselections; }
+	const TArray<TSharedPtrTS<FDashMPD_BaseURLType>>& GetBaseURLs() const override						{ return BaseURLs; }
+	TSharedPtrTS<FDashMPD_SegmentBaseType> GetSegmentBase() const										{ return SegmentBase; }
+	TSharedPtrTS<FDashMPD_SegmentListType> GetSegmentList() const										{ return SegmentList; }
+	TSharedPtrTS<FDashMPD_SegmentTemplateType> GetSegmentTemplate() const								{ return SegmentTemplate; }
+	TSharedPtrTS<FDashMPD_DescriptorType> GetAssetIdentifier() const									{ return AssetIdentifier; }
+	const TArray<TSharedPtrTS<FDashMPD_EventStreamType>>& GetEventStreams() const						{ return EventStreams; }
+	const TArray<TSharedPtrTS<FDashMPD_AdaptationSetType>>& GetAdaptationSets() const					{ return AdaptationSets; }
+	const TArray<TSharedPtrTS<FDashMPD_SubsetType>>& GetSubsets() const									{ return Subsets; }
+	const TArray<TSharedPtrTS<FDashMPD_DescriptorType>>& GetSupplementalProperties() const override		{ return SupplementalProperties; }
+	const TArray<TSharedPtrTS<FDashMPD_ServiceDescriptionType>>& GetServiceDescriptions() const			{ return ServiceDescriptions; }
+	const TArray<TSharedPtrTS<FDashMPD_AdaptationSetType>>& GetEmptyAdaptationSets() const				{ return EmptyAdaptationSets; }
+	const TArray<TSharedPtrTS<FDashMPD_LabelType>>& GetGroupLabels() const								{ return GroupLabels; }
+	const TArray<TSharedPtrTS<FDashMPD_PreselectionType>>& GetPreselections() const						{ return Preselections; }
 
 	const FString& GetID() const											{ return ID; }
 	const FTimeValue& GetStart() const										{ return Start; }
 	const FTimeValue& GetDuration() const									{ return Duration; }
 	const TMediaOptionalValue<bool>& GetBitstreamSwitching() const			{ return bBitstreamSwitching; }
-	virtual FXLink& GetXLink() override										{ return XLink; }
+	FXLink& GetXLink() override												{ return XLink; }
 
 	// Methods to manipulate the presentation type.
 	void SetDuration(const FTimeValue& NewDuration)
@@ -1643,9 +1680,9 @@ public:
 	}
 
 private:
-	virtual bool ProcessElement(FManifestParserDASH* Builder, const TCHAR* ElementName, const TCHAR* ElementData, int32 XmlFileLineNumber) override;
-	virtual bool ProcessAttribute(FManifestParserDASH* Builder, const TCHAR* AttributeName, const TCHAR* AttributeValue) override;
-	virtual int32 ReplaceChildElementWithRemoteEntities(TSharedPtrTS<IDashMPDElement> Element, const FDashMPD_RootEntities& NewRootEntities, int64 OldResolveID, int64 NewResolveID) override;
+	bool ProcessElement(FManifestParserDASH* Builder, const TCHAR* ElementName, const TCHAR* ElementData, int32 XmlFileLineNumber) override;
+	bool ProcessAttribute(FManifestParserDASH* Builder, const TCHAR* AttributeName, const TCHAR* AttributeValue) override;
+	int32 ReplaceChildElementWithRemoteEntities(TSharedPtrTS<IDashMPDElement> Element, const FDashMPD_RootEntities& NewRootEntities, int64 OldResolveID, int64 NewResolveID) override;
 
 	TArray<TSharedPtrTS<FDashMPD_BaseURLType>> BaseURLs;
 
@@ -1683,21 +1720,22 @@ public:
 	FDashMPD_MPDType(const TCHAR* Name, const TCHAR* Data) : Super(Name, Data)
 	{ }
 	virtual ~FDashMPD_MPDType() = default;
-	virtual EType GetElementType() const override
+	EType GetElementType() const override
 	{ return EType::MPD; }
 
-	const TArray<TSharedPtrTS<FDashMPD_ProgramInformationType>>& GetProgramInformations() const				{ return ProgramInformations; }
-	virtual const TArray<TSharedPtrTS<FDashMPD_BaseURLType>>& GetBaseURLs() const override					{ return BaseURLs; }
-	const TArray<TSharedPtrTS<FDashMPD_OtherType>>& GetLocations() const									{ return Locations; }
-	const TArray<TSharedPtrTS<FDashMPD_PeriodType>>& GetPeriods() const										{ return Periods; }
-	virtual const TArray<TSharedPtrTS<FDashMPD_DescriptorType>>& GetEssentialProperties() const override	{ return EssentialProperties; }
-	virtual const TArray<TSharedPtrTS<FDashMPD_DescriptorType>>& GetSupplementalProperties() const override	{ return SupplementalProperties; }
-	const TArray<TSharedPtrTS<FDashMPD_DescriptorType>>& GetUTCTimings() const								{ return UTCTimings; }
-	const TArray<TSharedPtrTS<FDashMPD_ServiceDescriptionType>>& GetServiceDescriptions() const				{ return ServiceDescriptions; }
-	const TArray<TSharedPtrTS<FDashMPD_InitializationSetType>>& GetInitializationSets() const				{ return InitializationSets; }
-	const TArray<TSharedPtrTS<FDashMPD_UIntVWithIDType>>& GetInitializationGroups() const					{ return InitializationGroups; }
-	const TArray<TSharedPtrTS<FDashMPD_UIntVWithIDType>>& GetInitializationPresentations() const			{ return InitializationPresentations; }
-	TSharedPtrTS<FDashMPD_LeapSecondInformationType> GetLeapSecondInformation() const						{ return LeapSecondInformation; }
+	const TArray<TSharedPtrTS<FDashMPD_ProgramInformationType>>& GetProgramInformations() const			{ return ProgramInformations; }
+	const TArray<TSharedPtrTS<FDashMPD_BaseURLType>>& GetBaseURLs() const override						{ return BaseURLs; }
+	const TArray<TSharedPtrTS<FDashMPD_OtherType>>& GetLocations() const								{ return Locations; }
+	const TArray<TSharedPtrTS<FDashMPD_PeriodType>>& GetPeriods() const									{ return Periods; }
+	const TArray<TSharedPtrTS<FDashMPD_DescriptorType>>& GetEssentialProperties() const override		{ return EssentialProperties; }
+	const TArray<TSharedPtrTS<FDashMPD_DescriptorType>>& GetSupplementalProperties() const override		{ return SupplementalProperties; }
+	const TArray<TSharedPtrTS<FDashMPD_DescriptorType>>& GetUTCTimings() const							{ return UTCTimings; }
+	const TArray<TSharedPtrTS<FDashMPD_ServiceDescriptionType>>& GetServiceDescriptions() const			{ return ServiceDescriptions; }
+	const TArray<TSharedPtrTS<FDashMPD_InitializationSetType>>& GetInitializationSets() const			{ return InitializationSets; }
+	const TArray<TSharedPtrTS<FDashMPD_UIntVWithIDType>>& GetInitializationGroups() const				{ return InitializationGroups; }
+	const TArray<TSharedPtrTS<FDashMPD_UIntVWithIDType>>& GetInitializationPresentations() const		{ return InitializationPresentations; }
+	TSharedPtrTS<FDashMPD_LeapSecondInformationType> GetLeapSecondInformation() const					{ return LeapSecondInformation; }
+	TSharedPtrTS<FDashMPD_ContentSteeringType> GetContentSteering() const								{ return ContentSteering; }
 
 	const FString& GetID() const							{ return ID; }
 	const TArray<FString>& GetProfiles()					{ return Profiles; }
@@ -1713,7 +1751,7 @@ public:
 	const FTimeValue& GetMaxSegmentDuration() const			{ return MaxSegmentDuration; }
 	const FTimeValue& GetMaxSubsegmentDuration() const		{ return MaxSubsegmentDuration; }
 
-	virtual FString GetDocumentURL() const override			{ return DocumentURL; }
+	FString GetDocumentURL() const override					{ return DocumentURL; }
 	void SetDocumentURL(const FString& InDocumentURL)		{ DocumentURL = InDocumentURL; }
 	const FTimeValue& GetFetchTime() const					{ FScopeLock Lock(&UpdateLock); return FetchTime; }
 	void SetFetchTime(const FTimeValue& InFetchTime)		{ FScopeLock Lock(&UpdateLock); FetchTime = InFetchTime; }
@@ -1733,6 +1771,9 @@ public:
 		Element->SetValue(InValue);
 	}
 
+	void SetDVBDash(bool bInDVBDash)						{ bIsDVBDash = bInDVBDash; }
+	bool IsDVBDash() const									{ return bIsDVBDash; }
+
 	// Methods to manipulate the presentation type.
 	void LockAccess()											{ UpdateLock.Lock(); }
 	void UnlockAccess()											{ UpdateLock.Unlock(); }
@@ -1746,9 +1787,9 @@ public:
 	void SetSuggestedPresentationDelay(const FTimeValue& InSPD)	{ SuggestedPresentationDelay = InSPD; }
 
 private:
-	virtual bool ProcessElement(FManifestParserDASH* Builder, const TCHAR* ElementName, const TCHAR* ElementData, int32 XmlFileLineNumber) override;
-	virtual bool ProcessAttribute(FManifestParserDASH* Builder, const TCHAR* AttributeName, const TCHAR* AttributeValue) override;
-	virtual int32 ReplaceChildElementWithRemoteEntities(TSharedPtrTS<IDashMPDElement> Element, const FDashMPD_RootEntities& NewRootEntities, int64 OldResolveID, int64 NewResolveID) override;
+	bool ProcessElement(FManifestParserDASH* Builder, const TCHAR* ElementName, const TCHAR* ElementData, int32 XmlFileLineNumber) override;
+	bool ProcessAttribute(FManifestParserDASH* Builder, const TCHAR* AttributeName, const TCHAR* AttributeValue) override;
+	int32 ReplaceChildElementWithRemoteEntities(TSharedPtrTS<IDashMPDElement> Element, const FDashMPD_RootEntities& NewRootEntities, int64 OldResolveID, int64 NewResolveID) override;
 
 	// The MPD is not expected to get modified after being built except for potential XLink changes
 	// and extending the validity if an MPD update returned a 304 on a conditional GET.
@@ -1760,7 +1801,10 @@ private:
 
 	// Optional ETag for conditional GET update requests.
 	FString ETag;
-	
+
+	// Set if the MPD is specifying a DVB-DASH profile.
+	bool bIsDVBDash = false;
+
 	// Time the MPD load request was started.
 	FTimeValue FetchTime;
 
@@ -1785,6 +1829,8 @@ private:
 	TArray<TSharedPtrTS<FDashMPD_UIntVWithIDType>> InitializationGroups;
 	TArray<TSharedPtrTS<FDashMPD_UIntVWithIDType>> InitializationPresentations;
 	TSharedPtrTS<FDashMPD_LeapSecondInformationType> LeapSecondInformation;
+	TSharedPtrTS<FDashMPD_ContentSteeringType> ContentSteering;
+
 
 	// Attributes of this MPD element.
 	FString ID;
