Electra: Added multiple CDN support and ContentSteering for DASH and HLS

diff --git a/Engine/Plugins/Media/ElectraPlayer/Source/ElectraPlayerRuntime/Private/Runtime/Player/DASH/ManifestBuilderDASH.cpp b/Engine/Plugins/Media/ElectraPlayer/Source/ElectraPlayerRuntime/Private/Runtime/Player/DASH/ManifestBuilderDASH.cpp
--- a/Engine/Plugins/Media/ElectraPlayer/Source/ElectraPlayerRuntime/Private/Runtime/Player/DASH/ManifestBuilderDASH.cpp
+++ b/Engine/Plugins/Media/ElectraPlayer/Source/ElectraPlayerRuntime/Private/Runtime/Player/DASH/ManifestBuilderDASH.cpp
@@ -39,6 +39,7 @@
 #define ERRCODE_DASH_MPD_BUILDER_XLINK_NOT_SUPPORTED_ON_ELEMENT		104
 #define ERRCODE_DASH_MPD_BUILDER_URL_FAILED_TO_RESOLVE				105
 #define ERRCODE_DASH_MPD_BUILDER_MISSING_REQUIRED_ELEMENT			106
+#define ERRCODE_DASH_MPD_BUILDER_AMBIGUOUS_URLS						107
 
 #define ERRCODE_DASH_MPD_BUILDER_BAD_PERIOD_START					200
 #define ERRCODE_DASH_MPD_BUILDER_EARLY_PERIODS_MUST_BE_LAST			201
@@ -83,6 +84,7 @@ namespace
 //		TEXT(""),
 		// DVB-DASH profiles
 		TEXT("urn:dvb:dash:profile:dvb-dash:2014"),
+		TEXT("urn:dvb:dash:profile:dvb-dash:isoff-ext-on-demand:2014"),
 		TEXT("urn:dvb:dash:profile:dvb-dash:isoff-ext-live:2014"),
 		// Low latency Live
 		TEXT("http://www.dashif.org/guidelines/low-latency-live-v5"),
@@ -91,6 +93,8 @@ namespace
 		TEXT("urn:mpeg:dash:profile:webm-on-demand:2012"),
 	};
 
+	const TCHAR* const DVBDashProfilePrefix = TEXT("urn:dvb:dash:profile:");
+
 	const TCHAR* const ScanTypeInterlace = TEXT("interlace");
 
 	const TCHAR* const AudioChannelConfigurationLegacy = TEXT("urn:mpeg:dash:23003:3:audio_channel_configuration:2011");
@@ -100,10 +104,14 @@ namespace
 	const TCHAR* const DASHRole = TEXT("urn:mpeg:dash:role:2011");
 
 	const TCHAR* const EssentialDashReservedPropertyPrefix = TEXT("urn:mpeg:dash:");
+
+	const TCHAR* const UrlParam2014 = TEXT("urn:mpeg:dash:urlparam:2014");
+	const TCHAR* const UrlParam2016 = TEXT("urn:mpeg:dash:urlparam:2016");
+
 	const TCHAR* const SupportedEssentialProperties[] =
 	{
-		TEXT("urn:mpeg:dash:urlparam:2014"),
-		TEXT("urn:mpeg:dash:urlparam:2016"),
+		UrlParam2014,
+		UrlParam2016
 	};
 
     const TCHAR* const Scheme_urn_mpeg_dash_mp4protection_2011 = TEXT("urn:mpeg:dash:mp4protection:2011");
@@ -157,6 +165,157 @@ namespace DASHUrlHelpers
 		return URL.StartsWith(SchemeHTTPS) || URL.StartsWith(SchemeHTTP) || URL.StartsWith(SchemeDATA);
 	}
 
+	bool AreBaseURLsValid(const TArray<TSharedPtrTS<FDashMPD_BaseURLType>>& InBaseURLs)
+	{
+		int32 NumRelativeURLs = 0;
+		int32 NumAbsoluteURLs = 0;
+		for(auto& BUrl : InBaseURLs)
+		{
+			if (DASHUrlHelpers::IsAbsoluteURL(BUrl->GetURL()))
+			{
+				++NumAbsoluteURLs;
+			}
+			else
+			{
+				++NumRelativeURLs;
+			}
+		}
+		//return NumRelativeURLs && NumAbsoluteURLs ? false : NumRelativeURLs <= 1;
+		return NumRelativeURLs <= 1;
+	}
+
+	bool AreURLQueryInfosValid(const TSharedPtrTS<const IDashMPDElement>& InElement)
+	{
+		if (!InElement.IsValid())
+		{
+			return true;
+		}
+		// Annex I.2 limits the 2014 param to at most one per hierarchy level.
+		// Note that for the 2016 version there is no such limit as different instances may be used to affect different types of requests.
+		int32 nEss2014 = InElement->GetEssentialProperties().FilterByPredicate([](const TSharedPtrTS<FDashMPD_DescriptorType>& d) {return d->GetSchemeIdUri().Equals(UrlParam2014);}).Num();
+		int32 nSupp2014 = InElement->GetSupplementalProperties().FilterByPredicate([](const TSharedPtrTS<FDashMPD_DescriptorType>& d) {return d->GetSchemeIdUri().Equals(UrlParam2014);}).Num();
+		return nEss2014 + nSupp2014 <= 1;
+	}
+
+
+	// Get all the <BaseURL>'s for this element throughout the MPD hierarchy, expanding relative URLs against their
+	// parent URLs up the the MPD document URL if necessary.
+	bool GetAllCandidateBaseURLs(TArray<FDASHMediaURL>& OutCandidateURLs, TSharedPtrTS<const IDashMPDElement> StartingAt, const FDASHMediaURL& InMPDUrl, bool bIsDVBDash)
+	{
+		// Collect all <BaseURL> elements in hierarchy order, starting at the current element going all the way to
+		// the top of the MPD.
+		TArray<TArray<TSharedPtrTS<FDashMPD_BaseURLType>>> HierarchicalBaseURLs;
+		while(StartingAt.IsValid())
+		{
+			const TArray<TSharedPtrTS<FDashMPD_BaseURLType>>& ElementBaseURLs = StartingAt->GetBaseURLs();
+			if (ElementBaseURLs.Num())
+			{
+				// Insert at the front, so the final list will be in hierarchy order.
+				HierarchicalBaseURLs.Insert(ElementBaseURLs, 0);
+			}
+			StartingAt = StartingAt->GetParentElement();
+		}
+		// Start building the candidate URLs
+		struct FParsedDASHMediaURL : FDASHMediaURL
+		{
+			FURL_RFC3986 Parsed;
+			void operator = (const FDASHMediaURL& InBase) { FDASHMediaURL::operator=(InBase); }
+			bool operator == (const FDASHMediaURL& InOther) const { return URL.Equals(InOther.URL) && CDN.Equals(InOther.CDN); }
+		};
+		FParsedDASHMediaURL RootDocument;
+		RootDocument = InMPDUrl;
+		if (!RootDocument.Parsed.Parse(RootDocument.URL))
+		{
+			UE_LOG(LogElectraPlayer, Error, TEXT("Could not parse the MPD URL when building element URL"));
+			return false;
+		}
+		TArray<FParsedDASHMediaURL> Base, Work;
+		Base.Emplace(RootDocument);
+		for(int32 nL=0; nL<HierarchicalBaseURLs.Num(); ++nL)
+		{
+			for(int32 nI=0; nI<Base.Num(); ++nI)
+			{
+				// Resolve each entry in `Base` against the <BaseURL> elements on the next hierarchy level.
+				for(int32 j=0; j<HierarchicalBaseURLs[nL].Num(); ++j)
+				{
+					auto bu = HierarchicalBaseURLs[nL][j];
+					FURL_RFC3986 nextUrl;
+					if (!nextUrl.Parse(bu->GetURL()))
+					{
+						UE_LOG(LogElectraPlayer, Error, TEXT("Could not parse a <BaseURL> element while resolving URL"));
+						return false;
+					}
+					FParsedDASHMediaURL& nxt(Work.Emplace_GetRef());
+					nxt.Parsed = Base[nI].Parsed;
+					nxt.Parsed.ResolveWith(nextUrl);
+					nxt.URL = nxt.Parsed.Get();
+					if (nextUrl.IsAbsolute())
+					{
+						nxt.CDN = bu->GetServiceLocation();
+						// An empty serviceLocation on an absolute URL uses the URL as the serviceLocation according to DVB-DASH.
+						if (nxt.CDN.IsEmpty())
+						{
+							if (bIsDVBDash)
+							{
+								nxt.CDN = nxt.URL;
+							}
+							else
+							{
+								// For non-DVB DASH we use the host as serviceLocation, mostly because it is a shorter string
+								// used in debugging...
+								nxt.CDN = nxt.Parsed.GetHost();
+							}
+						}
+						if (bu->GetDVBPriority().IsSet() && bu->GetDVBPriority().Value() < 0)
+						{
+							UE_LOG(LogElectraPlayer, Error, TEXT("Invalid @priority found on <BaseURL> element while resolving URL"));
+							return false;
+						}
+						nxt.DVBpriority = bu->GetDVBPriority().GetWithDefault(1);
+						if (bu->GetDVBWeight().IsSet() && bu->GetDVBWeight().Value() < 0)
+						{
+							UE_LOG(LogElectraPlayer, Error, TEXT("Invalid @weight found on <BaseURL> element while resolving URL"));
+							return false;
+						}
+						nxt.DVBweight = bu->GetDVBWeight().GetWithDefault(1);
+					}
+					else
+					{
+						nxt.CDN = Base[nI].CDN;
+						nxt.DVBpriority = Base[nI].DVBpriority;
+						nxt.DVBweight = Base[nI].DVBweight;
+					}
+					nxt.ATO = Base[nI].ATO + (bu->GetAvailabilityTimeOffset().IsValid() ? bu->GetAvailabilityTimeOffset() : FTimeValue::GetZero());
+					nxt.ByteRange = bu->GetByteRange().Len() ? bu->GetByteRange() : Base[nI].ByteRange;
+					if (!Base[nI].ATOComplete.IsSet())
+					{
+						nxt.ATOComplete = bu->GetAvailabilityTimeComplete();
+					}
+					else if (!Base[nI].ATOComplete.IsSet() && bu->GetAvailabilityTimeComplete().IsSet() && Base[nI].ATOComplete.Value() != bu->GetAvailabilityTimeComplete().Value())
+					{
+						UE_LOG(LogElectraPlayer, Warning, TEXT("Inconsistent @availabilityTimeComplete found on hierarchical <BaseURL> elements while resolving URL"));
+						// Use the child element's attribute here.
+						nxt.ATOComplete = bu->GetAvailabilityTimeComplete();
+					}
+				}
+			}
+			// The working list now becomes the new base list.
+			Base = MoveTemp(Work);
+		}
+
+		// Add unique (by URL and CDN) results to the return list.
+		// If on any hierarchy level there is a relative AND an absolute URL we may end up with the same
+		// URL multiple times.
+		for(auto& It : Base)
+		{
+			if (!OutCandidateURLs.Contains(It))
+			{
+				OutCandidateURLs.Add(It);
+			}
+		}
+		return true;
+	}
+
 	/**
 	 * Collects BaseURL elements beginning at the specified element up the MPD hierarchy.
 	 * At most one BaseURL per level is added to the output. If a preferred service location is specified the BaseURL element matching this will
@@ -204,9 +363,9 @@ namespace DASHUrlHelpers
 			// Get the essential and supplemental properties into a single list for processing.
 			TArray<TSharedPtrTS<FDashMPD_DescriptorType>> Desc;
 			Desc = StartingAt->GetEssentialProperties().FilterByPredicate([bInclude2014](const TSharedPtrTS<FDashMPD_DescriptorType>& d)
-				{ return d->GetSchemeIdUri().Equals(TEXT("urn:mpeg:dash:urlparam:2016")) || (bInclude2014 && d->GetSchemeIdUri().Equals(TEXT("urn:mpeg:dash:urlparam:2014"))); });
+				{ return d->GetSchemeIdUri().Equals(UrlParam2016) || (bInclude2014 && d->GetSchemeIdUri().Equals(UrlParam2014)); });
 			Desc.Append(StartingAt->GetSupplementalProperties().FilterByPredicate([bInclude2014](const TSharedPtrTS<FDashMPD_DescriptorType>& d)
-				{ return d->GetSchemeIdUri().Equals(TEXT("urn:mpeg:dash:urlparam:2016")) || (bInclude2014 && d->GetSchemeIdUri().Equals(TEXT("urn:mpeg:dash:urlparam:2014"))); }));
+				{ return d->GetSchemeIdUri().Equals(UrlParam2016) || (bInclude2014 && d->GetSchemeIdUri().Equals(UrlParam2014)); }));
 			for(int32 i=0; i<Desc.Num(); ++i)
 			{
 				const TArray<TSharedPtrTS<IDashMPDElement>>& WellKnown = Desc[i]->GetWellKnownDescriptors();
@@ -218,12 +377,16 @@ namespace DASHUrlHelpers
 						const TArray<FString>& IncIn = uq->GetIncludeInRequests();
 						for(int32 k=0; k<IncIn.Num(); ++k)
 						{
-							if ((ForRequestType == EUrlQueryRequestType::Segment  && IncIn[k].Equals(TEXT("segment"))) ||
+							// NOTE: The '*' wildcard is not yet officially supported.
+							//       This is an extension proposed in ETSI TS 103 998 V1.1.1 (2024-01)
+							if (IncIn[k].Equals(TEXT("*")) ||
+								(ForRequestType == EUrlQueryRequestType::Segment  && IncIn[k].Equals(TEXT("segment"))) ||
 								(ForRequestType == EUrlQueryRequestType::Xlink    && IncIn[k].Equals(TEXT("xlink"))) ||
 								(ForRequestType == EUrlQueryRequestType::Mpd      && IncIn[k].Equals(TEXT("mpd"))) ||
 								(ForRequestType == EUrlQueryRequestType::Callback && IncIn[k].Equals(TEXT("callback"))) ||
 								(ForRequestType == EUrlQueryRequestType::Chaining && IncIn[k].Equals(TEXT("chaining"))) ||
-								(ForRequestType == EUrlQueryRequestType::Fallback && IncIn[k].Equals(TEXT("fallback"))))
+								(ForRequestType == EUrlQueryRequestType::Fallback && IncIn[k].Equals(TEXT("fallback"))) ||
+								(ForRequestType == EUrlQueryRequestType::Steering && IncIn[k].Equals(TEXT("steering"))))
 							{
 								UrlQueries.Emplace(MoveTemp(uq));
 							}
@@ -562,7 +725,7 @@ namespace DASHUrlHelpers
 			bInUrlParsed = true;
 			InUrl.Parse(InOutURL);
 		}
-		InUrl.AddQueryParameters(TotalFinalQueryString, true);
+		InUrl.AddOrUpdateQueryParams(TotalFinalQueryString);
 		InOutURL = InUrl.Get();
 		OutRequestHeader = TotalFinalHttpRequestHeader;
 
@@ -580,11 +743,18 @@ namespace DASHUrlHelpers
 		FTimeValue SumOfUrlATOs(FTimeValue::GetZero());
 		ATO.SetToZero();
 		bATOComplete.Reset();
+
+		// Short circuit data: URLs
+		if (InElementURL.StartsWith(TEXT("data:")))
+		{
+			OutURL = InElementURL;
+			return true;
+		}
+
 		FURL_RFC3986 UrlParser;
-		FString ElementURL(InElementURL);
 
 		// If there is no element URL and BaseURLs this is probably for an MPD update in which case we return the original document URL.
-		if (ElementURL.IsEmpty() && BaseURLs.Num() == 0)
+		if (InElementURL.IsEmpty() && BaseURLs.Num() == 0)
 		{
 			OutURL = DocumentURL;
 			return true;
@@ -617,24 +787,9 @@ namespace DASHUrlHelpers
 		return UrlParser.IsAbsolute();
 	}
 
-	FString ApplyAnnexEByteRange(FString InURL, FString InRange, const TArray<TSharedPtrTS<const FDashMPD_BaseURLType>>& BaseURLs)
+	FString ApplyAnnexEByteRange(FString InURL, FString InRange, FString InByteRangeParam)
 	{
-		// Short circuit missing parameters.
-		if (InRange.IsEmpty() || BaseURLs.Num() == 0)
-		{
-			return MoveTemp(InURL);
-		}
-		// Find the first <BaseURL> element having a byteRange attribute.
-		FString byteRange;
-		for(int32 i=0; i<BaseURLs.Num(); ++i)
-		{
-			byteRange = BaseURLs[i]->GetByteRange();
-			if (!byteRange.IsEmpty())
-			{
-				break;
-			}
-		}
-		if (byteRange.IsEmpty())
+		if (InRange.IsEmpty() || InByteRangeParam.IsEmpty())
 		{
 			return MoveTemp(InURL);
 		}
@@ -648,12 +803,12 @@ namespace DASHUrlHelpers
 		InURL.FindChar(TCHAR('?'), QueryPos);
 		// Process the elements as described by the byte range.
 		FString NewURL;
-		while(!byteRange.IsEmpty())
+		while(!InByteRangeParam.IsEmpty())
 		{
 			int32 tokenPos = INDEX_NONE;
-			if (!byteRange.FindChar(TCHAR('$'), tokenPos))
+			if (!InByteRangeParam.FindChar(TCHAR('$'), tokenPos))
 			{
-				NewURL.Append(byteRange);
+				NewURL.Append(InByteRangeParam);
 				break;
 			}
 			else
@@ -661,14 +816,14 @@ namespace DASHUrlHelpers
 				// Append everything up to the first token.
 				if (tokenPos)
 				{
-					NewURL.Append(byteRange.Mid(0, tokenPos));
+					NewURL.Append(InByteRangeParam.Mid(0, tokenPos));
 				}
 				// Need to find another token.
-				int32 token2Pos = byteRange.Find(TEXT("$"), ESearchCase::CaseSensitive, ESearchDir::FromStart, tokenPos+1);
+				int32 token2Pos = InByteRangeParam.Find(TEXT("$"), ESearchCase::CaseSensitive, ESearchDir::FromStart, tokenPos+1);
 				if (token2Pos != INDEX_NONE)
 				{
-					FString token(byteRange.Mid(tokenPos+1, token2Pos-tokenPos-1));
-					byteRange.RightChopInline(token2Pos+1, EAllowShrinking::No);
+					FString token(InByteRangeParam.Mid(tokenPos+1, token2Pos-tokenPos-1));
+					InByteRangeParam.RightChopInline(token2Pos+1, EAllowShrinking::No);
 					// An empty token results from "$$" used to insert a single '$'.
 					if (token.IsEmpty())
 					{
@@ -691,9 +846,9 @@ namespace DASHUrlHelpers
 								NewURL.LeftChopInline(1, EAllowShrinking::No);
 							}
 							// If the next char in the template is a separator character it is to be removed.
-							if (byteRange.Len() && byteRange[0] == TCHAR('&'))
+							if (InByteRangeParam.Len() && InByteRangeParam[0] == TCHAR('&'))
 							{
-								byteRange.RightChopInline(1, EAllowShrinking::No);
+								InByteRangeParam.RightChopInline(1, EAllowShrinking::No);
 							}
 						}
 						else
@@ -726,6 +881,26 @@ namespace DASHUrlHelpers
 		return NewURL;
 	}
 
+	FString ApplyAnnexEByteRange(FString InURL, FString InRange, const TArray<TSharedPtrTS<const FDashMPD_BaseURLType>>& BaseURLs)
+	{
+		// Short circuit missing parameters.
+		if (InRange.IsEmpty() || BaseURLs.Num() == 0)
+		{
+			return MoveTemp(InURL);
+		}
+		// Find the first <BaseURL> element having a byteRange attribute.
+		FString byteRange;
+		for(int32 i=0; i<BaseURLs.Num(); ++i)
+		{
+			byteRange = BaseURLs[i]->GetByteRange();
+			if (!byteRange.IsEmpty())
+			{
+				break;
+			}
+		}
+		return ApplyAnnexEByteRange(InURL, InRange, byteRange);
+	}
+
 
 } // namespace DASHUrlHelpers
 
@@ -754,6 +929,7 @@ IManifestBuilderDASH* IManifestBuilderDASH::Create(IPlayerSessionServices* InPla
 }
 
 /*********************************************************************************************************************/
+
 FManifestBuilderDASH::FManifestBuilderDASH(IPlayerSessionServices* InPlayerSessionServices)
 	: PlayerSessionServices(InPlayerSessionServices)
 {
@@ -785,14 +961,18 @@ FErrorDetail FManifestBuilderDASH::BuildFromMPD(TSharedPtrTS<FManifestDASHIntern
 	{
 		return CreateErrorAndLog(PlayerSessionServices, FString::Printf(TEXT("MPD@profiles is required")), ERRCODE_DASH_MPD_BUILDER_MISSING_REQUIRED_ATTRIBUTE);
 	}
+	bool bUsesDVBDashProfile = false;
 	for(int32 nSupportedProfile=0; !bHasUsableProfile && nSupportedProfile<UE_ARRAY_COUNT(SupportedProfiles); ++nSupportedProfile)
 	{
-		for(int32 nProfile=0; !bHasUsableProfile && nProfile<MPDProfiles.Num(); ++nProfile)
+		for(int32 nProfile=0; nProfile<MPDProfiles.Num(); ++nProfile)
 		{
 			if (MPDProfiles[nProfile].Equals(SupportedProfiles[nSupportedProfile]))
 			{
 				bHasUsableProfile = true;
-				break;
+			}
+			if (MPDProfiles[nProfile].StartsWith(DVBDashProfilePrefix))
+			{
+				bUsesDVBDashProfile = true;
 			}
 		}
 	}
@@ -800,6 +980,7 @@ FErrorDetail FManifestBuilderDASH::BuildFromMPD(TSharedPtrTS<FManifestDASHIntern
 	{
 		return CreateErrorAndLog(PlayerSessionServices, FString::Printf(TEXT("MPD is of no supported profile")), ERRCODE_DASH_MPD_BUILDER_UNSUPPORTED_PROFILE);
 	}
+	NewMPD->SetDVBDash(bUsesDVBDashProfile);
 
 	// Check if this MPD requires essential properties we do not understand.
 	FString DashProtocol(TEXT("dash"));
@@ -958,7 +1139,7 @@ FErrorDetail FManifestDASHInternal::PrepareRemoteElementLoadRequest(TArray<TWeak
 		// Create the request.
 		TSharedPtrTS<FMPDLoadRequestDASH> LoadReq = MakeSharedTS<FMPDLoadRequestDASH>();
 		LoadReq->LoadType = LoadReqType;
-		LoadReq->URL = URL;
+		LoadReq->Url.URL = URL;
 		if (RequestHeader.Len())
 		{
 			LoadReq->Headers.Emplace(HTTP::FHTTPHeader({DASH::HTTPHeaderOptionName, RequestHeader}));
@@ -1147,6 +1328,18 @@ FErrorDetail FManifestDASHInternal::BuildAfterInitialRemoteElementDownload()
 	FTimeValue AST = GetAnchorTime();
 	SegmentFetchDelay.SetToZero();
 
+	// If on the MPD level there are multiple <BaseURL> elements then they need to be all absolute.
+	// There cannot be a mix between absolute and relative URLs and there can also not be multiple relative URLs either.
+	if (!DASHUrlHelpers::AreBaseURLsValid(MPDRoot->GetBaseURLs()))
+	{
+		return CreateErrorAndLog(PlayerSessionServices, FString::Printf(TEXT("Inconsistent relative <BaseURL> elements found on MPD level")), ERRCODE_DASH_MPD_BUILDER_AMBIGUOUS_URLS);
+	}
+	// There may be at most one UrlQueryInfo as per Annex I.2 on each hierarchy level.
+	if (!DASHUrlHelpers::AreURLQueryInfosValid(MPDRoot))
+	{
+		return CreateErrorAndLog(PlayerSessionServices, FString::Printf(TEXT("Multiple Annex I.2 URL query elements found on MPD level")), ERRCODE_DASH_MPD_BUILDER_AMBIGUOUS_URLS);
+	}
+
 	bool bWarnedPresentationDuration = false;
 	// Go over the periods one at a time as XLINK attributes could bring in additional periods.
 	FTimeValue CombinedPeriodDuration(FTimeValue::GetZero());
@@ -1166,6 +1359,17 @@ FErrorDetail FManifestDASHInternal::BuildAfterInitialRemoteElementDownload()
 			p->ID = FString::Printf(TEXT("$period.%d$"), nPeriod);
 		}
 
+		// As with the MPD <BaseURL> elements above the same is true on Period level.
+		if (!DASHUrlHelpers::AreBaseURLsValid(MPDperiods[nPeriod]->GetBaseURLs()))
+		{
+			return CreateErrorAndLog(PlayerSessionServices, FString::Printf(TEXT("Inconsistent relative <BaseURL> elements found in period \"%s\""), *p->ID), ERRCODE_DASH_MPD_BUILDER_AMBIGUOUS_URLS);
+		}
+		// And also the UrlQueryInfo
+		if (!DASHUrlHelpers::AreURLQueryInfosValid(MPDperiods[nPeriod]))
+		{
+			return CreateErrorAndLog(PlayerSessionServices, FString::Printf(TEXT("Multiple Annex I.2 URL query elements found in period \"%s\""), *p->ID), ERRCODE_DASH_MPD_BUILDER_AMBIGUOUS_URLS);
+		}
+
 		p->Start = MPDperiods[nPeriod]->GetStart();
 		p->Duration = MPDperiods[nPeriod]->GetDuration();
 		// If the period start is not specified this could be an early available period.
@@ -1395,12 +1599,48 @@ void FManifestDASHInternal::PreparePeriodAdaptationSets(TSharedPtrTS<FPeriod> Pe
 
 		IPlayerStreamFilter* StreamFilter = PlayerSessionServices->GetStreamFilter();
 
+		auto CheckEssentialProperties = [](const TArray<TSharedPtrTS<FDashMPD_DescriptorType>>& InEssentialProperties) -> bool
+		{
+			for(int32 nEssProp=0; nEssProp<InEssentialProperties.Num(); ++nEssProp)
+			{
+				// Check against a fixed list of properties we support internally.
+				for(int32 nSupportedEssProp=0; nSupportedEssProp<UE_ARRAY_COUNT(SupportedEssentialProperties); ++nSupportedEssProp)
+				{
+					if (!InEssentialProperties[nEssProp]->GetSchemeIdUri().Equals(SupportedEssentialProperties[nSupportedEssProp]))
+					{
+						return false;
+					}
+				}
+			}
+			return true;
+		};
+
 		Period->AdaptationSets.Empty();
 		const TArray<TSharedPtrTS<FDashMPD_AdaptationSetType>>& MPDAdaptationSets = MPDPeriod->GetAdaptationSets();
 		int32 nAdapt=0, nAdaptMax=MPDAdaptationSets.Num();
 		for(; nAdapt<nAdaptMax; ++nAdapt)
 		{
 			const TSharedPtrTS<FDashMPD_AdaptationSetType>& MPDAdaptationSet = MPDAdaptationSets[nAdapt];
+
+			// Check the <BaseURL> elements and the UrlQueryInfos in the adaptation set.
+			if (!DASHUrlHelpers::AreBaseURLsValid(MPDAdaptationSet->GetBaseURLs()))
+			{
+				LogMessage(PlayerSessionServices, IInfoLog::ELevel::Info, FString::Printf(TEXT("Inconsistent relative <BaseURL> elements found in AdaptationSet, ignoring this AdaptationSet.")));
+				continue;
+			}
+			// And also the UrlQueryInfo
+			if (!DASHUrlHelpers::AreURLQueryInfosValid(MPDAdaptationSet))
+			{
+				LogMessage(PlayerSessionServices, IInfoLog::ELevel::Info, FString::Printf(TEXT("Multiple Annex I.2 URL query elements found in AdaptationSet, ignoring this AdaptationSet.")));
+				continue;
+			}
+			// Check if there are any unsupported EssentialProperties. If so, this AdaptationSet cannot be used.
+			if (!CheckEssentialProperties(MPDAdaptationSet->GetEssentialProperties()))
+			{
+				LogMessage(PlayerSessionServices, IInfoLog::ELevel::Info, FString::Printf(TEXT("Unsupported <EssentialProperty> found in AdaptationSet, ignoring this AdaptationSet.")));
+				continue;
+			}
+
 			TSharedPtrTS<FAdaptationSet> AdaptationSet = MakeSharedTS<FAdaptationSet>();
 			// Store the index of this AdaptationSet in the array with the adaptation set.
 			// AdaptationSets do not need to have an ID and we will filter out some sets without removing
@@ -1527,6 +1767,26 @@ void FManifestDASHInternal::PreparePeriodAdaptationSets(TSharedPtrTS<FPeriod> Pe
 			for(int32 nRepr=0, nReprMax=MPDRepresentations.Num(); nRepr<nReprMax; ++nRepr)
 			{
 				const TSharedPtrTS<FDashMPD_RepresentationType>& MPDRepresentation = MPDRepresentations[nRepr];
+
+				// Check the <BaseURL> elements and the UrlQueryInfos
+				if (!DASHUrlHelpers::AreBaseURLsValid(MPDRepresentation->GetBaseURLs()))
+				{
+					LogMessage(PlayerSessionServices, IInfoLog::ELevel::Info, FString::Printf(TEXT("Inconsistent relative <BaseURL> elements found in Representation, ignoring this Representation.")));
+					continue;
+				}
+				// And also the UrlQueryInfo
+				if (!DASHUrlHelpers::AreURLQueryInfosValid(MPDRepresentation))
+				{
+					LogMessage(PlayerSessionServices, IInfoLog::ELevel::Info, FString::Printf(TEXT("Multiple Annex I.2 URL query elements found in Representation, ignoring this Representation.")));
+					continue;
+				}
+				// Check if there are any unsupported EssentialProperties. If so, this Representation cannot be used.
+				if (!CheckEssentialProperties(MPDRepresentation->GetEssentialProperties()))
+				{
+					LogMessage(PlayerSessionServices, IInfoLog::ELevel::Info, FString::Printf(TEXT("Unsupported <EssentialProperty> found in Representation, ignoring this Representation.")));
+					continue;
+				}
+
 				// Subrepresentations and dependent representations are not supported.
 				if (MPDRepresentation->GetSubRepresentations().Num())
 				{
diff --git a/Engine/Plugins/Media/ElectraPlayer/Source/ElectraPlayerRuntime/Private/Runtime/Player/DASH/ManifestDASH.cpp b/Engine/Plugins/Media/ElectraPlayer/Source/ElectraPlayerRuntime/Private/Runtime/Player/DASH/ManifestDASH.cpp
--- a/Engine/Plugins/Media/ElectraPlayer/Source/ElectraPlayerRuntime/Private/Runtime/Player/DASH/ManifestDASH.cpp
+++ b/Engine/Plugins/Media/ElectraPlayer/Source/ElectraPlayerRuntime/Private/Runtime/Player/DASH/ManifestDASH.cpp
@@ -21,7 +21,7 @@
 #include "Player/PlayerEntityCache.h"
 #include "Player/AdaptivePlayerOptionKeynames.h"
 #include "Player/DRM/DRMManager.h"
-
+#include "Player/ContentSteeringHandler.h"
 
 DECLARE_CYCLE_STAT(TEXT("FRepresentation::FindSegment"), STAT_ElectraPlayer_DASH_FindSegment, STATGROUP_ElectraPlayer);
 
@@ -532,6 +532,15 @@ void FManifestDASH::TriggerPlaylistRefresh()
 	}
 }
 
+void FManifestDASH::ReachedStableBuffer()
+{
+	TSharedPtrTS<FManifestDASHInternal> Manifest = CurrentManifest;
+	if (Manifest.IsValid())
+	{
+		PlayerSessionServices->GetContentSteeringHandler()->ReachedStableBuffer();
+	}
+}
+
 
 IStreamReader* FManifestDASH::CreateStreamReaderHandler()
 {
@@ -2013,11 +2022,11 @@ void FDASHPlayPeriod::HandleRepresentationInitSegmentLoading(const TArray<FInitS
 			is.bRequested = true;
 
 			IPlayerEntityCache::FCacheItem CachedItem;
-			if (!PlayerSessionServices->GetEntityCache()->GetCachedEntity(CachedItem, is.InitSegmentInfo.InitializationURL.URL, is.InitSegmentInfo.InitializationURL.Range))
+			if (!PlayerSessionServices->GetEntityCache()->GetCachedEntity(CachedItem, is.InitSegmentInfo.InitializationURL.Url.URL, is.InitSegmentInfo.InitializationURL.Range))
 			{
 				is.LoadRequest = MakeSharedTS<FMPDLoadRequestDASH>();
 				is.LoadRequest->LoadType = FMPDLoadRequestDASH::ELoadType::Segment;
-				is.LoadRequest->URL = is.InitSegmentInfo.InitializationURL.URL;
+				is.LoadRequest->Url.URL = is.InitSegmentInfo.InitializationURL.Url.URL;
 				is.LoadRequest->Range = is.InitSegmentInfo.InitializationURL.Range;
 				if (is.InitSegmentInfo.InitializationURL.CustomHeader.Len())
 				{
@@ -2066,7 +2075,7 @@ void FDASHPlayPeriod::InitSegmentMP4DownloadComplete(TSharedPtrTS<FMPDLoadReques
 			{
 				// Add this to the entity cache in case it needs to be retrieved again.
 				IPlayerEntityCache::FCacheItem CacheItem;
-				CacheItem.URL = LoadRequest->URL;
+				CacheItem.URL = LoadRequest->Url.URL;
 				CacheItem.Range = LoadRequest->Range;
 				CacheItem.Parsed14496_12Data = Init;
 				PlayerSessionServices->GetEntityCache()->CacheEntity(CacheItem);
@@ -2146,7 +2155,7 @@ FTimeValue CalculateSegmentAvailabilityTimeOffset(const TArray<TSharedPtrTS<FDas
 }
 
 
-
+#if 0
 TMediaOptionalValue<bool> GetSegmentAvailabilityTimeComplete(const TArray<TSharedPtrTS<FDashMPD_SegmentBaseType>>& Arr)
 {
 	TMediaOptionalValue<bool> TimeComplete;
@@ -2186,7 +2195,7 @@ TMediaOptionalValue<bool> GetSegmentAvailabilityTimeComplete(const TArray<TShare
 	}
 	return TimeComplete;
 }
-
+#endif
 }
 
 
@@ -2207,71 +2216,90 @@ FManifestDASHInternal::FRepresentation::ESearchResult FManifestDASHInternal::FRe
 
 		bNeedsSegmentIndex = false;
 
-		// We sort of need an index to figure out segment sizes and durations.
-		FString IndexRange = GET_ATTR(SegmentBase, GetIndexRange(), Len(), FString());
-		// If the index range is empty check if there is a RepresentationIndex. It should not be to specify the index URL but it may be there to specify the range!
-		if (IndexRange.IsEmpty())
+  		// Get the index, if any.
+		TSharedPtrTS<FDashMPD_URLType> RepresentationIndex = GET_ATTR(SegmentBase, GetRepresentationIndex(), IsValid(), TSharedPtrTS<FDashMPD_URLType>());
+		FString RepresentationIndexSegmentURL, RepresentationIndexSegmentURLRange;
+		if (RepresentationIndex.IsValid())
 		{
-			TSharedPtrTS<FDashMPD_URLType> RepresentationIndex = GET_ATTR(SegmentBase, GetRepresentationIndex(), IsValid(), TSharedPtrTS<FDashMPD_URLType>());
-			if (RepresentationIndex.IsValid())
+			RepresentationIndexSegmentURLRange = RepresentationIndex->GetRange();
+			RepresentationIndexSegmentURL = RepresentationIndex->GetSourceURL();
+			if (RepresentationIndexSegmentURL.Len())
 			{
-				IndexRange = RepresentationIndex->GetRange();
-				if (!RepresentationIndex->GetSourceURL().IsEmpty())
-				{
-					LogMessage(InPlayerSessionServices, IInfoLog::ELevel::Warning, FString::Printf(TEXT("Ignoring <RepresentationIndex> element URL present for <SegmentBase> of representation \"%s\""), *MPDRepresentation->GetID()));
-				}
+				LogMessage(InPlayerSessionServices, IInfoLog::ELevel::Warning, FString::Printf(TEXT("The <RepresentationIndex> within a <SegmentBase> specifies a URL, which is probably not correct.")));
+				/*
+					Ignore this URL since it is not clear what that means.
+					With <SegmentBase> the segment index is a `sidx` box and with almost 100% certainty the base segment file
+					will be a fragmented fmp4 consisting of several `moof`+`mdat` boxes, in which case the base address of the
+					moof will be the byte following the `sidx` box, but how would that work if this is represented by a
+					_different_ file? At what byte offset is the first `moof` to be found then?
+
+					This is expressed in section 8.3.3 of ISO/IEC 23009-1:2022
+				*/
+				RepresentationIndexSegmentURL.Empty();
 			}
 		}
-		if (!IndexRange.IsEmpty())
+		FString IndexRangeAttribute = GET_ATTR(SegmentBase, GetIndexRange(), Len(), FString());
+		if (RepresentationIndexSegmentURLRange.Len() && IndexRangeAttribute.Len() && !RepresentationIndexSegmentURLRange.Equals(IndexRangeAttribute))
 		{
-			ElectraHTTPStream::FHttpRange r;
-			r.Set(IndexRange);
-			if (r.IsSet())
-			{
-				check(r.GetStart() >= 0);
-				check(r.GetEndIncluding() > 0);
-				SegmentIndexRangeStart = r.GetStart();
-				SegmentIndexRangeSize = r.GetEndIncluding() + 1 - SegmentIndexRangeStart;
-			}
+			PostError(InPlayerSessionServices, FString::Printf(TEXT("The <RepresentationIndex> within a <SegmentBase> specifies a byte range different from the @indexRange attribute")), ERRCODE_DASH_MPD_BAD_REPRESENTATION);
+			return FManifestDASHInternal::FRepresentation::ESearchResult::BadType;
+		}
+		if (RepresentationIndexSegmentURLRange.IsEmpty())
+		{
+			RepresentationIndexSegmentURLRange = IndexRangeAttribute;
+		}
+		if (RepresentationIndexSegmentURLRange.IsEmpty())
+		{
+			PostError(InPlayerSessionServices, FString::Printf(TEXT("The <SegmentBase> does not specify any byte range for the index")), ERRCODE_DASH_MPD_BAD_REPRESENTATION);
+			return FManifestDASHInternal::FRepresentation::ESearchResult::BadType;
+		}
 
-			FString PreferredServiceLocation;
-			DashUtils::GetPlayerOption(InPlayerSessionServices, PreferredServiceLocation, DASH::OptionKey_CurrentCDN, FString());
-			// Get the relevant <BaseURL> elements from the hierarchy.
-			TArray<TSharedPtrTS<const FDashMPD_BaseURLType>> OutBaseURLs;
-			DASHUrlHelpers::GetAllHierarchyBaseURLs(InPlayerSessionServices, OutBaseURLs, MPDRepresentation, *PreferredServiceLocation);
-			if (OutBaseURLs.Num() == 0)
-			{
-				PostError(InPlayerSessionServices, FString::Printf(TEXT("Representation \"%s\" has no <BaseURL> element on any hierarchy level!"), *MPDRepresentation->GetID()), ERRCODE_DASH_MPD_BAD_REPRESENTATION);
-				bIsUsable = false;
-				return FManifestDASHInternal::FRepresentation::ESearchResult::BadType;
-			}
-			// Generate the absolute URL
-			FString DocumentURL = MPDRepresentation->GetDocumentURL();
-			FString URL, RequestHeader;
-			FTimeValue UrlATO;
-			TMediaOptionalValue<bool> bATOComplete;
-			if (!DASHUrlHelpers::BuildAbsoluteElementURL(URL, UrlATO, bATOComplete, DocumentURL, OutBaseURLs, FString()))
-			{
-				PostError(InPlayerSessionServices, FString::Printf(TEXT("Representation \"%s\" failed to resolve URL to an absolute URL!"), *MPDRepresentation->GetID()), ERRCODE_DASH_MPD_BAD_REPRESENTATION);
-				bIsUsable = false;
-				return FManifestDASHInternal::FRepresentation::ESearchResult::BadType;
-			}
 
+		// If there is a range for the index, remember it since forms the base address for moof addressing.
+		if (RepresentationIndexSegmentURLRange.Len())
+		{
+			ElectraHTTPStream::FHttpRange r;
+			r.Set(RepresentationIndexSegmentURLRange);
+			SegmentIndexRangeStart = r.GetStart();
+			SegmentIndexRangeSize = r.GetEndIncluding() + 1 - SegmentIndexRangeStart;
+		}
+
+		if (RepresentationIndexSegmentURLRange.Len() || RepresentationIndexSegmentURL.Len())
+		{
 			// The URL query might need to be changed. Look for the UrlQuery properties.
 			TArray<TSharedPtrTS<FDashMPD_UrlQueryInfoType>> UrlQueries;
 			DASHUrlHelpers::GetAllHierarchyUrlQueries(UrlQueries, MPDRepresentation, DASHUrlHelpers::EUrlQueryRequestType::Segment, true);
-			FErrorDetail Error = DASHUrlHelpers::ApplyUrlQueries(InPlayerSessionServices, DocumentURL, URL, RequestHeader, UrlQueries);
-			if (Error.IsSet())
-			{
-				PostError(InPlayerSessionServices, Error);
-				bIsUsable = false;
+
+			TArray<DASHUrlHelpers::FDASHMediaURL> CandidateURLs;
+
+			DASHUrlHelpers::FDASHMediaURL RootDocument;
+			RootDocument.URL = MPDRepresentation->GetDocumentURL();
+			// The MPD has no @serviceLocation or the DVB properties. According to DVB-DASH the @serviceLocation is the URL itself and the attributes default to 1.
+			RootDocument.CDN = RootDocument.URL;
+			if (!InPlayerSessionServices->GetContentSteeringHandler()->IsDVBDash())
+			{
+				FURL_RFC3986 RootDocURL;
+				RootDocURL.Parse(RootDocument.URL);
+				RootDocument.CDN = RootDocURL.GetHost();
+			}
+			RootDocument.DVBpriority = 1;
+			RootDocument.DVBweight = 1;
+			DASHUrlHelpers::GetAllCandidateBaseURLs(CandidateURLs, MPDRepresentation, RootDocument, InPlayerSessionServices->GetContentSteeringHandler()->IsDVBDash());
+			// Build a list we can pass to the steering handler.
+			TArray<FContentSteeringHandler::FCandidateURL> SteeringCandidates;
+			FString SteeringMessage;
+			PrepareSteeringCandidates(SteeringCandidates, InPlayerSessionServices, RootDocument.URL, CandidateURLs, UrlQueries, RepresentationIndexSegmentURL);
+			FContentSteeringHandler::FSelectedCandidateURL ChosenIndex = InPlayerSessionServices->GetContentSteeringHandler()->SelectBestCandidateFrom(SteeringMessage, FContentSteeringHandler::ESelectFor::Segment, SteeringCandidates);
+			if (ChosenIndex.MediaURL.URL.IsEmpty())
+			{
+				PostError(InPlayerSessionServices, SteeringMessage, ERRCODE_DASH_MPD_BAD_REPRESENTATION);
 				return FManifestDASHInternal::FRepresentation::ESearchResult::BadType;
 			}
-			URL = DASHUrlHelpers::ApplyAnnexEByteRange(URL, IndexRange, OutBaseURLs);
+			ChosenIndex.MediaURL.URL = DASHUrlHelpers::ApplyAnnexEByteRange(ChosenIndex.MediaURL.URL, RepresentationIndexSegmentURLRange, ChosenIndex.AdditionalParams.GetValue(DASHUrlHelpers::SteerOption_ByteRange).SafeGetFString());
 
 			// Check with entity cache if the index has been retrieved before.
 			IPlayerEntityCache::FCacheItem CachedItem;
-			if (InPlayerSessionServices->GetEntityCache()->GetCachedEntity(CachedItem, URL, IndexRange))
+			if (InPlayerSessionServices->GetEntityCache()->GetCachedEntity(CachedItem, ChosenIndex.MediaURL.URL, RepresentationIndexSegmentURLRange))
 			{
 				// Already cached. Use it.
 				SegmentIndexMP4 = CachedItem.Parsed14496_12Data;
@@ -2279,13 +2307,14 @@ FManifestDASHInternal::FRepresentation::ESearchResult FManifestDASHInternal::FRe
 			}
 			else
 			{
+				FString RequestHeader = ChosenIndex.AdditionalParams.GetValue(DASHUrlHelpers::SteerOption_AnnexIRequestHeader).SafeGetFString();
 				if (GetStreamContainerType() == EStreamContainerType::ISO14496_12)
 				{
 					// Create the request.
 					TSharedPtrTS<FMPDLoadRequestDASH> LoadReq = MakeSharedTS<FMPDLoadRequestDASH>();
 					LoadReq->LoadType = FMPDLoadRequestDASH::ELoadType::Segment;
-					LoadReq->URL = URL;
-					LoadReq->Range = IndexRange;
+					LoadReq->Url.URL = ChosenIndex.MediaURL.URL;
+					LoadReq->Range = RepresentationIndexSegmentURLRange;
 					if (RequestHeader.Len())
 					{
 						LoadReq->Headers.Emplace(HTTP::FHTTPHeader(DASH::HTTPHeaderOptionName, RequestHeader));
@@ -2310,8 +2339,8 @@ FManifestDASHInternal::FRepresentation::ESearchResult FManifestDASHInternal::FRe
 					}
 
 					LoadReq->LoadType = FMPDLoadRequestDASH::ELoadType::Segment;
-					LoadReq->URL = URL;
-					LoadReq->Range = IndexRange;
+					LoadReq->Url.URL = ChosenIndex.MediaURL.URL;;
+					LoadReq->Range = RepresentationIndexSegmentURLRange;
 					if (RequestHeader.Len())
 					{
 						LoadReq->Headers.Emplace(HTTP::FHTTPHeader(DASH::HTTPHeaderOptionName, RequestHeader));
@@ -2337,76 +2366,127 @@ FManifestDASHInternal::FRepresentation::ESearchResult FManifestDASHInternal::FRe
 }
 
 
+void FManifestDASHInternal::FRepresentation::PrepareSteeringCandidates(TArray<FContentSteeringHandler::FCandidateURL>& OutSteeringCandidates, IPlayerSessionServices* InPlayerSessionServices, const FString& InDocumentURL, const TArray<DASHUrlHelpers::FDASHMediaURL>& InCandidateURLs, const TArray<TSharedPtrTS<FDashMPD_UrlQueryInfoType>>& InUrlQueries, const FString& InTemplateURL)
+{
+	OutSteeringCandidates.Empty();
+	for(auto& cand : InCandidateURLs)
+	{
+		FContentSteeringHandler::FCandidateURL& sc(OutSteeringCandidates.Emplace_GetRef());
+		sc.MediaURL = cand;
+
+		// Resolve with media template string.
+		if (InTemplateURL.Len())
+		{
+			FURL_RFC3986 AbsoluteURL;
+			AbsoluteURL.Parse(sc.MediaURL.URL);
+			AbsoluteURL.ResolveWith(InTemplateURL);
+			sc.MediaURL.URL = AbsoluteURL.Get();
+		}
+		// Apply any of the Annex I UrlQuery elements
+		FString AnnexIRequestHeader;
+		DASHUrlHelpers::ApplyUrlQueries(InPlayerSessionServices, InDocumentURL, sc.MediaURL.URL, AnnexIRequestHeader, InUrlQueries);
+
+		// Set DASH specific URL parameters the steering handler does not need.
+		// Since the steering handler may clone a CDN on demand it needs to copy these values into the clone
+		// so the best way to do that is to set them in a dictionary.
+		sc.AdditionalParams.Set(DASHUrlHelpers::SteerOption_ByteRange, FVariantValue(cand.ByteRange));
+		sc.AdditionalParams.Set(DASHUrlHelpers::SteerOption_ATO, FVariantValue(cand.ATO));
+		sc.AdditionalParams.Set(DASHUrlHelpers::SteerOption_ATOComplete, FVariantValue((int64)(cand.ATOComplete.IsSet()?(cand.ATOComplete.Value()?1:0):-1)));
+		sc.AdditionalParams.Set(DASHUrlHelpers::SteerOption_AnnexIRequestHeader, FVariantValue(AnnexIRequestHeader));
+	}
+}
+
 bool FManifestDASHInternal::FRepresentation::PrepareDownloadURLs(IPlayerSessionServices* InPlayerSessionServices, FSegmentInformationCommon& InOutSegmentInfo, const TArray<TSharedPtrTS<FDashMPD_SegmentBaseType>>& SegmentBase)
 {
 	TSharedPtrTS<FDashMPD_RepresentationType> MPDRepresentation = Representation.Pin();
 
 	// Get the initialization, if any. If there is none the representation is supposed to be self-initializing.
 	TSharedPtrTS<FDashMPD_URLType> Initialization = GET_ATTR(SegmentBase, GetInitialization(), IsValid(), TSharedPtrTS<FDashMPD_URLType>());
+	FString InitializationSegmentURL, InitializationSegmentURLRange;
 	if (Initialization.IsValid())
 	{
-		InOutSegmentInfo.InitializationURL.Range = Initialization->GetRange();
-		if (!Initialization->GetSourceURL().IsEmpty())
-		{
-			LogMessage(InPlayerSessionServices, IInfoLog::ELevel::Warning, FString::Printf(TEXT("Ignoring <Initialization> element URL present for <SegmentBase> of representation \"%s\""), *MPDRepresentation->GetID()));
-		}
-	}
-
-	FString PreferredServiceLocation;
-	DashUtils::GetPlayerOption(InPlayerSessionServices, PreferredServiceLocation, DASH::OptionKey_CurrentCDN, FString());
-	// Get the relevant <BaseURL> elements from the hierarchy.
-	TArray<TSharedPtrTS<const FDashMPD_BaseURLType>> OutBaseURLs;
-	DASHUrlHelpers::GetAllHierarchyBaseURLs(InPlayerSessionServices, OutBaseURLs, MPDRepresentation, *PreferredServiceLocation);
-	if (OutBaseURLs.Num() == 0)
-	{
-		PostError(InPlayerSessionServices, FString::Printf(TEXT("Representation \"%s\" has no <BaseURL> element on any hierarchy level!"), *MPDRepresentation->GetID()), ERRCODE_DASH_MPD_BAD_REPRESENTATION);
-		return false;
-	}
-	// Generate the absolute URL
-	FString DocumentURL = MPDRepresentation->GetDocumentURL();
-	FString URL, RequestHeader;
-	FTimeValue UrlATO;
-	TMediaOptionalValue<bool> bATOComplete;
-	if (!DASHUrlHelpers::BuildAbsoluteElementURL(URL, UrlATO, bATOComplete, DocumentURL, OutBaseURLs, FString()))
-	{
-		PostError(InPlayerSessionServices, FString::Printf(TEXT("Representation \"%s\" failed to resolve URL to an absolute URL!"), *MPDRepresentation->GetID()), ERRCODE_DASH_MPD_BAD_REPRESENTATION);
-		return false;
+		InitializationSegmentURLRange = Initialization->GetRange();
+		InitializationSegmentURL = Initialization->GetSourceURL();
 	}
 
 	// The URL query might need to be changed. Look for the UrlQuery properties.
 	TArray<TSharedPtrTS<FDashMPD_UrlQueryInfoType>> UrlQueries;
 	DASHUrlHelpers::GetAllHierarchyUrlQueries(UrlQueries, MPDRepresentation, DASHUrlHelpers::EUrlQueryRequestType::Segment, true);
-	FErrorDetail Error = DASHUrlHelpers::ApplyUrlQueries(InPlayerSessionServices, DocumentURL, URL, RequestHeader, UrlQueries);
-	if (Error.IsSet())
-	{
-		PostError(InPlayerSessionServices, Error);
-		return false;
-	}
 
-	if (Initialization.IsValid())
-	{
-		InOutSegmentInfo.InitializationURL.URL = DASHUrlHelpers::ApplyAnnexEByteRange(URL, InOutSegmentInfo.InitializationURL.Range, OutBaseURLs);
-		InOutSegmentInfo.InitializationURL.CDN = PreferredServiceLocation;
-		InOutSegmentInfo.InitializationURL.CustomHeader = RequestHeader;
+	TArray<DASHUrlHelpers::FDASHMediaURL> CandidateURLs;
+
+	DASHUrlHelpers::FDASHMediaURL RootDocument;
+	RootDocument.URL = MPDRepresentation->GetDocumentURL();
+	// The MPD has no @serviceLocation or the DVB properties. According to DVB-DASH the @serviceLocation is the URL itself and the attributes default to 1.
+	RootDocument.CDN = RootDocument.URL;
+	if (!InPlayerSessionServices->GetContentSteeringHandler()->IsDVBDash())
+	{
+		FURL_RFC3986 RootDocURL;
+		RootDocURL.Parse(RootDocument.URL);
+		RootDocument.CDN = RootDocURL.GetHost();
+	}
+	RootDocument.DVBpriority = 1;
+	RootDocument.DVBweight = 1;
+	DASHUrlHelpers::GetAllCandidateBaseURLs(CandidateURLs, MPDRepresentation, RootDocument, InPlayerSessionServices->GetContentSteeringHandler()->IsDVBDash());
+	// Build a list we can pass to the steering handler.
+	TArray<FContentSteeringHandler::FCandidateURL> SteeringCandidates;
+	PrepareSteeringCandidates(SteeringCandidates, InPlayerSessionServices, RootDocument.URL, CandidateURLs, UrlQueries, FString());
+	// Have the steering handler select the desired CDN, potentially cloning an existing one on demand.
+	FString SteeringMessage;
+	FContentSteeringHandler::FSelectedCandidateURL ChosenMedia = InPlayerSessionServices->GetContentSteeringHandler()->SelectBestCandidateFrom(SteeringMessage, FContentSteeringHandler::ESelectFor::Segment, SteeringCandidates);
+	if (ChosenMedia.MediaURL.URL.IsEmpty())
+	{
+		PostError(InPlayerSessionServices, SteeringMessage, ERRCODE_DASH_MPD_BAD_REPRESENTATION);
+		return false;
 	}
 
+	// Setup the media URL parameters in the segment request.
+	InOutSegmentInfo.MediaURL.Url = ChosenMedia.MediaURL;
+	InOutSegmentInfo.MediaURL.CustomHeader = ChosenMedia.AdditionalParams.GetValue(DASHUrlHelpers::SteerOption_AnnexIRequestHeader).SafeGetFString();
+	InOutSegmentInfo.MediaURL.SteeringID = ChosenMedia.SteeringID;
+	// Content range?
 	if (InOutSegmentInfo.FirstByteOffset && InOutSegmentInfo.NumberOfBytes)
 	{
 		ElectraHTTPStream::FHttpRange r;
 		r.SetStart(InOutSegmentInfo.FirstByteOffset);
 		r.SetEndIncluding(InOutSegmentInfo.FirstByteOffset + InOutSegmentInfo.NumberOfBytes - 1);
 		InOutSegmentInfo.MediaURL.Range = r.GetString();
+		InOutSegmentInfo.MediaURL.Url.URL = DASHUrlHelpers::ApplyAnnexEByteRange(InOutSegmentInfo.MediaURL.Url.URL, InOutSegmentInfo.MediaURL.Range, ChosenMedia.AdditionalParams.GetValue(DASHUrlHelpers::SteerOption_ByteRange).SafeGetFString());
 	}
-	InOutSegmentInfo.MediaURL.URL = DASHUrlHelpers::ApplyAnnexEByteRange(URL, InOutSegmentInfo.MediaURL.Range, OutBaseURLs);
-	InOutSegmentInfo.MediaURL.CDN = PreferredServiceLocation;
-	InOutSegmentInfo.MediaURL.CustomHeader = RequestHeader;
 
-	// Add any availabilityTimeOffset from the <BaseURL>.
-	InOutSegmentInfo.ATO += UrlATO;
-	if (!InOutSegmentInfo.bAvailabilityTimeComplete.IsSet())
+	// Now do the same for the initialization segment if there is a dedicated one.
+	if (InitializationSegmentURL.Len())
+	{
+		PrepareSteeringCandidates(SteeringCandidates, InPlayerSessionServices, RootDocument.URL, CandidateURLs, UrlQueries, InitializationSegmentURL);
+		FContentSteeringHandler::FSelectedCandidateURL ChosenInit = InPlayerSessionServices->GetContentSteeringHandler()->SelectBestCandidateFrom(SteeringMessage, FContentSteeringHandler::ESelectFor::Segment, SteeringCandidates);
+		if (ChosenInit.MediaURL.URL.IsEmpty())
+		{
+			PostError(InPlayerSessionServices, SteeringMessage, ERRCODE_DASH_MPD_BAD_REPRESENTATION);
+			return false;
+		}
+
+		// Setup the init URL parameters in the segment request.
+		InOutSegmentInfo.InitializationURL.Url = ChosenInit.MediaURL;
+		InOutSegmentInfo.InitializationURL.CustomHeader = ChosenInit.AdditionalParams.GetValue(DASHUrlHelpers::SteerOption_AnnexIRequestHeader).SafeGetFString();
+		InOutSegmentInfo.InitializationURL.SteeringID = ChosenInit.SteeringID;
+		InOutSegmentInfo.InitializationURL.Range = InitializationSegmentURLRange;
+		InOutSegmentInfo.InitializationURL.Url.URL = DASHUrlHelpers::ApplyAnnexEByteRange(InOutSegmentInfo.InitializationURL.Url.URL, InOutSegmentInfo.InitializationURL.Range, ChosenInit.AdditionalParams.GetValue(DASHUrlHelpers::SteerOption_ByteRange).SafeGetFString());
+	}
+	else
 	{
-		InOutSegmentInfo.bAvailabilityTimeComplete = bATOComplete;
+		// The init segment is part of the media segment, so we can just copy this across.
+		InOutSegmentInfo.InitializationURL = InOutSegmentInfo.MediaURL;
+		// Except for the byte range, which is different.
+		InOutSegmentInfo.InitializationURL.Range = InitializationSegmentURLRange;
+		InOutSegmentInfo.InitializationURL.Url.URL = DASHUrlHelpers::ApplyAnnexEByteRange(InOutSegmentInfo.InitializationURL.Url.URL, InOutSegmentInfo.InitializationURL.Range, ChosenMedia.AdditionalParams.GetValue(DASHUrlHelpers::SteerOption_ByteRange).SafeGetFString());
 	}
+
+	// Finally apply the ATO of the media segment as the ATO in general.
+	// We load init segments only when requesting media segments at which time the init segment needs
+	// to be available anyway.
+	InOutSegmentInfo.ATO += ChosenMedia.AdditionalParams.GetValue(DASHUrlHelpers::SteerOption_ATO).SafeGetTimeValue(FTimeValue::GetZero());
+
+	// Set up additional things.
 	InOutSegmentInfo.bLowLatencyChunkedEncodingExpected = bAvailableAsLowLatency.GetWithDefault(false);
 	InOutSegmentInfo.ContainerType = StreamContainerType == EStreamContainerType::ISO14496_12 ? FSegmentInformationCommon::EContainerType::ISO14496_12 : FSegmentInformationCommon::EContainerType::Matroska;
 	return true;
@@ -2446,72 +2526,69 @@ bool FManifestDASHInternal::FRepresentation::PrepareDownloadURLs(IPlayerSessionS
 	FString MediaTemplateURL = ApplyTemplateStrings(MediaTemplate, InOutSegmentInfo);
 	FString InitTemplateURL = ApplyTemplateStrings(InitializationTemplate, InOutSegmentInfo);
 
-	// Get the preferred CDN and the <BaseURL> and <UrlQueryInfo> elements affecting the URL assembly.
-	FString PreferredServiceLocation;
-	DashUtils::GetPlayerOption(InPlayerSessionServices, PreferredServiceLocation, DASH::OptionKey_CurrentCDN, FString());
-	// Get the relevant <BaseURL> elements from the hierarchy.
-	TArray<TSharedPtrTS<const FDashMPD_BaseURLType>> OutBaseURLs;
-	DASHUrlHelpers::GetAllHierarchyBaseURLs(InPlayerSessionServices, OutBaseURLs, MPDRepresentation, *PreferredServiceLocation);
 	// The URL query might need to be changed. Look for the UrlQuery properties.
 	TArray<TSharedPtrTS<FDashMPD_UrlQueryInfoType>> UrlQueries;
 	DASHUrlHelpers::GetAllHierarchyUrlQueries(UrlQueries, MPDRepresentation, DASHUrlHelpers::EUrlQueryRequestType::Segment, true);
 
-	// Generate the absolute media URL
-	FString DocumentURL = MPDRepresentation->GetDocumentURL();
-	FString MediaURL, MediaRequestHeader;
-	FTimeValue MediaUrlATO;
-	TMediaOptionalValue<bool> bMediaATOComplete;
-	if (!DASHUrlHelpers::BuildAbsoluteElementURL(MediaURL, MediaUrlATO, bMediaATOComplete, DocumentURL, OutBaseURLs, MediaTemplateURL))
-	{
-		PostError(InPlayerSessionServices, FString::Printf(TEXT("Representation \"%s\" failed to resolve media segment URL to an absolute URL!"), *MPDRepresentation->GetID()), ERRCODE_DASH_MPD_BAD_REPRESENTATION);
-		return false;
-	}
-	FErrorDetail Error = DASHUrlHelpers::ApplyUrlQueries(InPlayerSessionServices, DocumentURL, MediaURL, MediaRequestHeader, UrlQueries);
-	if (Error.IsSet())
-	{
-		PostError(InPlayerSessionServices, Error);
-		return false;
-	}
-
-	// And also generate the absolute init segment URL
-	FString InitURL, InitRequestHeader;
-	FTimeValue InitUrlATO;
-	TMediaOptionalValue<bool> bInitATOComplete;
-	if (!DASHUrlHelpers::BuildAbsoluteElementURL(InitURL, InitUrlATO, bInitATOComplete, DocumentURL, OutBaseURLs, InitTemplateURL))
-	{
-		PostError(InPlayerSessionServices, FString::Printf(TEXT("Representation \"%s\" failed to resolve init segment URL to an absolute URL!"), *MPDRepresentation->GetID()), ERRCODE_DASH_MPD_BAD_REPRESENTATION);
-		return false;
-	}
-	Error = DASHUrlHelpers::ApplyUrlQueries(InPlayerSessionServices, DocumentURL, InitURL, InitRequestHeader, UrlQueries);
-	if (Error.IsSet())
-	{
-		PostError(InPlayerSessionServices, Error);
+	TArray<DASHUrlHelpers::FDASHMediaURL> CandidateURLs;
+
+	DASHUrlHelpers::FDASHMediaURL RootDocument;
+	RootDocument.URL = MPDRepresentation->GetDocumentURL();
+	// The MPD has no @serviceLocation or the DVB properties. According to DVB-DASH the @serviceLocation is the URL itself and the attributes default to 1.
+	RootDocument.CDN = RootDocument.URL;
+	if (!InPlayerSessionServices->GetContentSteeringHandler()->IsDVBDash())
+	{
+		FURL_RFC3986 RootDocURL;
+		RootDocURL.Parse(RootDocument.URL);
+		RootDocument.CDN = RootDocURL.GetHost();
+	}
+	RootDocument.DVBpriority = 1;
+	RootDocument.DVBweight = 1;
+	DASHUrlHelpers::GetAllCandidateBaseURLs(CandidateURLs, MPDRepresentation, RootDocument, InPlayerSessionServices->GetContentSteeringHandler()->IsDVBDash());
+	// Build a list we can pass to the steering handler.
+	TArray<FContentSteeringHandler::FCandidateURL> SteeringCandidates;
+	PrepareSteeringCandidates(SteeringCandidates, InPlayerSessionServices, RootDocument.URL, CandidateURLs, UrlQueries, MediaTemplateURL);
+	// Have the steering handler select the desired CDN, potentially cloning an existing one on demand.
+	FString SteeringMessage;
+	FContentSteeringHandler::FSelectedCandidateURL ChosenMedia = InPlayerSessionServices->GetContentSteeringHandler()->SelectBestCandidateFrom(SteeringMessage, FContentSteeringHandler::ESelectFor::Segment, SteeringCandidates);
+	if (ChosenMedia.MediaURL.URL.IsEmpty())
+	{
+		PostError(InPlayerSessionServices, SteeringMessage, ERRCODE_DASH_MPD_BAD_REPRESENTATION);
 		return false;
 	}
 
-	InOutSegmentInfo.InitializationURL.URL = DASHUrlHelpers::ApplyAnnexEByteRange(InitURL, InOutSegmentInfo.InitializationURL.Range, OutBaseURLs);
-	InOutSegmentInfo.InitializationURL.CDN = PreferredServiceLocation;
-	InOutSegmentInfo.InitializationURL.CustomHeader = InitRequestHeader;
+	// Setup the media URL parameters in the segment request.
+	InOutSegmentInfo.MediaURL.Url = ChosenMedia.MediaURL;
+	InOutSegmentInfo.MediaURL.CustomHeader = ChosenMedia.AdditionalParams.GetValue(DASHUrlHelpers::SteerOption_AnnexIRequestHeader).SafeGetFString();
+	InOutSegmentInfo.MediaURL.SteeringID = ChosenMedia.SteeringID;
 
-	if (InOutSegmentInfo.FirstByteOffset && InOutSegmentInfo.NumberOfBytes)
+	// Now do the same for the initialization segment
+	if (!InitTemplateURL.IsEmpty())
 	{
-		ElectraHTTPStream::FHttpRange r;
-		r.SetStart(InOutSegmentInfo.FirstByteOffset);
-		r.SetEndIncluding(InOutSegmentInfo.FirstByteOffset + InOutSegmentInfo.NumberOfBytes - 1);
-		InOutSegmentInfo.MediaURL.Range = r.GetString();
-	}
-	InOutSegmentInfo.MediaURL.URL = DASHUrlHelpers::ApplyAnnexEByteRange(MediaURL, InOutSegmentInfo.MediaURL.Range, OutBaseURLs);
-	InOutSegmentInfo.MediaURL.CDN = PreferredServiceLocation;
-	InOutSegmentInfo.MediaURL.CustomHeader = MediaRequestHeader;
+		PrepareSteeringCandidates(SteeringCandidates, InPlayerSessionServices, RootDocument.URL, CandidateURLs, UrlQueries, InitTemplateURL);
+		FContentSteeringHandler::FSelectedCandidateURL ChosenInit = InPlayerSessionServices->GetContentSteeringHandler()->SelectBestCandidateFrom(SteeringMessage, FContentSteeringHandler::ESelectFor::Segment, SteeringCandidates);
+		if (ChosenInit.MediaURL.URL.IsEmpty())
+		{
+			PostError(InPlayerSessionServices, SteeringMessage, ERRCODE_DASH_MPD_BAD_REPRESENTATION);
+			return false;
+		}
 
-	// Add any availabilityTimeOffset from the <BaseURL>. We use the media ATO since we request the init segment at the same time
-	// as the media segment and the init segment needs to be available at the same time anyway.
-	InOutSegmentInfo.ATO += MediaUrlATO;
-	if (!InOutSegmentInfo.bAvailabilityTimeComplete.IsSet())
-	{
-		InOutSegmentInfo.bAvailabilityTimeComplete = bMediaATOComplete;
+		// Setup the init URL parameters in the segment request.
+		InOutSegmentInfo.InitializationURL.Url = ChosenInit.MediaURL;
+		InOutSegmentInfo.InitializationURL.CustomHeader = ChosenInit.AdditionalParams.GetValue(DASHUrlHelpers::SteerOption_AnnexIRequestHeader).SafeGetFString();
+		InOutSegmentInfo.InitializationURL.SteeringID = ChosenInit.SteeringID;
+		// The init segment may have a byte range if provided through an <Initialization> element in the <SegmentTemplate> and that may need
+		// to be handled according to Annex E byte range request rules.
+		// NOTE: It is *not* possible to have a byte range for the media segments under a <SegmentTemplate>, so we do not do the same above for the media segments.
+		InOutSegmentInfo.InitializationURL.Url.URL = DASHUrlHelpers::ApplyAnnexEByteRange(InOutSegmentInfo.InitializationURL.Url.URL, InOutSegmentInfo.InitializationURL.Range, ChosenInit.AdditionalParams.GetValue(DASHUrlHelpers::SteerOption_ByteRange).SafeGetFString());
 	}
 
+	// Finally apply the ATO of the media segment as the ATO in general.
+	// We load init segments only when requesting media segments at which time the init segment needs
+	// to be available anyway.
+	InOutSegmentInfo.ATO += ChosenMedia.AdditionalParams.GetValue(DASHUrlHelpers::SteerOption_ATO).SafeGetTimeValue(FTimeValue::GetZero());
+
+	// Set up additional things.
 	InOutSegmentInfo.bLowLatencyChunkedEncodingExpected = bAvailableAsLowLatency.GetWithDefault(false);
 	InOutSegmentInfo.ContainerType = StreamContainerType == EStreamContainerType::ISO14496_12 ? FSegmentInformationCommon::EContainerType::ISO14496_12 : FSegmentInformationCommon::EContainerType::Matroska;
 	return true;
@@ -2819,7 +2896,7 @@ FManifestDASHInternal::FRepresentation::ESearchResult FManifestDASHInternal::FRe
 		LogMessage(InPlayerSessionServices, IInfoLog::ELevel::Warning, FString::Printf(TEXT("Representation timescale (%u) in MPD is not equal to timescale used in the segment index (%u) for Representation \"%s\"."), MPDTimescale, SidxTimescale, *MPDRepresentation->GetID()));
 	}
 	FTimeValue ATO = CalculateSegmentAvailabilityTimeOffset(SegmentBase);
-	TMediaOptionalValue<bool> bATOComplete = GetSegmentAvailabilityTimeComplete(SegmentBase);
+	//TMediaOptionalValue<bool> bATOComplete = GetSegmentAvailabilityTimeComplete(SegmentBase);
 
 
 	// Convert the local media search time to the timescale of the segment index.
@@ -2981,7 +3058,6 @@ FManifestDASHInternal::FRepresentation::ESearchResult FManifestDASHInternal::FRe
 	// Did we find it?
 	if (bFound && CurrentT < MediaLocalEndTime)
 	{
-		OutSegmentInfo.PeriodLocalSegmentStartTime.SetFromND(CurrentT - PTO, SidxTimescale);
 		OutSegmentInfo.Time = CurrentT;
 		OutSegmentInfo.PTO = PTO;
 		OutSegmentInfo.Duration = CurrentD;
@@ -2992,7 +3068,6 @@ FManifestDASHInternal::FRepresentation::ESearchResult FManifestDASHInternal::FRe
 		OutSegmentInfo.MediaLocalLastAUTime = MediaLocalEndTime;
 		OutSegmentInfo.Timescale = SidxTimescale;
 		OutSegmentInfo.ATO = ATO;
-		OutSegmentInfo.bAvailabilityTimeComplete = bATOComplete;
 		OutSegmentInfo.bIsLastInPeriod = CurrentT + CurrentD >= MediaLocalEndTime;
 		OutSegmentInfo.bFrameAccuracyRequired = InSearchOptions.bFrameAccurateSearch;
 		CollectInbandEventStreams(InPlayerSessionServices, OutSegmentInfo);
@@ -3044,7 +3119,7 @@ FManifestDASHInternal::FRepresentation::ESearchResult FManifestDASHInternal::FRe
 	// the search time onto the media internal timeline.
 	int64 PTO = 0;
 	FTimeValue ATO = CalculateSegmentAvailabilityTimeOffset(SegmentBase);
-	TMediaOptionalValue<bool> bATOComplete = GetSegmentAvailabilityTimeComplete(SegmentBase);
+	//TMediaOptionalValue<bool> bATOComplete = GetSegmentAvailabilityTimeComplete(SegmentBase);
 	{
 		uint32 MPDTimescale = GET_ATTR(SegmentBase, GetTimescale(), IsSet(), TMediaOptionalValue<uint32>(1)).Value();
 		uint64 pto = GET_ATTR(SegmentBase, GetPresentationTimeOffset(), IsSet(), TMediaOptionalValue<uint64>(0)).Value();
@@ -3085,7 +3160,6 @@ FManifestDASHInternal::FRepresentation::ESearchResult FManifestDASHInternal::FRe
 	// Did we find it?
 	if (Error == UEMEDIA_ERROR_OK)
 	{
-		OutSegmentInfo.PeriodLocalSegmentStartTime = TrackIt->GetTimestamp() - FTimeValue(PTO);
 		OutSegmentInfo.Time = TrackIt->GetTimestamp().GetAsHNS();
 		OutSegmentInfo.PTO = PTO;
 		OutSegmentInfo.Duration = TrackIt->GetClusterDuration().GetAsHNS();
@@ -3096,7 +3170,6 @@ FManifestDASHInternal::FRepresentation::ESearchResult FManifestDASHInternal::FRe
 		OutSegmentInfo.MediaLocalLastAUTime = MediaLocalEndTime;
 		OutSegmentInfo.Timescale = 10000000;
 		OutSegmentInfo.ATO = ATO;
-		OutSegmentInfo.bAvailabilityTimeComplete = bATOComplete;
 		OutSegmentInfo.bIsLastInPeriod = TrackIt->IsLastCluster() || (TrackIt->GetTimestamp() + TrackIt->GetClusterDuration()).GetAsHNS() >= MediaLocalEndTime;
 		OutSegmentInfo.bFrameAccuracyRequired = InSearchOptions.bFrameAccurateSearch;
 		if (!PrepareDownloadURLs(InPlayerSessionServices, OutSegmentInfo, SegmentBase))
@@ -3124,7 +3197,7 @@ FManifestDASHInternal::FRepresentation::ESearchResult FManifestDASHInternal::FRe
 	TMediaOptionalValue<uint32> Duration = GET_ATTR(SegmentTemplate, GetDuration(), IsSet(), TMediaOptionalValue<uint32>());
 	TMediaOptionalValue<int32> EptDelta = GET_ATTR(SegmentTemplate, GetEptDelta(), IsSet(), TMediaOptionalValue<int32>());
 	FTimeValue ATO = CalculateSegmentAvailabilityTimeOffset(SegmentTemplate);
-	TMediaOptionalValue<bool> bATOComplete = GetSegmentAvailabilityTimeComplete(SegmentTemplate);
+	//TMediaOptionalValue<bool> bATOComplete = GetSegmentAvailabilityTimeComplete(SegmentTemplate);
 
 	// The timescale should in all likelihood not be 1. While certainly allowed an accuracy of only one second is more likely to be
 	// an oversight when building the MPD.
@@ -3234,7 +3307,6 @@ FManifestDASHInternal::FRepresentation::ESearchResult FManifestDASHInternal::FRe
 		return FManifestDASHInternal::FRepresentation::ESearchResult::PastEOS;
 	}
 
-	OutSegmentInfo.PeriodLocalSegmentStartTime.SetFromND(SegmentNum * (int64)SegmentDuration, MPDTimescale);
 	OutSegmentInfo.Time = PTO + EPTdelta + SegmentNum * (int64)SegmentDuration;
 	OutSegmentInfo.PTO = PTO;
 	OutSegmentInfo.EPTdelta = EPTdelta;
@@ -3247,7 +3319,6 @@ FManifestDASHInternal::FRepresentation::ESearchResult FManifestDASHInternal::FRe
 	OutSegmentInfo.bIsLastInPeriod = OutSegmentInfo.bMayBeMissing && InSearchOptions.bHasFollowingPeriod;
 	OutSegmentInfo.bFrameAccuracyRequired = InSearchOptions.bFrameAccurateSearch;
 	OutSegmentInfo.ATO = ATO;
-	OutSegmentInfo.bAvailabilityTimeComplete = bATOComplete;
 	CollectInbandEventStreams(InPlayerSessionServices, OutSegmentInfo);
 	SetupProducerReferenceTimeInfo(InPlayerSessionServices, OutSegmentInfo);
 	if (!PrepareDownloadURLs(InPlayerSessionServices, OutSegmentInfo, SegmentTemplate))
@@ -3284,7 +3355,7 @@ FManifestDASHInternal::FRepresentation::ESearchResult FManifestDASHInternal::FRe
 	uint32 EndNumber = GET_ATTR(SegmentTemplate, GetEndNumber(), IsSet(), TMediaOptionalValue<uint32>(~0U)).Value();
 	//TMediaOptionalValue<uint32> Duration = GET_ATTR(SegmentTemplate, GetDuration(), IsSet(), TMediaOptionalValue<uint32>());
 	FTimeValue ATO = CalculateSegmentAvailabilityTimeOffset(SegmentTemplate);
-	TMediaOptionalValue<bool> bATOComplete = GetSegmentAvailabilityTimeComplete(SegmentTemplate);
+	//TMediaOptionalValue<bool> bATOComplete = GetSegmentAvailabilityTimeComplete(SegmentTemplate);
 
 	// The timescale should in all likelihood not be 1. While certainly allowed an accuracy of only one second is more likely to be
 	// an oversight when building the MPD.
@@ -3563,7 +3634,6 @@ FManifestDASHInternal::FRepresentation::ESearchResult FManifestDASHInternal::FRe
 	// Did we find it?
 	if (bFound && CurrentT < MediaLocalEndTime)
 	{
-		OutSegmentInfo.PeriodLocalSegmentStartTime.SetFromND(CurrentT - PTO, MPDTimescale);
 		OutSegmentInfo.Time = CurrentT;
 		OutSegmentInfo.PTO = PTO;
 		OutSegmentInfo.Duration = CurrentD;
@@ -3581,7 +3651,6 @@ FManifestDASHInternal::FRepresentation::ESearchResult FManifestDASHInternal::FRe
 			LogMessage(InPlayerSessionServices, IInfoLog::ELevel::Warning, FString::Printf(TEXT("Representation \"%s\" <SegmentTimeline> gap encountered for needed 't' value of %lld. Replacing with an empty filler segment."), *MPDRepresentation->GetID(), (long long int)CurrentT));
 		}
 		OutSegmentInfo.ATO = ATO;
-		OutSegmentInfo.bAvailabilityTimeComplete = bATOComplete;
 		CollectInbandEventStreams(InPlayerSessionServices, OutSegmentInfo);
 		SetupProducerReferenceTimeInfo(InPlayerSessionServices, OutSegmentInfo);
 		if (!PrepareDownloadURLs(InPlayerSessionServices, OutSegmentInfo, SegmentTemplate))
@@ -3647,7 +3716,6 @@ FManifestDASHInternal::FRepresentation::ESearchResult FManifestDASHInternal::FRe
 	OutSegmentInfo.MediaLocalLastAUTime = MediaLocalPeriodEnd;
 	OutSegmentInfo.Timescale = Timescale;
 	OutSegmentInfo.ATO.SetToZero();
-	OutSegmentInfo.bAvailabilityTimeComplete.Set(true);
 	OutSegmentInfo.bIsSideload = true;
 	OutSegmentInfo.bIsLastInPeriod = true;
 	TArray<TSharedPtrTS<FDashMPD_SegmentBaseType>> SegmentBase;
@@ -3792,7 +3860,7 @@ void FManifestDASHInternal::FRepresentation::SegmentIndexDownloadComplete(TShare
 						bOk = true;
 						// Add this to the entity cache in case it needs to be retrieved again.
 						IPlayerEntityCache::FCacheItem ci;
-						ci.URL = LoadRequest->URL;
+						ci.URL = LoadRequest->Url.URL;
 						ci.Range = LoadRequest->Range;
 						ci.Parsed14496_12Data = SegmentIndexMP4;
 						LoadRequest->PlayerSessionServices->GetEntityCache()->CacheEntity(ci);
@@ -3818,7 +3886,7 @@ void FManifestDASHInternal::FRepresentation::SegmentIndexDownloadComplete(TShare
 			{
 				TSharedPtrTS<FMPDLoadRequestDASH> NextChainedReq = MakeSharedTS<FMPDLoadRequestDASH>();
 				NextChainedReq->LoadType = FMPDLoadRequestDASH::ELoadType::Segment;
-				NextChainedReq->URL = LoadRequest->URL;
+				NextChainedReq->Url = LoadRequest->Url;
 				NextChainedReq->Range = LoadRequest->Range2;
 				NextChainedReq->Headers = LoadRequest->Headers;
 				NextChainedReq->PlayerSessionServices = LoadRequest->PlayerSessionServices;
@@ -3853,7 +3921,7 @@ void FManifestDASHInternal::FRepresentation::SegmentIndexDownloadComplete(TShare
 					{
 						// Add this to the entity cache in case it needs to be retrieved again.
 						IPlayerEntityCache::FCacheItem ci;
-						ci.URL = LoadRequest->URL;
+						ci.URL = LoadRequest->Url.URL;
 						ci.Range = LoadRequest->Range;
 						ci.ParsedMatroskaData = MKVParser;
 						LoadRequest->PlayerSessionServices->GetEntityCache()->CacheEntity(ci);
diff --git a/Engine/Plugins/Media/ElectraPlayer/Source/ElectraPlayerRuntime/Private/Runtime/Player/DASH/ManifestDASH.h b/Engine/Plugins/Media/ElectraPlayer/Source/ElectraPlayerRuntime/Private/Runtime/Player/DASH/ManifestDASH.h
--- a/Engine/Plugins/Media/ElectraPlayer/Source/ElectraPlayerRuntime/Private/Runtime/Player/DASH/ManifestDASH.h
+++ b/Engine/Plugins/Media/ElectraPlayer/Source/ElectraPlayerRuntime/Private/Runtime/Player/DASH/ManifestDASH.h
@@ -40,6 +40,7 @@ public:
 	void UpdateDynamicRefetchCounter() override;
 	void TriggerClockSync(IManifest::EClockSyncType InClockSyncType) override;
 	void TriggerPlaylistRefresh() override;
+	void ReachedStableBuffer() override;
 	IStreamReader* CreateStreamReaderHandler() override;
 
 	FResult FindPlayPeriod(TSharedPtrTS<IPlayPeriod>& OutPlayPeriod, const FPlayStartPosition& StartPosition, ESearchType SearchType) override;
diff --git a/Engine/Plugins/Media/ElectraPlayer/Source/ElectraPlayerRuntime/Private/Runtime/Player/DASH/ManifestParserDASH.cpp b/Engine/Plugins/Media/ElectraPlayer/Source/ElectraPlayerRuntime/Private/Runtime/Player/DASH/ManifestParserDASH.cpp
--- a/Engine/Plugins/Media/ElectraPlayer/Source/ElectraPlayerRuntime/Private/Runtime/Player/DASH/ManifestParserDASH.cpp
+++ b/Engine/Plugins/Media/ElectraPlayer/Source/ElectraPlayerRuntime/Private/Runtime/Player/DASH/ManifestParserDASH.cpp
@@ -103,6 +103,7 @@ DEF_ELEMENT(OperatingQuality);
 DEF_ELEMENT(OperatingBandwidth);
 DEF_ELEMENT(ServiceDescription);
 DEF_ELEMENT(Resync);
+DEF_ELEMENT(ContentSteering);
 
 const TCHAR* const Element_up_UrlQueryInfo = TEXT("up:UrlQueryInfo");
 const TCHAR* const Element_up_ExtUrlQueryInfo = TEXT("up:ExtUrlQueryInfo");
@@ -232,6 +233,15 @@ DEF_ATTR(dImin);
 DEF_ATTR(dImax);
 DEF_ATTR(marker);
 DEF_ATTR(rangeAccess);
+DEF_ATTR(defaultServiceLocation);
+DEF_ATTR(queryBeforeStart);
+DEF_ATTR(clientRequirement);
+DEF_ATTR(proxyServerURL);
+
+// DVB DASH attributes have a colon in them.
+const TCHAR* const Attr_dvb_priority = TEXT("dvb:priority");
+const TCHAR* const Attr_dvb_weight = TEXT("dvb:weight");
+
 
 
 // xlink attributes. Due to the xlink namespace they have a colon in them
@@ -715,19 +725,19 @@ FString IDashMPDElement::GetDocumentURL() const
 }
 
 const TArray<TSharedPtrTS<FDashMPD_BaseURLType>>& IDashMPDElement::GetBaseURLs() const
-{ 
+{
 	static TArray<TSharedPtrTS<FDashMPD_BaseURLType>> NoBaseURLs;
-	return NoBaseURLs; 
+	return NoBaseURLs;
 }
 const TArray<TSharedPtrTS<FDashMPD_DescriptorType>>& IDashMPDElement::GetEssentialProperties() const
-{ 
+{
 	static TArray<TSharedPtrTS<FDashMPD_DescriptorType>> NoEssentialProperties;
-	return NoEssentialProperties; 
+	return NoEssentialProperties;
 }
 const TArray<TSharedPtrTS<FDashMPD_DescriptorType>>& IDashMPDElement::GetSupplementalProperties() const
-{ 
+{
 	static TArray<TSharedPtrTS<FDashMPD_DescriptorType>> NoSupplementalProperties;
-	return NoSupplementalProperties; 
+	return NoSupplementalProperties;
 }
 IDashMPDElement::FXLink& IDashMPDElement::GetXLink()
 {
@@ -795,7 +805,7 @@ bool IDashMPDElement::ProcessAttribute(FManifestParserDASH* Builder, const TCHAR
 
 	// Since descriptors are largely application specific we cannot know if their elements are meaningful or not, so we don't show a warning.
 	// Similary for anything we don't really understand we also should not emit warnings.
-	if (GetElementType() == EType::Descriptor || GetElementType() == EType::Other)
+	if (GetElementType() == EType::Descriptor || GetElementType() == EType::ContentSteering || GetElementType() == EType::Other)
 	{
 		bEmitWarning = false;
 	}
@@ -907,6 +917,11 @@ bool FDashMPD_MPDType::ProcessElement(FManifestParserDASH* Builder, const TCHAR*
 		LeapSecondInformation = MakeSharedTS<FDashMPD_LeapSecondInformationType>(ElementName, ElementData);
 		return Builder->PushNewElement(LeapSecondInformation, XmlFileLineNumber);
 	}
+	else if (StringHelpers::StringEquals(ElementName, ELEMENT(ContentSteering)))
+	{
+		ContentSteering = MakeSharedTS<FDashMPD_ContentSteeringType>(ElementName, ElementData);
+		return Builder->PushNewElement(ContentSteering, XmlFileLineNumber);
+	}
 	else
 	{
 		// Metrics or other elements.
@@ -1117,6 +1132,16 @@ bool FDashMPD_BaseURLType::ProcessAttribute(FManifestParserDASH* Builder, const
 	{
 		return ParseDuration(TimeShiftBufferDepth, Builder, AttributeName, AttributeValue, GetXMLLineNumber());
 	}
+	// DVB attributes?
+	else if (AttrEquals(dvb_priority))
+	{
+		return ParseIntAttribute(DVBPriority, Builder, AttributeName, AttributeValue, GetXMLLineNumber());
+	}
+	else if (AttrEquals(dvb_weight))
+	{
+		return ParseIntAttribute(DVBWeight, Builder, AttributeName, AttributeValue, GetXMLLineNumber());
+	}
+	// Any other attribute.
 	else
 	{
 		return Super::ProcessAttribute(Builder, AttributeName, AttributeValue);
@@ -2137,6 +2162,32 @@ bool FDashMPD_LeapSecondInformationType::ProcessAttribute(FManifestParserDASH* B
 
 /*********************************************************************************************************************/
 
+bool FDashMPD_ContentSteeringType::ProcessAttribute(FManifestParserDASH* Builder, const TCHAR* AttributeName, const TCHAR* AttributeValue)
+{
+	if (AttrEquals(defaultServiceLocation))
+	{
+		return ParseStringAttribute(DefaultServiceLocation, Builder, AttributeName, AttributeValue, GetXMLLineNumber());
+	}
+	else if (AttrEquals(proxyServerURL))
+	{
+		return ParseStringAttribute(ProxyServerURL, Builder, AttributeName, AttributeValue, GetXMLLineNumber());
+	}
+	else if (AttrEquals(queryBeforeStart))
+	{
+		return ParseBooleanAttribute(bQueryBeforeStart, Builder, AttributeName, AttributeValue, GetXMLLineNumber());
+	}
+	else if (AttrEquals(clientRequirement))
+	{
+		return ParseBooleanAttribute(bClientRequirement, Builder, AttributeName, AttributeValue, GetXMLLineNumber());
+	}
+	else
+	{
+		return Super::ProcessAttribute(Builder, AttributeName, AttributeValue);
+	}
+}
+
+/*********************************************************************************************************************/
+
 bool FDashMPD_ContentPopularityRateType::FPRElement::ProcessAttribute(FManifestParserDASH* Builder, const TCHAR* AttributeName, const TCHAR* AttributeValue)
 {
 	if (AttrEquals(popularityRate))
