Electra: Added multiple CDN support and ContentSteering for DASH and HLS

diff --git a/Engine/Plugins/Media/ElectraPlayer/Source/ElectraPlayerRuntime/Private/Runtime/Player/DASH/PlaylistReaderDASH.cpp b/Engine/Plugins/Media/ElectraPlayer/Source/ElectraPlayerRuntime/Private/Runtime/Player/DASH/PlaylistReaderDASH.cpp
--- a/Engine/Plugins/Media/ElectraPlayer/Source/ElectraPlayerRuntime/Private/Runtime/Player/DASH/PlaylistReaderDASH.cpp
+++ b/Engine/Plugins/Media/ElectraPlayer/Source/ElectraPlayerRuntime/Private/Runtime/Player/DASH/PlaylistReaderDASH.cpp
@@ -20,6 +20,7 @@
 #include "Player/DASH/OptionKeynamesDASH.h"
 #include "Player/DASH/PlayerEventDASH.h"
 #include "Player/DASH/PlayerEventDASH_Internal.h"
+#include "Player/ContentSteeringHandler.h"
 #include "ElectraHTTPStream.h"
 #include "Misc/DateTime.h"
 
@@ -28,13 +29,12 @@
 #define ERRCODE_DASH_MPD_PARSING_FAILED						3
 #define ERRCODE_DASH_MPD_UNSUPPORTED_DOCUMENT_ENCODING		4
 #define ERRCODE_DASH_MPD_REMOTE_ENTITY_FAILED				5
+#define ERRCODE_DASH_MPD_BAD_CONTENT_STEERING				6
 
 
 DECLARE_CYCLE_STAT(TEXT("FPlaylistReaderDASH_WorkerThread"), STAT_ElectraPlayer_DASH_PlaylistWorker, STATGROUP_ElectraPlayer);
 
 
-#define DASH_PLAYLISTREADER_USE_DEDICATED_WORKER_THREAD		0
-
 //#define DO_NOT_PERFORM_CONDITIONAL_GET
 
 namespace Electra
@@ -44,9 +44,6 @@ namespace Electra
  * This class is responsible for downloading a DASH MPD and parsing it.
  */
 class FPlaylistReaderDASH : public TSharedFromThis<FPlaylistReaderDASH, ESPMode::ThreadSafe>, public IPlaylistReaderDASH, public IAdaptiveStreamingPlayerAEMSReceiver
-#if DASH_PLAYLISTREADER_USE_DEDICATED_WORKER_THREAD
-						  , public FMediaThread
-#endif
 {
 public:
 	FPlaylistReaderDASH();
@@ -83,13 +80,8 @@ private:
 	FErrorDetail GetXMLResponseString(FString& OutXMLString, FResourceLoadRequestPtr FromRequest);
 	FErrorDetail GetResponseString(FString& OutString, FResourceLoadRequestPtr FromRequest);
 
-	void StartWorkerThread();
-	void StopWorkerThread();
-	void WorkerThread();
-
 	void InternalSetup();
 	void InternalCleanup();
-	void InternalHandleOnce();
 
 	void ExecutePendingRequests(const FTimeValue& TimeNow);
 	void HandleCompletedRequests(const FTimeValue& TimeNow);
@@ -115,6 +107,10 @@ private:
 	void PrepareNextTimeSync();
 	void ManifestUpdateDownloadCompleted(FResourceLoadRequestPtr Request, bool bSuccess);
 	void InitialMPDXLinkElementDownloadCompleted(FResourceLoadRequestPtr Request, bool bSuccess);
+	void ContentSteeringDownloadCompleted(FResourceLoadRequestPtr Request, bool bSuccess);
+	void ContentSteeringUpdateDownloadCompleted(FResourceLoadRequestPtr Request, bool bSuccess);
+	void TriggerContentSteeringUpdate(bool bInitial);
+	void ContentSteeringUpdate(bool bInitial, FResourceLoadRequestPtr Request, bool bSuccess);
 
 	void Timesync_httphead_Completed(FResourceLoadRequestPtr Request, bool bSuccess);
 	void Timesync_httpxsdate_Completed(FResourceLoadRequestPtr Request, bool bSuccess);
@@ -124,9 +120,7 @@ private:
 	IPlayerSessionServices*									PlayerSessionServices = nullptr;
 	FString													MPDURL;
 	FString													Fragment;
-	FMediaSemaphore											WorkerThreadSignal;
-	bool													bIsWorkerThreadStarted = false;
-	volatile bool											bTerminateWorkerThread = false;
+	bool													bIsInitialized = false;
 
 	FCriticalSection										RequestsLock;
 	TArray<FResourceLoadRequestPtr>							PendingRequests;
@@ -208,9 +202,6 @@ TSharedPtrTS<IPlaylistReader> IPlaylistReaderDASH::Create(IPlayerSessionServices
 
 
 FPlaylistReaderDASH::FPlaylistReaderDASH()
-#if DASH_PLAYLISTREADER_USE_DEDICATED_WORKER_THREAD
-	: FMediaThread("ElectraPlayer::DASH MPD")
-#endif
 {
 }
 
@@ -236,19 +227,12 @@ void FPlaylistReaderDASH::Initialize(IPlayerSessionServices* InPlayerSessionServ
 
 void FPlaylistReaderDASH::Close()
 {
-	StopWorkerThread();
+	InternalCleanup();
 }
 
 void FPlaylistReaderDASH::HandleOnce()
 {
-#if !DASH_PLAYLISTREADER_USE_DEDICATED_WORKER_THREAD
-	InternalHandleOnce();
-#endif
-}
-
-void FPlaylistReaderDASH::InternalHandleOnce()
-{
-	if (bIsWorkerThreadStarted)
+	if (bIsInitialized)
 	{
 		SCOPE_CYCLE_COUNTER(STAT_ElectraPlayer_DASH_PlaylistWorker);
 		CSV_SCOPED_TIMING_STAT(ElectraPlayer, DASH_PlaylistWorker);
@@ -265,36 +249,14 @@ void FPlaylistReaderDASH::InternalHandleOnce()
 		ExecutePendingRequests(Now);
 		// Check if the MPD must be updated.
 		CheckForMPDUpdate();
+		// Check if a new steering manifest is needed
+		if (PlayerSessionServices->GetContentSteeringHandler()->NeedToObtainNewSteeringManifestNow())
+		{
+			TriggerContentSteeringUpdate(false);
+		}
 	}
 }
 
-void FPlaylistReaderDASH::StartWorkerThread()
-{
-	check(!bIsWorkerThreadStarted);
-#if DASH_PLAYLISTREADER_USE_DEDICATED_WORKER_THREAD
-	bTerminateWorkerThread = false;
-	ThreadStart(Electra::MakeDelegate(this, &FPlaylistReaderDASH::WorkerThread));
-#else
-	InternalSetup();
-#endif
-	bIsWorkerThreadStarted = true;
-}
-
-void FPlaylistReaderDASH::StopWorkerThread()
-{
-	if (bIsWorkerThreadStarted)
-	{
-#if DASH_PLAYLISTREADER_USE_DEDICATED_WORKER_THREAD
-		bTerminateWorkerThread = true;
-		WorkerThreadSignal.Release();
-		ThreadWaitDone();
-		ThreadReset();
-#else
-		InternalCleanup();
-#endif
-		bIsWorkerThreadStarted = false;
-	}
-}
 
 void FPlaylistReaderDASH::PostError(const FErrorDetail& Error)
 {
@@ -342,7 +304,7 @@ void FPlaylistReaderDASH::LogMessage(IInfoLog::ELevel Level, const FString& Mess
 void FPlaylistReaderDASH::LoadAndParse(const FString& URL)
 {
 	MPDURL = URL;
-	StartWorkerThread();
+	InternalSetup();
 }
 
 
@@ -497,56 +459,113 @@ void FPlaylistReaderDASH::CheckForMPDUpdate()
 		TSharedPtrTS<const FDashMPD_MPDType> MPDRoot = Manifest.IsValid() ? Manifest->GetMPDRoot() : nullptr;
 		if (MPDRoot.IsValid())
 		{
-			FString NewMPDLocation;
+			// The URL query might need to be changed. Look for the UrlQuery properties.
+			TArray<TSharedPtrTS<FDashMPD_UrlQueryInfoType>> UrlQueries;
+			DASHUrlHelpers::GetAllHierarchyUrlQueries(UrlQueries, MPDRoot, DASHUrlHelpers::EUrlQueryRequestType::Mpd, false);
+
+			DASHUrlHelpers::FDASHMediaURL RootDocument;
+			RootDocument.URL = MPDRoot->GetDocumentURL();
+			FURL_RFC3986 RootDocURL;
+			RootDocURL.Parse(RootDocument.URL);
+			// The MPD has no @serviceLocation or the DVB properties. According to DVB-DASH the @serviceLocation is the URL itself and the attributes default to 1.
+			RootDocument.CDN = MPDRoot->IsDVBDash() ? RootDocument.URL : RootDocURL.GetHost();
+			RootDocument.DVBpriority = 1;
+			RootDocument.DVBweight = 1;
+			// Build a list we can pass to the steering handler.
+			TArray<FContentSteeringHandler::FCandidateURL> SteeringCandidates;
+
 			// See if there are any <Location> elements on the current MPD that tell us from where to get the update.
 			// As per A.11 no <BaseURL> elements affect the construction of the new document URL. At most, if it is
 			// a relative URL it will be resolved against the current document URL.
 			const TArray<TSharedPtrTS<FDashMPD_OtherType>>& Locations = MPDRoot->GetLocations();
-			if (Locations.Num())
+			// There can be several new locations, but the elements have no DASH defined attributes to differentiate them
+			// as of the ISO/IEC 23009-1:2022 standard. For content steering ETSI TS 103 998 V1.1.1 (2024-01)
+			// proposes the addition of the @serviceLocation attribute used with the <BaseURL> element.
+			for(int32 i=0; i<Locations.Num(); ++i)
 			{
-				// There can be several new locations, but the elements have no DASH defined attributes to differentiate them.
-				// They are merely elements containing a URL. In absence of any preference we just use the first one.
-				NewMPDLocation = Locations[0]->GetData();
-			}
+				FContentSteeringHandler::FCandidateURL& sc = SteeringCandidates.Emplace_GetRef();
+				sc.MediaURL.URL = Locations[i]->GetData();
+				const TArray<IDashMPDElement::FXmlAttribute>& OtherAttributes = Locations[i]->GetOtherAttributes();
+				for(auto &Attribute : OtherAttributes)
+				{
+					if (Attribute.GetName().Equals(TEXT("serviceLocation")))
+					{
+						sc.MediaURL.CDN = Attribute.GetValue();
+						break;
+					}
+				}
+				// The URL given in the location must be an absolute URL.
+				if (!DASHUrlHelpers::IsAbsoluteURL(sc.MediaURL.URL))
+				{
+					LogMessage(IInfoLog::ELevel::Warning, FString::Printf(TEXT("<MPD.Location> URI is not an absolute URL, ignoring!")));
+					SteeringCandidates.Pop();
+					continue;
+				}
 
-			TArray<TSharedPtrTS<const FDashMPD_BaseURLType>> OutBaseURLs;
-			FString URL, RequestHeader;
-			FTimeValue UrlATO;
-			TMediaOptionalValue<bool> bATOComplete;
-			DASHUrlHelpers::BuildAbsoluteElementURL(URL, UrlATO, bATOComplete, MPDRoot->GetDocumentURL(), OutBaseURLs, NewMPDLocation);
+				// Apply any of the Annex I UrlQuery elements
+				FString AnnexIRequestHeader;
+				DASHUrlHelpers::ApplyUrlQueries(PlayerSessionServices, RootDocument.URL, sc.MediaURL.URL, AnnexIRequestHeader, UrlQueries);
+				sc.AdditionalParams.Set(DASHUrlHelpers::SteerOption_AnnexIRequestHeader, FVariantValue(AnnexIRequestHeader));
 
-			// The URL query might need to be changed. Look for the UrlQuery properties.
-			TArray<TSharedPtrTS<FDashMPD_UrlQueryInfoType>> UrlQueries;
-			DASHUrlHelpers::GetAllHierarchyUrlQueries(UrlQueries, MPDRoot, DASHUrlHelpers::EUrlQueryRequestType::Mpd, false);
-			FErrorDetail Error = DASHUrlHelpers::ApplyUrlQueries(PlayerSessionServices, MPDRoot->GetDocumentURL(), URL, RequestHeader, UrlQueries);
-			if (Error.IsOK())
-			{
-				FString ETag = MPDRoot->GetETag();
-
-				TSharedPtrTS<FMPDLoadRequestDASH> PlaylistLoadRequest = MakeSharedTS<FMPDLoadRequestDASH>();
-				PlaylistLoadRequest->LoadType = FMPDLoadRequestDASH::ELoadType::MPDUpdate;
-				PlaylistLoadRequest->URL = URL;
-				if (RequestHeader.Len())
+				// If there is no service location then as with DVB-DASH the URL is to be used for that.
+				if (sc.MediaURL.CDN.IsEmpty())
 				{
-					PlaylistLoadRequest->Headers.Emplace(HTTP::FHTTPHeader(DASH::HTTPHeaderOptionName, RequestHeader));
+					if (!MPDRoot->IsDVBDash())
+					{
+						FURL_RFC3986 tmp;tmp.Parse(sc.MediaURL.URL);
+						sc.MediaURL.CDN = tmp.GetHost();
 				}
-#ifndef DO_NOT_PERFORM_CONDITIONAL_GET
-				// As per the standard we should perform conditional GET requests.
-				if (ETag.Len())
+					else
+					{
+						sc.MediaURL.CDN = sc.MediaURL.URL;
+					}
+				}
+			}
+
+			FContentSteeringHandler::FSelectedCandidateURL ChosenLocation;
+			// Perform steering evaluation only when there are candidates, otherwise
+			// fall back to the original MPD URL.
+			if (SteeringCandidates.Num())
+			{
+				// Have the steering handler select the desired CDN, potentially cloning an existing one on demand.
+				FString SteeringMessage;
+				ChosenLocation = PlayerSessionServices->GetContentSteeringHandler()->SelectBestCandidateFrom(SteeringMessage, FContentSteeringHandler::ESelectFor::Playlist, SteeringCandidates);
+				if (ChosenLocation.MediaURL.URL.IsEmpty())
 				{
-					PlaylistLoadRequest->Headers.Emplace(HTTP::FHTTPHeader(TEXT("If-None-Match"), ETag));
+					LogMessage(IInfoLog::ELevel::Warning, FString::Printf(TEXT("No MPD location for reloading the MPD selected by content steering, falling back to original URL.")));
 				}
-#endif
-				PlaylistLoadRequest->CompleteCallback.BindThreadSafeSP(AsShared(), &FPlaylistReaderDASH::ManifestUpdateDownloadCompleted);
+			}
+			if (ChosenLocation.MediaURL.URL.IsEmpty())
+			{
+				ChosenLocation.MediaURL = RootDocument;
+				FString AnnexIRequestHeader;
+				DASHUrlHelpers::ApplyUrlQueries(PlayerSessionServices, RootDocument.URL, ChosenLocation.MediaURL.URL, AnnexIRequestHeader, UrlQueries);
+				ChosenLocation.AdditionalParams.Set(DASHUrlHelpers::SteerOption_AnnexIRequestHeader, FVariantValue(AnnexIRequestHeader));
+				ChosenLocation.SteeringID = -1;
+			}
+
+			FString RequestHeader = ChosenLocation.AdditionalParams.GetValue(DASHUrlHelpers::SteerOption_AnnexIRequestHeader).SafeGetFString();
+			FString ETag = MPDRoot->GetETag();
 
-				bIsMPDUpdateInProgress = true;
-				EnqueueResourceRequest(MoveTemp(PlaylistLoadRequest));
+			TSharedPtrTS<FMPDLoadRequestDASH> PlaylistLoadRequest = MakeSharedTS<FMPDLoadRequestDASH>();
+			PlaylistLoadRequest->LoadType = FMPDLoadRequestDASH::ELoadType::MPDUpdate;
+			PlaylistLoadRequest->Url = ChosenLocation.MediaURL;
+			PlaylistLoadRequest->SteeringID = ChosenLocation.SteeringID;
+			if (RequestHeader.Len())
+			{
+				PlaylistLoadRequest->Headers.Emplace(HTTP::FHTTPHeader(DASH::HTTPHeaderOptionName, RequestHeader));
 			}
-			else
+#ifndef DO_NOT_PERFORM_CONDITIONAL_GET
+			// As per the standard we should perform conditional GET requests.
+			if (ETag.Len())
 			{
-				// Failed to build the URL.
-				PostError(Error);
+				PlaylistLoadRequest->Headers.Emplace(HTTP::FHTTPHeader(TEXT("If-None-Match"), ETag));
 			}
+#endif
+			PlaylistLoadRequest->CompleteCallback.BindThreadSafeSP(AsShared(), &FPlaylistReaderDASH::ManifestUpdateDownloadCompleted);
+
+			bIsMPDUpdateInProgress = true;
+			EnqueueResourceRequest(MoveTemp(PlaylistLoadRequest));
 		}
 	}
 }
@@ -561,7 +580,6 @@ void FPlaylistReaderDASH::OnHTTPResourceRequestComplete(TSharedPtrTS<FHTTPResour
 		CompletedRequests.Emplace(lr);
 		ActiveRequests.Remove(lr);
 	}
-	WorkerThreadSignal.Release();
 }
 
 
@@ -569,7 +587,6 @@ void FPlaylistReaderDASH::EnqueueResourceRequest(FResourceLoadRequestPtr InReque
 {
 	FScopeLock lock(&RequestsLock);
 	PendingRequests.Emplace(MoveTemp(InRequest));
-	WorkerThreadSignal.Release();
 }
 
 void FPlaylistReaderDASH::EnqueueResourceRetryRequest(FResourceLoadRequestPtr InRequest, const FTimeValue& AtUTCTime)
@@ -578,7 +595,6 @@ void FPlaylistReaderDASH::EnqueueResourceRetryRequest(FResourceLoadRequestPtr In
 	InRequest->ExecuteAtUTC = AtUTCTime;
 	++InRequest->Attempt;
 	PendingRequests.Emplace(MoveTemp(InRequest));
-	WorkerThreadSignal.Release();
 }
 
 void FPlaylistReaderDASH::EnqueueInitialXLinkRequests()
@@ -597,7 +613,6 @@ void FPlaylistReaderDASH::EnqueueInitialXLinkRequests()
 			PendingRequests.Emplace(MoveTemp(LReq));
 		}
 	}
-	WorkerThreadSignal.Release();
 }
 
 void FPlaylistReaderDASH::AddElementLoadRequests(const TArray<TWeakPtrTS<FMPDLoadRequestDASH>>& RemoteElementLoadRequests)
@@ -614,7 +629,6 @@ void FPlaylistReaderDASH::AddElementLoadRequests(const TArray<TWeakPtrTS<FMPDLoa
 				PendingRequests.Emplace(MoveTemp(LReq));
 			}
 		}
-		WorkerThreadSignal.Release();
 	}
 }
 
@@ -641,13 +655,19 @@ void FPlaylistReaderDASH::SetStreamInbandEventUsage(EStreamType InStreamType, bo
 	switch(InStreamType)
 	{
 		case EStreamType::Video:
+		{
 			bInbandEventByStreamType[0] = bUsesInbandDASHEvents;
 			break;
+		}
 		case EStreamType::Audio:
+		{
 			bInbandEventByStreamType[1] = bUsesInbandDASHEvents;
 			break;
+		}
 		default:
+		{
 			break;
+		}
 	}
 }
 
@@ -730,21 +750,21 @@ void FPlaylistReaderDASH::TriggerTimeSynchronization()
 					{
 						TimeSyncRequest = MakeSharedTS<FMPDLoadRequestDASH>();
 						TimeSyncRequest->LoadType = FMPDLoadRequestDASH::ELoadType::TimeSync;
-						TimeSyncRequest->URL = GetRequestURL(Value, 0);
+						TimeSyncRequest->Url.URL = GetRequestURL(Value, 0);
 						TimeSyncRequest->CompleteCallback.BindThreadSafeSP(AsShared(), &FPlaylistReaderDASH::Timesync_httpxsdate_Completed);
 					}
 					else if (SchemeIdUri.Equals(DASH::Schemes::TimingSources::Scheme_urn_mpeg_dash_utc_httpiso2014))
 					{
 						TimeSyncRequest = MakeSharedTS<FMPDLoadRequestDASH>();
 						TimeSyncRequest->LoadType = FMPDLoadRequestDASH::ELoadType::TimeSync;
-						TimeSyncRequest->URL = GetRequestURL(Value, 0);
+						TimeSyncRequest->Url.URL = GetRequestURL(Value, 0);
 						TimeSyncRequest->CompleteCallback.BindThreadSafeSP(AsShared(), &FPlaylistReaderDASH::Timesync_httpiso_Completed);
 					}
 					else if (SchemeIdUri.Equals(DASH::Schemes::TimingSources::Scheme_urn_mpeg_dash_utc_httphead2014))
 					{
 						TimeSyncRequest = MakeSharedTS<FMPDLoadRequestDASH>();
 						TimeSyncRequest->LoadType = FMPDLoadRequestDASH::ELoadType::TimeSync;
-						TimeSyncRequest->URL = GetRequestURL(Value, 0);
+						TimeSyncRequest->Url.URL = GetRequestURL(Value, 0);
 						TimeSyncRequest->Verb = TEXT("HEAD");
 						TimeSyncRequest->CompleteCallback.BindThreadSafeSP(AsShared(), &FPlaylistReaderDASH::Timesync_httphead_Completed);
 					}
@@ -794,6 +814,7 @@ void FPlaylistReaderDASH::TriggerTimeSynchronization()
 
 void FPlaylistReaderDASH::InternalSetup()
 {
+	check(!bIsInitialized);
 	bInbandEventByStreamType[0] = false;
 	bInbandEventByStreamType[1] = false;
 	bInbandEventByStreamType[2] = false;
@@ -813,54 +834,38 @@ void FPlaylistReaderDASH::InternalSetup()
 	PlaylistLoadRequest->LoadType = FMPDLoadRequestDASH::ELoadType::MPD;
 	FURL_RFC3986 UrlParser;
 	UrlParser.Parse(MPDURL);
-	PlaylistLoadRequest->URL = UrlParser.Get(true, false);
+	PlaylistLoadRequest->Url.URL = UrlParser.Get(true, false);
 	Fragment = UrlParser.GetFragment();
 	PlaylistLoadRequest->CompleteCallback.BindThreadSafeSP(AsShared(), &FPlaylistReaderDASH::ManifestDownloadCompleted);
 	EnqueueResourceRequest(MoveTemp(PlaylistLoadRequest));
+	bIsInitialized = true;
 }
 
 void FPlaylistReaderDASH::InternalCleanup()
 {
-	// Unregister event callbacks
-	PlayerSessionServices->GetAEMSEventHandler()->RemoveAEMSReceiver(AsShared(), DASH::Schemes::ManifestEvents::Scheme_urn_mpeg_dash_event_ttfn_2016, TEXT(""), IAdaptiveStreamingPlayerAEMSReceiver::EDispatchMode::OnStart);
-	PlayerSessionServices->GetAEMSEventHandler()->RemoveAEMSReceiver(AsShared(), DASH::Schemes::ManifestEvents::Scheme_urn_mpeg_dash_event_callback_2015, TEXT("1"), IAdaptiveStreamingPlayerAEMSReceiver::EDispatchMode::OnStart);
-	PlayerSessionServices->GetAEMSEventHandler()->RemoveAEMSReceiver(AsShared(), DASH::Schemes::ManifestEvents::Scheme_urn_mpeg_dash_event_2012, TEXT(""), IAdaptiveStreamingPlayerAEMSReceiver::EDispatchMode::OnStart);
-
-	// Cleanup!
-	RequestsLock.Lock();
-	PendingRequests.Empty();
-	for(int32 i=0; i<ActiveRequests.Num(); ++i)
+	if (bIsInitialized)
 	{
-		if (ActiveRequests[i]->Request.IsValid())
-		{
-			ActiveRequests[i]->Request->Cancel();
-		}
-	}
-	ActiveRequests.Empty();
-	CompletedRequests.Empty();
-	RequestsLock.Unlock();
-	Builder.Reset();
-}
-
-void FPlaylistReaderDASH::WorkerThread()
-{
-#if DASH_PLAYLISTREADER_USE_DEDICATED_WORKER_THREAD
-	LLM_SCOPE(ELLMTag::ElectraPlayer);
-
-	InternalSetup();
+		// Unregister event callbacks
+		PlayerSessionServices->GetAEMSEventHandler()->RemoveAEMSReceiver(AsShared(), DASH::Schemes::ManifestEvents::Scheme_urn_mpeg_dash_event_ttfn_2016, TEXT(""), IAdaptiveStreamingPlayerAEMSReceiver::EDispatchMode::OnStart);
+		PlayerSessionServices->GetAEMSEventHandler()->RemoveAEMSReceiver(AsShared(), DASH::Schemes::ManifestEvents::Scheme_urn_mpeg_dash_event_callback_2015, TEXT("1"), IAdaptiveStreamingPlayerAEMSReceiver::EDispatchMode::OnStart);
+		PlayerSessionServices->GetAEMSEventHandler()->RemoveAEMSReceiver(AsShared(), DASH::Schemes::ManifestEvents::Scheme_urn_mpeg_dash_event_2012, TEXT(""), IAdaptiveStreamingPlayerAEMSReceiver::EDispatchMode::OnStart);
 
-	while(!bTerminateWorkerThread)
-	{
-		WorkerThreadSignal.Obtain(1000 * 100);
-		if (bTerminateWorkerThread)
+		// Cleanup!
+		RequestsLock.Lock();
+		PendingRequests.Empty();
+		for(int32 i=0; i<ActiveRequests.Num(); ++i)
 		{
-			break;
+			if (ActiveRequests[i]->Request.IsValid())
+			{
+				ActiveRequests[i]->Request->Cancel();
+			}
 		}
-		InternalHandleOnce();
+		ActiveRequests.Empty();
+		CompletedRequests.Empty();
+		RequestsLock.Unlock();
+		Builder.Reset();
+		bIsInitialized = false;
 	}
-
-	InternalCleanup();
-#endif
 }
 
 
@@ -877,10 +882,9 @@ void FPlaylistReaderDASH::ExecutePendingRequests(const FTimeValue& TimeNow)
 			--i;
 
 			// Check for reserved URIs
-			if (Request->URL.Equals(TEXT("urn:mpeg:dash:resolve-to-zero:2013")))
+			if (Request->Url.URL.Equals(TEXT("urn:mpeg:dash:resolve-to-zero:2013")))
 			{
 				CompletedRequests.Emplace(Request);
-				WorkerThreadSignal.Release();
 			}
 			else
 			{
@@ -888,7 +892,7 @@ void FPlaylistReaderDASH::ExecutePendingRequests(const FTimeValue& TimeNow)
 				ActiveRequests.Push(Request);
 
 				Request->Request = MakeSharedTS<FHTTPResourceRequest>();
-				Request->Request->URL(Request->URL).Verb(Request->Verb).Range(Request->Range).Headers(Request->Headers).Object(Request);
+				Request->Request->URL(Request->Url.URL).Verb(Request->Verb).Range(Request->Range).Headers(Request->Headers).Object(Request);
 				// Set up timeouts and also accepted encodings depending on the type of request.
 				SetupRequestTimeouts(Request);
 				Request->Request->AllowStaticQuery(IAdaptiveStreamingPlayerResourceRequest::EPlaybackResourceType::Playlist);
@@ -981,6 +985,15 @@ int32 FPlaylistReaderDASH::CheckForRetry(FResourceLoadRequestPtr InRequest, int3
 			}
 			return -1;
 		}
+		case FMPDLoadRequestDASH::ELoadType::Steering:
+		{
+			// Only retry connection issues, not HTTP failures.
+			if (Error < 100 && InRequest->Attempt < 2)
+			{
+				return 250 * (1 << InRequest->Attempt);
+			}
+			return -1;
+		}
 
 		case FMPDLoadRequestDASH::ELoadType::XLink_Period:
 		case FMPDLoadRequestDASH::ELoadType::XLink_AdaptationSet:
@@ -989,6 +1002,8 @@ int32 FPlaylistReaderDASH::CheckForRetry(FResourceLoadRequestPtr InRequest, int3
 		case FMPDLoadRequestDASH::ELoadType::XLink_URLQuery:
 		case FMPDLoadRequestDASH::ELoadType::XLink_InitializationSet:
 		case FMPDLoadRequestDASH::ELoadType::Callback:
+		case FMPDLoadRequestDASH::ELoadType::TimeSync:
+		case FMPDLoadRequestDASH::ELoadType::Sideload:
 		default:
 		{
 			break;
@@ -1026,7 +1041,7 @@ void FPlaylistReaderDASH::HandleCompletedRequests(const FTimeValue& TimeNow)
 				case FMPDLoadRequestDASH::ELoadType::MPD:
 				case FMPDLoadRequestDASH::ELoadType::MPDUpdate:
 				{
-					PlayerSessionServices->GetEntityCache()->SetRecentResponseHeaders(IPlayerEntityCache::EEntityType::Document, cr[i]->URL, req->GetConnectionInfo()->ResponseHeaders);
+					PlayerSessionServices->GetEntityCache()->SetRecentResponseHeaders(IPlayerEntityCache::EEntityType::Document, cr[i]->Url.URL, req->GetConnectionInfo()->ResponseHeaders);
 					break;
 				}
 				case FMPDLoadRequestDASH::ELoadType::XLink_Period:
@@ -1035,17 +1050,17 @@ void FPlaylistReaderDASH::HandleCompletedRequests(const FTimeValue& TimeNow)
 				case FMPDLoadRequestDASH::ELoadType::XLink_SegmentList:
 				case FMPDLoadRequestDASH::ELoadType::XLink_URLQuery:
 				{
-					PlayerSessionServices->GetEntityCache()->SetRecentResponseHeaders(IPlayerEntityCache::EEntityType::XLink, cr[i]->URL, req->GetConnectionInfo()->ResponseHeaders);
+					PlayerSessionServices->GetEntityCache()->SetRecentResponseHeaders(IPlayerEntityCache::EEntityType::XLink, cr[i]->Url.URL, req->GetConnectionInfo()->ResponseHeaders);
 					break;
 				}
 				case FMPDLoadRequestDASH::ELoadType::Callback:
 				{
-					PlayerSessionServices->GetEntityCache()->SetRecentResponseHeaders(IPlayerEntityCache::EEntityType::Callback, cr[i]->URL, req->GetConnectionInfo()->ResponseHeaders);
+					PlayerSessionServices->GetEntityCache()->SetRecentResponseHeaders(IPlayerEntityCache::EEntityType::Callback, cr[i]->Url.URL, req->GetConnectionInfo()->ResponseHeaders);
 					break;
 				}
 				case FMPDLoadRequestDASH::ELoadType::Segment:
 				{
-					PlayerSessionServices->GetEntityCache()->SetRecentResponseHeaders(IPlayerEntityCache::EEntityType::Segment, cr[i]->URL, req->GetConnectionInfo()->ResponseHeaders);
+					PlayerSessionServices->GetEntityCache()->SetRecentResponseHeaders(IPlayerEntityCache::EEntityType::Segment, cr[i]->Url.URL, req->GetConnectionInfo()->ResponseHeaders);
 					break;
 				}
 				default:
@@ -1063,13 +1078,13 @@ void FPlaylistReaderDASH::HandleCompletedRequests(const FTimeValue& TimeNow)
 			if (retryInMillisec >= 0)
 			{
 				// Try again
-				LogMessage(IInfoLog::ELevel::Warning, FString::Printf(TEXT("Failed to download %s (%s) from %s, retrying..."), cr[i]->GetRequestTypeName(), *cr[i]->URL, *req->GetErrorString()));
+				LogMessage(IInfoLog::ELevel::Warning, FString::Printf(TEXT("Failed to download %s (%s) from %s, retrying..."), cr[i]->GetRequestTypeName(), *cr[i]->Url.URL, *req->GetErrorString()));
 				EnqueueResourceRetryRequest(cr[i], TimeNow + FTimeValue().SetFromMilliseconds(retryInMillisec));
 			}
 			else
 			{
 				// Give up. Notify caller of failure.
-				LogMessage(IInfoLog::ELevel::Error, FString::Printf(TEXT("Failed to download %s (%s) from %s, giving up."), cr[i]->GetRequestTypeName(), *cr[i]->URL, *req->GetErrorString()));
+				LogMessage(IInfoLog::ELevel::Error, FString::Printf(TEXT("Failed to download %s (%s) from %s, giving up."), cr[i]->GetRequestTypeName(), *cr[i]->Url.URL, *req->GetErrorString()));
 				cr[i]->CompleteCallback.ExecuteIfBound(cr[i], false);
 			}
 		}
@@ -1204,7 +1219,7 @@ void FPlaylistReaderDASH::ManifestDownloadCompleted(FResourceLoadRequestPtr Requ
 	}
 	else
 	{
-		PostError(FString::Printf(TEXT("Failed to download MPD \"%s\" (%s)"), *Request->URL, *Request->GetErrorDetail()), ERRCODE_DASH_MPD_DOWNLOAD_FAILED, UEMEDIA_ERROR_READ_ERROR);
+		PostError(FString::Printf(TEXT("Failed to download MPD \"%s\" (%s)"), *Request->Url.URL, *Request->GetErrorDetail()), ERRCODE_DASH_MPD_DOWNLOAD_FAILED, UEMEDIA_ERROR_READ_ERROR);
 	}
 }
 
@@ -1323,9 +1338,49 @@ void FPlaylistReaderDASH::FinishManifestBuildingAfterTimesync(bool bGotTheTime)
 	PlayerSessionServices->SendMessageToPlayer(IPlaylistReader::PlaylistLoadedMessage::Create(!LastErrorDetail.IsTryAgain() ? LastErrorDetail : FErrorDetail(), InitialTimesync.ConnInfo.Get(), Playlist::EListType::Main, Playlist::ELoadType::Initial, Attempts));
 	if (LastErrorDetail.IsOK())
 	{
-		// Notify that the "variant playlists" are ready. There are no variants in DASH, but this is the trigger that the playlists are all set up and are good to go now.
-		PlayerSessionServices->SendMessageToPlayer(IPlaylistReader::PlaylistLoadedMessage::Create(LastErrorDetail, InitialTimesync.ConnInfo.Get(), Playlist::EListType::Variant, Playlist::ELoadType::Initial, Attempts));
-		PrepareNextTimeSync();
+		// Prepare content steering
+		FContentSteeringHandler::FInitialParams csp;
+		auto MPDRoot = Manifest->GetMPDRoot();
+		auto cs = MPDRoot->GetContentSteering();
+		csp.RootDocumentURL = MPDRoot->GetDocumentURL();
+		csp.bHasContentSteering = cs.IsValid();
+		csp.bUseDVBPriorities = !cs.IsValid() && MPDRoot->IsDVBDash();
+		if (cs.IsValid())
+		{
+			csp.FirstSteeringURL = cs->GetContentSteeringServer();
+			csp.ProxyURL = cs->GetProxyServerURL();
+			csp.InitialDefaultCDN = cs->GetDefaultServiceLocation();
+			csp.bQueryBeforeStart = cs->QueryBeforeStart();
+			const TArray<IDashMPDElement::FXmlAttribute>& OtherAttributes = cs->GetOtherAttributes();
+			for(auto &Attribute : OtherAttributes)
+			{
+				if (Attribute.GetName().Equals(TEXT("EpicInitialSelectionPriority")))
+				{
+					LexFromString(csp.CustomFirstCDNPrioritization, *Attribute.GetValue());
+					break;
+				}
+			}
+			// Without a URL we can use this only to select the CDN to start with.
+			csp.bHasContentSteering = !csp.FirstSteeringURL.IsEmpty();
+		}
+		if (PlayerSessionServices->GetContentSteeringHandler()->InitialSetup(FContentSteeringHandler::EStreamingProtocol::DASH, csp))
+		{
+			// If we do not have to query the content steering server right away we are now good to go.
+			if (!PlayerSessionServices->GetContentSteeringHandler()->NeedToObtainNewSteeringManifestNow())
+			{
+				// Notify that the "variant playlists" are ready. There are no variants in DASH, but this is the trigger that the playlists are all set up and are good to go now.
+				PlayerSessionServices->SendMessageToPlayer(IPlaylistReader::PlaylistLoadedMessage::Create(LastErrorDetail, InitialTimesync.ConnInfo.Get(), Playlist::EListType::Variant, Playlist::ELoadType::Initial, Attempts));
+				PrepareNextTimeSync();
+			}
+			else
+			{
+				TriggerContentSteeringUpdate(true);
+			}
+		}
+		else
+		{
+			PostError(FString::Printf(TEXT("Invalid content steering configuration")), ERRCODE_DASH_MPD_BAD_CONTENT_STEERING, UEMEDIA_ERROR_DETAIL);
+		}
 	}
 	// "try again" is returned when there are remote entities that need to be resolved first.
 	else if (LastErrorDetail.IsTryAgain())
@@ -1340,6 +1395,37 @@ void FPlaylistReaderDASH::FinishManifestBuildingAfterTimesync(bool bGotTheTime)
 	InitialTimesync.Reset();
 }
 
+void FPlaylistReaderDASH::TriggerContentSteeringUpdate(bool bInitial)
+{
+	auto MPDRoot = Manifest->GetMPDRoot();
+	// No <BaseURL> elements apply to content steering, but UrlQuery does.
+	TArray<TSharedPtrTS<const FDashMPD_BaseURLType>> OutBaseURLs;
+	FString URL, RequestHeader;
+	FTimeValue UrlATO;
+	TMediaOptionalValue<bool> bATOComplete;
+	DASHUrlHelpers::BuildAbsoluteElementURL(URL, UrlATO, bATOComplete, MPDRoot->GetDocumentURL(), OutBaseURLs, PlayerSessionServices->GetContentSteeringHandler()->GetBaseSteeringServerRequestURL());
+	TArray<TSharedPtrTS<FDashMPD_UrlQueryInfoType>> UrlQueries;
+	DASHUrlHelpers::GetAllHierarchyUrlQueries(UrlQueries, MPDRoot, DASHUrlHelpers::EUrlQueryRequestType::Steering, false);
+	LastErrorDetail = DASHUrlHelpers::ApplyUrlQueries(PlayerSessionServices, MPDRoot->GetDocumentURL(), URL, RequestHeader, UrlQueries);
+	URL = PlayerSessionServices->GetContentSteeringHandler()->GetFinalSteeringServerRequestURL(URL);
+	if (LastErrorDetail.IsOK())
+	{
+		TSharedPtrTS<FMPDLoadRequestDASH> SteeringLoadRequest = MakeSharedTS<FMPDLoadRequestDASH>();
+		SteeringLoadRequest->LoadType = FMPDLoadRequestDASH::ELoadType::Steering;
+		SteeringLoadRequest->Url.URL = URL;
+		if (RequestHeader.Len())
+		{
+			SteeringLoadRequest->Headers.Emplace(HTTP::FHTTPHeader(DASH::HTTPHeaderOptionName, RequestHeader));
+		}
+		SteeringLoadRequest->CompleteCallback.BindThreadSafeSP(AsShared(), bInitial ? &FPlaylistReaderDASH::ContentSteeringDownloadCompleted : &FPlaylistReaderDASH::ContentSteeringUpdateDownloadCompleted);
+		EnqueueResourceRequest(MoveTemp(SteeringLoadRequest));
+		PlayerSessionServices->GetContentSteeringHandler()->SetSteeringServerRequestIsPending();
+	}
+	else
+	{
+		PostError(LastErrorDetail);
+	}
+}
 
 
 void FPlaylistReaderDASH::ManifestUpdateDownloadCompleted(FResourceLoadRequestPtr Request, bool bSuccess)
@@ -1450,6 +1536,12 @@ void FPlaylistReaderDASH::ManifestUpdateDownloadCompleted(FResourceLoadRequestPt
 					PlayerSessionServices->SendMessageToPlayer(IPlaylistReader::PlaylistLoadedMessage::Create(!LastErrorDetail.IsTryAgain() ? LastErrorDetail : FErrorDetail(), ConnInfo, Playlist::EListType::Main, Playlist::ELoadType::Update, Attempt));
 					if (LastErrorDetail.IsOK())
 					{
+						// Let content steering know which CDN we used for downloading.
+						Metrics::FSegmentDownloadStats s;
+						s.SegmentType = Metrics::ESegmentType::PlaylistElement;
+						s.URL = Request->Url;
+						s.SteeringID = Request->SteeringID;
+						PlayerSessionServices->GetContentSteeringHandler()->FinishedDownloadRequestOn(s, FContentSteeringHandler::FStreamParams());
 					}
 					// "try again" is returned when there are remote entities that need to be resolved first.
 					else if (LastErrorDetail.IsTryAgain())
@@ -1474,7 +1566,7 @@ void FPlaylistReaderDASH::ManifestUpdateDownloadCompleted(FResourceLoadRequestPt
 	}
 	else
 	{
-		PostError(FString::Printf(TEXT("Failed to download MPD \"%s\" (%s)"), *Request->URL, *Request->GetErrorDetail()), ERRCODE_DASH_MPD_DOWNLOAD_FAILED, UEMEDIA_ERROR_READ_ERROR);
+		PostError(FString::Printf(TEXT("Failed to download MPD \"%s\" (%s)"), *Request->Url.URL, *Request->GetErrorDetail()), ERRCODE_DASH_MPD_DOWNLOAD_FAILED, UEMEDIA_ERROR_READ_ERROR);
 	}
 }
 
@@ -1511,6 +1603,39 @@ void FPlaylistReaderDASH::InitialMPDXLinkElementDownloadCompleted(FResourceLoadR
 }
 
 
+void FPlaylistReaderDASH::ContentSteeringUpdate(bool bInitial, FResourceLoadRequestPtr Request, bool bSuccess)
+{
+	HTTP::FConnectionInfo ci;
+	if (TSharedPtrTS<FHTTPResourceRequest> Req = Request->Request)
+	{
+		if (Req->GetConnectionInfo())
+		{
+			ci = *Req->GetConnectionInfo();
+		}
+	}
+	FString Response;
+	if (bSuccess)
+	{
+		GetResponseString(Response, Request);
+	}
+	PlayerSessionServices->GetContentSteeringHandler()->UpdateWithSteeringServerResponse(Response, ci.StatusInfo.HTTPStatus, ci.ResponseHeaders);
+	if (bInitial)
+	{
+		// Notify that the "variant playlists" are ready. There are no variants in DASH, but this is the trigger that the playlists are all set up and are good to go now.
+		PlayerSessionServices->SendMessageToPlayer(IPlaylistReader::PlaylistLoadedMessage::Create(LastErrorDetail, &ci, Playlist::EListType::Variant, Playlist::ELoadType::Initial, 0));
+		PrepareNextTimeSync();
+	}
+}
+void FPlaylistReaderDASH::ContentSteeringUpdateDownloadCompleted(FResourceLoadRequestPtr Request, bool bSuccess)
+{
+	ContentSteeringUpdate(false, Request, bSuccess);
+}
+void FPlaylistReaderDASH::ContentSteeringDownloadCompleted(FResourceLoadRequestPtr Request, bool bSuccess)
+{
+	ContentSteeringUpdate(true, Request, bSuccess);
+}
+
+
 void FPlaylistReaderDASH::Timesync_httphead_Completed(FResourceLoadRequestPtr Request, bool bSuccess)
 {
 	if (bSuccess)
@@ -1651,7 +1776,7 @@ void FPlaylistReaderDASH::OnMediaPlayerEventReceived(TSharedPtrTS<IAdaptiveStrea
 			TSharedPtrTS<FMPDLoadRequestDASH> CallbackRequest;
 			CallbackRequest = MakeSharedTS<FMPDLoadRequestDASH>();
 			CallbackRequest->LoadType = FMPDLoadRequestDASH::ELoadType::Callback;
-			CallbackRequest->URL = URL;
+			CallbackRequest->Url.URL = URL;
 			CallbackRequest->CompleteCallback.BindThreadSafeSP(AsShared(), &FPlaylistReaderDASH::Callback_Completed);
 			EnqueueResourceRequest(MoveTemp(CallbackRequest));
 		}
diff --git a/Engine/Plugins/Media/ElectraPlayer/Source/ElectraPlayerRuntime/Private/Runtime/Player/DASH/PlaylistReaderDASH_Internal.h b/Engine/Plugins/Media/ElectraPlayer/Source/ElectraPlayerRuntime/Private/Runtime/Player/DASH/PlaylistReaderDASH_Internal.h
--- a/Engine/Plugins/Media/ElectraPlayer/Source/ElectraPlayerRuntime/Private/Runtime/Player/DASH/PlaylistReaderDASH_Internal.h
+++ b/Engine/Plugins/Media/ElectraPlayer/Source/ElectraPlayerRuntime/Private/Runtime/Player/DASH/PlaylistReaderDASH_Internal.h
@@ -6,9 +6,11 @@
 #include "PlayerCore.h"
 
 #include "MPDElementsDASH.h"
+#include "MediaURLType.h"
 #include "Player/Manifest.h"
 #include "Player/AdaptiveStreamingPlayerResourceRequest.h"
 #include "Player/StreamSegmentReaderCommon.h"
+#include "Player/ContentSteeringHandler.h"
 #include "Utilities/URLParser.h"
 
 
@@ -36,7 +38,8 @@ struct FMPDLoadRequestDASH : public IHTTPResourceRequestObject
 		Callback,
 		Segment,
 		TimeSync,
-		Sideload
+		Sideload,
+		Steering
 	};
 	const TCHAR* const GetRequestTypeName() const
 	{
@@ -64,7 +67,8 @@ struct FMPDLoadRequestDASH : public IHTTPResourceRequestObject
 	}
 
 	FMPDLoadRequestDASH() : LoadType(ELoadType::MPD) {}
-	FString URL;	// For xlink requests this could be "urn:mpeg:dash:resolve-to-zero:2013" indicating removal of the element.
+	FMediaURL Url;	// For xlink requests this could be "urn:mpeg:dash:resolve-to-zero:2013" indicating removal of the element.
+	int64 SteeringID = -1;
 	FString Range;
 	FString Range2;
 	FString Verb;
@@ -105,6 +109,19 @@ struct FMPDLoadRequestDASH : public IHTTPResourceRequestObject
 
 namespace DASHUrlHelpers
 {
+	const FName SteerOption_ByteRange(TEXT("ByteRange"));
+	const FName SteerOption_ATO(TEXT("ATO"));
+	const FName SteerOption_ATOComplete(TEXT("ATOComplete"));
+	const FName SteerOption_AnnexIRequestHeader(TEXT("AnnexIRequestHeader"));
+
+	struct FDASHMediaURL : FMediaURL
+	{
+		FString ByteRange;
+		FTimeValue ATO {FTimeValue::GetZero()};
+		TMediaOptionalValue<bool> ATOComplete;
+	};
+	bool GetAllCandidateBaseURLs(TArray<FDASHMediaURL>& OutCandidateURLs, TSharedPtrTS<const IDashMPDElement> StartingAt, const FDASHMediaURL& InMPDUrl, bool bIsDVBDash);
+
 	bool IsAbsoluteURL(const FString& URL);
 	void GetAllHierarchyBaseURLs(IPlayerSessionServices* InPlayerSessionServices, TArray<TSharedPtrTS<const FDashMPD_BaseURLType>>& OutBaseURLs, TSharedPtrTS<const IDashMPDElement> StartingAt, const TCHAR* PreferredServiceLocation);
 	enum class EUrlQueryRequestType
@@ -114,12 +131,14 @@ namespace DASHUrlHelpers
 		Mpd,
 		Callback,
 		Chaining,
-		Fallback
+		Fallback,
+		Steering
 	};
 	void GetAllHierarchyUrlQueries(TArray<TSharedPtrTS<FDashMPD_UrlQueryInfoType>>& OutUrlQueries, TSharedPtrTS<const IDashMPDElement> StartingAt, EUrlQueryRequestType ForRequestType, bool bInclude2014);
 	FErrorDetail ApplyUrlQueries(IPlayerSessionServices* PlayerSessionServices, const FString& InMPDUrl, FString& InOutURL, FString& OutRequestHeader, const TArray<TSharedPtrTS<FDashMPD_UrlQueryInfoType>>& UrlQueries);
 	bool BuildAbsoluteElementURL(FString& OutURL, FTimeValue& ATO, TMediaOptionalValue<bool>& bATOComplete, const FString& DocumentURL, const TArray<TSharedPtrTS<const FDashMPD_BaseURLType>>& BaseURLs, const FString& ElementURL);
 	FString ApplyAnnexEByteRange(FString InURL, FString InRange, const TArray<TSharedPtrTS<const FDashMPD_BaseURLType>>& OutBaseURLs);
+	FString ApplyAnnexEByteRange(FString InURL, FString InRange, FString InByteRangeParam);
 }
 
 
@@ -223,6 +242,7 @@ public:
 		ESearchResult FindSegment_Timeline(IPlayerSessionServices* PlayerSessionServices, FSegmentInformationCommon& OutSegmentInfo, TArray<TWeakPtrTS<FMPDLoadRequestDASH>>& OutRemoteElementLoadRequests, const FSegmentSearchOption& InSearchOptions, const TSharedPtrTS<FDashMPD_RepresentationType>& MPDRepresentation, const TArray<TSharedPtrTS<FDashMPD_SegmentTemplateType>>& SegmentTemplate, const TSharedPtrTS<FDashMPD_SegmentTimelineType>& SegmentTimeline);
 		ESearchResult SetupSideloadedFile(IPlayerSessionServices* PlayerSessionServices, FSegmentInformationCommon& OutSegmentInfo, const FSegmentSearchOption& InSearchOptions, const TSharedPtrTS<FDashMPD_RepresentationType>& MPDRepresentation);
 
+		void PrepareSteeringCandidates(TArray<FContentSteeringHandler::FCandidateURL>& OutSteeringCandidates, IPlayerSessionServices* InPlayerSessionServices, const FString& InDocumentURL, const TArray<DASHUrlHelpers::FDASHMediaURL>& InCandidateURLs, const TArray<TSharedPtrTS<FDashMPD_UrlQueryInfoType>>& InUrlQueries, const FString& InTemplateURL);
 		bool PrepareDownloadURLs(IPlayerSessionServices* PlayerSessionServices, FSegmentInformationCommon& InOutSegmentInfo, const TArray<TSharedPtrTS<FDashMPD_SegmentBaseType>>& SegmentBase);
 		bool PrepareDownloadURLs(IPlayerSessionServices* PlayerSessionServices, FSegmentInformationCommon& InOutSegmentInfo, const TArray<TSharedPtrTS<FDashMPD_SegmentTemplateType>>& SegmentTemplate);
 		FString ApplyTemplateStrings(FString TemplateURL, const FSegmentInformationCommon& InSegmentInfo);
diff --git a/Engine/Plugins/Media/ElectraPlayer/Source/ElectraPlayerRuntime/Private/Runtime/Player/HLS/PlaylistHLS.cpp b/Engine/Plugins/Media/ElectraPlayer/Source/ElectraPlayerRuntime/Private/Runtime/Player/HLS/PlaylistHLS.cpp
--- a/Engine/Plugins/Media/ElectraPlayer/Source/ElectraPlayerRuntime/Private/Runtime/Player/HLS/PlaylistHLS.cpp
+++ b/Engine/Plugins/Media/ElectraPlayer/Source/ElectraPlayerRuntime/Private/Runtime/Player/HLS/PlaylistHLS.cpp
@@ -5,6 +5,8 @@
 #include "Player/AdaptivePlayerOptionKeynames.h"
 #include "Player/AdaptiveStreamingPlayerResourceRequest.h"
 #include "Player/StreamSegmentReaderCommon.h"
+#include "Player/AdaptiveStreamingPlayerABR.h"
+#include "Player/ContentSteeringHandler.h"
 #include "Utilities/Utilities.h"
 #include "Utilities/TimeUtilities.h"
 #include "Utilities/ISO639-Map.h"
@@ -16,7 +18,7 @@ namespace Electra
 {
 namespace HLS
 {
-const int32 kAssumedAudioBandwidth = 256000;
+const int32 kAssumedAudioBandwidth = 128000;
 const int32 kAssumedSubtitleBandwidth = 8000;
 }
 
@@ -154,18 +156,12 @@ FActiveHLSPlaylist::~FActiveHLSPlaylist()
 {
 }
 
-FErrorDetail FActiveHLSPlaylist::Create(TArray<TSharedPtr<FLoadRequestHLSPlaylist, ESPMode::ThreadSafe>>& OutPlaylistLoadRequests, IPlayerSessionServices* InPlayerSessionServices, TSharedPtr<FMultiVariantPlaylistHLS, ESPMode::ThreadSafe> InFromMultiVariantPlaylist)
-{
-	TUniquePtr<FInternalBuilder> Builder=MakeUnique<FInternalBuilder>();
-
-	if (!InPlayerSessionServices || !InPlayerSessionServices->GetStreamFilter() || !InFromMultiVariantPlaylist.IsValid())
-	{
-		return CreateError(FString::Printf(TEXT("Internal error")), HLS::ERRCODE_PLAYLIST_SETUP_FAILED);
-	}
 
-	TimelineMediaAsset ->PlayerSessionServices = PlayerSessionServices = InPlayerSessionServices;
+FErrorDetail FActiveHLSPlaylist::PreparePathway(TSharedPtrTS<FMultiVariantPlaylistHLS::FPathwayStreamInfs>& InOutPathway, TSharedPtr<FMultiVariantPlaylistHLS, ESPMode::ThreadSafe> InFromMultiVariantPlaylist)
+{
 	IPlayerStreamFilter* Filter = PlayerSessionServices->GetStreamFilter();
 
+	TUniquePtr<FInternalBuilder> Builder=MakeUnique<FInternalBuilder>();
 	// Set up the audio rendition groups filtered by supported codecs
 	for(int32 i=0, iMax=InFromMultiVariantPlaylist->RenditionGroupsOfType[static_cast<int32>(FMultiVariantPlaylistHLS::ERenditionGroupType::Audio)].Num(); i<iMax; ++i)
 	{
@@ -180,9 +176,11 @@ FErrorDetail FActiveHLSPlaylist::Create(TArray<TSharedPtr<FLoadRequestHLSPlaylis
 		}
 	}
 
+	TArray<FMultiVariantPlaylistHLS::FStreamInf>& sInfs(InOutPathway->StreamInfs);
+
 	// If there are audio-only variant groups check the for supported codec.
-	const TArray<FMultiVariantPlaylistHLS::FVideoVariantGroup>& vgrps = InFromMultiVariantPlaylist->VideoVariantGroups;
-	const TArray<FMultiVariantPlaylistHLS::FAudioVariantGroup>& agrps = InFromMultiVariantPlaylist->AudioOnlyVariantGroups;
+	const TArray<FMultiVariantPlaylistHLS::FVideoVariantGroup>& vgrps = InOutPathway->VideoVariantGroups;
+	const TArray<FMultiVariantPlaylistHLS::FAudioVariantGroup>& agrps = InOutPathway->AudioOnlyVariantGroups;
 	if (agrps.Num())
 	{
 		const FCodecSelectionPriorities& acSel = PlayerSessionServices->GetCodecSelectionPriorities(EStreamType::Audio);
@@ -208,7 +206,7 @@ FErrorDetail FActiveHLSPlaylist::Create(TArray<TSharedPtr<FLoadRequestHLSPlaylis
 		// Video groups may reference audio rendition groups, so the audio-only groups may not even be used.
 		if (Builder->UsableAudioVariantGroupIndices.IsEmpty() && vgrps.IsEmpty())
 		{
-			return CreateError(FString::Printf(TEXT("None of the audio-only variants can be played.")), HLS::ERRCODE_PLAYLIST_SETUP_FAILED);
+			return CreateError(FString::Printf(TEXT("None of the audio-only variants on pathway \"%s\" can be played."), *InOutPathway->PathwayID), HLS::ERRCODE_PLAYLIST_SETUP_FAILED);
 		}
 	}
 
@@ -237,7 +235,7 @@ FErrorDetail FActiveHLSPlaylist::Create(TArray<TSharedPtr<FLoadRequestHLSPlaylis
 		// If there is no usable video we fail. We do not play back audio-only in this case.
 		if (Builder->UsableVideoVariantGroupIndices.IsEmpty())
 		{
-			return CreateError(FString::Printf(TEXT("None of the video variants can be played.")), HLS::ERRCODE_PLAYLIST_SETUP_FAILED);
+			return CreateError(FString::Printf(TEXT("None of the video variants on pathway \"%s\" can be played."), *InOutPathway->PathwayID), HLS::ERRCODE_PLAYLIST_SETUP_FAILED);
 		}
 
 		// For each usable video variant group determine their associated audio group.
@@ -261,9 +259,9 @@ FErrorDetail FActiveHLSPlaylist::Create(TArray<TSharedPtr<FLoadRequestHLSPlaylis
 			for(int32 i=0; i<svg.StreamInfIndices.Num(); ++i)
 			{
 				FInternalBuilder::FVideoVariantAudio& vag = grp.Emplace_GetRef();
-				if (InFromMultiVariantPlaylist->StreamInfs[svg.StreamInfIndices[i]].NumAudioCodec)
+				if (sInfs[svg.StreamInfIndices[i]].NumAudioCodec)
 				{
-					vag.GroupName = InFromMultiVariantPlaylist->StreamInfs[svg.StreamInfIndices[i]].AudioGroup;
+					vag.GroupName = sInfs[svg.StreamInfIndices[i]].AudioGroup;
 					if (vag.GroupName.Len())
 					{
 						// References an audio group we haven't seen yet?
@@ -362,26 +360,139 @@ FErrorDetail FActiveHLSPlaylist::Create(TArray<TSharedPtr<FLoadRequestHLSPlaylis
 		Builder->SelectedVideoVariantGroup = GroupScores[0].Index;
 	}
 
+	FErrorDetail MetadataResult = CreateTrackMetadata(PlayerSessionServices, InOutPathway, InFromMultiVariantPlaylist, Builder.Get());
+	if (!MetadataResult.IsOK())
+	{
+		TimelineMediaAsset->MultiVariantPlaylist.Reset();
+		return MetadataResult;
+	}
+	return FErrorDetail();
+}
+
+FErrorDetail FActiveHLSPlaylist::Create(TArray<TSharedPtr<FLoadRequestHLSPlaylist, ESPMode::ThreadSafe>>& OutPlaylistLoadRequests, IPlayerSessionServices* InPlayerSessionServices, TSharedPtr<FMultiVariantPlaylistHLS, ESPMode::ThreadSafe> InFromMultiVariantPlaylist)
+{
+	if (!InPlayerSessionServices || !InPlayerSessionServices->GetStreamFilter() || !InFromMultiVariantPlaylist.IsValid())
+	{
+		return CreateError(FString::Printf(TEXT("Internal error")), HLS::ERRCODE_PLAYLIST_SETUP_FAILED);
+	}
+
+	TimelineMediaAsset->PlayerSessionServices = PlayerSessionServices = InPlayerSessionServices;
+	IPlayerStreamFilter* Filter = PlayerSessionServices->GetStreamFilter();
+
+	FErrorDetail LastErr;
+
+	// Handle each possible pathway on its own since they *could* all have different variants and renditions.
+	for(int32 pwIdx=0; pwIdx<InFromMultiVariantPlaylist->PathwayStreamInfs.Num(); ++pwIdx)
+	{
+		LastErr = PreparePathway(InFromMultiVariantPlaylist->PathwayStreamInfs[pwIdx], InFromMultiVariantPlaylist);
+		if (!LastErr.IsOK())
+		{
+			return LastErr;
+		}
+	}
+
 	TimelineMediaAsset->MultiVariantPlaylist = InFromMultiVariantPlaylist;
-	FErrorDetail Result = CreateTrackMetadata(PlayerSessionServices, Builder.Get());
-	if (Result.IsOK())
+
+	// Determine the pathway to be used.
+	FString CurrentPathway = TimelineMediaAsset->CurrentPathwayId;
+	FString NewPathwayId;
+	LastErr = DeterminePathwayToUse(PlayerSessionServices, NewPathwayId, CurrentPathway, InFromMultiVariantPlaylist);
+	if (!LastErr.IsOK())
+	{
+		TimelineMediaAsset->MultiVariantPlaylist.Reset();
+		return LastErr;
+	}
+	PlayerSessionServices->GetContentSteeringHandler()->SetCurrentlyActivePathway(NewPathwayId);
+
+	// Find the pathway and set it as the current one.
+	TSharedPtrTS<FMultiVariantPlaylistHLS::FPathwayStreamInfs> NewPathway;
+	for(auto& pwIt : InFromMultiVariantPlaylist->PathwayStreamInfs)
 	{
-		Result = GetInitialVariantPlaylistLoadRequests(OutPlaylistLoadRequests, PlayerSessionServices);
-		if (!Result.IsOK())
+		if (pwIt->PathwayID.Equals(NewPathwayId))
 		{
-			TimelineMediaAsset->MultiVariantPlaylist.Reset();
+			NewPathway = pwIt;
+			break;
 		}
 	}
-	else
+	TimelineMediaAsset->CurrentPathwayId = MoveTemp(NewPathwayId);
+	TimelineMediaAsset->SetCurrentPathway(NewPathway);
+
+	// Set up the variant and rendition playlist load requests for the streams we will be starting with.
+	LastErr = GetInitialVariantPlaylistLoadRequests(OutPlaylistLoadRequests, PlayerSessionServices);
+	if (!LastErr.IsOK())
 	{
 		TimelineMediaAsset->MultiVariantPlaylist.Reset();
 	}
-	return Result;
+	return LastErr;
+}
+
+FErrorDetail FActiveHLSPlaylist::DeterminePathwayToUse(IPlayerSessionServices* InPlayerSessionServices, FString& OutPathway, const FString& InCurrentPathway, const TSharedPtr<FMultiVariantPlaylistHLS, ESPMode::ThreadSafe>& InFromMultiVariantPlaylist)
+{
+	if (!InFromMultiVariantPlaylist.IsValid())
+	{
+		OutPathway = FString(TEXT("."));
+		return FErrorDetail();
+	}
+	TArray<FContentSteeringHandler::FCandidateURL> Candidates;
+	// If we are on a defined pathway add it to the candidate list first so it gets selected
+	// again in case there are no other matches.
+	if (InCurrentPathway.Len())
+	{
+		FContentSteeringHandler::FCandidateURL& pw(Candidates.Emplace_GetRef());
+		pw.MediaURL.CDN = InCurrentPathway;
+	}
+	for(auto& pwIt : InFromMultiVariantPlaylist->PathwayStreamInfs)
+	{
+		if (InCurrentPathway.IsEmpty() || (InCurrentPathway.Len() && !InCurrentPathway.Equals(pwIt->PathwayID)))
+		{
+			FContentSteeringHandler::FCandidateURL& pw(Candidates.Emplace_GetRef());
+			pw.MediaURL.CDN = pwIt->PathwayID;
+		}
+	}
+	FString SteeringMsg;
+	FContentSteeringHandler::FSelectedCandidateURL Selected = InPlayerSessionServices->GetContentSteeringHandler()->SelectBestCandidateFrom(SteeringMsg, FContentSteeringHandler::ESelectFor::Playlist, Candidates);
+	if (!Selected.MediaURL.CDN.IsEmpty())
+	{
+		OutPathway = Selected.MediaURL.CDN;
+		return FErrorDetail();
+	}
+	return CreateError(FString::Printf(TEXT("No PATHWAY is currently viable")), HLS::ERRCODE_PLAYLIST_SETUP_FAILED);
+}
+
+void FActiveHLSPlaylist::CheckForPathwaySwitch()
+{
+	// This gets called when the steering manifest was updated.
+	// Let's see if the update is requiring us to switch to a different pathway than we are on now.
+	TSharedPtrTS<FMultiVariantPlaylistHLS::FPathwayStreamInfs> Pwy = TimelineMediaAsset.IsValid() ? TimelineMediaAsset->GetCurrentPathway() : nullptr;
+	if (Pwy.IsValid() && TimelineMediaAsset->MultiVariantPlaylist.IsValid())
+	{
+		FString CurrentPathway = Pwy->PathwayID;
+		FString NewPathwayId;
+		FErrorDetail LastErr = DeterminePathwayToUse(PlayerSessionServices, NewPathwayId, CurrentPathway, TimelineMediaAsset->MultiVariantPlaylist);
+		if (LastErr.IsOK() && !NewPathwayId.Equals(CurrentPathway))
+		{
+			// Try to locate the new pathway
+			for(int32 i=0; i<TimelineMediaAsset->MultiVariantPlaylist->PathwayStreamInfs.Num(); ++i)
+			{
+				if (NewPathwayId.Equals(TimelineMediaAsset->MultiVariantPlaylist->PathwayStreamInfs[i]->PathwayID))
+				{
+					TimelineMediaAsset->SetCurrentPathway(TimelineMediaAsset->MultiVariantPlaylist->PathwayStreamInfs[i]);
+					PlayerSessionServices->GetContentSteeringHandler()->SetCurrentlyActivePathway(NewPathwayId);
+					if (auto stsel = PlayerSessionServices->GetStreamSelector())
+					{
+						stsel->PathwayChanged(NewPathwayId);
+					}
+					break;
+				}
+			}
+		}
+	}
 }
 
-FErrorDetail FActiveHLSPlaylist::CreateTrackMetadata(IPlayerSessionServices* InPlayerSessionServices, const FInternalBuilder* InBuilder)
+
+FErrorDetail FActiveHLSPlaylist::CreateTrackMetadata(IPlayerSessionServices* InPlayerSessionServices, TSharedPtrTS<FMultiVariantPlaylistHLS::FPathwayStreamInfs>& InPathway, TSharedPtr<FMultiVariantPlaylistHLS, ESPMode::ThreadSafe> InFromMultiVariantPlaylist, const FInternalBuilder* InBuilder)
 {
-	auto GetRenditionGroup = [mvp=TimelineMediaAsset->MultiVariantPlaylist](FMultiVariantPlaylistHLS::ERenditionGroupType InType, const FString& InName) -> const FMultiVariantPlaylistHLS::FRenditionGroup*
+	auto GetRenditionGroup = [mvp=InFromMultiVariantPlaylist](FMultiVariantPlaylistHLS::ERenditionGroupType InType, const FString& InName) -> const FMultiVariantPlaylistHLS::FRenditionGroup*
 	{
 		for(int32 i=0, iMax=mvp->RenditionGroupsOfType[static_cast<int32>(InType)].Num(); i<iMax; ++i)
 		{
@@ -397,29 +508,21 @@ FErrorDetail FActiveHLSPlaylist::CreateTrackMetadata(IPlayerSessionServices* InP
 	TArray<FString> VideoGroups;
 	TArray<FString> AudioGroups;
 	TArray<FString> SubtitleGroups;
-	TArray<FMultiVariantPlaylistHLS::FStreamInf> StreamInfsByPathway;
-	const FString& PrimaryPathwayId(TimelineMediaAsset->MultiVariantPlaylist->PrimaryPathwayId);
+	TArray<FMultiVariantPlaylistHLS::FStreamInf> GroupStreamInfs;
 
 	// Video (with or without audio) or audio-only?
 	if (InBuilder->SelectedVideoVariantGroup >= 0 || InBuilder->SelectedAudioVariantGroup >= 0)
 	{
-		const TArray<int32>& StreamInfIndices = InBuilder->SelectedVideoVariantGroup >= 0
-				? TimelineMediaAsset->MultiVariantPlaylist->VideoVariantGroups[InBuilder->SelectedVideoVariantGroup].StreamInfIndices
-				: TimelineMediaAsset->MultiVariantPlaylist->AudioOnlyVariantGroups[InBuilder->SelectedAudioVariantGroup].StreamInfIndices;
-		// Get the streaminfs that we are interested in for the primary pathway only.
-		for(int32 sIdx=0; sIdx<StreamInfIndices.Num(); ++sIdx)
+		for(auto& grSinf : InBuilder->SelectedVideoVariantGroup >= 0
+				? InPathway->VideoVariantGroups[InBuilder->SelectedVideoVariantGroup].StreamInfIndices
+				: InPathway->AudioOnlyVariantGroups[InBuilder->SelectedAudioVariantGroup].StreamInfIndices)
 		{
-			const FMultiVariantPlaylistHLS::FStreamInf& si = TimelineMediaAsset->MultiVariantPlaylist->StreamInfs[StreamInfIndices[sIdx]];
-			// If there are multiple CDNs we filter the list down to include only those of the primary pathway.
-			if (PrimaryPathwayId.IsEmpty() || PrimaryPathwayId.Equals(si.PathwayId))
-			{
-				StreamInfsByPathway.Emplace(si);
-			}
+			GroupStreamInfs.Emplace(InPathway->StreamInfs[grSinf]);
 		}
 		// Get the referenced group names from the streaminfs
-		for(int32 sIdx=0; sIdx<StreamInfsByPathway.Num(); ++sIdx)
+		for(int32 sIdx=0; sIdx<GroupStreamInfs.Num(); ++sIdx)
 		{
-			const FMultiVariantPlaylistHLS::FStreamInf& si = StreamInfsByPathway[sIdx];
+			const FMultiVariantPlaylistHLS::FStreamInf& si = GroupStreamInfs[sIdx];
 
 			FString GroupName;
 			GroupName = si.VideoGroup;
@@ -463,19 +566,19 @@ FErrorDetail FActiveHLSPlaylist::CreateTrackMetadata(IPlayerSessionServices* InP
 			}
 		}
 
-		auto AssignTracks = [&](EStreamType InType, FInternalTrackMetadata& tm) -> void
+		auto AssignTracks = [&](EStreamType InType, FMultiVariantPlaylistHLS::FInternalTrackMetadata& tm) -> void
 		{
-			for(int32 ns=0; ns<StreamInfsByPathway.Num(); ++ns)
+			for(int32 ns=0; ns<GroupStreamInfs.Num(); ++ns)
 			{
 				FStreamMetadata& sm = tm.Meta.StreamDetails.Emplace_GetRef();
-				sm.ID = StreamInfsByPathway[ns].StableVariantId;
-				sm.Bandwidth = StreamInfsByPathway[ns].Bandwidth;
-				sm.QualityIndex = StreamInfsByPathway[ns].QualityIndex;
-				for(int32 nc=0; nc<StreamInfsByPathway[ns].ParsedCodecs.Num(); ++nc)
+				sm.ID = GroupStreamInfs[ns].ID;
+				sm.Bandwidth = GroupStreamInfs[ns].Bandwidth;
+				sm.QualityIndex = GroupStreamInfs[ns].QualityIndex;
+				for(int32 nc=0; nc<GroupStreamInfs[ns].ParsedCodecs.Num(); ++nc)
 				{
-					if (StreamInfsByPathway[ns].ParsedCodecs[nc].IsCodec(InType))
+					if (GroupStreamInfs[ns].ParsedCodecs[nc].IsCodec(InType))
 					{
-						sm.CodecInformation = StreamInfsByPathway[ns].ParsedCodecs[nc];
+						sm.CodecInformation = GroupStreamInfs[ns].ParsedCodecs[nc];
 						break;
 					}
 				}
@@ -488,30 +591,30 @@ FErrorDetail FActiveHLSPlaylist::CreateTrackMetadata(IPlayerSessionServices* InP
 		};
 		if (VideoGroups.Num())
 		{
-			// Each group contains the streams of a certain bandwidth and the representations therein are the different "angles".
-			// We create the video tracks from the "angles" and assign the representations to the track.
-			// This requires all groups to have matching representations of course.
-			// Get the first group and create the "angle" tracks from it.
+			// Each group contains the alternatives to a variant of a certain bandwidth and the renditions therein are the different "angles".
+			// We create the video tracks from the "angles" and assign the renditions to the track.
+			// This requires all groups to have matching renditions of course.
+			// Get the first group (could be any, it does not matter) and create the "angle" tracks from it.
 			const FMultiVariantPlaylistHLS::FRenditionGroup* rg = GetRenditionGroup(FMultiVariantPlaylistHLS::ERenditionGroupType::Video, VideoGroups[0]);
 			for(int32 i=0; i<rg->Renditions.Num(); ++i)
 			{
-				FInternalTrackMetadata& tm = TimelineMediaAsset->VideoTracks.Emplace_GetRef();
-				tm.Meta.ID = FString::Printf(TEXT("$Angle%d"), i+1);
-				tm.Meta.Label = rg->Renditions[i].Name;
+				FMultiVariantPlaylistHLS::FInternalTrackMetadata& tm = InPathway->VideoTracks.Emplace_GetRef();
+				tm.Meta.ID = FString::Printf(TEXT("vid:%s"), *rg->Renditions[i].Name);
+				tm.Meta.Label = rg->Renditions[i].Name;	// Set the label to be the rendition's name. This gets used further down in comparisons!
 				tm.Meta.Kind = i==0 ? TEXT("main") : TEXT("alternative");
 			}
-			for(int32 vtIdx=0; vtIdx<TimelineMediaAsset->VideoTracks.Num(); ++vtIdx)
+			for(int32 vtIdx=0; vtIdx<InPathway->VideoTracks.Num(); ++vtIdx)
 			{
 				// Get the variants that reference this group.
 				for(int32 grpIdx=0; grpIdx<VideoGroups.Num(); ++grpIdx)
 				{
-					for(int32 vi=0; vi<StreamInfsByPathway.Num(); ++vi)
+					for(int32 vi=0; vi<GroupStreamInfs.Num(); ++vi)
 					{
-						if (!StreamInfsByPathway[vi].VideoGroup.Equals(VideoGroups[grpIdx]))
+						if (!GroupStreamInfs[vi].VideoGroup.Equals(VideoGroups[grpIdx]))
 						{
 							continue;
 						}
-						FInternalTrackMetadata& tm = TimelineMediaAsset->VideoTracks[vtIdx];
+						FMultiVariantPlaylistHLS::FInternalTrackMetadata& tm = InPathway->VideoTracks[vtIdx];
 						rg = GetRenditionGroup(FMultiVariantPlaylistHLS::ERenditionGroupType::Video, VideoGroups[grpIdx]);
 						bool bFoundRendition = false;
 						for(int32 k=0; k<rg->Renditions.Num(); ++k)
@@ -522,11 +625,11 @@ FErrorDetail FActiveHLSPlaylist::CreateTrackMetadata(IPlayerSessionServices* InP
 								// that we can't remove later.
 								if (Filter->CanDecodeStream(rg->Renditions[k].ParsedCodecFromStreamInf))
 								{
-									tm.VideoVariantBaseIDs.Emplace(StreamInfsByPathway[vi].StableVariantId);
+									tm.VideoVariantBaseIDs.Emplace(GroupStreamInfs[vi].ID);
 									FStreamMetadata& sm = tm.Meta.StreamDetails.Emplace_GetRef();
-									sm.Bandwidth = StreamInfsByPathway[vi].Bandwidth;
-									sm.QualityIndex = StreamInfsByPathway[vi].QualityIndex;
-									sm.ID = rg->Renditions[k].StableRenditionId;
+									sm.Bandwidth = GroupStreamInfs[vi].Bandwidth;
+									sm.QualityIndex = GroupStreamInfs[vi].QualityIndex;
+									sm.ID = rg->Renditions[k].Name;
 									sm.CodecInformation = rg->Renditions[k].ParsedCodecFromStreamInf;
 									if (sm.Bandwidth > tm.Meta.HighestBandwidth)
 									{
@@ -548,8 +651,8 @@ FErrorDetail FActiveHLSPlaylist::CreateTrackMetadata(IPlayerSessionServices* InP
 		}
 		else if (InBuilder->SelectedVideoVariantGroup >= 0)
 		{
-			FInternalTrackMetadata& tm = TimelineMediaAsset->VideoTracks.Emplace_GetRef();
-			tm.Meta.ID = TEXT("$Video");
+			FMultiVariantPlaylistHLS::FInternalTrackMetadata& tm = InPathway->VideoTracks.Emplace_GetRef();
+			tm.Meta.ID = TEXT("vid:");
 			tm.Meta.Kind = TEXT("main");
 			tm.bIsVariant = true;
 			AssignTracks(EStreamType::Video, tm);
@@ -562,13 +665,13 @@ FErrorDetail FActiveHLSPlaylist::CreateTrackMetadata(IPlayerSessionServices* InP
 			const FMultiVariantPlaylistHLS::FRenditionGroup* rg = GetRenditionGroup(FMultiVariantPlaylistHLS::ERenditionGroupType::Audio, AudioGroups[0]);
 			for(int32 i=0; i<rg->Renditions.Num(); ++i)
 			{
-				FInternalTrackMetadata& tm = TimelineMediaAsset->AudioTracks.Emplace_GetRef();
-				for(auto& vvIt : StreamInfsByPathway)
+				FMultiVariantPlaylistHLS::FInternalTrackMetadata& tm = InPathway->AudioTracks.Emplace_GetRef();
+				for(auto& vvIt : GroupStreamInfs)
 				{
-					tm.VideoVariantBaseIDs.Emplace(vvIt.StableVariantId);
+					tm.VideoVariantBaseIDs.Emplace(vvIt.ID);
 				}
 				tm.Rendition = rg->Renditions[i];
-				tm.Meta.ID = rg->Renditions[i].StableRenditionId;
+				tm.Meta.ID = FString::Printf(TEXT("aud:%s:%s"), *AudioGroups[0], *rg->Renditions[i].Name);
 				tm.Meta.Label = rg->Renditions[i].Name;
 				tm.Meta.Language = rg->Renditions[i].Language.IsEmpty() ? TEXT("und") : ISO639::RFC5646To639_1(rg->Renditions[i].Language);
 				tm.Meta.Kind = i==0 ? TEXT("main") : TEXT("translation");
@@ -586,7 +689,7 @@ FErrorDetail FActiveHLSPlaylist::CreateTrackMetadata(IPlayerSessionServices* InP
 					tm.Meta.HighestBandwidth = HLS::kAssumedAudioBandwidth;
 					tm.Meta.HighestBandwidthCodec = rg->Renditions[i].ParsedCodecFromStreamInf;
 					FStreamMetadata& sm = tm.Meta.StreamDetails.Emplace_GetRef();
-					sm.ID.Empty();	// do not set an ID here as we need to find the pathway of the variant
+					sm.ID.Empty();	// do not set an ID here to indicate this is a rendition
 					sm.Bandwidth = tm.Meta.HighestBandwidth;
 					sm.QualityIndex = 0;
 					sm.CodecInformation = tm.Meta.HighestBandwidthCodec;
@@ -597,8 +700,8 @@ FErrorDetail FActiveHLSPlaylist::CreateTrackMetadata(IPlayerSessionServices* InP
 		{
 			if (InBuilder->SelectedAudioVariantGroup >= 0)
 			{
-				FInternalTrackMetadata& tm = TimelineMediaAsset->AudioTracks.Emplace_GetRef();
-				tm.Meta.ID = TEXT("$Audio");
+				FMultiVariantPlaylistHLS::FInternalTrackMetadata& tm = InPathway->AudioTracks.Emplace_GetRef();
+				tm.Meta.ID = TEXT("aud:");
 				tm.Meta.Kind = TEXT("main");
 				tm.Meta.Language = TEXT("und");
 				tm.bIsVariant = true;
@@ -606,13 +709,17 @@ FErrorDetail FActiveHLSPlaylist::CreateTrackMetadata(IPlayerSessionServices* InP
 			}
 			else
 			{
+				// This is the case where there are no audio groups but at least one audio-only variant.
 				check(InBuilder->VideoGroupAssociatedAudio.Num());
 				check(InBuilder->SelectedVideoVariantGroup < InBuilder->VideoGroupAssociatedAudio.Num());
 
+				// There can be variant streams that include an audio codec and thus have inband-audio
+				// and others that have no audio codec and thus need to use the audio-only variant.
+				// Figure out which is which.
 				TArray<FInternalBuilder::FVideoVariantAudio> va;
-				for(int32 i=0; i<StreamInfsByPathway.Num(); ++i)
+				for(int32 i=0; i<GroupStreamInfs.Num(); ++i)
 				{
-					va.Emplace(InBuilder->VideoGroupAssociatedAudio[InBuilder->SelectedVideoVariantGroup][StreamInfsByPathway[i].Index]);
+					va.Emplace(InBuilder->VideoGroupAssociatedAudio[InBuilder->SelectedVideoVariantGroup][GroupStreamInfs[i].IndexOfSelfInArray]);
 					if (va.Last() != va[0])
 					{
 						return CreateError(FString::Printf(TEXT("Variant streams have inconsistent audio")), HLS::ERRCODE_PLAYLIST_SETUP_FAILED);
@@ -620,14 +727,15 @@ FErrorDetail FActiveHLSPlaylist::CreateTrackMetadata(IPlayerSessionServices* InP
 				}
 				if (va.Num() && !va[0].bHasNoAudio && !va[0].bIsEmpty)
 				{
-					FInternalTrackMetadata& tm = TimelineMediaAsset->AudioTracks.Emplace_GetRef();
+					FMultiVariantPlaylistHLS::FInternalTrackMetadata& tm = InPathway->AudioTracks.Emplace_GetRef();
 					tm.bIsVariant = true;
-					tm.Meta.ID = TEXT("$Audio");
+					tm.Meta.ID = TEXT("aud:");
 					tm.Meta.Kind = TEXT("main");
 					tm.Meta.Language = TEXT("und");
 					if (va[0].VariantGroupIndex < 0)
 					{
-						for(auto& siIt : StreamInfsByPathway)
+						// This is the case where a variant stream includes an audio codec.
+						for(auto& siIt : GroupStreamInfs)
 						{
 							if (siIt.Bandwidth > tm.Meta.HighestBandwidth)
 							{
@@ -652,16 +760,17 @@ FErrorDetail FActiveHLSPlaylist::CreateTrackMetadata(IPlayerSessionServices* InP
 					}
 					else
 					{
-						check(va[0].VariantGroupIndex < TimelineMediaAsset->MultiVariantPlaylist->AudioOnlyVariantGroups.Num());
-						check(TimelineMediaAsset->MultiVariantPlaylist->AudioOnlyVariantGroups[va[0].VariantGroupIndex].ParsedCodecs.Num());
+						// This is the case where earlier variant streams did not include an audio codec.
+						check(va[0].VariantGroupIndex < InPathway->AudioOnlyVariantGroups.Num());
+						check(InPathway->AudioOnlyVariantGroups[va[0].VariantGroupIndex].ParsedCodecs.Num());
 						tm.AudioVariantGroupIndex = va[0].VariantGroupIndex;
-						for(int nv=0; nv<TimelineMediaAsset->MultiVariantPlaylist->AudioOnlyVariantGroups[tm.AudioVariantGroupIndex].StreamInfIndices.Num(); ++nv)
+						for(int nv=0; nv<InPathway->AudioOnlyVariantGroups[tm.AudioVariantGroupIndex].StreamInfIndices.Num(); ++nv)
 						{
 							FStreamMetadata& sm = tm.Meta.StreamDetails.Emplace_GetRef();
-							sm.ID = TimelineMediaAsset->MultiVariantPlaylist->StreamInfs[TimelineMediaAsset->MultiVariantPlaylist->AudioOnlyVariantGroups[tm.AudioVariantGroupIndex].StreamInfIndices[nv]].StableVariantId;
-							sm.Bandwidth = TimelineMediaAsset->MultiVariantPlaylist->StreamInfs[TimelineMediaAsset->MultiVariantPlaylist->AudioOnlyVariantGroups[tm.AudioVariantGroupIndex].StreamInfIndices[nv]].Bandwidth;
-							sm.QualityIndex = TimelineMediaAsset->MultiVariantPlaylist->StreamInfs[TimelineMediaAsset->MultiVariantPlaylist->AudioOnlyVariantGroups[tm.AudioVariantGroupIndex].StreamInfIndices[nv]].QualityIndex;
-							sm.CodecInformation = TimelineMediaAsset->MultiVariantPlaylist->StreamInfs[TimelineMediaAsset->MultiVariantPlaylist->AudioOnlyVariantGroups[tm.AudioVariantGroupIndex].StreamInfIndices[nv]].ParsedCodecs[0];
+							sm.ID = InPathway->StreamInfs[InPathway->AudioOnlyVariantGroups[tm.AudioVariantGroupIndex].StreamInfIndices[nv]].ID;
+							sm.Bandwidth = InPathway->StreamInfs[InPathway->AudioOnlyVariantGroups[tm.AudioVariantGroupIndex].StreamInfIndices[nv]].Bandwidth;
+							sm.QualityIndex = InPathway->StreamInfs[InPathway->AudioOnlyVariantGroups[tm.AudioVariantGroupIndex].StreamInfIndices[nv]].QualityIndex;
+							sm.CodecInformation = InPathway->StreamInfs[InPathway->AudioOnlyVariantGroups[tm.AudioVariantGroupIndex].StreamInfIndices[nv]].ParsedCodecs[0];
 							if (sm.Bandwidth > tm.Meta.HighestBandwidth)
 							{
 								tm.Meta.HighestBandwidth = sm.Bandwidth;
@@ -680,10 +789,10 @@ FErrorDetail FActiveHLSPlaylist::CreateTrackMetadata(IPlayerSessionServices* InP
 			const FMultiVariantPlaylistHLS::FRenditionGroup* rg = GetRenditionGroup(FMultiVariantPlaylistHLS::ERenditionGroupType::Subtitles, SubtitleGroups[0]);
 			for(int32 i=0; i<rg->Renditions.Num(); ++i)
 			{
-				FInternalTrackMetadata& tm = TimelineMediaAsset->SubtitleTracks.Emplace_GetRef();
-				for(auto& vvIt : StreamInfsByPathway)
+				FMultiVariantPlaylistHLS::FInternalTrackMetadata& tm = InPathway->SubtitleTracks.Emplace_GetRef();
+				for(auto& vvIt : GroupStreamInfs)
 				{
-					tm.VideoVariantBaseIDs.Emplace(vvIt.StableVariantId);
+					tm.VideoVariantBaseIDs.Emplace(vvIt.ID);
 				}
 				tm.Rendition = rg->Renditions[i];
 				tm.Meta.ID = rg->Renditions[i].StableRenditionId;
@@ -700,7 +809,7 @@ FErrorDetail FActiveHLSPlaylist::CreateTrackMetadata(IPlayerSessionServices* InP
 				tm.Meta.HighestBandwidth = HLS::kAssumedSubtitleBandwidth;
 				tm.Meta.HighestBandwidthCodec = rg->Renditions[i].ParsedCodecFromStreamInf;
 				FStreamMetadata& sm = tm.Meta.StreamDetails.Emplace_GetRef();
-				sm.ID.Empty();	// do not set an ID here as we need to find the pathway of the variant
+				sm.ID.Empty();	// do not set an ID here to indicate this is a rendition
 				sm.Bandwidth = tm.Meta.HighestBandwidth;
 				sm.QualityIndex = 0;
 				sm.CodecInformation = tm.Meta.HighestBandwidthCodec;
@@ -709,12 +818,12 @@ FErrorDetail FActiveHLSPlaylist::CreateTrackMetadata(IPlayerSessionServices* InP
 	}
 
 	// Filter out the video streams that cannot be used on this device.
-	for(int32 vgIdx=0; vgIdx<TimelineMediaAsset->VideoTracks.Num(); ++vgIdx)
+	for(int32 vgIdx=0; vgIdx<InPathway->VideoTracks.Num(); ++vgIdx)
 	{
 		FStreamCodecInformation HighestBandwidthCodec;
 		int32 HighestBandwidth = 0;
 
-		FInternalTrackMetadata& vg(TimelineMediaAsset->VideoTracks[vgIdx]);
+		FMultiVariantPlaylistHLS::FInternalTrackMetadata& vg(InPathway->VideoTracks[vgIdx]);
 		for(int32 stIdx=0; stIdx<vg.Meta.StreamDetails.Num(); ++stIdx)
 		{
 			if (Filter->CanDecodeStream(vg.Meta.StreamDetails[stIdx].CodecInformation))
@@ -738,24 +847,24 @@ FErrorDetail FActiveHLSPlaylist::CreateTrackMetadata(IPlayerSessionServices* InP
 		}
 		else
 		{
-			TimelineMediaAsset->VideoTracks.RemoveAt(vgIdx);
+			InPathway->VideoTracks.RemoveAt(vgIdx);
 			--vgIdx;
 		}
 	}
 
 	// Create internal AdaptationSets that are used to interface with the player.
-	auto CreateAdaptationSet = [](TArray<TSharedPtrTS<FPlaybackAssetAdaptationSet>>& OutAdaptationSets, const TArray<FInternalTrackMetadata>& InTrackMetadata) -> void
+	auto CreateAdaptationSet = [](TArray<TSharedPtrTS<FMultiVariantPlaylistHLS::FPlaybackAssetAdaptationSet>>& OutAdaptationSets, const TArray<FMultiVariantPlaylistHLS::FInternalTrackMetadata>& InTrackMetadata) -> void
 	{
 		for(auto& tm : InTrackMetadata)
 		{
-			TSharedPtrTS<FPlaybackAssetAdaptationSet> as = MakeSharedTS<FPlaybackAssetAdaptationSet>();
+			TSharedPtrTS<FMultiVariantPlaylistHLS::FPlaybackAssetAdaptationSet> as = MakeSharedTS<FMultiVariantPlaylistHLS::FPlaybackAssetAdaptationSet>();
 			as->ID = tm.Meta.ID;
 			as->ListOfCodecs = tm.Meta.HighestBandwidthCodec.GetCodecSpecifierRFC6381();
 			as->Language = tm.Meta.Language;
 			int32 sIdx = 0;
 			for(auto& sm : tm.Meta.StreamDetails)
 			{
-				TSharedPtrTS<FPlaybackAssetRepresentation> repr = MakeSharedTS<FPlaybackAssetRepresentation>();
+				TSharedPtrTS<FMultiVariantPlaylistHLS::FPlaybackAssetRepresentation> repr = MakeSharedTS<FMultiVariantPlaylistHLS::FPlaybackAssetRepresentation>();
 
 				repr->StreamCodecInformation = sm.CodecInformation;
 				repr->Bandwidth = sm.Bandwidth;
@@ -767,9 +876,9 @@ FErrorDetail FActiveHLSPlaylist::CreateTrackMetadata(IPlayerSessionServices* InP
 			OutAdaptationSets.Emplace(MoveTemp(as));
 		}
 	};
-	CreateAdaptationSet(TimelineMediaAsset->VideoAdaptationSets, TimelineMediaAsset->VideoTracks);
-	CreateAdaptationSet(TimelineMediaAsset->AudioAdaptationSets, TimelineMediaAsset->AudioTracks);
-	CreateAdaptationSet(TimelineMediaAsset->SubtitleAdaptationSets, TimelineMediaAsset->SubtitleTracks);
+	CreateAdaptationSet(InPathway->VideoAdaptationSets, InPathway->VideoTracks);
+	CreateAdaptationSet(InPathway->AudioAdaptationSets, InPathway->AudioTracks);
+	CreateAdaptationSet(InPathway->SubtitleAdaptationSets, InPathway->SubtitleTracks);
 	return FErrorDetail();
 }
 
@@ -783,42 +892,42 @@ void FActiveHLSPlaylist::GetAllMediaPlaylistLoadRequests(TArray<TSharedPtr<FLoad
 		return;
 	}
 
-	TimelineMediaAsset->CurrentPathwayId = TimelineMediaAsset->MultiVariantPlaylist->PrimaryPathwayId;
+	auto Pwy = TimelineMediaAsset->GetCurrentPathway();
 
 	FErrorDetail Result;
 	TSharedPtr<FLoadRequestHLSPlaylist, ESPMode::ThreadSafe> PlaylistLoadRequest;
-	if (InForType == EStreamType::Video && TimelineMediaAsset->VideoTracks.Num())
+	if (InForType == EStreamType::Video && Pwy->VideoTracks.Num())
 	{
 		const int32 kTrack = 0;
-		for(int32 i=0,iMax=TimelineMediaAsset->VideoTracks[kTrack].Meta.StreamDetails.Num(); i<iMax; ++i)
+		for(int32 i=0,iMax=Pwy->VideoTracks[kTrack].Meta.StreamDetails.Num(); i<iMax; ++i)
 		{
-			Result = TimelineMediaAsset->GetVariantPlaylist(PlaylistLoadRequest, PlayerSessionServices, EStreamType::Video, TimelineMediaAsset->CurrentPathwayId, kTrack, i, 0, 0);
+			Result = TimelineMediaAsset->GetVariantPlaylist(PlaylistLoadRequest, PlayerSessionServices, EStreamType::Video, Pwy, kTrack, i, 0, 0);
 			if (Result.IsOK() && PlaylistLoadRequest.IsValid())
 			{
 				OutPlaylistLoadRequests.Emplace(MoveTemp(PlaylistLoadRequest));
 			}
 		}
 	}
-	else if (InForType == EStreamType::Audio && TimelineMediaAsset->AudioTracks.Num())
+	else if (InForType == EStreamType::Audio && Pwy->AudioTracks.Num())
 	{
 		const int32 kTrack = 0;
-		for(int32 i=0,iMax=TimelineMediaAsset->AudioTracks[kTrack].Meta.StreamDetails.Num(); i<iMax; ++i)
+		for(int32 i=0,iMax=Pwy->AudioTracks[kTrack].Meta.StreamDetails.Num(); i<iMax; ++i)
 		{
 			// TBD: Since this is used to determine alternatives for the initial variant playlist, can we safely pass 0,0 for the main stream indices?
-			Result = TimelineMediaAsset->GetVariantPlaylist(PlaylistLoadRequest, PlayerSessionServices, EStreamType::Audio, TimelineMediaAsset->CurrentPathwayId, kTrack, i, 0, 0);
+			Result = TimelineMediaAsset->GetVariantPlaylist(PlaylistLoadRequest, PlayerSessionServices, EStreamType::Audio, Pwy, kTrack, i, 0, 0);
 			if (Result.IsOK() && PlaylistLoadRequest.IsValid())
 			{
 				OutPlaylistLoadRequests.Emplace(MoveTemp(PlaylistLoadRequest));
 			}
 		}
 	}
-	else if (InForType == EStreamType::Subtitle && TimelineMediaAsset->SubtitleTracks.Num())
+	else if (InForType == EStreamType::Subtitle && Pwy->SubtitleTracks.Num())
 	{
 		const int32 kTrack = 0;
-		for(int32 i=0,iMax=TimelineMediaAsset->SubtitleTracks[kTrack].Meta.StreamDetails.Num(); i<iMax; ++i)
+		for(int32 i=0,iMax=Pwy->SubtitleTracks[kTrack].Meta.StreamDetails.Num(); i<iMax; ++i)
 		{
 			// TBD: Since this is used to determine alternatives for the initial variant playlist, can we safely pass 0,0 for the main stream indices?
-			Result = TimelineMediaAsset->GetVariantPlaylist(PlaylistLoadRequest, PlayerSessionServices, EStreamType::Subtitle, TimelineMediaAsset->CurrentPathwayId, kTrack, i, 0, 0);
+			Result = TimelineMediaAsset->GetVariantPlaylist(PlaylistLoadRequest, PlayerSessionServices, EStreamType::Subtitle, Pwy, kTrack, i, 0, 0);
 			if (Result.IsOK() && PlaylistLoadRequest.IsValid())
 			{
 				OutPlaylistLoadRequests.Emplace(MoveTemp(PlaylistLoadRequest));
@@ -831,20 +940,19 @@ void FActiveHLSPlaylist::GetAllMediaPlaylistLoadRequests(TArray<TSharedPtr<FLoad
 FErrorDetail FActiveHLSPlaylist::GetInitialVariantPlaylistLoadRequests(TArray<TSharedPtr<FLoadRequestHLSPlaylist, ESPMode::ThreadSafe>>& OutPlaylistLoadRequests, IPlayerSessionServices* InPlayerSessionServices)
 {
 	// Get the first variant or rendition playlists to load.
-	if (!TimelineMediaAsset->VideoTracks.Num() && !TimelineMediaAsset->AudioTracks.Num())
+	auto Pwy = TimelineMediaAsset->GetCurrentPathway();
+	if (!Pwy.IsValid() || (!Pwy->VideoAdaptationSets.Num() && !Pwy->AudioAdaptationSets.Num()))
 	{
 		return CreateError(FString::Printf(TEXT("Nothing usable found in multivariant playlist")), HLS::ERRCODE_PLAYLIST_SETUP_FAILED);
 	}
 
-	TimelineMediaAsset->CurrentPathwayId = TimelineMediaAsset->MultiVariantPlaylist->PrimaryPathwayId;
-
 	FErrorDetail Result;
 	TSharedPtr<FLoadRequestHLSPlaylist, ESPMode::ThreadSafe> PlaylistLoadRequest;
 	TArray<FString> URLs;
 	bool bIsPrimary = true;
-	if (TimelineMediaAsset->VideoTracks.Num())
+	if (Pwy->VideoAdaptationSets.Num())
 	{
-		Result = TimelineMediaAsset->GetVariantPlaylist(PlaylistLoadRequest, InPlayerSessionServices, EStreamType::Video, TimelineMediaAsset->CurrentPathwayId, 0, 0, 0, 0);
+		Result = TimelineMediaAsset->GetVariantPlaylist(PlaylistLoadRequest, InPlayerSessionServices, EStreamType::Video, Pwy, 0, 0, 0, 0);
 		if (Result.IsOK() && PlaylistLoadRequest.IsValid())
 		{
 			PlaylistLoadRequest->bIsPrimaryPlaylist = bIsPrimary;
@@ -854,9 +962,9 @@ FErrorDetail FActiveHLSPlaylist::GetInitialVariantPlaylistLoadRequests(TArray<TS
 			bIsPrimary = false;
 		}
 	}
-	if (Result.IsOK() && TimelineMediaAsset->AudioTracks.Num())
+	if (Result.IsOK() && Pwy->AudioAdaptationSets.Num())
 	{
-		Result = TimelineMediaAsset->GetVariantPlaylist(PlaylistLoadRequest, InPlayerSessionServices, EStreamType::Audio, TimelineMediaAsset->CurrentPathwayId, 0, 0, 0, 0);
+		Result = TimelineMediaAsset->GetVariantPlaylist(PlaylistLoadRequest, InPlayerSessionServices, EStreamType::Audio, Pwy, 0, 0, 0, 0);
 		if (Result.IsOK() && PlaylistLoadRequest.IsValid() && !URLs.Contains(PlaylistLoadRequest->ResourceRequest->GetURL()))
 		{
 			PlaylistLoadRequest->bIsPrimaryPlaylist = bIsPrimary;
@@ -1328,33 +1436,33 @@ TSharedPtr<FMediaPlaylistAndStateHLS, ESPMode::ThreadSafe> FActiveHLSPlaylist::F
 }
 
 
-FErrorDetail FActiveHLSPlaylist::FTimelineMediaAsset::GetVariantPlaylist(TSharedPtr<FLoadRequestHLSPlaylist, ESPMode::ThreadSafe>& OutPlaylistLoadRequest, IPlayerSessionServices* InPlayerSessionServices, EStreamType InStreamType, const FString& InPathway, int32 InTrackIndex, int32 InStreamIndex, int32 InMainTrackIndex, int32 InMainStreamIndex)
+FErrorDetail FActiveHLSPlaylist::FTimelineMediaAsset::GetVariantPlaylist(TSharedPtr<FLoadRequestHLSPlaylist, ESPMode::ThreadSafe>& OutPlaylistLoadRequest, IPlayerSessionServices* InPlayerSessionServices, EStreamType InStreamType, const TSharedPtrTS<FMultiVariantPlaylistHLS::FPathwayStreamInfs>& InPathway, int32 InTrackIndex, int32 InStreamIndex, int32 InMainTrackIndex, int32 InMainStreamIndex)
 {
 	auto GetTimeoutValue = [pss=InPlayerSessionServices](const FName& InOptionName, int32 InDefaultValueMillisec) -> FTimeValue
 	{
 		return pss->GetOptionValue(InOptionName).SafeGetTimeValue(FTimeValue().SetFromMilliseconds(InDefaultValueMillisec));
 	};
 
-	const TArray<FInternalTrackMetadata>* Tracks = nullptr;
-	const TArray<TSharedPtrTS<FPlaybackAssetAdaptationSet>>* AdaptationSets = nullptr;
+	const TArray<FMultiVariantPlaylistHLS::FInternalTrackMetadata>* Tracks = nullptr;
+	const TArray<TSharedPtrTS<FMultiVariantPlaylistHLS::FPlaybackAssetAdaptationSet>>* AdaptationSets = nullptr;
 	switch(InStreamType)
 	{
 		case EStreamType::Video:
 		{
-			Tracks = &VideoTracks;
-			AdaptationSets = &VideoAdaptationSets;
+			Tracks = &InPathway->VideoTracks;
+			AdaptationSets = &InPathway->VideoAdaptationSets;
 			break;
 		}
 		case EStreamType::Audio:
 		{
-			Tracks = &AudioTracks;
-			AdaptationSets = &AudioAdaptationSets;
+			Tracks = &InPathway->AudioTracks;
+			AdaptationSets = &InPathway->AudioAdaptationSets;
 			break;
 		}
 		case EStreamType::Subtitle:
 		{
-			Tracks = &SubtitleTracks;
-			AdaptationSets = &SubtitleAdaptationSets;
+			Tracks = &InPathway->SubtitleTracks;
+			AdaptationSets = &InPathway->SubtitleAdaptationSets;
 			break;
 		}
 		default:
@@ -1366,13 +1474,13 @@ FErrorDetail FActiveHLSPlaylist::FTimelineMediaAsset::GetVariantPlaylist(TShared
 	{
 		return CreateError(FString::Printf(TEXT("Invalid %s track index %d"), GetStreamTypeName(InStreamType), InTrackIndex), HLS::ERRCODE_PLAYLIST_SETUP_FAILED);
 	}
-	const FInternalTrackMetadata& tm = (*Tracks)[InTrackIndex];
+	const FMultiVariantPlaylistHLS::FInternalTrackMetadata& tm = (*Tracks)[InTrackIndex];
 	if (InStreamIndex >= tm.Meta.StreamDetails.Num())
 	{
 		return CreateError(FString::Printf(TEXT("Invalid %s stream index %d for track index %d"), GetStreamTypeName(InStreamType), InStreamIndex, InTrackIndex), HLS::ERRCODE_PLAYLIST_SETUP_FAILED);
 	}
 
-	TSharedPtrTS<FPlaybackAssetAdaptationSet> AdaptationSet = (*AdaptationSets)[InTrackIndex];
+	TSharedPtrTS<FMultiVariantPlaylistHLS::FPlaybackAssetAdaptationSet> AdaptationSet = (*AdaptationSets)[InTrackIndex];
 	TSharedPtrTS<IPlaybackAssetRepresentation> Representation;
 	const FStreamMetadata& sm = tm.Meta.StreamDetails[InStreamIndex];
 
@@ -1417,19 +1525,18 @@ FErrorDetail FActiveHLSPlaylist::FTimelineMediaAsset::GetVariantPlaylist(TShared
 
 	// Locate the variant.
 	const FMultiVariantPlaylistHLS::FStreamInf* si = nullptr;
-	for(int32 vIdx=0; vIdx<MultiVariantPlaylist->StreamInfs.Num(); ++vIdx)
+	for(int32 vIdx=0; vIdx<InPathway->StreamInfs.Num(); ++vIdx)
 	{
-		if (MultiVariantPlaylist->StreamInfs[vIdx].StableVariantId.Equals(VariantID) &&
-			MultiVariantPlaylist->StreamInfs[vIdx].PathwayId.Equals(InPathway))
+		if (InPathway->StreamInfs[vIdx].ID.Equals(VariantID))
 		{
-			si = &MultiVariantPlaylist->StreamInfs[vIdx];
+			si = &InPathway->StreamInfs[vIdx];
 			Representation = AdaptationSet->GetRepresentationByUniqueIdentifier(VariantID);
 			break;
 		}
 	}
 	if (!si)
 	{
-		return CreateError(FString::Printf(TEXT("Variant stream \"%s\" not found for pathway \"%s\""), *VariantID, *InPathway), HLS::ERRCODE_PLAYLIST_SETUP_FAILED);
+		return CreateError(FString::Printf(TEXT("Variant stream \"%s\" not found for pathway \"%s\""), *VariantID, *InPathway->PathwayID), HLS::ERRCODE_PLAYLIST_SETUP_FAILED);
 	}
 
 	// Were we looking for a variant?
@@ -1482,34 +1589,16 @@ FErrorDetail FActiveHLSPlaylist::FTimelineMediaAsset::GetVariantPlaylist(TShared
 			return CreateError(FString::Printf(TEXT("Rendition group \"%s\" not found"), *GroupName), HLS::ERRCODE_PLAYLIST_SETUP_FAILED);
 		}
 		const FMultiVariantPlaylistHLS::FRendition* Rendition = nullptr;
-		// Search for matching rendition by stable id.
-		if (tm.Rendition.IsSet())
-		{
-			for(int32 k=0; k<Renditions->Num(); ++k)
-			{
-				if ((*Renditions)[k].StableRenditionId.Equals(tm.Rendition.GetValue().StableRenditionId))
-				{
-					Rendition = &(*Renditions)[k];
-					// Adaptation sets created from renditions only have a single rendition in them.
-					check(AdaptationSet->GetNumberOfRepresentations() == 1);
-					Representation = AdaptationSet->GetRepresentationByIndex(0);
-					break;
-				}
-			}
-		}
-		// If not found (because not a steering manifest) look for matching name.
-		if (!Rendition)
+		for(int32 k=0; k<Renditions->Num(); ++k)
 		{
-			for(int32 k=0; k<Renditions->Num(); ++k)
+			if ((*Renditions)[k].Name.Equals(tm.Meta.Label))
 			{
-				if ((*Renditions)[k].Name.Equals(tm.Meta.Label))
-				{
-					Rendition = &(*Renditions)[k];
-					// Adaptation sets created from renditions only have a single rendition in them.
-					check(AdaptationSet->GetNumberOfRepresentations() == 1);
-					Representation = AdaptationSet->GetRepresentationByIndex(0);
-					break;
-				}
+				Rendition = &(*Renditions)[k];
+				// Adaptation sets created from renditions only have a single rendition in them.
+// FIXME: this is not quite right with multiple video angles
+				//check(AdaptationSet->GetNumberOfRepresentations() == 1);
+				Representation = AdaptationSet->GetRepresentationByIndex(0);
+				break;
 			}
 		}
 		if (!Rendition)
@@ -1533,7 +1622,7 @@ FErrorDetail FActiveHLSPlaylist::FTimelineMediaAsset::GetVariantPlaylist(TShared
 		OutPlaylistLoadRequest->PlaylistInfo.AdaptationSetID = AdaptationSet->GetUniqueIdentifier();
 		check(Representation.IsValid());
 		OutPlaylistLoadRequest->PlaylistInfo.RepresentationID = Representation.IsValid() ? Representation->GetUniqueIdentifier() : VariantID;
-		OutPlaylistLoadRequest->PlaylistInfo.PathwayID = InPathway;
+		OutPlaylistLoadRequest->PlaylistInfo.PathwayID = InPathway->PathwayID;
 		OutPlaylistLoadRequest->PlaylistInfo.RepresentationBandwidth = Representation.IsValid() ? Representation->GetBitrate() : 0;
 		OutPlaylistLoadRequest->ResourceRequest = MakeShared<FHTTPResourceRequest, ESPMode::ThreadSafe>();
 		OutPlaylistLoadRequest->ResourceRequest->Verb(TEXT("GET")).URL(URL)
@@ -1784,9 +1873,9 @@ FActiveHLSPlaylist::FTimelineMediaAsset::ESegSearchResult FActiveHLSPlaylist::FT
 	{
 		FURL_RFC3986 UrlParserInit(mp->ParsedURL);
 		UrlParserInit.ResolveWith(Seg.InitSegment->URL);
-		OutSegment->Segment.InitializationURL.URL = UrlParserInit.Get(true, true);
+		OutSegment->Segment.InitializationURL.Url.URL = UrlParserInit.Get(true, true);
 		OutSegment->Segment.InitializationURL.Range = Seg.InitSegment->ByteRange.GetForHTTP();
-		OutSegment->Segment.InitializationURL.CDN = CurrentPathwayId;
+		OutSegment->Segment.InitializationURL.Url.CDN = InPlaylist->PlaylistInfo.PathwayID;
 		if (Seg.InitSegment->Encryption.IsValid())
 		{
 			FDRMClientCacheHLS::FEntry DrmClientEntry;
@@ -1809,9 +1898,9 @@ FActiveHLSPlaylist::FTimelineMediaAsset::ESegSearchResult FActiveHLSPlaylist::FT
 
 	FURL_RFC3986 UrlParserMedia(mp->ParsedURL);
 	UrlParserMedia.ResolveWith(Seg.URL);
-	OutSegment->Segment.MediaURL.URL = UrlParserMedia.Get(true, true);
+	OutSegment->Segment.MediaURL.Url.URL = UrlParserMedia.Get(true, true);
 	OutSegment->Segment.MediaURL.Range = Seg.ByteRange.GetForHTTP();
-	OutSegment->Segment.MediaURL.CDN = CurrentPathwayId;
+	OutSegment->Segment.MediaURL.Url.CDN = InPlaylist->PlaylistInfo.PathwayID;
 	if (Seg.Encryption.IsValid())
 	{
 		FDRMClientCacheHLS::FEntry DrmClientEntry;
@@ -1837,53 +1926,61 @@ FActiveHLSPlaylist::FTimelineMediaAsset::ESegSearchResult FActiveHLSPlaylist::FT
 }
 
 
-const FActiveHLSPlaylist::FInternalTrackMetadata* FActiveHLSPlaylist::FTimelineMediaAsset::GetInternalTrackMetadata(const FString& InForID) const
+const FMultiVariantPlaylistHLS::FInternalTrackMetadata* FActiveHLSPlaylist::FTimelineMediaAsset::GetInternalTrackMetadata(const FString& InForID) const
 {
-	for(int32 i=0; i<VideoTracks.Num(); ++i)
+	auto Pwy = GetCurrentPathway();
+	if (Pwy.IsValid())
 	{
-		if (VideoTracks[i].Meta.ID.Equals(InForID))
+		for(int32 i=0; i<Pwy->VideoTracks.Num(); ++i)
 		{
-			return &VideoTracks[i];
+			if (Pwy->VideoTracks[i].Meta.ID.Equals(InForID))
+			{
+				return &Pwy->VideoTracks[i];
+			}
 		}
-	}
-	for(int32 i=0; i<AudioTracks.Num(); ++i)
-	{
-		if (AudioTracks[i].Meta.ID.Equals(InForID))
+		for(int32 i=0; i<Pwy->AudioTracks.Num(); ++i)
 		{
-			return &AudioTracks[i];
+			if (Pwy->AudioTracks[i].Meta.ID.Equals(InForID))
+			{
+				return &Pwy->AudioTracks[i];
+			}
 		}
-	}
-	for(int32 i=0; i<SubtitleTracks.Num(); ++i)
-	{
-		if (SubtitleTracks[i].Meta.ID.Equals(InForID))
+		for(int32 i=0; i<Pwy->SubtitleTracks.Num(); ++i)
 		{
-			return &SubtitleTracks[i];
+			if (Pwy->SubtitleTracks[i].Meta.ID.Equals(InForID))
+			{
+				return &Pwy->SubtitleTracks[i];
+			}
 		}
 	}
 	return nullptr;
 }
 
-TSharedPtrTS<FActiveHLSPlaylist::FPlaybackAssetAdaptationSet> FActiveHLSPlaylist::FTimelineMediaAsset::GetAdaptationSet(const FString& InForID) const
+TSharedPtrTS<FMultiVariantPlaylistHLS::FPlaybackAssetAdaptationSet> FActiveHLSPlaylist::FTimelineMediaAsset::GetAdaptationSet(const FString& InForID) const
 {
-	for(int32 i=0; i<VideoAdaptationSets.Num(); ++i)
+	auto Pwy = GetCurrentPathway();
+	if (Pwy.IsValid())
 	{
-		if (VideoAdaptationSets[i]->ID.Equals(InForID))
+		for(int32 i=0; i<Pwy->VideoAdaptationSets.Num(); ++i)
 		{
-			return VideoAdaptationSets[i];
+			if (Pwy->VideoAdaptationSets[i]->ID.Equals(InForID))
+			{
+				return Pwy->VideoAdaptationSets[i];
+			}
 		}
-	}
-	for(int32 i=0; i<AudioAdaptationSets.Num(); ++i)
-	{
-		if (AudioAdaptationSets[i]->ID.Equals(InForID))
+		for(int32 i=0; i<Pwy->AudioAdaptationSets.Num(); ++i)
 		{
-			return AudioAdaptationSets[i];
+			if (Pwy->AudioAdaptationSets[i]->ID.Equals(InForID))
+			{
+				return Pwy->AudioAdaptationSets[i];
+			}
 		}
-	}
-	for(int32 i=0; i<SubtitleAdaptationSets.Num(); ++i)
-	{
-		if (SubtitleAdaptationSets[i]->ID.Equals(InForID))
+		for(int32 i=0; i<Pwy->SubtitleAdaptationSets.Num(); ++i)
 		{
-			return SubtitleAdaptationSets[i];
+			if (Pwy->SubtitleAdaptationSets[i]->ID.Equals(InForID))
+			{
+				return Pwy->SubtitleAdaptationSets[i];
+			}
 		}
 	}
 	return nullptr;
@@ -2031,19 +2128,24 @@ FString FActiveHLSPlaylist::FTimelineMediaAsset::GetUniqueIdentifier() const
 
 int32 FActiveHLSPlaylist::FTimelineMediaAsset::GetNumberOfAdaptationSets(EStreamType InStreamType) const
 {
+	auto Pwy = GetCurrentPathway();
+	if (!Pwy.IsValid())
+	{
+		return 0;
+	}
 	switch(InStreamType)
 	{
 		case EStreamType::Video:
 		{
-			return VideoTracks.Num();
+			return Pwy->VideoTracks.Num();
 		}
 		case EStreamType::Audio:
 		{
-			return AudioTracks.Num();
+			return Pwy->AudioTracks.Num();
 		}
 		case EStreamType::Subtitle:
 		{
-			return SubtitleTracks.Num();
+			return Pwy->SubtitleTracks.Num();
 		}
 		default:
 		{
@@ -2054,41 +2156,49 @@ int32 FActiveHLSPlaylist::FTimelineMediaAsset::GetNumberOfAdaptationSets(EStream
 
 TSharedPtrTS<IPlaybackAssetAdaptationSet> FActiveHLSPlaylist::FTimelineMediaAsset::GetAdaptationSetByTypeAndIndex(EStreamType InStreamType, int32 InAdaptationSetIndex) const
 {
-	if (InStreamType == EStreamType::Video && InAdaptationSetIndex >= 0 && InAdaptationSetIndex < VideoAdaptationSets.Num())
-	{
-		return VideoAdaptationSets[InAdaptationSetIndex];
-	}
-	else if (InStreamType == EStreamType::Audio && InAdaptationSetIndex >= 0 && InAdaptationSetIndex < AudioAdaptationSets.Num())
-	{
-		return AudioAdaptationSets[InAdaptationSetIndex];
-	}
-	else if (InStreamType == EStreamType::Subtitle && InAdaptationSetIndex >= 0 && InAdaptationSetIndex < SubtitleAdaptationSets.Num())
+	auto Pwy = GetCurrentPathway();
+	if (Pwy.IsValid())
 	{
-		return SubtitleAdaptationSets[InAdaptationSetIndex];
+		if (InStreamType == EStreamType::Video && InAdaptationSetIndex >= 0 && InAdaptationSetIndex < Pwy->VideoAdaptationSets.Num())
+		{
+			return Pwy->VideoAdaptationSets[InAdaptationSetIndex];
+		}
+		else if (InStreamType == EStreamType::Audio && InAdaptationSetIndex >= 0 && InAdaptationSetIndex < Pwy->AudioAdaptationSets.Num())
+		{
+			return Pwy->AudioAdaptationSets[InAdaptationSetIndex];
+		}
+		else if (InStreamType == EStreamType::Subtitle && InAdaptationSetIndex >= 0 && InAdaptationSetIndex < Pwy->SubtitleAdaptationSets.Num())
+		{
+			return Pwy->SubtitleAdaptationSets[InAdaptationSetIndex];
+		}
 	}
 	return nullptr;
 }
 
 void FActiveHLSPlaylist::FTimelineMediaAsset::GetMetaData(TArray<FTrackMetadata>& OutMetadata, EStreamType InStreamType) const
 {
-	const TArray<FInternalTrackMetadata>* Tracks = nullptr;
-	if (InStreamType == EStreamType::Video)
-	{
-		Tracks = &VideoTracks;
-	}
-	else if (InStreamType == EStreamType::Audio)
-	{
-		Tracks = &AudioTracks;
-	}
-	else if (InStreamType == EStreamType::Subtitle)
-	{
-		Tracks = &SubtitleTracks;
-	}
-	if (Tracks)
+	auto Pwy = GetCurrentPathway();
+	if (Pwy.IsValid())
 	{
-		for(auto& trkIt : *Tracks)
+		const TArray<FMultiVariantPlaylistHLS::FInternalTrackMetadata>* Tracks = nullptr;
+		if (InStreamType == EStreamType::Video)
+		{
+			Tracks = &Pwy->VideoTracks;
+		}
+		else if (InStreamType == EStreamType::Audio)
+		{
+			Tracks = &Pwy->AudioTracks;
+		}
+		else if (InStreamType == EStreamType::Subtitle)
+		{
+			Tracks = &Pwy->SubtitleTracks;
+		}
+		if (Tracks)
 		{
-			OutMetadata.Emplace(trkIt.Meta);
+			for(auto& trkIt : *Tracks)
+			{
+				OutMetadata.Emplace(trkIt.Meta);
+			}
 		}
 	}
 }
@@ -2201,24 +2311,28 @@ void FActiveHLSPlaylist::GetTrackMetadata(TArray<FTrackMetadata>& OutMetadata, E
 {
 	if (TimelineMediaAsset.IsValid())
 	{
-		const TArray<FInternalTrackMetadata>* Tracks = nullptr;
-		if (InStreamType == EStreamType::Video)
+		auto Pwy = TimelineMediaAsset->GetCurrentPathway();
+		if (Pwy.IsValid())
 		{
-			Tracks = &TimelineMediaAsset->VideoTracks;
-		}
-		else if (InStreamType == EStreamType::Audio)
-		{
-			Tracks = &TimelineMediaAsset->AudioTracks;
-		}
-		else if (InStreamType == EStreamType::Subtitle)
-		{
-			Tracks = &TimelineMediaAsset->SubtitleTracks;
-		}
-		if (Tracks)
-		{
-			for(auto& trkIt : *Tracks)
+			const TArray<FMultiVariantPlaylistHLS::FInternalTrackMetadata>* Tracks = nullptr;
+			if (InStreamType == EStreamType::Video)
 			{
-				OutMetadata.Emplace(trkIt.Meta);
+				Tracks = &Pwy->VideoTracks;
+			}
+			else if (InStreamType == EStreamType::Audio)
+			{
+				Tracks = &Pwy->AudioTracks;
+			}
+			else if (InStreamType == EStreamType::Subtitle)
+			{
+				Tracks = &Pwy->SubtitleTracks;
+			}
+			if (Tracks)
+			{
+				for(auto& trkIt : *Tracks)
+				{
+					OutMetadata.Emplace(trkIt.Meta);
+				}
 			}
 		}
 	}
@@ -2244,6 +2358,14 @@ void FActiveHLSPlaylist::TriggerPlaylistRefresh()
 	// Not used with HLS
 }
 
+void FActiveHLSPlaylist::ReachedStableBuffer()
+{
+	if (TimelineMediaAsset.IsValid())
+	{
+		PlayerSessionServices->GetContentSteeringHandler()->ReachedStableBuffer();
+	}
+}
+
 IStreamReader *FActiveHLSPlaylist::CreateStreamReaderHandler()
 {
 	// We reset the internal timeline used with Live playback at this point since we either
@@ -2344,30 +2466,28 @@ void FActiveHLSPlaylist::FPlayPeriod::Load()
 
 void FActiveHLSPlaylist::FPlayPeriod::PrepareForPlay()
 {
-	//int64 StartingBitrate = PlayerSessionServices->GetOptionValue(OptionKeyCurrentAvgStartingVideoBitrate).SafeGetInt64(2*1000*1000);
-
+	auto Pwy = TimelineMediaAsset->GetCurrentPathway();
 	TArray<TSharedPtr<FLoadRequestHLSPlaylist, ESPMode::ThreadSafe>> NewLoadReq;
 	int32 NumPending = 0;
-
 	// Select streams by preference, or the first one of the type if no preference is given.
 	for(int32 nStreamTypeIdx=0; nStreamTypeIdx<3; ++nStreamTypeIdx)
 	{
 		FSelectedTrackStream& st = SelectedTrackStream[nStreamTypeIdx];
-		const TArray<FActiveHLSPlaylist::FInternalTrackMetadata>* Tracks = nullptr;
+		const TArray<FMultiVariantPlaylistHLS::FInternalTrackMetadata>* Tracks = nullptr;
 		const EStreamType StreamType = StreamArrayIndexToType(nStreamTypeIdx);
 
-		if (nStreamTypeIdx == 0 && TimelineMediaAsset->VideoTracks.Num())
+		if (nStreamTypeIdx == 0 && Pwy->VideoTracks.Num())
 		{
-			Tracks = &TimelineMediaAsset->VideoTracks;
+			Tracks = &Pwy->VideoTracks;
 		}
-		else if (nStreamTypeIdx == 1 && TimelineMediaAsset->AudioTracks.Num())
+		else if (nStreamTypeIdx == 1 && Pwy->AudioTracks.Num())
 		{
-			Tracks = &TimelineMediaAsset->AudioTracks;
+			Tracks = &Pwy->AudioTracks;
 		}
-		else if (nStreamTypeIdx == 2 && TimelineMediaAsset->SubtitleTracks.Num())
+		else if (nStreamTypeIdx == 2 && Pwy->SubtitleTracks.Num())
 		{
 		// Not setting up subtitle tracks for now as we have no WebVTT parser.
-			//Tracks = &TimelineMediaAsset->SubtitleTracks;
+			//Tracks = &Pwy->SubtitleTracks;
 		}
 
 		int32 SelectedTrackIndex = 0;
@@ -2419,7 +2539,7 @@ void FActiveHLSPlaylist::FPlayPeriod::PrepareForPlay()
 		}
 		if (Tracks && SelectedTrackIndex < Tracks->Num() && (*Tracks)[SelectedTrackIndex].Meta.StreamDetails.Num())
 		{
-			const FActiveHLSPlaylist::FInternalTrackMetadata& Track((*Tracks)[SelectedTrackIndex]);
+			const FMultiVariantPlaylistHLS::FInternalTrackMetadata& Track((*Tracks)[SelectedTrackIndex]);
 
 			st.MetaID = Track.Meta.ID;
 			st.TrackIndex = SelectedTrackIndex;
@@ -2435,7 +2555,7 @@ void FActiveHLSPlaylist::FPlayPeriod::PrepareForPlay()
 			st.BufferSourceInfo->HardIndex = st.StreamIndex;
 
 			FStreamLoadRequest LoadReq;
-			if (TimelineMediaAsset->GetVariantPlaylist(LoadReq.Request, PlayerSessionServices, StreamType, TimelineMediaAsset->CurrentPathwayId, SelectedTrackIndex,0, SelectedTrackStream[0].TrackIndex,0).IsOK())
+			if (TimelineMediaAsset->GetVariantPlaylist(LoadReq.Request, PlayerSessionServices, StreamType, Pwy, SelectedTrackIndex,0, SelectedTrackStream[0].TrackIndex,0).IsOK())
 			{
 				st.ActivePlaylist = TimelineMediaAsset->GetExistingMediaPlaylistFromLoadRequest(LoadReq.Request);
 				if (!st.ActivePlaylist.IsValid())
@@ -2463,21 +2583,25 @@ void FActiveHLSPlaylist::FPlayPeriod::PrepareForPlay()
 
 int64 FActiveHLSPlaylist::FPlayPeriod::GetDefaultStartingBitrate() const
 {
-	const FActiveHLSPlaylist::FInternalTrackMetadata* Track = nullptr;
-	if (TimelineMediaAsset->VideoTracks.Num())
-	{
-		Track = &TimelineMediaAsset->VideoTracks[0];
-	}
-	else if (TimelineMediaAsset->AudioTracks.Num())
+	auto Pwy = TimelineMediaAsset->GetCurrentPathway();
+	if (Pwy.IsValid())
 	{
-		Track = &TimelineMediaAsset->AudioTracks[0];
-	}
-	if (Track)
-	{
-		check(Track->Meta.StreamDetails.Num());
-		if (Track->Meta.StreamDetails.Num())
+		const FMultiVariantPlaylistHLS::FInternalTrackMetadata* Track = nullptr;
+		if (Pwy->VideoTracks.Num())
+		{
+			Track = &Pwy->VideoTracks[0];
+		}
+		else if (Pwy->AudioTracks.Num())
+		{
+			Track = &Pwy->AudioTracks[0];
+		}
+		if (Track)
 		{
-			return Track->Meta.StreamDetails[0].Bandwidth;
+			check(Track->Meta.StreamDetails.Num());
+			if (Track->Meta.StreamDetails.Num())
+			{
+				return Track->Meta.StreamDetails[0].Bandwidth;
+			}
 		}
 	}
 	return -1;
@@ -2519,24 +2643,27 @@ IManifest::IPlayPeriod::ETrackChangeResult FActiveHLSPlaylist::FPlayPeriod::Chan
 	{
 		return IManifest::IPlayPeriod::ETrackChangeResult::StartOver;
 	}
-
 	// Create a temporary period and prepare it for playback. This may result in media playlist load requests!
 	FPlayPeriod TempPeriod(PlayerSessionServices, TimelineMediaAsset);
 	TempPeriod.SetStreamPreferences(InStreamType, InStreamAttributes);
 	TempPeriod.PrepareForPlay();
-	const TArray<FActiveHLSPlaylist::FInternalTrackMetadata>* Tracks = nullptr;
+	const TArray<FMultiVariantPlaylistHLS::FInternalTrackMetadata>* Tracks = nullptr;
 	const int32 stIdx = StreamTypeToArrayIndex(InStreamType);
-	if (stIdx == 0 && TimelineMediaAsset->VideoTracks.Num())
+	auto Pwy = TimelineMediaAsset->GetCurrentPathway();
+	if (Pwy.IsValid())
 	{
-		Tracks = &TimelineMediaAsset->VideoTracks;
-	}
-	else if (stIdx == 1 && TimelineMediaAsset->AudioTracks.Num())
-	{
-		Tracks = &TimelineMediaAsset->AudioTracks;
-	}
-	else if (stIdx == 2 && TimelineMediaAsset->SubtitleTracks.Num())
-	{
-		Tracks = &TimelineMediaAsset->SubtitleTracks;
+		if (stIdx == 0 && Pwy->VideoTracks.Num())
+		{
+			Tracks = &Pwy->VideoTracks;
+		}
+		else if (stIdx == 1 && Pwy->AudioTracks.Num())
+		{
+			Tracks = &Pwy->AudioTracks;
+		}
+		else if (stIdx == 2 && Pwy->SubtitleTracks.Num())
+		{
+			Tracks = &Pwy->SubtitleTracks;
+		}
 	}
 	// If either the stream we are leaving or the one we want to switch to is a variant we have to start over.
 	if (!Tracks || ((*Tracks)[SelectedTrackStream[stIdx].TrackIndex].bIsVariant || (*Tracks)[TempPeriod.SelectedTrackStream[stIdx].TrackIndex].bIsVariant))
@@ -2571,7 +2698,7 @@ void FActiveHLSPlaylist::FPlayPeriod::SelectStream(const FString& InAdaptationSe
 		LogMessage(PlayerSessionServices, IInfoLog::ELevel::Warning, FString::Printf(TEXT("ABR tried to activate a stream from an inactive AdaptationSet!")));
 		return;
 	}
-	TSharedPtrTS<FActiveHLSPlaylist::FPlaybackAssetAdaptationSet> Adapt = TimelineMediaAsset->GetAdaptationSet(InAdaptationSetID);
+	TSharedPtrTS<FMultiVariantPlaylistHLS::FPlaybackAssetAdaptationSet> Adapt = TimelineMediaAsset->GetAdaptationSet(InAdaptationSetID);
 	if (!Adapt.IsValid())
 	{
 		LogMessage(PlayerSessionServices, IInfoLog::ELevel::Error, FString::Printf(TEXT("ABR tried to activate a bad AdaptationSet!")));
@@ -2670,7 +2797,8 @@ void FActiveHLSPlaylist::FPlayPeriod::SetTimestampAdjustIfNecessary(const TShare
 
 IManifest::FResult FActiveHLSPlaylist::FPlayPeriod::GetStartingSegment(TSharedPtrTS<IStreamSegment>& OutSegment, const FPlayerSequenceState& InSequenceState, const FPlayStartPosition& InStartPosition, ESearchType InSearchType)
 {
-	if (!TimelineMediaAsset.IsValid())
+	TSharedPtrTS<FMultiVariantPlaylistHLS::FPathwayStreamInfs> Pwy = TimelineMediaAsset.IsValid() ? TimelineMediaAsset->GetCurrentPathway() : nullptr;
+	if (!Pwy.IsValid())
 	{
 		return IManifest::FResult(IManifest::FResult::EType::NotLoaded);
 	}
@@ -2686,7 +2814,7 @@ IManifest::FResult FActiveHLSPlaylist::FPlayPeriod::GetStartingSegment(TSharedPt
 		{
 			EStreamType StreamType = nStreamTypeIdx==0 ? EStreamType::Video:nStreamTypeIdx==1 ?EStreamType::Audio : EStreamType::Subtitle;
 			LoadReq.Emplace_GetRef().Type = StreamType;
-			FErrorDetail Error = TimelineMediaAsset->GetVariantPlaylist(LoadReq.Last().Request, PlayerSessionServices, StreamType, TimelineMediaAsset->CurrentPathwayId,
+			FErrorDetail Error = TimelineMediaAsset->GetVariantPlaylist(LoadReq.Last().Request, PlayerSessionServices, StreamType, Pwy,
 				SelectedTrackStream[nStreamTypeIdx].TrackIndex, SelectedTrackStream[nStreamTypeIdx].StreamIndex,
 				SelectedTrackStream[0].TrackIndex, SelectedTrackStream[0].StreamIndex);
 			if (Error.IsSet())
@@ -2896,9 +3024,10 @@ IManifest::FResult FActiveHLSPlaylist::FPlayPeriod::GetSegment(TSharedPtrTS<IStr
 		return IManifest::FResult(IManifest::FResult::EType::NotFound);
 	}
 
+	auto Pwy = TimelineMediaAsset->GetCurrentPathway();
 	FStreamLoadRequest LoadReq;
 	LoadReq.Type = InSegment->StreamType;
-	FErrorDetail Error = TimelineMediaAsset->GetVariantPlaylist(LoadReq.Request, PlayerSessionServices, InSegment->StreamType, TimelineMediaAsset->CurrentPathwayId,
+	FErrorDetail Error = TimelineMediaAsset->GetVariantPlaylist(LoadReq.Request, PlayerSessionServices, InSegment->StreamType, Pwy,
 		SelectedTrackStream[SelectedTrackTypeIndex].TrackIndex, SelectedTrackStream[SelectedTrackTypeIndex].StreamIndex,
 		SelectedTrackStream[0].TrackIndex, SelectedTrackStream[0].StreamIndex);
 	if (Error.IsSet())
@@ -3183,7 +3312,7 @@ void FActiveHLSPlaylist::FPlayPeriod::ValidateDownloadedSegmentDuration(const FS
 		if (AbsDelta > 1.0)
 		{
 			FURL_RFC3986 MediaURL, PlaylistURL;
-			MediaURL.Parse(InRequest->DownloadStats.URL);
+			MediaURL.Parse(InRequest->DownloadStats.URL.URL);
 			PlaylistURL.Parse(InRequest->HLS.Playlist.IsValid() ? InRequest->HLS.Playlist->URL : TEXT(""));
 			PlayerSessionServices->PostLog(Facility::EFacility::HLSPlaylistHandler, IInfoLog::ELevel::Warning, FString::Printf(TEXT("Media segment duration for segment %s in variant playlist %s was given as %.3f seconds but really was %.3f seconds!"), *MediaURL.GetLastPathComponent(), *PlaylistURL.GetLastPathComponent(), InRequest->DownloadStats.Duration, InRequest->DownloadStats.DurationDownloaded));
 		}
