Electra: Added multiple CDN support and ContentSteering for DASH and HLS

diff --git a/Engine/Plugins/Media/ElectraPlayer/Source/ElectraPlayerRuntime/Private/Runtime/Player/HLS/PlaylistHLS.h b/Engine/Plugins/Media/ElectraPlayer/Source/ElectraPlayerRuntime/Private/Runtime/Player/HLS/PlaylistHLS.h
--- a/Engine/Plugins/Media/ElectraPlayer/Source/ElectraPlayerRuntime/Private/Runtime/Player/HLS/PlaylistHLS.h
+++ b/Engine/Plugins/Media/ElectraPlayer/Source/ElectraPlayerRuntime/Private/Runtime/Player/HLS/PlaylistHLS.h
@@ -54,6 +54,63 @@
 
 struct FMultiVariantPlaylistHLS
 {
+	class FPlaybackAssetRepresentation : public IPlaybackAssetRepresentation
+	{
+	public:
+		virtual ~FPlaybackAssetRepresentation()
+		{ }
+		FString GetUniqueIdentifier() const override
+		{ return ID; }
+		const FStreamCodecInformation& GetCodecInformation() const override
+		{ return StreamCodecInformation; }
+		int32 GetBitrate() const override
+		{ return Bandwidth; }
+		int32 GetQualityIndex() const override
+		{ return QualityIndex; }
+		bool CanBePlayed() const override
+		{ return true; }
+		FStreamCodecInformation StreamCodecInformation;
+		FString ID;
+		int32 Bandwidth = 0;
+		int32 QualityIndex = 0;
+	};
+
+	class FPlaybackAssetAdaptationSet : public IPlaybackAssetAdaptationSet
+	{
+	public:
+		virtual ~FPlaybackAssetAdaptationSet()
+		{ }
+		FString GetUniqueIdentifier() const override
+		{ return ID; }
+		FString GetListOfCodecs() const override
+		{ return ListOfCodecs; }
+		FString GetLanguage() const override
+		{ return Language; }
+		int32 GetNumberOfRepresentations() const override
+		{ return Representations.Num(); }
+		bool IsLowLatencyEnabled() const override
+		{ return false; }
+		TSharedPtrTS<IPlaybackAssetRepresentation> GetRepresentationByIndex(int32 InRepresentationIndex) const override
+		{ return InRepresentationIndex >= 0 && InRepresentationIndex < Representations.Num() ? Representations[InRepresentationIndex] : nullptr; }
+		TSharedPtrTS<IPlaybackAssetRepresentation> GetRepresentationByUniqueIdentifier(const FString& InUniqueIdentifier) const override
+		{
+			for(int32 i=0; i<Representations.Num(); ++i)
+			{
+				if (Representations[i]->GetUniqueIdentifier().Equals(InUniqueIdentifier))
+				{
+					return Representations[i];
+				}
+			}
+			return nullptr;
+		}
+
+		TArray<TSharedPtrTS<FPlaybackAssetRepresentation>> Representations;
+		FString ID;
+		FString ListOfCodecs;
+		FString Language;
+	};
+
+
 	struct FRendition
 	{
 		// Sadly there is no explicit codec. We *TRY* to associate it with a codec from the EXT-X-STREAM-INF, if present.
@@ -119,8 +176,10 @@
 		int32 NumVideoCodec = 0;
 		int32 NumAudioCodec = 0;
 		int32 NumSubtitleCodec = 0;
+		// Generated ID
+		FString ID;
 		// Generated indices
-		int32 Index = 0;
+		int32 IndexOfSelfInArray = -1;
 		int32 QualityIndex = 0;
 		// Temp check
 		bool bReferencesAudioRenditionWithoutCodec = false;
@@ -141,20 +200,77 @@
 		TArray<int32> StreamInfIndices;
 	};
 
+	struct FInternalTrackMetadata
+	{
+		FTrackMetadata Meta;
+
+		/*
+			If set this is the rendition this has been sourced from.
+			If not set this is generated from a variant.
+		*/
+		TOptional<FRendition> Rendition;
+
+		/*
+			If `true` this indicates that this track is being sourced from a variant
+			and not a rendition.
+			If `false` the track may be sourced from a separate rendition if the URI in
+			the rendition is set, otherwise it is included in the variant.
+		*/
+		bool bIsVariant = false;
+
+		/*
+			Used only when there are audio-only variant streams AND video-only variant streams, in which
+			case this gives the index of the audio-only group we have associated with the video.
+			`bIsVariant` will be `true` when this is set.
+		*/
+		int32 AudioVariantGroupIndex = -1;
+
+		/*
+			Used only when there are multiple video groups, aka "angles".
+			Each "angle" is one video track, with the streams in the track metadata being
+			the different quality levels of that "angle".
+			This array contains the variant IDs that correspond to the respective quality level.
+		*/
+		TArray<FString> VideoVariantBaseIDs;
+	};
+
+	struct FPathwayStreamInfs
+	{
+		FString PathwayID;
+		TArray<FStreamInf> StreamInfs;
+		TArray<FVideoVariantGroup> VideoVariantGroups;
+		TArray<FAudioVariantGroup> AudioOnlyVariantGroups;
+		// Generated metadata
+		TArray<FInternalTrackMetadata> VideoTracks;
+		TArray<FInternalTrackMetadata> AudioTracks;
+		TArray<FInternalTrackMetadata> SubtitleTracks;
+		TArray<TSharedPtrTS<FPlaybackAssetAdaptationSet>> VideoAdaptationSets;
+		TArray<TSharedPtrTS<FPlaybackAssetAdaptationSet>> AudioAdaptationSets;
+		TArray<TSharedPtrTS<FPlaybackAssetAdaptationSet>> SubtitleAdaptationSets;
+	};
+
+	struct FContentSteeringParams
+	{
+		FString PrimaryPathwayId;
+		FString SteeringURI;
+		FString CustomInitialSelectionPriority;
+		bool bQueryBeforeStart = false;
+		bool bHaveContentSteering = false;
+	};
+
 	// Variable names are case sensitive so we use an array instead of a map.
 	TArray<FPlaylistParserHLS::FVariableSubstitution> VariableSubstitutions;
 	TArray<FRenditionGroup> RenditionGroupsOfType[4];
-	TArray<FStreamInf> StreamInfs;
-	TArray<FVideoVariantGroup> VideoVariantGroups;
-	TArray<FAudioVariantGroup> AudioOnlyVariantGroups;
+	TArray<TSharedPtrTS<FPathwayStreamInfs>> PathwayStreamInfs;
 	FString URL;
 	FURL_RFC3986 ParsedURL;
 
 	FServerControlHLS ServerControl;
 	FStartTimeHLS StartTime;
-	FString PrimaryPathwayId;
-	FString SteeringURI;
-	bool bHaveContentSteering = false;
+	FContentSteeringParams ContentSteeringParams;
+
+	// Initial bucket for all #EXT-X-STREAM-INF's before grouping into their respective PATHWAY-IDs
+	TArray<FStreamInf> InitialStreamInfs;
 };
 
 
@@ -352,6 +468,7 @@
 	{
 		Undefined,
 		Main,
+		Steering,
 		Variant,
 		InitialVariant
 	};
@@ -361,6 +478,7 @@
 	FTimeValue ExecuteAtUTC;
 	int32 Attempt = 1;
 	bool bIsPrimaryPlaylist = false;
+	bool bIsPreStartSteering = false;
 
 	FMediaPlaylistInformationHLS PlaylistInfo;
 
@@ -406,6 +524,7 @@
 	void GetNewMediaPlaylistLoadRequests(TArray<TSharedPtr<FLoadRequestHLSPlaylist, ESPMode::ThreadSafe>>& OutPlaylistLoadRequests);
 	void GetActiveMediaPlaylists(TArray<TSharedPtr<FMediaPlaylistAndStateHLS, ESPMode::ThreadSafe>>& OutActivePlaylists, const FTimeValue& InNow);
 	void GetAllMediaPlaylistLoadRequests(TArray<TSharedPtr<FLoadRequestHLSPlaylist, ESPMode::ThreadSafe>>& OutPlaylistLoadRequests, EStreamType InForType);
+	void CheckForPathwaySwitch();
 
 	EType GetPresentationType() const override;
 	TSharedPtrTS<const FLowLatencyDescriptor> GetLowLatencyDescriptor() const override;
@@ -429,109 +548,22 @@
 	void UpdateDynamicRefetchCounter() override;
 	void TriggerClockSync(EClockSyncType InClockSyncType) override;
 	void TriggerPlaylistRefresh() override;
+	void ReachedStableBuffer() override;
 	IStreamReader *CreateStreamReaderHandler() override;
 	FResult FindPlayPeriod(TSharedPtrTS<IPlayPeriod>& OutPlayPeriod, const FPlayStartPosition& StartPosition, ESearchType SearchType) override;
 	FResult FindNextPlayPeriod(TSharedPtrTS<IPlayPeriod>& OutPlayPeriod, TSharedPtrTS<const IStreamSegment> CurrentSegment) override;
 
+	static FErrorDetail DeterminePathwayToUse(IPlayerSessionServices* InPlayerSessionServices, FString& OutPathway, const FString& InCurrentPathway, const TSharedPtr<FMultiVariantPlaylistHLS, ESPMode::ThreadSafe>& InFromMultiVariantPlaylist);
+	FErrorDetail PreparePathway(TSharedPtrTS<FMultiVariantPlaylistHLS::FPathwayStreamInfs>& InOutPathway, TSharedPtr<FMultiVariantPlaylistHLS, ESPMode::ThreadSafe> InFromMultiVariantPlaylist);
 private:
 	ELECTRA_IMPL_DEFAULT_ERROR_METHODS(HLSPlaylistBuilder);
 
 	struct FInternalBuilder;
 	class FPlayPeriod;
 
-	FErrorDetail CreateTrackMetadata(IPlayerSessionServices* InPlayerSessionServices, const FInternalBuilder* InBuilder);
+	FErrorDetail CreateTrackMetadata(IPlayerSessionServices* InPlayerSessionServices, TSharedPtrTS<FMultiVariantPlaylistHLS::FPathwayStreamInfs>& InPathway, TSharedPtr<FMultiVariantPlaylistHLS, ESPMode::ThreadSafe> InFromMultiVariantPlaylist, const FInternalBuilder* InBuilder);
 	FErrorDetail GetInitialVariantPlaylistLoadRequests(TArray<TSharedPtr<FLoadRequestHLSPlaylist, ESPMode::ThreadSafe>>& OutPlaylistLoadRequests, IPlayerSessionServices* InPlayerSessionServices);
 
-	struct FInternalTrackMetadata
-	{
-		FTrackMetadata Meta;
-
-		/*
-			If set this is the rendition this has been sourced from.
-			If not set this is generated from a variant.
-		*/
-		TOptional<FMultiVariantPlaylistHLS::FRendition> Rendition;
-
-		/*
-			If `true` this indicates that this track is being sourced from a variant
-			and not a rendition.
-			If `false` the track may be sourced from a separate rendition if the URI in
-			the rendition is set, otherwise it is included in the variant.
-		*/
-		bool bIsVariant = false;
-
-		/*
-			Used only when there are audio-only variant streams AND video-only variant streams, in which
-			case this gives the index of the audio-only group we have associated with the video.
-			`bIsVariant` will be `true` when this is set.
-		*/
-		int32 AudioVariantGroupIndex = -1;
-
-		/*
-			Used only when there are multiple video groups, aka "angles".
-			Each "angle" is one video track, with the streams in the track metadata being
-			the different quality levels of that "angle".
-			This array contains the variant IDs that correspond to the respective quality level.
-		*/
-		TArray<FString> VideoVariantBaseIDs;
-	};
-
-	class FPlaybackAssetRepresentation : public IPlaybackAssetRepresentation
-	{
-	public:
-		virtual ~FPlaybackAssetRepresentation()
-		{ }
-		FString GetUniqueIdentifier() const override
-		{ return ID; }
-		const FStreamCodecInformation& GetCodecInformation() const override
-		{ return StreamCodecInformation; }
-		int32 GetBitrate() const override
-		{ return Bandwidth; }
-		int32 GetQualityIndex() const override
-		{ return QualityIndex; }
-		bool CanBePlayed() const override
-		{ return true; }
-		FStreamCodecInformation StreamCodecInformation;
-		FString ID;
-		int32 Bandwidth = 0;
-		int32 QualityIndex = 0;
-	};
-
-	class FPlaybackAssetAdaptationSet : public IPlaybackAssetAdaptationSet
-	{
-	public:
-		virtual ~FPlaybackAssetAdaptationSet()
-		{ }
-		FString GetUniqueIdentifier() const override
-		{ return ID; }
-		FString GetListOfCodecs() const override
-		{ return ListOfCodecs; }
-		FString GetLanguage() const override
-		{ return Language; }
-		int32 GetNumberOfRepresentations() const override
-		{ return Representations.Num(); }
-		bool IsLowLatencyEnabled() const override
-		{ return false; }
-		TSharedPtrTS<IPlaybackAssetRepresentation> GetRepresentationByIndex(int32 InRepresentationIndex) const override
-		{ return InRepresentationIndex >= 0 && InRepresentationIndex < Representations.Num() ? Representations[InRepresentationIndex] : nullptr; }
-		TSharedPtrTS<IPlaybackAssetRepresentation> GetRepresentationByUniqueIdentifier(const FString& InUniqueIdentifier) const override
-		{
-			for(int32 i=0; i<Representations.Num(); ++i)
-			{
-				if (Representations[i]->GetUniqueIdentifier().Equals(InUniqueIdentifier))
-				{
-					return Representations[i];
-				}
-			}
-			return nullptr;
-		}
-
-		TArray<TSharedPtrTS<FPlaybackAssetRepresentation>> Representations;
-		FString ID;
-		FString ListOfCodecs;
-		FString Language;
-	};
-
 	class FTimelineMediaAsset : public ITimelineMediaAsset
 	{
 	public:
@@ -545,7 +577,7 @@
 		void GetMetaData(TArray<FTrackMetadata>& OutMetadata, EStreamType InStreamType) const override;
 		void UpdateRunningMetaData(const FString& InKindOfValue, const FVariant& InNewValue) override;
 
-		FErrorDetail GetVariantPlaylist(TSharedPtr<FLoadRequestHLSPlaylist, ESPMode::ThreadSafe>& OutPlaylistLoadRequest, IPlayerSessionServices* InPlayerSessionServices, EStreamType InStreamType, const FString& InPathway, int32 InTrackIndex, int32 InStreamIndex, int32 InMainTrackIndex, int32 InMainStreamIndex);
+		FErrorDetail GetVariantPlaylist(TSharedPtr<FLoadRequestHLSPlaylist, ESPMode::ThreadSafe>& OutPlaylistLoadRequest, IPlayerSessionServices* InPlayerSessionServices, EStreamType InStreamType, const TSharedPtrTS<FMultiVariantPlaylistHLS::FPathwayStreamInfs>& InPathway, int32 InTrackIndex, int32 InStreamIndex, int32 InMainTrackIndex, int32 InMainStreamIndex);
 		void UpdateWithMediaPlaylist(TSharedPtr<FMediaPlaylistAndStateHLS, ESPMode::ThreadSafe> InMediaPlaylist, bool bIsPrimary, bool bIsUpdate);
 		void GetNewMediaPlaylistLoadRequests(TArray<TSharedPtr<FLoadRequestHLSPlaylist, ESPMode::ThreadSafe>>& OutPlaylistLoadRequests);
 		void AddNewMediaPlaylistLoadRequests(const TArray<TSharedPtr<FLoadRequestHLSPlaylist, ESPMode::ThreadSafe>>& InNewPlaylistLoadRequests);
@@ -561,6 +593,14 @@
 		void UpdateTimelineFromMediaSegment(TSharedPtrTS<FStreamSegmentRequestCommon> InSegment);
 		void ResetInternalTimeline();
 
+		TSharedPtrTS<FMultiVariantPlaylistHLS::FPathwayStreamInfs> GetCurrentPathway() const
+		{
+			return CurrentPathway;
+		}
+		void SetCurrentPathway(const TSharedPtrTS<FMultiVariantPlaylistHLS::FPathwayStreamInfs>& InNewPathway)
+		{
+			CurrentPathway = InNewPathway;
+		}
 
 		struct FSegSearchParam
 		{
@@ -589,8 +629,8 @@
 		FErrorDetail GetLastError() const
 		{ return LastError; }
 
-		const FInternalTrackMetadata* GetInternalTrackMetadata(const FString& InForID) const;
-		TSharedPtrTS<FPlaybackAssetAdaptationSet> GetAdaptationSet(const FString& InForID) const;
+		const FMultiVariantPlaylistHLS::FInternalTrackMetadata* GetInternalTrackMetadata(const FString& InForID) const;
+		TSharedPtrTS<FMultiVariantPlaylistHLS::FPlaybackAssetAdaptationSet> GetAdaptationSet(const FString& InForID) const;
 
 		// Returns `true` if the very first media playlist loaded already had an EXT-X-ENDLIST
 		// If this is an EVENT or a Live presentation and it did not, we can tell that if it
@@ -603,15 +643,6 @@
 
 		TSharedPtr<FMultiVariantPlaylistHLS, ESPMode::ThreadSafe> MultiVariantPlaylist;
 
-		// Generated metadata
-		TArray<FInternalTrackMetadata> VideoTracks;
-		TArray<FInternalTrackMetadata> AudioTracks;
-		TArray<FInternalTrackMetadata> SubtitleTracks;
-
-		TArray<TSharedPtrTS<FPlaybackAssetAdaptationSet>> VideoAdaptationSets;
-		TArray<TSharedPtrTS<FPlaybackAssetAdaptationSet>> AudioAdaptationSets;
-		TArray<TSharedPtrTS<FPlaybackAssetAdaptationSet>> SubtitleAdaptationSets;
-
 		// Metadata from the primary media playlist
 		FServerControlHLS ServerControl;
 		FTimeValue FirstProgramDateTime;
@@ -662,13 +693,14 @@
 		mutable FCriticalSection InternalMediaTimelineLock;
 
 
-		//
-		FString CurrentPathwayId;
 		FCriticalSection MediaPlaylistsLock;
 		TArray<TSharedPtr<FMediaPlaylistAndStateHLS, ESPMode::ThreadSafe>> MediaPlaylists;
 		TArray<TSharedPtr<FLoadRequestHLSPlaylist, ESPMode::ThreadSafe>> NewMediaPlaylistLoadRequests;
 		FDRMClientCacheHLS LicenseKeyCache;
 		FErrorDetail LastError;
+		// The currently active pathway
+		TSharedPtrTS<FMultiVariantPlaylistHLS::FPathwayStreamInfs> CurrentPathway;
+		FString CurrentPathwayId;
 	};
 
 	FCriticalSection RequestedPeriodsLock;
diff --git a/Engine/Plugins/Media/ElectraPlayer/Source/ElectraPlayerRuntime/Private/Runtime/Player/HLS/PlaylistHandlerHLS.cpp b/Engine/Plugins/Media/ElectraPlayer/Source/ElectraPlayerRuntime/Private/Runtime/Player/HLS/PlaylistHandlerHLS.cpp
--- a/Engine/Plugins/Media/ElectraPlayer/Source/ElectraPlayerRuntime/Private/Runtime/Player/HLS/PlaylistHandlerHLS.cpp
+++ b/Engine/Plugins/Media/ElectraPlayer/Source/ElectraPlayerRuntime/Private/Runtime/Player/HLS/PlaylistHandlerHLS.cpp
@@ -5,12 +5,14 @@
 #include "PlaylistHLS.h"
 #include "Player/AdaptiveStreamingPlayerABR.h"
 #include "Player/AdaptiveStreamingPlayerResourceRequest.h"
+#include "Player/ContentSteeringHandler.h"
 #include "Utilities/Utilities.h"
 #include "Utilities/URLParser.h"
 #include "Utilities/TimeUtilities.h"
 #include "ElectraPlayerMisc.h"
 #include "StreamTypes.h"
 #include "Misc/SecureHash.h"
+#include "Internationalization/Regex.h"
 
 namespace Electra
 {
@@ -35,6 +37,12 @@ public:
 private:
 	ELECTRA_CLASS_DEFAULT_ERROR_METHODS(PlayerSessionServices, HLSPlaylistHandler);
 
+	static const FRegexPattern VariableSubstitutionPattern()
+	{
+		static const FRegexPattern vsp(TEXT("(\\{\\$.+?\\})"));
+		return vsp;
+	}
+
 	static bool ValidateNumbersOnly(const FString& In)
 	{
 		for(StringHelpers::FStringIterator It(In); It; ++It)
@@ -62,6 +70,8 @@ private:
 	TSharedPtr<FLoadRequestHLSPlaylist, ESPMode::ThreadSafe> CreateMediaPlaylistLoadRequest(const TSharedPtr<FMediaPlaylistAndStateHLS, ESPMode::ThreadSafe>& InPlaylist);
 	void RepeatPlaylistLoad(TSharedPtr<FLoadRequestHLSPlaylist, ESPMode::ThreadSafe> InRequestToRepeat, const FTimeValue& InAtUTC);
 	void RetryInitialPlaylistWith(TSharedPtr<FLoadRequestHLSPlaylist, ESPMode::ThreadSafe> InRequestToRepeat, bool bIsPrimary);
+	void SetupActivePlaylist();
+	void PerformContentSteeringCloning(bool bIsPreStartFetch);
 
 	FTimeValue GetTimeoutValue(const FName& InOptionName, int32 InDefaultValueMillisec);
 	TSharedPtr<FPlaylistParserHLS, ESPMode::ThreadSafe> HandleMainPlaylist(const FTimeValue& Now, TSharedPtr<FLoadRequestHLSPlaylist, ESPMode::ThreadSafe> InRequest);
@@ -73,6 +83,8 @@ private:
 	bool BuildPlaylist(TSharedPtr<FPlaylistParserHLS, ESPMode::ThreadSafe> InPlaylist, bool bIsMainPlaylist);
 	bool BuildMultiVariantPlaylist(TSharedPtr<FMultiVariantPlaylistHLS, ESPMode::ThreadSafe>& OutMultiVariantPlaylist, TSharedPtr<FPlaylistParserHLS, ESPMode::ThreadSafe> InPlaylist);
 	bool BuildMediaPlaylist(TSharedPtr<FMediaPlaylistHLS, ESPMode::ThreadSafe>& OutMediaPlaylist, TSharedPtr<FPlaylistParserHLS, ESPMode::ThreadSafe> InPlaylist, TSharedPtr<FMultiVariantPlaylistHLS, ESPMode::ThreadSafe> InMultiVariantPlaylist);
+	void GroupVariantStreamsByPathways(TSharedPtr<FMultiVariantPlaylistHLS, ESPMode::ThreadSafe> InMultiVariantPlaylist);
+	void AssignInternalVariantStreamIDs(TSharedPtr<FMultiVariantPlaylistHLS, ESPMode::ThreadSafe> InMultiVariantPlaylist);
 	void GroupVariantStreamsByVideoProperties(TSharedPtr<FMultiVariantPlaylistHLS, ESPMode::ThreadSafe> InMultiVariantPlaylist);
 	void GroupAudioOnlyVariantStreams(TSharedPtr<FMultiVariantPlaylistHLS, ESPMode::ThreadSafe> InMultiVariantPlaylist);
 	enum EFillInOptions
@@ -92,7 +104,6 @@ private:
 	void CheckForMissingResolution(TSharedPtr<FMultiVariantPlaylistHLS, ESPMode::ThreadSafe> InMultiVariantPlaylist);
 	void AssignResolutionAndFrameRateToCodecs(TSharedPtr<FMultiVariantPlaylistHLS, ESPMode::ThreadSafe> InMultiVariantPlaylist);
 	void AssignCodecsToRenditions(TSharedPtr<FMultiVariantPlaylistHLS, ESPMode::ThreadSafe> InMultiVariantPlaylist);
-	void AssignStableIds(TSharedPtr<FMultiVariantPlaylistHLS, ESPMode::ThreadSafe> InMultiVariantPlaylist);
 
 	IPlayerSessionServices* PlayerSessionServices = nullptr;
 	FString MultiVariantPlaylistEffectiveURL;
@@ -152,9 +163,31 @@ void FPlaylistHandlerHLS::Close()
 	}
 }
 
+void FPlaylistHandlerHLS::SetupActivePlaylist()
+{
+	TArray<TSharedPtr<FLoadRequestHLSPlaylist, ESPMode::ThreadSafe>> PlaylistLoadRequests;
+	ActivePlaylist = MakeShared<FActiveHLSPlaylist, ESPMode::ThreadSafe>();
+	LastError = ActivePlaylist->Create(PlaylistLoadRequests, PlayerSessionServices, MultiVariantPlaylist);
+	if (LastError.IsOK())
+	{
+		for(int32 nr=0; nr<PlaylistLoadRequests.Num(); ++nr)
+		{
+			PlaylistLoadRequests[nr]->LoadType = FLoadRequestHLSPlaylist::ELoadType::InitialVariant;
+			PendingLoadRequests.Emplace(MoveTemp(PlaylistLoadRequests[nr]));
+			++NumPendingInitialVariantRequest;
+		}
+	}
+	else
+	{
+		PostError(LastError);
+		ActivePlaylist.Reset();
+	}
+}
+
 void FPlaylistHandlerHLS::HandleOnce()
 {
 	FTimeValue Now = PlayerSessionServices->GetSynchronizedUTCTime()->GetTime();
+	bool bIsPreStartSteering = false;
 
 	// Get any new media playlist load requests.
 	if (ActivePlaylist.IsValid())
@@ -301,27 +334,86 @@ void FPlaylistHandlerHLS::HandleOnce()
 
 					if (BuildPlaylist(MainPlaylist, true))
 					{
-						TArray<TSharedPtr<FLoadRequestHLSPlaylist, ESPMode::ThreadSafe>> PlaylistLoadRequests;
-						ActivePlaylist = MakeShared<FActiveHLSPlaylist, ESPMode::ThreadSafe>();
-						LastError = ActivePlaylist->Create(PlaylistLoadRequests, PlayerSessionServices, MultiVariantPlaylist);
-						if (LastError.IsOK())
+						// Configure content steering now. This allows us to obtain an initial steering manifest before deciding
+						// which CDN to use.
+						FContentSteeringHandler::FInitialParams csp;
+						if (MultiVariantPlaylist->ContentSteeringParams.bHaveContentSteering)
 						{
-							for(int32 nr=0; nr<PlaylistLoadRequests.Num(); ++nr)
+							csp.RootDocumentURL = MultiVariantPlaylist->URL;
+							csp.FirstSteeringURL = MultiVariantPlaylist->ContentSteeringParams.SteeringURI;
+							csp.InitialDefaultCDN = MultiVariantPlaylist->ContentSteeringParams.PrimaryPathwayId;
+							csp.bQueryBeforeStart = MultiVariantPlaylist->ContentSteeringParams.bQueryBeforeStart;
+							csp.CustomFirstCDNPrioritization = MultiVariantPlaylist->ContentSteeringParams.CustomInitialSelectionPriority;
+							// Without a URL we can use this only to select the PATHWAY to start with.
+							csp.bHasContentSteering = !csp.FirstSteeringURL.IsEmpty();
+						}
+						// If we do not have an initial CDN specified in the multi variant playlist we still
+						// need to set up a list of pathways that have been specified wither through explicit
+						// PATHWAY-ID attributes on the #EXT-X-STREAM-INF or through generated ones when
+						// fallback variants have been detected.
+						if (csp.InitialDefaultCDN.IsEmpty())
+						{
+							csp.bAllowAnyPathwayNames = true;
+							for(int32 pwIdx=0; pwIdx<MultiVariantPlaylist->PathwayStreamInfs.Num(); ++pwIdx)
 							{
-								PlaylistLoadRequests[nr]->LoadType = FLoadRequestHLSPlaylist::ELoadType::InitialVariant;
-								PendingLoadRequests.Emplace(MoveTemp(PlaylistLoadRequests[nr]));
-								++NumPendingInitialVariantRequest;
+								if (pwIdx)
+								{
+									csp.InitialDefaultCDN.Append(TEXT(" "));
+								}
+								csp.InitialDefaultCDN.Append(MultiVariantPlaylist->PathwayStreamInfs[pwIdx]->PathwayID);
 							}
 						}
+						bool bNeedFirstSteering = PlayerSessionServices->GetContentSteeringHandler()->InitialSetup(FContentSteeringHandler::EStreamingProtocol::HLS, csp) &&
+												  PlayerSessionServices->GetContentSteeringHandler()->NeedToObtainNewSteeringManifestNow();
+						if (!bNeedFirstSteering)
+						{
+							SetupActivePlaylist();
+						}
 						else
 						{
-							PostError(LastError);
-							ActivePlaylist.Reset();
+							// Because there is no guarantee that the initial PATHWAY-ID (if it is given) is actually one that exists
+							// we perform a first pathway selection right now.
+							FString CurrentPathway = MultiVariantPlaylist->PathwayStreamInfs.Num() ? MultiVariantPlaylist->PathwayStreamInfs[0]->PathwayID : FString(TEXT("."));
+							FString NewPathwayId;
+							FActiveHLSPlaylist::DeterminePathwayToUse(PlayerSessionServices, NewPathwayId, CurrentPathway, MultiVariantPlaylist);
+							PlayerSessionServices->GetContentSteeringHandler()->SetCurrentlyActivePathway(NewPathwayId);
+							bIsPreStartSteering = true;
 						}
 					}
 				}
 				break;
 			}
+			case FLoadRequestHLSPlaylist::ELoadType::Steering:
+			{
+				if (MultiVariantPlaylist.IsValid())
+				{
+					// Update content steering. It does not matter if this was successful or not.
+					HTTP::FConnectionInfo ci;
+					if (const HTTP::FConnectionInfo* pci = FinishedLoadRequests[i]->ResourceRequest->GetConnectionInfo())
+					{
+						ci = *pci;
+					}
+					FString SteeringJSON;
+					TSharedPtrTS<FWaitableBuffer> ResponseBuffer = FinishedLoadRequests[i]->ResourceRequest->GetResponseBuffer();
+					if (ResponseBuffer.IsValid() && !StringHelpers::ArrayToString(SteeringJSON, MakeConstArrayView<const uint8>(ResponseBuffer->GetLinearReadData(), ResponseBuffer->GetLinearReadSize())))
+					{
+						SteeringJSON.Empty();
+					}
+					PlayerSessionServices->GetContentSteeringHandler()->UpdateWithSteeringServerResponse(SteeringJSON, ci.StatusInfo.HTTPStatus, ci.ResponseHeaders);
+					PerformContentSteeringCloning(FinishedLoadRequests[i]->bIsPreStartSteering);
+
+					// If this was the initial steering request we now need to continue with the selection of the first playlists.
+					if (FinishedLoadRequests[i]->bIsPreStartSteering)
+					{
+						SetupActivePlaylist();
+					}
+					else if (ActivePlaylist.IsValid())
+					{
+						ActivePlaylist->CheckForPathwaySwitch();
+					}
+				}
+				break;
+			}
 			case FLoadRequestHLSPlaylist::ELoadType::InitialVariant:
 			{
 				--NumPendingInitialVariantRequest;
@@ -424,8 +516,189 @@ void FPlaylistHandlerHLS::HandleOnce()
 			--i;
 		}
 	}
+
+	// Check if a new steering manifest is needed
+	if (PlayerSessionServices->GetContentSteeringHandler()->NeedToObtainNewSteeringManifestNow())
+	{
+		FString SteeringURL = PlayerSessionServices->GetContentSteeringHandler()->GetFinalSteeringServerRequestURL(PlayerSessionServices->GetContentSteeringHandler()->GetBaseSteeringServerRequestURL());
+		PlayerSessionServices->GetContentSteeringHandler()->SetSteeringServerRequestIsPending();
+		TSharedPtr<FLoadRequestHLSPlaylist, ESPMode::ThreadSafe> lr = MakeShared<FLoadRequestHLSPlaylist, ESPMode::ThreadSafe>();
+		lr->LoadType = FLoadRequestHLSPlaylist::ELoadType::Steering;
+		lr->bIsPreStartSteering = bIsPreStartSteering;
+		lr->ResourceRequest = MakeShared<FHTTPResourceRequest, ESPMode::ThreadSafe>();
+		lr->ResourceRequest->Verb(TEXT("GET")).URL(SteeringURL)
+			.ConnectionTimeout(GetTimeoutValue(HLS::OptionKeyPlaylistLoadConnectTimeout, 5000))
+			.NoDataTimeout(GetTimeoutValue(HLS::OptionKeyPlaylistLoadNoDataTimeout, 2000));
+			//.AllowStaticQuery(IAdaptiveStreamingPlayerResourceRequest::EPlaybackResourceType::Playlist);
+		PendingLoadRequests.Emplace(MoveTemp(lr));
+	}
 }
 
+void FPlaylistHandlerHLS::PerformContentSteeringCloning(bool bIsPreStartFetch)
+{
+	// See if there are clones that need to be created.
+	TArray<FContentSteeringHandler::FPathwayCloneEntry> NewClones(PlayerSessionServices->GetContentSteeringHandler()->GetCurrentCloneEntries());
+	if (NewClones.Num())
+	{
+		auto UpdateURL = [&](const FString& InURL, const FContentSteeringHandler::FPathwayCloneEntry& InCloneParam) -> FString
+		{
+			FURL_RFC3986 updatedURL;
+			updatedURL.Parse(InURL);
+			if (InCloneParam.Host.Len())
+			{
+				updatedURL.SetHost(InCloneParam.Host);
+			}
+			updatedURL.AddOrUpdateQueryParams(InCloneParam.Params);
+			FString NewURL = updatedURL.Get();
+
+			// Perform variable substitution on the new URL, just like it would be done
+			// on a URL that was originally appeared in the playlist.
+			// It does not state in the RFC that this should be done though.
+			FRegexMatcher Matcher(VariableSubstitutionPattern(), NewURL);
+			TArray<FString> Substitutions;
+			while(Matcher.FindNext())
+			{
+				Substitutions.Emplace(Matcher.GetCaptureGroup(1));
+			}
+			for(auto& subs : Substitutions)
+			{
+				for(auto& rpl : MultiVariantPlaylist->VariableSubstitutions)
+				{
+					if (subs.Equals(rpl.Name))
+					{
+						NewURL.ReplaceInline(*subs, *rpl.Value, ESearchCase::CaseSensitive);
+					}
+				}
+			}
+			return NewURL;
+		};
+
+		for(int32 cloneIdx=0; cloneIdx<NewClones.Num(); ++cloneIdx)
+		{
+			for(int32 pwIdx=0; pwIdx<MultiVariantPlaylist->PathwayStreamInfs.Num(); ++pwIdx)
+			{
+				auto src = MultiVariantPlaylist->PathwayStreamInfs[pwIdx];
+				if (src->PathwayID.Equals(NewClones[cloneIdx].BaseId))
+				{
+					TSharedPtrTS<FMultiVariantPlaylistHLS::FPathwayStreamInfs> cln = MakeSharedTS<FMultiVariantPlaylistHLS::FPathwayStreamInfs>();
+					cln->PathwayID = NewClones[cloneIdx].Id;
+					cln->StreamInfs = src->StreamInfs;
+					cln->VideoVariantGroups = src->VideoVariantGroups;
+					cln->AudioOnlyVariantGroups = src->AudioOnlyVariantGroups;
+
+					TArray<FString> GroupsToClone[4];
+					FString CloneSuffix(TEXT("@clone"));
+					// Replace the pathway on the cloned stream-inf and collect groups that are referenced.
+					for(int32 nSInf=0; nSInf<cln->StreamInfs.Num(); ++nSInf)
+					{
+						cln->StreamInfs[nSInf].PathwayId = cln->PathwayID;
+						if (cln->StreamInfs[nSInf].VideoGroup.Len())
+						{
+							GroupsToClone[0].AddUnique(cln->StreamInfs[nSInf].VideoGroup);
+							cln->StreamInfs[nSInf].VideoGroup.Append(CloneSuffix);
+						}
+						if (cln->StreamInfs[nSInf].AudioGroup.Len())
+						{
+							GroupsToClone[1].AddUnique(cln->StreamInfs[nSInf].AudioGroup);
+							cln->StreamInfs[nSInf].AudioGroup.Append(CloneSuffix);
+						}
+						if (cln->StreamInfs[nSInf].SubtitleGroup.Len())
+						{
+							GroupsToClone[2].AddUnique(cln->StreamInfs[nSInf].SubtitleGroup);
+							cln->StreamInfs[nSInf].SubtitleGroup.Append(CloneSuffix);
+						}
+						if (cln->StreamInfs[nSInf].ClosedCaptionGroup.Len())
+						{
+							GroupsToClone[3].AddUnique(cln->StreamInfs[nSInf].ClosedCaptionGroup);
+							cln->StreamInfs[nSInf].ClosedCaptionGroup.Append(CloneSuffix);
+						}
+					}
+					// Are there HOST and/or PARAMS in the clone description?
+					if (NewClones[cloneIdx].Host.Len() || NewClones[cloneIdx].Params.Num())
+					{
+						for(int32 nSInf=0; nSInf<cln->StreamInfs.Num(); ++nSInf)
+						{
+							cln->StreamInfs[nSInf].URI = UpdateURL(cln->StreamInfs[nSInf].URI, NewClones[cloneIdx]);
+						}
+					}
+					// Check if there are PER-VARIANT-URIS in the clone description
+					for(int32 nSInf=0; nSInf<cln->StreamInfs.Num(); ++nSInf)
+					{
+						if (cln->StreamInfs[nSInf].StableVariantId.Len())
+						{
+							for(auto& variantURI : NewClones[cloneIdx].PerVariantURIs)
+							{
+								if (variantURI.Key.Equals(cln->StreamInfs[nSInf].StableVariantId))
+								{
+									cln->StreamInfs[nSInf].URI = variantURI.Value;
+								}
+							}
+						}
+					}
+
+					// Need to clone groups?
+					for(int32 nCloneGrpIdx=0; nCloneGrpIdx<UE_ARRAY_COUNT(GroupsToClone); ++nCloneGrpIdx)
+					{
+						for(auto& OrgGrpName : GroupsToClone[nCloneGrpIdx])
+						{
+							// Locate the group requiring cloning
+							for(int32 nOrgGrpIdx=0; nOrgGrpIdx<MultiVariantPlaylist->RenditionGroupsOfType[nCloneGrpIdx].Num(); ++nOrgGrpIdx)
+							{
+								if (MultiVariantPlaylist->RenditionGroupsOfType[nCloneGrpIdx][nOrgGrpIdx].GroupID.Equals(OrgGrpName))
+								{
+									// Clone the group
+									FMultiVariantPlaylistHLS::FRenditionGroup ClonedGroup(MultiVariantPlaylist->RenditionGroupsOfType[nCloneGrpIdx][nOrgGrpIdx]);
+									ClonedGroup.GroupID.Append(CloneSuffix);
+									// Go over the renditions and update their URLs
+									for(auto& ClonedRendition : ClonedGroup.Renditions)
+									{
+										if (ClonedRendition.URI.Len())
+										{
+											ClonedRendition.URI = UpdateURL(ClonedRendition.URI, NewClones[cloneIdx]);
+										}
+										// Check if the rendition has a stable rendition id for which there is a new dedicated URL in the clone parameters
+										if (ClonedRendition.StableRenditionId.Len())
+										{
+											for(auto& renditionURI : NewClones[cloneIdx].PerRenditionURIs)
+											{
+												if (renditionURI.Key.Equals(ClonedRendition.StableRenditionId))
+												{
+													ClonedRendition.URI = renditionURI.Value;
+												}
+											}
+										}
+									}
+									// Add the cloned group to the list.
+									MultiVariantPlaylist->RenditionGroupsOfType[nCloneGrpIdx].Emplace(MoveTemp(ClonedGroup));
+									break;
+								}
+							}
+						}
+					}
+
+					// The metadata is more involved and cloning with patching is too much hassle.
+					// If this is the very first application of a clone on pre-start steering then
+					// we do not need to do anything since the metadata will be created in the next step.
+					// Otherwise we create the metadata for the clone now.
+					if (!bIsPreStartFetch && ActivePlaylist.IsValid())
+					{
+						FErrorDetail pwErr = ActivePlaylist->PreparePathway(cln, MultiVariantPlaylist);
+						check(pwErr.IsOK());(void)pwErr;
+					}
+
+					// Take note that we have created the clone, so we do not do this repeatedly.
+					PlayerSessionServices->GetContentSteeringHandler()->CreatedClone(cln->PathwayID);
+
+					// Add the clone to the pathway list
+					MultiVariantPlaylist->PathwayStreamInfs.Emplace(MoveTemp(cln));
+					break;
+				}
+			}
+		}
+	}
+}
+
+
 void FPlaylistHandlerHLS::LoadAndParse(const FString& URL)
 {
 	FURL_RFC3986 UrlParser;
@@ -748,6 +1021,8 @@ bool FPlaylistHandlerHLS::BuildPlaylist(TSharedPtr<FPlaylistParserHLS, ESPMode::
 		if (BuildMultiVariantPlaylist(MultiVariantPlaylist, InPlaylist))
 		{
 			FillInMissingInformation(MultiVariantPlaylist, EFillInOptions::All);
+			GroupVariantStreamsByPathways(MultiVariantPlaylist);
+			AssignInternalVariantStreamIDs(MultiVariantPlaylist);
 			GroupVariantStreamsByVideoProperties(MultiVariantPlaylist);
 			GroupAudioOnlyVariantStreams(MultiVariantPlaylist);
 			return true;
@@ -762,13 +1037,14 @@ bool FPlaylistHandlerHLS::BuildPlaylist(TSharedPtr<FPlaylistParserHLS, ESPMode::
 
 		TSharedPtr<FMultiVariantPlaylistHLS, ESPMode::ThreadSafe> NewVariantPlaylist = MakeShared<FMultiVariantPlaylistHLS, ESPMode::ThreadSafe>();
 		NewVariantPlaylist->URL = InPlaylist->GetURL();
-		FMultiVariantPlaylistHLS::FStreamInf& si = NewVariantPlaylist->StreamInfs.Emplace_GetRef();
-		si.Index = 0;
+		FMultiVariantPlaylistHLS::FStreamInf& si = NewVariantPlaylist->InitialStreamInfs.Emplace_GetRef();
 		si.URI = InPlaylist->GetURL();
 		si.Bandwidth = 500000;
 
 		MultiVariantPlaylist = MoveTemp(NewVariantPlaylist);
 		FillInMissingInformation(MultiVariantPlaylist, EFillInOptions::All);
+		GroupVariantStreamsByPathways(MultiVariantPlaylist);
+		AssignInternalVariantStreamIDs(MultiVariantPlaylist);
 		GroupVariantStreamsByVideoProperties(MultiVariantPlaylist);
 		GroupAudioOnlyVariantStreams(MultiVariantPlaylist);
 		return true;
@@ -1350,7 +1626,7 @@ bool FPlaylistHandlerHLS::BuildMultiVariantPlaylist(TSharedPtr<FMultiVariantPlay
 	TSharedPtr<FMultiVariantPlaylistHLS, ESPMode::ThreadSafe> NewVariantPlaylist = MakeShared<FMultiVariantPlaylistHLS, ESPMode::ThreadSafe>();
 	NewVariantPlaylist->URL = InPlaylist->GetURL();
 	NewVariantPlaylist->ParsedURL.Parse(InPlaylist->GetURL());
-	NewVariantPlaylist->bHaveContentSteering = InPlaylist->UsesContentSteering();
+	NewVariantPlaylist->ContentSteeringParams.bHaveContentSteering = InPlaylist->UsesContentSteering();
 
 	// First process all EXT-X-DEFINE elements to set up the values for the variable substitutions.
 	TArray<FPlaylistParserHLS::FVariableSubstitution> EmptySubstitutions;
@@ -1650,27 +1926,43 @@ bool FPlaylistHandlerHLS::BuildMultiVariantPlaylist(TSharedPtr<FMultiVariantPlay
 		{
 			if (const FPlaylistParserHLS::FAttribute* ServerURI = Elements[nElem]->GetAttribute(TEXT("SERVER-URI")))
 			{
-				if (!ServerURI->GetValue(NewVariantPlaylist->SteeringURI, NewVariantPlaylist->VariableSubstitutions))
+				if (!ServerURI->GetValue(NewVariantPlaylist->ContentSteeringParams.SteeringURI, NewVariantPlaylist->VariableSubstitutions))
 				{
 					LastError = PostError(FString::Printf(TEXT("EXT-X-CONTENT-STEERING has a bad SERVER-URI value")), HLS::ERRCODE_PLAYLIST_PARSING_FAILED);
 					return false;
 				}
-				// Make an absolute URL unless this is a data URL
-				if (!NewVariantPlaylist->SteeringURI.StartsWith(TEXT("data:")))
-				{
-					FURL_RFC3986 UrlParser(NewVariantPlaylist->ParsedURL);
-					UrlParser.ResolveWith(NewVariantPlaylist->SteeringURI);
-					NewVariantPlaylist->SteeringURI = UrlParser.Get(true, true);
-				}
 
 				if (const FPlaylistParserHLS::FAttribute* PathwayId = Elements[nElem]->GetAttribute(TEXT("PATHWAY-ID")))
 				{
-					if (!PathwayId->GetValue(NewVariantPlaylist->PrimaryPathwayId, NewVariantPlaylist->VariableSubstitutions))
+					if (!PathwayId->GetValue(NewVariantPlaylist->ContentSteeringParams.PrimaryPathwayId, NewVariantPlaylist->VariableSubstitutions))
 					{
 						LastError = PostError(FString::Printf(TEXT("EXT-X-CONTENT-STEERING has a bad PATHWAY-ID value")), HLS::ERRCODE_PLAYLIST_PARSING_FAILED);
 						return false;
 					}
 				}
+
+				// Check for our custom attributes
+				if (const FPlaylistParserHLS::FAttribute* QueryBeforeStart = Elements[nElem]->GetAttribute(TEXT("EPIC-QUERY-BEFORE-START")))
+				{
+					// This is the equivalent of the DASH <ContentSteering@queryBeforeStart>
+					FString qbsValue;
+					// We use a string value here to allow it to be controlled by variable substitution.
+					if (!QueryBeforeStart->GetValue(qbsValue, NewVariantPlaylist->VariableSubstitutions))
+					{
+						LastError = PostError(FString::Printf(TEXT("EXT-X-CONTENT-STEERING has a bad EPIC-QUERY-BEFORE-START value")), HLS::ERRCODE_PLAYLIST_PARSING_FAILED);
+						return false;
+					}
+					TArray<FString> qbfTrueValues { TEXT("true"), TEXT("yes"), TEXT("1") };
+					NewVariantPlaylist->ContentSteeringParams.bQueryBeforeStart = qbfTrueValues.ContainsByPredicate([&](const FString& InMaybe){ return InMaybe.Equals(qbsValue, ESearchCase::IgnoreCase); });
+				}
+				if (const FPlaylistParserHLS::FAttribute* InitialSelectionPriority = Elements[nElem]->GetAttribute(TEXT("EPIC-INITIAL-SELECTION-PRIORITY")))
+				{
+					if (!InitialSelectionPriority->GetValue(NewVariantPlaylist->ContentSteeringParams.CustomInitialSelectionPriority, NewVariantPlaylist->VariableSubstitutions))
+					{
+						LastError = PostError(FString::Printf(TEXT("EXT-X-CONTENT-STEERING has a bad EPIC-INITIAL-SELECTION-PRIORITY value")), HLS::ERRCODE_PLAYLIST_PARSING_FAILED);
+						return false;
+					}
+				}
 			}
 			else
 			{
@@ -2009,8 +2301,7 @@ bool FPlaylistHandlerHLS::BuildMultiVariantPlaylist(TSharedPtr<FMultiVariantPlay
 			NewStreamInf.URI = UrlParser.Get(true, true);
 
 			// Add the new stream inf to the list.
-			NewStreamInf.Index = NewVariantPlaylist->StreamInfs.Num();
-			NewVariantPlaylist->StreamInfs.Emplace(MoveTemp(NewStreamInf));
+			NewVariantPlaylist->InitialStreamInfs.Emplace(MoveTemp(NewStreamInf));
 		}
 	}
 
@@ -2025,6 +2316,43 @@ bool FPlaylistHandlerHLS::BuildMultiVariantPlaylist(TSharedPtr<FMultiVariantPlay
 	return true;
 }
 
+void FPlaylistHandlerHLS::GroupVariantStreamsByPathways(TSharedPtr<FMultiVariantPlaylistHLS, ESPMode::ThreadSafe> InMultiVariantPlaylist)
+{
+	for(int32 ns=0; ns<InMultiVariantPlaylist->InitialStreamInfs.Num(); ++ns)
+	{
+		// If there is no PATHWAY-ID set on the stream inf yet, set the default pathway of "."
+		if (InMultiVariantPlaylist->InitialStreamInfs[ns].PathwayId.IsEmpty())
+		{
+			InMultiVariantPlaylist->InitialStreamInfs[ns].PathwayId = TEXT(".");
+		}
+		// Assign to the pathway bucket, creating it if necessary.
+		TSharedPtrTS<FMultiVariantPlaylistHLS::FPathwayStreamInfs>* pwEntry = InMultiVariantPlaylist->PathwayStreamInfs.FindByPredicate([pwId=InMultiVariantPlaylist->InitialStreamInfs[ns].PathwayId](const TSharedPtrTS<FMultiVariantPlaylistHLS::FPathwayStreamInfs>& e){return e->PathwayID.Equals(pwId);});
+		if (!pwEntry)
+		{
+			TSharedPtrTS<FMultiVariantPlaylistHLS::FPathwayStreamInfs> newPw(MakeSharedTS<FMultiVariantPlaylistHLS::FPathwayStreamInfs>());
+			InMultiVariantPlaylist->PathwayStreamInfs.Emplace(newPw);
+			newPw->PathwayID = InMultiVariantPlaylist->InitialStreamInfs[ns].PathwayId;
+			pwEntry = &InMultiVariantPlaylist->PathwayStreamInfs.Last();
+		}
+		(*pwEntry)->StreamInfs.Emplace(InMultiVariantPlaylist->InitialStreamInfs[ns]);
+		(*pwEntry)->StreamInfs.Last().IndexOfSelfInArray = (*pwEntry)->StreamInfs.Num() - 1;
+	}
+	// From this point on the initial collection list of the stream infs must not be used any more.
+	InMultiVariantPlaylist->InitialStreamInfs.Empty();
+}
+
+void FPlaylistHandlerHLS::AssignInternalVariantStreamIDs(TSharedPtr<FMultiVariantPlaylistHLS, ESPMode::ThreadSafe> InMultiVariantPlaylist)
+{
+	for(auto& pwIt : InMultiVariantPlaylist->PathwayStreamInfs)
+	{
+		for(int32 i=0,iMax=pwIt->StreamInfs.Num(); i<iMax; ++i)
+		{
+			pwIt->StreamInfs[i].ID = FString::Printf(TEXT("%d"), i);
+		}
+	}
+}
+
+
 void FPlaylistHandlerHLS::GroupVariantStreamsByVideoProperties(TSharedPtr<FMultiVariantPlaylistHLS, ESPMode::ThreadSafe> InMultiVariantPlaylist)
 {
 	auto UpdateHashStr = [](FSHA1& InOutHash, const FString& InString)	{ if (InString.Len()) InOutHash.UpdateWithString(*InString, InString.Len()); };
@@ -2037,134 +2365,139 @@ void FPlaylistHandlerHLS::GroupVariantStreamsByVideoProperties(TSharedPtr<FMulti
 		}
 		return InCodec;
 	};
-	TMultiMap<FString, int32> SamePropertyMap;
-	for(int32 ns=0,nsMax=InMultiVariantPlaylist->StreamInfs.Num(); ns<nsMax; ++ns)
+	// Group video streams by property for each pathway, since pathways may have different variant streams.
+	for(int32 pwIdx=0; pwIdx<InMultiVariantPlaylist->PathwayStreamInfs.Num(); ++pwIdx)
 	{
-		FMultiVariantPlaylistHLS::FStreamInf& si = InMultiVariantPlaylist->StreamInfs[ns];
+		TSharedPtrTS<FMultiVariantPlaylistHLS::FPathwayStreamInfs> pwSinf(InMultiVariantPlaylist->PathwayStreamInfs[pwIdx]);
+		TArray<FMultiVariantPlaylistHLS::FStreamInf>& sInfs(pwSinf->StreamInfs);
 
-		if (si.NumVideoCodec > 1)
+		TMultiMap<FString, int32> SamePropertyMap;
+		for(int32 ns=0,nsMax=sInfs.Num(); ns<nsMax; ++ns)
 		{
-			LogMessage(IInfoLog::ELevel::Warning, FString::Printf(TEXT("Variant stream gives more than one video codec, ignoring this variant.")));
-			continue;
-		}
-
-		// Generate a (hopefully) unique hash over all the variant stream attributes in an ordered and consistent way.
-		FSHA1 Hash;
+			FMultiVariantPlaylistHLS::FStreamInf& si = sInfs[ns];
 
-		TArray<FString> BaseCodecs;
-		FString SupplementalCodec;
-		bool bHasVideo = false;
-		// Add all unique base codecs in sorted ascending order (so they same ones give the same hash)
-		for(int32 i=0; i<si.ParsedCodecs.Num(); ++i)
-		{
-			if (si.ParsedCodecs[i].IsVideoCodec())
+			if (si.NumVideoCodec > 1)
 			{
-				bHasVideo = true;
+				LogMessage(IInfoLog::ELevel::Warning, FString::Printf(TEXT("Variant stream gives more than one video codec, ignoring this variant.")));
+				continue;
 			}
-			BaseCodecs.AddUnique(GetBaseCodec(si.Codecs[i]));
-		}
-		// No recognized video codec, no need to create a group.
-		if (!bHasVideo)
-		{
-			continue;
-		}
-		if (si.SupplementalCodecs.Num())
-		{
-			SupplementalCodec = GetBaseCodec(si.SupplementalCodecs[0]);
-		}
-		BaseCodecs.Sort();
-		for(auto& Codec : BaseCodecs)
-		{
-			UpdateHashStr(Hash, Codec);
-		}
-		UpdateHashStr(Hash, SupplementalCodec);
-		UpdateHashStr(Hash, si.VideoRange);
-		/*
-			Do NOT add any of the video/audio/etc. groups as each variant may reference a different group
-			(a variant is the bandwidth-switchable entity, the groups are things like different angles)
-			If the groups use different codecs, like in
-				#EXT-X-STREAM-INF:AVERAGE-BANDWIDTH=2168183,BANDWIDTH=2177116,CODECS="avc1.640020,mp4a.40.2",RESOLUTION=960x540,FRAME-RATE=60.000,CLOSED-CAPTIONS="cc1",AUDIO="aud1",SUBTITLES="sub1"
-			and
-				#EXT-X-STREAM-INF:AVERAGE-BANDWIDTH=2390686,BANDWIDTH=2399619,CODECS="avc1.640020,ac-3",RESOLUTION=960x540,FRAME-RATE=60.000,CLOSED-CAPTIONS="cc1",AUDIO="aud2",SUBTITLES="sub1"
-			with
-				#EXT-X-MEDIA:TYPE=AUDIO,GROUP-ID="aud1",LANGUAGE="en",NAME="English",AUTOSELECT=YES,DEFAULT=YES,CHANNELS="2",URI="a1/prog_index.m3u8"
-				#EXT-X-MEDIA:TYPE=AUDIO,GROUP-ID="aud2",LANGUAGE="en",NAME="English",AUTOSELECT=YES,DEFAULT=YES,CHANNELS="6",URI="a2/prog_index.m3u8"
-			we will have covered this by considering all base codecs in the hash at the top.
-
-			Do not
-				UpdateHashStr(Hash, si.AudioGroup);
-				UpdateHashStr(Hash, si.SubtitleGroup);
-				UpdateHashStr(Hash, si.ClosedCaptionGroup);
-		*/
-		FString HashValue = Hash.Finalize().ToString();
-		SamePropertyMap.Emplace(MoveTemp(HashValue), ns);
-	}
-
 
-	// Go over the tuples of same hash values and create groups
-	TMap<FString, int32> SameGroupURLHashMap;
-	TArray<FString> KeyArray;
-	SamePropertyMap.GetKeys(KeyArray);
-	for(auto& kIt : KeyArray)
-	{
-		FMultiVariantPlaylistHLS::FVideoVariantGroup& vg = InMultiVariantPlaylist->VideoVariantGroups.Emplace_GetRef();
+			// Generate a (hopefully) unique hash over all the variant stream attributes in an ordered and consistent way.
+			FSHA1 Hash;
 
-		// Get the indices of the variant streams in the StreamInfs and sort them
-		// back ascending as the ordering in the map is not guaranteed.
-		for(auto it=SamePropertyMap.CreateConstKeyIterator(kIt); it; ++it)
-		{
-			vg.StreamInfIndices.Add((*it).Value);
+			TArray<FString> BaseCodecs;
+			FString SupplementalCodec;
+			bool bHasVideo = false;
+			// Add all unique base codecs in sorted ascending order (so they same ones give the same hash)
+			for(int32 i=0; i<si.ParsedCodecs.Num(); ++i)
+			{
+				if (si.ParsedCodecs[i].IsVideoCodec())
+				{
+					bHasVideo = true;
+				}
+				BaseCodecs.AddUnique(GetBaseCodec(si.Codecs[i]));
+			}
+			// No recognized video codec, no need to create a group.
+			if (!bHasVideo)
+			{
+				continue;
+			}
+			if (si.SupplementalCodecs.Num())
+			{
+				SupplementalCodec = GetBaseCodec(si.SupplementalCodecs[0]);
+			}
+			BaseCodecs.Sort();
+			for(auto& Codec : BaseCodecs)
+			{
+				UpdateHashStr(Hash, Codec);
+			}
+			UpdateHashStr(Hash, SupplementalCodec);
+			UpdateHashStr(Hash, si.VideoRange);
+			/*
+				Do NOT add any of the video/audio/etc. groups as each variant may reference a different group
+				(a variant is the bandwidth-switchable entity, the groups are things like different angles)
+				If the groups use different codecs, like in
+					#EXT-X-STREAM-INF:AVERAGE-BANDWIDTH=2168183,BANDWIDTH=2177116,CODECS="avc1.640020,mp4a.40.2",RESOLUTION=960x540,FRAME-RATE=60.000,CLOSED-CAPTIONS="cc1",AUDIO="aud1",SUBTITLES="sub1"
+				and
+					#EXT-X-STREAM-INF:AVERAGE-BANDWIDTH=2390686,BANDWIDTH=2399619,CODECS="avc1.640020,ac-3",RESOLUTION=960x540,FRAME-RATE=60.000,CLOSED-CAPTIONS="cc1",AUDIO="aud2",SUBTITLES="sub1"
+				with
+					#EXT-X-MEDIA:TYPE=AUDIO,GROUP-ID="aud1",LANGUAGE="en",NAME="English",AUTOSELECT=YES,DEFAULT=YES,CHANNELS="2",URI="a1/prog_index.m3u8"
+					#EXT-X-MEDIA:TYPE=AUDIO,GROUP-ID="aud2",LANGUAGE="en",NAME="English",AUTOSELECT=YES,DEFAULT=YES,CHANNELS="6",URI="a2/prog_index.m3u8"
+				we will have covered this by considering all base codecs in the hash at the top.
+
+				Do not
+					UpdateHashStr(Hash, si.AudioGroup);
+					UpdateHashStr(Hash, si.SubtitleGroup);
+					UpdateHashStr(Hash, si.ClosedCaptionGroup);
+			*/
+			FString HashValue = Hash.Finalize().ToString();
+			SamePropertyMap.Emplace(MoveTemp(HashValue), ns);
 		}
-		vg.StreamInfIndices.Sort();
 
-		// Gather the individual bandwidths to create a list of "qualities" and assign quality indices.
-		TArray<int64> SortedBandwidths;
-		for(int32 j=0; j<vg.StreamInfIndices.Num(); ++j)
+		// Go over the tuples of same hash values and create groups
+		TMap<FString, int32> SameGroupURLHashMap;
+		TArray<FString> KeyArray;
+		SamePropertyMap.GetKeys(KeyArray);
+		for(auto& kIt : KeyArray)
 		{
-			SortedBandwidths.AddUnique(InMultiVariantPlaylist->StreamInfs[vg.StreamInfIndices[j]].Bandwidth);
-		}
-		SortedBandwidths.Sort();
+			FMultiVariantPlaylistHLS::FVideoVariantGroup& vg = pwSinf->VideoVariantGroups.Emplace_GetRef();
 
-		// Add the URL of the media playlist in bandwidth order to the group hash.
-		// The idea is that groups that use the same media playlist as a variant but differ in the use
-		// of the VIDEO, AUDIO, SUBTITLES or CLOSED-CAPTIONS groups (and hence have a different combined BANDWIDTH
-		// so we can't use that as an identifier) can be identified by a hash over all the playlists URLs.
-		FSHA1 GroupURLHash;
-		TArray<FMultiVariantPlaylistHLS::FStreamInf> SortedStreamInfs;
-		for(int32 j=0; j<vg.StreamInfIndices.Num(); ++j)
-		{
-			SortedStreamInfs.Emplace(InMultiVariantPlaylist->StreamInfs[vg.StreamInfIndices[j]]);
-		}
-		SortedStreamInfs.StableSort([](const FMultiVariantPlaylistHLS::FStreamInf& a, const FMultiVariantPlaylistHLS::FStreamInf& b){return a.Bandwidth < b.Bandwidth;});
-		for(int32 j=0; j<SortedStreamInfs.Num(); ++j)
-		{
-			UpdateHashStr(GroupURLHash, SortedStreamInfs[j].URI);
-		}
-		FString HashValue = GroupURLHash.Finalize().ToString();
-		if (SameGroupURLHashMap.Contains(HashValue))
-		{
-			vg.SameAsVideoVariantGroupIndex.Emplace(SameGroupURLHashMap[HashValue]);
-		}
-		else
-		{
-			SameGroupURLHashMap.Emplace(HashValue, InMultiVariantPlaylist->VideoVariantGroups.Num()-1);
-		}
+			// Get the indices of the variant streams in the StreamInfs and sort them
+			// back ascending as the ordering in the map is not guaranteed.
+			for(auto it=SamePropertyMap.CreateConstKeyIterator(kIt); it; ++it)
+			{
+				vg.StreamInfIndices.Add((*it).Value);
+			}
+			vg.StreamInfIndices.Sort();
 
+			// Gather the individual bandwidths to create a list of "qualities" and assign quality indices.
+			TArray<int64> SortedBandwidths;
+			for(int32 j=0; j<vg.StreamInfIndices.Num(); ++j)
+			{
+				SortedBandwidths.AddUnique(sInfs[vg.StreamInfIndices[j]].Bandwidth);
+			}
+			SortedBandwidths.Sort();
 
-		for(int32 j=0; j<vg.StreamInfIndices.Num(); ++j)
-		{
-			FMultiVariantPlaylistHLS::FStreamInf& si = InMultiVariantPlaylist->StreamInfs[vg.StreamInfIndices[j]];
-			si.QualityIndex = SortedBandwidths.Find(si.Bandwidth);
-			for(int32 i=0; i<si.ParsedCodecs.Num(); ++i)
+			// Add the URL of the media playlist in bandwidth order to the group hash.
+			// The idea is that groups that use the same media playlist as a variant but differ in the use
+			// of the VIDEO, AUDIO, SUBTITLES or CLOSED-CAPTIONS groups (and hence have a different combined BANDWIDTH
+			// so we can't use that as an identifier) can be identified by a hash over all the playlists URLs.
+			FSHA1 GroupURLHash;
+			TArray<FMultiVariantPlaylistHLS::FStreamInf> SortedStreamInfs;
+			for(int32 j=0; j<vg.StreamInfIndices.Num(); ++j)
 			{
-				if (si.ParsedCodecs[i].IsVideoCodec())
+				SortedStreamInfs.Emplace(sInfs[vg.StreamInfIndices[j]]);
+			}
+			SortedStreamInfs.StableSort([](const FMultiVariantPlaylistHLS::FStreamInf& a, const FMultiVariantPlaylistHLS::FStreamInf& b){return a.Bandwidth < b.Bandwidth;});
+			for(int32 j=0; j<SortedStreamInfs.Num(); ++j)
+			{
+				UpdateHashStr(GroupURLHash, SortedStreamInfs[j].URI);
+			}
+			FString HashValue = GroupURLHash.Finalize().ToString();
+			if (SameGroupURLHashMap.Contains(HashValue))
+			{
+				vg.SameAsVideoVariantGroupIndex.Emplace(SameGroupURLHashMap[HashValue]);
+			}
+			else
+			{
+				SameGroupURLHashMap.Emplace(HashValue, pwSinf->VideoVariantGroups.Num()-1);
+			}
+
+			for(int32 j=0; j<vg.StreamInfIndices.Num(); ++j)
+			{
+				FMultiVariantPlaylistHLS::FStreamInf& si = sInfs[vg.StreamInfIndices[j]];
+				si.QualityIndex = SortedBandwidths.Find(si.Bandwidth);
+				for(int32 i=0; i<si.ParsedCodecs.Num(); ++i)
 				{
-					vg.ParsedCodecs.Emplace(si.ParsedCodecs[i]);
-					vg.BaseSupplementalCodecs.Append(si.SupplementalCodecs);
+					if (si.ParsedCodecs[i].IsVideoCodec())
+					{
+						vg.ParsedCodecs.Emplace(si.ParsedCodecs[i]);
+						vg.BaseSupplementalCodecs.Append(si.SupplementalCodecs);
+					}
 				}
+				check(vg.VideoRange.IsEmpty() || vg.VideoRange.Equals(si.VideoRange));
+				vg.VideoRange = si.VideoRange;
 			}
-			check(vg.VideoRange.IsEmpty() || vg.VideoRange.Equals(si.VideoRange));
-			vg.VideoRange = si.VideoRange;
 		}
 	}
 }
@@ -2181,71 +2514,79 @@ void FPlaylistHandlerHLS::GroupAudioOnlyVariantStreams(TSharedPtr<FMultiVariantP
 		}
 		return InCodec;
 	};
-	TMultiMap<FString, int32> SamePropertyMap;
-	bool bHasAnyVariantMissingAudioCodec = false;
-	for(int32 ns=0,nsMax=InMultiVariantPlaylist->StreamInfs.Num(); ns<nsMax; ++ns)
-	{
-		FMultiVariantPlaylistHLS::FStreamInf& si = InMultiVariantPlaylist->StreamInfs[ns];
-		bHasAnyVariantMissingAudioCodec |= si.bReferencesAudioRenditionWithoutCodec;
 
-		// Generate a (hopefully) unique hash over all the variant stream attributes in an ordered and consistent way.
-		FSHA1 Hash;
+	// Group audio streams by property for each pathway, since pathways may have different variant streams.
+	for(int32 pwIdx=0; pwIdx<InMultiVariantPlaylist->PathwayStreamInfs.Num(); ++pwIdx)
+	{
+		TSharedPtrTS<FMultiVariantPlaylistHLS::FPathwayStreamInfs> pwSinf(InMultiVariantPlaylist->PathwayStreamInfs[pwIdx]);
+		TArray<FMultiVariantPlaylistHLS::FStreamInf>& sInfs(pwSinf->StreamInfs);
 
-		bool bHasVideo = false;
-		for(int32 i=0; i<si.ParsedCodecs.Num(); ++i)
+		TMultiMap<FString, int32> SamePropertyMap;
+		bool bHasAnyVariantMissingAudioCodec = false;
+		for(int32 ns=0,nsMax=sInfs.Num(); ns<nsMax; ++ns)
 		{
-			// When video is present this will not be an audio-only group.
-			// We do not look at subtitles because an audio-only group is allowed to have subtitles (for Karaoke, for example).
-			if (si.ParsedCodecs[i].IsVideoCodec())
+			FMultiVariantPlaylistHLS::FStreamInf& si = sInfs[ns];
+			bHasAnyVariantMissingAudioCodec |= si.bReferencesAudioRenditionWithoutCodec;
+
+			// Generate a (hopefully) unique hash over all the variant stream attributes in an ordered and consistent way.
+			FSHA1 Hash;
+
+			bool bHasVideo = false;
+			for(int32 i=0; i<si.ParsedCodecs.Num(); ++i)
 			{
-				bHasVideo = true;
-				break;
+				// When video is present this will not be an audio-only group.
+				// We do not look at subtitles because an audio-only group is allowed to have subtitles (for Karaoke, for example).
+				if (si.ParsedCodecs[i].IsVideoCodec())
+				{
+					bHasVideo = true;
+					break;
+				}
+				else if (si.ParsedCodecs[i].IsAudioCodec())
+				{
+					UpdateHashStr(Hash, GetBaseCodec(si.Codecs[i]));
+				}
 			}
-			else if (si.ParsedCodecs[i].IsAudioCodec())
+			if (!bHasVideo)
 			{
-				UpdateHashStr(Hash, GetBaseCodec(si.Codecs[i]));
+				// Do NOT add the audio group as each variant may reference a different group
+				// (a variant is the bandwidth-switchable entity, the groups are things like different languages, commentary, etc.)
+				UpdateHashStr(Hash, si.SubtitleGroup);
+				FString HashValue = Hash.Finalize().ToString();
+				SamePropertyMap.Emplace(MoveTemp(HashValue), ns);
 			}
 		}
-		if (!bHasVideo)
-		{
-			// Do NOT add the audio group as each variant may reference a different group
-			// (a variant is the bandwidth-switchable entity, the groups are things like different languages, commentary, etc.)
-			UpdateHashStr(Hash, si.SubtitleGroup);
-			FString HashValue = Hash.Finalize().ToString();
-			SamePropertyMap.Emplace(MoveTemp(HashValue), ns);
-		}
-	}
-
-	// Go over the tuples of same hash values and create groups
-	TArray<FString> KeyArray;
-	SamePropertyMap.GetKeys(KeyArray);
-	for(auto& kIt : KeyArray)
-	{
-		FMultiVariantPlaylistHLS::FAudioVariantGroup& ag = InMultiVariantPlaylist->AudioOnlyVariantGroups.Emplace_GetRef();
 
-		// Get the indices of the variant streams in the StreamInfs and sort them
-		// back ascending as the ordering in the map is not guaranteed.
-		for(auto it=SamePropertyMap.CreateConstKeyIterator(kIt); it; ++it)
+		// Go over the tuples of same hash values and create groups
+		TArray<FString> KeyArray;
+		SamePropertyMap.GetKeys(KeyArray);
+		for(auto& kIt : KeyArray)
 		{
-			ag.StreamInfIndices.Add((*it).Value);
-		}
-		ag.StreamInfIndices.Sort();
+			FMultiVariantPlaylistHLS::FAudioVariantGroup& ag = pwSinf->AudioOnlyVariantGroups.Emplace_GetRef();
+
+			// Get the indices of the variant streams in the StreamInfs and sort them
+			// back ascending as the ordering in the map is not guaranteed.
+			for(auto it=SamePropertyMap.CreateConstKeyIterator(kIt); it; ++it)
+			{
+				ag.StreamInfIndices.Add((*it).Value);
+			}
+			ag.StreamInfIndices.Sort();
 
-		const FMultiVariantPlaylistHLS::FStreamInf& si = InMultiVariantPlaylist->StreamInfs[ag.StreamInfIndices[0]];
-		ag.ParsedCodecs.Append(si.ParsedCodecs);
+			const FMultiVariantPlaylistHLS::FStreamInf& si = sInfs[ag.StreamInfIndices[0]];
+			ag.ParsedCodecs.Append(si.ParsedCodecs);
 
-		// If any variant has missing audio codecs for a rendition group it references
-		// and this audio-only variant is in the same group, then add the codec.
-		if (bHasAnyVariantMissingAudioCodec && si.AudioGroup.Len())
-		{
-			for(int32 k=0,nsMax=InMultiVariantPlaylist->StreamInfs.Num(); k<nsMax; ++k)
+			// If any variant has missing audio codecs for a rendition group it references
+			// and this audio-only variant is in the same group, then add the codec.
+			if (bHasAnyVariantMissingAudioCodec && si.AudioGroup.Len())
 			{
-				FMultiVariantPlaylistHLS::FStreamInf& si2 = InMultiVariantPlaylist->StreamInfs[k];
-				if (si.Index != si2.Index && si2.bReferencesAudioRenditionWithoutCodec && si2.AudioGroup.Equals(si2.AudioGroup))
+				for(int32 k=0,nsMax=sInfs.Num(); k<nsMax; ++k)
 				{
-					si2.Codecs.Append(si.Codecs);
-					si2.ParsedCodecs.Append(si.ParsedCodecs);
-					si2.NumAudioCodec += si.ParsedCodecs.Num();
+					FMultiVariantPlaylistHLS::FStreamInf& si2 = sInfs[k];
+					if (si.IndexOfSelfInArray != si2.IndexOfSelfInArray && si2.bReferencesAudioRenditionWithoutCodec && si2.AudioGroup.Equals(si2.AudioGroup))
+					{
+						si2.Codecs.Append(si.Codecs);
+						si2.ParsedCodecs.Append(si.ParsedCodecs);
+						si2.NumAudioCodec += si.ParsedCodecs.Num();
+					}
 				}
 			}
 		}
@@ -2290,9 +2631,6 @@ void FPlaylistHandlerHLS::FillInMissingInformation(TSharedPtr<FMultiVariantPlayl
 		CheckForFallbackStreams(InMultiVariantPlaylist);
 	}
 
-	// Assign stable variant and rendition IDs
-	AssignStableIds(InMultiVariantPlaylist);
-
 	// Check for generally missing CODECS
 	if ((InFillInOpts & EFillInOptions::Codecs) != EFillInOptions::None)
 	{
@@ -2324,41 +2662,11 @@ void FPlaylistHandlerHLS::FillInMissingInformation(TSharedPtr<FMultiVariantPlayl
 	}
 }
 
-void FPlaylistHandlerHLS::AssignStableIds(TSharedPtr<FMultiVariantPlaylistHLS, ESPMode::ThreadSafe> InMultiVariantPlaylist)
-{
-	// Content steering is supposed to have stable IDs assigned already.
-	if (!InMultiVariantPlaylist->bHaveContentSteering)
-	{
-		int32 NextId = 0;
-		for(int32 i=0; i<InMultiVariantPlaylist->StreamInfs.Num(); ++i)
-		{
-			if (InMultiVariantPlaylist->StreamInfs[i].StableVariantId.IsEmpty())
-			{
-				InMultiVariantPlaylist->StreamInfs[i].StableVariantId = FString::Printf(TEXT("[v%02d]"), ++NextId);
-			}
-		}
-		NextId = 0;
-		for(int32 typ=0; typ<UE_ARRAY_COUNT(InMultiVariantPlaylist->RenditionGroupsOfType); ++typ)
-		{
-			for(int32 grp=0; grp<InMultiVariantPlaylist->RenditionGroupsOfType[typ].Num(); ++grp)
-			{
-				for(int32 i=0; i<InMultiVariantPlaylist->RenditionGroupsOfType[typ][grp].Renditions.Num(); ++i)
-				{
-					if (InMultiVariantPlaylist->RenditionGroupsOfType[typ][grp].Renditions[i].StableRenditionId.IsEmpty())
-					{
-						InMultiVariantPlaylist->RenditionGroupsOfType[typ][grp].Renditions[i].StableRenditionId = FString::Printf(TEXT("[r%02d]"), ++NextId);
-					}
-				}
-			}
-		}
-	}
-}
-
 void FPlaylistHandlerHLS::AssignResolutionAndFrameRateToCodecs(TSharedPtr<FMultiVariantPlaylistHLS, ESPMode::ThreadSafe> InMultiVariantPlaylist)
 {
-	for(int32 ns=0,nsMax=InMultiVariantPlaylist->StreamInfs.Num(); ns<nsMax; ++ns)
+	for(int32 ns=0,nsMax=InMultiVariantPlaylist->InitialStreamInfs.Num(); ns<nsMax; ++ns)
 	{
-		FMultiVariantPlaylistHLS::FStreamInf& si = InMultiVariantPlaylist->StreamInfs[ns];
+		FMultiVariantPlaylistHLS::FStreamInf& si = InMultiVariantPlaylist->InitialStreamInfs[ns];
 		if (si.NumVideoCodec)
 		{
 			for(int32 j=0; j<si.ParsedCodecs.Num(); ++j)
@@ -2389,9 +2697,9 @@ void FPlaylistHandlerHLS::AssignCodecsToRenditions(TSharedPtr<FMultiVariantPlayl
 	};
 
 	// Go over each variant stream and see what groups it references
-	for(int32 ns=0,nsMax=InMultiVariantPlaylist->StreamInfs.Num(); ns<nsMax; ++ns)
+	for(int32 ns=0,nsMax=InMultiVariantPlaylist->InitialStreamInfs.Num(); ns<nsMax; ++ns)
 	{
-		FMultiVariantPlaylistHLS::FStreamInf& si = InMultiVariantPlaylist->StreamInfs[ns];
+		FMultiVariantPlaylistHLS::FStreamInf& si = InMultiVariantPlaylist->InitialStreamInfs[ns];
 		for(int32 ng=0; ng<3; ++ng)
 		{
 			FMultiVariantPlaylistHLS::FRenditionGroup* rg = nullptr;
@@ -2488,9 +2796,9 @@ void FPlaylistHandlerHLS::AssignCodecsToRenditions(TSharedPtr<FMultiVariantPlayl
 
 void FPlaylistHandlerHLS::CheckForMissingCodecs(TSharedPtr<FMultiVariantPlaylistHLS, ESPMode::ThreadSafe> InMultiVariantPlaylist)
 {
-	for(int32 ns=0,nsMax=InMultiVariantPlaylist->StreamInfs.Num(); ns<nsMax; ++ns)
+	for(int32 ns=0,nsMax=InMultiVariantPlaylist->InitialStreamInfs.Num(); ns<nsMax; ++ns)
 	{
-		FMultiVariantPlaylistHLS::FStreamInf& si = InMultiVariantPlaylist->StreamInfs[ns];
+		FMultiVariantPlaylistHLS::FStreamInf& si = InMultiVariantPlaylist->InitialStreamInfs[ns];
 		if (si.Codecs.IsEmpty())
 		{
 			LogMessage(IInfoLog::ELevel::Warning, FString::Printf(TEXT("Variant stream is missing the CODECS attribute. We pretend it to be \"avc1.640028,mp4a.40.2\". Playback may not work correctly.")));
@@ -2509,9 +2817,9 @@ void FPlaylistHandlerHLS::CheckForMissingResolution(TSharedPtr<FMultiVariantPlay
 {
 	// See if any variant stream containing video is missing the RESOLUTION attribute.
 	bool bHasMissingResolution = false;
-	for(int32 ns=0,nsMax=InMultiVariantPlaylist->StreamInfs.Num(); ns<nsMax && !bHasMissingResolution; ++ns)
+	for(int32 ns=0,nsMax=InMultiVariantPlaylist->InitialStreamInfs.Num(); ns<nsMax && !bHasMissingResolution; ++ns)
 	{
-		FMultiVariantPlaylistHLS::FStreamInf& si = InMultiVariantPlaylist->StreamInfs[ns];
+		FMultiVariantPlaylistHLS::FStreamInf& si = InMultiVariantPlaylist->InitialStreamInfs[ns];
 		bHasMissingResolution |= si.NumVideoCodec && si.ResolutionH <= 0;
 	}
 	if (bHasMissingResolution)
@@ -2519,9 +2827,9 @@ void FPlaylistHandlerHLS::CheckForMissingResolution(TSharedPtr<FMultiVariantPlay
 		// If missing we now take the variant streams and sort them by BANDWIDTH,
 		// generating a fake resolution for it.
 		TArray<FMultiVariantPlaylistHLS::FStreamInf> VideoVariantStreams;
-		for(int32 ns=0,nsMax=InMultiVariantPlaylist->StreamInfs.Num(); ns<nsMax; ++ns)
+		for(int32 ns=0,nsMax=InMultiVariantPlaylist->InitialStreamInfs.Num(); ns<nsMax; ++ns)
 		{
-			FMultiVariantPlaylistHLS::FStreamInf& si = InMultiVariantPlaylist->StreamInfs[ns];
+			FMultiVariantPlaylistHLS::FStreamInf& si = InMultiVariantPlaylist->InitialStreamInfs[ns];
 			if (si.NumVideoCodec && !VideoVariantStreams.ContainsByPredicate([bw=si.Bandwidth](const FMultiVariantPlaylistHLS::FStreamInf& e){return e.Bandwidth == bw;}))
 			{
 				VideoVariantStreams.Emplace(si);
@@ -2538,9 +2846,9 @@ void FPlaylistHandlerHLS::CheckForMissingResolution(TSharedPtr<FMultiVariantPlay
 		for(int32 vsi=0; vsi<VideoVariantStreams.Num(); ++vsi)
 		{
 			const int32 Height = vsi < UE_ARRAY_COUNT(kCommonHeights) ? kCommonHeights[vsi] : kCommonHeights[UE_ARRAY_COUNT(kCommonHeights)-1];
-			for(int32 ns=0,nsMax=InMultiVariantPlaylist->StreamInfs.Num(); ns<nsMax; ++ns)
+			for(int32 ns=0,nsMax=InMultiVariantPlaylist->InitialStreamInfs.Num(); ns<nsMax; ++ns)
 			{
-				FMultiVariantPlaylistHLS::FStreamInf& si = InMultiVariantPlaylist->StreamInfs[ns];
+				FMultiVariantPlaylistHLS::FStreamInf& si = InMultiVariantPlaylist->InitialStreamInfs[ns];
 				if (si.NumVideoCodec && si.Bandwidth == VideoVariantStreams[vsi].Bandwidth && si.ResolutionH <= 0)
 				{
 					si.ResolutionH = Height;
@@ -2566,17 +2874,17 @@ void FPlaylistHandlerHLS::CheckForScore(TSharedPtr<FMultiVariantPlaylistHLS, ESP
 			"The SCORE attribute (if present) MUST be on every variant. Otherwise, the SCORE attribute will be ignored."
 	*/
 	int32 NumScores = 0;
-	for(int32 ns=0,nsMax=InMultiVariantPlaylist->StreamInfs.Num(); ns<nsMax; ++ns)
+	for(int32 ns=0,nsMax=InMultiVariantPlaylist->InitialStreamInfs.Num(); ns<nsMax; ++ns)
 	{
-		NumScores += InMultiVariantPlaylist->StreamInfs[ns].Score >= 0.0f ? 1 : 0;
+		NumScores += InMultiVariantPlaylist->InitialStreamInfs[ns].Score >= 0.0f ? 1 : 0;
 	}
 	// If only some variants have a score value, reset it on all of them
-	if (NumScores && NumScores != InMultiVariantPlaylist->StreamInfs.Num())
+	if (NumScores && NumScores != InMultiVariantPlaylist->InitialStreamInfs.Num())
 	{
 		LogMessage(IInfoLog::ELevel::Warning, FString::Printf(TEXT("Not all variant streams have a defined SCORE attribute. Ignoring the SCORE on all of them")));
-		for(int32 ns=0,nsMax=InMultiVariantPlaylist->StreamInfs.Num(); ns<nsMax; ++ns)
+		for(int32 ns=0,nsMax=InMultiVariantPlaylist->InitialStreamInfs.Num(); ns<nsMax; ++ns)
 		{
-			InMultiVariantPlaylist->StreamInfs[ns].Score = -1.0f;
+			InMultiVariantPlaylist->InitialStreamInfs[ns].Score = -1.0f;
 		}
 	}
 }
@@ -2586,9 +2894,9 @@ void FPlaylistHandlerHLS::CheckForFallbackStreams(TSharedPtr<FMultiVariantPlayli
 	auto UpdateHashStr = [](FSHA1& InOutHash, const FString& InString)	{ if (InString.Len()) InOutHash.UpdateWithString(*InString, InString.Len()); };
 	auto UpdateHashI64 = [](FSHA1& InOutHash, const int64 InValue)		{  InOutHash.Update(reinterpret_cast<const uint8*>(&InValue), sizeof(InValue)); };
 	TMultiMap<FString, int32> IdenticalAttributeMap;
-	for(int32 ns=0,nsMax=InMultiVariantPlaylist->StreamInfs.Num(); ns<nsMax; ++ns)
+	for(int32 ns=0,nsMax=InMultiVariantPlaylist->InitialStreamInfs.Num(); ns<nsMax; ++ns)
 	{
-		const FMultiVariantPlaylistHLS::FStreamInf& si = InMultiVariantPlaylist->StreamInfs[ns];
+		const FMultiVariantPlaylistHLS::FStreamInf& si = InMultiVariantPlaylist->InitialStreamInfs[ns];
 
 		// Generate a (hopefully) unique hash over all the variant stream attributes in an ordered and consistent way.
 		FSHA1 Hash;
@@ -2642,7 +2950,7 @@ void FPlaylistHandlerHLS::CheckForFallbackStreams(TSharedPtr<FMultiVariantPlayli
 	// We took those into account up above in the hash, so anything that still comes out as duplicates would also
 	// have to have the same id already, which is suspicious as content steering is intended for explicit CDN
 	// management and thus duplicate variants should not be listed.
-	if (InMultiVariantPlaylist->bHaveContentSteering)
+	if (InMultiVariantPlaylist->ContentSteeringParams.bHaveContentSteering)
 	{
 		LogMessage(IInfoLog::ELevel::Warning, FString::Printf(TEXT("Assigning generated PATHWAY-ID to like variants in a playlist that uses content steering. This may have undesirable effects.")));
 	}
@@ -2664,8 +2972,8 @@ void FPlaylistHandlerHLS::CheckForFallbackStreams(TSharedPtr<FMultiVariantPlayli
 		TArray<int32> DuplicatesToRemove;
 		for(int32 i=1; i<VariantIndices.Num(); ++i)
 		{
-			const FMultiVariantPlaylistHLS::FStreamInf& si0 = InMultiVariantPlaylist->StreamInfs[VariantIndices[i-1]];
-			const FMultiVariantPlaylistHLS::FStreamInf& si1 = InMultiVariantPlaylist->StreamInfs[VariantIndices[i]];
+			const FMultiVariantPlaylistHLS::FStreamInf& si0 = InMultiVariantPlaylist->InitialStreamInfs[VariantIndices[i-1]];
+			const FMultiVariantPlaylistHLS::FStreamInf& si1 = InMultiVariantPlaylist->InitialStreamInfs[VariantIndices[i]];
 			if (si1.URI.Equals(si0.URI))
 			{
 				// If the groups are identical as well, this really is a duplicate.
@@ -2685,23 +2993,18 @@ void FPlaylistHandlerHLS::CheckForFallbackStreams(TSharedPtr<FMultiVariantPlayli
 			{
 				continue;
 			}
-			FMultiVariantPlaylistHLS::FStreamInf& si = InMultiVariantPlaylist->StreamInfs[VariantIndices[i]];
+			FMultiVariantPlaylistHLS::FStreamInf& si = InMultiVariantPlaylist->InitialStreamInfs[VariantIndices[i]];
 			// Enclose the generated name with brackets, which are normally invalid characters for `PATHWAY-ID`
 			// to indicate that this is a generated ID.
 			++CDN;
 			si.PathwayId = FString::Printf(TEXT("[CDN-%02d]"), CDN);
-			si.StableVariantId = FString::Printf(TEXT("[Variant-%02d]"), j+1);
-			if (i==0 && InMultiVariantPlaylist->PrimaryPathwayId.IsEmpty())
-			{
-				InMultiVariantPlaylist->PrimaryPathwayId = si.PathwayId;
-			}
 		}
 	}
 	// Finally, remove all the duplicate variants for good.
 	VariantIndicesToRemove.Sort();
 	for(int32 i=VariantIndicesToRemove.Num()-1; i>=0; --i)
 	{
-		InMultiVariantPlaylist->StreamInfs.RemoveAt(VariantIndicesToRemove[i]);
+		InMultiVariantPlaylist->InitialStreamInfs.RemoveAt(VariantIndicesToRemove[i]);
 		//LogMessage(IInfoLog::ELevel::Warning, FString::Printf(TEXT("Removing duplicate variant stream.")));
 	}
 }
diff --git a/Engine/Plugins/Media/ElectraPlayer/Source/ElectraPlayerRuntime/Private/Runtime/Player/Manifest.h b/Engine/Plugins/Media/ElectraPlayer/Source/ElectraPlayerRuntime/Private/Runtime/Player/Manifest.h
--- a/Engine/Plugins/Media/ElectraPlayer/Source/ElectraPlayerRuntime/Private/Runtime/Player/Manifest.h
+++ b/Engine/Plugins/Media/ElectraPlayer/Source/ElectraPlayerRuntime/Private/Runtime/Player/Manifest.h
@@ -394,6 +394,7 @@ namespace Electra
 
 		virtual void TriggerPlaylistRefresh() = 0;
 
+		virtual void ReachedStableBuffer() = 0;
 
 		//-------------------------------------------------------------------------
 		// Stream fragment reader
diff --git a/Engine/Plugins/Media/ElectraPlayer/Source/ElectraPlayerRuntime/Private/Runtime/Player/PlayerSessionServices.h b/Engine/Plugins/Media/ElectraPlayer/Source/ElectraPlayerRuntime/Private/Runtime/Player/PlayerSessionServices.h
--- a/Engine/Plugins/Media/ElectraPlayer/Source/ElectraPlayerRuntime/Private/Runtime/Player/PlayerSessionServices.h
+++ b/Engine/Plugins/Media/ElectraPlayer/Source/ElectraPlayerRuntime/Private/Runtime/Player/PlayerSessionServices.h
@@ -21,6 +21,7 @@ namespace Electra
 	class IPlaylistReader;
 	class IAdaptiveStreamingPlayerAEMSHandler;
 	class FDRMManager;
+	class FContentSteeringHandler;
 	class IHTTPResponseCache;
 	class IExternalDataReader;
 
@@ -130,6 +131,10 @@ namespace Electra
 		 */
 		virtual	TSharedPtrTS<IPlaylistReader> GetManifestReader() = 0;
 
+		/**
+		 * Returns the content steering handler for this player instance.
+		 */
+		virtual TSharedPtrTS<FContentSteeringHandler> GetContentSteeringHandler() = 0;
 
 		/**
 		 * Returns the "Application Event or Metadata Streams" (AEMS) handler.
diff --git a/Engine/Plugins/Media/ElectraPlayer/Source/ElectraPlayerRuntime/Private/Runtime/Player/StreamSegmentReaderCommon.cpp b/Engine/Plugins/Media/ElectraPlayer/Source/ElectraPlayerRuntime/Private/Runtime/Player/StreamSegmentReaderCommon.cpp
--- a/Engine/Plugins/Media/ElectraPlayer/Source/ElectraPlayerRuntime/Private/Runtime/Player/StreamSegmentReaderCommon.cpp
+++ b/Engine/Plugins/Media/ElectraPlayer/Source/ElectraPlayerRuntime/Private/Runtime/Player/StreamSegmentReaderCommon.cpp
@@ -611,7 +611,7 @@ void FStreamSegmentReaderCommon::FStreamHandler::HTTPUpdateStats(const FTimeValu
 	CurrentConnectionInfo = InHTTPRequest->ConnectionInfo;
 	if (CurrentConnectionInfo.EffectiveURL.Len())
 	{
-		DownloadStats.URL = CurrentConnectionInfo.EffectiveURL;
+		DownloadStats.URL.URL = CurrentConnectionInfo.EffectiveURL;
 	}
 	DownloadStats.HTTPStatusCode = CurrentConnectionInfo.StatusInfo.HTTPStatus;
 	DownloadStats.TimeToFirstByte = CurrentConnectionInfo.TimeUntilFirstByte;
@@ -878,9 +878,9 @@ bool FStreamSegmentReaderCommon::FStreamHandler::FetchInitSegment(TSharedPtrTS<F
 	// Set the download stats with the information for downloading an init segment.
 	DownloadStats.StatsID = FMediaInterlockedIncrement(UniqueDownloadID);
 	DownloadStats.SegmentType = Metrics::ESegmentType::Init;
-	DownloadStats.URL = InitURLToUse.URL;
+	DownloadStats.URL = InitURLToUse.Url;
 	DownloadStats.Range = InitURLToUse.Range;
-	DownloadStats.CDN = InitURLToUse.CDN;
+	DownloadStats.SteeringID = InitURLToUse.SteeringID;
 
 	// The presence of an init segment is signalled differently for HLS and DASH.
 	if (InRequest->StreamingProtocol == FStreamSegmentRequestCommon::EStreamingProtocol::DASH)
@@ -893,13 +893,13 @@ bool FStreamSegmentReaderCommon::FStreamHandler::FetchInitSegment(TSharedPtrTS<F
 		check(InRequest->ExpectedContainerFormat != FStreamSegmentRequestCommon::EContainerFormat::Undefined);
 		if (InRequest->ExpectedContainerFormat == FStreamSegmentRequestCommon::EContainerFormat::ISO_14496_12)
 		{
-			if (InitURLToUse.URL.IsEmpty())
+			if (InitURLToUse.Url.URL.IsEmpty())
 			{
 				InitURLToUse = segInfo.MediaURL;
 				// Clear out the media range as we need to scan for the 'moov' box through the entire file.
 				InitURLToUse.Range.Empty();
 				bFetchFromSelfInitializing = true;
-				DownloadStats.URL = InitURLToUse.URL;
+				DownloadStats.URL = InitURLToUse.Url;
 				DownloadStats.Range = InitURLToUse.Range;
 			}
 		}
@@ -917,7 +917,7 @@ bool FStreamSegmentReaderCommon::FStreamHandler::FetchInitSegment(TSharedPtrTS<F
 	}
 
 	// When no init segment is needed, signal completion and return.
-	if (InitURLToUse.URL.IsEmpty())
+	if (InitURLToUse.Url.URL.IsEmpty())
 	{
 		return true;
 	}
@@ -928,7 +928,7 @@ bool FStreamSegmentReaderCommon::FStreamHandler::FetchInitSegment(TSharedPtrTS<F
 	if (EntityCache.IsValid())
 	{
 		IPlayerEntityCache::FCacheItem CachedItem;
-		if (EntityCache->GetCachedEntity(CachedItem, InitURLToUse.URL, InitURLToUse.Range))
+		if (EntityCache->GetCachedEntity(CachedItem, InitURLToUse.Url.URL, InitURLToUse.Range))
 		{
 			if (CachedItem.Parsed14496_12Data.IsValid())
 			{
@@ -967,7 +967,7 @@ bool FStreamSegmentReaderCommon::FStreamHandler::FetchInitSegment(TSharedPtrTS<F
 		TSharedPtrTS<FHTTPResourceRequest> rr = MakeSharedTS<FHTTPResourceRequest>();
 		TSharedPtr<FHTTPResourceRequestCompletionSignal, ESPMode::ThreadSafe> rrsig = FHTTPResourceRequestCompletionSignal::Create();
 
-		rr->Verb(TEXT("GET")).URL(InitURLToUse.URL).Range(InitURLToUse.Range).Headers(ReqHeaders).AcceptEncoding(TEXT("identity"))
+		rr->Verb(TEXT("GET")).URL(InitURLToUse.Url.URL).Range(InitURLToUse.Range).Headers(ReqHeaders).AcceptEncoding(TEXT("identity"))
 			.ConnectionTimeout(FTimeValue().SetFromMilliseconds(5000)).NoDataTimeout(FTimeValue().SetFromMilliseconds(2000))
 			.StreamTypeAndQuality(InRequest->StreamType, InRequest->QualityIndex, InRequest->MaxQualityIndex)
 			.CompletionSignal(rrsig)
@@ -1108,7 +1108,7 @@ bool FStreamSegmentReaderCommon::FStreamHandler::FetchInitSegment(TSharedPtrTS<F
 		TArray<UtilsMP4::FMP4RootBoxLocator::FBoxInfo> Boxes;
 		const TArray<uint32> FirstBox { UtilsMP4::Make4CC('f','t','y','p') };
 		const TArray<uint32> MoovBox { UtilsMP4::Make4CC('m','o','o','v') };
-		bool bSuccess = BoxLocator.LocateRootBoxes(Boxes, PlayerSessionService->GetHTTPManager(), InitURLToUse.URL, FirstBox, MoovBox, MoovBox, UtilsMP4::FMP4RootBoxLocator::FCancellationCheckDelegate::CreateLambda([&]()
+		bool bSuccess = BoxLocator.LocateRootBoxes(Boxes, PlayerSessionService->GetHTTPManager(), InitURLToUse.Url.URL, FirstBox, MoovBox, MoovBox, UtilsMP4::FMP4RootBoxLocator::FCancellationCheckDelegate::CreateLambda([&]()
 		{
 			return HasReadBeenAborted();
 		}));
@@ -1144,7 +1144,7 @@ bool FStreamSegmentReaderCommon::FStreamHandler::FetchInitSegment(TSharedPtrTS<F
 	// Set the response headers with the entity cache.
 	if (EntityCache.IsValid())
 	{
-		EntityCache->SetRecentResponseHeaders(IPlayerEntityCache::EEntityType::Segment, InitURLToUse.URL, CurrentConnectionInfo.ResponseHeaders);
+		EntityCache->SetRecentResponseHeaders(IPlayerEntityCache::EEntityType::Segment, InitURLToUse.Url.URL, CurrentConnectionInfo.ResponseHeaders);
 	}
 
 	// Parse the init segment.
@@ -1167,7 +1167,7 @@ bool FStreamSegmentReaderCommon::FStreamHandler::FetchInitSegment(TSharedPtrTS<F
 					if (EntityCache.IsValid())
 					{
 						IPlayerEntityCache::FCacheItem CacheItem;
-						CacheItem.URL = InitURLToUse.URL;
+						CacheItem.URL = InitURLToUse.Url.URL;
 						CacheItem.Range = InitURLToUse.Range;
 						CacheItem.Parsed14496_12Data = MP4InitSeg;
 						EntityCache->CacheEntity(CacheItem);
@@ -1180,7 +1180,7 @@ bool FStreamSegmentReaderCommon::FStreamHandler::FetchInitSegment(TSharedPtrTS<F
 				{
 					DownloadStats.bParseFailure = true;
 					StreamSelector->ReportDownloadEnd(DownloadStats);
-					SetError(FString::Printf(TEXT("Track preparation of init segment \"%s\" failed"), *InitURLToUse.URL), INTERNAL_SEG_ERROR_INIT_SEGMENT_PARSE_ERROR);
+					SetError(FString::Printf(TEXT("Track preparation of init segment \"%s\" failed"), *InitURLToUse.Url.URL), INTERNAL_SEG_ERROR_INIT_SEGMENT_PARSE_ERROR);
 					return false;
 				}
 			}
@@ -1188,7 +1188,7 @@ bool FStreamSegmentReaderCommon::FStreamHandler::FetchInitSegment(TSharedPtrTS<F
 			{
 				DownloadStats.bParseFailure = true;
 				StreamSelector->ReportDownloadEnd(DownloadStats);
-				SetError(FString::Printf(TEXT("Parse error of init segment \"%s\""), *InitURLToUse.URL), INTERNAL_SEG_ERROR_INIT_SEGMENT_PARSE_ERROR);
+				SetError(FString::Printf(TEXT("Parse error of init segment \"%s\""), *InitURLToUse.Url.URL), INTERNAL_SEG_ERROR_INIT_SEGMENT_PARSE_ERROR);
 				return false;
 			}
 			break;
@@ -1207,7 +1207,7 @@ bool FStreamSegmentReaderCommon::FStreamHandler::FetchInitSegment(TSharedPtrTS<F
 					if (EntityCache.IsValid())
 					{
 						IPlayerEntityCache::FCacheItem CacheItem;
-						CacheItem.URL = InitURLToUse.URL;
+						CacheItem.URL = InitURLToUse.Url.URL;
 						CacheItem.Range = InitURLToUse.Range;
 						CacheItem.ParsedMatroskaData = MKVInitSeg;
 						EntityCache->CacheEntity(CacheItem);
@@ -1220,7 +1220,7 @@ bool FStreamSegmentReaderCommon::FStreamHandler::FetchInitSegment(TSharedPtrTS<F
 				{
 					DownloadStats.bParseFailure = true;
 					StreamSelector->ReportDownloadEnd(DownloadStats);
-					SetError(FString::Printf(TEXT("Track preparation of init segment \"%s\" failed"), *InitURLToUse.URL), INTERNAL_SEG_ERROR_INIT_SEGMENT_PARSE_ERROR);
+					SetError(FString::Printf(TEXT("Track preparation of init segment \"%s\" failed"), *InitURLToUse.Url.URL), INTERNAL_SEG_ERROR_INIT_SEGMENT_PARSE_ERROR);
 					return false;
 				}
 			}
@@ -1228,7 +1228,7 @@ bool FStreamSegmentReaderCommon::FStreamHandler::FetchInitSegment(TSharedPtrTS<F
 			{
 				DownloadStats.bParseFailure = true;
 				StreamSelector->ReportDownloadEnd(DownloadStats);
-				SetError(FString::Printf(TEXT("Parse error of init segment \"%s\""), *InitURLToUse.URL), INTERNAL_SEG_ERROR_INIT_SEGMENT_PARSE_ERROR);
+				SetError(FString::Printf(TEXT("Parse error of init segment \"%s\""), *InitURLToUse.Url.URL), INTERNAL_SEG_ERROR_INIT_SEGMENT_PARSE_ERROR);
 				return false;
 			}
 			break;
@@ -1250,7 +1250,7 @@ bool FStreamSegmentReaderCommon::FStreamHandler::FetchInitSegment(TSharedPtrTS<F
 					if (EntityCache.IsValid())
 					{
 						IPlayerEntityCache::FCacheItem CacheItem;
-						CacheItem.URL = InitURLToUse.URL;
+						CacheItem.URL = InitURLToUse.Url.URL;
 						CacheItem.Range = InitURLToUse.Range;
 						CacheItem.RawPayloadData = TSPacketBlob;
 						EntityCache->CacheEntity(CacheItem);
@@ -1262,7 +1262,7 @@ bool FStreamSegmentReaderCommon::FStreamHandler::FetchInitSegment(TSharedPtrTS<F
 			}
 			DownloadStats.bParseFailure = true;
 			StreamSelector->ReportDownloadEnd(DownloadStats);
-			SetError(FString::Printf(TEXT("Track preparation of init segment \"%s\" failed"), *InitURLToUse.URL), INTERNAL_SEG_ERROR_INIT_SEGMENT_PARSE_ERROR);
+			SetError(FString::Printf(TEXT("Track preparation of init segment \"%s\" failed"), *InitURLToUse.Url.URL), INTERNAL_SEG_ERROR_INIT_SEGMENT_PARSE_ERROR);
 			return false;
 		}
 		default:
@@ -1280,9 +1280,9 @@ FStreamSegmentReaderCommon::FStreamHandler::EMediaSegmentTriggerResult FStreamSe
 	DownloadStats.ResetOutput();
 	DownloadStats.StatsID = FMediaInterlockedIncrement(UniqueDownloadID);
 	DownloadStats.SegmentType = Metrics::ESegmentType::Media;
-	DownloadStats.URL = InRequest->Segment.MediaURL.URL;
+	DownloadStats.URL = InRequest->Segment.MediaURL.Url;
 	DownloadStats.Range = InRequest->Segment.MediaURL.Range;
-	DownloadStats.CDN = InRequest->Segment.MediaURL.CDN;
+	DownloadStats.SteeringID = InRequest->Segment.MediaURL.SteeringID;
 
 	Parameters.EventListener->OnFragmentOpen(InRequest);
 
@@ -1318,7 +1318,7 @@ FStreamSegmentReaderCommon::FStreamHandler::EMediaSegmentTriggerResult FStreamSe
 	HTTPRequest->ProgressListener = ProgressListener;
 	HTTPRequest->ResponseCache = PlayerSessionService->GetHTTPResponseCache();
 	HTTPRequest->ExternalDataReader = PlayerSessionService->GetExternalDataReader();
-	HTTPRequest->Parameters.URL = InRequest->Segment.MediaURL.URL;
+	HTTPRequest->Parameters.URL = InRequest->Segment.MediaURL.Url.URL;
 	HTTPRequest->Parameters.Range.Set(InRequest->Segment.MediaURL.Range);
 	HTTPRequest->Parameters.StreamType = InRequest->StreamType;
 	HTTPRequest->Parameters.QualityIndex = InRequest->QualityIndex;
@@ -1442,7 +1442,7 @@ void FStreamSegmentReaderCommon::FStreamHandler::HandleCommonMediaEnd(EHandleRes
 	DownloadStats.bWasAborted = bAbortedByABR;
 	DownloadStats.bWasSuccessful = !bHasErrored && !bAbortedByABR;
 	DownloadStats.bDidTimeout = CurrentConnectionInfo.StatusInfo.ErrorCode == HTTP::EStatusErrorCode::ERRCODE_HTTP_CONNECTION_TIMEOUT;
-	DownloadStats.URL = CurrentConnectionInfo.EffectiveURL;
+	DownloadStats.URL.URL = CurrentConnectionInfo.EffectiveURL;
 	DownloadStats.HTTPStatusCode = CurrentConnectionInfo.StatusInfo.HTTPStatus;
 	DownloadStats.DurationDownloaded = DurationDownloaded.GetAsSeconds();
 	DownloadStats.DurationDelivered = DurationDelivered.GetAsSeconds();
@@ -1688,7 +1688,7 @@ FStreamSegmentReaderCommon::FStreamHandler::EHandleResult FStreamSegmentReaderCo
 	FTimeValue TimeOffset = InRequest->PeriodStart + InRequest->AST + InRequest->AdditionalAdjustmentTime;
 	bool bDoNotTruncateAtPresentationEnd = PlayerSessionService->GetOptionValue(OptionKeyDoNotTruncateAtPresentationEnd).SafeGetBool(false);
 	bool bIsLastSegment = InRequest->Segment.bIsLastInPeriod;
-	const FString& RequestURL(InRequest->Segment.MediaURL.URL);
+	const FString& RequestURL(InRequest->Segment.MediaURL.Url.URL);
 	int64 LastSuccessfulFilePos = 0;
 	bool bSkippedBecauseOfTimestampCheck = false;
 	int32 StreamTypeAUCount[4] { 0,0,0,0 };
@@ -2216,7 +2216,7 @@ FStreamSegmentReaderCommon::FStreamHandler::EHandleResult FStreamSegmentReaderCo
 	bool bDone = false;
 	FTimeValue TimeOffset = InRequest->PeriodStart + InRequest->AST + InRequest->AdditionalAdjustmentTime;
 	bool bDoNotTruncateAtPresentationEnd = PlayerSessionService->GetOptionValue(OptionKeyDoNotTruncateAtPresentationEnd).SafeGetBool(false);
-	const FString& RequestURL(InRequest->Segment.MediaURL.URL);
+	const FString& RequestURL(InRequest->Segment.MediaURL.Url.URL);
 	FTimeValue PTO(InRequest->Segment.PTO, InRequest->Segment.Timescale);
 	FTimeValue EarliestPTS(InRequest->Segment.MediaLocalFirstAUTime, InRequest->Segment.Timescale, InRequest->TimestampSequenceIndex);
 	EarliestPTS += TimeOffset - PTO;
@@ -3328,9 +3328,9 @@ FStreamSegmentReaderCommon::FStreamHandler::EHandleResult FStreamSegmentReaderCo
 	DownloadStats.ResetOutput();
 	DownloadStats.StatsID = FMediaInterlockedIncrement(UniqueDownloadID);
 	DownloadStats.SegmentType = Metrics::ESegmentType::Media;
-	DownloadStats.URL = InRequest->Segment.MediaURL.URL;
+	DownloadStats.URL = InRequest->Segment.MediaURL.Url;
 	DownloadStats.Range = InRequest->Segment.MediaURL.Range;
-	DownloadStats.CDN = InRequest->Segment.MediaURL.CDN;
+	DownloadStats.SteeringID = InRequest->Segment.MediaURL.SteeringID;
 
 	// Check with the entity cache if we already have it from earlier.
 	TSharedPtrTS<IPlayerEntityCache> EntityCache = PlayerSessionService ? PlayerSessionService->GetEntityCache() : nullptr;
@@ -3338,7 +3338,7 @@ FStreamSegmentReaderCommon::FStreamHandler::EHandleResult FStreamSegmentReaderCo
 	if (EntityCache.IsValid())
 	{
 		IPlayerEntityCache::FCacheItem CachedItem;
-		if (EntityCache->GetCachedEntity(CachedItem, InRequest->Segment.MediaURL.URL, InRequest->Segment.MediaURL.Range))
+		if (EntityCache->GetCachedEntity(CachedItem, InRequest->Segment.MediaURL.Url.URL, InRequest->Segment.MediaURL.Range))
 		{
 			if (CachedItem.RawPayloadData.IsValid())
 			{
@@ -3361,7 +3361,7 @@ FStreamSegmentReaderCommon::FStreamHandler::EHandleResult FStreamSegmentReaderCo
 		TSharedPtrTS<FHTTPResourceRequest> rr = MakeSharedTS<FHTTPResourceRequest>();
 		TSharedPtr<FHTTPResourceRequestCompletionSignal, ESPMode::ThreadSafe> rrsig = FHTTPResourceRequestCompletionSignal::Create();
 
-		rr->Verb(TEXT("GET")).URL(InRequest->Segment.MediaURL.URL).Range(InRequest->Segment.MediaURL.Range).Headers(ReqHeaders)
+		rr->Verb(TEXT("GET")).URL(InRequest->Segment.MediaURL.Url.URL).Range(InRequest->Segment.MediaURL.Range).Headers(ReqHeaders)
 			.ConnectionTimeout(FTimeValue().SetFromMilliseconds(5000)).NoDataTimeout(FTimeValue().SetFromMilliseconds(2000))
 			.StreamTypeAndQuality(InRequest->StreamType, InRequest->QualityIndex, InRequest->MaxQualityIndex)
 			.CompletionSignal(rrsig)
@@ -3398,11 +3398,11 @@ FStreamSegmentReaderCommon::FStreamHandler::EHandleResult FStreamSegmentReaderCo
 		if (EntityCache.IsValid())
 		{
 			// Set the response headers with the entity cache.
-			EntityCache->SetRecentResponseHeaders(IPlayerEntityCache::EEntityType::Segment, InRequest->Segment.MediaURL.URL, CurrentConnectionInfo.ResponseHeaders);
+			EntityCache->SetRecentResponseHeaders(IPlayerEntityCache::EEntityType::Segment, InRequest->Segment.MediaURL.Url.URL, CurrentConnectionInfo.ResponseHeaders);
 
 			// Cache the sideloaded file as well.
 			IPlayerEntityCache::FCacheItem CacheItem;
-			CacheItem.URL = InRequest->Segment.MediaURL.URL;
+			CacheItem.URL = InRequest->Segment.MediaURL.Url.URL;
 			CacheItem.Range = InRequest->Segment.MediaURL.Range;
 			CacheItem.RawPayloadData = SideloadedData;
 			EntityCache->CacheEntity(CacheItem);
diff --git a/Engine/Plugins/Media/ElectraPlayer/Source/ElectraPlayerRuntime/Private/Runtime/Player/StreamSegmentReaderCommon.h b/Engine/Plugins/Media/ElectraPlayer/Source/ElectraPlayerRuntime/Private/Runtime/Player/StreamSegmentReaderCommon.h
--- a/Engine/Plugins/Media/ElectraPlayer/Source/ElectraPlayerRuntime/Private/Runtime/Player/StreamSegmentReaderCommon.h
+++ b/Engine/Plugins/Media/ElectraPlayer/Source/ElectraPlayerRuntime/Private/Runtime/Player/StreamSegmentReaderCommon.h
@@ -41,10 +41,10 @@ struct FSegmentInformationCommon
 {
 	struct FURL
 	{
-		FString URL;
+		FMediaURL Url;
 		FString Range;
-		FString CDN;
 		FString CustomHeader;
+		int64 SteeringID = 0;
 	};
 	struct FInbandEventStream
 	{
@@ -72,7 +72,6 @@ struct FSegmentInformationCommon
 	int64 MediaLocalLastAUTime = 0;				// Time at which the last AU to use in thie segment ends in media local time
 	int64 MediaLocalFirstPTS = 0;
 	uint32 Timescale = 0;						// Local media timescale
-	TMediaOptionalValue<bool> bAvailabilityTimeComplete;
 	bool bLowLatencyChunkedEncodingExpected = false;
 	bool bFrameAccuracyRequired = false;		// true if the segment was located for frame accurate seeking.
 	bool bIsSideload = false;					// true if this is a side-loaded resource to be fetched and cached.
@@ -90,8 +89,6 @@ struct FSegmentInformationCommon
 	EContainerType ContainerType = EContainerType::ISO14496_12;
 	int64 NumberOfBytes = 0;
 	int64 FirstByteOffset = 0;
-// NOT NEEDED, REMOVE!
-	FTimeValue PeriodLocalSegmentStartTime;
 
 	FTimeValue CalculateASAST(const FTimeValue& AST, const FTimeValue& PeriodStart, bool bIsStatic)
 	{
diff --git a/Engine/Plugins/Media/ElectraPlayer/Source/ElectraPlayerRuntime/Private/Runtime/Player/mkv/ManifestMKV.h b/Engine/Plugins/Media/ElectraPlayer/Source/ElectraPlayerRuntime/Private/Runtime/Player/mkv/ManifestMKV.h
--- a/Engine/Plugins/Media/ElectraPlayer/Source/ElectraPlayerRuntime/Private/Runtime/Player/mkv/ManifestMKV.h
+++ b/Engine/Plugins/Media/ElectraPlayer/Source/ElectraPlayerRuntime/Private/Runtime/Player/mkv/ManifestMKV.h
@@ -66,6 +66,8 @@ public:
 	void UpdateDynamicRefetchCounter() override;
 	void TriggerClockSync(EClockSyncType InClockSyncType) override;
 	void TriggerPlaylistRefresh() override;
+	void ReachedStableBuffer() override
+	{ }
 
 	IStreamReader* CreateStreamReaderHandler() override;
 	FResult FindPlayPeriod(TSharedPtrTS<IPlayPeriod>& OutPlayPeriod, const FPlayStartPosition& StartPosition, ESearchType SearchType) override;
diff --git a/Engine/Plugins/Media/ElectraPlayer/Source/ElectraPlayerRuntime/Private/Runtime/Player/mkv/StreamReaderMKV.cpp b/Engine/Plugins/Media/ElectraPlayer/Source/ElectraPlayerRuntime/Private/Runtime/Player/mkv/StreamReaderMKV.cpp
--- a/Engine/Plugins/Media/ElectraPlayer/Source/ElectraPlayerRuntime/Private/Runtime/Player/mkv/StreamReaderMKV.cpp
+++ b/Engine/Plugins/Media/ElectraPlayer/Source/ElectraPlayerRuntime/Private/Runtime/Player/mkv/StreamReaderMKV.cpp
@@ -241,7 +241,7 @@ void FStreamReaderMKV::HTTPUpdateStats(const FTimeValue& CurrentTime, const IEle
 		Metrics::FSegmentDownloadStats& ds = SegmentRequest->DownloadStats;
 		if (Request->ConnectionInfo.EffectiveURL.Len())
 		{
-			ds.URL = Request->ConnectionInfo.EffectiveURL;
+			ds.URL.URL = Request->ConnectionInfo.EffectiveURL;
 		}
 		ds.HTTPStatusCode = Request->ConnectionInfo.StatusInfo.HTTPStatus;
 		ds.TimeToFirstByte = Request->ConnectionInfo.TimeUntilFirstByte;
@@ -349,7 +349,7 @@ void FStreamReaderMKV::HandleRequest()
 	ds.ByteSize = -1;
 	ds.NumBytesDownloaded = 0;
 	ds.bInsertedFillerData = false;
-	ds.URL = TimelineAsset->GetMediaURL();
+	ds.URL.URL = TimelineAsset->GetMediaURL();
 	ds.bIsMissingSegment = false;
 	ds.bParseFailure = false;
 	ds.RetryNumber = Request->NumOverallRetries;
@@ -690,7 +690,7 @@ void FStreamReaderMKV::HandleRequest()
 		ds.FailureReason = ClusterParser->GetLastError().GetMessage();
 	}
 	ds.bWasSuccessful = !bHasErrored;
-	ds.URL = Request->ConnectionInfo.EffectiveURL;
+	ds.URL.URL = Request->ConnectionInfo.EffectiveURL;
 	ds.HTTPStatusCode = Request->ConnectionInfo.StatusInfo.HTTPStatus;
 	ds.DurationDownloaded = DurationRead;
 	ds.DurationDelivered = DurationDelivered;
diff --git a/Engine/Plugins/Media/ElectraPlayer/Source/ElectraPlayerRuntime/Private/Runtime/Player/mp4/ManifestMP4.h b/Engine/Plugins/Media/ElectraPlayer/Source/ElectraPlayerRuntime/Private/Runtime/Player/mp4/ManifestMP4.h
--- a/Engine/Plugins/Media/ElectraPlayer/Source/ElectraPlayerRuntime/Private/Runtime/Player/mp4/ManifestMP4.h
+++ b/Engine/Plugins/Media/ElectraPlayer/Source/ElectraPlayerRuntime/Private/Runtime/Player/mp4/ManifestMP4.h
@@ -67,6 +67,8 @@ public:
 	void UpdateDynamicRefetchCounter() override;
 	void TriggerClockSync(EClockSyncType InClockSyncType) override;
 	void TriggerPlaylistRefresh() override;
+	void ReachedStableBuffer() override
+	{ }
 
 	IStreamReader* CreateStreamReaderHandler() override;
 	FResult FindPlayPeriod(TSharedPtrTS<IPlayPeriod>& OutPlayPeriod, const FPlayStartPosition& StartPosition, ESearchType SearchType) override;
diff --git a/Engine/Plugins/Media/ElectraPlayer/Source/ElectraPlayerRuntime/Private/Runtime/Player/mp4/StreamReaderMP4.cpp b/Engine/Plugins/Media/ElectraPlayer/Source/ElectraPlayerRuntime/Private/Runtime/Player/mp4/StreamReaderMP4.cpp
--- a/Engine/Plugins/Media/ElectraPlayer/Source/ElectraPlayerRuntime/Private/Runtime/Player/mp4/StreamReaderMP4.cpp
+++ b/Engine/Plugins/Media/ElectraPlayer/Source/ElectraPlayerRuntime/Private/Runtime/Player/mp4/StreamReaderMP4.cpp
@@ -355,7 +355,7 @@ void FStreamReaderMP4::HandleRequest()
 	ds.ByteSize 		   = -1;
 	ds.NumBytesDownloaded  = 0;
 	ds.bInsertedFillerData = false;
-	ds.URL  			   = TimelineAsset->GetMediaURL();
+	ds.URL.URL 			   = TimelineAsset->GetMediaURL();
 	ds.bIsMissingSegment   = false;
 	ds.bParseFailure	   = false;
 	ds.RetryNumber  	   = Request->NumOverallRetries;
@@ -675,10 +675,6 @@ void FStreamReaderMP4::HandleRequest()
 	DurationSuccessfullyDelivered = PrimaryTrack.DurationSuccessfullyDelivered;
 
 	// Set up download stat fields.
-	if (Request->ConnectionInfo.EffectiveURL.Len())
-	{
-		ds.URL = Request->ConnectionInfo.EffectiveURL;
-	}
 	ds.HTTPStatusCode = Request->ConnectionInfo.StatusInfo.HTTPStatus;
 	ds.TimeToFirstByte = Request->ConnectionInfo.TimeUntilFirstByte;
 	ds.TimeToDownload = (Request->ConnectionInfo.RequestEndTime - Request->ConnectionInfo.RequestStartTime).GetAsSeconds();
@@ -691,7 +687,7 @@ void FStreamReaderMP4::HandleRequest()
 		ds.FailureReason = ParsingErrorMessage;
 	}
 	ds.bWasSuccessful     = !bHasErrored;
-	ds.URL  			  = Request->ConnectionInfo.EffectiveURL;
+	ds.URL.URL 			  = Request->ConnectionInfo.EffectiveURL;
 	ds.HTTPStatusCode     = Request->ConnectionInfo.StatusInfo.HTTPStatus;
 	ds.DurationDownloaded = DurationSuccessfullyRead.GetAsSeconds();
 	ds.DurationDelivered  = DurationSuccessfullyDelivered.GetAsSeconds();
diff --git a/Engine/Plugins/Media/ElectraPlayer/Source/ElectraPlayerRuntime/Private/Runtime/Player/mpegaudio/StreamReaderMPEGAudio.cpp b/Engine/Plugins/Media/ElectraPlayer/Source/ElectraPlayerRuntime/Private/Runtime/Player/mpegaudio/StreamReaderMPEGAudio.cpp
--- a/Engine/Plugins/Media/ElectraPlayer/Source/ElectraPlayerRuntime/Private/Runtime/Player/mpegaudio/StreamReaderMPEGAudio.cpp
+++ b/Engine/Plugins/Media/ElectraPlayer/Source/ElectraPlayerRuntime/Private/Runtime/Player/mpegaudio/StreamReaderMPEGAudio.cpp
@@ -370,7 +370,7 @@ void FStreamReaderMPEGAudio::HandleRequest()
 	ds.ByteSize = -1;
 	ds.NumBytesDownloaded  = 0;
 	ds.bInsertedFillerData = false;
-	ds.URL = TimelineAsset->GetMediaURL();
+	ds.URL.URL = TimelineAsset->GetMediaURL();
 	ds.bIsMissingSegment = false;
 	ds.bParseFailure = false;
 	ds.RetryNumber = Request->NumOverallRetries;
@@ -802,10 +802,6 @@ void FStreamReaderMPEGAudio::HandleRequest()
 	}
 
 	// Set up download stat fields.
-	if (Request->ConnectionInfo.EffectiveURL.Len())
-	{
-		ds.URL = Request->ConnectionInfo.EffectiveURL;
-	}
 	ds.HTTPStatusCode = Request->ConnectionInfo.StatusInfo.HTTPStatus;
 	ds.TimeToFirstByte = Request->ConnectionInfo.TimeUntilFirstByte;
 	ds.TimeToDownload = (Request->ConnectionInfo.RequestEndTime - Request->ConnectionInfo.RequestStartTime).GetAsSeconds();
@@ -818,7 +814,7 @@ void FStreamReaderMPEGAudio::HandleRequest()
 		ds.FailureReason = ParsingErrorMessage;
 	}
 	ds.bWasSuccessful = !bHasErrored;
-	ds.URL = Request->ConnectionInfo.EffectiveURL;
+	ds.URL.URL = Request->ConnectionInfo.EffectiveURL;
 	ds.HTTPStatusCode = Request->ConnectionInfo.StatusInfo.HTTPStatus;
 	ds.DurationDownloaded = DurationSuccessfullyRead.GetAsSeconds();
 	ds.DurationDelivered = DurationSuccessfullyDelivered.GetAsSeconds();
diff --git a/Engine/Plugins/Media/ElectraPlayer/Source/ElectraPlayerRuntime/Private/Runtime/StreamTypes.h b/Engine/Plugins/Media/ElectraPlayer/Source/ElectraPlayerRuntime/Private/Runtime/StreamTypes.h
--- a/Engine/Plugins/Media/ElectraPlayer/Source/ElectraPlayerRuntime/Private/Runtime/StreamTypes.h
+++ b/Engine/Plugins/Media/ElectraPlayer/Source/ElectraPlayerRuntime/Private/Runtime/StreamTypes.h
@@ -6,6 +6,7 @@
 #include "PlayerTime.h"
 #include "ParameterDictionary.h"
 #include "MediaVideoDecoderOutput.h"
+#include "MediaURLType.h"
 
 namespace Electra
 {
