Electra: Fixed self-initializing DASH segment

#rnx

diff --git a/Engine/Plugins/Media/ElectraPlayer/Source/ElectraPlayerRuntime/Private/Runtime/Player/DASH/ManifestDASH.cpp b/Engine/Plugins/Media/ElectraPlayer/Source/ElectraPlayerRuntime/Private/Runtime/Player/DASH/ManifestDASH.cpp
--- a/Engine/Plugins/Media/ElectraPlayer/Source/ElectraPlayerRuntime/Private/Runtime/Player/DASH/ManifestDASH.cpp
+++ b/Engine/Plugins/Media/ElectraPlayer/Source/ElectraPlayerRuntime/Private/Runtime/Player/DASH/ManifestDASH.cpp
@@ -2472,6 +2472,10 @@ bool FManifestDASHInternal::FRepresentation::PrepareDownloadURLs(IPlayerSessionS
 		InOutSegmentInfo.InitializationURL.Range = InitializationSegmentURLRange;
 		InOutSegmentInfo.InitializationURL.Url.URL = DASHUrlHelpers::ApplyAnnexEByteRange(InOutSegmentInfo.InitializationURL.Url.URL, InOutSegmentInfo.InitializationURL.Range, ChosenInit.AdditionalParams.GetValue(DASHUrlHelpers::SteerOption_ByteRange).SafeGetFString());
 	}
+/*
+	Do not set up the init segment that way. Leave it empty as the segment reader knows that it needs to look for the init data in the media file.
+	If we set this explicitly then it won't know that and load the entire file.
+
 	else
 	{
 		// The init segment is part of the media segment, so we can just copy this across.
@@ -2480,6 +2484,7 @@ bool FManifestDASHInternal::FRepresentation::PrepareDownloadURLs(IPlayerSessionS
 		InOutSegmentInfo.InitializationURL.Range = InitializationSegmentURLRange;
 		InOutSegmentInfo.InitializationURL.Url.URL = DASHUrlHelpers::ApplyAnnexEByteRange(InOutSegmentInfo.InitializationURL.Url.URL, InOutSegmentInfo.InitializationURL.Range, ChosenMedia.AdditionalParams.GetValue(DASHUrlHelpers::SteerOption_ByteRange).SafeGetFString());
 	}
+*/
 
 	// Finally apply the ATO of the media segment as the ATO in general.
 	// We load init segments only when requesting media segments at which time the init segment needs
