ElectraProtron: Changed player to automatically select the first video and audio track if they exist. This is to ensure that there are valid selected track indices right away as the selection of tracks through the facade will be delayed until the `TracksChanged` event is processed, which is not even guaranteed to select a track if the default values in FMediaPlayerTrackOptions would get changed to -1 for video/audio.

#rnx

diff --git a/Engine/Plugins/Media/ElectraPlayer/Source/ElectraProtron/Private/Player/ElectraProtronPlayer.cpp b/Engine/Plugins/Media/ElectraPlayer/Source/ElectraProtron/Private/Player/ElectraProtronPlayer.cpp
--- a/Engine/Plugins/Media/ElectraPlayer/Source/ElectraProtron/Private/Player/ElectraProtronPlayer.cpp
+++ b/Engine/Plugins/Media/ElectraPlayer/Source/ElectraProtron/Private/Player/ElectraProtronPlayer.cpp
@@ -178,6 +178,10 @@ bool FElectraProtronPlayer::GetPlayerFeatureFlag(EFeatureFlag InWhichFeature) co
 		{
 			return true;
 		}
+		case EFeatureFlag::PlayerSelectsDefaultTracks:
+		{
+			return true;
+		}
 		default:
 		{
 			return IMediaPlayer::GetPlayerFeatureFlag(InWhichFeature);
diff --git a/Engine/Plugins/Media/ElectraPlayer/Source/ElectraProtron/Private/Player/ElectraProtronPlayerImpl.cpp b/Engine/Plugins/Media/ElectraPlayer/Source/ElectraProtron/Private/Player/ElectraProtronPlayerImpl.cpp
--- a/Engine/Plugins/Media/ElectraPlayer/Source/ElectraProtron/Private/Player/ElectraProtronPlayerImpl.cpp
+++ b/Engine/Plugins/Media/ElectraPlayer/Source/ElectraProtron/Private/Player/ElectraProtronPlayerImpl.cpp
@@ -145,6 +145,10 @@ uint32 FElectraProtronPlayer::FImpl::Run()
 
 						VideoDecoderThread.StartThread(open.Param, SharedPlayParams);
 						AudioDecoderThread.StartThread(open.Param, SharedPlayParams);
+
+						// Select the first video and audio track by default (if they exist).
+						SelectTrack(EMediaTrackType::Video, 0);
+						SelectTrack(EMediaTrackType::Audio, 0);
 					}
 					break;
 				}
