Electra: Allowing for an empty PATHWAY-PRIORITY array in the content steering manifest to keep the current ones as providing an empty list does not make a lot of sense

#rnx

diff --git a/Engine/Plugins/Media/ElectraPlayer/Source/ElectraPlayerRuntime/Private/Runtime/Player/ContentSteeringHandler.cpp b/Engine/Plugins/Media/ElectraPlayer/Source/ElectraPlayerRuntime/Private/Runtime/Player/ContentSteeringHandler.cpp
--- a/Engine/Plugins/Media/ElectraPlayer/Source/ElectraPlayerRuntime/Private/Runtime/Player/ContentSteeringHandler.cpp
+++ b/Engine/Plugins/Media/ElectraPlayer/Source/ElectraPlayerRuntime/Private/Runtime/Player/ContentSteeringHandler.cpp
@@ -392,15 +392,20 @@ void FContentSteeringHandler::UpdateWithSteeringServerResponse(const FString& In
 				if ((SteerParams->TryGetArrayField(TEXT("PATHWAY-PRIORITY"), PATHWAY_PRIORITIES)) ||
 					(StreamingProtocol == EStreamingProtocol::DASH && SteerParams->TryGetArrayField(TEXT("SERVICE-LOCATION-PRIORITY"), PATHWAY_PRIORITIES)))
 				{
-					CurrentCDNPriorities.Empty();
+					TArray<FString> NewPriorities;
 					for(auto& pwpri : *PATHWAY_PRIORITIES)
 					{
 						FString pw;
 						if (pwpri->TryGetString(pw))
 						{
-							CurrentCDNPriorities.Emplace(MoveTemp(pw));
+							NewPriorities.Emplace(MoveTemp(pw));
 						}
 					}
+					// If there are no pathways given we keep using the current ones.
+					if (NewPriorities.Num())
+					{
+						CurrentCDNPriorities = MoveTemp(NewPriorities);
+					}
 				}
 				// In HLS the PATHWAY-PRIORITY is mandatory while it's optional in DASH.
 				else if (StreamingProtocol == EStreamingProtocol::HLS)
