Media: Added a new "MediaBufferingComplete" event. This is sent by media players that also send the "MediaBuffering" event, but the complete event is still optional and is not sent by every media player. This new event should be considered an optional piece of information only.
Removed events that are either not sent by any player or no one is listening for.

#rnx

diff --git a/Engine/Plugins/Experimental/TextureMediaPlayer/Source/TextureMediaPlayer/Private/TextureMediaPlayer.cpp b/Engine/Plugins/Experimental/TextureMediaPlayer/Source/TextureMediaPlayer/Private/TextureMediaPlayer.cpp
--- a/Engine/Plugins/Experimental/TextureMediaPlayer/Source/TextureMediaPlayer/Private/TextureMediaPlayer.cpp
+++ b/Engine/Plugins/Experimental/TextureMediaPlayer/Source/TextureMediaPlayer/Private/TextureMediaPlayer.cpp
@@ -428,7 +428,6 @@ void FTextureMediaPlayer::OnFrame(const TArray<uint8>& TextureBuffer, FIntPoint
 			//       being ticked the time it uses for comparison will not have advanced and nothing
 			//       will get purged.
 			MediaSamples->PurgeOutdatedVideoSamples(FMediaTimeStamp(FTimespan::MaxValue()), false, FTimespan::Zero());
-			EventSink.ReceiveMediaEvent(EMediaEvent::Internal_ResetForDiscontinuity);
 		}
 
 		FElectraTextureSampleRef TextureSample = OutputTexturePool.AcquireShared();
@@ -481,7 +480,6 @@ void FTextureMediaPlayer::OnFrame(FTextureRHIRef TextureRHIRef, TRefCountPtr<ID3
 			//       being ticked the time it uses for comparison will not have advanced and nothing
 			//       will get purged.
 			MediaSamples->PurgeOutdatedVideoSamples(FMediaTimeStamp(FTimespan::MaxValue()), false, FTimespan::Zero());
-			EventSink.ReceiveMediaEvent(EMediaEvent::Internal_ResetForDiscontinuity);
 		}
 
 		FElectraTextureSampleRef TextureSample = OutputTexturePool.AcquireShared();
diff --git a/Engine/Plugins/Media/ElectraPlayer/Source/ElectraPlayerPlugin/Private/ElectraPlayerPlugin.cpp b/Engine/Plugins/Media/ElectraPlayer/Source/ElectraPlayerPlugin/Private/ElectraPlayerPlugin.cpp
--- a/Engine/Plugins/Media/ElectraPlayer/Source/ElectraPlayerPlugin/Private/ElectraPlayerPlugin.cpp
+++ b/Engine/Plugins/Media/ElectraPlayer/Source/ElectraPlayerPlugin/Private/ElectraPlayerPlugin.cpp
@@ -59,6 +59,7 @@ FElectraPlayerPlugin::FElectraPlayerPlugin()
 {
 	// Make sure a few assumptions are correct...
 	static_assert((int32)EMediaEvent::MediaBuffering == (int32)IElectraPlayerAdapterDelegate::EPlayerEvent::MediaBuffering, "check alignment of both enums");
+	static_assert((int32)EMediaEvent::MediaBufferingComplete == (int32)IElectraPlayerAdapterDelegate::EPlayerEvent::MediaBufferingComplete, "check alignment of both enums");
 	static_assert((int32)EMediaEvent::MediaClosed == (int32)IElectraPlayerAdapterDelegate::EPlayerEvent::MediaClosed, "check alignment of both enums");
 	static_assert((int32)EMediaEvent::MediaConnecting == (int32)IElectraPlayerAdapterDelegate::EPlayerEvent::MediaConnecting, "check alignment of both enums");
 	static_assert((int32)EMediaEvent::MediaOpened == (int32)IElectraPlayerAdapterDelegate::EPlayerEvent::MediaOpened, "check alignment of both enums");
@@ -70,9 +71,6 @@ FElectraPlayerPlugin::FElectraPlayerPlugin()
 	static_assert((int32)EMediaEvent::TracksChanged == (int32)IElectraPlayerAdapterDelegate::EPlayerEvent::TracksChanged, "check alignment of both enums");
 	static_assert((int32)EMediaEvent::MetadataChanged == (int32)IElectraPlayerAdapterDelegate::EPlayerEvent::MetadataChanged, "check alignment of both enums");
 	static_assert((int32)EMediaEvent::Internal_PurgeVideoSamplesHint == (int32)IElectraPlayerAdapterDelegate::EPlayerEvent::Internal_PurgeVideoSamplesHint, "check alignment of both enums");
-	static_assert((int32)EMediaEvent::Internal_ResetForDiscontinuity == (int32)IElectraPlayerAdapterDelegate::EPlayerEvent::Internal_ResetForDiscontinuity, "check alignment of both enums");
-	static_assert((int32)EMediaEvent::Internal_RenderClockStart == (int32)IElectraPlayerAdapterDelegate::EPlayerEvent::Internal_RenderClockStart, "check alignment of both enums");
-	static_assert((int32)EMediaEvent::Internal_RenderClockStop == (int32)IElectraPlayerAdapterDelegate::EPlayerEvent::Internal_RenderClockStop, "check alignment of both enums");
 	static_assert((int32)EMediaEvent::Internal_VideoSamplesAvailable == (int32)IElectraPlayerAdapterDelegate::EPlayerEvent::Internal_VideoSamplesAvailable, "check alignment of both enums");
 	static_assert((int32)EMediaEvent::Internal_VideoSamplesUnavailable == (int32)IElectraPlayerAdapterDelegate::EPlayerEvent::Internal_VideoSamplesUnavailable, "check alignment of both enums");
 	static_assert((int32)EMediaEvent::Internal_AudioSamplesAvailable == (int32)IElectraPlayerAdapterDelegate::EPlayerEvent::Internal_AudioSamplesAvailable, "check alignment of both enums");
diff --git a/Engine/Plugins/Media/ElectraPlayer/Source/ElectraPlayerRuntime/Private/ElectraPlayer.cpp b/Engine/Plugins/Media/ElectraPlayer/Source/ElectraPlayerRuntime/Private/ElectraPlayer.cpp
--- a/Engine/Plugins/Media/ElectraPlayer/Source/ElectraPlayerRuntime/Private/ElectraPlayer.cpp
+++ b/Engine/Plugins/Media/ElectraPlayer/Source/ElectraPlayerRuntime/Private/ElectraPlayer.cpp
@@ -767,10 +767,6 @@ void FElectraPlayer::OnVideoFlush()
 		if (PinnedAdapterDelegate.IsValid())
 		{
 			PinnedAdapterDelegate->OnVideoFlush();
-			// The event sink (player facade) needs to prepare for all new data.
-			// We tell it only here and not also in OnAudioFlush() since both flushes occur at the same time
-			// and there are always both the video and audio renderer. See OnAudioRenderingStarted()
-			PinnedAdapterDelegate->SendMediaEvent(IElectraPlayerAdapterDelegate::EPlayerEvent::Internal_ResetForDiscontinuity);
 		}
 	}
 }
@@ -831,44 +827,6 @@ void FElectraPlayer::OnSubtitleFlush()
 }
 
 
-void FElectraPlayer::OnVideoRenderingStarted()
-{
-	TSharedPtr<FInternalPlayerImpl, ESPMode::ThreadSafe> Player = CurrentPlayer;
-	if (Player.IsValid())
-	{
-		TSharedPtr<IElectraPlayerAdapterDelegate, ESPMode::ThreadSafe> PinnedAdapterDelegate = AdapterDelegate.Pin();
-		if (PinnedAdapterDelegate.IsValid())
-		{
-			PinnedAdapterDelegate->SendMediaEvent(IElectraPlayerAdapterDelegate::EPlayerEvent::Internal_RenderClockStart);
-		}
-	}
-}
-void FElectraPlayer::OnVideoRenderingStopped()
-{
-	TSharedPtr<FInternalPlayerImpl, ESPMode::ThreadSafe> Player = CurrentPlayer;
-	if (Player.IsValid())
-	{
-		TSharedPtr<IElectraPlayerAdapterDelegate, ESPMode::ThreadSafe> PinnedAdapterDelegate = AdapterDelegate.Pin();
-		if (PinnedAdapterDelegate.IsValid())
-		{
-			PinnedAdapterDelegate->SendMediaEvent(IElectraPlayerAdapterDelegate::EPlayerEvent::Internal_RenderClockStop);
-		}
-	}
-}
-
-void FElectraPlayer::OnAudioRenderingStarted()
-{
-	// Note: We do nothing here because there is always a video renderer and an audio renderer and both
-	//       get started and stopped at the same time. Thus we only need to relay the event from one of them.
-}
-
-void FElectraPlayer::OnAudioRenderingStopped()
-{
-	// Note: We do nothing here because there is always a video renderer and an audio renderer and both
-	//       get started and stopped at the same time. Thus we only need to relay the event from one of them.
-}
-
-
 //-----------------------------------------------------------------------------
 /**
  * Check timeline and moves sample over to media FACADE sinks
@@ -2774,6 +2732,8 @@ void FElectraPlayer::HandlePlayerEventPrerollEnd()
 	UE_LOG(LogElectraPlayer, Verbose, TEXT("[%u] Player prerolling complete"), InstanceID);
 	Statistics.LastState = "Ready";
 
+	DeferredEvents.Enqueue(IElectraPlayerAdapterDelegate::EPlayerEvent::MediaBufferingComplete);
+
 	// Enqueue a "PrerollEnd" event.
 	static const FString kEventNameElectraPrerollEnd(TEXT("Electra.PrerollEnd"));
 	if (Electra::IsAnalyticsEventEnabled(kEventNameElectraPrerollEnd))
diff --git a/Engine/Plugins/Media/ElectraPlayer/Source/ElectraPlayerRuntime/Private/ElectraPlayer.h b/Engine/Plugins/Media/ElectraPlayer/Source/ElectraPlayerRuntime/Private/ElectraPlayer.h
--- a/Engine/Plugins/Media/ElectraPlayer/Source/ElectraPlayerRuntime/Private/ElectraPlayer.h
+++ b/Engine/Plugins/Media/ElectraPlayer/Source/ElectraPlayerRuntime/Private/ElectraPlayer.h
@@ -60,11 +60,6 @@ public:
 	void OnSubtitleDecoded(ISubtitleDecoderOutputPtr DecoderOutput);
 	void OnSubtitleFlush();
 
-	void OnVideoRenderingStarted();
-	void OnVideoRenderingStopped();
-	void OnAudioRenderingStarted();
-	void OnAudioRenderingStopped();
-
 	void SendAnalyticMetrics(const TSharedPtr<IAnalyticsProviderET>& AnalyticsProvider, const FGuid& InPlayerGuid);
 	void SendAnalyticMetricsPerMinute(const TSharedPtr<IAnalyticsProviderET>& AnalyticsProvider);
 	void SendPendingAnalyticMetrics(const TSharedPtr<IAnalyticsProviderET>& AnalyticsProvider);
diff --git a/Engine/Plugins/Media/ElectraPlayer/Source/ElectraPlayerRuntime/Private/Runtime/Renderer/RendererAudio.cpp b/Engine/Plugins/Media/ElectraPlayer/Source/ElectraPlayerRuntime/Private/Runtime/Renderer/RendererAudio.cpp
--- a/Engine/Plugins/Media/ElectraPlayer/Source/ElectraPlayerRuntime/Private/Runtime/Renderer/RendererAudio.cpp
+++ b/Engine/Plugins/Media/ElectraPlayer/Source/ElectraPlayerRuntime/Private/Runtime/Renderer/RendererAudio.cpp
@@ -258,20 +258,10 @@ UEMediaError FElectraRendererAudio::Flush(const Electra::FParamDict& InOptions)
 // Begins rendering of the first sample buffer
 void FElectraRendererAudio::StartRendering(const Electra::FParamDict& InOptions)
 {
-	//UE_LOG(LogElectraPlayer, Log, TEXT("-- FElectraRendererAudio::StartRendering (NumBuffers = %d)"), NumBuffers);
-	if (TSharedPtr<FElectraPlayer, ESPMode::ThreadSafe> PinnedPlayer = Player.Pin())
-	{
-		PinnedPlayer->OnAudioRenderingStarted();
-	}
 }
 
 
 // Stops rendering of sample buffers
 void FElectraRendererAudio::StopRendering(const Electra::FParamDict& InOptions)
 {
-	//UE_LOG(LogElectraPlayer, Log, TEXT("-- FElectraRendererAudio::StopRendering (NumBuffers = %d)"), NumBuffers);
-	if (TSharedPtr<FElectraPlayer, ESPMode::ThreadSafe> PinnedPlayer = Player.Pin())
-	{
-		PinnedPlayer->OnAudioRenderingStopped();
-	}
 }
diff --git a/Engine/Plugins/Media/ElectraPlayer/Source/ElectraPlayerRuntime/Private/Runtime/Renderer/RendererVideo.cpp b/Engine/Plugins/Media/ElectraPlayer/Source/ElectraPlayerRuntime/Private/Runtime/Renderer/RendererVideo.cpp
--- a/Engine/Plugins/Media/ElectraPlayer/Source/ElectraPlayerRuntime/Private/Runtime/Renderer/RendererVideo.cpp
+++ b/Engine/Plugins/Media/ElectraPlayer/Source/ElectraPlayerRuntime/Private/Runtime/Renderer/RendererVideo.cpp
@@ -263,10 +263,6 @@ UEMediaError FElectraRendererVideo::Flush(const Electra::FParamDict& InOptions)
  */
 void FElectraRendererVideo::StartRendering(const Electra::FParamDict& InOptions)
 {
-	if (TSharedPtr<FElectraPlayer, ESPMode::ThreadSafe> PinnedPlayer = Player.Pin())
-	{
-		PinnedPlayer->OnVideoRenderingStarted();
-	}
 }
 
 /**
@@ -274,10 +270,6 @@ void FElectraRendererVideo::StartRendering(const Electra::FParamDict& InOptions)
 */
 void FElectraRendererVideo::StopRendering(const Electra::FParamDict& InOptions)
 {
-	if (TSharedPtr<FElectraPlayer, ESPMode::ThreadSafe> PinnedPlayer = Player.Pin())
-	{
-		PinnedPlayer->OnVideoRenderingStopped();
-	}
 }
 
 /**
diff --git a/Engine/Plugins/Media/ElectraPlayer/Source/ElectraPlayerRuntime/Public/IElectraPlayerInterface.h b/Engine/Plugins/Media/ElectraPlayer/Source/ElectraPlayerRuntime/Public/IElectraPlayerInterface.h
--- a/Engine/Plugins/Media/ElectraPlayer/Source/ElectraPlayerRuntime/Public/IElectraPlayerInterface.h
+++ b/Engine/Plugins/Media/ElectraPlayer/Source/ElectraPlayerRuntime/Public/IElectraPlayerInterface.h
@@ -57,6 +57,7 @@ public:
 	enum class EPlayerEvent
 	{
 		MediaBuffering = 0,
+		MediaBufferingComplete,
 		MediaClosed,
 		MediaConnecting,
 		MediaOpened,
@@ -70,9 +71,6 @@ public:
 
 		Internal_Start,
 		Internal_PurgeVideoSamplesHint = Internal_Start,
-		Internal_ResetForDiscontinuity,
-		Internal_RenderClockStart,
-		Internal_RenderClockStop,
 		Internal_VideoSamplesAvailable,
 		Internal_VideoSamplesUnavailable,
 		Internal_AudioSamplesAvailable,
diff --git a/Engine/Plugins/Media/MfMedia/Source/MfMedia/Private/Player/MfMediaPlayer.cpp b/Engine/Plugins/Media/MfMedia/Source/MfMedia/Private/Player/MfMediaPlayer.cpp
--- a/Engine/Plugins/Media/MfMedia/Source/MfMedia/Private/Player/MfMediaPlayer.cpp
+++ b/Engine/Plugins/Media/MfMedia/Source/MfMedia/Private/Player/MfMediaPlayer.cpp
@@ -582,7 +582,7 @@ bool FMfMediaPlayer::Seek(const FTimespan& Time)
 		UE_LOG(LogMfMedia, Verbose, TEXT("Player %p: Invalid seek time %s (media duration is %s)"), this, *Time.ToString(), *CurrentDuration.ToString());
 		return false;
 	}
-	
+
 	UE_LOG(LogMfMedia, Verbose, TEXT("Player %p: Seeking to %s"), this, *Time.ToString());
 
 	if (!CommitTime(Time))
@@ -725,6 +725,7 @@ void FMfMediaPlayer::ReceiveSourceReaderEvent(MediaEventType Event)
 
 	case MEBufferingStopped:
 		CurrentStatus &= ~EMediaStatus::Buffering;
+		EventSink.ReceiveMediaEvent(EMediaEvent::MediaBufferingComplete);
 		break;
 
 	case MEConnectEnd:
diff --git a/Engine/Plugins/Media/WmfMedia/Source/WmfMedia/Private/Wmf/WmfMediaSession.cpp b/Engine/Plugins/Media/WmfMedia/Source/WmfMedia/Private/Wmf/WmfMediaSession.cpp
--- a/Engine/Plugins/Media/WmfMedia/Source/WmfMedia/Private/Wmf/WmfMediaSession.cpp
+++ b/Engine/Plugins/Media/WmfMedia/Source/WmfMedia/Private/Wmf/WmfMediaSession.cpp
@@ -710,6 +710,7 @@ STDMETHODIMP FWmfMediaSession::Invoke(IMFAsyncResult* AsyncResult)
 
 	case MEBufferingStopped:
 		Status &= ~EMediaStatus::Buffering;
+		DeferredEvents.Enqueue(EMediaEvent::MediaBufferingComplete);
 		break;
 
 	case MEError:
diff --git a/Engine/Source/Runtime/Media/Public/IMediaEventSink.h b/Engine/Source/Runtime/Media/Public/IMediaEventSink.h
--- a/Engine/Source/Runtime/Media/Public/IMediaEventSink.h
+++ b/Engine/Source/Runtime/Media/Public/IMediaEventSink.h
@@ -11,6 +11,9 @@ enum class EMediaEvent
 	/** The media source started buffering data. */
 	MediaBuffering,
 
+	/** The media source completed buffering data. */
+	MediaBufferingComplete,
+
 	/** The current media source has been closed. */
 	MediaClosed,
 
@@ -49,13 +52,6 @@ enum class EMediaEvent
 	/** Player requests purge of outdated video samples from queues. */
 	Internal_PurgeVideoSamplesHint = Internal_Start,
 
-	/** Player requests an internal reset of all queues and timing. **/
-	Internal_ResetForDiscontinuity,
-
-	/** Player requests render clock to start and stop in preparation for moving playback to a new position. **/
-	Internal_RenderClockStart,
-	Internal_RenderClockStop,
-
 	/** Player notifies about availability or unavailablity of media samples. **/
 	Internal_VideoSamplesAvailable,
 	Internal_VideoSamplesUnavailable,
diff --git a/Engine/Source/Runtime/MediaAssets/Private/Assets/MediaPlayer.cpp b/Engine/Source/Runtime/MediaAssets/Private/Assets/MediaPlayer.cpp
--- a/Engine/Source/Runtime/MediaAssets/Private/Assets/MediaPlayer.cpp
+++ b/Engine/Source/Runtime/MediaAssets/Private/Assets/MediaPlayer.cpp
@@ -1209,6 +1209,14 @@ void UMediaPlayer::HandlePlayerMediaEvent(EMediaEvent Event)
 	case EMediaEvent::MetadataChanged:
 		OnMetadataChanged.Broadcast();
 		break;
+
+	case EMediaEvent::MediaBuffering:
+		OnBufferingStart.Broadcast();
+		break;
+
+	case EMediaEvent::MediaBufferingComplete:
+		OnBufferingCompleted.Broadcast();
+		break;
 	}
 }
 
diff --git a/Engine/Source/Runtime/MediaAssets/Public/MediaPlayer.h b/Engine/Source/Runtime/MediaAssets/Public/MediaPlayer.h
--- a/Engine/Source/Runtime/MediaAssets/Public/MediaPlayer.h
+++ b/Engine/Source/Runtime/MediaAssets/Public/MediaPlayer.h
@@ -820,7 +820,7 @@ public:
 
 	/**
 	 * Seeks to the specified playback time as soon as possible.
-	 * 
+	 *
 	 * If the player is currently seeking, the new seek request will be queued
 	 * and executed as soon as the current one is completed. It will guarantee that
 	 * at least one sample of the previous seek makes it to the sample queue.
@@ -1090,6 +1090,24 @@ public:
 	 */
 	UPROPERTY(BlueprintAssignable, Category="Media|MediaPlayer", meta = (HideInDetailPanel))
 	FOnMediaPlayerMediaEvent OnMetadataChanged;
+
+	/**
+	 * A delegate that is invoked when the player starts buffering data.
+	 * This event is not sent by all media players.
+	 *
+	 * @see OnBufferingCompleted
+	 */
+	UPROPERTY(BlueprintAssignable, Category="Media|MediaPlayer", meta = (HideInDetailPanel))
+	FOnMediaPlayerMediaEvent OnBufferingStart;
+
+	/**
+	 * A delegate that may be invoked when the player completes buffering data.
+	 * This event should not be relied upon since not all media players support this.
+	 *
+	 * @see OnBufferingStart
+	 */
+	UPROPERTY(BlueprintAssignable, Category="Media|MediaPlayer", meta = (HideInDetailPanel))
+	FOnMediaPlayerMediaEvent OnBufferingCompleted;
 public:
 
 	/**
diff --git a/Engine/Source/Runtime/MediaUtils/Private/MediaHelpers.cpp b/Engine/Source/Runtime/MediaUtils/Private/MediaHelpers.cpp
--- a/Engine/Source/Runtime/MediaUtils/Private/MediaHelpers.cpp
+++ b/Engine/Source/Runtime/MediaUtils/Private/MediaHelpers.cpp
@@ -14,6 +14,7 @@ namespace MediaUtils
 		switch (Event)
 		{
 		case EMediaEvent::MediaBuffering: return TEXT("MediaBuffering");
+		case EMediaEvent::MediaBufferingComplete: return TEXT("MediaBufferingComplete");
 		case EMediaEvent::MediaClosed: return TEXT("MediaClosed");
 		case EMediaEvent::MediaConnecting: return TEXT("MediaConnecting");
 		case EMediaEvent::MediaOpened: return TEXT("MediaOpened");
@@ -26,9 +27,6 @@ namespace MediaUtils
 		case EMediaEvent::MetadataChanged: return TEXT("MetadataChanged");
 
 		case EMediaEvent::Internal_PurgeVideoSamplesHint: return TEXT("Internal: PurgeVideoSamplesHint");
-		case EMediaEvent::Internal_ResetForDiscontinuity: return TEXT("Internal: ResetForDiscontinuity");
-		case EMediaEvent::Internal_RenderClockStart: return TEXT("Internal: RenderClockStart");
-		case EMediaEvent::Internal_RenderClockStop: return TEXT("Internal: RenderClockStop");
 		case EMediaEvent::Internal_VideoSamplesAvailable: return TEXT("Internal: VideoSamplesAvailable");
 		case EMediaEvent::Internal_VideoSamplesUnavailable: return TEXT("Internal: VideoSamplesUnavailable");
 		case EMediaEvent::Internal_AudioSamplesAvailable: return TEXT("Internal: AudioSamplesAvailable");
diff --git a/Engine/Source/Runtime/MediaUtils/Private/MediaPlayerFacade.cpp b/Engine/Source/Runtime/MediaUtils/Private/MediaPlayerFacade.cpp
--- a/Engine/Source/Runtime/MediaUtils/Private/MediaPlayerFacade.cpp
+++ b/Engine/Source/Runtime/MediaUtils/Private/MediaPlayerFacade.cpp
@@ -3610,21 +3610,6 @@ void FMediaPlayerFacade::ReceiveMediaEvent(EMediaEvent Event)
 				break;
 			}
 
-			case EMediaEvent::Internal_ResetForDiscontinuity:
-			{
-				break;
-			}
-			case EMediaEvent::Internal_RenderClockStart:
-			{
-				UE_LOG(LogMediaUtils, VeryVerbose, TEXT("PlayerFacade: Render clock shall start"));
-				break;
-			}
-			case EMediaEvent::Internal_RenderClockStop:
-			{
-				UE_LOG(LogMediaUtils, VeryVerbose, TEXT("PlayerFacade: Render clock shall stop"));
-				break;
-			}
-
 			case EMediaEvent::Internal_VideoSamplesAvailable:
 			{
 				UE_LOG(LogMediaUtils, VeryVerbose, TEXT("PlayerFacade: Video samples ARE available"));
