Protron: Added a decoded image cache; added support for playback range

#rnx

diff --git a/Engine/Plugins/Media/ElectraPlayer/Source/ElectraProtron/Private/Player/ElectraProtronPlayer.cpp b/Engine/Plugins/Media/ElectraPlayer/Source/ElectraProtron/Private/Player/ElectraProtronPlayer.cpp
--- a/Engine/Plugins/Media/ElectraPlayer/Source/ElectraProtron/Private/Player/ElectraProtronPlayer.cpp
+++ b/Engine/Plugins/Media/ElectraPlayer/Source/ElectraProtron/Private/Player/ElectraProtronPlayer.cpp
@@ -14,8 +14,6 @@
 FElectraProtronPlayer::FElectraProtronPlayer(IMediaEventSink& InEventSink)
 	: EventSink(InEventSink)
 {
-	EmptySampleQueue = MakeShared<FMediaSamples, ESPMode::ThreadSafe>();
-
 	CurrentState = EMediaState::Closed;
 	CurrentStatus = EMediaStatus::None;
 }
@@ -33,9 +31,6 @@ void FElectraProtronPlayer::Close()
 	CurrentStatus = EMediaStatus::None;
 	CurrentURL.Empty();
 	CurrentInternalState = EInternalState::Closed;
-	SampleQueueLock.Lock();
-	CurrentSampleQueue.Reset();
-	SampleQueueLock.Unlock();
 	if (ToClose.IsValid())
 	{
 		ToClose->Close(FImpl::FCompletionDelegate());
@@ -64,12 +59,7 @@ FGuid FElectraProtronPlayer::GetPlayerPluginGUID() const
 
 IMediaSamples& FElectraProtronPlayer::GetSamples()
 {
-	FScopeLock lock(&SampleQueueLock);
-	if (CurrentSampleQueue.IsValid())
-	{
-		return *CurrentSampleQueue;
-	}
-	return *EmptySampleQueue;
+	return *this;
 }
 
 IMediaTracks& FElectraProtronPlayer::GetTracks()
@@ -118,16 +108,12 @@ bool FElectraProtronPlayer::Open(const FString& InUrl, const IMediaOptions* InOp
 	CurrentTexturePool = MakeShareable(new FElectraTextureSamplePool);
 	CurrentAudioSamplePool = MakeShareable(new FElectraAudioSamplePool);
 
-	TSharedPtr<FMediaSamples, ESPMode::ThreadSafe> NewSampleQueue = MakeShared<FMediaSamples, ESPMode::ThreadSafe>();
-	SampleQueueLock.Lock();
-	CurrentSampleQueue = NewSampleQueue;
-	SampleQueueLock.Unlock();
 	EventSink.ReceiveMediaEvent(EMediaEvent::MediaConnecting);
 	FImpl::FOpenParam op;
 	op.Filename = Filename;
-	op.SampleQueue = NewSampleQueue;
 	op.TexturePool = CurrentTexturePool;
 	op.AudioSamplePool = CurrentAudioSamplePool;
+	op.InitialPlaybackRange = CurrentPlaybackRange;
 	CurrentPlayer->Open(op, FImpl::FCompletionDelegate::CreateSPLambda(AsShared(), [this](FImpl::ImplPointer InImpl)
 	{
 		if (InImpl == CurrentPlayer)
@@ -163,21 +149,9 @@ bool FElectraProtronPlayer::GetPlayerFeatureFlag(EFeatureFlag InWhichFeature) co
 	switch(InWhichFeature)
 	{
 		case EFeatureFlag::AllowShutdownOnClose:
-		{
-			return true;
-		}
 		case EFeatureFlag::UsePlaybackTimingV2:
-		{
-			return true;
-		}
 		case EFeatureFlag::PlayerUsesInternalFlushOnSeek:
-		{
-			return true;
-		}
 		case EFeatureFlag::IsTrackSwitchSeamless:
-		{
-			return true;
-		}
 		case EFeatureFlag::PlayerSelectsDefaultTracks:
 		{
 			return true;
@@ -310,6 +284,10 @@ bool FElectraProtronPlayer::CanControl(EMediaControl InControl) const
 	{
 		return CurrentState == EMediaState::Playing || CurrentState == EMediaState::Paused || CurrentState == EMediaState::Stopped;
 	}
+	else if (InControl == EMediaControl::PlaybackRange)
+	{
+		return true;
+	}
 	return false;
 }
 
@@ -355,36 +333,21 @@ bool FElectraProtronPlayer::IsLooping() const
 
 bool FElectraProtronPlayer::Seek(const FTimespan& InTime, const FMediaSeekParams& InAdditionalParams)
 {
-	// Create a new sample queue into which the new decoded samples will be added.
-	TSharedPtr<FMediaSamples, ESPMode::ThreadSafe> SampleQueue = MakeShared<FMediaSamples, ESPMode::ThreadSafe>();
-	SampleQueueLock.Lock();
-	Swap(CurrentSampleQueue, SampleQueue);
-	SampleQueueLock.Unlock();
-
 	TSharedPtr<FImpl, ESPMode::ThreadSafe> Player(GetCurrentPlayer());
 	if (!Player.IsValid())
 	{
 		return false;
 	}
 	check(InAdditionalParams.NewSequenceIndex.IsSet());
-	int32 SequenceIndexNow = CurrentSequenceIndex;
 	CurrentSequenceIndex = InAdditionalParams.NewSequenceIndex.GetValue();
-	CurrentSampleQueue->SetMinExpectedNextSequenceIndex(InAdditionalParams.NewSequenceIndex);
 	TOptional<int32> NewLoopIndex = 0;
-	bool bOk = Player->Seek(InTime, CurrentSampleQueue, CurrentSequenceIndex, NewLoopIndex);
-	if (bOk)
-	{
-		// Send a seek complete event even if that is not really true and the seek is in progress or
-		// has not even started. This is mostly to satisfy code that cannot handle async seeking or
-		// seeks that override an ongoing seek (so the number of completed seeks will not match the
-		// number of issued seeks).
-		EventSink.ReceiveMediaEvent(EMediaEvent::SeekCompleted);
-	}
-	else
-	{
-		CurrentSequenceIndex = SequenceIndexNow;
-	}
-	return bOk;
+	Player->Seek(InTime, CurrentSequenceIndex, NewLoopIndex);
+	// Send a seek complete event even if that is not really true and the seek is in progress or
+	// has not even started. This is mostly to satisfy code that cannot handle async seeking or
+	// seeks that override an ongoing seek (so the number of completed seeks will not match the
+	// number of issued seeks).
+	EventSink.ReceiveMediaEvent(EMediaEvent::SeekCompleted);
+	return true;
 }
 
 bool FElectraProtronPlayer::SetLooping(bool bInLooping)
@@ -399,9 +362,17 @@ bool FElectraProtronPlayer::SetRate(float InRate)
 	return Player.IsValid() ? Player->SetRate(InRate) : false;
 }
 
-bool FElectraProtronPlayer::SetNativeVolume(float Volume)
+TRange<FTimespan> FElectraProtronPlayer::GetPlaybackTimeRange(EMediaTimeRangeType InRangeToGet) const
 {
-	return false;
+	TSharedPtr<FImpl, ESPMode::ThreadSafe> Player(GetCurrentPlayer());
+	return Player.IsValid() ? Player->GetPlaybackTimeRange(InRangeToGet) : (InRangeToGet == EMediaTimeRangeType::Current && CurrentPlaybackRange.IsSet()) ? CurrentPlaybackRange.GetValue() :TRange<FTimespan>(FTimespan(0), GetDuration());
+}
+
+bool FElectraProtronPlayer::SetPlaybackTimeRange(const TRange<FTimespan>& InTimeRange)
+{
+	TSharedPtr<FImpl, ESPMode::ThreadSafe> Player(GetCurrentPlayer());
+	CurrentPlaybackRange = InTimeRange;
+	return Player.IsValid() ? Player->SetPlaybackTimeRange(InTimeRange) : false;
 }
 
 int32 FElectraProtronPlayer::GetNumTracks(EMediaTrackType InTrackType) const
@@ -481,3 +452,102 @@ int32 FElectraProtronPlayer::GetSampleCount(EMediaCacheState InState) const
 	TSharedPtr<FImpl, ESPMode::ThreadSafe> Player(GetCurrentPlayer());
 	return Player.IsValid() ? Player->GetSampleCount(InState) : 0;
 }
+
+
+
+IMediaSamples::EFetchBestSampleResult FElectraProtronPlayer::FetchBestVideoSampleForTimeRange(const TRange<FMediaTimeStamp>& InTimeRange, TSharedPtr<IMediaTextureSample, ESPMode::ThreadSafe>& OutSample, bool bInReverse, bool bInConsistentResult)
+{
+	TSharedPtr<FImpl, ESPMode::ThreadSafe> Player(GetCurrentPlayer());
+	return Player.IsValid() ? Player->FetchBestVideoSampleForTimeRange(InTimeRange, OutSample, bInReverse, bInConsistentResult) : IMediaSamples::EFetchBestSampleResult::NoSample;
+}
+bool FElectraProtronPlayer::FetchAudio(TRange<FMediaTimeStamp> InTimeRange, TSharedPtr<IMediaAudioSample, ESPMode::ThreadSafe>& OutSample)
+{
+	TSharedPtr<FImpl, ESPMode::ThreadSafe> Player(GetCurrentPlayer());
+	return Player.IsValid() ? Player->FetchAudio(InTimeRange, OutSample) : false;
+}
+bool FElectraProtronPlayer::FetchCaption(TRange<FMediaTimeStamp> InTimeRange, TSharedPtr<IMediaOverlaySample, ESPMode::ThreadSafe>& OutSample)
+{
+	TSharedPtr<FImpl, ESPMode::ThreadSafe> Player(GetCurrentPlayer());
+	return Player.IsValid() ? Player->FetchCaption(InTimeRange, OutSample) : false;
+}
+bool FElectraProtronPlayer::FetchMetadata(TRange<FMediaTimeStamp> InTimeRange, TSharedPtr<IMediaBinarySample, ESPMode::ThreadSafe>& OutSample)
+{
+	TSharedPtr<FImpl, ESPMode::ThreadSafe> Player(GetCurrentPlayer());
+	return Player.IsValid() ? Player->FetchMetadata(InTimeRange, OutSample) : false;
+}
+bool FElectraProtronPlayer::FetchSubtitle(TRange<FMediaTimeStamp> InTimeRange, TSharedPtr<IMediaOverlaySample, ESPMode::ThreadSafe>& OutSample)
+{
+	TSharedPtr<FImpl, ESPMode::ThreadSafe> Player(GetCurrentPlayer());
+	return Player.IsValid() ? Player->FetchSubtitle(InTimeRange, OutSample) : false;
+}
+void FElectraProtronPlayer::FlushSamples()
+{
+	TSharedPtr<FImpl, ESPMode::ThreadSafe> Player(GetCurrentPlayer());
+	if (Player.IsValid())
+	{
+		Player->FlushSamples();
+	}
+}
+void FElectraProtronPlayer::SetMinExpectedNextSequenceIndex(TOptional<int32> InNextSequenceIndex)
+{
+	TSharedPtr<FImpl, ESPMode::ThreadSafe> Player(GetCurrentPlayer());
+	if (Player.IsValid())
+	{
+		Player->SetMinExpectedNextSequenceIndex(InNextSequenceIndex);
+	}
+}
+bool FElectraProtronPlayer::PeekVideoSampleTime(FMediaTimeStamp& OutTimeStamp)
+{
+	TSharedPtr<FImpl, ESPMode::ThreadSafe> Player(GetCurrentPlayer());
+	return Player.IsValid() ? Player->PeekVideoSampleTime(OutTimeStamp) : false;
+}
+bool FElectraProtronPlayer::CanReceiveVideoSamples(uint32 Num) const
+{
+	TSharedPtr<FImpl, ESPMode::ThreadSafe> Player(GetCurrentPlayer());
+	return Player.IsValid() ? Player->CanReceiveVideoSamples(Num) : false;
+}
+bool FElectraProtronPlayer::CanReceiveAudioSamples(uint32 Num) const
+{
+	TSharedPtr<FImpl, ESPMode::ThreadSafe> Player(GetCurrentPlayer());
+	return Player.IsValid() ? Player->CanReceiveAudioSamples(Num) : false;
+}
+bool FElectraProtronPlayer::CanReceiveSubtitleSamples(uint32 Num) const
+{
+	TSharedPtr<FImpl, ESPMode::ThreadSafe> Player(GetCurrentPlayer());
+	return Player.IsValid() ? Player->CanReceiveSubtitleSamples(Num) : false;
+}
+bool FElectraProtronPlayer::CanReceiveCaptionSamples(uint32 Num) const
+{
+	TSharedPtr<FImpl, ESPMode::ThreadSafe> Player(GetCurrentPlayer());
+	return Player.IsValid() ? Player->CanReceiveCaptionSamples(Num) : false;
+}
+bool FElectraProtronPlayer::CanReceiveMetadataSamples(uint32 Num) const
+{
+	TSharedPtr<FImpl, ESPMode::ThreadSafe> Player(GetCurrentPlayer());
+	return Player.IsValid() ? Player->CanReceiveMetadataSamples(Num) : false;
+}
+int32 FElectraProtronPlayer::NumAudioSamples() const
+{
+	TSharedPtr<FImpl, ESPMode::ThreadSafe> Player(GetCurrentPlayer());
+	return Player.IsValid() ? Player->NumAudioSamples() : 0;
+}
+int32 FElectraProtronPlayer::NumCaptionSamples() const
+{
+	TSharedPtr<FImpl, ESPMode::ThreadSafe> Player(GetCurrentPlayer());
+	return Player.IsValid() ? Player->NumCaptionSamples() : 0;
+}
+int32 FElectraProtronPlayer::NumMetadataSamples() const
+{
+	TSharedPtr<FImpl, ESPMode::ThreadSafe> Player(GetCurrentPlayer());
+	return Player.IsValid() ? Player->NumMetadataSamples() : 0;
+}
+int32 FElectraProtronPlayer::NumSubtitleSamples() const
+{
+	TSharedPtr<FImpl, ESPMode::ThreadSafe> Player(GetCurrentPlayer());
+	return Player.IsValid() ? Player->NumSubtitleSamples() : 0;
+}
+int32 FElectraProtronPlayer::NumVideoSamples() const
+{
+	TSharedPtr<FImpl, ESPMode::ThreadSafe> Player(GetCurrentPlayer());
+	return Player.IsValid() ? Player->NumVideoSamples() : 0;
+}
diff --git a/Engine/Plugins/Media/ElectraPlayer/Source/ElectraProtron/Private/Player/ElectraProtronPlayer.h b/Engine/Plugins/Media/ElectraPlayer/Source/ElectraProtron/Private/Player/ElectraProtronPlayer.h
--- a/Engine/Plugins/Media/ElectraPlayer/Source/ElectraProtron/Private/Player/ElectraProtronPlayer.h
+++ b/Engine/Plugins/Media/ElectraPlayer/Source/ElectraProtron/Private/Player/ElectraProtronPlayer.h
@@ -12,6 +12,7 @@
 #include "IMediaPlayer.h"
 #include "IMediaTracks.h"
 #include "IMediaView.h"
+#include "IMediaSamples.h"
 #include "IMediaEventSink.h"
 #include "Misc/Timespan.h"
 #include "Misc/Timecode.h"
@@ -29,6 +30,7 @@ class FElectraProtronPlayer : public TSharedFromThis<FElectraProtronPlayer, ESPM
 	, protected IMediaControls
 	, protected IMediaTracks
 	, protected IMediaView
+	, protected IMediaSamples
 {
 public:
 
@@ -60,15 +62,12 @@ public:
 	bool Open(const TSharedRef<FArchive, ESPMode::ThreadSafe>& InArchive, const FString& InOriginalUrl, const IMediaOptions* InOptions) override;
 	bool Open(const FString& InUrl, const IMediaOptions* InOptions, const FMediaPlayerOptions* InPlayerOptions) override;
 	FVariant GetMediaInfo(FName InInfoName) const override;
-	bool GetPlayerFeatureFlag(EFeatureFlag /*flag*/) const override;
-
+	bool GetPlayerFeatureFlag(EFeatureFlag InWhichFeature) const override;
 	void TickFetch(FTimespan InDeltaTime, FTimespan InTimecode) override;
 	void TickInput(FTimespan InDeltaTime, FTimespan InTimecode) override;
 
 protected:
-
 	//~ IMediaControls interface
-
 	bool CanControl(EMediaControl InControl) const override;
 	FTimespan GetDuration() const override;
 	float GetRate() const override;
@@ -82,12 +81,11 @@ protected:
 	bool Seek(const FTimespan& InNewTime, const FMediaSeekParams& InAdditionalParams) override;
 	bool SetLooping(bool bInLooping) override;
 	bool SetRate(float InRate) override;
-	bool SetNativeVolume(float InVolume) override;
+	TRange<FTimespan> GetPlaybackTimeRange(EMediaTimeRangeType InRangeToGet) const override;
+	bool SetPlaybackTimeRange(const TRange<FTimespan>& InTimeRange) override;
 
-protected:
 
 	//~ IMediaTracks interface
-
 	bool GetAudioTrackFormat(int32 InTrackIndex, int32 InFormatIndex, FMediaAudioTrackFormat& OutFormat) const override;
 	int32 GetNumTracks(EMediaTrackType InTrackType) const override;
 	int32 GetNumTrackFormats(EMediaTrackType InTrackType, int32 InTrackIndex) const override;
@@ -100,16 +98,33 @@ protected:
 	bool SelectTrack(EMediaTrackType InTrackType, int32 InTrackIndex) override;
 	bool SetTrackFormat(EMediaTrackType InTrackType, int32 InTrackIndex, int32 InFormatIndex) override;
 
-protected:
-
 	//~ IMediaCache interface
 	bool QueryCacheState(EMediaCacheState InState, TRangeSet<FTimespan>& OutTimeRanges) const override;
 	int32 GetSampleCount(EMediaCacheState InState) const override;
 
+	//~ IMediaSamples interface
+	EFetchBestSampleResult FetchBestVideoSampleForTimeRange(const TRange<FMediaTimeStamp>& InTimeRange, TSharedPtr<IMediaTextureSample, ESPMode::ThreadSafe>& OutSample, bool bInReverse, bool bInConsistentResult) override;
+	bool FetchAudio(TRange<FMediaTimeStamp> InTimeRange, TSharedPtr<IMediaAudioSample, ESPMode::ThreadSafe>& OutSample) override;
+	bool FetchCaption(TRange<FMediaTimeStamp> InTimeRange, TSharedPtr<IMediaOverlaySample, ESPMode::ThreadSafe>& OutSample) override;
+	bool FetchMetadata(TRange<FMediaTimeStamp> InTimeRange, TSharedPtr<IMediaBinarySample, ESPMode::ThreadSafe>& OutSample) override;
+	bool FetchSubtitle(TRange<FMediaTimeStamp> InTimeRange, TSharedPtr<IMediaOverlaySample, ESPMode::ThreadSafe>& OutSample) override;
+	void FlushSamples() override;
+	void SetMinExpectedNextSequenceIndex(TOptional<int32> InNextSequenceIndex) override;
+	bool PeekVideoSampleTime(FMediaTimeStamp& OutTimeStamp) override;
+	bool CanReceiveVideoSamples(uint32 InNum) const override;
+	bool CanReceiveAudioSamples(uint32 InNum) const override;
+	bool CanReceiveSubtitleSamples(uint32 InNum) const override;
+	bool CanReceiveCaptionSamples(uint32 InNum) const override;
+	bool CanReceiveMetadataSamples(uint32 InNum) const override;
+	int32 NumAudioSamples() const override;
+	int32 NumCaptionSamples() const override;
+	int32 NumMetadataSamples() const override;
+	int32 NumSubtitleSamples() const override;
+	int32 NumVideoSamples() const override;
+
 private:
 	class FImpl;
 
-
 	TSharedPtr<FImpl, ESPMode::ThreadSafe> GetCurrentPlayer() const
 	{
 		return CurrentPlayer;
@@ -127,6 +142,7 @@ private:
 
 
 	TSharedPtr<FImpl, ESPMode::ThreadSafe> CurrentPlayer;
+	TOptional<TRange<FTimespan>> CurrentPlaybackRange;
 	int32 CurrentSequenceIndex = 0;
 	EInternalState CurrentInternalState = EInternalState::Closed;
 
@@ -135,9 +151,6 @@ private:
 	FString LastError;
 	FString CurrentURL;
 
-	FCriticalSection SampleQueueLock;
-	TSharedPtr<FMediaSamples, ESPMode::ThreadSafe> EmptySampleQueue;
-	TSharedPtr<FMediaSamples, ESPMode::ThreadSafe> CurrentSampleQueue;
 	TSharedPtr<FElectraTextureSamplePool, ESPMode::ThreadSafe> CurrentTexturePool;
 	TSharedPtr<FElectraAudioSamplePool, ESPMode::ThreadSafe> CurrentAudioSamplePool;
 
diff --git a/Engine/Plugins/Media/ElectraPlayer/Source/ElectraProtron/Private/Player/ElectraProtronPlayerCache.cpp b/Engine/Plugins/Media/ElectraPlayer/Source/ElectraProtron/Private/Player/ElectraProtronPlayerCache.cpp
new file mode 100644
--- /dev/null
+++ b/Engine/Plugins/Media/ElectraPlayer/Source/ElectraProtron/Private/Player/ElectraProtronPlayerCache.cpp
@@ -0,0 +1,568 @@
+// Copyright Epic Games, Inc. All Rights Reserved.
+
+#include "ElectraProtronPlayerCache.h"
+
+FProtronVideoCache::FProtronVideoCache()
+{
+}
+
+void FProtronVideoCache::SetMaxFramesToCache(int32 InNumVideoFramesToCacheAhead, int32 InNumVideoFramesToCacheBehind)
+{
+	NumVideoFramesToCacheAhead = InNumVideoFramesToCacheAhead;
+	NumVideoFramesToCacheBehind = InNumVideoFramesToCacheBehind;
+	MaxFramesToCache = NumVideoFramesToCacheAhead + NumVideoFramesToCacheBehind;
+	Empty();
+}
+
+void FProtronVideoCache::SetPlaybackRange(TRange<FTimespan> InRange)
+{
+	FScopeLock lock(&Lock);
+	PlaybackRange = MoveTemp(InRange);
+	PlaybackRangeEndInclusive = PlaybackRange.GetUpperBoundValue() - FTimespan(1);
+	while(Entries.Num() && !Entries[0].RawRange.Overlaps(PlaybackRange))
+	{
+		Entries.RemoveAt(0);
+	}
+	TRange<FTimespan> TempInclusiveRange(PlaybackRange.GetLowerBoundValue(), PlaybackRange.GetUpperBoundValue() + FTimespan(1));
+	while(Entries.Num() && !Entries.Last().RawRange.Overlaps(TempInclusiveRange))
+	{
+		Entries.Pop();
+	}
+}
+
+void FProtronVideoCache::SeekIssuedTo(FTimespan InNewPosition)
+{
+	FScopeLock lock(&Lock);
+	CurrentReadTime.Reset();
+
+	// Check if we have an entry for this time.
+	int32 Idx = INDEX_NONE;
+	for(int32 i=0,iMax=Entries.Num(); i<iMax; ++i)
+	{
+		if (Entries[i].RawRange.Contains(InNewPosition))
+		{
+			Idx = i;
+			break;
+		}
+	}
+	if (Idx != INDEX_NONE)
+	{
+		PerformMaintenance(Idx, true, false);
+	}
+	else
+	{
+		Entries.Empty();
+	}
+}
+
+void FProtronVideoCache::SetPlaybackRate(float InNewRate)
+{
+	FScopeLock lock(&Lock);
+	PlaybackRate = InNewRate;
+}
+
+int32 FProtronVideoCache::Find(FTimespan InPTS)
+{
+	FScopeLock lock(&Lock);
+	for(int32 i=0,iMax=Entries.Num(); i<iMax; ++i)
+	{
+		if (Entries[i].RawRange.GetLowerBoundValue() == InPTS)
+		{
+			return i;
+		}
+	}
+	return INDEX_NONE;
+}
+
+bool FProtronVideoCache::Contains(FTimespan InPTS)
+{
+	return Find(InPTS) != INDEX_NONE;
+}
+
+bool FProtronVideoCache::CanAccept(FTimespan InPTS)
+{
+	FScopeLock lock(&Lock);
+
+	// If we are starting at a new time we can accept the sample.
+	if (!CurrentReadTime.IsSet())
+	{
+		return true;
+	}
+
+	// If we already have that time then we can accept the sample again.
+	if (Contains(InPTS))
+	{
+		return true;
+	}
+
+	// When we are full we can't accept anything.
+	if (Entries.Num() >= MaxFramesToCache)
+	{
+		return false;
+	}
+
+
+	int32 CurrentReadTimeIndex = Find(CurrentReadTime.GetValue());
+	check(CurrentReadTimeIndex != INDEX_NONE);
+
+	// Playing forwards?
+	TArray<int32> DummyIndices;
+	int32 DummyWrapIndex;
+	if (PlaybackRate >= 0.0f)
+	{
+		int32 HaveFuture = GetConsecutiveFutureSamples(DummyIndices, DummyWrapIndex, CurrentReadTimeIndex, true);
+		return HaveFuture < NumVideoFramesToCacheAhead;
+	}
+	else
+	{
+		int32 HavePrevious = GetConsecutivePreviousSamples(DummyIndices, DummyWrapIndex, CurrentReadTimeIndex, true);
+		return HavePrevious < NumVideoFramesToCacheAhead;
+	}
+}
+
+void FProtronVideoCache::Empty()
+{
+	FScopeLock lock(&Lock);
+	Entries.Empty();
+	CurrentReadTime.Reset();
+}
+
+FProtronVideoCache::EGetResult FProtronVideoCache::GetFrame(TSharedPtr<IMediaTextureSample, ESPMode::ThreadSafe>& OutFrame, const TRange<FMediaTimeStamp>& InTimeRange, bool bIsLooping, bool bInReverse, bool bInUseFirstMatch)
+{
+	auto WrappedModulo = [](FTimespan InTime, FTimespan InDuration) -> FTimespan
+	{
+		return (InTime >= FTimespan::Zero()) ? (InTime % InDuration) : (InDuration + (InTime % InDuration));
+	};
+
+
+	/*
+		If the time range to find a frame for wraps around for looping then we need to inspect multiple sub ranges.
+		Here we do not unwrap more than one loop index since this makes no difference.
+	*/
+	TRange<FMediaTimeStamp> AdjustedTimeRange(InTimeRange);
+	if (AdjustedTimeRange.GetLowerBoundValue().GetTime() < PlaybackRange.GetLowerBoundValue())
+	{
+		check(!"not handled");
+	}
+	// Note: This should probably have happened in the player facade GetCurrentPlaybackTimeRange() already, but it doesn't.
+	//       All that happens there is checking if the range contains a wrap around, but not if the range as a whole needs to wrap!
+	if (AdjustedTimeRange.GetLowerBoundValue().GetTime() > PlaybackRangeEndInclusive || AdjustedTimeRange.GetUpperBoundValue().GetTime() > PlaybackRangeEndInclusive)
+	{
+		check(!bInReverse);
+		FTimespan s0 = WrappedModulo(AdjustedTimeRange.GetLowerBoundValue().GetTime(), PlaybackRange.GetUpperBoundValue());
+		FTimespan s1 = WrappedModulo(AdjustedTimeRange.GetUpperBoundValue().GetTime(), PlaybackRange.GetUpperBoundValue());
+		if (s0 < AdjustedTimeRange.GetLowerBoundValue().GetTime())
+		{
+			AdjustedTimeRange.SetLowerBoundValue(FMediaTimeStamp(s0, AdjustedTimeRange.GetLowerBoundValue().GetSequenceIndex(), AdjustedTimeRange.GetLowerBoundValue().GetLoopIndex() + 1));
+		}
+		if (s1 < AdjustedTimeRange.GetUpperBoundValue().GetTime())
+		{
+			AdjustedTimeRange.SetUpperBoundValue(FMediaTimeStamp(s1, AdjustedTimeRange.GetUpperBoundValue().GetSequenceIndex(), AdjustedTimeRange.GetUpperBoundValue().GetLoopIndex() + 1));
+		}
+	}
+
+	TArray<TRange<FTimespan>> RangesToCheck;
+	TArray<int32> LoopIndices;
+	// Does the time range loop?
+	if (AdjustedTimeRange.GetLowerBoundValue().GetLoopIndex() == AdjustedTimeRange.GetUpperBoundValue().GetLoopIndex())
+	{
+		// No. Add the range as it is to the list of ranges to check.
+		RangesToCheck.Emplace(TRange<FTimespan>(AdjustedTimeRange.GetLowerBoundValue().GetTime(), AdjustedTimeRange.GetUpperBoundValue().GetTime()));
+		LoopIndices.Emplace(AdjustedTimeRange.GetLowerBoundValue().GetLoopIndex());
+	}
+	// Looping forward?
+	else if (AdjustedTimeRange.GetLowerBoundValue().GetLoopIndex() < AdjustedTimeRange.GetUpperBoundValue().GetLoopIndex())
+	{
+		RangesToCheck.Emplace(TRange<FTimespan>(AdjustedTimeRange.GetLowerBoundValue().GetTime(), PlaybackRange.GetUpperBoundValue()));
+		RangesToCheck.Emplace(TRange<FTimespan>(FTimespan::Zero(), AdjustedTimeRange.GetUpperBoundValue().GetTime()));
+		LoopIndices.Emplace(AdjustedTimeRange.GetLowerBoundValue().GetLoopIndex());
+		LoopIndices.Emplace(AdjustedTimeRange.GetUpperBoundValue().GetLoopIndex());
+	}
+	// Looping backward.
+	else
+	{
+		RangesToCheck.Emplace(TRange<FTimespan>(FTimespan::Zero(), AdjustedTimeRange.GetLowerBoundValue().GetTime()));
+		RangesToCheck.Emplace(TRange<FTimespan>(AdjustedTimeRange.GetUpperBoundValue().GetTime(), PlaybackRange.GetUpperBoundValue()));
+		LoopIndices.Emplace(AdjustedTimeRange.GetLowerBoundValue().GetLoopIndex());
+		LoopIndices.Emplace(AdjustedTimeRange.GetUpperBoundValue().GetLoopIndex());
+	}
+
+	FScopeLock lock(&Lock);
+
+	FTimespan OverlapDuration(FTimespan::Zero());
+	int32 BestIndex = INDEX_NONE;
+	TOptional<int32> LoopIndex;
+
+	// Searching for forward play direction?
+	if (!bInReverse)
+	{
+		// Do we want the first or last match?
+		if (!bInUseFirstMatch)
+		{
+			// We want to find the *last* match with the largest overlap duration with the time range,
+			// so we check the samples in reverse to find such a match fast.
+			for(int32 nSubRangeIdx=RangesToCheck.Num()-1; nSubRangeIdx>=0; --nSubRangeIdx)
+			{
+				const TRange<FTimespan>& TimeRange(RangesToCheck[nSubRangeIdx]);
+				for(int32 SmpIdx=Entries.Num()-1; SmpIdx>=0; --SmpIdx)
+				{
+					if (TimeRange.Overlaps(Entries[SmpIdx].RawRange))
+					{
+						FTimespan Overlap(TRange<FTimespan>::Intersection(Entries[SmpIdx].RawRange, TimeRange).Size<FTimespan>());
+						// No candidate or a better overlap?
+						if (BestIndex == INDEX_NONE || Overlap > OverlapDuration)
+						{
+							OverlapDuration = Overlap;
+							BestIndex = SmpIdx;
+							LoopIndex = LoopIndices[nSubRangeIdx];
+						}
+						else
+						{
+							break;
+						}
+					}
+				}
+			}
+		}
+		else
+		{
+			// We want to find the *first* match with the largest overlap duration.
+			for(int32 nSubRangeIdx=0; nSubRangeIdx<RangesToCheck.Num(); ++nSubRangeIdx)
+			{
+				const TRange<FTimespan>& TimeRange(RangesToCheck[nSubRangeIdx]);
+				for(int32 SmpIdx=0; SmpIdx<Entries.Num(); ++SmpIdx)
+				{
+					if (TimeRange.Overlaps(Entries[SmpIdx].RawRange))
+					{
+						FTimespan Overlap(TRange<FTimespan>::Intersection(Entries[SmpIdx].RawRange, TimeRange).Size<FTimespan>());
+						// No candidate or a better overlap?
+						if (BestIndex == INDEX_NONE || Overlap > OverlapDuration)
+						{
+							OverlapDuration = Overlap;
+							BestIndex = SmpIdx;
+							LoopIndex = LoopIndices[nSubRangeIdx];
+						}
+						else
+						{
+							break;
+						}
+					}
+				}
+			}
+		}
+	}
+	// Searching for reverse play direction
+	else
+	{
+		// Do we want the first or last match?
+		if (!bInUseFirstMatch)
+		{
+			// We want to find the *last* match with the largest overlap duration with the time range
+			// but in reverse direction, so we check the samples forwards to find such a match fast.
+			for(int32 nSubRangeIdx=RangesToCheck.Num()-1; nSubRangeIdx>=0; --nSubRangeIdx)
+			{
+				const TRange<FTimespan>& TimeRange(RangesToCheck[nSubRangeIdx]);
+				for(int32 SmpIdx=0; SmpIdx<Entries.Num(); ++SmpIdx)
+				{
+					if (TimeRange.Overlaps(Entries[SmpIdx].RawRange))
+					{
+						FTimespan Overlap(TRange<FTimespan>::Intersection(Entries[SmpIdx].RawRange, TimeRange).Size<FTimespan>());
+						// No candidate or a better overlap?
+						if (BestIndex == INDEX_NONE || Overlap > OverlapDuration)
+						{
+							OverlapDuration = Overlap;
+							BestIndex = SmpIdx;
+							LoopIndex = LoopIndices[nSubRangeIdx];
+						}
+						else
+						{
+							break;
+						}
+					}
+				}
+			}
+		}
+		else
+		{
+			// We want to find the *first* match with the largest overlap duration,
+			// but in reverse direction.
+			for(int32 nSubRangeIdx=0; nSubRangeIdx<RangesToCheck.Num(); ++nSubRangeIdx)
+			{
+				const TRange<FTimespan>& TimeRange(RangesToCheck[nSubRangeIdx]);
+				for(int32 SmpIdx=Entries.Num()-1; SmpIdx>=0; --SmpIdx)
+				{
+					if (TimeRange.Overlaps(Entries[SmpIdx].RawRange))
+					{
+						FTimespan Overlap(TRange<FTimespan>::Intersection(Entries[SmpIdx].RawRange, TimeRange).Size<FTimespan>());
+						// No candidate or a better overlap?
+						if (BestIndex == INDEX_NONE || Overlap > OverlapDuration)
+						{
+							OverlapDuration = Overlap;
+							BestIndex = SmpIdx;
+							LoopIndex = LoopIndices[nSubRangeIdx];
+						}
+						else
+						{
+							break;
+						}
+					}
+				}
+			}
+		}
+	}
+
+
+	// Found?
+	if (BestIndex != INDEX_NONE)
+	{
+		OutFrame = Entries[BestIndex].Frame;
+		// Set the sequence and loop index
+		check(LoopIndex.IsSet());
+		static_cast<FElectraTextureSample*>(OutFrame.Get())->SetTime(FMediaTimeStamp(OutFrame->GetTime().GetTime(), AdjustedTimeRange.GetLowerBoundValue().GetSequenceIndex(), LoopIndex.GetValue()));
+
+		CurrentReadTime = Entries[BestIndex].RawRange.GetLowerBoundValue();
+
+		PerformMaintenance(BestIndex, bIsLooping, bInReverse);
+		return EGetResult::Hit;
+	}
+	// Not found.
+	else
+	{
+		// If there are samples we need to dump all of them and return that fact.
+		if (Entries.Num())
+		{
+			Empty();
+			return EGetResult::PurgedEmpty;
+		}
+		return EGetResult::Miss;
+	}
+}
+
+
+void FProtronVideoCache::PerformMaintenance(int32 InAtIndex, bool bIsLooping, bool bInReverse)
+{
+	TArray<int32> NextIndices, PrevIndices;
+	int32 NextIndexWrap, PrevIndexWrap;
+	GetConsecutiveFutureSamples(NextIndices, NextIndexWrap, InAtIndex, bIsLooping);
+	GetConsecutivePreviousSamples(PrevIndices, PrevIndexWrap, InAtIndex, bIsLooping);
+
+	// If the next indices contains the oldest then we have the entire video in the cache.
+	// So either the video is really short or the cache massively large.
+	// Either way, we do not need to perform any evicting of frames.
+	if (NextIndices.Num() && PrevIndices.Num() && NextIndices.Contains(PrevIndices.Last()))
+	{
+		return;
+	}
+
+	// Perform cache maintenance.
+	if (!bInReverse)
+	{
+		if (PrevIndices.Num() > NumVideoFramesToCacheBehind)
+		{
+			// Remove the indices of the samples we want to keep.
+			PrevIndices.RemoveAt(0, NumVideoFramesToCacheBehind);
+			// Check the indices we want to remove for being a non-wrapping range.
+			if (PrevIndices[0] >= PrevIndices.Last())
+			{
+				Entries.RemoveAt(PrevIndices.Last(), PrevIndices[0] - PrevIndices.Last() + 1);
+			}
+			else
+			{
+				check(PrevIndexWrap != INDEX_NONE);
+				PrevIndexWrap -= NumVideoFramesToCacheBehind;
+				if (PrevIndexWrap < 0)
+				{
+					PrevIndexWrap = 0;
+				}
+				// First remove the range with the larger indices, which ensures the array doesn't shift down.
+				Entries.RemoveAt(PrevIndices.Last(), PrevIndices[PrevIndexWrap] - PrevIndices.Last() + 1);
+				// Then, if there is still a valid wrap index, remove the ones with the smaller indices which causes the
+				// array to shift down.
+				if (PrevIndexWrap > 0)
+				{
+					Entries.RemoveAt(PrevIndices[PrevIndexWrap - 1], PrevIndices[0] - PrevIndices[PrevIndexWrap - 1] + 1);
+				}
+			}
+		}
+	}
+	else
+	{
+		if (NextIndices.Num() > NumVideoFramesToCacheBehind)
+		{
+			// Remove the indices of the samples we want to keep.
+			NextIndices.RemoveAt(0, NumVideoFramesToCacheBehind);
+			// Check the indices we want to remove for being a non-wrapping range.
+			if (NextIndices[0] <= NextIndices.Last())
+			{
+				Entries.RemoveAt(NextIndices[0], NextIndices.Last() - NextIndices[0] + 1);
+			}
+			else
+			{
+				check(NextIndexWrap != INDEX_NONE);
+				NextIndexWrap -= NumVideoFramesToCacheBehind;
+				if (NextIndexWrap < 0)
+				{
+					NextIndexWrap = 0;
+				}
+				// First remove the range with the larger indices, like above. Here these are at the start of the list.
+				if (NextIndexWrap > 0)
+				{
+					Entries.RemoveAt(NextIndices[0], NextIndices[NextIndexWrap - 1] - NextIndices[0] + 1);
+				}
+				Entries.RemoveAt(NextIndices[NextIndexWrap], NextIndices.Last() - NextIndices[NextIndexWrap] + 1);
+			}
+		}
+	}
+}
+
+
+int32 FProtronVideoCache::GetConsecutiveFutureSamples(TArray<int32>& OutIndices, int32& OutWrapIndex, int32 InStartFrameIndex, bool bInLooping)
+{
+	TRange<FTimespan> SampleRange(Entries[InStartFrameIndex].RawRange);
+	int32 curIdx = InStartFrameIndex;
+	OutWrapIndex = INDEX_NONE;
+	for(int32 i=1,iMax=Entries.Num(); i<iMax; ++i)
+	{
+		int32 nextIdx = (curIdx + 1) % iMax;
+		TRange<FTimespan> NextSampleRange(Entries[nextIdx].RawRange);
+		// Wrap around?
+		if (nextIdx < curIdx)
+		{
+			// When not considering looping we are done.
+			if (!bInLooping)
+			{
+				break;
+			}
+			// Are we wrapping around on the last sample of the range?
+			// If this is not the last sample then we are not looping back to the beginning.
+			if (!SampleRange.Contains(PlaybackRangeEndInclusive))
+			{
+				break;
+			}
+
+			// If the next sample range, which is the range of the lowest sample in the cache,
+			// does not contain the start of the play range then this is not a contiguous loop.
+			if (!NextSampleRange.Contains(PlaybackRange.GetLowerBoundValue()))
+			{
+				break;
+			}
+			else
+			{
+				OutWrapIndex = OutIndices.Num();
+				OutIndices.Emplace(nextIdx);
+			}
+		}
+		else
+		{
+			if (NextSampleRange.Adjoins(SampleRange))
+			{
+				OutIndices.Emplace(nextIdx);
+			}
+			else
+			{
+				break;
+			}
+		}
+		curIdx = nextIdx;
+		SampleRange = MoveTemp(NextSampleRange);
+	}
+	return OutIndices.Num();
+}
+
+
+int32 FProtronVideoCache::GetConsecutivePreviousSamples(TArray<int32>& OutIndices, int32& OutWrapIndex, int32 InStartFrameIndex, bool bInLooping)
+{
+	TRange<FTimespan> SampleRange(Entries[InStartFrameIndex].RawRange);
+	int32 curIdx = InStartFrameIndex;
+	OutWrapIndex = INDEX_NONE;
+	for(int32 i=1,iMax=Entries.Num(); i<iMax; ++i)
+	{
+		int32 nextIdx = curIdx - 1;
+		if (nextIdx < 0)
+		{
+			nextIdx += iMax;
+		}
+		TRange<FTimespan> PrevSampleRange(Entries[nextIdx].RawRange);
+		// Wrap around?
+		if (nextIdx > curIdx)
+		{
+			// When not considering looping we are done.
+			if (!bInLooping)
+			{
+				break;
+			}
+			// Are we wrapping around on the first sample of the range?
+			// If this is not the first sample then we are not looping back to the end.
+			if (!SampleRange.Contains(PlaybackRange.GetLowerBoundValue()))
+			{
+				break;
+			}
+
+			// If the previous sample range, which is the range of the highest sample in the cache,
+			// does not contain the end of the play range then this is not a contiguous loop.
+			if (!PrevSampleRange.Contains(PlaybackRangeEndInclusive))
+			{
+				break;
+			}
+			else
+			{
+				OutWrapIndex = OutIndices.Num();
+				OutIndices.Emplace(nextIdx);
+			}
+		}
+		else
+		{
+			if (PrevSampleRange.Adjoins(SampleRange))
+			{
+				OutIndices.Emplace(nextIdx);
+			}
+			else
+			{
+				break;
+			}
+		}
+		curIdx = nextIdx;
+		SampleRange = MoveTemp(PrevSampleRange);
+	}
+	return OutIndices.Num();
+}
+
+
+void FProtronVideoCache::AddFrame(TSharedRef<IMediaTextureSample, ESPMode::ThreadSafe> InFrame, FTimespan InRawPTS, FTimespan InRawDuration)
+{
+	/*
+		Note: We must not try to make room in the cache here.
+			    If we do that then we basically allow the decoder to run freely from start to end
+				as we would be throwing out frames that haven't been used yet to make room for the
+				next frame to add.
+	*/
+	FScopeLock lock(&Lock);
+
+	// If there is no known time of the sample value last returned in GetFrame(), which is the case at the
+	// start or after a seek, then we set that time to the time of the sample being delivered now.
+	// That will become the time around which we check which cache entries to keep and which to evict.
+	if (!CurrentReadTime.IsSet())
+	{
+		CurrentReadTime = InRawPTS;
+	}
+
+	int32 Idx = Find(InRawPTS);
+	if (Idx == INDEX_NONE)
+	{
+		FEntry ne;
+		ne.RawRange = TRange<FTimespan>(InRawPTS, InRawPTS + InRawDuration);
+		ne.Frame = MoveTemp(InFrame);
+		Entries.Emplace(MoveTemp(ne));
+		Entries.Sort([](const FEntry& a, const FEntry& b) { return a.RawRange.GetLowerBoundValue() < b.RawRange.GetLowerBoundValue(); });
+	}
+}
+
+
+void FProtronVideoCache::QueryCacheState(TRangeSet<FTimespan>& OutTimeRanges)
+{
+	FScopeLock lock(&Lock);
+	for(auto &e : Entries)
+	{
+		OutTimeRanges.Add(e.RawRange);
+	}
+}
diff --git a/Engine/Plugins/Media/ElectraPlayer/Source/ElectraProtron/Private/Player/ElectraProtronPlayerCache.h b/Engine/Plugins/Media/ElectraPlayer/Source/ElectraProtron/Private/Player/ElectraProtronPlayerCache.h
new file mode 100644
--- /dev/null
+++ b/Engine/Plugins/Media/ElectraPlayer/Source/ElectraProtron/Private/Player/ElectraProtronPlayerCache.h
@@ -0,0 +1,67 @@
+// Copyright Epic Games, Inc. All Rights Reserved.
+
+#pragma once
+
+#include "ElectraProtronPlayer.h"
+#include "Containers/Queue.h"
+#include "Misc/TVariant.h"
+#include "Misc/Optional.h"
+#include "Math/Range.h"
+#include "Math/RangeSet.h"
+#include "MediaSamples.h"
+
+
+class FProtronVideoCache
+{
+public:
+	FProtronVideoCache();
+	void SetMaxFramesToCache(int32 InNumVideoFramesToCacheAhead, int32 InNumVideoFramesToCacheBehind);
+	void SetPlaybackRange(TRange<FTimespan> InRange);
+	void SeekIssuedTo(FTimespan InNewPosition);
+	void SetPlaybackRate(float InNewRate);
+
+	bool Contains(FTimespan InPTS);
+	int32 Find(FTimespan InPTS);
+	bool CanAccept(FTimespan InPTS);
+	void Empty();
+
+	void QueryCacheState(TRangeSet<FTimespan>& OutTimeRanges);
+
+	enum class EGetResult
+	{
+		Hit,
+		Miss,
+		PurgedEmpty
+	};
+
+	EGetResult GetFrame(TSharedPtr<IMediaTextureSample, ESPMode::ThreadSafe>& OutFrame, const TRange<FMediaTimeStamp>& InTimeRange, bool bIsLooping, bool bInReverse, bool bInUseFirstMatch);
+
+	void PerformMaintenance(int32 InAtIndex, bool bIsLooping, bool bInReverse);
+
+	void AddFrame(TSharedRef<IMediaTextureSample, ESPMode::ThreadSafe> InFrame, FTimespan InRawPTS, FTimespan InRawDuration);
+
+private:
+	FCriticalSection Lock;
+	struct FEntry
+	{
+		TRange<FTimespan> RawRange;
+		TSharedPtr<IMediaTextureSample, ESPMode::ThreadSafe> Frame;
+	};
+
+	int32 GetConsecutiveFutureSamples(TArray<int32>& OutIndices, int32& OutWrapIndex, int32 InStartFrameIndex, bool bInLooping);
+	int32 GetConsecutivePreviousSamples(TArray<int32>& OutIndices, int32& OutWrapIndex, int32 InStartFrameIndex, bool bInLooping);
+
+
+
+	TArray<FEntry> Entries;
+	TRange<FTimespan> PlaybackRange;
+	FTimespan PlaybackRangeEndInclusive;
+
+	// The position around which to check for consecutive future and past samples when adding.
+	TOptional<FTimespan> CurrentReadTime;
+
+	float PlaybackRate = 0.0f;
+	int32 MaxFramesToCache = 0;
+	int32 NumVideoFramesToCacheAhead = 0;
+	int32 NumVideoFramesToCacheBehind = 0;
+};
diff --git a/Engine/Plugins/Media/ElectraPlayer/Source/ElectraProtron/Private/Player/ElectraProtronPlayerControls.cpp b/Engine/Plugins/Media/ElectraPlayer/Source/ElectraProtron/Private/Player/ElectraProtronPlayerControls.cpp
--- a/Engine/Plugins/Media/ElectraPlayer/Source/ElectraProtron/Private/Player/ElectraProtronPlayerControls.cpp
+++ b/Engine/Plugins/Media/ElectraPlayer/Source/ElectraProtron/Private/Player/ElectraProtronPlayerControls.cpp
@@ -32,8 +32,8 @@ TRangeSet<float> FElectraProtronPlayer::FImpl::GetSupportedRates(EMediaRateThinn
 				auto Track = Tracks[UsableTrackArrayIndicesByType[TypeIdx][TrackSelection.SelectedTrackIndex[TypeIdx]]];
 				if (Track.IsValid() && Track->bIsKeyframeOnlyFormat)
 				{
-					TempRates.Add(TRange<float>::Inclusive(0.0f, 64.0f));
-					TempRates.Add(TRange<float>::Inclusive(-64.0f, 0.0f));
+					TempRates.Add(TRange<float>::Inclusive(0.0f, 8.0f));
+					TempRates.Add(TRange<float>::Inclusive(-8.0f, 0.0f));
 				}
 			}
 		}
@@ -108,9 +108,6 @@ void FElectraProtronPlayer::FImpl::HandleRateChanges()
 {
 	if (IntendedRate != CurrentRate)
 	{
-		VideoDecoderThread.SetRate(IntendedRate);
-		AudioDecoderThread.SetRate(IntendedRate);
-
 		// Update the general playback direction. When going into pause retain
 		// the last direction for the loader to know into which direction
 		// the operation went before pausing.
@@ -125,6 +122,10 @@ void FElectraProtronPlayer::FImpl::HandleRateChanges()
 		}
 		SharedPlayParams->DesiredPlayRate = IntendedRate;
 
+		CurrentSampleQueueInterface->SetPlaybackRate(IntendedRate);
+		VideoDecoderThread.SetRate(IntendedRate);
+		AudioDecoderThread.SetRate(IntendedRate);
+
 		CurrentRate = IntendedRate;
 	}
 }
@@ -148,11 +149,9 @@ void FElectraProtronPlayer::FImpl::HandleSeekRequest(const FSeekRequest& InSeekR
 	UpdateTrackLoader(CodecTypeIndex(ElectraProtronUtils::FCodecInfo::EType::Video));
 	UpdateTrackLoader(CodecTypeIndex(ElectraProtronUtils::FCodecInfo::EType::Audio));
 
-	VideoDecoderThread.SetSampleQueue(InSeekRequest.NewSampleQueue);
 	VideoDecoderThread.SetTime(InSeekRequest.NewTime, InSeekRequest.NewSequenceIndex, InSeekRequest.NewLoopIndex);
 	VideoDecoderThread.ResumeAfterSeek();
 
-	AudioDecoderThread.SetSampleQueue(InSeekRequest.NewSampleQueue);
 	AudioDecoderThread.SetTime(InSeekRequest.NewTime, InSeekRequest.NewSequenceIndex, InSeekRequest.NewLoopIndex);
 	AudioDecoderThread.ResumeAfterSeek();
 }
diff --git a/Engine/Plugins/Media/ElectraPlayer/Source/ElectraProtron/Private/Player/ElectraProtronPlayerDecoder.cpp b/Engine/Plugins/Media/ElectraPlayer/Source/ElectraProtron/Private/Player/ElectraProtronPlayerDecoder.cpp
--- a/Engine/Plugins/Media/ElectraPlayer/Source/ElectraProtron/Private/Player/ElectraProtronPlayerDecoder.cpp
+++ b/Engine/Plugins/Media/ElectraPlayer/Source/ElectraProtron/Private/Player/ElectraProtronPlayerDecoder.cpp
@@ -15,7 +15,6 @@
 #include "IElectraDecoder.h"
 #include "IElectraDecoderOutputVideo.h"
 #include "IElectraDecoderOutputAudio.h"
-#include "IElectraDecoderFeaturesAndOptions.h"
 #include "IElectraPlayerPluginModule.h"
 #include "ElectraDecodersUtils.h"
 
@@ -119,10 +118,18 @@ uint32 FElectraProtronPlayer::FImpl::FDecoderThread::Run()
 				DestroyDecoder();
 				TrackSampleBuffer = MoveTemp(NewTrackSampleBuffer);
 				GetSampleDelegate = MoveTemp(NewGetSampleDelegate);
+				UpdateTrackSampleDurationMap();
+				FirstRangeSampleIt.Reset();
+				LastRangeSampleIt.Reset();
+				HandlePlaybackRangeChanges();
 				TimeLock.Lock();
+				CurrentTime = ClampTimeIntoPlaybackRange(CurrentTime);
 				FTimespan Time = CurrentTime;
 				TimeLock.Unlock();
 				UpdateTrackIterator(Time);
+				// When switching buffers we can skip decoding of samples prior to the current time.
+				SeekTimeToDecodeTo = SeekTimeToHandleTo = Time;
+				SeekTimeNumFramesDecoded = SeekTimeNumFramesSkipped = 0;
 			}
 			else
 			{
@@ -139,16 +146,14 @@ uint32 FElectraProtronPlayer::FImpl::FDecoderThread::Run()
 		// Check for rate changes.
 		if (IntendedRate != CurrentRate)
 		{
-			// Is this a direction change after we reached the end?
-			if (bReachedEnd && IntendedRate != 0.0f && TrackIterator.IsValid())
+			const bool bDirectionChange = (IntendedRate < 0.0f && PlaybackDirection >= 0.0f) || (IntendedRate > 0.0f && PlaybackDirection < 0.0f);
+			FTimespan NewTime;
+			if (bDirectionChange)
 			{
-				bool bDirectionChange = (IntendedRate < 0.0f && PlaybackDirection >= 0.0f) || (IntendedRate > 0.0f && PlaybackDirection < 0.0f);
-				if (bDirectionChange)
-				{
-					PlaybackDirection = IntendedRate;
-					bReachedEnd = false;
-					StepTrackIterator(false);
-				}
+				TimeLock.Lock();
+				NewTime = CurrentTime;
+				TimeLock.Unlock();
+				PerformFlush();
 			}
 
 			// Going into pause?
@@ -172,6 +177,11 @@ uint32 FElectraProtronPlayer::FImpl::FDecoderThread::Run()
 				}
 			}
 			CurrentRate = IntendedRate;
+
+			if (bDirectionChange)
+			{
+				UpdateTrackIterator(NewTime);
+			}
 		}
 
 		// Are we to do something?
@@ -204,6 +214,10 @@ uint32 FElectraProtronPlayer::FImpl::FDecoderThread::Run()
 			SeekTimeNumFramesSkipped = 0;
 		}
 
+		// Update the playback range if it changed.
+		HandlePlaybackRangeChanges();
+// FIXME: need to check if we are out of bounds and need to flush?
+
 		DecodeOneFrame();
 	}
 	DestroyDecoder();
@@ -242,6 +256,14 @@ bool FElectraProtronPlayer::FImpl::FDecoderThread::SetLooping(bool bInLooping)
 	return true;
 }
 
+void FElectraProtronPlayer::FImpl::FDecoderThread::SetPlaybackRange(TRange<FTimespan> InRange)
+{
+	FScopeLock lock(&PendingPlayRangeChange.Lock);
+	PendingPlayRangeChange.NewRange = MoveTemp(InRange);
+	PendingPlayRangeChange.bIsSet = true;
+}
+
+
 bool FElectraProtronPlayer::FImpl::FDecoderThread::HasReachedEnd()
 {
 	return bReachedEnd;
@@ -277,11 +299,6 @@ bool FElectraProtronPlayer::FImpl::FDecoderThread::IsPausedForSeek()
 	return bPausedForSeek;
 }
 
-void FElectraProtronPlayer::FImpl::FDecoderThread::SetSampleQueue(const TSharedPtr<FMediaSamples, ESPMode::ThreadSafe>& InNewSampleQueue)
-{
-	Params.SampleQueue = InNewSampleQueue;
-}
-
 void FElectraProtronPlayer::FImpl::FDecoderThread::SetSampleBuffer(const TSharedPtr<FElectraProtronPlayer::FImpl::FMP4TrackSampleBuffer, ESPMode::ThreadSafe>& InTrackSampleBuffer, FElectraProtronPlayer::FImpl::FGetSampleDlg InGetSampleDelegate)
 {
 	FScopeLock lock(&PendingBufferChange.Lock);
@@ -315,6 +332,12 @@ void FElectraProtronPlayer::FImpl::FDecoderThread::SetEstimatedPlaybackTime(FTim
 	CurrentTime = InTime;
 }
 
+FTimespan FElectraProtronPlayer::FImpl::FDecoderThread::GetEstimatedPlaybackTime()
+{
+	FScopeLock lock(&TimeLock);
+	return CurrentTime;
+}
+
 
 void FElectraProtronPlayer::FImpl::FDecoderThread::Flush(const TSharedPtr<FMediaEvent, ESPMode::ThreadSafe>& InFlushedSignal)
 {
@@ -322,6 +345,48 @@ void FElectraProtronPlayer::FImpl::FDecoderThread::Flush(const TSharedPtr<FMedia
 	bFlushPending = true;
 }
 
+void FElectraProtronPlayer::FImpl::FDecoderThread::UpdateTrackSampleDurationMap()
+{
+	// Only do this for video tracks.
+	if (DecoderTypeIndex != CodecTypeIndex(ElectraProtronUtils::FCodecInfo::EType::Video))
+	{
+		return;
+	}
+
+	// Clear the video sample cache.
+	Params.SampleQueueInterface->GetVideoCache().Empty();
+
+	// Iterate the track and collect the effective PTS values into a sorted list,
+	// then use their difference as the sample duration.
+	FTimespan ClipDur(TrackSampleBuffer->TrackAndCodecInfo->MP4Track->GetFullMovieDuration().GetAsTimespan());
+	struct FSampleTimes
+	{
+		FTimespan PTS;
+		FTimespan Duration;
+	};
+	TArray<FSampleTimes> SampleTimeDurations;
+	TSharedPtr<Electra::UtilitiesMP4::FMP4Track::FIterator, ESPMode::ThreadSafe> TkIt = TrackSampleBuffer->TrackAndCodecInfo->MP4Track->CreateIterator();
+	SampleTimeDurations.Reserve(TkIt->GetNumSamples());
+	auto It = TrackSampleBuffer->TrackAndCodecInfo->MP4Track->CreateIterator();
+	do
+	{
+		SampleTimeDurations.Emplace(FSampleTimes({.PTS=It->GetEffectivePTS().GetAsTimespan()}));
+	} while(It->Next());
+	SampleTimeDurations.Sort([](const FSampleTimes& a, const FSampleTimes& b){return a.PTS < b.PTS;});
+	for(int32 i=1,iMax=SampleTimeDurations.Num(); i<iMax; ++i)
+	{
+		SampleTimeDurations[i-1].Duration = SampleTimeDurations[i].PTS - SampleTimeDurations[i-1].PTS;
+	}
+	// The last sample's duration extends to the end of the clip.
+	SampleTimeDurations.Last().Duration = ClipDur - SampleTimeDurations.Last().PTS;
+	// Get the sample durations over into the lookup map.
+	SampleTimeToDurationMap.Empty();
+	for(auto& st : SampleTimeDurations)
+	{
+		SampleTimeToDurationMap.Emplace(st.PTS, st.Duration);
+	}
+}
+
 void FElectraProtronPlayer::FImpl::FDecoderThread::UpdateTrackIterator(const FTimespan& InForTime)
 {
 	if (!TrackSampleBuffer.IsValid())
@@ -338,15 +403,86 @@ void FElectraProtronPlayer::FImpl::FDecoderThread::UpdateTrackIterator(const FTi
 		UE_LOG(LogElectraProtron, Error, TEXT("%s"), *LastErrorMessage);
 		return;
 	}
-	// Get the timestamps of the media local timeline needed to adjust the decoded samples.
-	FirstSampleEffectiveStartTime = FTimespan(0);
-	LatestSampleEffectiveEndTime = FirstSampleEffectiveStartTime + TkIt->GetEffectiveTrackDuration().GetAsTimespan();
 	// Use the new iterator.
 	TrackIterator = MoveTemp(TkIt);
 	bReachedEnd = false;
+	// Recalculate the playrange iterators.
+	FirstRangeSampleIt.Reset();
+	LastRangeSampleIt.Reset();
+	HandlePlaybackRangeChanges();
 }
 
 
+FTimespan FElectraProtronPlayer::FImpl::FDecoderThread::ClampTimeIntoPlaybackRange(const FTimespan& InTime)
+{
+	if (InTime < PlaybackRange.GetLowerBoundValue())
+	{
+		return PlaybackRange.GetLowerBoundValue();
+	}
+	else if (InTime > PlaybackRange.GetUpperBoundValue())
+	{
+		return PlaybackRange.GetUpperBoundValue();
+	}
+	return InTime;
+}
+
+
+void FElectraProtronPlayer::FImpl::FDecoderThread::HandlePlaybackRangeChanges()
+{
+	// Can't do anything without a sample buffer.
+	if (!TrackSampleBuffer.IsValid())
+	{
+		return;
+	}
+
+	// Change in playback range?
+	if (PendingPlayRangeChange.bIsSet)
+	{
+		FScopeLock lock(&TimeLock);
+		PendingPlayRangeChange.bIsSet = false;
+		if (PlaybackRange != PendingPlayRangeChange.NewRange)
+		{
+			PlaybackRange = MoveTemp(PendingPlayRangeChange.NewRange);
+			FirstRangeSampleIt.Reset();
+			LastRangeSampleIt.Reset();
+		}
+	}
+	// If the range is not valid we set it to encompass the entire clip.
+	if (PlaybackRange.IsEmpty() || PlaybackRange.IsDegenerate() || !PlaybackRange.HasLowerBound() || !PlaybackRange.HasUpperBound() || PlaybackRange.GetLowerBoundValue() > PlaybackRange.GetUpperBoundValue())
+	{
+		PlaybackRange = TRange<FTimespan>(FTimespan(0), Params.SampleQueueInterface->GetMovieDuration());
+	}
+	// Need to reset the first/last iterators?
+	if (!FirstRangeSampleIt.IsValid())
+	{
+		// Locate the first and last sample numbers for the range
+		FTrackIterator RangeIt = TrackSampleBuffer->TrackAndCodecInfo->MP4Track->CreateIteratorAtKeyframe(Electra::FTimeValue().SetFromTimespan(PlaybackRange.GetLowerBoundValue()), Electra::FTimeValue::GetZero());
+		if (!RangeIt.IsValid())
+		{
+			LastErrorMessage = TrackSampleBuffer->TrackAndCodecInfo->MP4Track->GetLastError();
+			UE_LOG(LogElectraProtron, Error, TEXT("%s"), *LastErrorMessage);
+			return;
+		}
+		FirstRangeSampleIt = RangeIt->Clone();
+		// Move forward until we reach the end or both the effective DTS *and* PTS are greater or equal than the end of the range.
+		// We need to look at both DTS and PTS because the effective PTS can be smaller than the effective DTS due to composition time offsets.
+		while(!RangeIt->IsLastEffective())
+		{
+			if (RangeIt->GetEffectiveDTS().GetAsTimespan() >= PlaybackRange.GetUpperBoundValue() && RangeIt->GetEffectivePTS().GetAsTimespan() >= PlaybackRange.GetUpperBoundValue())
+			{
+				// We want the last iterator to represent the last sample included in the playback range,
+				// so we need to step one back here as we are currently outside the range.
+				RangeIt->PrevEffective();
+				break;
+			}
+			RangeIt->NextEffective();
+		}
+		LastRangeSampleIt = MoveTemp(RangeIt);
+	}
+}
+
+
+
 void FElectraProtronPlayer::FImpl::FDecoderThread::PerformFlush()
 {
 	CurrentDecoderOutput.Reset();
@@ -458,7 +594,8 @@ void FElectraProtronPlayer::FImpl::FDecoderThread::DecodeOneFrame()
 				}
 				CurrentInputSample->DecAU.DTS = CurrentInputSample->Sample->DTS;
 				CurrentInputSample->DecAU.PTS = CurrentInputSample->Sample->PTS;
-				CurrentInputSample->DecAU.Duration = CurrentInputSample->Sample->Duration;
+				const FTimespan* DurFromMap = SampleTimeToDurationMap.Find(CurrentInputSample->Sample->EffectivePTS);
+				CurrentInputSample->DecAU.Duration = DurFromMap ? *DurFromMap : CurrentInputSample->Sample->Duration;
 				CurrentInputSample->DecAU.UserValue = NextUserValue;
 				CurrentInputSample->DecAU.Flags = CurrentInputSample->Sample->bIsSyncOrRap ? EElectraDecoderFlags::IsSyncSample : EElectraDecoderFlags::None;
 				if (CurrentInputSample->Sample->bIsSyncOrRap)
@@ -480,7 +617,7 @@ void FElectraProtronPlayer::FImpl::FDecoderThread::DecodeOneFrame()
 				if (SeekTimeToDecodeTo.IsSet())
 				{
 					FTimespan Start = CurrentInputSample->Sample->EffectivePTS;
-					FTimespan End = CurrentInputSample->Sample->EffectivePTS + CurrentInputSample->Sample->Duration;
+					FTimespan End = Start + CurrentInputSample->DecAU.Duration;
 					FTimespan X = SeekTimeToDecodeTo.GetValue();
 					if (X >= Start && X < End)
 					{
@@ -559,14 +696,14 @@ void FElectraProtronPlayer::FImpl::FDecoderThread::DecodeOneFrame()
 			}
 
 			// Move to the next or previous frame
-			StepTrackIterator(true);
+			StepTrackIterator();
 		}
 	}
 }
 
-void FElectraProtronPlayer::FImpl::FDecoderThread::StepTrackIterator(bool bPerformFlush)
+void FElectraProtronPlayer::FImpl::FDecoderThread::StepTrackIterator()
 {
-	if (!TrackSampleBuffer.IsValid() || !TrackIterator.IsValid())
+	if (!TrackSampleBuffer.IsValid() || !TrackIterator.IsValid() || !FirstRangeSampleIt.IsValid() || !LastRangeSampleIt.IsValid())
 	{
 		return;
 	}
@@ -574,17 +711,14 @@ void FElectraProtronPlayer::FImpl::FDecoderThread::StepTrackIterator(bool bPerfo
 	// Going forward or backwards?
 	if (PlaybackDirection >= 0.0f)
 	{
-		if (!TrackIterator->NextEffective())
+		// Move to the next sample. If we are on the last sample of the playback range (inclusive) or there is no next sample
+		// we either have to loop or are done.
+		if (TrackIterator->GetSampleNumber() == LastRangeSampleIt->GetSampleNumber() || !TrackIterator->NextEffective())
 		{
-			if (bPerformFlush)
-			{
-				FlushForEndOrLooping();
-			}
+			FlushForEndOrLooping();
 			if (bShouldLoop)
 			{
-				TSharedPtr<Electra::UtilitiesMP4::FMP4Track::FIterator, ESPMode::ThreadSafe> TkIt = TrackSampleBuffer->TrackAndCodecInfo->MP4Track->CreateIterator(TrackIterator->GetFirstSampleInfo().SyncSampleNumber);
-				check(TkIt.IsValid());
-				TrackIterator = MoveTemp(TkIt);
+				TrackIterator = FirstRangeSampleIt->Clone();
 				++LoopIndex;
 			}
 			else
@@ -595,21 +729,18 @@ void FElectraProtronPlayer::FImpl::FDecoderThread::StepTrackIterator(bool bPerfo
 	}
 	else
 	{
-		if (!TrackIterator->PrevEffective())
+		// If we are on the first sample of the playback range (inclusive) or there is no previous sample to go back to
+		// we either have to loop or are done.
+		if (TrackIterator->GetSampleNumber() == FirstRangeSampleIt->GetSampleNumber() || !TrackIterator->PrevEffective())
 		{
-			if (bPerformFlush)
-			{
-				FlushForEndOrLooping();
-			}
+			FlushForEndOrLooping();
 			if (bShouldLoop)
 			{
 				// Re-check this is keyframe-only. For non keyframe-only we cannot go back to the last frame since we need to
 				// go to whatever the last keyframe is to decode up to the last frame from there.
 				if (TrackSampleBuffer.IsValid() && TrackSampleBuffer->TrackAndCodecInfo->bIsKeyframeOnlyFormat)
 				{
-					TSharedPtr<Electra::UtilitiesMP4::FMP4Track::FIterator, ESPMode::ThreadSafe> TkIt = TrackSampleBuffer->TrackAndCodecInfo->MP4Track->CreateIterator(TrackIterator->GetLastSampleInfo().LastSampleNumber);
-					check(TkIt.IsValid());
-					TrackIterator = MoveTemp(TkIt);
+					TrackIterator = LastRangeSampleIt->Clone();
 					--LoopIndex;
 				}
 				else
@@ -673,10 +804,14 @@ void FElectraProtronPlayer::FImpl::FDecoderThread::HandleOutputFrame()
 		{
 			check(bIsDrainingAtEOS);
 			bIsDrainingAtEOS = false;
+			// Assuming we loop back to the playback range set the skip-until to the start of the range.
+			// When we do not loop or seek somewhere else setting this up here does no harm either.
+			SeekTimeToDecodeTo = SeekTimeToHandleTo = PlaybackRange.GetLowerBoundValue();
+			SeekTimeNumFramesDecoded = SeekTimeNumFramesSkipped = 0;
 		}
 	}
 	// Handle the output.
-	if (CurrentDecoderOutput.IsValid() && InputForCurrentDecoderOutput.IsValid() && Params.SampleQueue.IsValid())
+	if (CurrentDecoderOutput.IsValid() && InputForCurrentDecoderOutput.IsValid() && Params.SampleQueueInterface.IsValid())
 	{
 		TOptional<FTimespan> AdjustToTime;
 		bool bSendToOutput = true;
@@ -687,7 +822,7 @@ void FElectraProtronPlayer::FImpl::FDecoderThread::HandleOutputFrame()
 			if (SeekTimeToHandleTo.IsSet())
 			{
 				FTimespan Start = InputForCurrentDecoderOutput->Sample->EffectivePTS;
-				FTimespan End = InputForCurrentDecoderOutput->Sample->EffectivePTS + InputForCurrentDecoderOutput->Sample->Duration;
+				FTimespan End = InputForCurrentDecoderOutput->Sample->EffectivePTS + InputForCurrentDecoderOutput->DecAU.Duration;
 				FTimespan X = SeekTimeToHandleTo.GetValue();
 				if (X >= Start && X < End)
 				{
@@ -706,13 +841,20 @@ void FElectraProtronPlayer::FImpl::FDecoderThread::HandleOutputFrame()
 			}
 		};
 
+		// Get the timestamps of the media local timeline needed to adjust the decoded samples.
+		FTimespan FirstSampleEffectiveStartTime = PlaybackRange.GetLowerBoundValue();
+		FTimespan LatestSampleEffectiveEndTime = PlaybackRange.GetUpperBoundValue();
+
 		if (DecoderTypeIndex == CodecTypeIndex(ElectraProtronUtils::FCodecInfo::EType::Video))
 		{
 			AdjustForSeek();
 
+			FTimespan RawPTS(InputForCurrentDecoderOutput->Sample->EffectivePTS);
+			FTimespan RawDur(InputForCurrentDecoderOutput->DecAU.Duration);
+
 			// When sending the output we need to wait until there is room to receive it.
 			// If the output is not to be sent we can handle it quickly without waiting.
-			if (!bSendToOutput || Params.SampleQueue->CanReceiveVideoSamples(1))
+			if (!bSendToOutput || Params.SampleQueueInterface->CanEnqueueVideoSample(RawPTS))
 			{
 				// First perform any processing that needs to be done regardless of whether the output will be sent or not.
 
@@ -766,11 +908,11 @@ void FElectraProtronPlayer::FImpl::FDecoderThread::HandleOutputFrame()
 					}
 				}
 
-				FTimespan PTS(InputForCurrentDecoderOutput->Sample->EffectivePTS);
-				FTimespan Dur(InputForCurrentDecoderOutput->Sample->Duration);
+				FTimespan PTS(RawPTS);
+				FTimespan Dur(RawDur);
 				if (AdjustToTime.IsSet())
 				{
-					FTimespan Diff(AdjustToTime.GetValue() - InputForCurrentDecoderOutput->Sample->EffectivePTS);
+					FTimespan Diff(AdjustToTime.GetValue() - RawPTS);
 					PTS += Diff;
 					Dur -= Diff;
 				}
@@ -787,12 +929,11 @@ void FElectraProtronPlayer::FImpl::FDecoderThread::HandleOutputFrame()
 					FTimespan Diff = PTS + Dur - LatestSampleEffectiveEndTime;
 					Dur -= Diff;
 				}
-
-				// Set the current time to be the end of this decoded frame, which is the PTS plus the
-				// frame's duration when playing forward, or just the PTS when playing in reverse.
-				TimeLock.Lock();
-				CurrentTime = PTS + (PlaybackDirection >= 0.0f ? Dur : FTimespan(0));
-				TimeLock.Unlock();
+				// Still useful for display?
+				if (Dur <= FTimespan::Zero())
+				{
+					bSendToOutput = false;
+				}
 
 				// Send the output along?
 				if (bSendToOutput)
@@ -835,7 +976,7 @@ void FElectraProtronPlayer::FImpl::FDecoderThread::HandleOutputFrame()
 					{
 						FElectraTextureSampleRef TextureSample = Params.TexturePool->AcquireShared();
 						TextureSample->Initialize(NewDecoderOutput.Get());
-						Params.SampleQueue->AddVideo(TextureSample);
+						Params.SampleQueueInterface->EnqueueVideoSample(TextureSample, RawPTS, RawDur);
 					}
 				}
 				CurrentDecoderOutput.Reset();
@@ -848,7 +989,7 @@ void FElectraProtronPlayer::FImpl::FDecoderThread::HandleOutputFrame()
 
 			// When sending the output we need to wait until there is room to receive it.
 			// If the output is not to be sent we can handle it quickly without waiting.
-			if (!bSendToOutput || Params.SampleQueue->CanReceiveAudioSamples(1))
+			if (!bSendToOutput || Params.SampleQueueInterface->CanEnqueueAudioSample())
 			{
 				TSharedPtr<IElectraDecoderAudioOutput, ESPMode::ThreadSafe> AudioDecoderOutput = StaticCastSharedPtr<IElectraDecoderAudioOutput>(CurrentDecoderOutput);
 
@@ -893,6 +1034,7 @@ void FElectraProtronPlayer::FImpl::FDecoderThread::HandleOutputFrame()
 
 						// Set the current time to be the end of this decoded frame, which is the PTS plus the
 						// frame's duration when playing forward, or just the PTS when playing in reverse.
+						// We do this for audio since it is difficult to get the actual playback position from the outside.
 						TimeLock.Lock();
 						CurrentTime = PTS.Time + (PlaybackDirection >= 0.0f ? Dur : FTimespan(0));
 						TimeLock.Unlock();
@@ -923,7 +1065,7 @@ void FElectraProtronPlayer::FImpl::FDecoderThread::HandleOutputFrame()
 
 									AudioChannelMapper.MapChannels(const_cast<void*>(AudioSample->GetBuffer()), (int32)AudioSample->GetAllocatedSize(),
 																	ElectraDecodersUtil::AdvancePointer(PCMBuffer, ByteOffsetToFirstSample), NumSamplesProduced * NumBytesPerFrame, NumSamplesProduced);
-									Params.SampleQueue->AddAudio(AudioSample);
+									Params.SampleQueueInterface->EnqueueAudioSample(AudioSample);
 								}
 								else
 								{
@@ -1013,7 +1155,15 @@ bool FElectraProtronPlayer::FImpl::FDecoderThread::CreateDecoder()
 		DecoderCfgParams.Add(TEXT("height"), FVariant((uint32) vi.Height));
 		DecoderCfgParams.Add(TEXT("max_width"), FVariant((uint32) vi.Width));
 		DecoderCfgParams.Add(TEXT("max_height"), FVariant((uint32) vi.Height));
-		DecoderCfgParams.Add(TEXT("max_output_buffers"), FVariant((uint32) 8));
+		uint32 kNumBufs = Params.SampleQueueInterface->GetMaxVideoFramesToCache() > 0 ? (uint32)Params.SampleQueueInterface->GetMaxVideoFramesToCache() : 8U;
+		// To support playback ranges with GOP codecs we need additional buffers when looping from the range end to its start
+		// since it is possible that we need to decode-discard frames prior to the range. Although these won't be used for display they
+		// will still be produced and thus need additional room.
+		if (!TrackSampleBuffer->TrackAndCodecInfo->bIsKeyframeOnlyFormat)
+		{
+			kNumBufs += 5;
+		}
+		DecoderCfgParams.Add(TEXT("max_output_buffers"), FVariant(kNumBufs));
 		if (vi.FrameRate.IsValid())
 		{
 			DecoderCfgParams.Add(TEXT("fps"), FVariant((double)vi.FrameRate.GetAsDouble()));
diff --git a/Engine/Plugins/Media/ElectraPlayer/Source/ElectraProtron/Private/Player/ElectraProtronPlayerImpl.cpp b/Engine/Plugins/Media/ElectraPlayer/Source/ElectraProtron/Private/Player/ElectraProtronPlayerImpl.cpp
--- a/Engine/Plugins/Media/ElectraPlayer/Source/ElectraProtron/Private/Player/ElectraProtronPlayerImpl.cpp
+++ b/Engine/Plugins/Media/ElectraPlayer/Source/ElectraProtron/Private/Player/ElectraProtronPlayerImpl.cpp
@@ -28,6 +28,11 @@ FElectraProtronPlayer::FImpl::FImpl()
 
 	// Create the track-by-type array upfront in case queries to tracks are made before opening a source.
 	UsableTrackArrayIndicesByType.SetNum(UE_ARRAY_COUNT(kCodecTrackIndexMap));
+
+	// Create the sample queue interface.
+	const int32 kVideoFramesToCacheAhead = 8;
+	const int32 kVideoFramesToCacheBehind = 8;
+	CurrentSampleQueueInterface = MakeShared<FSampleQueueInterface, ESPMode::ThreadSafe>(kVideoFramesToCacheAhead, kVideoFramesToCacheBehind);
 }
 
 FElectraProtronPlayer::FImpl::~FImpl()
@@ -47,6 +52,7 @@ void FElectraProtronPlayer::FImpl::Open(const FOpenParam& InParam, FElectraProtr
 {
 	FWorkerThreadMessage Msg;
 	FWorkerThreadMessage::FParamOpen open { InParam };
+	open.Param.SampleQueueInterface = CurrentSampleQueueInterface;
 	Msg.Param.Emplace<FWorkerThreadMessage::FParamOpen>(MoveTemp(open));
 	Msg.Self = AsShared();
 	Msg.Type = FWorkerThreadMessage::EType::Open;
@@ -140,6 +146,13 @@ uint32 FElectraProtronPlayer::FImpl::Run()
 					// Start loader threads when opening was successful.
 					if (LastErrorMessage.IsEmpty())
 					{
+						// Set the duration of the movie on the sample queue for looping/wrapping purposes.
+						CurrentSampleQueueInterface->SetMovieDuration(Duration);
+						// If there an initial playback range set then apply it, otherwise set the entire movie.
+						SetPlaybackTimeRange(open.Param.InitialPlaybackRange.Get(TRange<FTimespan>::Empty()));
+						// By default we start at the beginning of the playback range.
+						CurrentPlayPosTime = CurrentPlaybackRange.GetLowerBoundValue();
+
 						VideoLoaderThread.StartThread(open.Param.Filename, SharedPlayParams);
 						AudioLoaderThread.StartThread(open.Param.Filename, SharedPlayParams);
 
@@ -274,6 +287,7 @@ void FElectraProtronPlayer::FImpl::InternalOpen(const FString& InFilename)
 	int32 NumTimecodeTracks = 0;
 	int64 LongestTrackDuration = -1;
 	int64 LongestSupportedTrackDuration = -1;
+	int64 ShortestSupportedTrackDuration = TNumericLimits<int64>::Max();
 	for(int32 i=0; i<AllTracks.Num(); ++i)
 	{
 		Tracks[i] = MakeShared<FTrackInfo>();
@@ -385,6 +399,8 @@ void FElectraProtronPlayer::FImpl::InternalOpen(const FString& InFilename)
 
 				// Take note of the supported track with the longest duration.
 				LongestSupportedTrackDuration = LongestSupportedTrackDuration < Tkhd->GetDuration() ? Tkhd->GetDuration() : LongestSupportedTrackDuration;
+				// Likewise for the shortest.
+				ShortestSupportedTrackDuration = ShortestSupportedTrackDuration > Tkhd->GetDuration() ? Tkhd->GetDuration() : ShortestSupportedTrackDuration;
 			}
 			else
 			{
@@ -395,6 +411,7 @@ void FElectraProtronPlayer::FImpl::InternalOpen(const FString& InFilename)
 
 	// Check that the duration given in the `mvhd` box is not larger than the longest track is.
 	int64 MvhdDur = MovieDuration.GetNumerator();
+
 	if (MvhdDur > LongestTrackDuration)
 	{
 		UE_LOG(LogElectraProtron, Warning, TEXT("Movie duration in `mvhd` box (%lld) is larger than the longest track (%lld) in the file, adjusting movie duration down."), (long long int)MvhdDur, (long long int)LongestTrackDuration);
@@ -402,6 +419,9 @@ void FElectraProtronPlayer::FImpl::InternalOpen(const FString& InFilename)
 		Duration = MovieDuration.GetAsTimebase(ETimespan::TicksPerSecond);
 		MvhdDur = LongestTrackDuration;
 	}
+	Electra::FTimeFraction EntireMovieDuration(MovieDuration);
+
+#if 0
 	// Check that the movie duration is not larger than the longest supported track.
 	if (LongestSupportedTrackDuration > 0 && MvhdDur > LongestSupportedTrackDuration)
 	{
@@ -409,7 +429,15 @@ void FElectraProtronPlayer::FImpl::InternalOpen(const FString& InFilename)
 		MovieDuration.SetFromND(LongestSupportedTrackDuration, MovieDuration.GetDenominator());
 		Duration = MovieDuration.GetAsTimebase(ETimespan::TicksPerSecond);
 	}
-
+#else
+	// Check that the movie duration is not larger than the longest supported track.
+	if (ShortestSupportedTrackDuration < TNumericLimits<int64>::Max() && MvhdDur > ShortestSupportedTrackDuration)
+	{
+		UE_LOG(LogElectraProtron, Warning, TEXT("Movie duration in `mvhd` box (%lld) is larger than the shortest supported track (%lld) in the file, adjusting movie duration down."), (long long int)MvhdDur, (long long int)ShortestSupportedTrackDuration);
+		MovieDuration.SetFromND(ShortestSupportedTrackDuration, MovieDuration.GetDenominator());
+		Duration = MovieDuration.GetAsTimebase(ETimespan::TicksPerSecond);
+	}
+#endif
 
 	// If there are timecode tracks, find which tracks references them.
 	// If there are none, then - if there is only a single timecode track - we apply it to all other tracks.
@@ -571,7 +599,7 @@ void FElectraProtronPlayer::FImpl::InternalOpen(const FString& InFilename)
 			continue;
 		}
 		Tracks[i]->MP4Track = Electra::UtilitiesMP4::FMP4Track::Create(Tracks[i]->TrackBox);
-		if (!Tracks[i]->MP4Track->Prepare(MovieDuration))
+		if (!Tracks[i]->MP4Track->Prepare(EntireMovieDuration, MovieDuration))
 		{
 			LastErrorMessage = Tracks[i]->MP4Track->GetLastError();
 			check(!LastErrorMessage.IsEmpty());
@@ -1203,6 +1231,18 @@ bool FElectraProtronPlayer::FImpl::QueryCacheState(EMediaCacheState InState, TRa
 			return true;
 		}
 	}
+	else if (InState == EMediaCacheState::Loaded)
+	{
+		if (TrackSelection.SelectedTrackIndex[CodecTypeIndex(ElectraProtronUtils::FCodecInfo::EType::Video)] >= 0)
+		{
+			TSharedPtr<FSampleQueueInterface, ESPMode::ThreadSafe> sqi(GetCurrentSampleQueueInterface());
+			if (sqi.IsValid())
+			{
+				sqi->GetVideoCache().QueryCacheState(OutTimeRanges);
+				return true;
+			}
+		}
+	}
 	return false;
 }
 
@@ -1251,17 +1291,59 @@ bool FElectraProtronPlayer::FImpl::IsLooping()
 	return SharedPlayParams->bShouldLoop;
 }
 
-bool FElectraProtronPlayer::FImpl::Seek(const FTimespan& InTime, const TSharedPtr<FMediaSamples, ESPMode::ThreadSafe>& InNewSampleQueue, int32 InNewSequenceIndex, const TOptional<int32>& InNewLoopIndex)
+void FElectraProtronPlayer::FImpl::Seek(const FTimespan& InTime, int32 InNewSequenceIndex, const TOptional<int32>& InNewLoopIndex)
 {
 	FSeekRequest sr;
-	sr.NewSampleQueue = InNewSampleQueue;
 	sr.NewTime = InTime;
 	sr.NewSequenceIndex = InNewSequenceIndex;
 	sr.NewLoopIndex = InNewLoopIndex;
+	if (CurrentSampleQueueInterface.IsValid())
+	{
+		CurrentSampleQueueInterface->SeekIssuedTo(InTime, TOptional<int32>(InNewSequenceIndex));
+	}
 	SeekRequestLock.Lock();
 	PendingSeekRequest = MoveTemp(sr);
 	SeekRequestLock.Unlock();
 	WorkMessageSignal.Signal();
+}
+
+TRange<FTimespan> FElectraProtronPlayer::FImpl::GetPlaybackTimeRange(EMediaTimeRangeType InRangeToGet)
+{
+	return InRangeToGet == EMediaTimeRangeType::Absolute ? TRange<FTimespan>(FTimespan(0), GetDuration()) : CurrentPlaybackRange;
+}
+
+bool FElectraProtronPlayer::FImpl::SetPlaybackTimeRange(const TRange<FTimespan>& InTimeRange)
+{
+	// For proper validation we need to have the content duration.
+	if (Duration <= FTimespan::Zero() || InTimeRange.IsDegenerate() || !InTimeRange.HasLowerBound() || !InTimeRange.HasUpperBound() || InTimeRange.GetLowerBoundValue() > InTimeRange.GetUpperBoundValue())
+	{
+		return false;
+	}
+	// If we get an empty range we instead set the range to be the entire movie.
+	if (InTimeRange.IsEmpty())
+	{
+		CurrentPlaybackRange = TRange<FTimespan>(FTimespan(0), GetDuration());
+	}
+	else
+	{
+		TRange<FTimespan> r(InTimeRange);
+		if (r.GetLowerBoundValue() < FTimespan::Zero())
+		{
+			UE_LOG(LogElectraProtron, Warning, TEXT("Clamping start of playback range to zero as it was set negative."));
+			r.SetLowerBoundValue(FTimespan::Zero());
+		}
+		if (r.GetUpperBoundValue() > Duration)
+		{
+			UE_LOG(LogElectraProtron, Warning, TEXT("Clamping end of playback range to movie duration as it was set larger."));
+			r.SetUpperBoundValue(Duration);
+		}
+		CurrentPlaybackRange = MoveTemp(r);
+	}
+	CurrentSampleQueueInterface->SetPlaybackRange(CurrentPlaybackRange);
+	VideoLoaderThread.SetPlaybackRange(CurrentPlaybackRange);
+	AudioLoaderThread.SetPlaybackRange(CurrentPlaybackRange);
+	VideoDecoderThread.SetPlaybackRange(CurrentPlaybackRange);
+	AudioDecoderThread.SetPlaybackRange(CurrentPlaybackRange);
 	return true;
 }
 
@@ -1291,25 +1373,35 @@ void FElectraProtronPlayer::FImpl::TickInput(FTimespan InDeltaTime, FTimespan In
 		}
 	}
 
-	double Rate = CurrentRate;
-	if (Rate != 0.0)
+	auto sqi = GetCurrentSampleQueueInterface();
+	if (sqi.IsValid())
 	{
-		FTimespan ScaledDeltaTime = InDeltaTime * Rate;
-		FTimespan NewPos = CurrentPlayPosTime + ScaledDeltaTime;
-		if (NewPos < FTimespan::Zero())
+		FMediaTimeStamp ts = sqi->GetLastHandedOutTimestamp();
+		const bool bIsVideoActive = TrackSelection.ActiveTrackIndex[CodecTypeIndex(ElectraProtronUtils::FCodecInfo::EType::Video)] != -1;
+		const bool bIsAudioActive = TrackSelection.ActiveTrackIndex[CodecTypeIndex(ElectraProtronUtils::FCodecInfo::EType::Audio)] != -1;
+		if (ts.IsValid())
 		{
-			NewPos = IsLooping() ? Duration - (NewPos % Duration) : FTimespan::Zero();
+			FTimespan NewPos = ts.GetTime();
+			if (bIsVideoActive)
+			{
+				VideoDecoderThread.SetEstimatedPlaybackTime(NewPos);
+				CurrentPlayPosTime = NewPos;
+			}
+			if (!bIsAudioActive)
+			{
+				AudioDecoderThread.SetEstimatedPlaybackTime(NewPos);
+			}
 		}
-		else if (NewPos > Duration)
+		else if (bIsAudioActive)
 		{
-			NewPos = IsLooping() ? NewPos % Duration : Duration;
+			FTimespan NewPos = AudioDecoderThread.GetEstimatedPlaybackTime();
+			VideoDecoderThread.SetEstimatedPlaybackTime(NewPos);
+			CurrentPlayPosTime = NewPos;
 		}
-		VideoDecoderThread.SetEstimatedPlaybackTime(NewPos);
-		AudioDecoderThread.SetEstimatedPlaybackTime(NewPos);
-		CurrentPlayPosTime = NewPos;
 	}
 }
 
+
 void FElectraProtronPlayer::FImpl::UpdateTrackLoader(int32 InCodecTypeIndex)
 {
 	if (TrackSelection.SelectedTrackIndex[InCodecTypeIndex] >= 0)
@@ -1339,3 +1431,100 @@ void FElectraProtronPlayer::FImpl::UpdateTrackLoader(int32 InCodecTypeIndex)
 		}
 	}
 }
+
+
+
+IMediaSamples::EFetchBestSampleResult FElectraProtronPlayer::FImpl::FetchBestVideoSampleForTimeRange(const TRange<FMediaTimeStamp>& InTimeRange, TSharedPtr<IMediaTextureSample, ESPMode::ThreadSafe>& OutSample, bool bInReverse, bool bInConsistentResult)
+{
+	TSharedPtr<FSampleQueueInterface, ESPMode::ThreadSafe> sqi(GetCurrentSampleQueueInterface());
+	if (sqi.IsValid())
+	{
+		FProtronVideoCache::EGetResult gr = sqi->GetVideoCache().GetFrame(OutSample, InTimeRange, IsLooping(), bInReverse, bInConsistentResult);
+		if (gr == FProtronVideoCache::EGetResult::Hit)
+		{
+			sqi->UpdateNextExpectedTimestamp(OutSample, bInReverse, IsLooping());
+			sqi->UpdateLastHandedOutTimestamp(OutSample);
+			return IMediaSamples::EFetchBestSampleResult::Ok;
+		}
+		else if (gr == FProtronVideoCache::EGetResult::PurgedEmpty)
+		{
+			sqi->ResetCurrentTimestamps();
+			return IMediaSamples::EFetchBestSampleResult::PurgedToEmpty;
+		}
+	}
+	return IMediaSamples::EFetchBestSampleResult::NoSample;
+}
+bool FElectraProtronPlayer::FImpl::FetchAudio(TRange<FMediaTimeStamp> InTimeRange, TSharedPtr<IMediaAudioSample, ESPMode::ThreadSafe>& OutSample)
+{
+	TSharedPtr<FSampleQueueInterface, ESPMode::ThreadSafe> sqi(GetCurrentSampleQueueInterface());
+	return sqi.IsValid() ? sqi->GetCurrentSampleQueue()->FetchAudio(InTimeRange, OutSample) : false;
+}
+bool FElectraProtronPlayer::FImpl::FetchCaption(TRange<FMediaTimeStamp> InTimeRange, TSharedPtr<IMediaOverlaySample, ESPMode::ThreadSafe>& OutSample)
+{
+	return false;
+}
+bool FElectraProtronPlayer::FImpl::FetchMetadata(TRange<FMediaTimeStamp> InTimeRange, TSharedPtr<IMediaBinarySample, ESPMode::ThreadSafe>& OutSample)
+{
+	return false;
+}
+bool FElectraProtronPlayer::FImpl::FetchSubtitle(TRange<FMediaTimeStamp> InTimeRange, TSharedPtr<IMediaOverlaySample, ESPMode::ThreadSafe>& OutSample)
+{
+	return false;
+}
+void FElectraProtronPlayer::FImpl::FlushSamples()
+{
+}
+void FElectraProtronPlayer::FImpl::SetMinExpectedNextSequenceIndex(TOptional<int32> InNextSequenceIndex)
+{
+	TSharedPtr<FSampleQueueInterface, ESPMode::ThreadSafe> sqi(GetCurrentSampleQueueInterface());
+	if (sqi.IsValid())
+	{
+		sqi->GetCurrentSampleQueue()->SetMinExpectedNextSequenceIndex(InNextSequenceIndex);
+	}
+}
+bool FElectraProtronPlayer::FImpl::PeekVideoSampleTime(FMediaTimeStamp& OutTimeStamp)
+{
+	TSharedPtr<FSampleQueueInterface, ESPMode::ThreadSafe> sqi(GetCurrentSampleQueueInterface());
+	return sqi.IsValid() ? sqi->PeekVideoSampleTime(OutTimeStamp) : false;
+}
+bool FElectraProtronPlayer::FImpl::CanReceiveVideoSamples(uint32 InNum) const
+{
+	return true;
+}
+bool FElectraProtronPlayer::FImpl::CanReceiveAudioSamples(uint32 InNum) const
+{
+	TSharedPtr<FSampleQueueInterface, ESPMode::ThreadSafe> sqi(GetCurrentSampleQueueInterface());
+	return sqi.IsValid() ? sqi->CanEnqueueAudioSample() : true;
+}
+bool FElectraProtronPlayer::FImpl::CanReceiveSubtitleSamples(uint32 InNum) const
+{
+	return true;
+}
+bool FElectraProtronPlayer::FImpl::CanReceiveCaptionSamples(uint32 InNum) const
+{
+	return true;
+}
+bool FElectraProtronPlayer::FImpl::CanReceiveMetadataSamples(uint32 InNum) const
+{
+	return true;
+}
+int32 FElectraProtronPlayer::FImpl::NumAudioSamples() const
+{
+	return 0;
+}
+int32 FElectraProtronPlayer::FImpl::NumCaptionSamples() const
+{
+	return 0;
+}
+int32 FElectraProtronPlayer::FImpl::NumMetadataSamples() const
+{
+	return 0;
+}
+int32 FElectraProtronPlayer::FImpl::NumSubtitleSamples() const
+{
+	return 0;
+}
+int32 FElectraProtronPlayer::FImpl::NumVideoSamples() const
+{
+	return 0;
+}
diff --git a/Engine/Plugins/Media/ElectraPlayer/Source/ElectraProtron/Private/Player/ElectraProtronPlayerImpl.h b/Engine/Plugins/Media/ElectraPlayer/Source/ElectraProtron/Private/Player/ElectraProtronPlayerImpl.h
--- a/Engine/Plugins/Media/ElectraPlayer/Source/ElectraProtron/Private/Player/ElectraProtronPlayerImpl.h
+++ b/Engine/Plugins/Media/ElectraPlayer/Source/ElectraProtron/Private/Player/ElectraProtronPlayerImpl.h
@@ -21,6 +21,9 @@
 #include "Utils/MPEG/ElectraUtilsMPEGVideo.h"
 #include "Utils/AudioChannelMapper.h"
 #include "Decoder/VideoDecoderHelpers.h"
+#include "MediaSamples.h"
+#include "ElectraProtronPlayerCache.h"
+
 
 
 /**
@@ -30,6 +33,186 @@
 class FElectraProtronPlayer::FImpl : public TSharedFromThis<FElectraProtronPlayer::FImpl, ESPMode::ThreadSafe>, public FRunnable
 {
 public:
+	class FSampleQueueInterface : public TSharedFromThis<FSampleQueueInterface, ESPMode::ThreadSafe>
+	{
+	public:
+		FSampleQueueInterface(int32 InNumVideoFramesToCacheAhead, int32 InNumVideoFramesToCacheBehind)
+			: SampleQueue(MakeShared<FMediaSamples, ESPMode::ThreadSafe>())
+			, NumVideoFramesToCache(InNumVideoFramesToCacheAhead + InNumVideoFramesToCacheBehind)
+		{
+			// We need to have some future frames to mimick the behavior of the FMediaSample struct.
+			check(InNumVideoFramesToCacheAhead >= 4);
+			// And we also need to retain some old samples, which is the whole point of having a cache.
+			check(InNumVideoFramesToCacheBehind >= 4);
+			VideoCache.SetMaxFramesToCache(InNumVideoFramesToCacheAhead, InNumVideoFramesToCacheBehind);
+		}
+
+		int32 GetMaxVideoFramesToCache()
+		{
+			return NumVideoFramesToCache;
+		}
+
+		void SetMovieDuration(FTimespan InDuration)
+		{
+			Duration = InDuration;
+			VideoCache.SetPlaybackRange(TRange<FTimespan>(FTimespan(0), InDuration));
+		}
+
+		FTimespan GetMovieDuration()
+		{
+			return Duration;
+		}
+
+		void SetPlaybackRange(TRange<FTimespan> InRange)
+		{
+			PlaybackRange = MoveTemp(InRange);
+		}
+
+		TRange<FTimespan> GetPlaybackRange()
+		{
+			return PlaybackRange;
+		}
+
+		void SetPlaybackRate(float InNewRate)
+		{
+			PlaybackRate = InNewRate;
+			VideoCache.SetPlaybackRate(InNewRate);
+		}
+
+		void SeekIssuedTo(FTimespan InToTime, TOptional<int32> InNextSequenceIndex)
+		{
+			MinSeqIdx = InNextSequenceIndex;
+			SampleQueue->SetMinExpectedNextSequenceIndex(InNextSequenceIndex);
+			VideoCache.SeekIssuedTo(InToTime);
+			FScopeLock lock(&TimestampLock);
+			NextExpectedTimestamp.Invalidate();
+			LastHandedOutTimestamp.Invalidate();
+		}
+
+		bool CanEnqueueVideoSample(FTimespan InPTS)
+		{
+			return VideoCache.CanAccept(InPTS);
+		}
+
+		bool CanEnqueueAudioSample()
+		{
+			return SampleQueue->CanReceiveAudioSamples(1);
+		}
+
+		void EnqueueVideoSample(const TSharedRef<IMediaTextureSample, ESPMode::ThreadSafe>& InSample, FTimespan InRawPTS, FTimespan InRawDuration)
+		{
+			if (InSample->GetTime().GetSequenceIndex() < MinSeqIdx.Get(0))
+			{
+				return;
+			}
+			VideoCache.AddFrame(InSample, InRawPTS, InRawDuration);
+
+			FScopeLock lock(&TimestampLock);
+			if (!NextExpectedTimestamp.IsValid())
+			{
+				NextExpectedTimestamp = InSample->GetTime();
+			}
+		}
+
+		void EnqueueAudioSample(const TSharedRef<IMediaAudioSample, ESPMode::ThreadSafe>& InSample)
+		{
+			SampleQueue->AddAudio(InSample);
+		}
+
+		TSharedPtr<FMediaSamples, ESPMode::ThreadSafe> GetCurrentSampleQueue()
+		{
+			return SampleQueue;
+		}
+
+		bool PeekVideoSampleTime(FMediaTimeStamp& OutTimeStamp)
+		{
+			FScopeLock lock(&TimestampLock);
+			if (NextExpectedTimestamp.IsValid())
+			{
+				OutTimeStamp = NextExpectedTimestamp;
+				return true;
+			}
+			return false;
+		}
+
+		void UpdateLastHandedOutTimestamp(const TSharedPtr<IMediaTextureSample, ESPMode::ThreadSafe>& InSample)
+		{
+			FScopeLock lock(&TimestampLock);
+			LastHandedOutTimestamp = InSample->GetTime();
+		}
+
+		FMediaTimeStamp GetLastHandedOutTimestamp()
+		{
+			FScopeLock lock(&TimestampLock);
+			return LastHandedOutTimestamp;
+		}
+
+		void UpdateNextExpectedTimestamp(const TSharedPtr<IMediaTextureSample, ESPMode::ThreadSafe>& InSample, bool bInReverse, bool bInIsLooping)
+		{
+			FScopeLock lock(&TimestampLock);
+			if (!bInReverse)
+			{
+				NextExpectedTimestamp = InSample->GetTime() + InSample->GetDuration();
+				if (NextExpectedTimestamp.GetTime() >= PlaybackRange.GetUpperBoundValue())
+				{
+					if (bInIsLooping)
+					{
+						NextExpectedTimestamp -= PlaybackRange.GetUpperBoundValue() - PlaybackRange.GetLowerBoundValue();
+						NextExpectedTimestamp.AdjustLoopIndex(1);
+					}
+					else
+					{
+						// Set to the time of the last sample. This must be less than the end of the
+						// playback range to work.
+						NextExpectedTimestamp.SetTime(InSample->GetTime().GetTime());
+					}
+				}
+			}
+			else
+			{
+				NextExpectedTimestamp = InSample->GetTime() - InSample->GetDuration();
+				if (NextExpectedTimestamp.GetTime() < PlaybackRange.GetLowerBoundValue())
+				{
+					if (bInIsLooping)
+					{
+						NextExpectedTimestamp += PlaybackRange.GetUpperBoundValue() - PlaybackRange.GetLowerBoundValue();
+						NextExpectedTimestamp.AdjustLoopIndex(-1);
+					}
+					else
+					{
+						// Set to the lower bound of the playback range
+						NextExpectedTimestamp.SetTime(PlaybackRange.GetLowerBoundValue());
+					}
+				}
+			}
+		}
+
+		void ResetCurrentTimestamps()
+		{
+			FScopeLock lock(&TimestampLock);
+			LastHandedOutTimestamp.Invalidate();
+			NextExpectedTimestamp.Invalidate();
+		}
+
+		FProtronVideoCache& GetVideoCache()
+		{
+			return VideoCache;
+		}
+
+	private:
+		FProtronVideoCache VideoCache;
+		TSharedPtr<FMediaSamples, ESPMode::ThreadSafe> SampleQueue;
+		FCriticalSection TimestampLock;
+		FMediaTimeStamp NextExpectedTimestamp;
+		FMediaTimeStamp LastHandedOutTimestamp;
+		TOptional<int32> MinSeqIdx;
+		FTimespan Duration;
+		TRange<FTimespan> PlaybackRange;
+		float PlaybackRate = 0.0f;
+		int32 NumVideoFramesToCache = 0;
+	};
+
+
 	// Define pointer type to get around the ',' in the name that confuses the delegate declaration macro.
 	using ImplPointer = TSharedPtr<FImpl, ESPMode::ThreadSafe>;
 	DECLARE_DELEGATE_OneParam(FCompletionDelegate, ImplPointer);
@@ -39,9 +222,10 @@ public:
 	struct FOpenParam
 	{
 		FString Filename;
-		TSharedPtr<FMediaSamples, ESPMode::ThreadSafe> SampleQueue;
+		TSharedPtr<FSampleQueueInterface, ESPMode::ThreadSafe> SampleQueueInterface;
 		TSharedPtr<FElectraTextureSamplePool, ESPMode::ThreadSafe> TexturePool;
 		TSharedPtr<FElectraAudioSamplePool, ESPMode::ThreadSafe> AudioSamplePool;
+		TOptional<TRange<FTimespan>> InitialPlaybackRange;
 	};
 	void Open(const FOpenParam& InParam, FCompletionDelegate InCompletionDelegate);
 	void Close(FCompletionDelegate InCompletionDelegate);
@@ -72,16 +256,37 @@ public:
 	FTimespan GetTime();
 	bool SetLooping(bool bInLooping);
 	bool IsLooping();
-	bool Seek(const FTimespan& InTime, const TSharedPtr<FMediaSamples, ESPMode::ThreadSafe>& InNewSampleQueue, int32 InNewSequenceIndex, const TOptional<int32>& InNewLoopIndex);
+	void Seek(const FTimespan& InTime, int32 InNewSequenceIndex, const TOptional<int32>& InNewLoopIndex);
+	TRange<FTimespan> GetPlaybackTimeRange(EMediaTimeRangeType InRangeToGet);
+	bool SetPlaybackTimeRange(const TRange<FTimespan>& InTimeRange);
 
 	void TickFetch(FTimespan InDeltaTime, FTimespan InTimecode);
 	void TickInput(FTimespan InDeltaTime, FTimespan InTimecode);
 
-
 	// Inherited from FRunnable
 	uint32 Run() override;
 	void Exit() override;
 
+	// Forwarded IMediaSamples interface methods from the enclosing FElectraProtronPlayer
+	EFetchBestSampleResult FetchBestVideoSampleForTimeRange(const TRange<FMediaTimeStamp>& TimeRange, TSharedPtr<IMediaTextureSample, ESPMode::ThreadSafe>& OutSample, bool bInReverse, bool bInConsistentResult);
+	bool FetchAudio(TRange<FMediaTimeStamp> InTimeRange, TSharedPtr<IMediaAudioSample, ESPMode::ThreadSafe>& OutSample);
+	bool FetchCaption(TRange<FMediaTimeStamp> InTimeRange, TSharedPtr<IMediaOverlaySample, ESPMode::ThreadSafe>& OutSample);
+	bool FetchMetadata(TRange<FMediaTimeStamp> InTimeRange, TSharedPtr<IMediaBinarySample, ESPMode::ThreadSafe>& OutSample);
+	bool FetchSubtitle(TRange<FMediaTimeStamp> InTimeRange, TSharedPtr<IMediaOverlaySample, ESPMode::ThreadSafe>& OutSample);
+	void FlushSamples();
+	void SetMinExpectedNextSequenceIndex(TOptional<int32> InNextSequenceIndex);
+	bool PeekVideoSampleTime(FMediaTimeStamp& OutTimeStamp);
+	bool CanReceiveVideoSamples(uint32 InNum) const;
+	bool CanReceiveAudioSamples(uint32 InNum) const;
+	bool CanReceiveSubtitleSamples(uint32 InNum) const;
+	bool CanReceiveCaptionSamples(uint32 InNum) const;
+	bool CanReceiveMetadataSamples(uint32 InNum) const;
+	int32 NumAudioSamples() const;
+	int32 NumCaptionSamples() const;
+	int32 NumMetadataSamples() const;
+	int32 NumSubtitleSamples() const;
+	int32 NumVideoSamples() const;
+
 private:
 	static constexpr int32 CodecTypeIndex(const ElectraProtronUtils::FCodecInfo::EType InType)
 	{ return (int32) InType; }
@@ -104,7 +309,6 @@ private:
 
 	struct FSeekRequest
 	{
-		TSharedPtr<FMediaSamples, ESPMode::ThreadSafe> NewSampleQueue;
 		FTimespan NewTime;
 		int32 NewSequenceIndex = 0;
 		TOptional<int32> NewLoopIndex;
@@ -177,6 +381,8 @@ private:
 	};
 	using FMP4SamplePtr = TSharedPtr<FMP4Sample, ESPMode::ThreadSafe>;
 
+	using FTrackIterator = TSharedPtr<Electra::UtilitiesMP4::FMP4Track::FIterator, ESPMode::ThreadSafe>;
+
 	struct FMP4TrackSampleBuffer
 	{
 		FCriticalSection Lock;
@@ -184,10 +390,14 @@ private:
 		TSortedMap<uint32, FMP4SamplePtr> SampleMap;
 		TSharedPtr<FTrackInfo, ESPMode::ThreadSafe> TrackAndCodecInfo;
 		uint32 TrackID;
+
+		// Used by the sample loader
+		TRange<FTimespan> CurrentPlaybackRange;
+		FTrackIterator FirstRangeSampleIt;
+		FTrackIterator LastRangeSampleIt;
 	};
 	using FMP4TrackSampleBufferPtr = TSharedPtr<FMP4TrackSampleBuffer, ESPMode::ThreadSafe>;
 
-	using FTrackIterator = TSharedPtr<Electra::UtilitiesMP4::FMP4Track::FIterator, ESPMode::ThreadSafe>;
 
 	void StartThread();
 	void SendWorkerThreadMessage(FWorkerThreadMessage&& InMessage);
@@ -229,6 +439,7 @@ private:
 	TRangeSet<float> ThinnedRates;
 	bool bAreRatesValid = false;
 	FTimespan CurrentPlayPosTime;
+	TRange<FTimespan> CurrentPlaybackRange;
 	float CurrentRate = 0.0f;
 	float IntendedRate = 0.0f;
 
@@ -239,6 +450,13 @@ private:
 
 	TSharedPtr<FSharedPlayParams, ESPMode::ThreadSafe> SharedPlayParams;
 
+	TSharedPtr<FSampleQueueInterface, ESPMode::ThreadSafe> CurrentSampleQueueInterface;
+	TSharedPtr<FSampleQueueInterface, ESPMode::ThreadSafe> GetCurrentSampleQueueInterface() const
+	{
+		return CurrentSampleQueueInterface;
+	}
+
+
 	// ===== Loader and decoder threads =====
 
 	DECLARE_DELEGATE_RetVal_ThreeParams(FMP4SamplePtr, FGetSampleDlg, FMP4TrackSampleBufferPtr /*FromBuffer*/, const FTrackIterator& /*AtIterator*/, int32 /*WaitMicros*/);
@@ -261,6 +479,7 @@ private:
 		void StopThread();
 
 		// Called by the player
+		void SetPlaybackRange(TRange<FTimespan> InRange);
 		void RequestLoad(FMP4TrackSampleBufferPtr InTrackSampleBuffer, FTimespan InTime);
 		TRangeSet<FTimespan> GetTimeRangesToLoad();
 
@@ -300,6 +519,7 @@ private:
 		FOpenRequest OpenRequest;
 		FLoadRequest PendingLoadRequest;
 		FLoadRequest ActiveLoadRequest;
+		TRange<FTimespan> PlaybackRange;
 		FCriticalSection TimeRangeLock;
 		TRangeSet<FTimespan> TimeRangesToLoad;
 		volatile int32 LoadRequestDirty = -1;
@@ -351,6 +571,9 @@ private:
 		// Whether or not to loop.
 		bool SetLooping(bool bInLooping);
 
+		// Sets the active playback range.
+		void SetPlaybackRange(TRange<FTimespan> InRange);
+
 		// Did decoding reach the last sample (or first sample when going in reverse)?
 		bool HasReachedEnd();
 
@@ -368,13 +591,16 @@ private:
 		bool IsPausedForSeek();
 		void SetTime(FTimespan InTime, int32 InSeqIdx, TOptional<int32> InLoopIdx);
 		void SetEstimatedPlaybackTime(FTimespan InTime);
+		FTimespan GetEstimatedPlaybackTime();
 
 		void Flush(const TSharedPtr<FMediaEvent, ESPMode::ThreadSafe>& InFlushedSignal);
-		void SetSampleQueue(const TSharedPtr<FMediaSamples, ESPMode::ThreadSafe>& InNewSampleQueue);
 	private:
 
+		void UpdateTrackSampleDurationMap();
+		void HandlePlaybackRangeChanges();
+		FTimespan ClampTimeIntoPlaybackRange(const FTimespan& InTime);
 		void UpdateTrackIterator(const FTimespan& InForTime);
-		void StepTrackIterator(bool bPerformFlush);
+		void StepTrackIterator();
 		bool CreateDecoder();
 		void DestroyDecoder();
 		void DecodeOneFrame();
@@ -410,6 +636,13 @@ private:
 			bool bIsSet = false;
 		};
 
+		struct FPendingPlayRange
+		{
+			FCriticalSection Lock;
+			TRange<FTimespan> NewRange;
+			bool bIsSet = false;
+		};
+
 		const FConfig& Config;
 		int32 DecoderTypeIndex;
 		FRunnableThread* Thread = nullptr;
@@ -420,13 +653,16 @@ private:
 		FOpenParam Params;
 		TSharedPtr<FSharedPlayParams, ESPMode::ThreadSafe> SharedPlayParams;
 		TSharedPtr<FElectraProtronPlayer::FImpl::FMP4TrackSampleBuffer, ESPMode::ThreadSafe> TrackSampleBuffer;
+		TSortedMap<FTimespan, FTimespan> SampleTimeToDurationMap;
 		FGetSampleDlg GetSampleDelegate;
 		FTrackIterator TrackIterator;
-		FTimespan FirstSampleEffectiveStartTime;
-		FTimespan LatestSampleEffectiveEndTime;
+		FTrackIterator FirstRangeSampleIt;
+		FTrackIterator LastRangeSampleIt;
+
 
 		FCriticalSection TimeLock;
 		FTimespan CurrentTime;
+		TRange<FTimespan> PlaybackRange;
 		float CurrentRate = 0.0f;
 		float IntendedRate = 0.0f;
 		float PlaybackDirection = 0.0f;
@@ -441,6 +677,8 @@ private:
 
 		FPendingBufferChange PendingBufferChange;
 
+		FPendingPlayRange PendingPlayRangeChange;
+
 		FPendingSeek PendingSeek;
 		TOptional<FTimespan> SeekTimeToHandleTo;
 		TOptional<FTimespan> SeekTimeToDecodeTo;
diff --git a/Engine/Plugins/Media/ElectraPlayer/Source/ElectraProtron/Private/Player/ElectraProtronPlayerLoader.cpp b/Engine/Plugins/Media/ElectraPlayer/Source/ElectraProtron/Private/Player/ElectraProtronPlayerLoader.cpp
--- a/Engine/Plugins/Media/ElectraPlayer/Source/ElectraProtron/Private/Player/ElectraProtronPlayerLoader.cpp
+++ b/Engine/Plugins/Media/ElectraPlayer/Source/ElectraProtron/Private/Player/ElectraProtronPlayerLoader.cpp
@@ -109,7 +109,8 @@ void FElectraProtronPlayer::FImpl::FLoaderThread::CalcRangeToLoad(FSampleRange&
 
 	// Given a track iterator, figure out the samples we need to load prior and following the iterator's current position.
 	check(InSampleIt->IsValid());
-	if (!InSampleIt->IsValid())
+	check(InTrackSampleBuffer->FirstRangeSampleIt.IsValid() && InTrackSampleBuffer->LastRangeSampleIt.IsValid());
+	if (!InSampleIt->IsValid() || !InTrackSampleBuffer->FirstRangeSampleIt.IsValid() || !InTrackSampleBuffer->LastRangeSampleIt.IsValid())
 	{
 		return;
 	}
@@ -119,15 +120,14 @@ void FElectraProtronPlayer::FImpl::FLoaderThread::CalcRangeToLoad(FSampleRange&
 
 	int64 FwdDurNeeded = Electra::FTimeFraction(Config.DurationCacheAhead).GetAsTimebase(TrackTimescale);
 	int64 RevDurNeeded = Electra::FTimeFraction(Config.DurationCacheBehind).GetAsTimebase(TrackTimescale);
-	// See if the combined cache ahead and behind duration encompasses the entire track.
-	// This happens if the track is short (or the cache durations excessively large)
-	if (FwdDurNeeded + RevDurNeeded >= InSampleIt->GetEffectiveTrackDuration().GetNumerator())
+	// See if the combined cache ahead and behind duration encompasses the entire playback range.
+	if (FwdDurNeeded + RevDurNeeded >= InTrackSampleBuffer->LastRangeSampleIt->GetEffectiveDTS().GetNumerator() - InTrackSampleBuffer->FirstRangeSampleIt->GetEffectiveDTS().GetNumerator())
 	{
-		const uint32 NumTrackSamples = InSampleIt->GetLastSampleInfo().LastSampleNumber + 1 - InSampleIt->GetFirstSampleInfo().SyncSampleNumber;
+		const uint32 NumTrackSamples = InTrackSampleBuffer->LastRangeSampleIt->GetSampleNumber() + 1 - InTrackSampleBuffer->FirstRangeSampleIt->GetSampleNumber();
 		OutRange.NumSamplesAfter = SharedPlayParams->PlaybackDirection >= 0.0f ? NumTrackSamples * 3 / 4 : NumTrackSamples / 4;
 		OutRange.NumSamplesBefore = NumTrackSamples - OutRange.NumSamplesAfter;
-		OutRange.TimeRanges.Add(TRange<FTimespan>(FTimespan(), InSampleIt->GetEffectiveTrackDuration().GetAsTimespan()));
-		OutRange.SampleRanges.Add(TRange<uint32>(InSampleIt->GetFirstSampleInfo().SyncSampleNumber, InSampleIt->GetLastSampleInfo().LastSampleNumber + 1));
+		OutRange.TimeRanges.Add(InTrackSampleBuffer->CurrentPlaybackRange);
+		OutRange.SampleRanges.Add(TRange<uint32>(InTrackSampleBuffer->FirstRangeSampleIt->GetSampleNumber(), InTrackSampleBuffer->LastRangeSampleIt->GetSampleNumber() + 1));
 		return;
 	}
 
@@ -145,11 +145,11 @@ void FElectraProtronPlayer::FImpl::FLoaderThread::CalcRangeToLoad(FSampleRange&
 		const int64 dur = FwdTkIt->GetDuration().GetNumerator();
 		DurHandled += dur;
 		++OutRange.NumSamplesAfter;
-		if (!FwdTkIt->NextEffective())
+		if (!FwdTkIt->NextEffective() || FwdTkIt->GetSampleNumber() >= InTrackSampleBuffer->LastRangeSampleIt->GetSampleNumber())
 		{
 			OutRange.TimeRanges.Add(TRange<FTimespan>(Electra::FTimeFraction(StartDTS, TrackTimescale).GetAsTimespan(), Electra::FTimeFraction(FwdTkIt->GetEffectiveDTS().GetNumerator() + dur, TrackTimescale).GetAsTimespan()));
 			OutRange.SampleRanges.Add(TRange<uint32>(StartSampleNum, FwdTkIt->GetSampleNumber()+1));
-			FwdTkIt = InTrackSampleBuffer->TrackAndCodecInfo->MP4Track->CreateIterator(FwdTkIt->GetFirstSampleInfo().SyncSampleNumber);
+			FwdTkIt = InTrackSampleBuffer->FirstRangeSampleIt->Clone();
 			StartDTS = FwdTkIt->GetEffectiveDTS().GetNumerator();
 			StartSampleNum = FwdTkIt->GetSampleNumber();
 		}
@@ -166,7 +166,7 @@ void FElectraProtronPlayer::FImpl::FLoaderThread::CalcRangeToLoad(FSampleRange&
 	bool bHaveEnough = DurHandled >= RevDurNeeded;
 	while(!bHaveEnough)
 	{
-		if (!RevTkIt->PrevEffective())
+		if (RevTkIt->GetSampleNumber() <= InTrackSampleBuffer->FirstRangeSampleIt->GetSampleNumber() || !RevTkIt->PrevEffective())
 		{
 			TRange<uint32> SmpRng(RevTkIt->GetSampleNumber(), EndSampleNum);
 			if (!SmpRng.IsEmpty())
@@ -177,7 +177,7 @@ void FElectraProtronPlayer::FImpl::FLoaderThread::CalcRangeToLoad(FSampleRange&
 				OutRange.TimeRanges.Add(TRange<FTimespan>(ts, te));
 				OutRange.SampleRanges.Add(SmpRng);
 			}
-			RevTkIt = InTrackSampleBuffer->TrackAndCodecInfo->MP4Track->CreateIterator(RevTkIt->GetLastSampleInfo().LastSampleNumber);
+			RevTkIt = InTrackSampleBuffer->LastRangeSampleIt->Clone();
 			EndDTS = RevTkIt->GetEffectiveDTS().GetNumerator() + RevTkIt->GetDuration().GetNumerator();
 			EndSampleNum = RevTkIt->GetSampleNumber() + 1;
 		}
@@ -247,6 +247,13 @@ TRangeSet<FTimespan> FElectraProtronPlayer::FImpl::FLoaderThread::GetTimeRangesT
 	return TimeRangesToLoad;
 }
 
+void FElectraProtronPlayer::FImpl::FLoaderThread::SetPlaybackRange(TRange<FTimespan> InRange)
+{
+	TimeRangeLock.Lock();
+	PlaybackRange = MoveTemp(InRange);
+	TimeRangeLock.Unlock();
+}
+
 void FElectraProtronPlayer::FImpl::FLoaderThread::RequestLoad(FMP4TrackSampleBufferPtr InTrackSampleBuffer, FTimespan InTime)
 {
 	if (!Reader.IsValid() || !InTrackSampleBuffer.IsValid())
@@ -254,6 +261,39 @@ void FElectraProtronPlayer::FImpl::FLoaderThread::RequestLoad(FMP4TrackSampleBuf
 		return;
 	}
 
+	// Did the playback range change?
+	TimeRangeLock.Lock();
+	TRange<FTimespan> RangeNow(PlaybackRange);
+	TimeRangeLock.Unlock();
+	if (RangeNow != InTrackSampleBuffer->CurrentPlaybackRange)
+	{
+		InTrackSampleBuffer->CurrentPlaybackRange = RangeNow;
+		// Locate the first and last sample numbers for the range
+		FTrackIterator RangeIt = InTrackSampleBuffer->TrackAndCodecInfo->MP4Track->CreateIteratorAtKeyframe(Electra::FTimeValue().SetFromTimespan(RangeNow.GetLowerBoundValue()), Electra::FTimeValue::GetZero());
+		if (!RangeIt.IsValid())
+		{
+			LastErrorMessage = InTrackSampleBuffer->TrackAndCodecInfo->MP4Track->GetLastError();
+			UE_LOG(LogElectraProtron, Error, TEXT("%s"), *LastErrorMessage);
+			return;
+		}
+		InTrackSampleBuffer->FirstRangeSampleIt = RangeIt->Clone();
+
+		// Move forward until we reach the end or both the effective DTS *and* PTS are greater or equal than the end of the range.
+		// We need to look at both DTS and PTS because the effective PTS can be smaller than the effective DTS due to composition time offsets.
+		while(!RangeIt->IsLastEffective())
+		{
+			if (RangeIt->GetEffectiveDTS().GetAsTimespan() >= RangeNow.GetUpperBoundValue() && RangeIt->GetEffectivePTS().GetAsTimespan() >= RangeNow.GetUpperBoundValue())
+			{
+				// We want the last iterator to represent the last sample included in the playback range,
+				// so we need to step one back here as we are currently outside the range.
+				RangeIt->PrevEffective();
+				break;
+			}
+			RangeIt->NextEffective();
+		}
+		InTrackSampleBuffer->LastRangeSampleIt = MoveTemp(RangeIt);
+	}
+
 	FTrackIterator TkIt =
 		InTrackSampleBuffer->TrackAndCodecInfo->MP4Track->CreateIteratorAtKeyframe(Electra::FTimeValue().SetFromTimespan(InTime),
 			Electra::FTimeValue().SetFromMilliseconds(InTrackSampleBuffer->TrackAndCodecInfo->bIsKeyframeOnlyFormat ? 0 : Config.NextKeyframeThresholdMillis));
@@ -364,9 +404,9 @@ FElectraProtronPlayer::FImpl::FLoaderThread::ELoadResult FElectraProtronPlayer::
 			{
 				FwdTkItTC->Next();
 			}
-			if (!FwdTkIt->NextEffective())
+			if (FwdTkIt->GetSampleNumber() >= InTrackSampleBuffer->LastRangeSampleIt->GetSampleNumber() || !FwdTkIt->NextEffective())
 			{
-				FwdTkIt = InTrackSampleBuffer->TrackAndCodecInfo->MP4Track->CreateIterator(FwdTkIt->GetFirstSampleInfo().SyncSampleNumber);
+				FwdTkIt = InTrackSampleBuffer->FirstRangeSampleIt->Clone();
 				FwdTkItTC.Reset();
 				if (auto TimecodeTrack = InTrackSampleBuffer->TrackAndCodecInfo->ReferencedTimecodeTrack.Pin())
 				{
@@ -382,13 +422,13 @@ FElectraProtronPlayer::FImpl::FLoaderThread::ELoadResult FElectraProtronPlayer::
 			{
 				RevTkItTC->Prev();
 			}
-			if (!RevTkIt->PrevEffective())
+			if (RevTkIt->GetSampleNumber() <= InTrackSampleBuffer->FirstRangeSampleIt->GetSampleNumber() || !RevTkIt->PrevEffective())
 			{
-				RevTkIt = InTrackSampleBuffer->TrackAndCodecInfo->MP4Track->CreateIterator(RevTkIt->GetLastSampleInfo().LastSampleNumber);
+				RevTkIt = InTrackSampleBuffer->LastRangeSampleIt->Clone();
 				RevTkItTC.Reset();
 				if (auto TimecodeTrack = InTrackSampleBuffer->TrackAndCodecInfo->ReferencedTimecodeTrack.Pin())
 				{
-					RevTkItTC = TimecodeTrack->MP4Track->CreateIterator(FwdTkIt->GetSampleNumber());
+					RevTkItTC = TimecodeTrack->MP4Track->CreateIterator(RevTkIt->GetSampleNumber());
 				}
 			}
 
diff --git a/Engine/Plugins/Media/ElectraUtil/Source/ElectraBase/Private/Utilities/MP4Boxes/MP4Track.cpp b/Engine/Plugins/Media/ElectraUtil/Source/ElectraBase/Private/Utilities/MP4Boxes/MP4Track.cpp
--- a/Engine/Plugins/Media/ElectraUtil/Source/ElectraBase/Private/Utilities/MP4Boxes/MP4Track.cpp
+++ b/Engine/Plugins/Media/ElectraUtil/Source/ElectraBase/Private/Utilities/MP4Boxes/MP4Track.cpp
@@ -8,7 +8,7 @@ namespace Electra
 	namespace UtilitiesMP4
 	{
 
-		bool FMP4Track::Prepare(FTimeFraction InMovieDurationFromMVHDBox)
+		bool FMP4Track::Prepare(FTimeFraction InFullMovieDuration, FTimeFraction InAdjustedMovieDuration)
 		{
 			if (!TrakBox.IsValid())
 			{
@@ -22,7 +22,8 @@ namespace Electra
 				return false;
 			}
 			Convs.TrackID = TkhdBox->GetTrackID();
-			Convs.DurationFromMvhdBox = InMovieDurationFromMVHDBox;
+			Convs.FullMovieDuration = InFullMovieDuration;
+			Convs.DurationFromMvhdBox = InAdjustedMovieDuration;
 			Convs.DurationFromTkhdBox.SetFromND(TkhdBox->GetDuration(), Convs.DurationFromMvhdBox.GetDenominator());
 
 			// Check for correct box hierarchy.
diff --git a/Engine/Plugins/Media/ElectraUtil/Source/ElectraBase/Public/Utilities/MP4Boxes/MP4Track.h b/Engine/Plugins/Media/ElectraUtil/Source/ElectraBase/Public/Utilities/MP4Boxes/MP4Track.h
--- a/Engine/Plugins/Media/ElectraUtil/Source/ElectraBase/Public/Utilities/MP4Boxes/MP4Track.h
+++ b/Engine/Plugins/Media/ElectraUtil/Source/ElectraBase/Public/Utilities/MP4Boxes/MP4Track.h
@@ -35,6 +35,7 @@ namespace Electra
 			{
 				FFirstSample FirstSample;
 				FLastSample LastSample;
+				FTimeFraction FullMovieDuration;
 				FTimeFraction DurationFromMvhdBox;
 				FTimeFraction DurationFromTkhdBox;
 				FTimeFraction DurationFromMdhdBox;
@@ -51,7 +52,7 @@ namespace Electra
 			{ return MakeShareable<FMP4Track>(new FMP4Track(InTrakBox)); }
 			virtual ~FMP4Track() = default;
 
-			ELECTRABASE_API bool Prepare(FTimeFraction InMovieDurationFromMVHDBox);
+			ELECTRABASE_API bool Prepare(FTimeFraction InFullMovieDuration, FTimeFraction InAdjustedMovieDuration);
 			ELECTRABASE_API FString GetLastError();
 			ELECTRABASE_API const FMP4TrackMetadataCommon& GetCommonMetadata();
 
@@ -131,7 +132,7 @@ namespace Electra
 				bool IsFirstEffective() const
 				{ return SampleNumber <= Convs.FirstSample.SyncSampleNumber; }
 				bool IsLastEffective() const
-				{ return SampleNumber+1 >= Convs.LastSample.LastSampleNumber; }
+				{ return SampleNumber >= Convs.LastSample.LastSampleNumber; }
 
 				// Creates a copy of this iterator.
 				TSharedPtr<FIterator, ESPMode::ThreadSafe> Clone()
@@ -185,6 +186,10 @@ namespace Electra
 			const FLastSample& GetLastSampleInfo() const
 			{ return Convs.LastSample; }
 
+			// Returns the duration of the movie as a whole, which is set from the longest track.
+			const FTimeFraction& GetFullMovieDuration() const
+			{ return Convs.FullMovieDuration; }
+
 		protected:
 			FMP4Track(const TSharedPtr<FMP4BoxTRAK, ESPMode::ThreadSafe>& InTrakBox) : TrakBox(InTrakBox)
 			{ }
