Protron: Seek to a later keyframe during reverse playback of keyframe-only codec if that is closer to the intended time than the located frame.

#rnx

diff --git a/Engine/Plugins/Media/ElectraPlayer/Source/ElectraProtron/Private/Player/ElectraProtronPlayerDecoder.cpp b/Engine/Plugins/Media/ElectraPlayer/Source/ElectraProtron/Private/Player/ElectraProtronPlayerDecoder.cpp
--- a/Engine/Plugins/Media/ElectraPlayer/Source/ElectraProtron/Private/Player/ElectraProtronPlayerDecoder.cpp
+++ b/Engine/Plugins/Media/ElectraPlayer/Source/ElectraProtron/Private/Player/ElectraProtronPlayerDecoder.cpp
@@ -403,6 +403,20 @@ void FElectraProtronPlayer::FImpl::FDecoderThread::UpdateTrackIterator(const FTi
 		UE_LOG(LogElectraProtron, Error, TEXT("%s"), *LastErrorMessage);
 		return;
 	}
+	// Searching for a keyframe in the track returns the frame on or before the time. When playing in reverse
+	// (for keyframe-only codecs) we may need to pick the next frame.
+	if (PlaybackDirection < 0.0f && TrackSampleBuffer->TrackAndCodecInfo->bIsKeyframeOnlyFormat)
+	{
+		auto NextIt = TkIt->Clone();
+		NextIt->NextEffective();
+		int64 t0 = FMath::Abs((TkIt->GetEffectivePTS().GetAsTimespan() - InForTime).GetTicks());
+		int64 t1 = FMath::Abs((NextIt->GetEffectivePTS().GetAsTimespan() - InForTime).GetTicks());
+		if (t1 < t0)
+		{
+			TkIt = MoveTemp(NextIt);
+		}
+	}
+
 	// Use the new iterator.
 	TrackIterator = MoveTemp(TkIt);
 	bReachedEnd = false;
