ElectraPlayer: Fix for webm in DASH

#rnx

diff --git a/Engine/Plugins/Media/ElectraPlayer/Source/ElectraPlayerRuntime/Private/Runtime/Player/DASH/ManifestDASH.cpp b/Engine/Plugins/Media/ElectraPlayer/Source/ElectraPlayerRuntime/Private/Runtime/Player/DASH/ManifestDASH.cpp
--- a/Engine/Plugins/Media/ElectraPlayer/Source/ElectraPlayerRuntime/Private/Runtime/Player/DASH/ManifestDASH.cpp
+++ b/Engine/Plugins/Media/ElectraPlayer/Source/ElectraPlayerRuntime/Private/Runtime/Player/DASH/ManifestDASH.cpp
@@ -2472,11 +2472,7 @@ bool FManifestDASHInternal::FRepresentation::PrepareDownloadURLs(IPlayerSessionS
 		InOutSegmentInfo.InitializationURL.Range = InitializationSegmentURLRange;
 		InOutSegmentInfo.InitializationURL.Url.URL = DASHUrlHelpers::ApplyAnnexEByteRange(InOutSegmentInfo.InitializationURL.Url.URL, InOutSegmentInfo.InitializationURL.Range, ChosenInit.AdditionalParams.GetValue(DASHUrlHelpers::SteerOption_ByteRange).SafeGetFString());
 	}
-/*
-	Do not set up the init segment that way. Leave it empty as the segment reader knows that it needs to look for the init data in the media file.
-	If we set this explicitly then it won't know that and load the entire file.
-
-	else
+	else if (GetStreamContainerType() == EStreamContainerType::Matroska)
 	{
 		// The init segment is part of the media segment, so we can just copy this across.
 		InOutSegmentInfo.InitializationURL = InOutSegmentInfo.MediaURL;
@@ -2484,7 +2480,7 @@ bool FManifestDASHInternal::FRepresentation::PrepareDownloadURLs(IPlayerSessionS
 		InOutSegmentInfo.InitializationURL.Range = InitializationSegmentURLRange;
 		InOutSegmentInfo.InitializationURL.Url.URL = DASHUrlHelpers::ApplyAnnexEByteRange(InOutSegmentInfo.InitializationURL.Url.URL, InOutSegmentInfo.InitializationURL.Range, ChosenMedia.AdditionalParams.GetValue(DASHUrlHelpers::SteerOption_ByteRange).SafeGetFString());
 	}
-*/
+
 
 	// Finally apply the ATO of the media segment as the ATO in general.
 	// We load init segments only when requesting media segments at which time the init segment needs
