ElectraPlayer: Fix for ABR not working in multi-angle HLS videos

#rnx

diff --git a/Engine/Plugins/Media/ElectraPlayer/Source/ElectraPlayerRuntime/Private/Runtime/Player/HLS/PlaylistHLS.cpp b/Engine/Plugins/Media/ElectraPlayer/Source/ElectraPlayerRuntime/Private/Runtime/Player/HLS/PlaylistHLS.cpp
--- a/Engine/Plugins/Media/ElectraPlayer/Source/ElectraPlayerRuntime/Private/Runtime/Player/HLS/PlaylistHLS.cpp
+++ b/Engine/Plugins/Media/ElectraPlayer/Source/ElectraPlayerRuntime/Private/Runtime/Player/HLS/PlaylistHLS.cpp
@@ -630,7 +630,7 @@ FErrorDetail FActiveHLSPlaylist::CreateTrackMetadata(IPlayerSessionServices* InP
 									FStreamMetadata& sm = tm.Meta.StreamDetails.Emplace_GetRef();
 									sm.Bandwidth = GroupStreamInfs[vi].Bandwidth;
 									sm.QualityIndex = GroupStreamInfs[vi].QualityIndex;
-									sm.ID = rg->Renditions[k].Name;
+									sm.ID.Empty();	// do not set an ID here to indicate this is a rendition
 									sm.CodecInformation = rg->Renditions[k].ParsedCodecFromStreamInf;
 									if (sm.Bandwidth > tm.Meta.HighestBandwidth)
 									{
@@ -1549,26 +1549,26 @@ FErrorDetail FActiveHLSPlaylist::FTimelineMediaAsset::GetVariantPlaylist(TShared
 	else
 	{
 		// We need to find the rendition in the respective group of this variant so we are on the correct pathway.
-		const TArray<FMultiVariantPlaylistHLS::FRenditionGroup>* RenditionGroup = nullptr;
+		const TArray<FMultiVariantPlaylistHLS::FRenditionGroup>* RenditionGroups = nullptr;
 		FString GroupName;
 		switch(InStreamType)
 		{
 			case EStreamType::Video:
 			{
 				GroupName = si->VideoGroup;
-				RenditionGroup = &MultiVariantPlaylist->RenditionGroupsOfType[static_cast<int32>(FMultiVariantPlaylistHLS::ERenditionGroupType::Video)];
+				RenditionGroups = &MultiVariantPlaylist->RenditionGroupsOfType[static_cast<int32>(FMultiVariantPlaylistHLS::ERenditionGroupType::Video)];
 				break;
 			}
 			case EStreamType::Audio:
 			{
 				GroupName = si->AudioGroup;
-				RenditionGroup = &MultiVariantPlaylist->RenditionGroupsOfType[static_cast<int32>(FMultiVariantPlaylistHLS::ERenditionGroupType::Audio)];
+				RenditionGroups = &MultiVariantPlaylist->RenditionGroupsOfType[static_cast<int32>(FMultiVariantPlaylistHLS::ERenditionGroupType::Audio)];
 				break;
 			}
 			case EStreamType::Subtitle:
 			{
 				GroupName = si->SubtitleGroup;
-				RenditionGroup = &MultiVariantPlaylist->RenditionGroupsOfType[static_cast<int32>(FMultiVariantPlaylistHLS::ERenditionGroupType::Subtitles)];
+				RenditionGroups = &MultiVariantPlaylist->RenditionGroupsOfType[static_cast<int32>(FMultiVariantPlaylistHLS::ERenditionGroupType::Subtitles)];
 				break;
 			}
 			default:
@@ -1577,11 +1577,11 @@ FErrorDetail FActiveHLSPlaylist::FTimelineMediaAsset::GetVariantPlaylist(TShared
 			}
 		}
 		const TArray<FMultiVariantPlaylistHLS::FRendition>* Renditions = nullptr;
-		for(int32 gIdx=0; gIdx<RenditionGroup->Num(); ++gIdx)
+		for(int32 gIdx=0; gIdx<RenditionGroups->Num(); ++gIdx)
 		{
-			if ((*RenditionGroup)[gIdx] == GroupName)
+			if ((*RenditionGroups)[gIdx] == GroupName)
 			{
-				Renditions = &(*RenditionGroup)[gIdx].Renditions;
+				Renditions = &(*RenditionGroups)[gIdx].Renditions;
 				break;
 			}
 		}
@@ -1595,10 +1595,8 @@ FErrorDetail FActiveHLSPlaylist::FTimelineMediaAsset::GetVariantPlaylist(TShared
 			if ((*Renditions)[k].Name.Equals(tm.Meta.Label))
 			{
 				Rendition = &(*Renditions)[k];
-				// Adaptation sets created from renditions only have a single rendition in them.
-// FIXME: this is not quite right with multiple video angles
-				//check(AdaptationSet->GetNumberOfRepresentations() == 1);
-				Representation = AdaptationSet->GetRepresentationByIndex(0);
+				check(AdaptationSet->GetNumberOfRepresentations() > InStreamIndex);
+				Representation = AdaptationSet->GetRepresentationByIndex(InStreamIndex);
 				break;
 			}
 		}
