ElectraPlayer: Fixed a crash when entering a partial URL into the stream media source that fails to parse into a filename with extension

#rnx

diff --git a/Engine/Plugins/Media/ElectraPlayer/Source/ElectraPlayerFactory/Private/ElectraPlayerFactoryModule.cpp b/Engine/Plugins/Media/ElectraPlayer/Source/ElectraPlayerFactory/Private/ElectraPlayerFactoryModule.cpp
--- a/Engine/Plugins/Media/ElectraPlayer/Source/ElectraPlayerFactory/Private/ElectraPlayerFactoryModule.cpp
+++ b/Engine/Plugins/Media/ElectraPlayer/Source/ElectraPlayerFactory/Private/ElectraPlayerFactoryModule.cpp
@@ -77,7 +77,7 @@ public:
 		// Check for known extensions
 		TArray<FString> PathComponents;
 		UrlParser.GetPathComponents(PathComponents);
-		FString LowerCaseExtension = FPaths::GetExtension(PathComponents.Last().ToLower());
+		FString LowerCaseExtension = PathComponents.Num() ? FPaths::GetExtension(PathComponents.Last().ToLower()) : FString();
 		// If the extension is known, we are confident that we can play this.
 		// At this point there is no information provided on the codecs used in the media,
 		// so we cannot check for this.
