Remove electra module dependencies from Server.

#rnx

diff --git a/Engine/Plugins/Media/ElectraPlayer/ElectraPlayer.uplugin b/Engine/Plugins/Media/ElectraPlayer/ElectraPlayer.uplugin
--- a/Engine/Plugins/Media/ElectraPlayer/ElectraPlayer.uplugin
+++ b/Engine/Plugins/Media/ElectraPlayer/ElectraPlayer.uplugin
@@ -18,8 +18,7 @@
     {
       "Name": "ElectraPlayerPluginHandler",
       "Type": "RuntimeNoCommandlet",
-      "LoadingPhase": "PreLoadingScreen",
-      "TargetDenyList": [ "Server" ]
+      "LoadingPhase": "PreLoadingScreen"
     },
     {
       "Name": "ElectraPlayerRuntime",
