ElectraPlayer: Fixed an issue with HLS segment duration inaccuracies leading to assumed overlapping segments when this was not true.

#rnx

diff --git a/Engine/Plugins/Media/ElectraPlayer/Source/ElectraPlayerRuntime/Private/Runtime/Player/HLS/PlaylistHLS.cpp b/Engine/Plugins/Media/ElectraPlayer/Source/ElectraPlayerRuntime/Private/Runtime/Player/HLS/PlaylistHLS.cpp
--- a/Engine/Plugins/Media/ElectraPlayer/Source/ElectraPlayerRuntime/Private/Runtime/Player/HLS/PlaylistHLS.cpp
+++ b/Engine/Plugins/Media/ElectraPlayer/Source/ElectraPlayerRuntime/Private/Runtime/Player/HLS/PlaylistHLS.cpp
@@ -1855,10 +1855,13 @@ FActiveHLSPlaylist::FTimelineMediaAsset::ESegSearchResult FActiveHLSPlaylist::FT
 	if (!OutSegment->HLS.bNoPDTMapping)
 	{
 		OutSegment->Segment.Time = Seg.ProgramDateTime.GetAsHNS();
-		OutSegment->Segment.MediaLocalFirstAUTime = SearchTime.GetAsHNS();
-		OutSegment->Segment.MediaLocalFirstPTS = SearchTime.GetAsHNS();
 		OutSegment->Segment.MediaLocalLastAUTime = InParam.LastPTS.GetAsHNS();
 		OutSegment->Segment.bFrameAccuracyRequired = InParam.bFrameAccurateSearch;
+		if (InParam.bFrameAccurateSearch)
+		{
+			OutSegment->Segment.MediaLocalFirstAUTime = SearchTime.GetAsHNS();
+			OutSegment->Segment.MediaLocalFirstPTS = SearchTime.GetAsHNS();
+		}
 	}
 	else
 	{
