Electra: Get framerate of H.264 video from container if VUI indicates variable framerate

#rnx

diff --git a/Engine/Plugins/Media/ElectraCodecs/Source/ElectraDecoders/Private/Utils/MPEG/ElectraUtilsMPEGVideo_H264.cpp b/Engine/Plugins/Media/ElectraCodecs/Source/ElectraDecoders/Private/Utils/MPEG/ElectraUtilsMPEGVideo_H264.cpp
--- a/Engine/Plugins/Media/ElectraCodecs/Source/ElectraDecoders/Private/Utils/MPEG/ElectraUtilsMPEGVideo_H264.cpp
+++ b/Engine/Plugins/Media/ElectraCodecs/Source/ElectraDecoders/Private/Utils/MPEG/ElectraUtilsMPEGVideo_H264.cpp
@@ -379,9 +379,9 @@ namespace ElectraDecodersUtil
 			}
 			FFractionalValue FSequenceParameterSet::GetTiming() const
 			{
-				if (vui_parameters_present_flag && timing_info_present_flag)
+				if (vui_parameters_present_flag && timing_info_present_flag && fixed_frame_rate_flag)
 				{
-					return FFractionalValue(time_scale, num_units_in_tick * 2);
+					return FFractionalValue(time_scale, num_units_in_tick * 2 /* rate per field, we assume progressive frames, so 2 fields */);
 				}
 				return FFractionalValue();
 			}
@@ -2233,6 +2233,38 @@ namespace ElectraDecodersUtil
 
 
 
+			int32 ELECTRADECODERS_API CalcNumBytesInNAL(int32 InProfileLevel, int32 InFrameWidth, int32 InFrameHeight, double InFrameRate)
+			{
+				int32 MaxMBPS = 0;
+				int32 MinCR = InProfileLevel == 31 || InProfileLevel == 32 || InProfileLevel == 40 ? 4 : 2;
+				double fR = InProfileLevel >= 60 ? 1/300.0 : 1/172.0;
+				if (InProfileLevel <= 10) MaxMBPS = 1485;
+				else if (InProfileLevel == 11) MaxMBPS = 3000;
+				else if (InProfileLevel == 12) MaxMBPS = 6000;
+				else if (InProfileLevel == 13) MaxMBPS = 11880;
+				else if (InProfileLevel == 20) MaxMBPS = 11880;
+				else if (InProfileLevel == 21) MaxMBPS = 19800;
+				else if (InProfileLevel == 22) MaxMBPS = 20250;
+				else if (InProfileLevel == 30) MaxMBPS = 40500;
+				else if (InProfileLevel == 31) MaxMBPS = 108000;
+				else if (InProfileLevel == 32) MaxMBPS = 216000;
+				else if (InProfileLevel == 40) MaxMBPS = 245760;
+				else if (InProfileLevel == 41) MaxMBPS = 245760;
+				else if (InProfileLevel == 42) MaxMBPS = 522240;
+				else if (InProfileLevel == 50) MaxMBPS = 589824;
+				else if (InProfileLevel == 51) MaxMBPS = 983040;
+				else if (InProfileLevel == 52) MaxMBPS = 2073600;
+				else if (InProfileLevel == 60) MaxMBPS = 4177920;
+				else if (InProfileLevel == 61) MaxMBPS = 8355840;
+				else if (InProfileLevel == 62) MaxMBPS = 16711680;
+				if (MaxMBPS == 0 || InFrameWidth <= 0 || InFrameHeight <= 0 || InFrameRate <= 1.0)
+				{
+					return 1536 << 10;
+				}
+				int32 PicSizeInMbs = Align(InFrameWidth, 16) * Align(InFrameHeight, 16) / 256;
+				int32 NumBytesInNAL = (int32)(384.0 * (FMath::Max((double)PicSizeInMbs, fR * (double)MaxMBPS) + ((double)MaxMBPS * 1.0 / InFrameRate)) / (double)MinCR);
+				return NumBytesInNAL;
+			}
 
 
 			bool FAVCDecoderConfigurationRecord::Parse(const TConstArrayView<uint8>& InDCR)
diff --git a/Engine/Plugins/Media/ElectraCodecs/Source/ElectraDecoders/Public/Utils/MPEG/ElectraUtilsMPEGVideo_H264.h b/Engine/Plugins/Media/ElectraCodecs/Source/ElectraDecoders/Public/Utils/MPEG/ElectraUtilsMPEGVideo_H264.h
--- a/Engine/Plugins/Media/ElectraCodecs/Source/ElectraDecoders/Public/Utils/MPEG/ElectraUtilsMPEGVideo_H264.h
+++ b/Engine/Plugins/Media/ElectraCodecs/Source/ElectraDecoders/Public/Utils/MPEG/ElectraUtilsMPEGVideo_H264.h
@@ -598,6 +598,7 @@ namespace ElectraDecodersUtil
 			{ return ParseSequenceParameterSet(OutSequenceParameterSet, InData.GetData(), InData.Num()); }
 
 
+			int32 ELECTRADECODERS_API CalcNumBytesInNAL(int32 InProfileLevel, int32 InFrameWidth, int32 InFrameHeight, double InFrameRate);
 
 
 
diff --git a/Engine/Plugins/Media/ElectraPlayer/Source/ElectraPlayerRuntime/Private/Runtime/Demuxer/ParserISO14496-12.cpp b/Engine/Plugins/Media/ElectraPlayer/Source/ElectraPlayerRuntime/Private/Runtime/Demuxer/ParserISO14496-12.cpp
--- a/Engine/Plugins/Media/ElectraPlayer/Source/ElectraPlayerRuntime/Private/Runtime/Demuxer/ParserISO14496-12.cpp
+++ b/Engine/Plugins/Media/ElectraPlayer/Source/ElectraPlayerRuntime/Private/Runtime/Demuxer/ParserISO14496-12.cpp
@@ -6,6 +6,7 @@
 #include "Misc/Optional.h"
 #include "Misc/FrameRate.h"
 #include "Misc/Timecode.h"
+#include "Math/BigInt.h"
 
 #include "Demuxer/ParserISO14496-12.h"
 
@@ -7442,7 +7443,25 @@ namespace Electra
 									FStreamCodecInformation::FAspectRatio ar;
 									sps.GetAspect(ar.Width, ar.Height);
 									Track->CodecInformation.SetAspectRatio(ar);
-									Track->CodecInformation.SetFrameRate(sps.GetTiming().Denom ? FTimeFraction(sps.GetTiming().Num, sps.GetTiming().Denom) : FTimeFraction());
+									// Have a valid fixed framerate?
+									if (sps.GetTiming().Denom)
+									{
+										Track->CodecInformation.SetFrameRate(FTimeFraction(sps.GetTiming().Num, sps.GetTiming().Denom));
+									}
+									else if (Track->MDHDBox && Track->STTSBox && Track->MDHDBox->GetTimescale() && Track->STTSBox->GetNumberOfEntries())
+									{
+										// We assume the frame rate to be constant.
+										Track->CodecInformation.SetFrameRate(FTimeFraction(Track->MDHDBox->GetTimescale(), Track->STTSBox->GetEntry(0).SampleDelta));
+									}
+									else if (Track->MDHDBox && Track->STSZBox && Track->STSZBox->GetNumTotalSamples() && Track->MDHDBox->GetDuration() && Track->MDHDBox->GetTimescale())
+									{
+										// We assume the frame rate to be constant. Or at least the average to be close to what
+										// a common constant framerate would be.
+										TBigInt<128> n((int64)Track->STSZBox->GetNumTotalSamples() * (int32)Track->MDHDBox->GetTimescale() * 1000);
+										TBigInt<128> d((int64)Track->MDHDBox->GetDuration());
+										n /= d;
+										Track->CodecInformation.SetFrameRate(FTimeFraction(n.ToInt(), 1000));
+									}
 									Track->CodecInformation.SetProfile(sps.profile_idc);
 									Track->CodecInformation.SetProfileLevel(sps.level_idc);
 									uint8 Constraints = (sps.constraint_set0_flag << 7) | (sps.constraint_set1_flag << 6) | (sps.constraint_set2_flag << 5) | (sps.constraint_set3_flag << 4) | (sps.constraint_set4_flag << 3) | (sps.constraint_set5_flag << 2);
diff --git a/Engine/Plugins/Media/ElectraPlayer/Source/ElectraPlayerRuntime/Private/Runtime/Demuxer/ParserMKV.cpp b/Engine/Plugins/Media/ElectraPlayer/Source/ElectraPlayerRuntime/Private/Runtime/Demuxer/ParserMKV.cpp
--- a/Engine/Plugins/Media/ElectraPlayer/Source/ElectraPlayerRuntime/Private/Runtime/Demuxer/ParserMKV.cpp
+++ b/Engine/Plugins/Media/ElectraPlayer/Source/ElectraPlayerRuntime/Private/Runtime/Demuxer/ParserMKV.cpp
@@ -2780,7 +2780,7 @@ namespace Electra
 					clli.max_pic_average_light_level = InColours->GetMaxFALL();
 					InOutCodecInformation.GetExtras().Set(TEXT("clli_box"), FVariantValue(ElectraDecodersUtil::MPEG::BuildCLLIBox(clli)));
 				}
-				
+
 				// Mastering display color volume present?
 				TMKVElementPtr<FMKVMasteringMetadata> MasteringMetadata = InColours->GetMasteringMetadata();
 				if (MasteringMetadata.IsValid())
@@ -2841,25 +2841,27 @@ namespace Electra
 				FStreamCodecInformation::FAspectRatio ar;
 				sps.GetAspect(ar.Width, ar.Height);
 				OutCodecInformation.SetAspectRatio(ar);
-				OutCodecInformation.SetFrameRate(sps.GetTiming().Denom ? FTimeFraction(sps.GetTiming().Num, sps.GetTiming().Denom) : FTimeFraction());
+				if (sps.GetTiming().Denom)
+				{
+					OutCodecInformation.SetFrameRate(FTimeFraction(sps.GetTiming().Num, sps.GetTiming().Denom));
+				}
+				else if (InFromTrack->GetDefaultDurationNanos())
+				{
+					OutCodecInformation.SetFrameRate(FTimeFraction(1000000000, InFromTrack->GetDefaultDurationNanos()));
+				}
 				OutCodecInformation.SetProfile(sps.profile_idc);
 				OutCodecInformation.SetProfileLevel(sps.level_idc);
 				uint8 Constraints = (sps.constraint_set0_flag << 7) | (sps.constraint_set1_flag << 6) | (sps.constraint_set2_flag << 5) | (sps.constraint_set3_flag << 4) | (sps.constraint_set4_flag << 3) | (sps.constraint_set5_flag << 2);
 				OutCodecInformation.SetProfileConstraints(Constraints);
 				OutCodecInformation.SetCodecSpecifierRFC6381(FString::Printf(TEXT("avc1.%02x%02x%02x"), sps.profile_idc, Constraints, sps.level_idc));
-				// If there is no default duration set we try to calculate it from the VUI timing of the SPS if present.
-				if (InFromTrack->GetDefaultDurationNanos() == 0)
+				// If there is no default duration set we try to calculate it from the frame rate
+				if (InFromTrack->GetDefaultDurationNanos() == 0 && OutCodecInformation.GetFrameRate().IsValid())
 				{
-					ElectraDecodersUtil::FFractionalValue fv = sps.GetTiming();
-					FTimeFraction fr(fv.Denom, fv.Num);
-					if (fr.IsValid())
+					FTimeFraction ooFPS(OutCodecInformation.GetFrameRate().GetDenominator(), OutCodecInformation.GetFrameRate().GetNumerator());
+					int64 nanos = ooFPS.GetAsTimebase(1000000000);
+					if (nanos >= 0)
 					{
-						int64 nanos = fr.GetAsTimebase(1000000000);
-						check(nanos >= 0);
-						if (nanos >= 0)
-						{
-							InFromTrack->SetDefaultDurationNanos((uint64)nanos);
-						}
+						InFromTrack->SetDefaultDurationNanos((uint64)nanos);
 					}
 				}
 			}
diff --git a/Engine/Plugins/Media/ElectraPlayer/Source/ElectraProtron/Private/Player/ElectraProtronPlayerImpl.cpp b/Engine/Plugins/Media/ElectraPlayer/Source/ElectraProtron/Private/Player/ElectraProtronPlayerImpl.cpp
--- a/Engine/Plugins/Media/ElectraPlayer/Source/ElectraProtron/Private/Player/ElectraProtronPlayerImpl.cpp
+++ b/Engine/Plugins/Media/ElectraPlayer/Source/ElectraProtron/Private/Player/ElectraProtronPlayerImpl.cpp
@@ -734,7 +734,14 @@ void FElectraProtronPlayer::FImpl::GetTrackCodecInfo(ElectraProtronUtils::FCodec
 							{
 								ElectraDecodersUtil::FFractionalValue fr = sps.GetTiming();
 								sps.GetDisplaySize(Video.Width, Video.Height);
-								Video.FrameRate.SetFromND(fr.Num, fr.Denom);
+								if (fr.Num && fr.Denom)
+								{
+									Video.FrameRate.SetFromND(fr.Num, fr.Denom);
+								}
+								else
+								{
+									SetResoAndFPSFromBox(Video, Visuals[0]);
+								}
 								Prefix = TEXT("avc1");
 							}
 							else
