ElectraPlayer: Activating a playback time range means the range is now set and thus the change flag must be cleared.

#rnx

diff --git a/Engine/Plugins/Media/ElectraPlayer/Source/ElectraPlayerRuntime/Private/Runtime/Player/AdaptiveStreamingPlayerInternal.h b/Engine/Plugins/Media/ElectraPlayer/Source/ElectraPlayerRuntime/Private/Runtime/Player/AdaptiveStreamingPlayerInternal.h
--- a/Engine/Plugins/Media/ElectraPlayer/Source/ElectraPlayerRuntime/Private/Runtime/Player/AdaptiveStreamingPlayerInternal.h
+++ b/Engine/Plugins/Media/ElectraPlayer/Source/ElectraPlayerRuntime/Private/Runtime/Player/AdaptiveStreamingPlayerInternal.h
@@ -533,7 +533,7 @@ struct FPlaybackState
 	{
 		FScopeLock lock(&Lock);
 		ActivePlaybackRange = InTimeRange ? *InTimeRange : NewPlaybackRange;
-		bPlayrangeHasChanged = InTimeRange ? bPlayrangeHasChanged : false;
+		bPlayrangeHasChanged = false;
 	}
 	FTimeRange GetActivePlayRange()
 	{
