ElectraPlayer: Release DX11 resources asynchronously when the decoder output pool is destroyed when the player closes as that can take a bit and cause a hitch on the game thread

#rnx

diff --git a/Engine/Plugins/Media/ElectraPlayer/Source/ElectraPlayerRuntime/Private/Runtime/Decoder/Windows/ElectraVideoDecoder_PC.cpp b/Engine/Plugins/Media/ElectraPlayer/Source/ElectraPlayerRuntime/Private/Runtime/Decoder/Windows/ElectraVideoDecoder_PC.cpp
--- a/Engine/Plugins/Media/ElectraPlayer/Source/ElectraPlayerRuntime/Private/Runtime/Decoder/Windows/ElectraVideoDecoder_PC.cpp
+++ b/Engine/Plugins/Media/ElectraPlayer/Source/ElectraPlayerRuntime/Private/Runtime/Decoder/Windows/ElectraVideoDecoder_PC.cpp
@@ -24,8 +24,33 @@ THIRD_PARTY_INCLUDES_START
 THIRD_PARTY_INCLUDES_END
 #include "Windows/HideWindowsPlatformTypes.h"
 
+#include "Async/Async.h"
+
 // ------------------------------------------------------------------------------------
 
+FElectraPlayerVideoDecoderOutputPC::~FElectraPlayerVideoDecoderOutputPC()
+{
+	// Release DX11 resources asynchronously. There is some internal synchronization going on that takes a while
+	// to finish that would block the game thread.
+	if (GIsRunning && (Texture.IsValid() || SharedTexture.IsValid() || D3D11Device.IsValid()))
+	{
+		struct FRefCountedObjs
+		{
+			TRefCountPtr<ID3D11Texture2D> Texture;
+			TRefCountPtr<ID3D11Texture2D> SharedTexture;
+			TRefCountPtr<ID3D11Device> D3D11Device;
+		};
+		FRefCountedObjs* rco = new FRefCountedObjs;
+		rco->Texture = MoveTemp(Texture);
+		rco->SharedTexture = MoveTemp(SharedTexture);
+		rco->D3D11Device = MoveTemp(D3D11Device);
+		Async(EAsyncExecution::TaskGraph, [rco]()
+		{
+			delete rco;
+		});
+	}
+}
+
 void FElectraPlayerVideoDecoderOutputPC::SetOwner(const TSharedPtr<IDecoderOutputOwner, ESPMode::ThreadSafe>& InOwningRenderer)
 {
 	OwningRenderer = InOwningRenderer;
diff --git a/Engine/Plugins/Media/ElectraPlayer/Source/ElectraPlayerRuntime/Private/Runtime/Decoder/Windows/ElectraVideoDecoder_PC.h b/Engine/Plugins/Media/ElectraPlayer/Source/ElectraPlayerRuntime/Private/Runtime/Decoder/Windows/ElectraVideoDecoder_PC.h
--- a/Engine/Plugins/Media/ElectraPlayer/Source/ElectraPlayerRuntime/Private/Runtime/Decoder/Windows/ElectraVideoDecoder_PC.h
+++ b/Engine/Plugins/Media/ElectraPlayer/Source/ElectraPlayerRuntime/Private/Runtime/Decoder/Windows/ElectraVideoDecoder_PC.h
@@ -45,6 +45,7 @@ class FElectraPlayerVideoDecoderOutputPC : public FVideoDecoderOutputPC
 public:
 	FElectraPlayerVideoDecoderOutputPC() : SampleDim(0, 0)
 	{ }
+	~FElectraPlayerVideoDecoderOutputPC();
 
 	// Hardware decode to buffer (Win7/DX12)
 	void InitializeWithBuffer(const void* InBuffer, uint32 InSize, uint32 InStride, FIntPoint Dim, TSharedPtr<Electra::FParamDict, ESPMode::ThreadSafe> InParamDict);
