ElectraCodecs: Added support for new decoder API

#rnx

diff --git a/Engine/Plugins/Media/ElectraPlayer/Source/ElectraPlayerRuntime/Private/Runtime/Decoder/VideoDecoder.cpp b/Engine/Plugins/Media/ElectraPlayer/Source/ElectraPlayerRuntime/Private/Runtime/Decoder/VideoDecoder.cpp
--- a/Engine/Plugins/Media/ElectraPlayer/Source/ElectraPlayerRuntime/Private/Runtime/Decoder/VideoDecoder.cpp
+++ b/Engine/Plugins/Media/ElectraPlayer/Source/ElectraPlayerRuntime/Private/Runtime/Decoder/VideoDecoder.cpp
@@ -1181,7 +1181,7 @@ FVideoDecoderImpl::ENextDecodingState FVideoDecoderImpl::HandleDecoding()
 				// Likewise we are no longer waiting for a sync sample.
 				bWaitForSyncSample = false;
 			}
-			else if (DecErr == IElectraDecoder::EDecoderError::NoBuffer)
+			else if (DecErr == IElectraDecoder::EDecoderError::NoBuffer || DecErr == IElectraDecoder::EDecoderError::EndOfData)
 			{
 				// Try again later...
 				return ENextDecodingState::NormalDecoding;
