ElectraPlayer: Imposing a maximum segment size in addition to a duration limit when generating an artificial mp4 segment

#rnx

diff --git a/Engine/Plugins/Media/ElectraPlayer/Source/ElectraPlayerRuntime/Private/Runtime/Player/mp4/ManifestMP4.cpp b/Engine/Plugins/Media/ElectraPlayer/Source/ElectraPlayerRuntime/Private/Runtime/Player/mp4/ManifestMP4.cpp
--- a/Engine/Plugins/Media/ElectraPlayer/Source/ElectraPlayerRuntime/Private/Runtime/Player/mp4/ManifestMP4.cpp
+++ b/Engine/Plugins/Media/ElectraPlayer/Source/ElectraPlayerRuntime/Private/Runtime/Player/mp4/ManifestMP4.cpp
@@ -854,9 +854,10 @@ void FManifestMP4Internal::FTimelineAssetMP4::LogMessage(IInfoLog::ELevel Level,
 
 void FManifestMP4Internal::FTimelineAssetMP4::LimitSegmentDownloadSize(TSharedPtrTS<IStreamSegment>& InOutSegment)
 {
-	// Limit the segment size by content duration.
+	// Limit the segment size by content duration and/or size.
 	// This helps with downloads that might otherwise take too long or keep the connection open for too long (when downloading a large mp4 from start to finish).
 	const int64 MaxSegmentDurationMSec = 3000;
+	const int64 MaxSegmentByteCount = 1024 * 1024 * 64;
 	if (InOutSegment.IsValid())
 	{
 		FStreamSegmentRequestMP4* Request = static_cast<FStreamSegmentRequestMP4*>(InOutSegment.Get());
@@ -896,7 +897,8 @@ void FManifestMP4Internal::FTimelineAssetMP4::LimitSegmentDownloadSize(TSharedPt
 				}
 				int64 CurrentTrackOffset = LastTrackOffset;
 				bool bDurationLimitReached = TrackDurationLimit > 0 && TrackDur > TrackDurationLimit;
-				if (CurrentTrackOffset >= EndOffset || bDurationLimitReached)
+				bool bSizeLimitReached = (MaxSegmentByteCount > 0) && (CurrentTrackOffset - StartOffset >= MaxSegmentByteCount);
+				if (CurrentTrackOffset >= EndOffset || bDurationLimitReached || bSizeLimitReached)
 				{
 					// Limit reached.
 					Request->FileEndOffset = CurrentTrackOffset - 1;
