[ElectraProton] Fix IsPaused() returns false when the player is ready and has a rate of 0 (which is defined as Paused by Electra)

Players define the "Paused" state as being "ready" (i.e. opened) but having a play rate of 0. For consistency, we also want Protron to be the same.

diff --git a/Engine/Plugins/Media/ElectraPlayer/Source/ElectraProtron/Private/Player/ElectraProtronPlayer.cpp b/Engine/Plugins/Media/ElectraPlayer/Source/ElectraProtron/Private/Player/ElectraProtronPlayer.cpp
--- a/Engine/Plugins/Media/ElectraPlayer/Source/ElectraProtron/Private/Player/ElectraProtronPlayer.cpp
+++ b/Engine/Plugins/Media/ElectraPlayer/Source/ElectraProtron/Private/Player/ElectraProtronPlayer.cpp
@@ -211,10 +211,14 @@ void FElectraProtronPlayer::CheckForStateChanges()
 
 		// Rate changes?
 		float CurrentRate = Player->GetRate();
-		if (CurrentRate == 0.0f && CurrentState == EMediaState::Playing)
+		if (CurrentRate == 0.0f && (CurrentState == EMediaState::Playing || CurrentState == EMediaState::Stopped))
 		{
+			const EMediaState PreviousState = CurrentState;
 			CurrentState = EMediaState::Paused;
-			EventSink.ReceiveMediaEvent(EMediaEvent::PlaybackSuspended);
+			if (PreviousState == EMediaState::Playing)
+			{
+				EventSink.ReceiveMediaEvent(EMediaEvent::PlaybackSuspended);
+			}
 		}
 		else if (CurrentRate != 0.0f && (CurrentState == EMediaState::Paused || CurrentState == EMediaState::Stopped))
 		{
