Electra: Added passing of a result error message to the method converting the video decoder output

#rnx

diff --git a/Engine/Plugins/Media/ElectraPlayer/Source/ElectraPlayerRuntime/Private/ElectraPlayerDecoderResourceManager.cpp b/Engine/Plugins/Media/ElectraPlayer/Source/ElectraPlayerRuntime/Private/ElectraPlayerDecoderResourceManager.cpp
--- a/Engine/Plugins/Media/ElectraPlayer/Source/ElectraPlayerRuntime/Private/ElectraPlayerDecoderResourceManager.cpp
+++ b/Engine/Plugins/Media/ElectraPlayer/Source/ElectraPlayerRuntime/Private/ElectraPlayerDecoderResourceManager.cpp
@@ -16,9 +16,9 @@ TSharedPtr<Electra::IVideoDecoderResourceDelegate, ESPMode::ThreadSafe> FElectra
 	return Electra::PlatformCreateVideoDecoderResourceDelegate(InElectraPlayerAdapterDelegate);
 }
 
-bool FElectraPlayerDecoderResourceManager::SetupRenderBufferFromDecoderOutput(Electra::IMediaRenderer::IBuffer* InOutBufferToSetup, TSharedPtr<Electra::FParamDict, ESPMode::ThreadSafe> InOutBufferPropertes, TSharedPtr<IElectraDecoderVideoOutput, ESPMode::ThreadSafe> InDecoderOutput, IElectraDecoderResourceDelegateBase::IDecoderPlatformResource* InPlatformSpecificResource)
+bool FElectraPlayerDecoderResourceManager::SetupRenderBufferFromDecoderOutput(FString& OutErrorMessage, Electra::IMediaRenderer::IBuffer* InOutBufferToSetup, TSharedPtr<Electra::FParamDict, ESPMode::ThreadSafe> InOutBufferPropertes, TSharedPtr<IElectraDecoderVideoOutput, ESPMode::ThreadSafe> InDecoderOutput, IElectraDecoderResourceDelegateBase::IDecoderPlatformResource* InPlatformSpecificResource)
 {
-	return Electra::FPlatformElectraDecoderResourceManager::SetupRenderBufferFromDecoderOutput(InOutBufferToSetup, InOutBufferPropertes, InDecoderOutput, InPlatformSpecificResource);
+	return Electra::FPlatformElectraDecoderResourceManager::SetupRenderBufferFromDecoderOutput(OutErrorMessage, InOutBufferToSetup, InOutBufferPropertes, InDecoderOutput, InPlatformSpecificResource);
 }
 
 FVideoDecoderOutput* FElectraPlayerDecoderResourceManager::FVideo::Create()
diff --git a/Engine/Plugins/Media/ElectraPlayer/Source/ElectraPlayerRuntime/Private/Runtime/Decoder/Android/AndroidElectraDecoderResourceManager.cpp b/Engine/Plugins/Media/ElectraPlayer/Source/ElectraPlayerRuntime/Private/Runtime/Decoder/Android/AndroidElectraDecoderResourceManager.cpp
--- a/Engine/Plugins/Media/ElectraPlayer/Source/ElectraPlayerRuntime/Private/Runtime/Decoder/Android/AndroidElectraDecoderResourceManager.cpp
+++ b/Engine/Plugins/Media/ElectraPlayer/Source/ElectraPlayerRuntime/Private/Runtime/Decoder/Android/AndroidElectraDecoderResourceManager.cpp
@@ -50,7 +50,7 @@ public:
 	TWeakPtr<IVideoDecoderResourceDelegate, ESPMode::ThreadSafe> VideoDecoderResourceDelegate;
 	TWeakPtr<ISurfaceRequestCallback, ESPMode::ThreadSafe> SurfaceRequestCallback;
 	TWeakPtr<IBufferReleaseCallback, ESPMode::ThreadSafe> BufferReleaseCallback;
-	
+
 	IDecoderPlatformResourceAndroid::ISurfaceRequestCallback::ESurfaceType SurfaceType = IDecoderPlatformResourceAndroid::ISurfaceRequestCallback::ESurfaceType::Error;
 	volatile uint32 CurrentID = 0;
 
@@ -183,7 +183,7 @@ public:
 };
 
 
-bool FElectraDecoderResourceManagerAndroid::SetupRenderBufferFromDecoderOutput(IMediaRenderer::IBuffer* InOutBufferToSetup, TSharedPtr<FParamDict, ESPMode::ThreadSafe> InOutBufferPropertes, TSharedPtr<IElectraDecoderVideoOutput, ESPMode::ThreadSafe> InDecoderOutput, IDecoderPlatformResource* InPlatformSpecificResource)
+bool FElectraDecoderResourceManagerAndroid::SetupRenderBufferFromDecoderOutput(FString& OutErrorMessage, IMediaRenderer::IBuffer* InOutBufferToSetup, TSharedPtr<FParamDict, ESPMode::ThreadSafe> InOutBufferPropertes, TSharedPtr<IElectraDecoderVideoOutput, ESPMode::ThreadSafe> InDecoderOutput, IDecoderPlatformResource* InPlatformSpecificResource)
 {
 	check(InOutBufferToSetup);
 	check(InOutBufferPropertes.IsValid());
@@ -192,7 +192,7 @@ bool FElectraDecoderResourceManagerAndroid::SetupRenderBufferFromDecoderOutput(I
 	{
 		return false;
 	}
-	
+
 	FInstanceVars* InInst = static_cast<FInstanceVars*>(InPlatformSpecificResource);
 	TSharedPtr<IVideoDecoderResourceDelegate, ESPMode::ThreadSafe> PinnedResourceDelegate = InInst->VideoDecoderResourceDelegate.Pin();
 	TSharedPtr<FElectraPlayerVideoDecoderOutputAndroid, ESPMode::ThreadSafe> DecoderOutput = InOutBufferToSetup->GetBufferProperties().GetValue(RenderOptionKeys::Texture).GetSharedPointer<FElectraPlayerVideoDecoderOutputAndroid>();
diff --git a/Engine/Plugins/Media/ElectraPlayer/Source/ElectraPlayerRuntime/Private/Runtime/Decoder/Android/AndroidElectraDecoderResourceManager.h b/Engine/Plugins/Media/ElectraPlayer/Source/ElectraPlayerRuntime/Private/Runtime/Decoder/Android/AndroidElectraDecoderResourceManager.h
--- a/Engine/Plugins/Media/ElectraPlayer/Source/ElectraPlayerRuntime/Private/Runtime/Decoder/Android/AndroidElectraDecoderResourceManager.h
+++ b/Engine/Plugins/Media/ElectraPlayer/Source/ElectraPlayerRuntime/Private/Runtime/Decoder/Android/AndroidElectraDecoderResourceManager.h
@@ -18,7 +18,7 @@ public:
 	static void Shutdown();
 	static TSharedPtr<IElectraDecoderResourceDelegate, ESPMode::ThreadSafe> GetDelegate();
 
-	static bool SetupRenderBufferFromDecoderOutput(IMediaRenderer::IBuffer* InOutBufferToSetup, TSharedPtr<FParamDict, ESPMode::ThreadSafe> InOutBufferPropertes, TSharedPtr<IElectraDecoderVideoOutput, ESPMode::ThreadSafe> InDecoderOutput, IDecoderPlatformResource* InPlatformSpecificResource);
+	static bool SetupRenderBufferFromDecoderOutput(FString& OutErrorMessage, IMediaRenderer::IBuffer* InOutBufferToSetup, TSharedPtr<FParamDict, ESPMode::ThreadSafe> InOutBufferPropertes, TSharedPtr<IElectraDecoderVideoOutput, ESPMode::ThreadSafe> InDecoderOutput, IDecoderPlatformResource* InPlatformSpecificResource);
 
 	virtual ~FElectraDecoderResourceManagerAndroid();
 	IDecoderPlatformResource* CreatePlatformResource(void* InOwnerHandle, EDecoderPlatformResourceType InDecoderResourceType, const TMap<FString, FVariant> InOptions) override;
diff --git a/Engine/Plugins/Media/ElectraPlayer/Source/ElectraPlayerRuntime/Private/Runtime/Decoder/Apple/AppleElectraDecoderResourceManager.cpp b/Engine/Plugins/Media/ElectraPlayer/Source/ElectraPlayerRuntime/Private/Runtime/Decoder/Apple/AppleElectraDecoderResourceManager.cpp
--- a/Engine/Plugins/Media/ElectraPlayer/Source/ElectraPlayerRuntime/Private/Runtime/Decoder/Apple/AppleElectraDecoderResourceManager.cpp
+++ b/Engine/Plugins/Media/ElectraPlayer/Source/ElectraPlayerRuntime/Private/Runtime/Decoder/Apple/AppleElectraDecoderResourceManager.cpp
@@ -85,7 +85,7 @@ bool FElectraDecoderResourceManagerApple::GetMetalDevice(void **OutMetalDevice)
 }
 
 
-bool FElectraDecoderResourceManagerApple::SetupRenderBufferFromDecoderOutput(IMediaRenderer::IBuffer* InOutBufferToSetup, TSharedPtr<FParamDict, ESPMode::ThreadSafe> InOutBufferPropertes, TSharedPtr<IElectraDecoderVideoOutput, ESPMode::ThreadSafe> InDecoderOutput, FElectraDecoderResourceManagerApple::IDecoderPlatformResource* InPlatformSpecificResource)
+bool FElectraDecoderResourceManagerApple::SetupRenderBufferFromDecoderOutput(FString& OutErrorMessage, IMediaRenderer::IBuffer* InOutBufferToSetup, TSharedPtr<FParamDict, ESPMode::ThreadSafe> InOutBufferPropertes, TSharedPtr<IElectraDecoderVideoOutput, ESPMode::ThreadSafe> InDecoderOutput, FElectraDecoderResourceManagerApple::IDecoderPlatformResource* InPlatformSpecificResource)
 {
 	check(InOutBufferToSetup);
 	check(InOutBufferPropertes.IsValid());
diff --git a/Engine/Plugins/Media/ElectraPlayer/Source/ElectraPlayerRuntime/Private/Runtime/Decoder/Apple/AppleElectraDecoderResourceManager.h b/Engine/Plugins/Media/ElectraPlayer/Source/ElectraPlayerRuntime/Private/Runtime/Decoder/Apple/AppleElectraDecoderResourceManager.h
--- a/Engine/Plugins/Media/ElectraPlayer/Source/ElectraPlayerRuntime/Private/Runtime/Decoder/Apple/AppleElectraDecoderResourceManager.h
+++ b/Engine/Plugins/Media/ElectraPlayer/Source/ElectraPlayerRuntime/Private/Runtime/Decoder/Apple/AppleElectraDecoderResourceManager.h
@@ -29,7 +29,7 @@ public:
 	static void Shutdown();
 	static TSharedPtr<IElectraDecoderResourceDelegate, ESPMode::ThreadSafe> GetDelegate();
 
-	static bool SetupRenderBufferFromDecoderOutput(IMediaRenderer::IBuffer* InOutBufferToSetup, TSharedPtr<FParamDict, ESPMode::ThreadSafe> InOutBufferPropertes, TSharedPtr<IElectraDecoderVideoOutput, ESPMode::ThreadSafe> InDecoderOutput, IDecoderPlatformResource* InPlatformSpecificResource);
+	static bool SetupRenderBufferFromDecoderOutput(FString& OutErrorMessage, IMediaRenderer::IBuffer* InOutBufferToSetup, TSharedPtr<FParamDict, ESPMode::ThreadSafe> InOutBufferPropertes, TSharedPtr<IElectraDecoderVideoOutput, ESPMode::ThreadSafe> InDecoderOutput, IDecoderPlatformResource* InPlatformSpecificResource);
 
 	IDecoderPlatformResource* CreatePlatformResource(void* InOwnerHandle, EDecoderPlatformResourceType InDecoderResourceType, const TMap<FString, FVariant> InOptions) override;
 	void ReleasePlatformResource(void* InOwnerHandle, IDecoderPlatformResource* InHandleToDestroy) override;
diff --git a/Engine/Plugins/Media/ElectraPlayer/Source/ElectraPlayerRuntime/Private/Runtime/Decoder/ElectraDecoderResourceManager.h b/Engine/Plugins/Media/ElectraPlayer/Source/ElectraPlayerRuntime/Private/Runtime/Decoder/ElectraDecoderResourceManager.h
--- a/Engine/Plugins/Media/ElectraPlayer/Source/ElectraPlayerRuntime/Private/Runtime/Decoder/ElectraDecoderResourceManager.h
+++ b/Engine/Plugins/Media/ElectraPlayer/Source/ElectraPlayerRuntime/Private/Runtime/Decoder/ElectraDecoderResourceManager.h
@@ -25,7 +25,7 @@ public:
 
 	class IDecoderPlatformResource;
 
-	static bool SetupRenderBufferFromDecoderOutput(Electra::IMediaRenderer::IBuffer* InOutBufferToSetup, Electra::FParamDict* InOutBufferPropertes, TSharedPtr<IElectraDecoderVideoOutput, ESPMode::ThreadSafe> InDecoderOutput, IDecoderPlatformResource* InPlatformSpecificResource)
+	static bool SetupRenderBufferFromDecoderOutput(FString& OutErrorMessage, Electra::IMediaRenderer::IBuffer* InOutBufferToSetup, Electra::FParamDict* InOutBufferPropertes, TSharedPtr<IElectraDecoderVideoOutput, ESPMode::ThreadSafe> InDecoderOutput, IDecoderPlatformResource* InPlatformSpecificResource)
 	{ return false; }
 };
 
diff --git a/Engine/Plugins/Media/ElectraPlayer/Source/ElectraPlayerRuntime/Private/Runtime/Decoder/Linux/LinuxElectraDecoderResourceManager.cpp b/Engine/Plugins/Media/ElectraPlayer/Source/ElectraPlayerRuntime/Private/Runtime/Decoder/Linux/LinuxElectraDecoderResourceManager.cpp
--- a/Engine/Plugins/Media/ElectraPlayer/Source/ElectraPlayerRuntime/Private/Runtime/Decoder/Linux/LinuxElectraDecoderResourceManager.cpp
+++ b/Engine/Plugins/Media/ElectraPlayer/Source/ElectraPlayerRuntime/Private/Runtime/Decoder/Linux/LinuxElectraDecoderResourceManager.cpp
@@ -245,7 +245,7 @@ static void ConvertDecodedImageToNV12(TArray<uint8>& OutNV12Buffer, FIntPoint Ou
 #endif
 
 
-bool FElectraDecoderResourceManagerLinux::SetupRenderBufferFromDecoderOutput(IMediaRenderer::IBuffer* InOutBufferToSetup, TSharedPtr<FParamDict, ESPMode::ThreadSafe> InOutBufferPropertes, TSharedPtr<IElectraDecoderVideoOutput, ESPMode::ThreadSafe> InDecoderOutput, IDecoderPlatformResource* InPlatformSpecificResource)
+bool FElectraDecoderResourceManagerLinux::SetupRenderBufferFromDecoderOutput(FString& OutErrorMessage, IMediaRenderer::IBuffer* InOutBufferToSetup, TSharedPtr<FParamDict, ESPMode::ThreadSafe> InOutBufferPropertes, TSharedPtr<IElectraDecoderVideoOutput, ESPMode::ThreadSafe> InDecoderOutput, IDecoderPlatformResource* InPlatformSpecificResource)
 {
 	check(InOutBufferToSetup);
 	check(InOutBufferPropertes.IsValid());
diff --git a/Engine/Plugins/Media/ElectraPlayer/Source/ElectraPlayerRuntime/Private/Runtime/Decoder/Linux/LinuxElectraDecoderResourceManager.h b/Engine/Plugins/Media/ElectraPlayer/Source/ElectraPlayerRuntime/Private/Runtime/Decoder/Linux/LinuxElectraDecoderResourceManager.h
--- a/Engine/Plugins/Media/ElectraPlayer/Source/ElectraPlayerRuntime/Private/Runtime/Decoder/Linux/LinuxElectraDecoderResourceManager.h
+++ b/Engine/Plugins/Media/ElectraPlayer/Source/ElectraPlayerRuntime/Private/Runtime/Decoder/Linux/LinuxElectraDecoderResourceManager.h
@@ -24,7 +24,7 @@ public:
 	static void Shutdown();
 	static TSharedPtr<IElectraDecoderResourceDelegate, ESPMode::ThreadSafe> GetDelegate();
 
-	static bool SetupRenderBufferFromDecoderOutput(IMediaRenderer::IBuffer* InOutBufferToSetup, TSharedPtr<FParamDict, ESPMode::ThreadSafe> InOutBufferPropertes, TSharedPtr<IElectraDecoderVideoOutput, ESPMode::ThreadSafe> InDecoderOutput, IDecoderPlatformResource* InPlatformSpecificResource);
+	static bool SetupRenderBufferFromDecoderOutput(FString& OutErrorMessage, IMediaRenderer::IBuffer* InOutBufferToSetup, TSharedPtr<FParamDict, ESPMode::ThreadSafe> InOutBufferPropertes, TSharedPtr<IElectraDecoderVideoOutput, ESPMode::ThreadSafe> InDecoderOutput, IDecoderPlatformResource* InPlatformSpecificResource);
 
 	virtual ~FElectraDecoderResourceManagerLinux();
 };
diff --git a/Engine/Plugins/Media/ElectraPlayer/Source/ElectraPlayerRuntime/Private/Runtime/Decoder/VideoDecoder.cpp b/Engine/Plugins/Media/ElectraPlayer/Source/ElectraPlayerRuntime/Private/Runtime/Decoder/VideoDecoder.cpp
--- a/Engine/Plugins/Media/ElectraPlayer/Source/ElectraPlayerRuntime/Private/Runtime/Decoder/VideoDecoder.cpp
+++ b/Engine/Plugins/Media/ElectraPlayer/Source/ElectraPlayerRuntime/Private/Runtime/Decoder/VideoDecoder.cpp
@@ -1080,9 +1080,17 @@ IElectraDecoder::EOutputStatus FVideoDecoderImpl::HandleOutput()
 						}
 					}
 
-					if (!FPlatformElectraDecoderResourceManager::SetupRenderBufferFromDecoderOutput(CurrentOutputBuffer, BufferProperties, CurrentDecoderOutput, PlatformResource))
+					FString DecoderOutputErrorMsg;
+					if (!FPlatformElectraDecoderResourceManager::SetupRenderBufferFromDecoderOutput(DecoderOutputErrorMsg, CurrentOutputBuffer, BufferProperties, CurrentDecoderOutput, PlatformResource))
 					{
-						PostError(0, TEXT("Failed to set up the decoder output!"), ERRCODE_VIDEO_INTERNAL_FAILED_TO_CONVERT_OUTPUT);
+						if (DecoderOutputErrorMsg.IsEmpty())
+						{
+							PostError(0, TEXT("Failed to set up the decoder output!"), ERRCODE_VIDEO_INTERNAL_FAILED_TO_CONVERT_OUTPUT);
+						}
+						else
+						{
+							PostError(0, FString::Printf(TEXT("Failed to set up the decoder output (%s)!"), *DecoderOutputErrorMsg), ERRCODE_VIDEO_INTERNAL_FAILED_TO_CONVERT_OUTPUT);
+						}
 						return IElectraDecoder::EOutputStatus::Error;
 					}
 					Renderer->ReturnBuffer(CurrentOutputBuffer, bUseOutput, *BufferProperties);
diff --git a/Engine/Plugins/Media/ElectraPlayer/Source/ElectraPlayerRuntime/Private/Runtime/Decoder/Windows/WindowsElectraDecoderResourceManager.cpp b/Engine/Plugins/Media/ElectraPlayer/Source/ElectraPlayerRuntime/Private/Runtime/Decoder/Windows/WindowsElectraDecoderResourceManager.cpp
--- a/Engine/Plugins/Media/ElectraPlayer/Source/ElectraPlayerRuntime/Private/Runtime/Decoder/Windows/WindowsElectraDecoderResourceManager.cpp
+++ b/Engine/Plugins/Media/ElectraPlayer/Source/ElectraPlayerRuntime/Private/Runtime/Decoder/Windows/WindowsElectraDecoderResourceManager.cpp
@@ -149,7 +149,7 @@ bool FElectraDecoderResourceManagerWindows::RunCodeAsync(TFunction<void()>&& Cod
 }
 
 
-bool FElectraDecoderResourceManagerWindows::SetupRenderBufferFromDecoderOutputFromMFSample(IMediaRenderer::IBuffer* InOutBufferToSetup, TSharedPtr<FParamDict, ESPMode::ThreadSafe> InOutBufferPropertes, TSharedPtr<IElectraDecoderVideoOutput, ESPMode::ThreadSafe> InDecoderOutput, FElectraDecoderResourceManagerWindows::IDecoderPlatformResource* InPlatformSpecificResource)
+bool FElectraDecoderResourceManagerWindows::SetupRenderBufferFromDecoderOutputFromMFSample(FString& OutErrorMessage, IMediaRenderer::IBuffer* InOutBufferToSetup, TSharedPtr<FParamDict, ESPMode::ThreadSafe> InOutBufferPropertes, TSharedPtr<IElectraDecoderVideoOutput, ESPMode::ThreadSafe> InDecoderOutput, FElectraDecoderResourceManagerWindows::IDecoderPlatformResource* InPlatformSpecificResource)
 {
 	TSharedPtr<FElectraPlayerVideoDecoderOutputPC, ESPMode::ThreadSafe> DecoderOutput = InOutBufferToSetup->GetBufferProperties().GetValue(RenderOptionKeys::Texture).GetSharedPointer<FElectraPlayerVideoDecoderOutputPC>();
 	if (DecoderOutput.IsValid())
@@ -396,7 +396,7 @@ bool FElectraDecoderResourceManagerWindows::SetupRenderBufferFromDecoderOutputFr
 }
 
 
-bool FElectraDecoderResourceManagerWindows::SetupRenderBufferFromDecoderOutput(IMediaRenderer::IBuffer* InOutBufferToSetup, TSharedPtr<FParamDict, ESPMode::ThreadSafe> InOutBufferPropertes, TSharedPtr<IElectraDecoderVideoOutput, ESPMode::ThreadSafe> InDecoderOutput, FElectraDecoderResourceManagerWindows::IDecoderPlatformResource* InPlatformSpecificResource)
+bool FElectraDecoderResourceManagerWindows::SetupRenderBufferFromDecoderOutput(FString& OutErrorMessage, IMediaRenderer::IBuffer* InOutBufferToSetup, TSharedPtr<FParamDict, ESPMode::ThreadSafe> InOutBufferPropertes, TSharedPtr<IElectraDecoderVideoOutput, ESPMode::ThreadSafe> InDecoderOutput, FElectraDecoderResourceManagerWindows::IDecoderPlatformResource* InPlatformSpecificResource)
 {
 	check(InOutBufferToSetup);
 	check(InOutBufferPropertes.IsValid());
@@ -561,7 +561,7 @@ bool FElectraDecoderResourceManagerWindows::SetupRenderBufferFromDecoderOutput(I
 			//
 			if (InDecoderOutput->GetPlatformOutputHandle(EElectraDecoderPlatformOutputHandleType::MFSample) != nullptr)
 			{
-				return SetupRenderBufferFromDecoderOutputFromMFSample(InOutBufferToSetup, InOutBufferPropertes, InDecoderOutput, InPlatformSpecificResource);
+				return SetupRenderBufferFromDecoderOutputFromMFSample(OutErrorMessage, InOutBufferToSetup, InOutBufferPropertes, InDecoderOutput, InPlatformSpecificResource);
 			}
 		}
 	}
diff --git a/Engine/Plugins/Media/ElectraPlayer/Source/ElectraPlayerRuntime/Private/Runtime/Decoder/Windows/WindowsElectraDecoderResourceManager.h b/Engine/Plugins/Media/ElectraPlayer/Source/ElectraPlayerRuntime/Private/Runtime/Decoder/Windows/WindowsElectraDecoderResourceManager.h
--- a/Engine/Plugins/Media/ElectraPlayer/Source/ElectraPlayerRuntime/Private/Runtime/Decoder/Windows/WindowsElectraDecoderResourceManager.h
+++ b/Engine/Plugins/Media/ElectraPlayer/Source/ElectraPlayerRuntime/Private/Runtime/Decoder/Windows/WindowsElectraDecoderResourceManager.h
@@ -28,7 +28,7 @@ public:
 	static void Shutdown();
 	static TSharedPtr<IElectraDecoderResourceDelegate, ESPMode::ThreadSafe> GetDelegate();
 
-	static bool SetupRenderBufferFromDecoderOutput(IMediaRenderer::IBuffer* InOutBufferToSetup, TSharedPtr<FParamDict, ESPMode::ThreadSafe> InOutBufferPropertes, TSharedPtr<IElectraDecoderVideoOutput, ESPMode::ThreadSafe> InDecoderOutput, IDecoderPlatformResource* InPlatformSpecificResource);
+	static bool SetupRenderBufferFromDecoderOutput(FString& OutErrorMessage, IMediaRenderer::IBuffer* InOutBufferToSetup, TSharedPtr<FParamDict, ESPMode::ThreadSafe> InOutBufferPropertes, TSharedPtr<IElectraDecoderVideoOutput, ESPMode::ThreadSafe> InDecoderOutput, IDecoderPlatformResource* InPlatformSpecificResource);
 
 	IDecoderPlatformResource* CreatePlatformResource(void* InOwnerHandle, EDecoderPlatformResourceType InDecoderResourceType, const TMap<FString, FVariant> InOptions) override;
 	void ReleasePlatformResource(void* InOwnerHandle, IDecoderPlatformResource* InHandleToDestroy) override;
@@ -42,7 +42,7 @@ public:
 
 private:
 	class FInstanceVars;
-	static bool SetupRenderBufferFromDecoderOutputFromMFSample(IMediaRenderer::IBuffer* InOutBufferToSetup, TSharedPtr<FParamDict, ESPMode::ThreadSafe> InOutBufferPropertes, TSharedPtr<IElectraDecoderVideoOutput, ESPMode::ThreadSafe> InDecoderOutput, IDecoderPlatformResource* InPlatformSpecificResource);
+	static bool SetupRenderBufferFromDecoderOutputFromMFSample(FString& OutErrorMessage, IMediaRenderer::IBuffer* InOutBufferToSetup, TSharedPtr<FParamDict, ESPMode::ThreadSafe> InOutBufferPropertes, TSharedPtr<IElectraDecoderVideoOutput, ESPMode::ThreadSafe> InDecoderOutput, IDecoderPlatformResource* InPlatformSpecificResource);
 };
 
 using FPlatformElectraDecoderResourceManager = FElectraDecoderResourceManagerWindows;
diff --git a/Engine/Plugins/Media/ElectraPlayer/Source/ElectraPlayerRuntime/Public/IElectraPlayerDecoderResourceManager.h b/Engine/Plugins/Media/ElectraPlayer/Source/ElectraPlayerRuntime/Public/IElectraPlayerDecoderResourceManager.h
--- a/Engine/Plugins/Media/ElectraPlayer/Source/ElectraPlayerRuntime/Public/IElectraPlayerDecoderResourceManager.h
+++ b/Engine/Plugins/Media/ElectraPlayer/Source/ElectraPlayerRuntime/Public/IElectraPlayerDecoderResourceManager.h
@@ -21,7 +21,7 @@ class FElectraPlayerDecoderResourceManager
 public:
 	ELECTRAPLAYERRUNTIME_API static TSharedPtr<IElectraDecoderResourceDelegate, ESPMode::ThreadSafe> GetDelegate();
 	ELECTRAPLAYERRUNTIME_API static TSharedPtr<Electra::IVideoDecoderResourceDelegate, ESPMode::ThreadSafe> CreatePlatformVideoDecoderResourceDelegate(TSharedPtr<IElectraPlayerAdapterDelegate, ESPMode::ThreadSafe> InElectraPlayerAdapterDelegate);
-	ELECTRAPLAYERRUNTIME_API static bool SetupRenderBufferFromDecoderOutput(Electra::IMediaRenderer::IBuffer* InOutBufferToSetup, TSharedPtr<Electra::FParamDict, ESPMode::ThreadSafe> InOutBufferPropertes, TSharedPtr<IElectraDecoderVideoOutput, ESPMode::ThreadSafe> InDecoderOutput, IElectraDecoderResourceDelegateBase::IDecoderPlatformResource* InPlatformSpecificResource);
+	ELECTRAPLAYERRUNTIME_API static bool SetupRenderBufferFromDecoderOutput(FString& OutErrorMessage, Electra::IMediaRenderer::IBuffer* InOutBufferToSetup, TSharedPtr<Electra::FParamDict, ESPMode::ThreadSafe> InOutBufferPropertes, TSharedPtr<IElectraDecoderVideoOutput, ESPMode::ThreadSafe> InDecoderOutput, IElectraDecoderResourceDelegateBase::IDecoderPlatformResource* InPlatformSpecificResource);
 	class FVideo
 	{
 	public:
diff --git a/Engine/Plugins/Media/ElectraPlayer/Source/ElectraProtron/Private/Player/ElectraProtronPlayerDecoder.cpp b/Engine/Plugins/Media/ElectraPlayer/Source/ElectraProtron/Private/Player/ElectraProtronPlayerDecoder.cpp
--- a/Engine/Plugins/Media/ElectraPlayer/Source/ElectraProtron/Private/Player/ElectraProtronPlayerDecoder.cpp
+++ b/Engine/Plugins/Media/ElectraPlayer/Source/ElectraProtron/Private/Player/ElectraProtronPlayerDecoder.cpp
@@ -986,7 +986,8 @@ void FElectraProtronPlayer::FImpl::FDecoderThread::HandleOutputFrame()
 						}
 					}
 
-					if (ensure(FElectraPlayerDecoderResourceManager::SetupRenderBufferFromDecoderOutput(&DecOutputPtrWrapper, BufferProperties, VideoDecoderOutput, PlatformResource)))
+					FString DecoderOutputErrorMsg;
+					if (ensure(FElectraPlayerDecoderResourceManager::SetupRenderBufferFromDecoderOutput(DecoderOutputErrorMsg, &DecOutputPtrWrapper, BufferProperties, VideoDecoderOutput, PlatformResource)))
 					{
 						FElectraTextureSampleRef TextureSample = Params.TexturePool->AcquireShared();
 						TextureSample->Initialize(NewDecoderOutput.Get());
