ElectraPlayer: Fixed timestamp of raw WebVTT segments

#rnx

diff --git a/Engine/Plugins/Media/ElectraPlayer/Source/ElectraPlayerRuntime/Private/Runtime/Player/StreamSegmentReaderCommon.cpp b/Engine/Plugins/Media/ElectraPlayer/Source/ElectraPlayerRuntime/Private/Runtime/Player/StreamSegmentReaderCommon.cpp
--- a/Engine/Plugins/Media/ElectraPlayer/Source/ElectraPlayerRuntime/Private/Runtime/Player/StreamSegmentReaderCommon.cpp
+++ b/Engine/Plugins/Media/ElectraPlayer/Source/ElectraPlayerRuntime/Private/Runtime/Player/StreamSegmentReaderCommon.cpp
@@ -3584,9 +3584,9 @@ FStreamSegmentReaderCommon::FStreamHandler::EHandleResult FStreamSegmentReaderCo
 			ci.SetCodecSpecifierRFC6381(TEXT("wvtt"));
 		}
 
-		FTimeValue PTS;
-		PTS.SetFromHNS(InRequest->Segment.MediaLocalFirstAUTime);
-		FTimeValue EarliestPTS(PTS);
+		FTimeValue PTS, EarliestPTS;
+		EarliestPTS.SetFromHNS(InRequest->Segment.MediaLocalFirstAUTime);
+		PTS.SetFromHNS(InRequest->Segment.Time);
 		if (InRequest->Segment.bFrameAccuracyRequired && InRequest->FrameAccurateStartTime.IsValid())
 		{
 			EarliestPTS = InRequest->FrameAccurateStartTime;
