ElectraPlayer: allow same playback rates for HLS as we do for other formats

#rnx

diff --git a/Engine/Plugins/Media/ElectraPlayer/Source/ElectraPlayerRuntime/Private/Runtime/Player/HLS/PlaylistHLS.cpp b/Engine/Plugins/Media/ElectraPlayer/Source/ElectraPlayerRuntime/Private/Runtime/Player/HLS/PlaylistHLS.cpp
--- a/Engine/Plugins/Media/ElectraPlayer/Source/ElectraPlayerRuntime/Private/Runtime/Player/HLS/PlaylistHLS.cpp
+++ b/Engine/Plugins/Media/ElectraPlayer/Source/ElectraPlayerRuntime/Private/Runtime/Player/HLS/PlaylistHLS.cpp
@@ -2298,7 +2298,14 @@ IManifest::ELiveEdgePlayMode FActiveHLSPlaylist::GetLiveEdgePlayMode() const
 TRangeSet<double> FActiveHLSPlaylist::GetPossiblePlaybackRates(EPlayRateType InForType) const
 {
 	TRangeSet<double> Ranges;
-	Ranges.Add(TRange<double>{1.0}); // normal (real-time) playback rate
+	if (InForType == IManifest::EPlayRateType::UnthinnedRate)
+	{
+		Ranges.Add(TRange<double>::Inclusive(0.1, 4.0));
+	}
+	else
+	{
+		Ranges.Add(TRange<double>{1.0}); // normal (real-time) playback rate
+	}
 	Ranges.Add(TRange<double>{0.0}); // and pause
 	return Ranges;
 }
