[ElectraProtron] Unable to specify Protron Player for external video assets

- Adding a setting to select Protron over Electra. Disabled by default.

diff --git a/Engine/Plugins/Media/ElectraPlayer/Source/ElectraProtronFactory/Private/ElectraProtronFactoryModule.cpp b/Engine/Plugins/Media/ElectraPlayer/Source/ElectraProtronFactory/Private/ElectraProtronFactoryModule.cpp
--- a/Engine/Plugins/Media/ElectraPlayer/Source/ElectraProtronFactory/Private/ElectraProtronFactoryModule.cpp
+++ b/Engine/Plugins/Media/ElectraPlayer/Source/ElectraProtronFactory/Private/ElectraProtronFactoryModule.cpp
@@ -1,6 +1,7 @@
 // Copyright Epic Games, Inc. All Rights Reserved.
 
 #include "ElectraProtronFactoryPrivate.h"
+#include "ElectraProtronFactorySettings.h"
 
 #include "Containers/Array.h"
 #include "Containers/UnrealString.h"
@@ -17,6 +18,11 @@
 
 #include "IElectraProtronModule.h"
 
+#if WITH_EDITOR
+#include "ISettingsModule.h"
+#include "ISettingsSection.h"
+#endif
+
 DEFINE_LOG_CATEGORY(LogElectraProtronFactory);
 
 #define LOCTEXT_NAMESPACE "ElectraProtronFactoryModule"
@@ -75,6 +81,15 @@ public:
 		// so we cannot check for this.
 		if (SupportedFileExtensions.Contains(LowerCaseExtension))
 		{
+			// See if Protron is preferred over Electra for the current executable.
+			const UElectraProtronFactorySettings* Settings = GetDefault<UElectraProtronFactorySettings>();
+			const bool bIsInGame = !GIsEditor || IsRunningGame();
+			if (Settings && ((bIsInGame && Settings->bPreferProtronInGame) || (!bIsInGame && Settings->bPreferProtronInEditor)))
+			{
+				// Electra's confidence score is 100 (and most other players are 80).
+				return 101;
+			}
+			
 			// Low score to select other players instead.
 			return 1;
 		}
@@ -141,10 +156,34 @@ public:
 
 		// register player factory
 		MediaModule->RegisterPlayerFactory(*this);
+
+#if WITH_EDITOR
+		// register settings
+		ISettingsModule* SettingsModule = FModuleManager::GetModulePtr<ISettingsModule>("Settings");
+
+		if (SettingsModule != nullptr)
+		{
+			ISettingsSectionPtr SettingsSection = SettingsModule->RegisterSettings("Project", "Plugins", "ElectraProtronFactory",
+				LOCTEXT("ElectraProtronFactorySettingsName", "Electra Protron Factory"),
+				LOCTEXT("ElectraProtronFactorySettingsDescription", "Configure the Electra Protron Factory."),
+				GetMutableDefault<UElectraProtronFactorySettings>()
+			);
+		}
+#endif //WITH_EDITOR
 	}
 
 	void ShutdownModule() override
 	{
+#if WITH_EDITOR
+		// unregister settings
+		ISettingsModule* SettingsModule = FModuleManager::GetModulePtr<ISettingsModule>("Settings");
+
+		if (SettingsModule != nullptr)
+		{
+			SettingsModule->UnregisterSettings("Project", "Plugins", "ElectraProtronFactory");
+		}
+#endif //WITH_EDITOR
+
 		// Get media module once more to be sure it is still there
 		MediaModule = FModuleManager::GetModulePtr<IMediaModule>("Media");
 
diff --git a/Engine/Plugins/Media/ElectraPlayer/Source/ElectraProtronFactory/Private/ElectraProtronFactorySettings.h b/Engine/Plugins/Media/ElectraPlayer/Source/ElectraProtronFactory/Private/ElectraProtronFactorySettings.h
new file mode 100644
--- /dev/null
+++ b/Engine/Plugins/Media/ElectraPlayer/Source/ElectraProtronFactory/Private/ElectraProtronFactorySettings.h
@@ -0,0 +1,23 @@
+// Copyright Epic Games, Inc. All Rights Reserved.
+
+#pragma once
+
+#include "ElectraProtronFactorySettings.generated.h"
+
+/**
+ * Settings for the Electra Protron Factory module.
+ */
+UCLASS(config=Engine)
+class UElectraProtronFactorySettings : public UObject
+{
+	GENERATED_BODY()
+	
+public:
+	/** When the media source leaves the player selection to automatic, Protron will be used instead of Electra in editor. */
+	UPROPERTY(config, EditAnywhere, Category=General)
+	bool bPreferProtronInEditor = false;
+
+	/** When the media source leaves the player selection to automatic, Protron will be used instead of Electra in game. */
+	UPROPERTY(config, EditAnywhere, Category=General)
+	bool bPreferProtronInGame = false;
+};
\ No newline at end of file
