Removal of unused ElectraCSA

#rnx

diff --git a/Engine/Plugins/Media/ElectraPlayer/Source/ElectraPlayerPlugin/Private/ElectraPlayerPluginModule.cpp b/Engine/Plugins/Media/ElectraPlayer/Source/ElectraPlayerPlugin/Private/ElectraPlayerPluginModule.cpp
--- a/Engine/Plugins/Media/ElectraPlayer/Source/ElectraPlayerPlugin/Private/ElectraPlayerPluginModule.cpp
+++ b/Engine/Plugins/Media/ElectraPlayer/Source/ElectraPlayerPlugin/Private/ElectraPlayerPluginModule.cpp
@@ -10,7 +10,6 @@
 #include "IElectraPlayerRuntimeModule.h"
 #include "ElectraPlayerPlugin.h"
 #include "ParameterDictionary.h"
-#include "SimpleElectraAudioPlayer.h"
 #include "IElectraDecoderResourceDelegateBase.h"
 #include "IElectraPlayerDecoderResourceManager.h"
 
@@ -54,7 +53,6 @@ public:
 	void SendAnalyticMetricsPerMinute(const TSharedPtr<IAnalyticsProviderET>& AnalyticsProvider) override
 	{
 		SendAnalyticMetricsPerMinuteDelegate.Broadcast(AnalyticsProvider);
-		ISimpleElectraAudioPlayer::SendAnalyticMetrics(AnalyticsProvider);
 	}
 
 	void ReportVideoStreamingError(const FGuid& PlayerGuid, const FString& LastError) override
diff --git a/Engine/Plugins/Media/ElectraPlayer/Source/ElectraPlayerRuntime/Private/Runtime/Player/SimpleElectraAudioPlayer.cpp b/Engine/Plugins/Media/ElectraPlayer/Source/ElectraPlayerRuntime/Private/Runtime/Player/SimpleElectraAudioPlayer.cpp
deleted file mode 100644
diff --git a/Engine/Plugins/Media/ElectraPlayer/Source/ElectraPlayerRuntime/Public/SimpleElectraAudioPlayer.h b/Engine/Plugins/Media/ElectraPlayer/Source/ElectraPlayerRuntime/Public/SimpleElectraAudioPlayer.h
deleted file mode 100644
