ElectraPlayer: fixed extraction of embedded mp4 content metadata

#rnx

diff --git a/Engine/Plugins/Media/ElectraPlayer/Source/ElectraPlayerRuntime/Private/Runtime/Player/mp4/PlaylistReaderMP4.cpp b/Engine/Plugins/Media/ElectraPlayer/Source/ElectraPlayerRuntime/Private/Runtime/Player/mp4/PlaylistReaderMP4.cpp
--- a/Engine/Plugins/Media/ElectraPlayer/Source/ElectraPlayerRuntime/Private/Runtime/Player/mp4/PlaylistReaderMP4.cpp
+++ b/Engine/Plugins/Media/ElectraPlayer/Source/ElectraPlayerRuntime/Private/Runtime/Player/mp4/PlaylistReaderMP4.cpp
@@ -228,7 +228,7 @@ void FPlaylistReaderMP4::WorkerThread()
 	CSV_SCOPED_TIMING_STAT(ElectraPlayer, MP4_PlaylistWorker);
 
 	const TArray<uint32> StopAfterBoxes { UtilsMP4::Make4CC('m','o','o','v') };
-	const TArray<uint32> ReadBoxes { UtilsMP4::Make4CC('m','o','o','v') };
+	const TArray<uint32> ReadBoxes { UtilsMP4::Make4CC('f','t','y','p'), UtilsMP4::Make4CC('m','o','o','v') };
 	const TArray<uint32> FirstBoxes { UtilsMP4::Make4CC('f','t','y','p'), UtilsMP4::Make4CC('s','t','y','p'), UtilsMP4::Make4CC('s','i','d','x'), UtilsMP4::Make4CC('f','r','e','e'), UtilsMP4::Make4CC('s','k','i','p') };
 	UtilsMP4::FMP4RootBoxLocator BoxLocator;
 	TArray<UtilsMP4::FMP4RootBoxLocator::FBoxInfo> BoxInfos;
@@ -269,12 +269,21 @@ void FPlaylistReaderMP4::WorkerThread()
 				LogMessage(IInfoLog::ELevel::Info, FString::Printf(TEXT("The mp4 at \"%s\" is not fast-startable. Consider moving the 'moov' box in front of the 'mdat' for faster startup times."), *ConnectionInfo.EffectiveURL));
 			}
 
-			ParseBuffer = BoxInfos.FindByPredicate([InType=UtilsMP4::Make4CC('m','o','o','v')](const UtilsMP4::FMP4RootBoxLocator::FBoxInfo& InBox){return InBox.Type == InType;})->DataBuffer;
+			MP4Parser = IParserISO14496_12::CreateParser();
+
+			// First parse the ftyp on which some elements in the moov depend.
+			ParseBuffer = BoxInfos.FindByPredicate([InType=UtilsMP4::Make4CC('f','t','y','p')](const UtilsMP4::FMP4RootBoxLocator::FBoxInfo& InBox){return InBox.Type == InType;})->DataBuffer;
 			ParsePos = 0;
 			ParseBufferSize = ParseBuffer->Num();
-
-			MP4Parser = IParserISO14496_12::CreateParser();
 			UEMediaError parseError = MP4Parser->ParseHeader(this, this, PlayerSessionServices, nullptr);
+			// If successful parse the moov
+			if (parseError == UEMEDIA_ERROR_OK || parseError == UEMEDIA_ERROR_END_OF_STREAM)
+			{
+				ParseBuffer = BoxInfos.FindByPredicate([InType=UtilsMP4::Make4CC('m','o','o','v')](const UtilsMP4::FMP4RootBoxLocator::FBoxInfo& InBox){return InBox.Type == InType;})->DataBuffer;
+				ParsePos = 0;
+				ParseBufferSize = ParseBuffer->Num();
+				parseError = MP4Parser->ParseHeader(this, this, PlayerSessionServices, nullptr);
+			}
 			if (parseError == UEMEDIA_ERROR_OK || parseError == UEMEDIA_ERROR_END_OF_STREAM)
 			{
 				// Prepare the tracks in the stream that are of a supported codec.
