Electra: Removal of the platform resource delegate and the wrapping plugin

diff --git a/Engine/Plugins/Media/ElectraCodecs/Source/ElectraDecoders/Private/Windows/h265/H265_VideoDecoder_Windows.cpp b/Engine/Plugins/Media/ElectraCodecs/Source/ElectraDecoders/Private/Windows/h265/H265_VideoDecoder_Windows.cpp
--- a/Engine/Plugins/Media/ElectraCodecs/Source/ElectraDecoders/Private/Windows/h265/H265_VideoDecoder_Windows.cpp
+++ b/Engine/Plugins/Media/ElectraCodecs/Source/ElectraDecoders/Private/Windows/h265/H265_VideoDecoder_Windows.cpp
@@ -6,7 +6,7 @@
 
 #include "DX/VideoDecoderH265_DX.h"
 #include "DX/DecoderErrors_DX.h"
-
+#include "ElectraDecodersPlatformResources.h"
 #include "ElectraDecodersUtils.h"
 #include "IElectraDecoderFeaturesAndOptions.h"
 
@@ -102,9 +102,9 @@ public:
 		IElectraVideoDecoderH265_DX::PlatformGetConfigurationOptions(OutOptions);
 	}
 
-	TSharedPtr<IElectraDecoder, ESPMode::ThreadSafe> CreateDecoderForFormat(const FString& InCodecFormat, const TMap<FString, FVariant>& InOptions, TSharedPtr<IElectraDecoderResourceDelegate, ESPMode::ThreadSafe> InResourceDelegate) override
+	TSharedPtr<IElectraDecoder, ESPMode::ThreadSafe> CreateDecoderForFormat(const FString& InCodecFormat, const TMap<FString, FVariant>& InOptions) override
 	{
-		return IElectraVideoDecoderH265_DX::Create(InOptions, InResourceDelegate);
+		return IElectraVideoDecoderH265_DX::Create(InOptions);
 	}
 
 };
@@ -195,15 +195,15 @@ void IElectraVideoDecoderH265_DX::PlatformGetConfigurationOptions(TMap<FString,
 }
 
 
-IElectraDecoder::FError IElectraVideoDecoderH265_DX::PlatformCreateMFDecoderTransform(IPlatformHandle** OutPlatformHandle, TSharedPtr<IElectraDecoderResourceDelegate, ESPMode::ThreadSafe> InResourceDelegate, const TMap<FString, FVariant>& InOptions)
+IElectraDecoder::FError IElectraVideoDecoderH265_DX::PlatformCreateMFDecoderTransform(IPlatformHandle** OutPlatformHandle, const TMap<FString, FVariant>& InOptions)
 {
 	IElectraDecoder::FError Error;
-	check(OutPlatformHandle && InResourceDelegate.IsValid());
+	check(OutPlatformHandle);
 	if (OutPlatformHandle)
 	{
 		*OutPlatformHandle = nullptr;
 	}
-	if (OutPlatformHandle && InResourceDelegate.IsValid())
+	if (OutPlatformHandle)
 	{
 		// Must be Windows 8 or better
 		if (FPlatformMisc::VerifyWindowsVersion(6, 2))
@@ -255,7 +255,7 @@ IElectraDecoder::FError IElectraVideoDecoderH265_DX::PlatformCreateMFDecoderTran
 				TRefCountPtr<IDXGIAdapter> DXGIAdapter;
 				void* ApplicationD3DDevice = nullptr;
 				int32 ApplicationD3DDeviceVersion = 0;
-				if (InResourceDelegate->GetD3DDevice(&ApplicationD3DDevice, &ApplicationD3DDeviceVersion))
+				if (FElectraDecodersPlatformResources::GetD3DDeviceAndVersion(&ApplicationD3DDevice, &ApplicationD3DDeviceVersion))
 				{
 					if (ApplicationD3DDeviceVersion >= 12000)
 					{
diff --git a/Engine/Plugins/Media/ElectraCodecs/Source/ElectraDecoders/Public/Android/AndroidElectraDecoderResourceDelegate.h b/Engine/Plugins/Media/ElectraCodecs/Source/ElectraDecoders/Public/Android/AndroidElectraDecoderResourceDelegate.h
--- a/Engine/Plugins/Media/ElectraCodecs/Source/ElectraDecoders/Public/Android/AndroidElectraDecoderResourceDelegate.h
+++ b/Engine/Plugins/Media/ElectraCodecs/Source/ElectraDecoders/Public/Android/AndroidElectraDecoderResourceDelegate.h
@@ -10,59 +10,6 @@ class IElectraDecoderResourceDelegateAndroid : public IElectraDecoderResourceDel
 {
 public:
 	virtual ~IElectraDecoderResourceDelegateAndroid() = default;
-
-	class IDecoderPlatformResourceAndroid : public IDecoderPlatformResource
-	{
-	public:
-		class ISurfaceRequestCallback : public TSharedFromThis<ISurfaceRequestCallback, ESPMode::ThreadSafe>
-		{
-		public:
-			virtual ~ISurfaceRequestCallback() = default;
-			enum class ESurfaceType
-			{
-				NoSurface,		// No decoding surface. Must use CPU side buffer.
-				Surface,
-				SurfaceView,
-				Error
-			};
-			virtual void OnNewSurface(ESurfaceType InSurfaceType, void* InSurface) = 0;
-		};
-		// The decoder calls this to request a surface onto which to decode.
-		virtual void RequestSurface(TWeakPtr<ISurfaceRequestCallback, ESPMode::ThreadSafe> InRequestCallback) = 0;
-
-		enum class ESurfaceChangeResult
-		{
-			NoChange,
-			NewSurface,
-			Error
-		};
-		/*
-			Called by the decoder prior to decoding an access unit if the initial surface returned by RequestSurface()
-			is a SurfaceView. If the SurfaceView has changed since it was initially returned the new SurfaceView must
-			be set to `OutNewSurfaceView` and `NewSurface` returned.
-			If there is no change returning `NoChange` is sufficient. In case of any error return `Error` to stop decoding.
-		*/
-		virtual ESurfaceChangeResult VerifySurfaceView(void*& OutNewSurfaceView, void* InCurrentSurfaceView) = 0;
-
-		class IBufferReleaseCallback : public TSharedFromThis<IBufferReleaseCallback, ESPMode::ThreadSafe>
-		{
-		public:
-			virtual ~IBufferReleaseCallback() = default;
-			virtual void OnReleaseBuffer(int32 InBufferIndex, TOptional<int32> InBufferValidCount, TOptional<bool> InDoRender, TOptional<int64> InRenderAt) = 0;
-		};
-		// The decoder calls this to set the buffer release callback.
-		virtual void SetBufferReleaseCallback(TWeakPtr<IBufferReleaseCallback, ESPMode::ThreadSafe> InBufferReleaseCallback) = 0;
-	};
-};
-
-class IElectraDecoderVideoOutputCopyResources
-{
-public:
-	virtual ~IElectraDecoderVideoOutputCopyResources() = default;
-	virtual void SetBufferIndex(int32 InIndex) = 0;
-	virtual void SetValidCount(int32 InValidCount) = 0;
-	virtual bool ShouldReleaseBufferImmediately() = 0;
 };
 
-
 using IElectraDecoderResourceDelegate = IElectraDecoderResourceDelegateAndroid;
diff --git a/Engine/Plugins/Media/ElectraCodecs/Source/ElectraDecoders/Public/Apple/AppleElectraDecoderResourceDelegate.h b/Engine/Plugins/Media/ElectraCodecs/Source/ElectraDecoders/Public/Apple/AppleElectraDecoderResourceDelegate.h
--- a/Engine/Plugins/Media/ElectraCodecs/Source/ElectraDecoders/Public/Apple/AppleElectraDecoderResourceDelegate.h
+++ b/Engine/Plugins/Media/ElectraCodecs/Source/ElectraDecoders/Public/Apple/AppleElectraDecoderResourceDelegate.h
@@ -12,9 +12,6 @@ class IElectraDecoderResourceDelegateApple : public IElectraDecoderResourceDeleg
 {
 public:
 	virtual ~IElectraDecoderResourceDelegateApple() = default;
-
-	// Hardware decoders might need to know the current Metal device.
-	virtual bool GetMetalDevice(void **OutMetalDevice) = 0;
 };
 
 
diff --git a/Engine/Plugins/Media/ElectraCodecs/Source/ElectraDecoders/Public/IElectraDecoder.h b/Engine/Plugins/Media/ElectraCodecs/Source/ElectraDecoders/Public/IElectraDecoder.h
--- a/Engine/Plugins/Media/ElectraCodecs/Source/ElectraDecoders/Public/IElectraDecoder.h
+++ b/Engine/Plugins/Media/ElectraCodecs/Source/ElectraDecoders/Public/IElectraDecoder.h
@@ -28,14 +28,6 @@ public:
 };
 
 
-class IElectraDecoderDefaultOutputFormat
-{
-public:
-	virtual ~IElectraDecoderDefaultOutputFormat() = default;
-};
-
-
-
 enum class EElectraDecoderFlags : uint32
 {
 	None = 0,
@@ -243,17 +235,6 @@ public:
 	 */
 	virtual bool ResetToCleanStart() = 0;
 
-	/**
-	 * Asks the decoder to provide information on the output that is expected to be produced
-	 * given the codec specific data as input.
-	 * This may be useful to check ahead of decoding what output format can be expected.
-	 * However, the decoder may still produce different output while performing actual decoding
-	 * if the output format cannot be derived from the CSD alone.
-	 *
-	 * If no output information can be provided a nullptr will be returned.
-	 */
-	virtual TSharedPtr<IElectraDecoderDefaultOutputFormat, ESPMode::ThreadSafe> GetDefaultOutputFormatFromCSD(const TMap<FString, FVariant>& CSDAndAdditionalOptions) = 0;
-
 	/**
 	 * Decodes one access unit.
 	 * The data sent for decoding must form one complete decodable unit. Partial data is not permitted.
diff --git a/Engine/Plugins/Media/ElectraCodecs/Source/ElectraDecoders/Public/IElectraDecoderOutputAudio.h b/Engine/Plugins/Media/ElectraCodecs/Source/ElectraDecoders/Public/IElectraDecoderOutputAudio.h
--- a/Engine/Plugins/Media/ElectraCodecs/Source/ElectraDecoders/Public/IElectraDecoderOutputAudio.h
+++ b/Engine/Plugins/Media/ElectraCodecs/Source/ElectraDecoders/Public/IElectraDecoderOutputAudio.h
@@ -4,36 +4,6 @@
 
 #include "IElectraDecoder.h"
 
-class IElectraDecoderDefaultAudioOutputFormat : public IElectraDecoderDefaultOutputFormat
-{
-public:
-	virtual ~IElectraDecoderDefaultAudioOutputFormat() = default;
-
-	/**
-	 * Returns the number of channels the format will decode into.
-	 * This could be different than what is known by the container format.
-	 * For instance, a mono HE-AACv2 stream may decode into stereo.
-	 */
-	virtual int32 GetNumChannels() const = 0;
-
-	/**
-	 * Returns the sampling rate the format will decode into.
-	 * This could be different than what is known by the container format.
-	 * For instance, a HE-AAC stream may decode into twice the sampling rate
-	 * if the decoder supports it.
-	 */
-	virtual int32 GetSampleRate() const = 0;
-
-	/**
-	 * Returns the number of frames one input access unit will decode into.
-	 * This could be different than what is known by the container format.
-	 * For instance, a HE-AAC stream may decode into twice the sampling rate
-	 * and thus twice the number of decoded frames if the decoder supports it.
-	 */
-	virtual int32 GetNumFrames() const = 0;
-};
-
-
 class IElectraDecoderAudioOutput : public IElectraDecoderOutput
 {
 public:
diff --git a/Engine/Plugins/Media/ElectraCodecs/Source/ElectraDecoders/Public/IElectraDecoderOutputVideo.h b/Engine/Plugins/Media/ElectraCodecs/Source/ElectraDecoders/Public/IElectraDecoderOutputVideo.h
--- a/Engine/Plugins/Media/ElectraCodecs/Source/ElectraDecoders/Public/IElectraDecoderOutputVideo.h
+++ b/Engine/Plugins/Media/ElectraCodecs/Source/ElectraDecoders/Public/IElectraDecoderOutputVideo.h
@@ -32,23 +32,6 @@ struct FElectraVideoDecoderOutputCropValues
 	int32 Right = 0;
 };
 
-class IElectraDecoderDefaultVideoOutputFormat : public IElectraDecoderDefaultOutputFormat
-{
-public:
-	virtual ~IElectraDecoderDefaultVideoOutputFormat() = default;
-
-};
-
-
-class IElectraDecoderVideoOutputTransferHandle
-{
-public:
-	virtual ~IElectraDecoderVideoOutputTransferHandle() = default;
-	virtual void* GetHandle() = 0;
-	virtual void ReleaseHandle() = 0;
-};
-
-
 class IElectraDecoderVideoOutputCopyResources;
 
 
@@ -173,14 +156,4 @@ public:
 	 * relevant parameters (eg. a device or texture handle, etc.) can be obtained.
 	 */
 	virtual EImageCopyResult CopyPlatformImage(IElectraDecoderVideoOutputCopyResources* InCopyResources) const = 0;
-
-	/**
-	 * Returns a "transfer" buffer handle.
-	 * This is a value particular to the target hardware.
-	 * It is usually a resource handle of sorts and is owned by the
-	 * platform's resource manager.
-	 *
-	 * USAGE TBD
-	 */
-	virtual IElectraDecoderVideoOutputTransferHandle* GetTransferHandle() const = 0;
 };
diff --git a/Engine/Plugins/Media/ElectraCodecs/Source/ElectraDecoders/Public/IElectraDecoderResourceDelegateBase.h b/Engine/Plugins/Media/ElectraCodecs/Source/ElectraDecoders/Public/IElectraDecoderResourceDelegateBase.h
--- a/Engine/Plugins/Media/ElectraCodecs/Source/ElectraDecoders/Public/IElectraDecoderResourceDelegateBase.h
+++ b/Engine/Plugins/Media/ElectraCodecs/Source/ElectraDecoders/Public/IElectraDecoderResourceDelegateBase.h
@@ -28,22 +28,6 @@ public:
 		Video
 	};
 
-	/**
-	* Interface class to abstract async job synchronization primitives on platforms implementing them
-	*/
-	class IAsyncConsecutiveTaskSync
-	{
-	public:
-		virtual ~IAsyncConsecutiveTaskSync() = default;
-	protected:
-		IAsyncConsecutiveTaskSync() = default;
-	};
-
-	virtual TSharedPtr<IAsyncConsecutiveTaskSync, ESPMode::ThreadSafe> CreateAsyncConsecutiveTaskSync()
-	{ return nullptr; }
-	virtual bool RunCodeAsync(TFunction<void()>&& CodeToRun, IAsyncConsecutiveTaskSync* TaskSync = nullptr)
-	{ return false; }
-
 	virtual IDecoderPlatformResource* CreatePlatformResource(void* InOwnerHandle, EDecoderPlatformResourceType InDecoderResourceType, const TMap<FString, FVariant> InOptions)
 	{ return nullptr; }
 	virtual void ReleasePlatformResource(void* InOwnerHandle, IDecoderPlatformResource* InHandleToDestroy)
diff --git a/Engine/Plugins/Media/ElectraCodecs/Source/ElectraDecoders/Public/Windows/WindowsElectraDecoderPlatformOutputHandleTypes.h b/Engine/Plugins/Media/ElectraCodecs/Source/ElectraDecoders/Public/Windows/WindowsElectraDecoderPlatformOutputHandleTypes.h
--- a/Engine/Plugins/Media/ElectraCodecs/Source/ElectraDecoders/Public/Windows/WindowsElectraDecoderPlatformOutputHandleTypes.h
+++ b/Engine/Plugins/Media/ElectraCodecs/Source/ElectraDecoders/Public/Windows/WindowsElectraDecoderPlatformOutputHandleTypes.h
@@ -3,7 +3,7 @@
 #pragma once
 
 #include "CoreMinimal.h"
-#include "IElectraDecoderResourceDelegateBase.h"
+#include "ElectraDecodersPlatformResources.h"
 
 #include "Windows/AllowWindowsPlatformTypes.h"
 THIRD_PARTY_INCLUDES_START
@@ -55,7 +55,7 @@ class FElectraDecoderOutputSync final
 {
 public:
 	FElectraDecoderOutputSync() : SyncValue(0), CopyDoneSyncValue(0) {}
-	FElectraDecoderOutputSync(TRefCountPtr<IUnknown> InSync, uint64 InSyncValue, TRefCountPtr<IMFSample> InMFSample = nullptr, TSharedPtr<IElectraDecoderResourceDelegateBase::IAsyncConsecutiveTaskSync> InTaskSync = nullptr)
+	FElectraDecoderOutputSync(TRefCountPtr<IUnknown> InSync, uint64 InSyncValue, TRefCountPtr<IMFSample> InMFSample = nullptr, TSharedPtr<FElectraDecodersPlatformResources::FAsyncConsecutiveTaskSync, ESPMode::ThreadSafe> InTaskSync = nullptr)
 		: Sync(InSync), SyncValue(InSyncValue), TaskSync(InTaskSync), MFSample(InMFSample)
 	{}
 	FElectraDecoderOutputSync(TRefCountPtr<IUnknown> InSync, uint64 InSyncValue, TRefCountPtr<ID3D12Fence> InCopyDoneSync, uint64 InCopyDoneSyncValue)
@@ -64,7 +64,7 @@ public:
 
 	TRefCountPtr<IUnknown>		Sync;														// Optional IUnknown based sync primitive (fence, or DX12 specific WMF object)
 	uint64						SyncValue;													// Optionally needed sync value
-	TSharedPtr<IElectraDecoderResourceDelegateBase::IAsyncConsecutiveTaskSync> TaskSync;	// Optional Sync/Perquisite object to use to run code async after the decoder data arrives
+	TSharedPtr<FElectraDecodersPlatformResources::FAsyncConsecutiveTaskSync, ESPMode::ThreadSafe> TaskSync;	// Optional Sync/Perquisite object to use to run code async after the decoder data arrives
 
 	TRefCountPtr<ID3D12Fence>	CopyDoneSync;												// Optional fence to signal end of copy from this buffer back to decoder
 	uint64						CopyDoneSyncValue;											// Value to be used with CopyDoneSync fence
diff --git a/Engine/Plugins/Media/ElectraCodecs/Source/ElectraDecoders/Public/Windows/WindowsElectraDecoderResourceDelegate.h b/Engine/Plugins/Media/ElectraCodecs/Source/ElectraDecoders/Public/Windows/WindowsElectraDecoderResourceDelegate.h
--- a/Engine/Plugins/Media/ElectraCodecs/Source/ElectraDecoders/Public/Windows/WindowsElectraDecoderResourceDelegate.h
+++ b/Engine/Plugins/Media/ElectraCodecs/Source/ElectraDecoders/Public/Windows/WindowsElectraDecoderResourceDelegate.h
@@ -9,9 +9,6 @@ class IElectraDecoderResourceDelegateWindows : public IElectraDecoderResourceDel
 {
 public:
 	virtual ~IElectraDecoderResourceDelegateWindows() = default;
-
-	// DirectX video decoding needs to know the current D3D device and the version of DX being used.
-	virtual bool GetD3DDevice(void **OutD3DDevice, int32* OutD3DVersionTimes1000) = 0;
 };
 
 
diff --git a/Engine/Plugins/Media/ElectraPlayer/Source/ElectraPlayerPlugin/ElectraPlayerPlugin.Build.cs b/Engine/Plugins/Media/ElectraPlayer/Source/ElectraPlayerPlugin/ElectraPlayerPlugin.Build.cs
--- a/Engine/Plugins/Media/ElectraPlayer/Source/ElectraPlayerPlugin/ElectraPlayerPlugin.Build.cs
+++ b/Engine/Plugins/Media/ElectraPlayer/Source/ElectraPlayerPlugin/ElectraPlayerPlugin.Build.cs
@@ -18,40 +18,19 @@ namespace UnrealBuildTool.Rules
 			PrivateDependencyModuleNames.AddRange(
 				new string[] {
 					"Core",
-					"MediaUtils",
-					"RenderCore",
 					"RHI",
-					"ElectraPlayerRuntime",
-					"ElectraSamples",
-					"ElectraBase",
-					"ElectraDecoders"
+					"ElectraPlayerRuntime"
 				});
 
 			PrivateIncludePathModuleNames.AddRange(
-			new string[] {
-				"Media",
+			    new string[] {
+				    "Media",
 			});
 
 			if (Target.bCompileAgainstEngine)
 			{
 				PrivateDependencyModuleNames.Add("Engine");
 			}
-
-			if (Target.Platform.IsInGroup(UnrealPlatformGroup.Windows))
-			{
-				PrivateDependencyModuleNames.Add("D3D12RHI");
-				AddEngineThirdPartyPrivateStaticDependencies(Target, "DX12");
-			}
-
-			if (Target.IsInPlatformGroup(UnrealPlatformGroup.Apple))
-			{
-				PrivateDependencyModuleNames.Add("MetalRHI");
-			}
-
-			if (Target.Type == TargetType.Editor)
-			{
-				PrivateDependencyModuleNames.Add("DeveloperSettings");
-			}
 		}
 	}
 }
diff --git a/Engine/Plugins/Media/ElectraPlayer/Source/ElectraPlayerPlugin/Private/ElectraPlayerPluginModule.cpp b/Engine/Plugins/Media/ElectraPlayer/Source/ElectraPlayerPlugin/Private/ElectraPlayerPluginModule.cpp
--- a/Engine/Plugins/Media/ElectraPlayer/Source/ElectraPlayerPlugin/Private/ElectraPlayerPluginModule.cpp
+++ b/Engine/Plugins/Media/ElectraPlayer/Source/ElectraPlayerPlugin/Private/ElectraPlayerPluginModule.cpp
@@ -1,46 +1,36 @@
 // Copyright Epic Games, Inc. All Rights Reserved.
 
-#include "CoreMinimal.h"
-
+#include "CoreTypes.h"
+#include "Templates/SharedPointer.h"
 #include "Modules/ModuleManager.h"
+#include "Logging/LogMacros.h"
 #include "RHI.h"
 
-#include "ElectraPlayerMisc.h"
-#include "IElectraPlayerPluginModule.h"
+
 #include "IElectraPlayerRuntimeModule.h"
-#include "ElectraPlayerPlugin.h"
-#include "ParameterDictionary.h"
-#include "IElectraDecoderResourceDelegateBase.h"
-#include "IElectraPlayerDecoderResourceManager.h"
+#include "IElectraPlayerPluginModule.h"			// for loading the Electra player runtime
+#include "IElectraPlayerInterface.h"			// for metrics delegate declaration
 
-#define LOCTEXT_NAMESPACE "ElectraPlayerPluginModule"
 
+#define LOCTEXT_NAMESPACE "ElectraPlayerPluginModule"
+DECLARE_LOG_CATEGORY_EXTERN(LogElectraPlayerPlugin, Log, All);
 DEFINE_LOG_CATEGORY(LogElectraPlayerPlugin);
 
-DECLARE_CYCLE_STAT(TEXT("Electra AsyncJob"), STAT_ElectraAsyncJob, STATGROUP_Media);
-
 // -----------------------------------------------------------------------------------------------------------------------------------
 
 class FElectraPlayerPluginModule : public IElectraPlayerPluginModule
 {
 public:
-	// IElectraPlayerPluginModule interface
-
 	bool IsInitialized() const override
 	{
 		return bInitialized;
 	}
 
-	TSharedPtr<IMediaPlayer, ESPMode::ThreadSafe> CreatePlayer(IMediaEventSink& EventSink) override
+	TSharedPtr<IMediaPlayer, ESPMode::ThreadSafe> CreatePlayer(IMediaEventSink& InEventSink) override
 	{
 		if (bInitialized)
 		{
-			TSharedPtr<FElectraPlayerPlugin, ESPMode::ThreadSafe> NewPlayer = MakeShared<FElectraPlayerPlugin, ESPMode::ThreadSafe>();
-			check(NewPlayer.IsValid());
-			if (NewPlayer->Initialize(EventSink, SendAnalyticMetricsDelegate, SendAnalyticMetricsPerMinuteDelegate, ReportVideoStreamingErrorDelegate, ReportSubtitlesMetricsDelegate))
-			{
-				return NewPlayer;
-			}
+			return FElectraPlayerRuntimeFactory::CreatePlayer(InEventSink, SendAnalyticMetricsDelegate, SendAnalyticMetricsPerMinuteDelegate, ReportVideoStreamingErrorDelegate, ReportSubtitlesMetricsDelegate);
 		}
 		return nullptr;
 	}
@@ -65,120 +55,24 @@ public:
 		ReportSubtitlesMetricsDelegate.Broadcast(PlayerGuid, URL, ResponseTime, LastError);
 	}
 
-	// Create a suitable video decoder resource delegate for and via the Electra Player runtime to be used with it by external means.
-	TSharedPtr<Electra::IVideoDecoderResourceDelegate, ESPMode::ThreadSafe> CreatePlatformVideoDecoderResourceDelegate() override
-	{
-	#if PLATFORM_WINDOWS
-		// This is a bit of a convoluted process due to this whole plugin/adapter/delegate nature.
-		class FDummyAdapter : public IElectraPlayerAdapterDelegate
-		{
-		public:
-			FDummyAdapter(TSharedPtr<IElectraPlayerResourceDelegate, ESPMode::ThreadSafe> InPlayerResourceDlg) : PlayerResourceDlg(MoveTemp(InPlayerResourceDlg)) { }
-			virtual ~FDummyAdapter() { }
-			Electra::FVariantValue QueryOptions(EOptionType Type, const Electra::FVariantValue& Param = Electra::FVariantValue()) override { return Electra::FVariantValue(); }
-			void BlobReceived(const TSharedPtr<TArray<uint8>, ESPMode::ThreadSafe>& InBlobData, EBlobResultType InResultType, int32 InResultCode, const Electra::FParamDict* InExtraInfo) override { }
-			void SendMediaEvent(EPlayerEvent Event) override { }
-			void OnVideoFlush() override { }
-			void OnAudioFlush() override { }
-			void OnSubtitleFlush() override { }
-			void PresentVideoFrame(const FVideoDecoderOutputPtr& InVideoFrame) override { }
-			void PresentAudioFrame(const IAudioDecoderOutputPtr& InAudioFrame) override { }
-			void PresentSubtitleSample(const ISubtitleDecoderOutputPtr& InSubtitleSample) override { }
-			void PresentMetadataSample(const IMetaDataDecoderOutputPtr& InMetadataSample) override { }
-			bool CanReceiveVideoSamples(int32 NumFrames) override { return false; }
-			bool CanReceiveAudioSamples(int32 NumFrames) override { return false; }
-			FString GetVideoAdapterName() const override { return FString(); }
-			TSharedPtr<IElectraPlayerResourceDelegate, ESPMode::ThreadSafe> GetResourceDelegate() const override { return PlayerResourceDlg; }
-		private:
-			TSharedPtr<IElectraPlayerResourceDelegate, ESPMode::ThreadSafe> PlayerResourceDlg;
-		};
-		return FElectraPlayerDecoderResourceManager::CreatePlatformVideoDecoderResourceDelegate(MakeShareable(new FDummyAdapter(MakeShareable(FElectraPlayerPlugin::PlatformCreateStaticPlayerResourceDelegate()))));
-	#else
-		return nullptr;
-	#endif
-	}
-
-public:
-	// IModuleInterface interface
-	static void GetDynamicRHIInfo(void** OutGDynamicRHI, int64* OutGDynamicRHIType)
-	{
-		if (OutGDynamicRHI)
-		{
-			*OutGDynamicRHI = GDynamicRHI ? GDynamicRHI->RHIGetNativeDevice() : nullptr;
-		}
-		if (OutGDynamicRHIType)
-		{
-			*OutGDynamicRHIType = (int64)RHIGetInterfaceType();
-		}
-	}
-
-	class FAsyncConsecutiveTaskSync : public IElectraDecoderResourceDelegateBase::IAsyncConsecutiveTaskSync
-	{
-	public:
-		FAsyncConsecutiveTaskSync() {}
-		virtual ~FAsyncConsecutiveTaskSync() {}
-
-		FGraphEventRef	GraphEvent;
-	};
-
-	static TSharedPtr<IElectraDecoderResourceDelegateBase::IAsyncConsecutiveTaskSync, ESPMode::ThreadSafe> CreateAsyncConsecutiveTaskSync()
-	{
-		return MakeShared<FAsyncConsecutiveTaskSync, ESPMode::ThreadSafe>();
-	}
-
-	static void RunCodeAsync(TFunction<void()>&& CodeToRun, IElectraDecoderResourceDelegateBase::IAsyncConsecutiveTaskSync* TaskSync)
-	{
-		FScopeLock Lock(&AsyncJobAccessCS);
-
-		// Execute code async
-		// (We assume this to be code copying buffer data around. Hence we allow only ONE at any time to not clog up
-		//  the buses even more and delay the copy process further)
-		FGraphEventArray Events;
-		auto& Event = TaskSync ? static_cast<FAsyncConsecutiveTaskSync*>(TaskSync)->GraphEvent : RunCodeAsyncEvent;
-		if (Event.IsValid())
-		{
-			Events.Add(Event);
-		}
-		Event = FFunctionGraphTask::CreateAndDispatchWhenReady(MoveTemp(CodeToRun), GET_STATID(STAT_ElectraAsyncJob), &Events);
-	}
-
 	void StartupModule() override
 	{
 		// Check that we have the player module and that it has initialized successfully.
 		if (FModuleManager::Get().GetModule("ElectraPlayerRuntime"))
 		{
 			IElectraPlayerRuntimeModule* ElectraPlayer = &FModuleManager::Get().GetModuleChecked<IElectraPlayerRuntimeModule>("ElectraPlayerRuntime");
-			if (!ElectraPlayer || !ElectraPlayer->IsInitialized())
-			{
-				return;
-			}
-		}
-		else
-		{
-			return;
-		}
-
-		if (!bInitialized)
-		{
-//WE COULD ALSO CHECK FOR SINGLE THREADING ISSUES HERE!? -- kinda piggy back on this?
-			// to detect cooking and other commandlets that run with NullRHI
-			if (GDynamicRHI == nullptr || RHIGetInterfaceType() == ERHIInterfaceType::Null)
-			{
-				UE_LOG(LogElectraPlayerPlugin, Log, TEXT("Dummy Dynamic RHI detected. Electra Player plugin is not initialised."));
-				return;
-			}
-
-			Electra::FParamDict Params;
-			Params.Set(FName(TEXT("GetDeviceTypeCallback")), Electra::FVariantValue((void*)&FElectraPlayerPluginModule::GetDynamicRHIInfo));
-			Params.Set(FName(TEXT("CreateAsyncConsecutiveTaskSync")), Electra::FVariantValue((void*)&FElectraPlayerPluginModule::CreateAsyncConsecutiveTaskSync));
-			Params.Set(FName(TEXT("RunCodeAsyncCallback")), Electra::FVariantValue((void*)&FElectraPlayerPluginModule::RunCodeAsync));
-			if (!FElectraPlayerPlatform::StartupPlatformResources(Params))
+			if (!bInitialized && ElectraPlayer && ElectraPlayer->IsInitialized())
 			{
-				UE_LOG(LogElectraPlayerPlugin, Log, TEXT("Platform resource setup failed! Electra Player plugin is not initialised."));
-				return;
+				// Detect cooking and other commandlets that run with NullRHI
+				if (GDynamicRHI && RHIGetInterfaceType() != ERHIInterfaceType::Null)
+				{
+					bInitialized = true;
+				}
+				else
+				{
+					UE_LOG(LogElectraPlayerPlugin, Log, TEXT("Dummy Dynamic RHI detected. Electra Player plugin is not initialized."));
+				}
 			}
-
-			bInitialized = true;
 		}
 	}
 
@@ -188,22 +82,13 @@ public:
 	}
 
 private:
+	FElectraPlayerSendAnalyticMetricsDelegate SendAnalyticMetricsDelegate;
+	FElectraPlayerSendAnalyticMetricsPerMinuteDelegate SendAnalyticMetricsPerMinuteDelegate;
+	FElectraPlayerReportVideoStreamingErrorDelegate ReportVideoStreamingErrorDelegate;
+	FElectraPlayerReportSubtitlesMetricsDelegate ReportSubtitlesMetricsDelegate;
 	bool bInitialized = false;
-
-	FElectraPlayerSendAnalyticMetricsDelegate			SendAnalyticMetricsDelegate;
-	FElectraPlayerSendAnalyticMetricsPerMinuteDelegate	SendAnalyticMetricsPerMinuteDelegate;
-	FElectraPlayerReportVideoStreamingErrorDelegate		ReportVideoStreamingErrorDelegate;
-	FElectraPlayerReportSubtitlesMetricsDelegate		ReportSubtitlesMetricsDelegate;
-
-	static FCriticalSection AsyncJobAccessCS;
-	static FGraphEventRef RunCodeAsyncEvent;
 };
 
-FCriticalSection FElectraPlayerPluginModule::AsyncJobAccessCS;
-FGraphEventRef FElectraPlayerPluginModule::RunCodeAsyncEvent;
-
 IMPLEMENT_MODULE(FElectraPlayerPluginModule, ElectraPlayerPlugin);
 
 #undef LOCTEXT_NAMESPACE
-
-
diff --git a/Engine/Plugins/Media/ElectraPlayer/Source/ElectraPlayerPlugin/Public/IElectraPlayerPluginModule.h b/Engine/Plugins/Media/ElectraPlayer/Source/ElectraPlayerPlugin/Public/IElectraPlayerPluginModule.h
--- a/Engine/Plugins/Media/ElectraPlayer/Source/ElectraPlayerPlugin/Public/IElectraPlayerPluginModule.h
+++ b/Engine/Plugins/Media/ElectraPlayer/Source/ElectraPlayerPlugin/Public/IElectraPlayerPluginModule.h
@@ -4,127 +4,11 @@
 
 #include "Templates/SharedPointer.h"
 #include "Modules/ModuleInterface.h"
-#include "Misc/ScopeLock.h"
-#include "Misc/Timespan.h"
-#include "Containers/Array.h"
-#include "IMediaOptions.h"
-#include "UObject/Object.h"
-#include "UObject/StrongObjectPtrTemplates.h"
-#include "UObject/WeakObjectPtr.h"
 
 class IAnalyticsProviderET;
 class IMediaEventSink;
 class IMediaPlayer;
-class IElectraPlayerDataCache;
-class UObject;
-namespace Electra
-{
-class IVideoDecoderResourceDelegate;
-}
-
-/**
- * This class is used to get safe access to an IMediaOptions interface.
- * When passing IMediaOptions through media framework as a POD pointer there is the risk
- * that it is actually from a derived class like UMediaSource which is subject to GC.
- * Thus it is not safe to keep the POD IMediaOptions around.
- * This class is intended to be derived and instantiated from and stored as a TSharedPtr
- * in a derived UMediaSource class.
- * Then, as the media player is opened and the POD IMediaOptions is passed along, a
- * GetMediaOption("GetSafeMediaOptions") query will be made on it to get this instance.
- * If future access to the IMediaOptions is needed it will be made through this class
- * by first locking, getting and using the IMediaOptions pointer returned here if it
- * not null and unlocking afterwards.
- */
-class IElectraSafeMediaOptionInterface : public IMediaOptions::FDataContainer
-{
-public:
-	virtual ~IElectraSafeMediaOptionInterface() = default;
-	virtual void Lock() = 0;
-	virtual void Unlock() = 0;
-	virtual IMediaOptions* GetMediaOptionInterface() = 0;
-
-	class FScopedLock
-	{
-	public:
-		FScopedLock(const TSharedPtr<IElectraSafeMediaOptionInterface, ESPMode::ThreadSafe>& InSafeMediaOptionInterface)
-			: SafeMediaOptionInterface(InSafeMediaOptionInterface)
-		{
-			if (SafeMediaOptionInterface.IsValid())
-			{
-				SafeMediaOptionInterface->Lock();
-			}
-		}
-		~FScopedLock()
-		{
-			if (SafeMediaOptionInterface.IsValid())
-			{
-				SafeMediaOptionInterface->Unlock();
-			}
-		}
-
-	private:
-		TSharedPtr<IElectraSafeMediaOptionInterface, ESPMode::ThreadSafe> SafeMediaOptionInterface;
-	};
-};
-
-
-/**
- * A ready to use implementation of the IElectraSafeMediaOptionInterface
- */
-class FElectraSafeMediaOptionInterface : public IElectraSafeMediaOptionInterface
-{
-public:
-	FElectraSafeMediaOptionInterface(IMediaOptions* InOwner)
-		: Owner(InOwner)
-		, OwnerObject(InOwner ? InOwner->ToUObject() : nullptr)
-	{ }
-	virtual ~FElectraSafeMediaOptionInterface()
-	{
-		ClearOwner();
-	}
-	void ClearOwner()
-	{
-		FScopeLock lock(&OwnerLock);
-		Owner = nullptr;
-		OwnerObject.Reset();
-	}
-	virtual void Lock() override
-	{
-		OwnerLock.Lock();
-	}
-	virtual void Unlock() override
-	{
-		OwnerLock.Unlock();
-	}
-	virtual IMediaOptions* GetMediaOptionInterface() override
-	{
-		return OwnerObject.IsStale(true, true) ? nullptr : Owner;
-	}
-private:
-	FCriticalSection OwnerLock;
-	IMediaOptions* Owner = nullptr;
-	FWeakObjectPtr OwnerObject;
-};
-
-
-
-//! Data type for use with media options interface
-class FElectraSeekablePositions : public IMediaOptions::FDataContainer
-{
-public:
-	FElectraSeekablePositions(const TArray<FTimespan>& InData) : Data(InData) {}
-	virtual ~FElectraSeekablePositions() {}
 
-	TArray<FTimespan> Data;
-};
-
-class FElectraPlayerDataCacheContainer : public IMediaOptions::FDataContainer
-{
-public:
-	FElectraPlayerDataCacheContainer(TSharedPtr<IElectraPlayerDataCache, ESPMode::ThreadSafe> InData) : Data(InData) {}
-	virtual ~FElectraPlayerDataCacheContainer() {}
-	TSharedPtr<IElectraPlayerDataCache, ESPMode::ThreadSafe> Data;
-};
 
 /**
  * Interface for the ElectraPlayerPlugin module.
@@ -133,6 +17,7 @@ class IElectraPlayerPluginModule
 	: public IModuleInterface
 {
 public:
+	virtual ~IElectraPlayerPluginModule() = default;
 
 	/**
 	 * Is the ElectraPlayerPlugin module initialized?
@@ -153,9 +38,4 @@ public:
 
 	virtual void ReportVideoStreamingError(const FGuid& PlayerGuid, const FString& LastError) = 0;
 	virtual void ReportSubtitlesMetrics(const FGuid& PlayerGuid, const FString& URL, double ResponseTime, const FString& LastError) = 0;
-
-	// Create a suitable video decoder resource delegate for and via the Electra Player runtime to be used with it by external means.
-	virtual TSharedPtr<Electra::IVideoDecoderResourceDelegate, ESPMode::ThreadSafe> CreatePlatformVideoDecoderResourceDelegate() = 0;
-public:
-	virtual ~IElectraPlayerPluginModule() { }
 };
diff --git a/Engine/Plugins/Media/ElectraPlayer/Source/ElectraPlayerRuntime/ElectraPlayerRuntime.Build.cs b/Engine/Plugins/Media/ElectraPlayer/Source/ElectraPlayerRuntime/ElectraPlayerRuntime.Build.cs
--- a/Engine/Plugins/Media/ElectraPlayer/Source/ElectraPlayerRuntime/ElectraPlayerRuntime.Build.cs
+++ b/Engine/Plugins/Media/ElectraPlayer/Source/ElectraPlayerRuntime/ElectraPlayerRuntime.Build.cs
@@ -18,8 +18,13 @@ namespace UnrealBuildTool.Rules
 			PrivateDependencyModuleNames.AddRange(
 				new string[] {
 					"Core",
+                    "RHI",
+					"RenderCore",
 					"Json",
+					"MediaUtils",
+                    "Media",
 					"ElectraBase",
+					"ElectraSamples",
 					"ElectraCodecFactory",
 					"ElectraDecoders",
 					"ElectraHTTPStream",
@@ -74,10 +79,10 @@ namespace UnrealBuildTool.Rules
 
 				PublicIncludePaths.Add("$(ModuleDir)/Public/Windows");
 				PrivateIncludePaths.Add("ElectraPlayerRuntime/Private/Runtime/Decoder/Windows");
-				PrivateIncludePaths.Add("ElectraPlayerRuntime/Private/Windows");
 			}
 			else if (Target.IsInPlatformGroup(UnrealPlatformGroup.Apple))
 			{
+				PrivateDependencyModuleNames.Add("MetalRHI");
 				PublicFrameworks.AddRange(
 				new string[] {
 								"CoreMedia",
