Electra: Removal of the platform resource delegate and the wrapping plugin

diff --git a/Engine/Plugins/Media/ElectraPlayer/Source/ElectraPlayerRuntime/Private/ElectraPlayer.h b/Engine/Plugins/Media/ElectraPlayer/Source/ElectraPlayerRuntime/Private/ElectraPlayer.h
--- a/Engine/Plugins/Media/ElectraPlayer/Source/ElectraPlayerRuntime/Private/ElectraPlayer.h
+++ b/Engine/Plugins/Media/ElectraPlayer/Source/ElectraPlayerRuntime/Private/ElectraPlayer.h
@@ -2,20 +2,41 @@
 
 #pragma once
 
+#include "IElectraPlayerInterface.h"
 
 #include "Containers/UnrealString.h"
 #include "Containers/Queue.h"
 #include "Containers/SortedMap.h"
-#include "Misc/Guid.h"
 #include "Templates/SharedPointer.h"
 #include "Templates/Greater.h"
+#include "Logging/LogMacros.h"
+#include "Misc/Optional.h"
+#include "Misc/Guid.h"
+#include "Math/Range.h"
+
+#include "IMediaPlayer.h"
+#include "IMediaCache.h"
+#include "IMediaView.h"
+#include "IMediaControls.h"
+#include "IMediaTracks.h"
+#include "IMediaEventSink.h"
+#include "IMediaPlayerLifecycleManager.h"
+#include "MediaSamples.h"
+#include "MediaSampleQueue.h"
+
 #include "IAnalyticsProviderET.h"
 
-#include "IElectraPlayerInterface.h"
+#include "ElectraTextureSample.h"
+#include "ElectraPlayerAudioSample.h"
 
 #include "Player/AdaptiveStreamingPlayer.h"
 #include "MediaStreamMetadata.h"
 #include "PlayerRuntimeGlobal.h"
+#include "StreamTypes.h"
+
+#include <atomic>
+
+
 
 class FVideoDecoderOutput;
 using FVideoDecoderOutputPtr = TSharedPtr<FVideoDecoderOutput, ESPMode::ThreadSafe>;
@@ -26,117 +47,209 @@ using IMetaDataDecoderOutputPtr = TSharedPtr<IMetaDataDecoderOutput, ESPMode::Th
 class ISubtitleDecoderOutput;
 using ISubtitleDecoderOutputPtr = TSharedPtr<ISubtitleDecoderOutput, ESPMode::ThreadSafe>;
 
-namespace Electra
-{
-class IVideoDecoderResourceDelegate;
-}
-
 class FElectraRendererVideo;
 class FElectraRendererAudio;
 
 using namespace Electra;
 
-DECLARE_MULTICAST_DELEGATE_TwoParams(FElectraPlayerSendAnalyticMetricsDelegate, const TSharedPtr<IAnalyticsProviderET>& /*AnalyticsProvider*/, const FGuid& /*PlayerGuid*/);
-DECLARE_MULTICAST_DELEGATE_OneParam(FElectraPlayerSendAnalyticMetricsPerMinuteDelegate, const TSharedPtr<IAnalyticsProviderET>& /*AnalyticsProvider*/);
-DECLARE_MULTICAST_DELEGATE_TwoParams(FElectraPlayerReportVideoStreamingErrorDelegate, const FGuid& /*PlayerGuid*/, const FString& /*LastError*/);
-DECLARE_MULTICAST_DELEGATE_FourParams(FElectraPlayerReportSubtitlesMetricsDelegate, const FGuid& /*PlayerGuid*/, const FString& /*URL*/, double /*ResponseTime*/, const FString& /*LastError*/);
 
 class FElectraPlayer
-	: public IElectraPlayerInterface
+	: public IMediaPlayer
+	, public IMediaCache
+	, public IMediaControls
+	, public IMediaTracks
+	, public IMediaView
 	, public IAdaptiveStreamingPlayerMetrics
+	, public TSharedFromThis<FElectraPlayer, ESPMode::ThreadSafe>
 {
 public:
-	FElectraPlayer(const TSharedPtr<IElectraPlayerAdapterDelegate, ESPMode::ThreadSafe>& AdapterDelegate,
+	FElectraPlayer(IMediaEventSink& InEventSink,
 			  FElectraPlayerSendAnalyticMetricsDelegate& InSendAnalyticMetricsDelegate,
 			  FElectraPlayerSendAnalyticMetricsPerMinuteDelegate& InSendAnalyticMetricsPerMinuteDelegate,
 			  FElectraPlayerReportVideoStreamingErrorDelegate& InReportVideoStreamingErrorDelegate,
 			  FElectraPlayerReportSubtitlesMetricsDelegate& InReportSubtitlesFileMetricsDelegate);
-	~FElectraPlayer();
-
-	void OnVideoDecoded(const FVideoDecoderOutputPtr& DecoderOutput, bool bDoNotRender);
+	virtual ~FElectraPlayer();
+
+	//-------------------------------------------------------------------------
+	// Methods from IMediaPlayer
+	//
+	void Close() override;
+	IMediaCache& GetCache() override
+	{ return *this; }
+	IMediaControls& GetControls() override
+	{ return *this; }
+	FString GetInfo() const override;
+	FGuid GetPlayerPluginGUID() const override;
+	IMediaSamples& GetSamples() override;
+	FString GetStats() const override;
+	IMediaTracks& GetTracks() override
+	{ return *this; }
+	FString GetUrl() const override;
+	IMediaView& GetView() override
+	{ return *this; }
+	bool Open(const FString& InUrl, const IMediaOptions* InOptions) override;
+	bool Open(const TSharedRef<FArchive, ESPMode::ThreadSafe>& InArchive, const FString& InOriginalUrl, const IMediaOptions* InOptions) override;
+	bool Open(const FString& InUrl, const IMediaOptions* InOptions, const FMediaPlayerOptions* InPlayerOptions) override;
+	FVariant GetMediaInfo(FName InInfoName) const override;
+	TSharedPtr<TMap<FString, TArray<TUniquePtr<IMediaMetadataItem>>>, ESPMode::ThreadSafe> GetMediaMetadata() const override;
+	void SetGuid(const FGuid& InGuid) override;
+	void TickInput(FTimespan InDeltaTime, FTimespan InTimecode) override;
+	bool FlushOnSeekStarted() const override;
+	bool FlushOnSeekCompleted() const override;
+	bool GetPlayerFeatureFlag(EFeatureFlag InFlag) const override;
+	bool SetAsyncResourceReleaseNotification(IAsyncResourceReleaseNotificationRef InAsyncDestructNotification) override;
+	uint32 GetNewResourcesOnOpen() const override;
+
+	//-------------------------------------------------------------------------
+	// Methods from IMediaCache
+	//
+	bool QueryCacheState(EMediaCacheState InState, TRangeSet<FTimespan>& OutTimeRanges) const override;
+
+	//-------------------------------------------------------------------------
+	// Methods from IMediaControls
+	//
+	bool CanControl(EMediaControl InControl) const override;
+	FTimespan GetDuration() const override;
+	float GetRate() const override;
+	EMediaState GetState() const override;
+	EMediaStatus GetStatus() const override;
+	TRangeSet<float> GetSupportedRates(EMediaRateThinning InThinning) const override;
+	FTimespan GetTime() const override;
+	bool IsLooping() const override;
+	bool Seek(const FTimespan& InTime) override
+	{ check(!"You have to call the override with additional options!"); return false; }
+	bool SetLooping(bool bInLooping) override;
+	bool SetRate(float InRate) override;
+	bool Seek(const FTimespan& InNewTime, const FMediaSeekParams& InAdditionalParams) override;
+	TRange<FTimespan> GetPlaybackTimeRange(EMediaTimeRangeType InRangeToGet) const override;
+	bool SetPlaybackTimeRange(const TRange<FTimespan>& InTimeRange) override;
+
+	//-------------------------------------------------------------------------
+	// Methods from IMediaTracks
+	//
+	bool GetAudioTrackFormat(int32 InTrackIndex, int32 InFormatIndex, FMediaAudioTrackFormat& OutFormat) const override;
+	int32 GetNumTracks(EMediaTrackType InTrackType) const override;
+	int32 GetNumTrackFormats(EMediaTrackType InTrackType, int32 InTrackIndex) const override;
+	int32 GetSelectedTrack(EMediaTrackType InTrackType) const override;
+	FText GetTrackDisplayName(EMediaTrackType InTrackType, int32 InTrackIndex) const override;
+	int32 GetTrackFormat(EMediaTrackType InTrackType, int32 InTrackIndex) const override;
+	FString GetTrackLanguage(EMediaTrackType InTrackType, int32 InTrackIndex) const override;
+	FString GetTrackName(EMediaTrackType InTrackType, int32 InTrackIndex) const override;
+	bool GetVideoTrackFormat(int32 InTrackIndex, int32 InFormatIndex, FMediaVideoTrackFormat& OutFormat) const override;
+	bool SelectTrack(EMediaTrackType InTrackType, int32 InTrackIndex) override;
+	bool SetTrackFormat(EMediaTrackType InTrackType, int32 InTrackIndex, int32 InFormatIndex) override;
+	bool SetVideoTrackFrameRate(int32 InTrackIndex, int32 InFormatIndex, float InFrameRate) override;
+
+	//-------------------------------------------------------------------------
+	// Methods from IAdaptiveStreamingPlayerMetrics
+	//
+	void ReportOpenSource(const FString& InURL) override;
+	void ReportReceivedMainPlaylist(const FString& InEffectiveURL) override;
+	void ReportReceivedPlaylists() override;
+	void ReportTracksChanged() override;
+	void ReportPlaylistDownload(const Metrics::FPlaylistDownloadStats& InPlaylistDownloadStats) override;
+	void ReportCleanStart() override;
+	void ReportBufferingStart(Metrics::EBufferingReason InBufferingReason) override;
+	void ReportBufferingEnd(Metrics::EBufferingReason InBufferingReason) override;
+	void ReportBandwidth(int64 InEffectiveBps, int64 InThroughputBps, double InLatencyInSeconds) override;
+	void ReportBufferUtilization(const Metrics::FBufferStats& InBufferStats) override;
+	void ReportSegmentDownload(const Metrics::FSegmentDownloadStats& InSegmentDownloadStats) override;
+	void ReportLicenseKey(const Metrics::FLicenseKeyStats& InLicenseKeyStats) override;
+	void ReportDataAvailabilityChange(const Metrics::FDataAvailabilityChange& InDataAvailability) override;
+	void ReportVideoQualityChange(int32 InNewBitrate, int32 InPreviousBitrate, bool bInIsDrasticDownswitch) override;
+	void ReportAudioQualityChange(int32 InNewBitrate, int32 InPreviousBitrate, bool bInIsDrasticDownswitch) override;
+	void ReportDecodingFormatChange(const FStreamCodecInformation& InNewDecodingFormat) override;
+	void ReportPrerollStart() override;
+	void ReportPrerollEnd() override;
+	void ReportPlaybackStart() override;
+	void ReportPlaybackPaused() override;
+	void ReportPlaybackResumed() override;
+	void ReportPlaybackEnded() override;
+	void ReportJumpInPlayPosition(const FTimeValue& InToNewTime, const FTimeValue& InFromTime, Metrics::ETimeJumpReason InTimejumpReason) override;
+	void ReportPlaybackStopped() override;
+	void ReportSeekCompleted() override;
+	void ReportMediaMetadataChanged(TSharedPtrTS<Electra::UtilsMP4::FMetadataParser> InMetadata) override;
+	void ReportError(const FString& InErrorReason) override;
+	void ReportLogMessage(IInfoLog::ELevel InLogLevel, const FString& InLogMessage, int64 InPlayerWallclockMilliseconds) override;
+	void ReportDroppedVideoFrame() override;
+	void ReportDroppedAudioFrame() override;
+
+
+	//-------------------------------------------------------------------------
+	// Implementation methods
+	//
+	void OnVideoDecoded(const FVideoDecoderOutputPtr& InDecoderOutput, bool bInDoNotRender);
 	void OnVideoFlush();
-	void OnAudioDecoded(const IAudioDecoderOutputPtr& DecoderOutput);
+	void OnAudioDecoded(const IAudioDecoderOutputPtr& InDecoderOutput);
 	void OnAudioFlush();
-	void OnSubtitleDecoded(ISubtitleDecoderOutputPtr DecoderOutput);
+	void OnSubtitleDecoded(ISubtitleDecoderOutputPtr InDecoderOutput);
 	void OnSubtitleFlush();
+	bool CanPresentVideoFrames(uint64 InNumFrames);
+	bool CanPresentAudioFrames(uint64 InNumFrames);
+	void PlatformSuspendOrResumeDecoders(bool bInSuspend, const Electra::FParamDict& InOptions);
 
-	void SendAnalyticMetrics(const TSharedPtr<IAnalyticsProviderET>& AnalyticsProvider, const FGuid& InPlayerGuid);
-	void SendAnalyticMetricsPerMinute(const TSharedPtr<IAnalyticsProviderET>& AnalyticsProvider);
-	void SendPendingAnalyticMetrics(const TSharedPtr<IAnalyticsProviderET>& AnalyticsProvider);
-	void ReportVideoStreamingError(const FGuid& InPlayerGuid, const FString& LastError);
-	void ReportSubtitlesMetrics(const FGuid& InPlayerGuid, const FString& URL, double ResponseTime, const FString& LastError);
-
-	void DropOldFramesFromPresentationQueue();
-
-	bool CanPresentVideoFrames(uint64 NumFrames);
-	bool CanPresentAudioFrames(uint64 NumFrames);
-
-	FString GetUrl() const override
-	{
-		return MediaUrl;
-	}
-
-	void SetGuid(const FGuid& Guid) override
-	{
-		PlayerGuid = Guid;
-	}
-
-	void SetAsyncResourceReleaseNotification(IAsyncResourceReleaseNotifyContainer* AsyncResourceReleaseNotification) override;
-
-	// -------- PlayerAdapter (Plugin/Native) API
-
-	bool OpenInternal(const FString& Url, const FParamDict& PlayerOptions, const FPlaystartOptions& InPlaystartOptions, EOpenType InOpenType) override;
-	void CloseInternal(bool bKillAfterClose) override;
-
-	void Tick(FTimespan DeltaTime, FTimespan Timecode) override;
-
-	bool IsKillAfterCloseAllowed() const override { return bAllowKillAfterCloseEvent;  }
-
-	EPlayerState GetState() const override;
-	EPlayerStatus GetStatus() const override;
-
-	bool IsLooping() const override;
-	bool SetLooping(bool bLooping) override;
-	int32 GetLoopCount() const override;
-
-	FTimespan GetTime() const override;
-	FTimespan GetDuration() const override;
-
-	bool IsLive() const override;
-	FTimespan GetSeekableDuration() const override;
-
-	void SetPlaybackRange(const FPlaybackRange& InPlaybackRange) override;
-	void GetPlaybackRange(FPlaybackRange& OutPlaybackRange) const override;
-	TRange<FTimespan> GetPlaybackRange(ETimeRangeType InRangeToGet) const override;
-
-	TRangeSet<float> GetSupportedRates(EPlayRateType InPlayRateType) const override;
-	float GetRate() const override;
-	bool SetRate(float Rate) override;
-
-	bool Seek(const FTimespan& Time, const FSeekParam& Param) override;
-	void SetFrameAccurateSeekMode(bool bEnableFrameAccuracy) override;
+	void SendAnalyticMetrics(const TSharedPtr<IAnalyticsProviderET>& InAnalyticsProvider, const FGuid& InPlayerGuid);
+	void SendAnalyticMetricsPerMinute(const TSharedPtr<IAnalyticsProviderET>& InAnalyticsProvider);
+	void SendPendingAnalyticMetrics(const TSharedPtr<IAnalyticsProviderET>& InAnalyticsProvider);
+	void ReportVideoStreamingError(const FGuid& InPlayerGuid, const FString& InLastError);
+	void ReportSubtitlesMetrics(const FGuid& InPlayerGuid, const FString& InURL, double InResponseTime, const FString& InLastError);
 
-	bool GetAudioTrackFormat(int32 TrackIndex, int32 FormatIndex, FAudioTrackFormat& OutFormat) const override;
-	bool GetVideoTrackFormat(int32 TrackIndex, int32 FormatIndex, FVideoTrackFormat& OutFormat) const override;
-
-	int32 GetNumTracks(EPlayerTrackType TrackType) const override;
-	int32 GetNumTrackFormats(EPlayerTrackType TrackType, int32 TrackIndex) const override;
-	int32 GetSelectedTrack(EPlayerTrackType TrackType) const override;
-	FText GetTrackDisplayName(EPlayerTrackType TrackType, int32 TrackIndex) const override;
-	int32 GetTrackFormat(EPlayerTrackType TrackType, int32 TrackIndex) const override;
-	FString GetTrackLanguage(EPlayerTrackType TrackType, int32 TrackIndex) const override;
-	FString GetTrackName(EPlayerTrackType TrackType, int32 TrackIndex) const override;
-	bool SelectTrack(EPlayerTrackType TrackType, int32 TrackIndex) override;
-
-	int32 GetNumVideoStreams(int32 TrackIndex) const override;
-	bool GetVideoStreamFormat(FVideoStreamFormat& OutFormat, int32 InTrackIndex, int32 InStreamIndex) const override;
-	bool GetActiveVideoStreamFormat(FVideoStreamFormat& OutFormat) const override;
-	Electra::FVariantValue GetMediaInfo(FName InInfoName) const override;
-	TSharedPtr<TMap<FString, TArray<TSharedPtr<Electra::IMediaStreamMetadata::IItem, ESPMode::ThreadSafe>>>, ESPMode::ThreadSafe> GetMediaMetadata() const override;
-
-	bool GetStreamBufferInformation(FStreamBufferInfo& OutBufferInformation, EPlayerTrackType InTrackType) const override;
+private:
+	struct FAnalyticsEvent;
+	void ClearToDefaultState();
+	void SendMediaSinkEvent(EMediaEvent InEventToSend);
+	bool OpenInternal(const FString& InUrl, const IMediaOptions* InOptions, const FMediaPlayerOptions* InPlayerOptions);
+	void CloseInternal();
+	void TickInternal(FTimespan InDeltaTime, FTimespan InTimecode);
+	void TriggerFirstSeekIfNecessary();
+	void CalculateTargetSeekTime(FTimespan& OutTargetTime, const FTimespan& InTime);
+	bool IsLive();
+	TSharedPtr<FTrackMetadata, ESPMode::ThreadSafe> GetTrackStreamMetadata(EMediaTrackType InTrackType, int32 InTrackIndex) const;
+	void AddCommonAnalyticsAttributes(TArray<FAnalyticsEventAttribute>& InOutParamArray);
+	TSharedPtr<FAnalyticsEvent> CreateAnalyticsEvent(FString InEventName);
+	void EnqueueAnalyticsEvent(TSharedPtr<FAnalyticsEvent> InAnalyticEvent);
+	void UpdateAnalyticsCustomValues();
+	void UpdatePlayEndStatistics();
+	void LogStatistics();
+	void OnMediaPlayerEventReceived(TSharedPtr<IAdaptiveStreamingPlayerAEMSEvent, ESPMode::ThreadSafe> InEvent, IAdaptiveStreamingPlayerAEMSReceiver::EDispatchMode InDispatchMode);
+	void HandleDeferredPlayerEvents();
+	void HandlePlayerEventOpenSource(const FString& InURL);
+	void HandlePlayerEventReceivedMainPlaylist(const FString& InEffectiveURL);
+	void HandlePlayerEventReceivedPlaylists();
+	void HandlePlayerEventTracksChanged();
+	void HandlePlayerEventPlaylistDownload(const Metrics::FPlaylistDownloadStats& InPlaylistDownloadStats);
+	void HandlePlayerEventBufferingStart(Metrics::EBufferingReason InBufferingReason);
+	void HandlePlayerEventBufferingEnd(Metrics::EBufferingReason InBufferingReason);
+	void HandlePlayerEventBandwidth(int64 InEffectiveBps, int64 InThroughputBps, double InLatencyInSeconds);
+	void HandlePlayerEventBufferUtilization(const Metrics::FBufferStats& InBufferStats);
+	void HandlePlayerEventSegmentDownload(const Metrics::FSegmentDownloadStats& InSegmentDownloadStats);
+	void HandlePlayerEventLicenseKey(const Metrics::FLicenseKeyStats& InLicenseKeyStats);
+	void HandlePlayerEventDataAvailabilityChange(const Metrics::FDataAvailabilityChange& InDataAvailability);
+	void HandlePlayerEventVideoQualityChange(int32 InNewBitrate, int32 InPreviousBitrate, bool bInIsDrasticDownswitch);
+	void HandlePlayerEventAudioQualityChange(int32 InNewBitrate, int32 InPreviousBitrate, bool bInIsDrasticDownswitch);
+	void HandlePlayerEventCodecFormatChange(const Electra::FStreamCodecInformation& InNewDecodingFormat);
+	void HandlePlayerEventPrerollStart();
+	void HandlePlayerEventPrerollEnd();
+	void HandlePlayerEventPlaybackStart();
+	void HandlePlayerEventPlaybackPaused();
+	void HandlePlayerEventPlaybackResumed();
+	void HandlePlayerEventPlaybackEnded();
+	void HandlePlayerEventJumpInPlayPosition(const FTimeValue& InToNewTime, const FTimeValue& InFromTime, Metrics::ETimeJumpReason InTimejumpReason);
+	void HandlePlayerEventPlaybackStopped();
+	void HandlePlayerEventSeekCompleted();
+	void HandlePlayerMediaMetadataChanged(const TSharedPtrTS<Electra::UtilsMP4::FMetadataParser>& InMetadata);
+	void HandlePlayerEventError(const FString& InErrorReason);
+	void HandlePlayerEventLogMessage(IInfoLog::ELevel InLogLevel, const FString& InLogMessage, int64 InPlayerWallclockMilliseconds);
+	void HandlePlayerEventDroppedVideoFrame();
+	void HandlePlayerEventDroppedAudioFrame();
+	void MediaStateOnPreparingFinished();
+	bool MediaStateOnPlay();
+	bool MediaStateOnPause();
+	void MediaStateOnEndReached();
+	void MediaStateOnSeekFinished();
 
-	void SuspendOrResumeDecoders(bool bSuspend, const Electra::FParamDict& InOptions) override;
+//+++++++++++++++++++++++++++++++++++++++++++++++++++++++++++++
+//+++++++++++++++++++++++++++++++++++++++++++++++++++++++++++++
 
 private:
 	DECLARE_DELEGATE_TwoParams(FOnMediaPlayerEventReceivedDelegate, TSharedPtrTS<IAdaptiveStreamingPlayerAEMSEvent> /*InEvent*/, IAdaptiveStreamingPlayerAEMSReceiver::EDispatchMode /*InDispatchMode*/);
@@ -171,6 +284,90 @@ private:
 		FOnMediaPlayerSubtitleFlushDelegate SubtitleFlushDelegate;
 	};
 
+	class FAdaptiveStreamingPlayerResourceProvider : public IAdaptiveStreamingPlayerResourceProvider
+	{
+	public:
+		FAdaptiveStreamingPlayerResourceProvider(TWeakPtr<IElectraSafeMediaOptionInterface, ESPMode::ThreadSafe> InOptionInterface)
+		{ OptionInterface = MoveTemp(InOptionInterface); }
+		virtual ~FAdaptiveStreamingPlayerResourceProvider() = default;
+
+		virtual void ProvideStaticPlaybackDataForURL(TSharedPtr<IAdaptiveStreamingPlayerResourceRequest, ESPMode::ThreadSafe> InOutRequest) override;
+
+		void ProcessPendingStaticResourceRequests();
+		void ClearPendingRequests();
+	private:
+		TWeakPtr<IElectraSafeMediaOptionInterface, ESPMode::ThreadSafe> OptionInterface;
+		/** Requests for static resource fetches we want to perform on the main thread **/
+		TQueue<TSharedPtr<IAdaptiveStreamingPlayerResourceRequest, ESPMode::ThreadSafe>, EQueueMode::Mpsc> PendingStaticResourceRequests;
+	};
+
+	class FInternalPlayerImpl
+	{
+	public:
+		/** The media player itself **/
+		TSharedPtr<IAdaptiveStreamingPlayer, ESPMode::ThreadSafe> AdaptivePlayer;
+
+		/** Renderers to use **/
+		TSharedPtr<FElectraRendererVideo, ESPMode::ThreadSafe> RendererVideo;
+		TSharedPtr<FElectraRendererAudio, ESPMode::ThreadSafe> RendererAudio;
+
+		/** Media samples to delete when closing */
+		TUniquePtr<FMediaSamples> MediaSamplesToDelete;
+
+		static void DoCloseAsync(TSharedPtr<FInternalPlayerImpl, ESPMode::ThreadSafe>&& InPlayer, uint32 InPlayerID, TSharedPtr<IAsyncResourceReleaseNotification, ESPMode::ThreadSafe> InAsyncResourceReleaseNotification);
+	};
+
+	struct FPlaystartOptions
+	{
+		void Reset()
+		{
+			TimeOffset.Reset();
+			InitialVideoTrackAttributes.Reset();
+			InitialAudioTrackAttributes.Reset();
+			InitialSubtitleTrackAttributes.Reset();
+			MaxVerticalStreamResolution.Reset();
+			MaxBandwidthForStreaming.Reset();
+			bDoNotPreload = false;
+		}
+		TOptional<FTimespan>		TimeOffset;
+		FStreamSelectionAttributes	InitialVideoTrackAttributes;
+		FStreamSelectionAttributes	InitialAudioTrackAttributes;
+		FStreamSelectionAttributes	InitialSubtitleTrackAttributes;
+		TOptional<int32>			MaxVerticalStreamResolution;
+		TOptional<int32>			MaxBandwidthForStreaming;
+		bool						bDoNotPreload = false;
+	};
+
+	struct FVideoStreamFormat
+	{
+		FIntPoint Resolution {};
+		double FrameRate = 0.0;
+		int32 Bitrate = 0;
+	};
+
+	struct FPlayerState
+	{
+		FPlayerState()
+		{ Reset(); }
+		TOptional<float> IntendedPlayRate;
+		float CurrentPlayRate;
+		EMediaState State;
+		EMediaStatus Status;
+
+		void Reset()
+		{
+			IntendedPlayRate.Reset();
+			CurrentPlayRate = 0.0f;
+			State = EMediaState::Closed;
+			Status = EMediaStatus::None;
+		}
+
+		float GetRate() const;
+		EMediaState GetState() const;
+		EMediaStatus GetStatus() const;
+		void SetIntendedPlayRate(float InIntendedRate);
+		void SetPlayRateFromPlayer(float InCurrentPlayerPlayRate);
+	};
 
 
 	struct FPlayerMetricEventBase
@@ -309,240 +506,12 @@ private:
 	};
 
 
-	void CalculateTargetSeekTime(FTimespan& OutTargetTime, const FTimespan& InTime);
-
-	bool PresentVideoFrame(const FVideoDecoderOutputPtr& InVideoFrame);
-	bool PresentAudioFrame(const IAudioDecoderOutputPtr& DecoderOutput);
-	bool PresentSubtitle(const ISubtitleDecoderOutputPtr& DecoderOutput);
-
-	void PlatformSuspendOrResumeDecoders(bool bSuspend, const Electra::FParamDict& InOptions);
-
-	void OnMediaPlayerEventReceived(TSharedPtrTS<IAdaptiveStreamingPlayerAEMSEvent> InEvent, IAdaptiveStreamingPlayerAEMSReceiver::EDispatchMode InDispatchMode);
-
-	// Methods from IAdaptiveStreamingPlayerMetrics
-	virtual void ReportOpenSource(const FString& URL) override
-	{ DeferredPlayerEvents.Enqueue(MakeSharedTS<FPlayerMetricEvent_OpenSource>(URL)); }
-	virtual void ReportReceivedMainPlaylist(const FString& EffectiveURL) override
-	{ DeferredPlayerEvents.Enqueue(MakeSharedTS<FPlayerMetricEvent_ReceivedMainPlaylist>(EffectiveURL)); }
-	virtual void ReportReceivedPlaylists() override
-	{ DeferredPlayerEvents.Enqueue(MakeSharedTS<FPlayerMetricEventBase>(FPlayerMetricEventBase::EType::ReceivedPlaylists)); }
-	virtual void ReportTracksChanged() override
-	{ DeferredPlayerEvents.Enqueue(MakeSharedTS<FPlayerMetricEventBase>(FPlayerMetricEventBase::EType::TracksChanged)); }
-	virtual void ReportPlaylistDownload(const Metrics::FPlaylistDownloadStats& PlaylistDownloadStats) override
-	{ DeferredPlayerEvents.Enqueue(MakeSharedTS<FPlayerMetricEvent_PlaylistDownload>(PlaylistDownloadStats)); }
-	virtual void ReportCleanStart() override
-	{ /*DeferredPlayerEvents.Enqueue(MakeSharedTS<FPlayerMetricEventBase>(FPlayerMetricEventBase::EType::CleanStart));*/
-		bDiscardOutputUntilCleanStart = false;
-	}
-	virtual void ReportBufferingStart(Metrics::EBufferingReason BufferingReason) override
-	{ DeferredPlayerEvents.Enqueue(MakeSharedTS<FPlayerMetricEvent_BufferingStart>(BufferingReason)); }
-	virtual void ReportBufferingEnd(Metrics::EBufferingReason BufferingReason) override
-	{ DeferredPlayerEvents.Enqueue(MakeSharedTS<FPlayerMetricEvent_BufferingEnd>(BufferingReason)); }
-	virtual void ReportBandwidth(int64 EffectiveBps, int64 ThroughputBps, double LatencyInSeconds) override
-	{ DeferredPlayerEvents.Enqueue(MakeSharedTS<FPlayerMetricEvent_Bandwidth>(EffectiveBps, ThroughputBps, LatencyInSeconds)); }
-	virtual void ReportBufferUtilization(const Metrics::FBufferStats& BufferStats) override
-	{ DeferredPlayerEvents.Enqueue(MakeSharedTS<FPlayerMetricEvent_BufferUtilization>(BufferStats)); }
-	virtual void ReportSegmentDownload(const Metrics::FSegmentDownloadStats& SegmentDownloadStats) override
-	{ DeferredPlayerEvents.Enqueue(MakeSharedTS<FPlayerMetricEvent_SegmentDownload>(SegmentDownloadStats)); }
-	virtual void ReportLicenseKey(const Metrics::FLicenseKeyStats& LicenseKeyStats) override
-	{ DeferredPlayerEvents.Enqueue(MakeSharedTS<FPlayerMetricEvent_LicenseKey>(LicenseKeyStats)); }
-	virtual void ReportDataAvailabilityChange(const Metrics::FDataAvailabilityChange& DataAvailability) override
-	{ DeferredPlayerEvents.Enqueue(MakeSharedTS<FPlayerMetricEvent_DataAvailabilityChange>(DataAvailability)); }
-	virtual void ReportVideoQualityChange(int32 NewBitrate, int32 PreviousBitrate, bool bIsDrasticDownswitch) override
-	{ DeferredPlayerEvents.Enqueue(MakeSharedTS<FPlayerMetricEvent_VideoQualityChange>(NewBitrate, PreviousBitrate, bIsDrasticDownswitch)); }
-	virtual void ReportAudioQualityChange(int32 NewBitrate, int32 PreviousBitrate, bool bIsDrasticDownswitch) override
-	{ DeferredPlayerEvents.Enqueue(MakeSharedTS<FPlayerMetricEvent_AudioQualityChange>(NewBitrate, PreviousBitrate, bIsDrasticDownswitch)); }
-	virtual void ReportDecodingFormatChange(const FStreamCodecInformation& NewDecodingFormat) override
-	{ DeferredPlayerEvents.Enqueue(MakeSharedTS<FPlayerMetricEvent_CodecFormatChange>(NewDecodingFormat)); }
-	virtual void ReportPrerollStart() override
-	{ DeferredPlayerEvents.Enqueue(MakeSharedTS<FPlayerMetricEventBase>(FPlayerMetricEventBase::EType::PrerollStart)); }
-	virtual void ReportPrerollEnd() override
-	{ DeferredPlayerEvents.Enqueue(MakeSharedTS<FPlayerMetricEventBase>(FPlayerMetricEventBase::EType::PrerollEnd)); }
-	virtual void ReportPlaybackStart() override
-	{ DeferredPlayerEvents.Enqueue(MakeSharedTS<FPlayerMetricEventBase>(FPlayerMetricEventBase::EType::PlaybackStart)); }
-	virtual void ReportPlaybackPaused() override
-	{ DeferredPlayerEvents.Enqueue(MakeSharedTS<FPlayerMetricEventBase>(FPlayerMetricEventBase::EType::PlaybackPaused)); }
-	virtual void ReportPlaybackResumed() override
-	{ DeferredPlayerEvents.Enqueue(MakeSharedTS<FPlayerMetricEventBase>(FPlayerMetricEventBase::EType::PlaybackResumed)); }
-	virtual void ReportPlaybackEnded() override
-	{ DeferredPlayerEvents.Enqueue(MakeSharedTS<FPlayerMetricEventBase>(FPlayerMetricEventBase::EType::PlaybackEnded)); }
-	virtual void ReportJumpInPlayPosition(const FTimeValue& ToNewTime, const FTimeValue& FromTime, Metrics::ETimeJumpReason TimejumpReason) override
-	{ DeferredPlayerEvents.Enqueue(MakeSharedTS<FPlayerMetricEvent_JumpInPlayPosition>(ToNewTime, FromTime, TimejumpReason)); }
-	virtual void ReportPlaybackStopped() override
-	{ DeferredPlayerEvents.Enqueue(MakeSharedTS<FPlayerMetricEventBase>(FPlayerMetricEventBase::EType::PlaybackStopped)); }
-	virtual void ReportSeekCompleted() override
-	{ DeferredPlayerEvents.Enqueue(MakeSharedTS<FPlayerMetricEventBase>(FPlayerMetricEventBase::EType::SeekCompleted)); }
-	virtual void ReportMediaMetadataChanged(TSharedPtrTS<Electra::UtilsMP4::FMetadataParser> Metadata) override
-	{ DeferredPlayerEvents.Enqueue(MakeSharedTS<FPlayerMetricEvent_MediaMetadataChange>(Metadata)); }
-	virtual void ReportError(const FString& ErrorReason) override
-	{ DeferredPlayerEvents.Enqueue(MakeSharedTS<FPlayerMetricEvent_Error>(ErrorReason)); }
-	virtual void ReportLogMessage(IInfoLog::ELevel InLogLevel, const FString& InLogMessage, int64 InPlayerWallclockMilliseconds) override
-	{ DeferredPlayerEvents.Enqueue(MakeSharedTS<FPlayerMetricEvent_LogMessage>(InLogLevel, InLogMessage, InPlayerWallclockMilliseconds)); }
-	virtual void ReportDroppedVideoFrame() override
-	{ DeferredPlayerEvents.Enqueue(MakeSharedTS<FPlayerMetricEventBase>(FPlayerMetricEventBase::EType::DroppedVideoFrame)); }
-	virtual void ReportDroppedAudioFrame() override
-	{ DeferredPlayerEvents.Enqueue(MakeSharedTS<FPlayerMetricEventBase>(FPlayerMetricEventBase::EType::DroppedAudioFrame)); }
-
-	void LogPresentationFramesQueues(FTimespan DeltaTime);
-
-	void ClearToDefaultState();
-
-	void MediaStateOnPreparingFinished();
-	bool MediaStateOnPlay();
-	bool MediaStateOnPause();
-	void MediaStateOnEndReached();
-	void MediaStateOnSeekFinished();
-	void TriggerFirstSeekIfNecessary();
-
-	TSharedPtr<FTrackMetadata, ESPMode::ThreadSafe> GetTrackStreamMetadata(EPlayerTrackType TrackType, int32 TrackIndex) const;
-
-	// Delegate to talk back to adapter host
-	TWeakPtr<IElectraPlayerAdapterDelegate, ESPMode::ThreadSafe>	AdapterDelegate;
-
-	// Contains number of audio tracks available to expose it later.
-	int32											NumTracksAudio;
-	int32											NumTracksVideo;
-	int32											NumTracksSubtitle;
-	int32											SelectedQuality;
-	mutable int32									SelectedVideoTrackIndex;
-	mutable int32									SelectedAudioTrackIndex;
-	mutable int32									SelectedSubtitleTrackIndex;
-	mutable bool									bVideoTrackIndexDirty;
-	mutable bool									bAudioTrackIndexDirty;
-	mutable bool									bSubtitleTrackIndexDirty;
-
-	bool											bInitialSeekPerformed;
-	bool											bDiscardOutputUntilCleanStart;
-	bool											bIsFirstBuffering;
-
-	FPlaybackRange									CurrentPlaybackRange;
-	TOptional<bool>									bFrameAccurateSeeking;
-	TOptional<bool>									bEnableLooping;
-
-	TSharedPtr<TMap<FString, TArray<TSharedPtr<Electra::IMediaStreamMetadata::IItem, ESPMode::ThreadSafe>>>, ESPMode::ThreadSafe> CurrentStreamMetadata;
-
-	TOptional<FVideoStreamFormat>					CurrentlyActiveVideoStreamFormat;
-
-	FIntPoint										LastPresentedFrameDimension;
-
-
-	struct FPlayerState
-	{
-		TOptional<float>		IntendedPlayRate;
-		float					CurrentPlayRate = 0.0f;
-
-		TAtomic<EPlayerState>	State;
-		TAtomic<EPlayerStatus>	Status;
-
-		bool					bUseInternal = false;
-
-		void Reset()
-		{
-			IntendedPlayRate.Reset();
-			CurrentPlayRate = 0.0f;
-			State = EPlayerState::Closed;
-			Status = EPlayerStatus::None;
-		}
-
-		float GetRate() const;
-		EPlayerState GetState() const;
-		EPlayerStatus GetStatus() const;
-
-		void SetIntendedPlayRate(float InIntendedRate);
-		void SetPlayRateFromPlayer(float InCurrentPlayerPlayRate);
-	};
-
-
-	/** Media player Guid */
-	FGuid											PlayerGuid;
-	/** Metric delegates */
-	FElectraPlayerSendAnalyticMetricsDelegate&			SendAnalyticMetricsDelegate;
-	FElectraPlayerSendAnalyticMetricsPerMinuteDelegate&	SendAnalyticMetricsPerMinuteDelegate;
-	FElectraPlayerReportVideoStreamingErrorDelegate&	ReportVideoStreamingErrorDelegate;
-	FElectraPlayerReportSubtitlesMetricsDelegate&		ReportSubtitlesMetricsDelegate;
-
-	/** Option interface **/
-	FPlaystartOptions								PlaystartOptions;
-
-	FPlayerState									PlayerState;
-
-	TAtomic<bool>									bPlayerHasClosed;
-	TAtomic<bool>									bHasPendingError;
-
-	bool											bAllowKillAfterCloseEvent;
-
-	/** Queued events */
-	TQueue<IElectraPlayerAdapterDelegate::EPlayerEvent>	DeferredEvents;
-	TQueue<TSharedPtrTS<FPlayerMetricEventBase>>	DeferredPlayerEvents;
-	TSharedPtrTS<FAEMSEventReceiver>				MediaPlayerEventReceiver;
-	TSharedPtrTS<FSubtitleEventReceiver>			MediaPlayerSubtitleReceiver;
-
-
-	/** The URL of the currently opened media. */
-	FString											MediaUrl;
-
-	class FInternalPlayerImpl
-	{
-	public:
-		/** The media player itself **/
-		TSharedPtr<IAdaptiveStreamingPlayer, ESPMode::ThreadSafe>	AdaptivePlayer;
-
-		/** Renderers to use **/
-		TSharedPtr<FElectraRendererVideo, ESPMode::ThreadSafe>		RendererVideo;
-		TSharedPtr<FElectraRendererAudio, ESPMode::ThreadSafe>		RendererAudio;
-
-		/** */
-		static void DoCloseAsync(TSharedPtr<FInternalPlayerImpl, ESPMode::ThreadSafe>&& Player, uint32 PlayerID, TSharedPtr<IAsyncResourceReleaseNotifyContainer, ESPMode::ThreadSafe> AsyncDestructNotification);
-	};
-
-	mutable FCriticalSection										PlayerLock;
-	TSharedPtr<FInternalPlayerImpl, ESPMode::ThreadSafe>			CurrentPlayer;
-	FEvent*															WaitForPlayerDestroyedEvent;
-	TSharedPtr<Electra::FApplicationTerminationHandler, ESPMode::ThreadSafe> AppTerminationHandler;
-
-	TSharedPtr<IAsyncResourceReleaseNotifyContainer, ESPMode::ThreadSafe> AsyncResourceReleaseNotification;
-
-	class FAdaptiveStreamingPlayerResourceProvider : public IAdaptiveStreamingPlayerResourceProvider
-	{
-	public:
-		FAdaptiveStreamingPlayerResourceProvider(const TWeakPtr<IElectraPlayerAdapterDelegate, ESPMode::ThreadSafe>& AdapterDelegate);
-		virtual ~FAdaptiveStreamingPlayerResourceProvider();
-
-		virtual void ProvideStaticPlaybackDataForURL(TSharedPtr<IAdaptiveStreamingPlayerResourceRequest, ESPMode::ThreadSafe> InOutRequest) override;
-
-		void SetExternalDataReader(TWeakPtr<IElectraPlayerExternalDataReader, ESPMode::ThreadSafe> InExternalDataReader);
-		void ProcessPendingStaticResourceRequests();
-		void ClearPendingRequests();
-
-	private:
-		/** Requests for static resource fetches we want to perform on the main thread **/
-		TQueue<TSharedPtr<IAdaptiveStreamingPlayerResourceRequest, ESPMode::ThreadSafe>, EQueueMode::Mpsc> PendingStaticResourceRequests;
-
-		// Player adapter delegate
-		TWeakPtr<IElectraPlayerAdapterDelegate, ESPMode::ThreadSafe> AdapterDelegate;
-		// External data reader
-		TWeakPtr<IElectraPlayerExternalDataReader, ESPMode::ThreadSafe> ExternalDataReader;
-		IElectraPlayerExternalDataReader::FElectraPlayerExternalDataReaderOnRequestCompleted ExternalDataCompletedDelegate;
-		static void OnExternalDataReadCompleted(IElectraPlayerExternalDataReader::FResponseDataPtr InResponseData, int64 InTotalFileSize, const IElectraPlayerExternalDataReader::FReadParam& InFromRequestParams);
-	};
-
-	TSharedPtr<FAdaptiveStreamingPlayerResourceProvider, ESPMode::ThreadSafe> StaticResourceProvider;
-	TSharedPtr<IVideoDecoderResourceDelegate, ESPMode::ThreadSafe> VideoDecoderResourceDelegate;
-
-	class FBlobRequest
+	/** Analytics event */
+	struct FAnalyticsEvent
 	{
-	public:
-		FBlobRequest() : Request(MakeShared<Electra::FHTTPResourceRequest, ESPMode::ThreadSafe>())
-		{ }
-		void OnBlobRequestComplete(TSharedPtrTS<Electra::FHTTPResourceRequest> InRequest)
-		{
-			bIsComplete = true;
-		}
-		TSharedPtr<Electra::FHTTPResourceRequest, ESPMode::ThreadSafe> Request;
-		bool bIsComplete = false;
-		bool bDispatched = false;
+		FString EventName;
+		TArray<FAnalyticsEventAttribute> ParamArray;
 	};
-	TSharedPtr<FBlobRequest, ESPMode::ThreadSafe> PendingBlobRequest;
 
 
 	class FAverageValue
@@ -748,72 +717,92 @@ private:
 		uint32					NumErrOther;
 	};
 
-	struct FAnalyticsEvent
-	{
-		FString EventName;
-		TArray<FAnalyticsEventAttribute> ParamArray;
-	};
 
-	void UpdatePlayEndStatistics();
-	void LogStatistics();
-	void AddCommonAnalyticsAttributes(TArray<FAnalyticsEventAttribute>& InOutParamArray);
-	TSharedPtr<FAnalyticsEvent> CreateAnalyticsEvent(FString InEventName);
-	void EnqueueAnalyticsEvent(TSharedPtr<FAnalyticsEvent> InAnalyticEvent);
-	void UpdateAnalyticsCustomValues();
+	/** Critical section protecting some members used by both the game and worker threads. */
+	FCriticalSection ParentObjectLock;
 
-	FCriticalSection						StatisticsLock;
-	FStatistics								Statistics;
+	/** Media event sink */
+	IMediaEventSink* EventSink = nullptr;
 
-	TQueue<TSharedPtr<FAnalyticsEvent>>		QueuedAnalyticEvents;
-	int32									NumQueuedAnalyticEvents = 0;
-	FString									AnalyticsOSVersion;
-	FString									AnalyticsGPUType;
-
-	/** Unique player instance GUID sent with each analytics event. This allows finding all events of a particular playback session. **/
-	FString									AnalyticsInstanceGuid;
-	/** Sequential analytics event number. Helps sorting events. **/
-	uint32									AnalyticsInstanceEventCount;
+	/** Metric delegates */
+	FElectraPlayerSendAnalyticMetricsDelegate& SendAnalyticMetricsDelegate;
+	FElectraPlayerSendAnalyticMetricsPerMinuteDelegate&	SendAnalyticMetricsPerMinuteDelegate;
+	FElectraPlayerReportVideoStreamingErrorDelegate& ReportVideoStreamingErrorDelegate;
+	FElectraPlayerReportSubtitlesMetricsDelegate& ReportSubtitlesMetricsDelegate;
 
-	/** Custom analytics constants. **/
-	FString									AnalyticsCustomValues[8];
+	/** Media sample queues */
+	mutable FCriticalSection MediaSamplesLock;
+	TUniquePtr<FMediaSamples> EmptyMediaSamples;
+	TUniquePtr<FMediaSamples> MediaSamples;
 
+	/** Option interface */
+	TWeakPtr<IElectraSafeMediaOptionInterface, ESPMode::ThreadSafe> OptionInterface;
 
-	/** Unique instance ID */
-	uint32 InstanceID = 0;
+	/** Current player stream metadata */
+	mutable TSharedPtr<TMap<FString, TArray<TUniquePtr<IMediaMetadataItem>>>, ESPMode::ThreadSafe> CurrentMetadata;
+	mutable bool bMetadataChanged = false;
 
-	void HandleBlobDownload();
+	/** Output sample pools */
+	TSharedPtr<FElectraTextureSamplePool, ESPMode::ThreadSafe> OutputTexturePool;
+	int64 OutputTexturePoolDecoderID = 0;
+	FElectraPlayerAudioSamplePool OutputAudioPool;
 
-	void HandleDeferredPlayerEvents();
-	void HandlePlayerEventOpenSource(const FString& URL);
-	void HandlePlayerEventReceivedMainPlaylist(const FString& EffectiveURL);
-	void HandlePlayerEventReceivedPlaylists();
-	void HandlePlayerEventTracksChanged();
-	void HandlePlayerEventPlaylistDownload(const Metrics::FPlaylistDownloadStats& PlaylistDownloadStats);
-	void HandlePlayerEventBufferingStart(Metrics::EBufferingReason BufferingReason);
-	void HandlePlayerEventBufferingEnd(Metrics::EBufferingReason BufferingReason);
-	void HandlePlayerEventBandwidth(int64 EffectiveBps, int64 ThroughputBps, double LatencyInSeconds);
-	void HandlePlayerEventBufferUtilization(const Metrics::FBufferStats& BufferStats);
-	void HandlePlayerEventSegmentDownload(const Metrics::FSegmentDownloadStats& SegmentDownloadStats);
-	void HandlePlayerEventLicenseKey(const Metrics::FLicenseKeyStats& LicenseKeyStats);
-	void HandlePlayerEventDataAvailabilityChange(const Metrics::FDataAvailabilityChange& DataAvailability);
-	void HandlePlayerEventVideoQualityChange(int32 NewBitrate, int32 PreviousBitrate, bool bIsDrasticDownswitch);
-	void HandlePlayerEventAudioQualityChange(int32 NewBitrate, int32 PreviousBitrate, bool bIsDrasticDownswitch);
-	void HandlePlayerEventCodecFormatChange(const Electra::FStreamCodecInformation& NewDecodingFormat);
-	void HandlePlayerEventPrerollStart();
-	void HandlePlayerEventPrerollEnd();
-	void HandlePlayerEventPlaybackStart();
-	void HandlePlayerEventPlaybackPaused();
-	void HandlePlayerEventPlaybackResumed();
-	void HandlePlayerEventPlaybackEnded();
-	void HandlePlayerEventJumpInPlayPosition(const FTimeValue& ToNewTime, const FTimeValue& FromTime, Metrics::ETimeJumpReason TimejumpReason);
-	void HandlePlayerEventPlaybackStopped();
-	void HandlePlayerEventSeekCompleted();
-	void HandlePlayerMediaMetadataChanged(const TSharedPtrTS<Electra::UtilsMP4::FMetadataParser>& InMetadata);
-	void HandlePlayerEventError(const FString& ErrorReason);
-	void HandlePlayerEventLogMessage(IInfoLog::ELevel InLogLevel, const FString& InLogMessage, int64 InPlayerWallclockMilliseconds);
-	void HandlePlayerEventDroppedVideoFrame();
-	void HandlePlayerEventDroppedAudioFrame();
-};
+	/** Application system event handler */
+	TSharedPtr<Electra::FApplicationTerminationHandler, ESPMode::ThreadSafe> AppTerminationHandler;
 
-ENUM_CLASS_FLAGS(FElectraPlayer::EPlayerStatus);
+	/** Media player Guid */
+	FGuid PlayerGuid;
+	uint32 PlayerUniqueID = 0;
+	static std::atomic<uint32> NextPlayerUniqueID;
+
+	/** Analytics */
+	TQueue<TSharedPtr<FAnalyticsEvent>> QueuedAnalyticEvents;
+	FString AnalyticsCustomValues[8];
+	FString AnalyticsOSVersion;
+	FString AnalyticsGPUType;
+	FString AnalyticsInstanceGuid;
+	uint32 AnalyticsInstanceEventCount;
+	std::atomic<int32> NumQueuedAnalyticEvents { 0 };
+
+	/** Statistics */
+	FCriticalSection StatisticsLock;
+	FStatistics Statistics;
+
+	/** External resource provider. */
+	TSharedPtr<FAdaptiveStreamingPlayerResourceProvider, ESPMode::ThreadSafe> StaticResourceProvider;
 
+	/** Internal states. */
+	FPlaystartOptions PlaystartOptions;
+	FString MediaUrl;
+	TOptional<bool> bFrameAccurateSeeking;
+	TOptional<bool> bEnableLooping;
+
+	TQueue<TSharedPtrTS<FPlayerMetricEventBase>> DeferredPlayerEvents;
+	TQueue<EMediaEvent> DeferredMediaEvents;
+	FCriticalSection VideoFormatLock;
+	TOptional<FVideoStreamFormat> CurrentlyActiveVideoStreamFormat;
+	FIntPoint LastPresentedFrameDimension;
+	int32 NumTracksAudio = 0;
+	int32 NumTracksVideo = 0;
+	int32 NumTracksSubtitle = 0;
+	int32 SelectedQuality = 0;
+	mutable int32 SelectedVideoTrackIndex = -1;
+	mutable int32 SelectedAudioTrackIndex = -1;
+	mutable int32 SelectedSubtitleTrackIndex = -1;
+	mutable bool bVideoTrackIndexDirty = true;
+	mutable bool bAudioTrackIndexDirty = true;
+	mutable bool bSubtitleTrackIndexDirty = true;
+	bool bInitialSeekPerformed = false;
+	bool bIsFirstBuffering = true;
+	TSharedPtr<TMap<FString, TArray<TUniquePtr<IMediaMetadataItem>>>, ESPMode::ThreadSafe> CurrentStreamMetadata;
+
+	TSharedPtr<FInternalPlayerImpl, ESPMode::ThreadSafe> CurrentPlayer;
+	TSharedPtr<FAEMSEventReceiver, ESPMode::ThreadSafe> MediaPlayerEventReceiver;
+	TSharedPtr<FSubtitleEventReceiver, ESPMode::ThreadSafe> MediaPlayerSubtitleReceiver;
+	TSharedPtr<IAsyncResourceReleaseNotification, ESPMode::ThreadSafe> AsyncResourceReleaseNotification;
+
+	TRange<FTimespan> CurrentPlaybackRange;
+	FPlayerState PlayerState;
+	bool bHasPendingError = false;
+	bool bHasClosedDueToError = false;
+};
diff --git a/Engine/Plugins/Media/ElectraPlayer/Source/ElectraPlayerRuntime/Private/ElectraPlayerAudioSample.cpp b/Engine/Plugins/Media/ElectraPlayer/Source/ElectraPlayerRuntime/Private/ElectraPlayerAudioSample.cpp
new file mode 100644
--- /dev/null
+++ b/Engine/Plugins/Media/ElectraPlayer/Source/ElectraPlayerRuntime/Private/ElectraPlayerAudioSample.cpp
@@ -0,0 +1,13 @@
+// Copyright Epic Games, Inc. All Rights Reserved.
+
+#include "ElectraPlayerAudioSample.h"
+
+void FElectraPlayerAudioSample::InitializePoolable()
+{
+}
+
+void FElectraPlayerAudioSample::ShutdownPoolable()
+{
+	DecoderOutput.Reset();
+}
+
diff --git a/Engine/Plugins/Media/ElectraPlayer/Source/ElectraPlayerRuntime/Private/ElectraPlayerAudioSample.h b/Engine/Plugins/Media/ElectraPlayer/Source/ElectraPlayerRuntime/Private/ElectraPlayerAudioSample.h
new file mode 100644
--- /dev/null
+++ b/Engine/Plugins/Media/ElectraPlayer/Source/ElectraPlayerRuntime/Private/ElectraPlayerAudioSample.h
@@ -0,0 +1,91 @@
+// Copyright Epic Games, Inc. All Rights Reserved.
+#pragma once
+
+#include "IMediaAudioSample.h"
+#include "MediaObjectPool.h"
+#include "Misc/Timespan.h"
+#include "Templates/SharedPointer.h"
+
+#include "MediaAudioDecoderOutput.h"
+
+
+// IMediaAudioSample impl
+// is used as audio decoder output
+// contains decoded audio samples ready for playback along with corresponding metadata
+// owns binary data so can be cached
+class FElectraPlayerAudioSample :
+	public IMediaAudioSample,
+	public IMediaPoolable
+{
+public:
+
+	FElectraPlayerAudioSample()
+	{
+	}
+
+	~FElectraPlayerAudioSample()
+	{
+	}
+
+	void Initialize(const IAudioDecoderOutputPtr& InDecoderOutput)
+	{
+		DecoderOutput = InDecoderOutput;
+	}
+
+	const void* GetBuffer() override
+	{ return DecoderOutput ? DecoderOutput->GetBuffer() : nullptr; }
+
+	uint32 GetChannels() const override
+	{ return DecoderOutput ? DecoderOutput->GetChannels() : 0; }
+
+	FTimespan GetDuration() const override
+	{ return DecoderOutput ? DecoderOutput->GetDuration() : FTimespan::Zero(); }
+
+	EMediaAudioSampleFormat GetFormat() const override
+	{
+		static_assert((int32)EMediaAudioSampleFormat::Undefined == (int32)IAudioDecoderOutput::ESampleFormat::Undefined, "check enums are identical");
+		static_assert((int32)EMediaAudioSampleFormat::Double == (int32)IAudioDecoderOutput::ESampleFormat::Double, "check enums are identical");
+		static_assert((int32)EMediaAudioSampleFormat::Float == (int32)IAudioDecoderOutput::ESampleFormat::Float, "check enums are identical");
+		static_assert((int32)EMediaAudioSampleFormat::Int8 == (int32)IAudioDecoderOutput::ESampleFormat::Int8, "check enums are identical");
+		static_assert((int32)EMediaAudioSampleFormat::Int16 == (int32)IAudioDecoderOutput::ESampleFormat::Int16, "check enums are identical");
+		static_assert((int32)EMediaAudioSampleFormat::Int32 == (int32)IAudioDecoderOutput::ESampleFormat::Int32, "check enums are identical");
+
+		return DecoderOutput ? (EMediaAudioSampleFormat)DecoderOutput->GetFormat() : EMediaAudioSampleFormat::Undefined;
+	}
+
+	uint32 GetFrames() const override
+	{ return DecoderOutput ? DecoderOutput->GetFrames() : 0; }
+
+	uint32 GetSampleRate() const override
+	{ return DecoderOutput ? DecoderOutput->GetSampleRate() : 0; }
+
+	FMediaTimeStamp GetTime() const override
+	{
+		if (!DecoderOutput)
+		{
+			return FMediaTimeStamp();
+		}
+		FDecoderTimeStamp TimeStamp = DecoderOutput->GetTime();
+		return FMediaTimeStamp(TimeStamp.Time, TimeStamp.SequenceIndex);
+	}
+
+	uint32 GetMaxBufferBytes() const
+	{ return DecoderOutput ? DecoderOutput->GetReservedBufferBytes() : 0; }
+
+	virtual void SetTime(const FMediaTimeStamp& InTime)
+	{
+		if (DecoderOutput)
+		{
+			DecoderOutput->SetTime(FDecoderTimeStamp(InTime.Time, InTime.GetIndexValue()));
+		}
+	}
+
+	virtual void InitializePoolable() override;
+	virtual void ShutdownPoolable() override;
+
+private:
+	IAudioDecoderOutputPtr DecoderOutput;
+};
+
+using FElectraPlayerAudioSamplePtr = TSharedPtr<FElectraPlayerAudioSample, ESPMode::ThreadSafe>;
+using FElectraPlayerAudioSamplePool = TMediaObjectPool<FElectraPlayerAudioSample>;
diff --git a/Engine/Plugins/Media/ElectraPlayer/Source/ElectraPlayerRuntime/Private/ElectraPlayerDecoderResourceManager.cpp b/Engine/Plugins/Media/ElectraPlayer/Source/ElectraPlayerRuntime/Private/ElectraPlayerDecoderResourceManager.cpp
--- a/Engine/Plugins/Media/ElectraPlayer/Source/ElectraPlayerRuntime/Private/ElectraPlayerDecoderResourceManager.cpp
+++ b/Engine/Plugins/Media/ElectraPlayer/Source/ElectraPlayerRuntime/Private/ElectraPlayerDecoderResourceManager.cpp
@@ -2,7 +2,6 @@
 
 #include "IElectraPlayerDecoderResourceManager.h"
 #include "Decoder/ElectraDecoderResourceManager.h"
-#include "ElectraPlayerPlatform.h"
 #include "Renderer/RendererVideo.h"
 
 
@@ -11,11 +10,6 @@ TSharedPtr<IElectraDecoderResourceDelegate, ESPMode::ThreadSafe> FElectraPlayerD
 	return Electra::FPlatformElectraDecoderResourceManager::GetDelegate();
 }
 
-TSharedPtr<Electra::IVideoDecoderResourceDelegate, ESPMode::ThreadSafe> FElectraPlayerDecoderResourceManager::CreatePlatformVideoDecoderResourceDelegate(TSharedPtr<IElectraPlayerAdapterDelegate, ESPMode::ThreadSafe> InElectraPlayerAdapterDelegate)
-{
-	return Electra::PlatformCreateVideoDecoderResourceDelegate(InElectraPlayerAdapterDelegate);
-}
-
 bool FElectraPlayerDecoderResourceManager::SetupRenderBufferFromDecoderOutput(FString& OutErrorMessage, Electra::IMediaRenderer::IBuffer* InOutBufferToSetup, TSharedPtr<Electra::FParamDict, ESPMode::ThreadSafe> InOutBufferPropertes, TSharedPtr<IElectraDecoderVideoOutput, ESPMode::ThreadSafe> InDecoderOutput, IElectraDecoderResourceDelegateBase::IDecoderPlatformResource* InPlatformSpecificResource)
 {
 	return Electra::FPlatformElectraDecoderResourceManager::SetupRenderBufferFromDecoderOutput(OutErrorMessage, InOutBufferToSetup, InOutBufferPropertes, InDecoderOutput, InPlatformSpecificResource);
diff --git a/Engine/Plugins/Media/ElectraPlayer/Source/ElectraPlayerRuntime/Private/ElectraPlayerRuntimeModule.cpp b/Engine/Plugins/Media/ElectraPlayer/Source/ElectraPlayerRuntime/Private/ElectraPlayerRuntimeModule.cpp
--- a/Engine/Plugins/Media/ElectraPlayer/Source/ElectraPlayerRuntime/Private/ElectraPlayerRuntimeModule.cpp
+++ b/Engine/Plugins/Media/ElectraPlayer/Source/ElectraPlayerRuntime/Private/ElectraPlayerRuntimeModule.cpp
@@ -7,14 +7,10 @@
 
 #include "PlayerRuntimeGlobal.h"
 #include "Core/MediaInterlocked.h"
-#include "Renderer/RendererVideo.h"
 
-#include "Renderer/RendererAudio.h"
-
-#include "ElectraPlayerPrivate.h"
 #include "IElectraPlayerRuntimeModule.h"
+#include "ElectraPlayerPrivate.h"
 #include "ElectraPlayer.h"
-#include "ElectraPlayerPlatform.h"
 
 DEFINE_LOG_CATEGORY(LogElectraPlayer);
 DEFINE_LOG_CATEGORY(LogElectraPlayerStats);
@@ -27,26 +23,19 @@ DEFINE_LOG_CATEGORY(LogElectraPlayerStats);
 class FElectraPlayerRuntimeModule : public IElectraPlayerRuntimeModule
 {
 public:
-	// IElectraPlayerRuntimeModule interface
-
-	virtual bool IsInitialized() const override
+	bool IsInitialized() const override
 	{
 		return bInitialized;
 	}
 
-public:
-	// IModuleInterface interface
-
 	void StartupModule() override
 	{
 		LLM_SCOPE_BYNAME(TEXT("StartupModule"));
 		check(!bInitialized);
-
-		// Core
-		Electra::Configuration ElectraConfig;
-		FString AnalyticsEventsFromIni;
 	
 		// Read Analytics Events from ini configuring the events sent during playback
+		Electra::Configuration ElectraConfig;
+		FString AnalyticsEventsFromIni;
 		if (GConfig->GetString(TEXT("ElectraPlayer"), TEXT("AnalyticsEvents"), AnalyticsEventsFromIni, GEngineIni))
 		{
 			// Parse comma delimited strings into arrays
@@ -58,18 +47,11 @@ public:
 				ElectraConfig.EnabledAnalyticsEvents.Add(EnabledEvent, true);
 			}
 		}
-
 		UE_LOG(LogElectraPlayer, Verbose, TEXT("Found %d enabled ElectraPlayer analytic events."), ElectraConfig.EnabledAnalyticsEvents.Num());
 
+		// Core
 		Electra::Startup(ElectraConfig);
 
-		// Video renderer
-		FElectraRendererVideo::SystemConfiguration SysCfgVideo;
-		FElectraRendererVideo::Startup(SysCfgVideo);
-		// Audio renderer
-		FElectraRendererAudio::SystemConfiguration SysCfgAudio;
-		FElectraRendererAudio::Startup(SysCfgAudio);
-
 		bInitialized = true;
 	}
 
@@ -81,10 +63,7 @@ public:
 			// to avoid potential hangs in there.
 			if (Electra::WaitForAllPlayersToHaveTerminated())
 			{
-				FElectraRendererAudio::Shutdown();
-				FElectraRendererVideo::Shutdown();
 				Electra::Shutdown();
-				Electra::PlatformShutdown();
 			}
 			else
 			{
diff --git a/Engine/Plugins/Media/ElectraPlayer/Source/ElectraPlayerRuntime/Private/Runtime/Decoder/Android/AndroidElectraDecoderResourceManager.cpp b/Engine/Plugins/Media/ElectraPlayer/Source/ElectraPlayerRuntime/Private/Runtime/Decoder/Android/AndroidElectraDecoderResourceManager.cpp
--- a/Engine/Plugins/Media/ElectraPlayer/Source/ElectraPlayerRuntime/Private/Runtime/Decoder/Android/AndroidElectraDecoderResourceManager.cpp
+++ b/Engine/Plugins/Media/ElectraPlayer/Source/ElectraPlayerRuntime/Private/Runtime/Decoder/Android/AndroidElectraDecoderResourceManager.cpp
@@ -7,10 +7,6 @@
 #include "Renderer/RendererBase.h"
 #include "ElectraVideoDecoder_Android.h"
 
-#include "VideoDecoderResourceDelegate.h"
-
-//#include "Renderer/RendererVideo.h"
-//#include "MediaDecoderOutput.h"
 
 /***************************************************************************************************************************************************/
 /***************************************************************************************************************************************************/
@@ -18,15 +14,6 @@
 namespace Electra
 {
 
-bool FElectraDecoderResourceManagerAndroid::Startup()
-{
-	return true;
-}
-
-void FElectraDecoderResourceManagerAndroid::Shutdown()
-{
-}
-
 TSharedPtr<IElectraDecoderResourceDelegate, ESPMode::ThreadSafe> FElectraDecoderResourceManagerAndroid::GetDelegate()
 {
 	static TSharedPtr<IElectraDecoderResourceDelegate, ESPMode::ThreadSafe> Self = MakeShared<FElectraDecoderResourceManagerAndroid, ESPMode::ThreadSafe>();
@@ -34,169 +21,15 @@ TSharedPtr<IElectraDecoderResourceDelegate, ESPMode::ThreadSafe> FElectraDecoder
 }
 
 
-FElectraDecoderResourceManagerAndroid::~FElectraDecoderResourceManagerAndroid()
-{
-}
-
-
-class FElectraDecoderResourceManagerAndroid::FInstanceVars : public IElectraDecoderResourceDelegateAndroid::IDecoderPlatformResourceAndroid
-{
-public:
-	virtual ~FInstanceVars() {}
-	void RequestSurface(TWeakPtr<ISurfaceRequestCallback, ESPMode::ThreadSafe> InRequestCallback) override;
-	ESurfaceChangeResult VerifySurfaceView(void*& OutNewSurfaceView, void* InCurrentSurfaceView) override;
-	void SetBufferReleaseCallback(TWeakPtr<IBufferReleaseCallback, ESPMode::ThreadSafe> InBufferReleaseCallback) override;
-
-	TWeakPtr<IVideoDecoderResourceDelegate, ESPMode::ThreadSafe> VideoDecoderResourceDelegate;
-	TWeakPtr<ISurfaceRequestCallback, ESPMode::ThreadSafe> SurfaceRequestCallback;
-	TWeakPtr<IBufferReleaseCallback, ESPMode::ThreadSafe> BufferReleaseCallback;
-
-	IDecoderPlatformResourceAndroid::ISurfaceRequestCallback::ESurfaceType SurfaceType = IDecoderPlatformResourceAndroid::ISurfaceRequestCallback::ESurfaceType::Error;
-	volatile uint32 CurrentID = 0;
-
-	static std::atomic<uint32> NextID;
-};
-std::atomic<uint32> FElectraDecoderResourceManagerAndroid::FInstanceVars::NextID { 0 };
-
-
-void FElectraDecoderResourceManagerAndroid::FInstanceVars::RequestSurface(TWeakPtr<IElectraDecoderResourceDelegateAndroid::IDecoderPlatformResourceAndroid::ISurfaceRequestCallback, ESPMode::ThreadSafe> InRequestCallback)
-{
-	SurfaceRequestCallback = MoveTemp(InRequestCallback);
-
-	auto Notify = [Callback=SurfaceRequestCallback](IDecoderPlatformResourceAndroid::ISurfaceRequestCallback::ESurfaceType InType, void* InSurface) -> void
-	{
-		TSharedPtr<ISurfaceRequestCallback, ESPMode::ThreadSafe> cb = Callback.Pin();
-		if (cb.IsValid())
-		{
-			cb->OnNewSurface(InType, InSurface);
-		}
-	};
-
-	TSharedPtr<IVideoDecoderResourceDelegate, ESPMode::ThreadSafe> rd = VideoDecoderResourceDelegate.Pin();
-	if (!rd.IsValid())
-	{
-		Notify(IDecoderPlatformResourceAndroid::ISurfaceRequestCallback::ESurfaceType::Error, nullptr);
-		return;
-	}
-
-	/**
-		Note: The surface returned (if any) must be an _additional_ globalref, not the actual surface.
-		      This is because the actual surface might go away at any time, leaving us with an otherwise stale
-			  reference that will result in a JNI crash when used.
-			  The decoder implementation will pass along this globalref to the decoder and then release it,
-			  so this must not be the actual surface handle!
-	 */
-	void* Surface = rd->VideoDecoderResourceDelegate_GetCodecSurface();
-	SurfaceType = IDecoderPlatformResourceAndroid::ISurfaceRequestCallback::ESurfaceType::Surface;
-	Notify(SurfaceType, Surface);
-}
-
-FElectraDecoderResourceManagerAndroid::FInstanceVars::ESurfaceChangeResult FElectraDecoderResourceManagerAndroid::FInstanceVars::VerifySurfaceView(void*& OutNewSurfaceView, void* InCurrentSurfaceView)
-{
-	check(!"This should not be called since we are not supporting SurfaceViews at the moment!");
-	OutNewSurfaceView = nullptr;
-	return ESurfaceChangeResult::Error;
-}
-
-
-void FElectraDecoderResourceManagerAndroid::FInstanceVars::SetBufferReleaseCallback(TWeakPtr<IBufferReleaseCallback, ESPMode::ThreadSafe> InBufferReleaseCallback)
-{
-	CurrentID = ++NextID;
-	BufferReleaseCallback = MoveTemp(InBufferReleaseCallback);
-}
-
-
-
-void FElectraDecoderResourceManagerAndroid::ReleaseToSurface(uint32 InID, const FDecoderTimeStamp& Time)
-{
-	TSharedPtr<FElectraDecoderResourceManagerAndroid, ESPMode::ThreadSafe> Dlg = StaticCastSharedPtr<FElectraDecoderResourceManagerAndroid>(GetDelegate());
-	if (Dlg.IsValid())
-	{
-		FScopeLock lock(&Dlg->Lock);
-		for(auto &Inst : Dlg->InstanceVars)
-		{
-			if (Inst->CurrentID == InID)
-			{
-				TSharedPtr<IElectraDecoderResourceDelegateAndroid::IDecoderPlatformResourceAndroid::IBufferReleaseCallback, ESPMode::ThreadSafe> cb = Inst->BufferReleaseCallback.Pin();
-				if (cb.IsValid())
-				{
-//					cb->OnReleaseBuffer();
-				}
-			}
-		}
-	}
-}
-
-
-
-IElectraDecoderResourceDelegateAndroid::IDecoderPlatformResource* FElectraDecoderResourceManagerAndroid::CreatePlatformResource(void* InOwnerHandle, IElectraDecoderResourceDelegateAndroid::EDecoderPlatformResourceType InDecoderResourceType, const TMap<FString, FVariant> InOptions)
-{
-/*
-	if (InDecoderResourceType != IElectraDecoderResourceDelegateAndroid::EDecoderPlatformResourceType::Video)
-	{
-		return nullptr;
-	}
-*/
-	IVideoDecoderResourceDelegate* Delegate = reinterpret_cast<IVideoDecoderResourceDelegate*>(ElectraDecodersUtil::GetVariantValueSafeU64(InOptions, TEXT("VideoResourceDelegate"), 0));
-	if (Delegate)
-	{
-		TSharedPtr<IVideoDecoderResourceDelegate, ESPMode::ThreadSafe> VideoDecoderResourceDelegate = Delegate->AsShared();
-		FInstanceVars* Inst = new FInstanceVars;
-		Inst->VideoDecoderResourceDelegate = VideoDecoderResourceDelegate;
-		FScopeLock lock(&Lock);
-		InstanceVars.Emplace(Inst);
-		return Inst;
-	}
-	return nullptr;
-}
-
-void FElectraDecoderResourceManagerAndroid::ReleasePlatformResource(void* InOwnerHandle, IDecoderPlatformResource* InHandleToDestroy)
-{
-	FInstanceVars* InInst = static_cast<FInstanceVars*>(InHandleToDestroy);
-	if (InInst)
-	{
-		FScopeLock lock(&Lock);
-		check(InstanceVars.Contains(InInst));
-		if (InstanceVars.Contains(InInst))
-		{
-			InstanceVars.Remove(InInst);
-			delete InInst;
-		}
-	}
-}
-
-
-class FElectraDecoderVideoOutputCopyResources : public IElectraDecoderVideoOutputCopyResources
-{
-public:
-	virtual ~FElectraDecoderVideoOutputCopyResources() {}
-	void SetBufferIndex(int32 InIndex) override
-	{ BufferIndex = InIndex; }
-	void SetValidCount(int32 InValidCount) override
-	{ ValidCount = InValidCount; }
-	bool ShouldReleaseBufferImmediately() override
-	{ return bReleaseImmediately; }
-
-	int32 BufferIndex = -1;
-	int32 ValidCount = -1;
-	bool bReleaseImmediately = false;
-};
-
 
 bool FElectraDecoderResourceManagerAndroid::SetupRenderBufferFromDecoderOutput(FString& OutErrorMessage, IMediaRenderer::IBuffer* InOutBufferToSetup, TSharedPtr<FParamDict, ESPMode::ThreadSafe> InOutBufferPropertes, TSharedPtr<IElectraDecoderVideoOutput, ESPMode::ThreadSafe> InDecoderOutput, IDecoderPlatformResource* InPlatformSpecificResource)
 {
 	check(InOutBufferToSetup);
 	check(InOutBufferPropertes.IsValid());
 	check(InDecoderOutput.IsValid());
-	if (!InPlatformSpecificResource)
-	{
-		return false;
-	}
 
-	FInstanceVars* InInst = static_cast<FInstanceVars*>(InPlatformSpecificResource);
-	TSharedPtr<IVideoDecoderResourceDelegate, ESPMode::ThreadSafe> PinnedResourceDelegate = InInst->VideoDecoderResourceDelegate.Pin();
 	TSharedPtr<FElectraPlayerVideoDecoderOutputAndroid, ESPMode::ThreadSafe> DecoderOutput = InOutBufferToSetup->GetBufferProperties().GetValue(RenderOptionKeys::Texture).GetSharedPointer<FElectraPlayerVideoDecoderOutputAndroid>();
-	if (DecoderOutput.IsValid() && PinnedResourceDelegate.IsValid())
+	if (DecoderOutput.IsValid())
 	{
 		FElectraVideoDecoderOutputCropValues Crop = InDecoderOutput->GetCropValues();
 		InOutBufferPropertes->Set(IDecoderOutputOptionNames::Width, FVariantValue((int64)InDecoderOutput->GetWidth()));
@@ -214,42 +47,8 @@ bool FElectraDecoderResourceManagerAndroid::SetupRenderBufferFromDecoderOutput(F
 		InOutBufferPropertes->Set(IDecoderOutputOptionNames::PixelFormat, FVariantValue((int64)((NumBits > 8)  ? EPixelFormat::PF_A2B10G10R10 : EPixelFormat::PF_B8G8R8A8)));
 		InOutBufferPropertes->Set(IDecoderOutputOptionNames::BitsPerComponent, FVariantValue((int64) NumBits));
 		InOutBufferPropertes->Set(IDecoderOutputOptionNames::Pitch, FVariantValue((int64)InDecoderOutput->GetDecodedWidth() * ((NumBits > 8) ? 2 : 1)));
-
-		FElectraDecoderVideoOutputCopyResources cr;
-		cr.bReleaseImmediately = InInst->SurfaceType == IDecoderPlatformResourceAndroid::ISurfaceRequestCallback::ESurfaceType::Surface;
-		IElectraDecoderVideoOutput::EImageCopyResult CopyResult = InDecoderOutput->CopyPlatformImage(&cr);
-		if (CopyResult != IElectraDecoderVideoOutput::EImageCopyResult::Succeeded)
-		{
-			return false;
-		}
-
-		check(InInst->SurfaceType == IDecoderPlatformResourceAndroid::ISurfaceRequestCallback::ESurfaceType::Surface);
-
-		DecoderOutput->Initialize(InInst->SurfaceType == IDecoderPlatformResourceAndroid::ISurfaceRequestCallback::ESurfaceType::Surface     ? FVideoDecoderOutputAndroid::EOutputType::DirectToSurfaceAsQueue :
-								  InInst->SurfaceType == IDecoderPlatformResourceAndroid::ISurfaceRequestCallback::ESurfaceType::SurfaceView ? FVideoDecoderOutputAndroid::EOutputType::DirectToSurfaceAsView :
-								  FVideoDecoderOutputAndroid::EOutputType::Unknown, cr.BufferIndex, cr.ValidCount, ReleaseToSurface, InInst->CurrentID, InOutBufferPropertes);
+		DecoderOutput->Initialize(FVideoDecoderOutputAndroid::EOutputType::DirectToSurfaceAsQueue, InOutBufferPropertes);
 		return true;
-#if 0
------------------
-		DecoderOutput->Initialize(bSurfaceIsView ? FVideoDecoderOutputAndroid::EOutputType::DirectToSurfaceAsView : FVideoDecoderOutputAndroid::EOutputType::DirectToSurfaceAsQueue, NextImage.OutputBufferInfo.BufferIndex, NextImage.OutputBufferInfo.ValidCount, ReleaseToSurface, NativeDecoderID, OutputBufferSampleProperties.Get());
-
-		if (!bSurfaceIsView)
-		{
-			// Release the decoder output buffer, thereby enqueuing it on our output surface.
-			// (we are issuing an RHI thread based update to our texture for each of these, so we should always have a 1:1 mapping - assuming we are fast enough
-			//  to not make the surface drop a frame before we get to it)
-			DecoderInstance->ReleaseOutputBuffer(NextImage.OutputBufferInfo.BufferIndex, NextImage.OutputBufferInfo.ValidCount, true, -1);
-		}
-		else
-		{
-			// We decode right into a Surface. Queue up output buffers until we are ready to show them
-			ReadyOutputBuffersToSurface.Emplace(FOutputBufferInfo(FDecoderTimeStamp(NextImage.SourceInfo->AdjustedPTS.GetAsTimespan(), NextImage.SourceInfo->AdjustedPTS.GetSequenceIndex()), NextImage.OutputBufferInfo.BufferIndex, NextImage.OutputBufferInfo.ValidCount));
-		}
-		// Note: we are returning the buffer to the renderer before we are done getting data
-		// (but: this will sync all up as the render command queues are all in order - and hence the async process will be done before MediaTextureResources sees this)
-		Renderer->ReturnBuffer(RenderOutputBuffer, true, *OutputBufferSampleProperties);
-----------------
-#endif
 	}
 	return false;
 }
diff --git a/Engine/Plugins/Media/ElectraPlayer/Source/ElectraPlayerRuntime/Private/Runtime/Decoder/Android/AndroidElectraDecoderResourceManager.h b/Engine/Plugins/Media/ElectraPlayer/Source/ElectraPlayerRuntime/Private/Runtime/Decoder/Android/AndroidElectraDecoderResourceManager.h
--- a/Engine/Plugins/Media/ElectraPlayer/Source/ElectraPlayerRuntime/Private/Runtime/Decoder/Android/AndroidElectraDecoderResourceManager.h
+++ b/Engine/Plugins/Media/ElectraPlayer/Source/ElectraPlayerRuntime/Private/Runtime/Decoder/Android/AndroidElectraDecoderResourceManager.h
@@ -6,7 +6,6 @@
 #include "Renderer/RendererBase.h"
 
 class IElectraDecoderVideoOutput;
-struct FDecoderTimeStamp;
 
 namespace Electra
 {
@@ -14,21 +13,12 @@ namespace Electra
 class FElectraDecoderResourceManagerAndroid : public IElectraDecoderResourceDelegateAndroid
 {
 public:
-	static bool Startup();
-	static void Shutdown();
 	static TSharedPtr<IElectraDecoderResourceDelegate, ESPMode::ThreadSafe> GetDelegate();
 
 	static bool SetupRenderBufferFromDecoderOutput(FString& OutErrorMessage, IMediaRenderer::IBuffer* InOutBufferToSetup, TSharedPtr<FParamDict, ESPMode::ThreadSafe> InOutBufferPropertes, TSharedPtr<IElectraDecoderVideoOutput, ESPMode::ThreadSafe> InDecoderOutput, IDecoderPlatformResource* InPlatformSpecificResource);
 
-	virtual ~FElectraDecoderResourceManagerAndroid();
-	IDecoderPlatformResource* CreatePlatformResource(void* InOwnerHandle, EDecoderPlatformResourceType InDecoderResourceType, const TMap<FString, FVariant> InOptions) override;
-	void ReleasePlatformResource(void* InOwnerHandle, IDecoderPlatformResource* InHandleToDestroy) override;
-
+	virtual ~FElectraDecoderResourceManagerAndroid() = default;
 private:
-	class FInstanceVars;
-	FCriticalSection Lock;
-	TSet<FInstanceVars*> InstanceVars;
-	static void ReleaseToSurface(uint32 InID, const FDecoderTimeStamp& Time);
 };
 
 using FPlatformElectraDecoderResourceManager = FElectraDecoderResourceManagerAndroid;
diff --git a/Engine/Plugins/Media/ElectraPlayer/Source/ElectraPlayerRuntime/Private/Runtime/Decoder/Android/ElectraVideoDecoder_Android.cpp b/Engine/Plugins/Media/ElectraPlayer/Source/ElectraPlayerRuntime/Private/Runtime/Decoder/Android/ElectraVideoDecoder_Android.cpp
--- a/Engine/Plugins/Media/ElectraPlayer/Source/ElectraPlayerRuntime/Private/Runtime/Decoder/Android/ElectraVideoDecoder_Android.cpp
+++ b/Engine/Plugins/Media/ElectraPlayer/Source/ElectraPlayerRuntime/Private/Runtime/Decoder/Android/ElectraVideoDecoder_Android.cpp
@@ -7,19 +7,10 @@
 namespace Electra
 {
 
-void FElectraPlayerVideoDecoderOutputAndroid::Initialize(EOutputType InOutputType, int32 InBufferIndex, int32 InValidCount, const TFunction<void(uint32, const FDecoderTimeStamp&)>& InSurfaceReleaseFN, uint32 InNativeDecoderID, TSharedPtr<FParamDict, ESPMode::ThreadSafe> InParamDict)
+void FElectraPlayerVideoDecoderOutputAndroid::Initialize(EOutputType InOutputType, TSharedPtr<FParamDict, ESPMode::ThreadSafe> InParamDict)
 {
 	FVideoDecoderOutputAndroid::Initialize(MoveTemp(InParamDict));
-	SurfaceReleaseFN = InSurfaceReleaseFN;
 	OutputType = InOutputType;
-	BufferIndex = InBufferIndex;
-	ValidCount = InValidCount;
-	NativeDecoderID = InNativeDecoderID;
-}
-
-void FElectraPlayerVideoDecoderOutputAndroid::ReleaseToSurface() const
-{
-	SurfaceReleaseFN(NativeDecoderID, GetTime());
 }
 
 void FElectraPlayerVideoDecoderOutputAndroid::SetOwner(const TSharedPtr<IDecoderOutputOwner, ESPMode::ThreadSafe>& InOwningRenderer)
diff --git a/Engine/Plugins/Media/ElectraPlayer/Source/ElectraPlayerRuntime/Private/Runtime/Decoder/Android/ElectraVideoDecoder_Android.h b/Engine/Plugins/Media/ElectraPlayer/Source/ElectraPlayerRuntime/Private/Runtime/Decoder/Android/ElectraVideoDecoder_Android.h
--- a/Engine/Plugins/Media/ElectraPlayer/Source/ElectraPlayerRuntime/Private/Runtime/Decoder/Android/ElectraVideoDecoder_Android.h
+++ b/Engine/Plugins/Media/ElectraPlayer/Source/ElectraPlayerRuntime/Private/Runtime/Decoder/Android/ElectraVideoDecoder_Android.h
@@ -20,11 +20,7 @@ public:
 	FElectraPlayerVideoDecoderOutputAndroid() = default;
 	~FElectraPlayerVideoDecoderOutputAndroid() = default;
 
-	static void ReleaseToSurface(uint32 NativeDecoderID, const FDecoderTimeStamp& Time);
-
-	void Initialize(EOutputType InOutputType, int32 InBufferIndex, int32 InValidCount, const TFunction<void(uint32, const FDecoderTimeStamp&)>& InSurfaceReleaseFN, uint32 InNativeDecoderID, TSharedPtr<FParamDict, ESPMode::ThreadSafe> InParamDict);
-
-	void ReleaseToSurface() const override;
+	void Initialize(EOutputType InOutputType, TSharedPtr<FParamDict, ESPMode::ThreadSafe> InParamDict);
 
 	void SetOwner(const TSharedPtr<IDecoderOutputOwner, ESPMode::ThreadSafe>& InOwningRenderer) override;
 
@@ -33,21 +29,9 @@ public:
 	virtual EOutputType GetOutputType() const override;
 
 private:
-	// Decoder specific surface release function
-	TFunction<void(uint32, const FDecoderTimeStamp&)> SurfaceReleaseFN;
-
 	// Decoder output type
 	EOutputType OutputType = EOutputType::Unknown;
 
-	// Decoder output buffer index associated with this sample (Surface type)
-	int32 BufferIndex = -1;
-
-	// Valid count to check for decoder Android changes within a single logical decoder (Surface type)
-	int32 ValidCount = -1;
-
-	// ID of native decoder
-	uint32 NativeDecoderID = 0;
-
 	// We hold a weak reference to the video renderer. During destruction the video renderer could be destroyed while samples are still out there..
 	TWeakPtr<IDecoderOutputOwner, ESPMode::ThreadSafe> OwningRenderer;
 };
diff --git a/Engine/Plugins/Media/ElectraPlayer/Source/ElectraPlayerRuntime/Private/Runtime/Decoder/Apple/AppleElectraDecoderResourceManager.cpp b/Engine/Plugins/Media/ElectraPlayer/Source/ElectraPlayerRuntime/Private/Runtime/Decoder/Apple/AppleElectraDecoderResourceManager.cpp
--- a/Engine/Plugins/Media/ElectraPlayer/Source/ElectraPlayerRuntime/Private/Runtime/Decoder/Apple/AppleElectraDecoderResourceManager.cpp
+++ b/Engine/Plugins/Media/ElectraPlayer/Source/ElectraPlayerRuntime/Private/Runtime/Decoder/Apple/AppleElectraDecoderResourceManager.cpp
@@ -17,21 +17,6 @@
 namespace Electra
 {
 
-namespace AppleDecoderResources
-{
-	static FElectraDecoderResourceManagerApple::FCallbacks Callbacks;
-}
-
-
-bool FElectraDecoderResourceManagerApple::Startup(const FCallbacks& InCallbacks)
-{
-	AppleDecoderResources::Callbacks = InCallbacks;
-	return AppleDecoderResources::Callbacks.GetMetalDevice ? true : false;
-}
-
-void FElectraDecoderResourceManagerApple::Shutdown()
-{
-}
 
 TSharedPtr<IElectraDecoderResourceDelegate, ESPMode::ThreadSafe> FElectraDecoderResourceManagerApple::GetDelegate()
 {
@@ -40,50 +25,6 @@ TSharedPtr<IElectraDecoderResourceDelegate, ESPMode::ThreadSafe> FElectraDecoder
 }
 
 
-FElectraDecoderResourceManagerApple::~FElectraDecoderResourceManagerApple()
-{
-}
-
-class FElectraDecoderResourceManagerApple::FInstanceVars : public FElectraDecoderResourceManagerApple::IDecoderPlatformResource
-{
-public:
-	uint32 Codec4CC = 0;
-};
-
-
-
-IElectraDecoderResourceDelegateApple::IDecoderPlatformResource* FElectraDecoderResourceManagerApple::CreatePlatformResource(void* InOwnerHandle, EDecoderPlatformResourceType InDecoderResourceType, const TMap<FString, FVariant> InOptions)
-{
-	FInstanceVars* Vars = new FInstanceVars;
-	check(Vars);
-	if (Vars)
-	{
-		uint32 Codec4CC = (uint32)ElectraDecodersUtil::GetVariantValueSafeU64(InOptions, TEXT("codec_4cc"), 0);
-		Vars->Codec4CC = Codec4CC;
-	}
-	return Vars;
-}
-
-void FElectraDecoderResourceManagerApple::ReleasePlatformResource(void* InOwnerHandle, IDecoderPlatformResource* InHandleToDestroy)
-{
-	FInstanceVars* Vars = static_cast<FInstanceVars*>(InHandleToDestroy);
-	if (Vars)
-	{
-		delete Vars;
-	}
-}
-
-
-bool FElectraDecoderResourceManagerApple::GetMetalDevice(void **OutMetalDevice)
-{
-	check(OutMetalDevice);
-	if (AppleDecoderResources::Callbacks.GetMetalDevice)
-	{
-		return AppleDecoderResources::Callbacks.GetMetalDevice(OutMetalDevice, AppleDecoderResources::Callbacks.UserValue);
-	}
-	return false;
-}
-
 
 bool FElectraDecoderResourceManagerApple::SetupRenderBufferFromDecoderOutput(FString& OutErrorMessage, IMediaRenderer::IBuffer* InOutBufferToSetup, TSharedPtr<FParamDict, ESPMode::ThreadSafe> InOutBufferPropertes, TSharedPtr<IElectraDecoderVideoOutput, ESPMode::ThreadSafe> InDecoderOutput, FElectraDecoderResourceManagerApple::IDecoderPlatformResource* InPlatformSpecificResource)
 {
@@ -92,7 +33,6 @@ bool FElectraDecoderResourceManagerApple::SetupRenderBufferFromDecoderOutput(FSt
 	check(InDecoderOutput.IsValid());
 
 	TSharedPtr<FElectraPlayerVideoDecoderOutputApple, ESPMode::ThreadSafe> DecoderOutput = InOutBufferToSetup->GetBufferProperties().GetValue(RenderOptionKeys::Texture).GetSharedPointer<FElectraPlayerVideoDecoderOutputApple>();
-	FInstanceVars* Vars = static_cast<FInstanceVars*>(InPlatformSpecificResource);
 
 	if (DecoderOutput.IsValid())
 	{
@@ -184,13 +124,12 @@ bool FElectraDecoderResourceManagerApple::SetupRenderBufferFromDecoderOutput(FSt
 			}
 			else
 			{
-				// Note: we assume a DX11 texture at this point, but this could also be interpreted as IUknown and then either DX11 or DX12 resources being derived from it...
+				//
+				// GPU texture (CVImageBuffer)
+				//
 				CVImageBufferRef ImageBuffer = static_cast<CVImageBufferRef>(ImageBuffers->GetBufferTextureByIndex(0));
 				if (ImageBuffer != nullptr)
 				{
-					//
-					// GPU texture (CVImageBuffer)
-					//
 					DecoderOutput->Initialize(ImageBuffer, InOutBufferPropertes);
 					return true;
 				}
diff --git a/Engine/Plugins/Media/ElectraPlayer/Source/ElectraPlayerRuntime/Private/Runtime/Decoder/Apple/AppleElectraDecoderResourceManager.h b/Engine/Plugins/Media/ElectraPlayer/Source/ElectraPlayerRuntime/Private/Runtime/Decoder/Apple/AppleElectraDecoderResourceManager.h
--- a/Engine/Plugins/Media/ElectraPlayer/Source/ElectraPlayerRuntime/Private/Runtime/Decoder/Apple/AppleElectraDecoderResourceManager.h
+++ b/Engine/Plugins/Media/ElectraPlayer/Source/ElectraPlayerRuntime/Private/Runtime/Decoder/Apple/AppleElectraDecoderResourceManager.h
@@ -1,17 +1,13 @@
 // Copyright Epic Games, Inc. All Rights Reserved.
-
 #pragma once
 
 #include "IElectraDecoderResourceDelegate.h"
 
-#include "HAL/Platform.h"
-
 #if PLATFORM_MAC || PLATFORM_IOS || PLATFORM_TVOS
 
 #include "Renderer/RendererBase.h"
 
 class IElectraDecoderVideoOutput;
-struct FDecoderTimeStamp;
 
 namespace Electra
 {
@@ -19,27 +15,11 @@ namespace Electra
 class FElectraDecoderResourceManagerApple : public IElectraDecoderResourceDelegateApple
 {
 public:
-	struct FCallbacks
-	{
-		TFunction<bool(void ** /*OutMetalDevice*/, void* /*UserValue*/)> GetMetalDevice;
-		void* UserValue = nullptr;
-	};
-
-	static bool Startup(const FCallbacks& InCallbacks);
-	static void Shutdown();
 	static TSharedPtr<IElectraDecoderResourceDelegate, ESPMode::ThreadSafe> GetDelegate();
 
 	static bool SetupRenderBufferFromDecoderOutput(FString& OutErrorMessage, IMediaRenderer::IBuffer* InOutBufferToSetup, TSharedPtr<FParamDict, ESPMode::ThreadSafe> InOutBufferPropertes, TSharedPtr<IElectraDecoderVideoOutput, ESPMode::ThreadSafe> InDecoderOutput, IDecoderPlatformResource* InPlatformSpecificResource);
 
-	IDecoderPlatformResource* CreatePlatformResource(void* InOwnerHandle, EDecoderPlatformResourceType InDecoderResourceType, const TMap<FString, FVariant> InOptions) override;
-	void ReleasePlatformResource(void* InOwnerHandle, IDecoderPlatformResource* InHandleToDestroy) override;
-
-	virtual bool GetMetalDevice(void **OutMetalDevice) override;
-
-	virtual ~FElectraDecoderResourceManagerApple();
-
-private:
-	class FInstanceVars;
+	virtual ~FElectraDecoderResourceManagerApple() = default;
 };
 
 using FPlatformElectraDecoderResourceManager = FElectraDecoderResourceManagerApple;
diff --git a/Engine/Plugins/Media/ElectraPlayer/Source/ElectraPlayerRuntime/Private/Runtime/Decoder/AudioDecoder.cpp b/Engine/Plugins/Media/ElectraPlayer/Source/ElectraPlayerRuntime/Private/Runtime/Decoder/AudioDecoder.cpp
--- a/Engine/Plugins/Media/ElectraPlayer/Source/ElectraPlayerRuntime/Private/Runtime/Decoder/AudioDecoder.cpp
+++ b/Engine/Plugins/Media/ElectraPlayer/Source/ElectraPlayerRuntime/Private/Runtime/Decoder/AudioDecoder.cpp
@@ -490,7 +490,7 @@ bool FAudioDecoderImpl::InternalDecoderCreate()
 	{
 		return PostError(-2, TEXT("No decoder factory found to create an audio decoder"), ERRCODE_AUDIO_INTERNAL_COULD_NOT_CREATE_DECODER);
 	}
-	DecoderInstance = Factory->CreateDecoderForFormat(Format, AddtlCfg, FPlatformElectraDecoderResourceManager::GetDelegate());
+	DecoderInstance = Factory->CreateDecoderForFormat(Format, AddtlCfg);
 	if (!DecoderInstance.IsValid() || DecoderInstance->GetError().IsSet())
 	{
 		InternalDecoderDestroy();
diff --git a/Engine/Plugins/Media/ElectraPlayer/Source/ElectraPlayerRuntime/Private/Runtime/Decoder/Linux/LinuxElectraDecoderResourceManager.cpp b/Engine/Plugins/Media/ElectraPlayer/Source/ElectraPlayerRuntime/Private/Runtime/Decoder/Linux/LinuxElectraDecoderResourceManager.cpp
--- a/Engine/Plugins/Media/ElectraPlayer/Source/ElectraPlayerRuntime/Private/Runtime/Decoder/Linux/LinuxElectraDecoderResourceManager.cpp
+++ b/Engine/Plugins/Media/ElectraPlayer/Source/ElectraPlayerRuntime/Private/Runtime/Decoder/Linux/LinuxElectraDecoderResourceManager.cpp
@@ -17,15 +17,6 @@
 namespace Electra
 {
 
-bool FElectraDecoderResourceManagerLinux::Startup()
-{
-	return true;
-}
-
-void FElectraDecoderResourceManagerLinux::Shutdown()
-{
-}
-
 TSharedPtr<IElectraDecoderResourceDelegate, ESPMode::ThreadSafe> FElectraDecoderResourceManagerLinux::GetDelegate()
 {
 	static TSharedPtr<IElectraDecoderResourceDelegate, ESPMode::ThreadSafe> Self = MakeShared<FElectraDecoderResourceManagerLinux, ESPMode::ThreadSafe>();
@@ -33,9 +24,6 @@ TSharedPtr<IElectraDecoderResourceDelegate, ESPMode::ThreadSafe> FElectraDecoder
 }
 
 
-FElectraDecoderResourceManagerLinux::~FElectraDecoderResourceManagerLinux()
-{
-}
 
 #if WITH_LIBAV
 namespace LibavcodecConversion
diff --git a/Engine/Plugins/Media/ElectraPlayer/Source/ElectraPlayerRuntime/Private/Runtime/Decoder/Linux/LinuxElectraDecoderResourceManager.h b/Engine/Plugins/Media/ElectraPlayer/Source/ElectraPlayerRuntime/Private/Runtime/Decoder/Linux/LinuxElectraDecoderResourceManager.h
--- a/Engine/Plugins/Media/ElectraPlayer/Source/ElectraPlayerRuntime/Private/Runtime/Decoder/Linux/LinuxElectraDecoderResourceManager.h
+++ b/Engine/Plugins/Media/ElectraPlayer/Source/ElectraPlayerRuntime/Private/Runtime/Decoder/Linux/LinuxElectraDecoderResourceManager.h
@@ -11,7 +11,6 @@
 #include "Renderer/RendererBase.h"
 
 class IElectraDecoderVideoOutput;
-struct FDecoderTimeStamp;
 
 namespace Electra
 {
@@ -19,14 +18,11 @@ namespace Electra
 class FElectraDecoderResourceManagerLinux : public IElectraDecoderResourceDelegateLinux
 {
 public:
-
-	static bool Startup();
-	static void Shutdown();
 	static TSharedPtr<IElectraDecoderResourceDelegate, ESPMode::ThreadSafe> GetDelegate();
 
 	static bool SetupRenderBufferFromDecoderOutput(FString& OutErrorMessage, IMediaRenderer::IBuffer* InOutBufferToSetup, TSharedPtr<FParamDict, ESPMode::ThreadSafe> InOutBufferPropertes, TSharedPtr<IElectraDecoderVideoOutput, ESPMode::ThreadSafe> InDecoderOutput, IDecoderPlatformResource* InPlatformSpecificResource);
 
-	virtual ~FElectraDecoderResourceManagerLinux();
+	virtual ~FElectraDecoderResourceManagerLinux() = default;
 };
 
 using FPlatformElectraDecoderResourceManager = FElectraDecoderResourceManagerLinux;
diff --git a/Engine/Plugins/Media/ElectraPlayer/Source/ElectraPlayerRuntime/Private/Runtime/Decoder/VideoDecoder.cpp b/Engine/Plugins/Media/ElectraPlayer/Source/ElectraPlayerRuntime/Private/Runtime/Decoder/VideoDecoder.cpp
--- a/Engine/Plugins/Media/ElectraPlayer/Source/ElectraPlayerRuntime/Private/Runtime/Decoder/VideoDecoder.cpp
+++ b/Engine/Plugins/Media/ElectraPlayer/Source/ElectraPlayerRuntime/Private/Runtime/Decoder/VideoDecoder.cpp
@@ -56,7 +56,6 @@ public:
 	bool Reopen(TSharedPtrTS<FAccessUnit::CodecData> InCodecData, const FParamDict& InAdditionalOptions, const FStreamCodecInformation* InMaxStreamConfiguration) override;
 	void Close() override;
 	void DrainForCodecChange() override;
-	void SetVideoResourceDelegate(TWeakPtr<IVideoDecoderResourceDelegate, ESPMode::ThreadSafe> InVideoResourceDelegate) override;
 	void SetRenderer(TSharedPtr<IMediaRenderer, ESPMode::ThreadSafe> InRenderer) override;
 	void SuspendOrResumeDecoder(bool bSuspend, const FParamDict& InOptions) override;
 	void AUdataPushAU(FAccessUnit* AccessUnit) override;
@@ -193,7 +192,6 @@ private:
 	TSharedPtrTS<FFGBGNotificationHandlers>									FGBGHandlers;
 	int32																	ApplicationSuspendCount = 0;
 
-	TWeakPtr<IVideoDecoderResourceDelegate, ESPMode::ThreadSafe>			VideoResourceDelegate;
 	TSharedPtr<IMediaRenderer, ESPMode::ThreadSafe>							Renderer;
 	int32																	MaxOutputBuffers = 0;
 
@@ -375,11 +373,6 @@ void FVideoDecoderImpl::DrainForCodecChange()
 	bDrainForCodecChange = true;
 }
 
-void FVideoDecoderImpl::SetVideoResourceDelegate(TWeakPtr<IVideoDecoderResourceDelegate, ESPMode::ThreadSafe> InVideoResourceDelegate)
-{
-	VideoResourceDelegate = InVideoResourceDelegate;
-}
-
 void FVideoDecoderImpl::SetRenderer(TSharedPtr<IMediaRenderer, ESPMode::ThreadSafe> InRenderer)
 {
 	Renderer = InRenderer;
@@ -563,22 +556,14 @@ bool FVideoDecoderImpl::InternalDecoderCreate()
 	}
 
 	// Create platform specifics.
-	TSharedPtr<IVideoDecoderResourceDelegate, ESPMode::ThreadSafe> PinnedVideoResourceDelegate = VideoResourceDelegate.Pin();
 	TMap<FString, FVariant> PlatformSpecificCfg(DecoderFactoryAddtlCfg);
-	PlatformSpecificCfg.Emplace(TEXT("VideoResourceDelegate"), FVariant((uint64)PinnedVideoResourceDelegate.Get()));
 	PlatformResource = FPlatformElectraDecoderResourceManager::GetDelegate()->CreatePlatformResource(this, IElectraDecoderResourceDelegateBase::EDecoderPlatformResourceType::Video, PlatformSpecificCfg);
 
-	// Put a pointer to the renderer into the decoder creation configuration.
-	// The decoder itself does not need it but it will pass this into the resource manager when creating a resource handler instance.
-	// That way the resource handler gets the pointer to make calls to AcquireBuffer() if necessary.
-	TMap<FString, FVariant> DecoderCreateCfg(DecoderFactoryAddtlCfg);
-	DecoderCreateCfg.Emplace(TEXT("renderer"), FVariant((uint64)Renderer.Get()));
-	DecoderCreateCfg.Emplace(TEXT("platformresource"), FVariant((uint64)PlatformResource));
-
 	// Add in video decoder special options passed from the application.
-	InitialAdditionalOptions.ConvertKeysStartingWithTo(DecoderCreateCfg, TEXT("videoDecoder"), FString());
+	TMap<FString, FVariant> DecoderCreateCfg(DecoderFactoryAddtlCfg);
+	InitialAdditionalOptions.ConvertKeysStartingWithTo(DecoderCreateCfg, TEXT("output_texturepool_id"), FString());
 
-	DecoderInstance = DecoderFactory->CreateDecoderForFormat(DecoderFormat, DecoderCreateCfg, FPlatformElectraDecoderResourceManager::GetDelegate());
+	DecoderInstance = DecoderFactory->CreateDecoderForFormat(DecoderFormat, DecoderCreateCfg);
 	if (!DecoderInstance.IsValid() || DecoderInstance->GetError().IsSet())
 	{
 		InternalDecoderDestroy();
@@ -587,7 +572,7 @@ bool FVideoDecoderImpl::InternalDecoderCreate()
 	if (DecoderInstance->GetType() != IElectraDecoder::EType::Video)
 	{
 		InternalDecoderDestroy();
-		return PostError(-2, TEXT("Created decoder is not an video decoder!"), ERRCODE_VIDEO_INTERNAL_COULD_NOT_CREATE_DECODER);
+		return PostError(-2, TEXT("Created decoder is not a video decoder!"), ERRCODE_VIDEO_INTERNAL_COULD_NOT_CREATE_DECODER);
 	}
 
 	TMap<FString, FVariant> Features;
@@ -909,39 +894,24 @@ IElectraDecoder::EOutputStatus FVideoDecoderImpl::HandleOutput()
 				return IElectraDecoder::EOutputStatus::Error;
 			}
 
-			// Check if the output has a "transfer buffer".
-			// If it does, then we know that this is actually a buffer of the renderer that was acquired by the decoder
-			// through the platform's resource manager, which has been implemented alongside this decoder implementation.
-			if (CurrentDecoderOutput->GetTransferHandle() == nullptr)
+			// Need a new output buffer?
+			if (CurrentOutputBuffer == nullptr && Renderer.IsValid())
 			{
-				// Need a new output buffer?
-				if (CurrentOutputBuffer == nullptr && Renderer.IsValid())
-				{
-					SCOPE_CYCLE_COUNTER(STAT_ElectraPlayer_VideoConvertOutput);
-					CSV_SCOPED_TIMING_STAT(ElectraPlayer, VideoConvertOutput);
-					UEMediaError bufResult = Renderer->AcquireBuffer(CurrentOutputBuffer, 0, EmptyOptions);
-					check(bufResult == UEMEDIA_ERROR_OK || bufResult == UEMEDIA_ERROR_INSUFFICIENT_DATA);
-					if (bufResult != UEMEDIA_ERROR_OK && bufResult != UEMEDIA_ERROR_INSUFFICIENT_DATA)
-					{
-						PostError(0, TEXT("Failed to acquire sample buffer"), ERRCODE_VIDEO_INTERNAL_COULD_NOT_GET_SAMPLE_BUFFER, bufResult);
-						return IElectraDecoder::EOutputStatus::Error;
-					}
-				}
-				// Didn't get a buffer? This should not really happen since the renderer said it could accept a frame.
-				if (!CurrentOutputBuffer)
+				SCOPE_CYCLE_COUNTER(STAT_ElectraPlayer_VideoConvertOutput);
+				CSV_SCOPED_TIMING_STAT(ElectraPlayer, VideoConvertOutput);
+				UEMediaError bufResult = Renderer->AcquireBuffer(CurrentOutputBuffer, 0, EmptyOptions);
+				check(bufResult == UEMEDIA_ERROR_OK || bufResult == UEMEDIA_ERROR_INSUFFICIENT_DATA);
+				if (bufResult != UEMEDIA_ERROR_OK && bufResult != UEMEDIA_ERROR_INSUFFICIENT_DATA)
 				{
-					NotifyReadyBufferListener(false);
-					return IElectraDecoder::EOutputStatus::TryAgainLater;
+					PostError(0, TEXT("Failed to acquire sample buffer"), ERRCODE_VIDEO_INTERNAL_COULD_NOT_GET_SAMPLE_BUFFER, bufResult);
+					return IElectraDecoder::EOutputStatus::Error;
 				}
 			}
-			else
+			// Didn't get a buffer? This should not really happen since the renderer said it could accept a frame.
+			if (!CurrentOutputBuffer)
 			{
-				// If we did not return the last buffer to the renderer, do it now.
-				// Transfer handles are unique to the buffer and cannot be used for different frames.
-				ReturnUnusedOutputBuffer();
-				// Get the transfer handle as the current decoder output.
-				CurrentOutputBuffer = reinterpret_cast<IMediaRenderer::IBuffer*>(CurrentDecoderOutput->GetTransferHandle()->GetHandle());
-				CurrentDecoderOutput->GetTransferHandle()->ReleaseHandle();
+				NotifyReadyBufferListener(false);
+				return IElectraDecoder::EOutputStatus::TryAgainLater;
 			}
 
 			// Check if the output can actually be output or if the decoder says this is not to be output (incorrectly decoded)
diff --git a/Engine/Plugins/Media/ElectraPlayer/Source/ElectraPlayerRuntime/Private/Runtime/Decoder/VideoDecoder.h b/Engine/Plugins/Media/ElectraPlayer/Source/ElectraPlayerRuntime/Private/Runtime/Decoder/VideoDecoder.h
--- a/Engine/Plugins/Media/ElectraPlayer/Source/ElectraPlayerRuntime/Private/Runtime/Decoder/VideoDecoder.h
+++ b/Engine/Plugins/Media/ElectraPlayer/Source/ElectraPlayerRuntime/Private/Runtime/Decoder/VideoDecoder.h
@@ -6,7 +6,6 @@
 
 #include "StreamAccessUnitBuffer.h"
 #include "Decoder/DecoderBase.h"
-#include "VideoDecoderResourceDelegate.h"
 
 
 
@@ -31,9 +30,6 @@ namespace Electra
 		virtual void Close() = 0;
 		virtual void DrainForCodecChange() = 0;
 
-		//! Sets a platform specific video resource delegate needed to handle decoder output.
-		virtual void SetVideoResourceDelegate(TWeakPtr<IVideoDecoderResourceDelegate, ESPMode::ThreadSafe> InVideoResourceDelegate) = 0;
-
 		//-------------------------------------------------------------------------
 		// Methods from IDecoderBase
 		//
diff --git a/Engine/Plugins/Media/ElectraPlayer/Source/ElectraPlayerRuntime/Private/Runtime/Decoder/Windows/ElectraVideoDecoder_PC.cpp b/Engine/Plugins/Media/ElectraPlayer/Source/ElectraPlayerRuntime/Private/Runtime/Decoder/Windows/ElectraVideoDecoder_PC.cpp
--- a/Engine/Plugins/Media/ElectraPlayer/Source/ElectraPlayerRuntime/Private/Runtime/Decoder/Windows/ElectraVideoDecoder_PC.cpp
+++ b/Engine/Plugins/Media/ElectraPlayer/Source/ElectraPlayerRuntime/Private/Runtime/Decoder/Windows/ElectraVideoDecoder_PC.cpp
@@ -25,6 +25,7 @@ THIRD_PARTY_INCLUDES_END
 #include "Windows/HideWindowsPlatformTypes.h"
 
 #include "Async/Async.h"
+#include "ID3D12DynamicRHI.h"
 
 // ------------------------------------------------------------------------------------
 
@@ -147,7 +148,7 @@ void FElectraPlayerVideoDecoderOutputPC::InitializeWithBuffer(TSharedPtr<TArray<
 	D3DFence = nullptr;
 }
 
-void FElectraPlayerVideoDecoderOutputPC::InitializeWithResource(const TRefCountPtr<ID3D12Device>& InD3D12Device, const TRefCountPtr<ID3D12Resource> InResource, uint32 ResourcePitch, const FElectraDecoderOutputSync& OutputSync, const FIntPoint& InOutputDim, TSharedPtr<Electra::FParamDict, ESPMode::ThreadSafe> InParamDict, TWeakPtr<Electra::IVideoDecoderResourceDelegate, ESPMode::ThreadSafe> InResourceDelegate,
+void FElectraPlayerVideoDecoderOutputPC::InitializeWithResource(const TRefCountPtr<ID3D12Device>& InD3D12Device, const TRefCountPtr<ID3D12Resource> InResource, uint32 ResourcePitch, const FElectraDecoderOutputSync& OutputSync, const FIntPoint& InOutputDim, TSharedPtr<Electra::FParamDict, ESPMode::ThreadSafe> InParamDict,
 																TSharedPtr<FElectraMediaDecoderOutputBufferPool_DX12>& InOutD3D12ResourcePool, uint32 MaxWidth, uint32 MaxHeight, uint32 MaxOutputBuffers)
 {
 	// We must not allow re-initialization of used instances
@@ -381,32 +382,26 @@ void FElectraPlayerVideoDecoderOutputPC::InitializeWithResource(const TRefCountP
 	bool bTriggerOk = false;
 	if (OutputSync.TaskSync.IsValid())
 	{
-		// Yes. Trigger the GPU copy once the CPU is ready, so we can avoid any stalls on the copy-queue waiting for decoder data...
-		auto ElectraDecoderResourceDelegate = Electra::FElectraDecoderResourceManagerWindows::GetDelegate();
-
 		// Note: we capture "this" as we ensure that this instance only dies once the copy triggered here is actually done, hence ensuring any reference to "this" is done
-		bTriggerOk = ElectraDecoderResourceDelegate->RunCodeAsync([D3DCmdList=this->D3DCmdList, D3DCmdAllocator=this->D3DCmdAllocator, D3DFence=this->D3DFence, FenceValue=this->FenceValue, OutputSync, ResourceDelegate = InResourceDelegate.Pin()]()
+		bTriggerOk = FElectraDecodersPlatformResources::RunCodeAsync([D3DCmdList=this->D3DCmdList, D3DCmdAllocator=this->D3DCmdAllocator, D3DFence=this->D3DFence, FenceValue=this->FenceValue, OutputSync]()
 			{
-				TriggerDataCopy(D3DCmdList, D3DCmdAllocator, D3DFence, FenceValue, OutputSync, ResourceDelegate.Get());
-			}, OutputSync.TaskSync.Get());
+				TriggerDataCopy(D3DCmdList, D3DCmdAllocator, D3DFence, FenceValue, OutputSync);
+			}, OutputSync.TaskSync);
 	}
 
 	if (!bTriggerOk)
 	{
 		// We could not run the trigger async. Schedule the copy right away. Any needed synchronization will be done in the copy-queue by the GPU
-		TriggerDataCopy(D3DCmdList, D3DCmdAllocator, D3DFence, FenceValue, OutputSync, InResourceDelegate.Pin().Get());
+		TriggerDataCopy(D3DCmdList, D3DCmdAllocator, D3DFence, FenceValue, OutputSync);
 	}
 }
 
-void FElectraPlayerVideoDecoderOutputPC::TriggerDataCopy(TRefCountPtr<ID3D12GraphicsCommandList> D3DCmdList, TRefCountPtr<ID3D12CommandAllocator> D3DCmdAllocator, TRefCountPtr<ID3D12Fence> D3DFence, uint64 FenceValue, const FElectraDecoderOutputSync& OutputSync, Electra::IVideoDecoderResourceDelegate* InResourceDelegate)
+void FElectraPlayerVideoDecoderOutputPC::TriggerDataCopy(TRefCountPtr<ID3D12GraphicsCommandList> D3DCmdList, TRefCountPtr<ID3D12CommandAllocator> D3DCmdAllocator, TRefCountPtr<ID3D12Fence> D3DFence, uint64 FenceValue, const FElectraDecoderOutputSync& OutputSync)
 {
-	if (!InResourceDelegate)
-	{
-		return;
-	}
 	// Trigger copy (this will eventually execute on the submission thread of RHI if running in UE)
 	// (note: we pass in all of FElectraDecoderOutputSync to guarantee any references needed to make the decoder output sync work are passed along, too!)
-	InResourceDelegate->ExecuteCodeWithCopyCommandQueueUsage([CmdList = D3DCmdList, CmdAllocator = D3DCmdAllocator, DestFence = D3DFence, DestFenceValue = FenceValue, OutputSync](ID3D12CommandQueue* D3DCmdQueue)
+	GetID3D12PlatformDynamicRHI()->RHIRunOnQueue(ED3D12RHIRunOnQueueType::Copy,
+		[CmdList = D3DCmdList, CmdAllocator = D3DCmdAllocator, DestFence = D3DFence, DestFenceValue = FenceValue, OutputSync](ID3D12CommandQueue* D3DCmdQueue)
 		{
 			TRefCountPtr<ID3D12Fence> ResourceFence;
 	#if ALLOW_MFSAMPLE_WITH_DX12
@@ -455,7 +450,7 @@ void FElectraPlayerVideoDecoderOutputPC::TriggerDataCopy(TRefCountPtr<ID3D12Grap
 
 			// Notify user of the output data we just copied of its arrival
 			D3DCmdQueue->Signal(DestFence, DestFenceValue);
-		});
+		}, false);
 }
 
 void FElectraPlayerVideoDecoderOutputPC::InitializeWithSharedTexture(const TRefCountPtr<ID3D11Device>& InD3D11Device, const TRefCountPtr<ID3D11DeviceContext> InDeviceContext, const TRefCountPtr<IMFSample> InMFSample, const FIntPoint& OutputDim, TSharedPtr<Electra::FParamDict, ESPMode::ThreadSafe> InParamDict)
diff --git a/Engine/Plugins/Media/ElectraPlayer/Source/ElectraPlayerRuntime/Private/Runtime/Decoder/Windows/ElectraVideoDecoder_PC.h b/Engine/Plugins/Media/ElectraPlayer/Source/ElectraPlayerRuntime/Private/Runtime/Decoder/Windows/ElectraVideoDecoder_PC.h
--- a/Engine/Plugins/Media/ElectraPlayer/Source/ElectraPlayerRuntime/Private/Runtime/Decoder/Windows/ElectraVideoDecoder_PC.h
+++ b/Engine/Plugins/Media/ElectraPlayer/Source/ElectraPlayerRuntime/Private/Runtime/Decoder/Windows/ElectraVideoDecoder_PC.h
@@ -5,7 +5,6 @@
 #include "PlayerCore.h"
 #include "ElectraPlayerPrivate.h"
 #include "MediaVideoDecoderOutputPC.h"
-#include "VideoDecoderResourceDelegate.h"
 
 /*********************************************************************************************************************/
 /*********************************************************************************************************************/
@@ -52,7 +51,7 @@ public:
 
 	void InitializeWithBuffer(TSharedPtr<TArray<uint8>, ESPMode::ThreadSafe> InBuffer, uint32 InStride, FIntPoint Dim, TSharedPtr<Electra::FParamDict, ESPMode::ThreadSafe> InParamDict);
 
-	void InitializeWithResource(const TRefCountPtr<ID3D12Device>& InD3D12Device, const TRefCountPtr<ID3D12Resource> Resource, uint32 ResourcePitch, const FElectraDecoderOutputSync& OutputSync, const FIntPoint& OutputDim, TSharedPtr<Electra::FParamDict, ESPMode::ThreadSafe> InParamDict, TWeakPtr<Electra::IVideoDecoderResourceDelegate, ESPMode::ThreadSafe> ResourceDelegate,
+	void InitializeWithResource(const TRefCountPtr<ID3D12Device>& InD3D12Device, const TRefCountPtr<ID3D12Resource> Resource, uint32 ResourcePitch, const FElectraDecoderOutputSync& OutputSync, const FIntPoint& OutputDim, TSharedPtr<Electra::FParamDict, ESPMode::ThreadSafe> InParamDict,
 								TSharedPtr<FElectraMediaDecoderOutputBufferPool_DX12>& InOutD3D12ResourcePool, uint32 MaxWidth, uint32 MaxHeight, uint32 MaxOutputBuffers);
 
 	// Hardware decode to shared DX11 texture (Win8+) from IMFSample
@@ -73,7 +72,7 @@ public:
 	TRefCountPtr<IUnknown> GetSync(uint64& SyncValue) const override;
 
 private:
-	static void TriggerDataCopy(TRefCountPtr<ID3D12GraphicsCommandList> D3DCmdList, TRefCountPtr<ID3D12CommandAllocator> D3DCmdAllocator, TRefCountPtr<ID3D12Fence> D3DFence, uint64 FenceValue, const FElectraDecoderOutputSync& OutputSync, Electra::IVideoDecoderResourceDelegate* InResourceDelegate);
+	static void TriggerDataCopy(TRefCountPtr<ID3D12GraphicsCommandList> D3DCmdList, TRefCountPtr<ID3D12CommandAllocator> D3DCmdAllocator, TRefCountPtr<ID3D12Fence> D3DFence, uint64 FenceValue, const FElectraDecoderOutputSync& OutputSync);
 
 	// Decoder output type
 	EOutputType OutputType = EOutputType::Unknown;
