Electra: Removal of the platform resource delegate and the wrapping plugin

diff --git a/Engine/Plugins/Media/ElectraPlayer/Source/ElectraPlayerRuntime/Private/Runtime/Decoder/Windows/WindowsElectraDecoderResourceManager.cpp b/Engine/Plugins/Media/ElectraPlayer/Source/ElectraPlayerRuntime/Private/Runtime/Decoder/Windows/WindowsElectraDecoderResourceManager.cpp
--- a/Engine/Plugins/Media/ElectraPlayer/Source/ElectraPlayerRuntime/Private/Runtime/Decoder/Windows/WindowsElectraDecoderResourceManager.cpp
+++ b/Engine/Plugins/Media/ElectraPlayer/Source/ElectraPlayerRuntime/Private/Runtime/Decoder/Windows/WindowsElectraDecoderResourceManager.cpp
@@ -6,10 +6,10 @@
 
 #include "IElectraDecoderOutputVideo.h"
 #include "ElectraDecodersUtils.h"
+#include "ElectraDecodersPlatformResources.h"
 
 #include "Renderer/RendererBase.h"
 #include "ElectraVideoDecoder_PC.h"
-#include "VideoDecoderResourceDelegate.h"
 
 #include "ElectraPlayerMisc.h"
 
@@ -21,48 +21,6 @@
 namespace Electra
 {
 
-namespace WindowsDecoderResources
-{
-	static FElectraDecoderResourceManagerWindows::FCallbacks Callbacks;
-	static bool bDidInitializeMF = false;
-}
-
-bool FElectraDecoderResourceManagerWindows::Startup(const FCallbacks& InCallbacks)
-{
-	// Windows 8 or better?
-	if (!FPlatformMisc::VerifyWindowsVersion(6, 2))
-	{
-		UE_LOG(LogElectraPlayer, Log, TEXT("Electra is incompatible with Windows prior to 8.0 version: %s"), *FPlatformMisc::GetOSVersion());
-		return false;
-	}
-	if (!(FPlatformProcess::GetDllHandle(TEXT("mf.dll"))
-		&& FPlatformProcess::GetDllHandle(TEXT("mfplat.dll"))
-		&& FPlatformProcess::GetDllHandle(TEXT("msmpeg2vdec.dll"))
-		&& FPlatformProcess::GetDllHandle(TEXT("MSAudDecMFT.dll"))))
-	{
-		UE_LOG(LogElectraPlayer, Error, TEXT("Electra can't load Media Foundation, %s"), *FPlatformMisc::GetOSVersion());
-		return false;
-	}
-	HRESULT Result = MFStartup(MF_VERSION);
-	if (FAILED(Result))
-	{
-		UE_LOG(LogElectraPlayer, Error, TEXT("MFStartup failed with 0x%08x"), Result);
-		return false;
-	}
-	WindowsDecoderResources::bDidInitializeMF = true;
-	WindowsDecoderResources::Callbacks = InCallbacks;
-	return WindowsDecoderResources::Callbacks.GetD3DDevice ? true : false;
-}
-
-void FElectraDecoderResourceManagerWindows::Shutdown()
-{
-	if (WindowsDecoderResources::bDidInitializeMF)
-	{
-		MFShutdown();
-		WindowsDecoderResources::bDidInitializeMF = false;
-	}
-}
-
 TSharedPtr<IElectraDecoderResourceDelegate, ESPMode::ThreadSafe> FElectraDecoderResourceManagerWindows::GetDelegate()
 {
 	static TSharedPtr<IElectraDecoderResourceDelegate, ESPMode::ThreadSafe> Self = MakeShared<FElectraDecoderResourceManagerWindows, ESPMode::ThreadSafe>();
@@ -78,34 +36,26 @@ FElectraDecoderResourceManagerWindows::~FElectraDecoderResourceManagerWindows()
 class FElectraDecoderResourceManagerWindows::FInstanceVars : public IElectraDecoderResourceDelegateWindows::IDecoderPlatformResource
 {
 public:
-	uint32 Codec4CC = 0;
+	void* D3DDevice = nullptr;
+	int32 D3DVersionTimes1000 = 0;
+
+
 	uint32 MaxWidth = 0;
 	uint32 MaxHeight = 0;
 	uint32 MaxOutputBuffers = 0;
-	TWeakPtr<IVideoDecoderResourceDelegate, ESPMode::ThreadSafe> VideoDecoderResourceDelegate;
 	TSharedPtr<FElectraMediaDecoderOutputBufferPool_DX12> D3D12ResourcePool;
 };
 
 
 IElectraDecoderResourceDelegateWindows::IDecoderPlatformResource* FElectraDecoderResourceManagerWindows::CreatePlatformResource(void* InOwnerHandle, EDecoderPlatformResourceType InDecoderResourceType, const TMap<FString, FVariant> InOptions)
 {
-	IVideoDecoderResourceDelegate* Delegate = reinterpret_cast<IVideoDecoderResourceDelegate*>(ElectraDecodersUtil::GetVariantValueSafeU64(InOptions, TEXT("VideoResourceDelegate"), 0));
-	if (Delegate)
-	{
-		TSharedPtr<IVideoDecoderResourceDelegate, ESPMode::ThreadSafe> VideoDecoderResourceDelegate = Delegate->AsShared();
-		FInstanceVars* Vars = new FInstanceVars();
-		check(Vars);
-		if (Vars)
-		{
-			Vars->VideoDecoderResourceDelegate = VideoDecoderResourceDelegate;
-			Vars->Codec4CC = (uint32)ElectraDecodersUtil::GetVariantValueSafeU64(InOptions, TEXT("codec_4cc"), 0);
-			Vars->MaxWidth  = Align((uint32)ElectraDecodersUtil::GetVariantValueSafeU64(InOptions, TEXT("max_width"), 1920), 16);
-			Vars->MaxHeight = Align((uint32)ElectraDecodersUtil::GetVariantValueSafeU64(InOptions, TEXT("max_height"), 1080), 16);
-			Vars->MaxOutputBuffers = (uint32)ElectraDecodersUtil::GetVariantValueSafeU64(InOptions, TEXT("max_output_buffers"), 5);
-		}
-		return Vars;
-	}
-	return nullptr;
+	FInstanceVars* Vars = new FInstanceVars();
+	FElectraDecodersPlatformResources::GetD3DDeviceAndVersion(&Vars->D3DDevice, &Vars->D3DVersionTimes1000);
+
+	Vars->MaxWidth  = Align((uint32)ElectraDecodersUtil::GetVariantValueSafeU64(InOptions, TEXT("max_width"), 1920), 16);
+	Vars->MaxHeight = Align((uint32)ElectraDecodersUtil::GetVariantValueSafeU64(InOptions, TEXT("max_height"), 1080), 16);
+	Vars->MaxOutputBuffers = (uint32)ElectraDecodersUtil::GetVariantValueSafeU64(InOptions, TEXT("max_output_buffers"), 5);
+	return Vars;
 }
 
 
@@ -119,42 +69,11 @@ void FElectraDecoderResourceManagerWindows::ReleasePlatformResource(void* InOwne
 }
 
 
-bool FElectraDecoderResourceManagerWindows::GetD3DDevice(void **OutD3DDevice, int32* OutD3DVersionTimes1000)
-{
-	check(OutD3DDevice);
-	check(OutD3DVersionTimes1000);
-	if (WindowsDecoderResources::Callbacks.GetD3DDevice)
-	{
-		return WindowsDecoderResources::Callbacks.GetD3DDevice(OutD3DDevice, OutD3DVersionTimes1000, WindowsDecoderResources::Callbacks.UserValue);
-	}
-	return false;
-}
-
-TSharedPtr<IElectraDecoderResourceDelegateBase::IAsyncConsecutiveTaskSync, ESPMode::ThreadSafe> FElectraDecoderResourceManagerWindows::CreateAsyncConsecutiveTaskSync()
-{
-	if (WindowsDecoderResources::Callbacks.CreateAsyncConsecutiveTaskSync)
-	{
-		return WindowsDecoderResources::Callbacks.CreateAsyncConsecutiveTaskSync();
-	}
-	return nullptr;
-}
-
-bool FElectraDecoderResourceManagerWindows::RunCodeAsync(TFunction<void()>&& CodeToRun, IAsyncConsecutiveTaskSync* TaskSync)
-{
-	if (WindowsDecoderResources::Callbacks.RunCodeAsync)
-	{
-		return WindowsDecoderResources::Callbacks.RunCodeAsync(MoveTemp(CodeToRun), TaskSync);
-	}
-	return false;
-}
-
-
 bool FElectraDecoderResourceManagerWindows::SetupRenderBufferFromDecoderOutputFromMFSample(FString& OutErrorMessage, IMediaRenderer::IBuffer* InOutBufferToSetup, TSharedPtr<FParamDict, ESPMode::ThreadSafe> InOutBufferPropertes, TSharedPtr<IElectraDecoderVideoOutput, ESPMode::ThreadSafe> InDecoderOutput, FElectraDecoderResourceManagerWindows::IDecoderPlatformResource* InPlatformSpecificResource)
 {
 	TSharedPtr<FElectraPlayerVideoDecoderOutputPC, ESPMode::ThreadSafe> DecoderOutput = InOutBufferToSetup->GetBufferProperties().GetValue(RenderOptionKeys::Texture).GetSharedPointer<FElectraPlayerVideoDecoderOutputPC>();
 	if (DecoderOutput.IsValid())
 	{
-
 		FElectraVideoDecoderOutputCropValues Crop = InDecoderOutput->GetCropValues();
 		InOutBufferPropertes->Set(IDecoderOutputOptionNames::Width, FVariantValue((int64)InDecoderOutput->GetWidth()));
 		InOutBufferPropertes->Set(IDecoderOutputOptionNames::Height, FVariantValue((int64)InDecoderOutput->GetHeight()));
@@ -404,8 +323,7 @@ bool FElectraDecoderResourceManagerWindows::SetupRenderBufferFromDecoderOutput(F
 
 	TSharedPtr<FElectraPlayerVideoDecoderOutputPC, ESPMode::ThreadSafe> DecoderOutput = InOutBufferToSetup->GetBufferProperties().GetValue(RenderOptionKeys::Texture).GetSharedPointer<FElectraPlayerVideoDecoderOutputPC>();
 	FInstanceVars* Vars = static_cast<FInstanceVars*>(InPlatformSpecificResource);
-	TSharedPtr<IVideoDecoderResourceDelegate, ESPMode::ThreadSafe> PinnedResourceDelegate = Vars ? Vars->VideoDecoderResourceDelegate.Pin() : nullptr;
-	if (DecoderOutput.IsValid() && Vars != nullptr && PinnedResourceDelegate.IsValid())
+	if (DecoderOutput.IsValid() && Vars != nullptr)
 	{
 		//
 		// Image buffers?
@@ -506,15 +424,9 @@ bool FElectraDecoderResourceManagerWindows::SetupRenderBufferFromDecoderOutput(F
 				if (TextureCommon.IsValid())
 				{
 					HRESULT Res;
-
-					void* DDevice;
-					int32 DVersion;
-					check(WindowsDecoderResources::Callbacks.GetD3DDevice);
-					WindowsDecoderResources::Callbacks.GetD3DDevice(&DDevice, &DVersion, WindowsDecoderResources::Callbacks.UserValue);
-
-					if (DVersion >= 12000)
+					if (Vars->D3DVersionTimes1000 >= 12000)
 					{
-						TRefCountPtr<ID3D12Device> D3D12Device(static_cast<ID3D12Device*>(DDevice));
+						TRefCountPtr<ID3D12Device> D3D12Device(static_cast<ID3D12Device*>(Vars->D3DDevice));
 
 						// Can we get a DX12 texture? (this will only work if the transform is associated with a DX12 device earlier on; hence no need for a SDK version guard here)
 						TRefCountPtr<ID3D12Resource> Resource;
@@ -529,7 +441,7 @@ bool FElectraDecoderResourceManagerWindows::SetupRenderBufferFromDecoderOutput(F
 							FElectraDecoderOutputSync OutputSync;
 							ImageBuffers->GetBufferTextureSyncByIndex(0, OutputSync);
 
-							DecoderOutput->InitializeWithResource(D3D12Device, Resource, Pitch, OutputSync, FIntPoint(InDecoderOutput->GetDecodedWidth(), InDecoderOutput->GetDecodedHeight()), InOutBufferPropertes, PinnedResourceDelegate,
+							DecoderOutput->InitializeWithResource(D3D12Device, Resource, Pitch, OutputSync, FIntPoint(InDecoderOutput->GetDecodedWidth(), InDecoderOutput->GetDecodedHeight()), InOutBufferPropertes,
 																  Vars->D3D12ResourcePool, Vars->MaxWidth, Vars->MaxHeight, Vars->MaxOutputBuffers);
 							return true;
 						}
diff --git a/Engine/Plugins/Media/ElectraPlayer/Source/ElectraPlayerRuntime/Private/Runtime/Decoder/Windows/WindowsElectraDecoderResourceManager.h b/Engine/Plugins/Media/ElectraPlayer/Source/ElectraPlayerRuntime/Private/Runtime/Decoder/Windows/WindowsElectraDecoderResourceManager.h
--- a/Engine/Plugins/Media/ElectraPlayer/Source/ElectraPlayerRuntime/Private/Runtime/Decoder/Windows/WindowsElectraDecoderResourceManager.h
+++ b/Engine/Plugins/Media/ElectraPlayer/Source/ElectraPlayerRuntime/Private/Runtime/Decoder/Windows/WindowsElectraDecoderResourceManager.h
@@ -15,17 +15,6 @@ namespace Electra
 class FElectraDecoderResourceManagerWindows : public IElectraDecoderResourceDelegateWindows
 {
 public:
-	struct FCallbacks
-	{
-		TFunction<bool(void ** /*OutD3DDevice*/, int32* /*OutD3DVersionTimes1000*/, void* /*UserValue*/)> GetD3DDevice;
-		TFunction<bool(TFunction<void()>&& CodeToRun, IAsyncConsecutiveTaskSync* TaskSync)> RunCodeAsync;
-		TFunction<TSharedPtr<IElectraDecoderResourceDelegateBase::IAsyncConsecutiveTaskSync, ESPMode::ThreadSafe>()> CreateAsyncConsecutiveTaskSync;
-
-		void* UserValue = nullptr;
-	};
-
-	static bool Startup(const FCallbacks& InCallbacks);
-	static void Shutdown();
 	static TSharedPtr<IElectraDecoderResourceDelegate, ESPMode::ThreadSafe> GetDelegate();
 
 	static bool SetupRenderBufferFromDecoderOutput(FString& OutErrorMessage, IMediaRenderer::IBuffer* InOutBufferToSetup, TSharedPtr<FParamDict, ESPMode::ThreadSafe> InOutBufferPropertes, TSharedPtr<IElectraDecoderVideoOutput, ESPMode::ThreadSafe> InDecoderOutput, IDecoderPlatformResource* InPlatformSpecificResource);
@@ -33,11 +22,6 @@ public:
 	IDecoderPlatformResource* CreatePlatformResource(void* InOwnerHandle, EDecoderPlatformResourceType InDecoderResourceType, const TMap<FString, FVariant> InOptions) override;
 	void ReleasePlatformResource(void* InOwnerHandle, IDecoderPlatformResource* InHandleToDestroy) override;
 
-	virtual bool GetD3DDevice(void **OutD3DDevice, int32* OutD3DVersionTimes1000) override;
-
-	virtual TSharedPtr<IElectraDecoderResourceDelegateBase::IAsyncConsecutiveTaskSync, ESPMode::ThreadSafe> CreateAsyncConsecutiveTaskSync() override;
-	virtual bool RunCodeAsync(TFunction<void()>&& CodeToRun, IAsyncConsecutiveTaskSync* TaskSync) override;
-
 	virtual ~FElectraDecoderResourceManagerWindows();
 
 private:
diff --git a/Engine/Plugins/Media/ElectraPlayer/Source/ElectraPlayerRuntime/Private/Runtime/HTTP/HTTPManager.cpp b/Engine/Plugins/Media/ElectraPlayer/Source/ElectraPlayerRuntime/Private/Runtime/HTTP/HTTPManager.cpp
--- a/Engine/Plugins/Media/ElectraPlayer/Source/ElectraPlayerRuntime/Private/Runtime/HTTP/HTTPManager.cpp
+++ b/Engine/Plugins/Media/ElectraPlayer/Source/ElectraPlayerRuntime/Private/Runtime/HTTP/HTTPManager.cpp
@@ -6,7 +6,6 @@
 #include "Utilities/URLParser.h"
 #include "PlayerTime.h"
 #include "Player/PlayerSessionServices.h"
-#include "Player/IExternalDataReader.h"
 #include "HAL/LowLevelMemTracker.h"
 #include "HAL/PlatformMisc.h"
 
@@ -89,27 +88,6 @@ namespace Electra
 			FString										MimeType;
 		};
 
-		struct FExternalReader : public FLocalByteStream, public TSharedFromThis<FExternalReader, ESPMode::ThreadSafe>
-		{
-			virtual ~FExternalReader() = default;
-			void SetConnected(TSharedPtrTS<FRequest> Request) override;
-			int32 Read(TSharedPtrTS<FWaitableBuffer> RcvBuffer, TSharedPtrTS<FRequest> Request) override;
-			void OnReadComplete(IExternalDataReader::FResponseDataPtr InResponseData, int64 InTotalFileSize, const IExternalDataReader::FReadParams& InFromRequestParams)
-			{
-				ResponseData = MoveTemp(InResponseData);
-				TotalFileSize = InTotalFileSize;
-				FPlatformMisc::MemoryBarrier();
-				bCompleted = true;
-			}
-			TSharedPtrTS<IExternalDataReader> ExternalDataReader;
-			IExternalDataReader::FReadParams ReadParams;
-			IExternalDataReader::FElectraExternalDataReadCompleted CompletionDelegate;
-			IExternalDataReader::FResponseDataPtr ResponseData;
-			int64 TotalFileSize = -1;
-			volatile bool bCompleted = false;
-			bool bRangedRequest = false;
-		};
-
 		class FHTTPCallbackWrapper
 		{
 		public:
@@ -358,7 +336,6 @@ namespace Electra
 		};
 
 		FHandle* CreateLocalFileHandle(const FTimeValue& Now, FTransportError& OutError, const TSharedPtrTS<FRequest>& Request);
-		FHandle* CreateExternalHandle(const FTimeValue& Now, FTransportError& OutError, const TSharedPtrTS<FRequest>& Request);
 		FHandle* CreateHTTPHandle(const FTimeValue& Now, FTransportError& OutError, const TSharedPtrTS<FRequest>& Request);
 		bool PrepareHTTPHandle(const FTimeValue& Now, FHandle* InHandle, const TSharedPtrTS<FRequest>& Request, bool bIsFirstSetup);
 		void AddPendingRequests(const FTimeValue& Now);
@@ -576,26 +553,6 @@ namespace Electra
 		return Handle.Release();
 	}
 
-	FElectraHttpManager::FHandle* FElectraHttpManager::CreateExternalHandle(const FTimeValue& Now, FTransportError& OutError, const TSharedPtrTS<IElectraHttpManager::FRequest>& Request)
-	{
-		TSharedPtrTS<IExternalDataReader> ExternalDataReader = Request->ExternalDataReader.Pin();
-		if (ExternalDataReader.IsValid())
-		{
-			TUniquePtr<FHandle> Handle(new FHandle);
-			Handle->Owner = this;
-			Handle->HandleType = FHandle::EHandleType::ExternalHandle;
-			FExternalReader* ExternalStream = new FExternalReader;
-			Handle->LocalByteStream = MakeShareable<FExternalReader>(ExternalStream);
-			ExternalStream->ExternalDataReader = MoveTemp(ExternalDataReader);
-			return Handle.Release();
-		}
-		else
-		{
-			OutError.Set(HTTP::EStatusErrorCode::ERRCODE_HTTP_FILE_COULDNT_READ_FILE, FString::Printf(TEXT("External data reader is not valid to read from \"%s\""), *Request->Parameters.URL));
-			return nullptr;
-		}
-	}
-
 	FElectraHttpManager::FHandle* FElectraHttpManager::CreateHTTPHandle(const FTimeValue& Now, FTransportError& OutError, const TSharedPtrTS<IElectraHttpManager::FRequest>& Request)
 	{
 		TUniquePtr<FHandle> Handle(new FHandle);
@@ -743,28 +700,9 @@ namespace Electra
 			FTransportError HttpError;
 			Request->ConnectionInfo.EffectiveURL = Request->Parameters.URL;
 
-			// If there is an external data reader set and this is _not_ a data url
-			if (Request->ExternalDataReader.IsValid() && !(Request->Parameters.URL.Len() > 5 && Request->Parameters.URL.Left(5) == TEXT("data:")))
-			{
-				FHandle* Handle = CreateExternalHandle(Now, HttpError, Request);
-				if (Handle)
-				{
-					ActiveRequests.Add(Handle, Request);
-
-					Request->ConnectionInfo.RequestStartTime = Now;
-					Handle->RequestStartTime = Now;
-					Handle->TimeAtNextProgressCallback = Now + ProgressInterval;
-					Handle->TimeAtConnectionTimeoutCheck.SetToPositiveInfinity();
-				}
-				else
-				{
-					Request->ConnectionInfo.StatusInfo.ErrorDetail.SetError(UEMEDIA_ERROR_INTERNAL).SetFacility(Facility::EFacility::HTTPReader).SetMessage(HttpError.Message).SetCode((uint16)HttpError.ErrorCode);
-					RequestsCompleted.Enqueue(Request);
-				}
-			}
-			// Otherwise, is this a local file or a data url?
-			else if ((Request->Parameters.URL.Len() > 7 && Request->Parameters.URL.Left(7) == TEXT("file://")) ||
-					 (Request->Parameters.URL.Len() > 5 && Request->Parameters.URL.Left(5) == TEXT("data:")))
+			// Is this a local file or a data url?
+			if ((Request->Parameters.URL.Len() > 7 && Request->Parameters.URL.Left(7) == TEXT("file://")) ||
+				(Request->Parameters.URL.Len() > 5 && Request->Parameters.URL.Left(5) == TEXT("data:")))
 			{
 				FHandle* Handle = CreateLocalFileHandle(Now, HttpError, Request);
 				if (Handle)
@@ -1664,80 +1602,6 @@ namespace Electra
 	}
 
 
-	void FElectraHttpManager::FExternalReader::SetConnected(TSharedPtrTS<FRequest> Request)
-	{
-		if (!bIsConnected)
-		{
-			bIsConnected = true;
-			// Go through the notions of this being a network request.
-			Request->ConnectionInfo.bIsConnected = true;
-			Request->ConnectionInfo.bHaveResponseHeaders = true;
-			Request->ConnectionInfo.ContentType = "application/octet-stream";
-			Request->ConnectionInfo.EffectiveURL = Request->Parameters.URL;
-			Request->ConnectionInfo.HTTPVersionReceived = 11;
-			Request->ConnectionInfo.bIsChunked = false;
-
-			// Trigger the external read request.
-			ReadParams.URI = Request->Parameters.URL;
-			if (!Request->Parameters.Range.IsSet())
-			{
-				ReadParams.AbsoluteFileOffset = 0;
-				ReadParams.NumBytesToRead = TNumericLimits<int64>::Max();
-			}
-			else
-			{
-				int64 end = Request->Parameters.Range.EndIncluding;
-				int64 numBytes = (end >= 0) ? end - Request->Parameters.Range.Start + 1 : TNumericLimits<int64>::Max();
-				ReadParams.AbsoluteFileOffset = Request->Parameters.Range.Start;
-				ReadParams.NumBytesToRead = numBytes;
-				bRangedRequest = true;
-			}
-			FileStartOffset = ReadParams.AbsoluteFileOffset;
-			FileSizeToGo = ReadParams.NumBytesToRead;
-
-			CompletionDelegate.BindThreadSafeSP(AsShared(), &FExternalReader::OnReadComplete);
-			ExternalDataReader->ReadData(ReadParams, CompletionDelegate);
-		}
-	}
-
-	int32 FElectraHttpManager::FExternalReader::Read(TSharedPtrTS<FWaitableBuffer> RcvBuffer, TSharedPtrTS<FRequest> Request)
-	{
-		if (bCompleted)
-		{
-			int64 NumBytesRead = ResponseData.IsValid() ? ResponseData->Num() : 0;
-			FileSize = TotalFileSize;
-			FileSizeToGo = 0;
-			if (FileSize >= 0)
-			{
-				RcvBuffer->SetExternalData(MoveTemp(ResponseData));
-				if (!bRangedRequest)
-				{
-					Request->ConnectionInfo.StatusInfo.HTTPStatus = 200;
-					Request->ConnectionInfo.ContentLength = FileSize;
-					Request->ConnectionInfo.ContentLengthHeader = FString::Printf(TEXT("Content-Length: %lld"), (long long int)FileSize);
-				}
-				else
-				{
-					Request->ConnectionInfo.StatusInfo.HTTPStatus = 206;
-					Request->ConnectionInfo.ContentLength = NumBytesRead;
-					Request->ConnectionInfo.ContentLengthHeader = FString::Printf(TEXT("Content-Length: %lld"), (long long int)NumBytesRead);
-					Request->ConnectionInfo.ContentRangeHeader = FString::Printf(TEXT("Content-Range: bytes %lld-%lld/%lld"), (long long int)FileStartOffset, (long long int)FileStartOffset+NumBytesRead-1, (long long int)TotalFileSize);
-				}
-				return NumBytesRead;
-			}
-			else
-			{
-				HTTP::FConnectionInfo& ci = Request->ConnectionInfo;
-				ci.StatusInfo.HTTPStatus = 404;
-				ci.StatusInfo.ErrorDetail.SetError(UEMEDIA_ERROR_READ_ERROR).SetFacility(Facility::EFacility::HTTPReader).SetCode(HTTP::EStatusErrorCode::ERRCODE_HTTPMODULE_FAILURE);
-				ci.StatusInfo.ErrorDetail.SetMessage(FString::Printf(TEXT("External reader returned -1 for resource size, indicating file not found (HTTP status %d)"), ci.StatusInfo.HTTPStatus));
-				ci.StatusInfo.ErrorCode = HTTP::EStatusErrorCode::ERRCODE_HTTP_RETURNED_ERROR;
-			}
-		}
-		return -1;
-	}
-
-
 	bool FElectraHttpManager::FDataUrl::SetData(const FString& InUrl)
 	{
 		// See https://en.wikipedia.org/wiki/Data_URI_scheme
diff --git a/Engine/Plugins/Media/ElectraPlayer/Source/ElectraPlayerRuntime/Private/Runtime/HTTP/HTTPManager.h b/Engine/Plugins/Media/ElectraPlayer/Source/ElectraPlayerRuntime/Private/Runtime/HTTP/HTTPManager.h
--- a/Engine/Plugins/Media/ElectraPlayer/Source/ElectraPlayerRuntime/Private/Runtime/HTTP/HTTPManager.h
+++ b/Engine/Plugins/Media/ElectraPlayer/Source/ElectraPlayerRuntime/Private/Runtime/HTTP/HTTPManager.h
@@ -19,7 +19,6 @@ DECLARE_LOG_CATEGORY_EXTERN(LogElectraHTTPManager, Log, All);
 
 namespace Electra
 {
-	class IExternalDataReader;
 	class IHTTPResponseCache;
 
 	namespace HTTP
@@ -251,7 +250,6 @@ namespace Electra
 			TWeakPtrTS<FWaitableBuffer>			ReceiveBuffer;
 			TWeakPtrTS<FProgressListener>		ProgressListener;
 			TSharedPtrTS<IHTTPResponseCache>	ResponseCache;
-			TWeakPtrTS<IExternalDataReader>		ExternalDataReader;
 			bool								bAutoRemoveWhenComplete = false;
 		};
 
diff --git a/Engine/Plugins/Media/ElectraPlayer/Source/ElectraPlayerRuntime/Private/Runtime/Player/AdaptivePlayerOptionKeynames.h b/Engine/Plugins/Media/ElectraPlayer/Source/ElectraPlayerRuntime/Private/Runtime/Player/AdaptivePlayerOptionKeynames.h
--- a/Engine/Plugins/Media/ElectraPlayer/Source/ElectraPlayerRuntime/Private/Runtime/Player/AdaptivePlayerOptionKeynames.h
+++ b/Engine/Plugins/Media/ElectraPlayer/Source/ElectraPlayerRuntime/Private/Runtime/Player/AdaptivePlayerOptionKeynames.h
@@ -6,6 +6,9 @@
 
 namespace Electra
 {
+	//! (int64) ID of the platform specific output texture pool registered with the decoder platform resources.
+	const FName OptionKeyOutputTexturePoolID(TEXT("output_texturepool_id"));
+
 	//! (FString) Use of worker threads. ("shared", "worker" or "worker_and_events")
 	const FName OptionKeyWorkerThreads(TEXT("worker_threads"));
 
@@ -39,9 +42,6 @@ namespace Electra
 	//! (bool) true to not truncate the media segment access units at the end of the presentation. Must only be used without a set playback range end!
 	const FName OptionKeyDoNotTruncateAtPresentationEnd(TEXT("do_not_truncate_at_presentation_end"));
 
-	//! (bool) true to have every request to read data break out to an external data reader.
-	const FName OptionKeyUseExternalDataReader(TEXT("use_external_data_reader"));
-
 	const FName OptionKeyCurrentAvgStartingVideoBitrate(TEXT("current:avg_video_bitrate"));
 
 	const FName OptionKeyExcludedCodecsVideo(TEXT("excluded_codecs_video"));
diff --git a/Engine/Plugins/Media/ElectraPlayer/Source/ElectraPlayerRuntime/Private/Runtime/Player/AdaptiveStreamingPlayer.cpp b/Engine/Plugins/Media/ElectraPlayer/Source/ElectraPlayerRuntime/Private/Runtime/Player/AdaptiveStreamingPlayer.cpp
--- a/Engine/Plugins/Media/ElectraPlayer/Source/ElectraPlayerRuntime/Private/Runtime/Player/AdaptiveStreamingPlayer.cpp
+++ b/Engine/Plugins/Media/ElectraPlayer/Source/ElectraPlayerRuntime/Private/Runtime/Player/AdaptiveStreamingPlayer.cpp
@@ -276,17 +276,6 @@ void FAdaptiveStreamingPlayer::SetStaticResourceProviderCallback(const TSharedPt
 	StaticResourceProvider = InStaticResourceProvider;
 }
 
-void FAdaptiveStreamingPlayer::SetVideoDecoderResourceDelegate(const TSharedPtr<IVideoDecoderResourceDelegate, ESPMode::ThreadSafe>& InResourceDelegate)
-{
-	VideoDecoderResourceDelegate = InResourceDelegate;
-}
-
-void FAdaptiveStreamingPlayer::SetPlayerDataCache(const TSharedPtr<IElectraPlayerDataCache, ESPMode::ThreadSafe>& InPlayerDataCache)
-{
-	ExternalCache = InPlayerDataCache;
-}
-
-
 void FAdaptiveStreamingPlayer::AddMetricsReceiver(IAdaptiveStreamingPlayerMetrics* InMetricsReceiver)
 {
 	if (InMetricsReceiver)
@@ -643,7 +632,7 @@ void FAdaptiveStreamingPlayer::Stop()
  */
 void FAdaptiveStreamingPlayer::SetLooping(const FLoopParam& InLoopParams)
 {
-	WorkerThread.SendLoopMessage(InLoopParams);
+	PlaybackState.RequestLooping(InLoopParams);
 }
 
 
@@ -1215,11 +1204,6 @@ TSharedPtrTS<IElectraHttpManager> FAdaptiveStreamingPlayer::GetHTTPManager()
 	return HttpManager;
 }
 
-TSharedPtrTS<IExternalDataReader> FAdaptiveStreamingPlayer::GetExternalDataReader()
-{
-	return GetStaticResourceProvider().IsValid() ? ExternalDataReader : nullptr;
-}
-
 TSharedPtrTS<IAdaptiveStreamSelector> FAdaptiveStreamingPlayer::GetStreamSelector()
 {
 	return StreamSelector;
@@ -1507,13 +1491,6 @@ bool FAdaptiveStreamingPlayer::InternalHandleThreadMessages()
 			}
 			case FWorkerThreadMessages::FMessage::EType::LoadManifest:
 			{
-				// Check if a new or different cache has been set prior to loading the manifest.
-				if (ExternalCache.IsValid())
-				{
-					HttpResponseCache.Reset();
-					HttpResponseCache = IHTTPResponseCache::Create(GetOptionValue(OptionKeyResponseCacheMaxByteSize).SafeGetInt64(0), (int32)GetOptionValue(OptionKeyResponseCacheMaxEntries).SafeGetInt64(8192), MoveTemp(ExternalCache));
-				}
-
 				FWorkerThreadMessages::FMessage::FLoadManifest& ev = msg.Data.Get<FWorkerThreadMessages::FMessage::FLoadManifest>();
 				InternalLoadManifest(ev.URL, ev.MimeType);
 				break;
@@ -1539,17 +1516,6 @@ bool FAdaptiveStreamingPlayer::InternalHandleThreadMessages()
 				bShouldBePlaying = true;
 				break;
 			}
-			case FWorkerThreadMessages::FMessage::EType::Loop:
-			{
-				FWorkerThreadMessages::FMessage::FLoop& ev = msg.Data.Get<FWorkerThreadMessages::FMessage::FLoop>();
-				InternalSetLoop(ev.Loop);
-				/*
-				if (ev.Signal)
-				{
-					ev.Signal->Signal();
-				*/
-				break;
-			}
 			case FWorkerThreadMessages::FMessage::EType::Close:
 			{
 				if (!bIsClosing)
@@ -1797,6 +1763,9 @@ void FAdaptiveStreamingPlayer::HandlePlayStateChanges()
 		}
 	}
 
+	// Handle pending loop type changes.
+	InternalHandlePendingLoopRequest();
+
 	// Did the play range or loop state change such that we need to start over?
 	if (PlaybackState.GetPlayRangeHasChanged() || PlaybackState.GetLoopStateHasChanged())
 	{
@@ -4552,14 +4521,8 @@ void FAdaptiveStreamingPlayer::InternalInitialize()
 	// Create an entity cache.
 	EntityCache = IPlayerEntityCache::Create(this);
 
-	// Create an HTTP response cache. Hand over any externally set cache. We do not need it in here any further.
-	HttpResponseCache = IHTTPResponseCache::Create(GetOptionValue(OptionKeyResponseCacheMaxByteSize).SafeGetInt64(0), (int32)GetOptionValue(OptionKeyResponseCacheMaxEntries).SafeGetInt64(8192), MoveTemp(ExternalCache));
-
-	// If all read requests for this player are to be routed to an external reader we need to create a wrapper for it.
-	if (PlayerOptions.GetValue(OptionKeyUseExternalDataReader).SafeGetBool(false))
-	{
-		ExternalDataReader = MakeSharedTS<FExternalDataReader>(this);
-	}
+	// Create an HTTP response cache.
+	HttpResponseCache = IHTTPResponseCache::Create(GetOptionValue(OptionKeyResponseCacheMaxByteSize).SafeGetInt64(0), (int32)GetOptionValue(OptionKeyResponseCacheMaxEntries).SafeGetInt64(8192), nullptr);
 
 	// Create the content steering handler.
 	ContentSteeringHandler = MakeSharedTS<FContentSteeringHandler>(this);
@@ -4655,7 +4618,6 @@ void FAdaptiveStreamingPlayer::InternalClose()
 	HttpManager.Reset();
 	EntityCache.Reset();
 	HttpResponseCache.Reset();
-	ExternalDataReader.Reset();
 
 	// Reset remaining internal state
 	CurrentState   	= EPlayerState::eState_Idle;
@@ -4693,42 +4655,46 @@ void FAdaptiveStreamingPlayer::InternalClose()
 }
 
 
-void FAdaptiveStreamingPlayer::InternalSetLoop(const FLoopParam& LoopParam)
+void FAdaptiveStreamingPlayer::InternalHandlePendingLoopRequest()
 {
-	// Looping only makes sense for on demand presentations since Live does not have an end to loop at.
-	// In case this is called that early on that we can't determine the presentation type yet we allow
-	// looping to be set. It won't get evaluated anyway since the Live presentation doesn't end.
-	// The only issue with this is that until a manifest is loaded querying loop enable will return true.
-	if (!Manifest || Manifest->GetPresentationType() == IManifest::EType::OnDemand)
+	TUniquePtr<FLoopParam> NewLoopParam(PlaybackState.GetPendingLoopRequest());
+	if (NewLoopParam.IsValid())
 	{
-		bool bStartOver = false;
-		// When looping is to be enabled we need to check if we can do this without having to start over.
-		if (!CurrentLoopParam.bEnableLooping && LoopParam.bEnableLooping)
+		// Looping only makes sense for on demand presentations since Live does not have an end to loop at.
+		// In case this is called that early on that we can't determine the presentation type yet we allow
+		// looping to be set. It won't get evaluated anyway since the Live presentation doesn't end.
+		// The only issue with this is that until a manifest is loaded querying loop enable will return true.
+		if (!Manifest || Manifest->GetPresentationType() == IManifest::EType::OnDemand)
 		{
-			// If there are any completed segment requests the loop state will be checked when all reach
-			// the end. Otherwise, if any of the receive buffers has the end-of-data flag set that
-			// that streams reached the end and will not loop, so we need to start over.
-			if (CompletedSegmentRequests.Num() == 0)
+			bool bStartOver = false;
+			// When looping is to be enabled we need to check if we can do this without having to start over.
+			if (!CurrentLoopParam.bEnableLooping && NewLoopParam->bEnableLooping)
 			{
-				TSharedPtrTS<FMultiTrackAccessUnitBuffer> RcvBuffer;
-				bool bAnyReceiveBufferAtEOD = false;
-				RcvBuffer = GetCurrentReceiveStreamBuffer(EStreamType::Video);
-				bAnyReceiveBufferAtEOD |= RcvBuffer.IsValid() ? RcvBuffer->IsEODFlagSet() || RcvBuffer->IsEndOfTrack() : false;
-				RcvBuffer = GetCurrentReceiveStreamBuffer(EStreamType::Audio);
-				bAnyReceiveBufferAtEOD |= RcvBuffer.IsValid() ? RcvBuffer->IsEODFlagSet() || RcvBuffer->IsEndOfTrack() : false;
-				RcvBuffer = GetCurrentReceiveStreamBuffer(EStreamType::Subtitle);
-				bAnyReceiveBufferAtEOD |= RcvBuffer.IsValid() ? RcvBuffer->IsEODFlagSet() || RcvBuffer->IsEndOfTrack() : false;
-
-				bStartOver = bAnyReceiveBufferAtEOD;
+				// If there are any completed segment requests the loop state will be checked when all reach
+				// the end. Otherwise, if any of the receive buffers has the end-of-data flag set that
+				// that streams reached the end and will not loop, so we need to start over.
+				if (CompletedSegmentRequests.Num() == 0)
+				{
+					TSharedPtrTS<FMultiTrackAccessUnitBuffer> RcvBuffer;
+					bool bAnyReceiveBufferAtEOD = false;
+					RcvBuffer = GetCurrentReceiveStreamBuffer(EStreamType::Video);
+					bAnyReceiveBufferAtEOD |= RcvBuffer.IsValid() ? RcvBuffer->IsEODFlagSet() || RcvBuffer->IsEndOfTrack() : false;
+					RcvBuffer = GetCurrentReceiveStreamBuffer(EStreamType::Audio);
+					bAnyReceiveBufferAtEOD |= RcvBuffer.IsValid() ? RcvBuffer->IsEODFlagSet() || RcvBuffer->IsEndOfTrack() : false;
+					RcvBuffer = GetCurrentReceiveStreamBuffer(EStreamType::Subtitle);
+					bAnyReceiveBufferAtEOD |= RcvBuffer.IsValid() ? RcvBuffer->IsEODFlagSet() || RcvBuffer->IsEndOfTrack() : false;
+
+					bStartOver = bAnyReceiveBufferAtEOD;
+				}
+			}
+			CurrentLoopParam = *NewLoopParam;
+			// For the lack of better knowledge update the current state immediately.
+			CurrentLoopState.bIsEnabled = NewLoopParam->bEnableLooping;
+			PlaybackState.SetLoopStateEnable(NewLoopParam->bEnableLooping);
+			if (bStartOver)
+			{
+				PlaybackState.SetLoopStateHasChanged(true);
 			}
-		}
-		CurrentLoopParam = LoopParam;
-		// For the lack of better knowledge update the current state immediately.
-		CurrentLoopState.bIsEnabled = LoopParam.bEnableLooping;
-		PlaybackState.SetLoopStateEnable(LoopParam.bEnableLooping);
-		if (bStartOver)
-		{
-			PlaybackState.SetLoopStateHasChanged(true);
 		}
 	}
 }
diff --git a/Engine/Plugins/Media/ElectraPlayer/Source/ElectraPlayerRuntime/Private/Runtime/Player/AdaptiveStreamingPlayer.h b/Engine/Plugins/Media/ElectraPlayer/Source/ElectraPlayerRuntime/Private/Runtime/Player/AdaptiveStreamingPlayer.h
--- a/Engine/Plugins/Media/ElectraPlayer/Source/ElectraPlayerRuntime/Private/Runtime/Player/AdaptiveStreamingPlayer.h
+++ b/Engine/Plugins/Media/ElectraPlayer/Source/ElectraPlayerRuntime/Private/Runtime/Player/AdaptiveStreamingPlayer.h
@@ -20,7 +20,7 @@ class IOptionPointerValueContainer;
 namespace Electra
 {
 class IMediaRenderer;
-class IVideoDecoderResourceDelegate;
+
 
 /**
  *
@@ -50,8 +50,6 @@ public:
 	// until player destruction.
 	//
 	virtual void SetStaticResourceProviderCallback(const TSharedPtr<IAdaptiveStreamingPlayerResourceProvider, ESPMode::ThreadSafe>& InStaticResourceProvider) = 0;
-	virtual void SetVideoDecoderResourceDelegate(const TSharedPtr<IVideoDecoderResourceDelegate, ESPMode::ThreadSafe>& ResourceDelegate) = 0;
-	virtual void SetPlayerDataCache(const TSharedPtr<IElectraPlayerDataCache, ESPMode::ThreadSafe>& InPlayerDataCache) = 0;
 
 
 	//-------------------------------------------------------------------------
diff --git a/Engine/Plugins/Media/ElectraPlayer/Source/ElectraPlayerRuntime/Private/Runtime/Player/AdaptiveStreamingPlayerDecoders.cpp b/Engine/Plugins/Media/ElectraPlayer/Source/ElectraPlayerRuntime/Private/Runtime/Player/AdaptiveStreamingPlayerDecoders.cpp
--- a/Engine/Plugins/Media/ElectraPlayer/Source/ElectraPlayerRuntime/Private/Runtime/Player/AdaptiveStreamingPlayerDecoders.cpp
+++ b/Engine/Plugins/Media/ElectraPlayer/Source/ElectraPlayerRuntime/Private/Runtime/Player/AdaptiveStreamingPlayerDecoders.cpp
@@ -288,8 +288,6 @@ int32 FAdaptiveStreamingPlayer::CreateDecoder(EStreamType type)
 					VideoDecoder.Decoder->SetReadyBufferListener(&VideoDecoder);
 					// Have the video decoder send its output to the video renderer
 					VideoDecoder.Decoder->SetRenderer(VideoRender.Renderer);
-					// Hand it (may be nullptr) a delegate for platform for resource queries
-					VideoDecoder.Decoder->SetVideoResourceDelegate(VideoDecoderResourceDelegate);
 					// Open the decoder after having set all listeners.
 					VideoDecoder.Decoder->Open(VideoDecoder.LastSentAUCodecData, PlayerOptions.GetDictionary(), &HighestStream);
 				}
diff --git a/Engine/Plugins/Media/ElectraPlayer/Source/ElectraPlayerRuntime/Private/Runtime/Player/AdaptiveStreamingPlayerInternal.h b/Engine/Plugins/Media/ElectraPlayer/Source/ElectraPlayerRuntime/Private/Runtime/Player/AdaptiveStreamingPlayerInternal.h
--- a/Engine/Plugins/Media/ElectraPlayer/Source/ElectraPlayerRuntime/Private/Runtime/Player/AdaptiveStreamingPlayerInternal.h
+++ b/Engine/Plugins/Media/ElectraPlayer/Source/ElectraPlayerRuntime/Private/Runtime/Player/AdaptiveStreamingPlayerInternal.h
@@ -17,7 +17,6 @@
 #include "Decoder/VideoDecoder.h"
 
 #include "Player/AdaptiveStreamingPlayerResourceRequest.h"
-#include "Player/ExternalDataReader.h"
 #include "Player/PlayerStreamReader.h"
 #include "Player/PlaylistReader.h"
 #include "Player/PlayerStreamFilter.h"
@@ -269,6 +268,7 @@ struct FPlaybackState
 	bool									bPlayrangeHasChanged;
 	bool									bLoopStateHasChanged;
 	bool									bShouldPlayOnLiveEdge;
+	TUniquePtr<IAdaptiveStreamingPlayer::FLoopParam>	PendingLoopParams;
 	TArray<FTrackMetadata>					VideoTracks;
 	TArray<FTrackMetadata>					AudioTracks;
 	TArray<FTrackMetadata>					SubtitleTracks;
@@ -482,6 +482,12 @@ struct FPlaybackState
 	{
 		FScopeLock lock(&Lock);
 		OutLoopState = LoopState;
+		// If there is a pending loop change request we return the requested state
+		// until we have handled the request.
+		if (PendingLoopParams.IsValid())
+		{
+			OutLoopState.bIsEnabled = PendingLoopParams->bEnableLooping;
+		}
 	}
 	void SetLoopStateEnable(bool bEnable)
 	{
@@ -586,6 +592,16 @@ struct FPlaybackState
 	}
 
 
+	void RequestLooping(const IAdaptiveStreamingPlayer::FLoopParam& InLoopParams)
+	{
+		FScopeLock lock(&Lock);
+		PendingLoopParams.Reset(new IAdaptiveStreamingPlayer::FLoopParam(InLoopParams));
+	}
+	TUniquePtr<IAdaptiveStreamingPlayer::FLoopParam> GetPendingLoopRequest()
+	{
+		FScopeLock lock(&Lock);
+		return MoveTemp(PendingLoopParams);
+	}
 	void SetLoopStateHasChanged(bool bWasChanged)
 	{
 		FScopeLock lock(&Lock);
@@ -1092,8 +1108,6 @@ public:
 	virtual ~FAdaptiveStreamingPlayer();
 
 	void SetStaticResourceProviderCallback(const TSharedPtr<IAdaptiveStreamingPlayerResourceProvider, ESPMode::ThreadSafe>& InStaticResourceProvider) override;
-	void SetVideoDecoderResourceDelegate(const TSharedPtr<IVideoDecoderResourceDelegate, ESPMode::ThreadSafe>& ResourceDelegate) override;
-	void SetPlayerDataCache(const TSharedPtr<IElectraPlayerDataCache, ESPMode::ThreadSafe>& InPlayerDataCache) override;
 
 	void AddMetricsReceiver(IAdaptiveStreamingPlayerMetrics* InMetricsReceiver) override;
 	void RemoveMetricsReceiver(IAdaptiveStreamingPlayerMetrics* InMetricsReceiver) override;
@@ -1184,7 +1198,6 @@ private:
 	ISynchronizedUTCTime* GetSynchronizedUTCTime() override;
 	TSharedPtr<IAdaptiveStreamingPlayerResourceProvider, ESPMode::ThreadSafe> GetStaticResourceProvider() override;
 	TSharedPtrTS<IElectraHttpManager> GetHTTPManager() override;
-	TSharedPtrTS<IExternalDataReader> GetExternalDataReader() override;
 	TSharedPtrTS<IAdaptiveStreamSelector> GetStreamSelector() override;
 	IPlayerStreamFilter* GetStreamFilter() override;
 	const FCodecSelectionPriorities& GetCodecSelectionPriorities(EStreamType ForStream) override;
@@ -1561,7 +1574,6 @@ private:
 				LoadBlob,
 				Pause,
 				Resume,
-				Loop,
 				Close,
 				ChangeBitrate,
 				LimitResolution,
@@ -1707,14 +1719,6 @@ private:
 			Msg.Type = FMessage::EType::Resume;
 			TriggerSharedWorkerThread(MoveTemp(Msg));
 		}
-		void SendLoopMessage(const FLoopParam& InLoopParams)
-		{
-			FMessage::FLoop loop { InLoopParams };
-			FMessage Msg;
-			Msg.Type = FMessage::EType::Loop;
-			Msg.Data.Emplace<FMessage::FLoop>(MoveTemp(loop));
-			TriggerSharedWorkerThread(MoveTemp(Msg));
-		}
 		void SendCloseMessage(FMediaEvent* InEventSignal)
 		{
 			check(InEventSignal);
@@ -2246,7 +2250,7 @@ private:
 	void InternalRebuffer();
 	void InternalStop(bool bHoldCurrentFrame);
 	void InternalClose();
-	void InternalSetLoop(const FLoopParam& LoopParam);
+	void InternalHandlePendingLoopRequest();
 	void InternalSetPlaybackEnded();
 	void InternalDeselectStream(EStreamType StreamType);
 
@@ -2266,8 +2270,6 @@ private:
 	FGuid																ExternalPlayerGUID;
 
 	TWeakPtr<IAdaptiveStreamingPlayerResourceProvider, ESPMode::ThreadSafe> StaticResourceProvider;
-	TWeakPtr<IVideoDecoderResourceDelegate, ESPMode::ThreadSafe>		VideoDecoderResourceDelegate;
-	TSharedPtr<IElectraPlayerDataCache, ESPMode::ThreadSafe>			ExternalCache;
 
 	TSharedPtrTS<FAdaptiveStreamingPlayerEventHandler>					EventDispatcher;
 	TSharedPtrTS<FAdaptiveStreamingPlayerWorkerThread>					SharedWorkerThread;
@@ -2284,7 +2286,6 @@ private:
 	TSharedPtrTS<IElectraHttpManager>									HttpManager;
 	TSharedPtrTS<IPlayerEntityCache>									EntityCache;
 	TSharedPtrTS<IHTTPResponseCache>									HttpResponseCache;
-	TSharedPtrTS<IExternalDataReader>									ExternalDataReader;
 
 	TSharedPtrTS<FDRMManager>											DrmManager;
 
diff --git a/Engine/Plugins/Media/ElectraPlayer/Source/ElectraPlayerRuntime/Private/Runtime/Player/AdaptiveStreamingPlayerResourceRequest.cpp b/Engine/Plugins/Media/ElectraPlayer/Source/ElectraPlayerRuntime/Private/Runtime/Player/AdaptiveStreamingPlayerResourceRequest.cpp
--- a/Engine/Plugins/Media/ElectraPlayer/Source/ElectraPlayerRuntime/Private/Runtime/Player/AdaptiveStreamingPlayerResourceRequest.cpp
+++ b/Engine/Plugins/Media/ElectraPlayer/Source/ElectraPlayerRuntime/Private/Runtime/Player/AdaptiveStreamingPlayerResourceRequest.cpp
@@ -122,8 +122,6 @@ void FHTTPResourceRequest::StartGet(IPlayerSessionServices* InPlayerSessionServi
 	}
 	else
 	{
-		// Is an external reader to be used?
-		Request->ExternalDataReader = PlayerSessionServices->GetExternalDataReader();
 		bWasAdded = true;
 		TSharedPtrTS<IElectraHttpManager> PinnedHTTPManager = HTTPManager.Pin();
 		if (PinnedHTTPManager.IsValid())
@@ -159,11 +157,6 @@ void FHTTPResourceRequest::StaticDataReady()
 		if (!bStaticDataReady)
 		{
 			// Do the actual HTTP request now.
-			// Is an external reader to be used?
-			if (!Req->ExternalDataReader.IsValid())
-			{
-				Req->ExternalDataReader = PlayerSessionServices->GetExternalDataReader();
-			}
 			bWasAdded = true;
 			TSharedPtrTS<IElectraHttpManager> PinnedHTTPManager = HTTPManager.Pin();
 			if (PinnedHTTPManager.IsValid())
diff --git a/Engine/Plugins/Media/ElectraPlayer/Source/ElectraPlayerRuntime/Private/Runtime/Player/AdaptiveStreamingPlayerResourceRequest.h b/Engine/Plugins/Media/ElectraPlayer/Source/ElectraPlayerRuntime/Private/Runtime/Player/AdaptiveStreamingPlayerResourceRequest.h
--- a/Engine/Plugins/Media/ElectraPlayer/Source/ElectraPlayerRuntime/Private/Runtime/Player/AdaptiveStreamingPlayerResourceRequest.h
+++ b/Engine/Plugins/Media/ElectraPlayer/Source/ElectraPlayerRuntime/Private/Runtime/Player/AdaptiveStreamingPlayerResourceRequest.h
@@ -20,14 +20,7 @@ public:
 	{
 		Empty,
 		Playlist,
-		LicenseKey,
-		BinaryData
-	};
-
-	struct FBinaryDataParams
-	{
-		int64 AbsoluteFileOffset = 0;
-		int64 NumBytesToRead = 0;
+		LicenseKey
 	};
 
 	virtual ~IAdaptiveStreamingPlayerResourceRequest() = default;
@@ -36,8 +29,6 @@ public:
 	virtual EPlaybackResourceType GetResourceType() const = 0;
 	//! Returns the URL of the requested resource.
 	virtual FString GetResourceURL() const = 0;
-	//! Returns the read parameters for a binary data request.
-	virtual FBinaryDataParams GetBinaryDataParams() const = 0;
 
 	//! Sets the binary resource data. If data can not be provided do not set anything (or a nullptr)
 	//! For a request of `BinaryData` you need to set `TotalResourceSize` to the total size of the file so the reader
@@ -223,12 +214,6 @@ private:
 			return p.IsValid() ? p->GetURL() : FString();
 		}
 
-		FBinaryDataParams GetBinaryDataParams() const override
-		{
-			FBinaryDataParams None;
-			return None;
-		}
-
 		void SetPlaybackData(TSharedPtr<TArray<uint8>, ESPMode::ThreadSafe>	PlaybackData, int64 TotalResourceSize) override
 		{
 			if (PlaybackData.IsValid())
diff --git a/Engine/Plugins/Media/ElectraPlayer/Source/ElectraPlayerRuntime/Private/Runtime/Player/PlayerSessionServices.h b/Engine/Plugins/Media/ElectraPlayer/Source/ElectraPlayerRuntime/Private/Runtime/Player/PlayerSessionServices.h
--- a/Engine/Plugins/Media/ElectraPlayer/Source/ElectraPlayerRuntime/Private/Runtime/Player/PlayerSessionServices.h
+++ b/Engine/Plugins/Media/ElectraPlayer/Source/ElectraPlayerRuntime/Private/Runtime/Player/PlayerSessionServices.h
@@ -23,7 +23,6 @@ namespace Electra
 	class FDRMManager;
 	class FContentSteeringHandler;
 	class IHTTPResponseCache;
-	class IExternalDataReader;
 
 
 	class IPlayerMessage
@@ -90,13 +89,6 @@ namespace Electra
 		 */
 		virtual TSharedPtrTS<IElectraHttpManager> GetHTTPManager() = 0;
 
-		/**
-		 * Returns the optional external data reader.
-		 *
-		 * @return Pointer to the external data reader, or nullptr if none is to be used.
-		 */
-		virtual TSharedPtrTS<IExternalDataReader> GetExternalDataReader() = 0;
-
 		/**
 		 * Returns the ABR stream selector instance.
 		 *
diff --git a/Engine/Plugins/Media/ElectraPlayer/Source/ElectraPlayerRuntime/Private/Runtime/Player/StreamSegmentReaderCommon.cpp b/Engine/Plugins/Media/ElectraPlayer/Source/ElectraPlayerRuntime/Private/Runtime/Player/StreamSegmentReaderCommon.cpp
--- a/Engine/Plugins/Media/ElectraPlayer/Source/ElectraPlayerRuntime/Private/Runtime/Player/StreamSegmentReaderCommon.cpp
+++ b/Engine/Plugins/Media/ElectraPlayer/Source/ElectraPlayerRuntime/Private/Runtime/Player/StreamSegmentReaderCommon.cpp
@@ -1341,7 +1341,6 @@ FStreamSegmentReaderCommon::FStreamHandler::EMediaSegmentTriggerResult FStreamSe
 	HTTPRequest->ReceiveBuffer = ReadBuffer.ReceiveBuffer;
 	HTTPRequest->ProgressListener = ProgressListener;
 	HTTPRequest->ResponseCache = PlayerSessionService->GetHTTPResponseCache();
-	HTTPRequest->ExternalDataReader = PlayerSessionService->GetExternalDataReader();
 	HTTPRequest->Parameters.URL = InRequest->Segment.MediaURL.Url.URL;
 	HTTPRequest->Parameters.Range.Set(InRequest->Segment.MediaURL.Range);
 	HTTPRequest->Parameters.StreamType = InRequest->StreamType;
diff --git a/Engine/Plugins/Media/ElectraPlayer/Source/ElectraPlayerRuntime/Private/Runtime/Player/mp4/StreamReaderMP4.cpp b/Engine/Plugins/Media/ElectraPlayer/Source/ElectraPlayerRuntime/Private/Runtime/Player/mp4/StreamReaderMP4.cpp
--- a/Engine/Plugins/Media/ElectraPlayer/Source/ElectraPlayerRuntime/Private/Runtime/Player/mp4/StreamReaderMP4.cpp
+++ b/Engine/Plugins/Media/ElectraPlayer/Source/ElectraPlayerRuntime/Private/Runtime/Player/mp4/StreamReaderMP4.cpp
@@ -387,7 +387,6 @@ void FStreamReaderMP4::HandleRequest()
 	HTTP->ReceiveBuffer = ReadBuffer.ReceiveBuffer;
 	HTTP->ProgressListener = ProgressListener;
 	HTTP->ResponseCache = PlayerSessionServices->GetHTTPResponseCache();
-	HTTP->ExternalDataReader = PlayerSessionServices->GetExternalDataReader();
 	PlayerSessionServices->GetHTTPManager()->AddRequest(HTTP, false);
 
 
diff --git a/Engine/Plugins/Media/ElectraPlayer/Source/ElectraPlayerRuntime/Private/Runtime/Player/mpegaudio/StreamReaderMPEGAudio.cpp b/Engine/Plugins/Media/ElectraPlayer/Source/ElectraPlayerRuntime/Private/Runtime/Player/mpegaudio/StreamReaderMPEGAudio.cpp
--- a/Engine/Plugins/Media/ElectraPlayer/Source/ElectraPlayerRuntime/Private/Runtime/Player/mpegaudio/StreamReaderMPEGAudio.cpp
+++ b/Engine/Plugins/Media/ElectraPlayer/Source/ElectraPlayerRuntime/Private/Runtime/Player/mpegaudio/StreamReaderMPEGAudio.cpp
@@ -404,7 +404,6 @@ void FStreamReaderMPEGAudio::HandleRequest()
 		HTTP->ReceiveBuffer = ReceiveBuffer;
 		HTTP->ProgressListener = ProgressListener;
 		HTTP->ResponseCache = PlayerSessionServices->GetHTTPResponseCache();
-		HTTP->ExternalDataReader = PlayerSessionServices->GetExternalDataReader();
 		PlayerSessionServices->GetHTTPManager()->AddRequest(HTTP, false);
 	}
 	else
diff --git a/Engine/Plugins/Media/ElectraPlayer/Source/ElectraPlayerRuntime/Private/Runtime/Renderer/RendererAudio.cpp b/Engine/Plugins/Media/ElectraPlayer/Source/ElectraPlayerRuntime/Private/Runtime/Renderer/RendererAudio.cpp
--- a/Engine/Plugins/Media/ElectraPlayer/Source/ElectraPlayerRuntime/Private/Runtime/Renderer/RendererAudio.cpp
+++ b/Engine/Plugins/Media/ElectraPlayer/Source/ElectraPlayerRuntime/Private/Runtime/Renderer/RendererAudio.cpp
@@ -21,15 +21,6 @@ void FElectraRendererAudio::SampleReleasedToPool(IDecoderOutput* InDecoderOutput
 	}
 }
 
-bool FElectraRendererAudio::Startup(const FElectraRendererAudio::SystemConfiguration& Configuration)
-{
-	return true;
-}
-
-void FElectraRendererAudio::Shutdown(void)
-{
-}
-
 FElectraRendererAudio::FElectraRendererAudio(TSharedPtr<FElectraPlayer, ESPMode::ThreadSafe> InPlayer)
 	: Player(InPlayer)
 	, MaxBufferSize(0)
diff --git a/Engine/Plugins/Media/ElectraPlayer/Source/ElectraPlayerRuntime/Private/Runtime/Renderer/RendererAudio.h b/Engine/Plugins/Media/ElectraPlayer/Source/ElectraPlayerRuntime/Private/Runtime/Renderer/RendererAudio.h
--- a/Engine/Plugins/Media/ElectraPlayer/Source/ElectraPlayerRuntime/Private/Runtime/Renderer/RendererAudio.h
+++ b/Engine/Plugins/Media/ElectraPlayer/Source/ElectraPlayerRuntime/Private/Runtime/Renderer/RendererAudio.h
@@ -24,12 +24,6 @@ public:
 class FElectraRendererAudio : public Electra::IMediaRenderer, public TSharedFromThis<FElectraRendererAudio, ESPMode::ThreadSafe>
 {
 public:
-	struct SystemConfiguration
-	{
-	};
-	static bool Startup(const SystemConfiguration& Configuration);
-	static void Shutdown(void);
-
 	FElectraRendererAudio(TSharedPtr<FElectraPlayer, ESPMode::ThreadSafe> InPlayer);
 	virtual ~FElectraRendererAudio();
 
diff --git a/Engine/Plugins/Media/ElectraPlayer/Source/ElectraPlayerRuntime/Private/Runtime/Renderer/RendererVideo.cpp b/Engine/Plugins/Media/ElectraPlayer/Source/ElectraPlayerRuntime/Private/Runtime/Renderer/RendererVideo.cpp
--- a/Engine/Plugins/Media/ElectraPlayer/Source/ElectraPlayerRuntime/Private/Runtime/Renderer/RendererVideo.cpp
+++ b/Engine/Plugins/Media/ElectraPlayer/Source/ElectraPlayerRuntime/Private/Runtime/Renderer/RendererVideo.cpp
@@ -12,8 +12,6 @@
 #include "ElectraPlayer.h"
 #include "ElectraPlayerPrivate.h"
 
-DECLARE_DWORD_COUNTER_STAT(TEXT("Electra::MEDIArendererVideoUE_NumUsedOutputVideoSamples"), STAT_ElectraPlayer_MEDIArendererVideoUE_NumUsedOutputVideoSamples, STATGROUP_ElectraPlayer);
-
 // -----------------------------------------------------------------------------------------------------------------------
 
 void FElectraRendererVideo::SampleReleasedToPool(IDecoderOutput* InDecoderOutput)
@@ -28,22 +26,6 @@ void FElectraRendererVideo::SampleReleasedToPool(IDecoderOutput* InDecoderOutput
 }
 
 
-FElectraRendererVideo::SystemConfiguration::SystemConfiguration()
-{
-}
-
-
-bool FElectraRendererVideo::Startup(const FElectraRendererVideo::SystemConfiguration& Configuration)
-{
-	return true;
-}
-
-void FElectraRendererVideo::Shutdown(void)
-{
-}
-
-
-
 FElectraRendererVideo::FElectraRendererVideo(TSharedPtr<FElectraPlayer, ESPMode::ThreadSafe> InPlayer)
 	: Player(InPlayer)
 	, NumOutputTexturesInUse(0)
@@ -134,13 +116,6 @@ UEMediaError FElectraRendererVideo::AcquireBuffer(IBuffer*& OutBuffer, int32 Tim
 		return UEMEDIA_ERROR_BAD_ARGUMENTS;
 	}
 
-	// Trigger removal of any old frames in the presentation queue of the player so we have all buffers we can have available
-	TSharedPtr<FElectraPlayer, ESPMode::ThreadSafe> PinnedPlayer = Player.Pin();
-	if (PinnedPlayer.IsValid())
-	{
-		PinnedPlayer->DropOldFramesFromPresentationQueue();
-	}
-
 	FVideoDecoderOutputPtr DelayedImage;
 	AcquireFromPool(DelayedImage);
 	if (!DelayedImage.IsValid())
@@ -271,16 +246,3 @@ void FElectraRendererVideo::StartRendering(const Electra::FParamDict& InOptions)
 void FElectraRendererVideo::StopRendering(const Electra::FParamDict& InOptions)
 {
 }
-
-/**
- * Called from game thread to update buffer management
- */
-void FElectraRendererVideo::TickInput(FTimespan DeltaTime, FTimespan Timecode)
-{
-	LLM_SCOPE(ELLMTag::ElectraPlayer);
-	CSV_SCOPED_TIMING_STAT(ElectraPlayer, MEDIArendererVideoUE_Tick);
-
-	// Keep stats up to date
-	SET_DWORD_STAT(STAT_ElectraPlayer_MEDIArendererVideoUE_NumUsedOutputVideoSamples, NumOutputTexturesInUse);
-}
-
diff --git a/Engine/Plugins/Media/ElectraPlayer/Source/ElectraPlayerRuntime/Private/Runtime/Renderer/RendererVideo.h b/Engine/Plugins/Media/ElectraPlayer/Source/ElectraPlayerRuntime/Private/Runtime/Renderer/RendererVideo.h
--- a/Engine/Plugins/Media/ElectraPlayer/Source/ElectraPlayerRuntime/Private/Runtime/Renderer/RendererVideo.h
+++ b/Engine/Plugins/Media/ElectraPlayer/Source/ElectraPlayerRuntime/Private/Runtime/Renderer/RendererVideo.h
@@ -21,14 +21,6 @@ public:
 class FElectraRendererVideo : public Electra::IMediaRenderer, public TSharedFromThis<FElectraRendererVideo, ESPMode::ThreadSafe>
 {
 public:
-	struct SystemConfiguration
-	{
-		SystemConfiguration();
-	};
-
-	static bool Startup(const SystemConfiguration& Configuration);
-	static void Shutdown(void);
-
 	FElectraRendererVideo(TSharedPtr<FElectraPlayer, ESPMode::ThreadSafe> InPlayer);
 	virtual ~FElectraRendererVideo();
 
@@ -52,7 +44,6 @@ public:
 	//-------------------------------------------------------------------------
 	// Methods for UE interface
 	//
-	void TickInput(FTimespan DeltaTime, FTimespan Timecode);
 	void DetachPlayer();
 
 
diff --git a/Engine/Plugins/Media/ElectraPlayer/Source/ElectraPlayerRuntime/Public/ElectraPlayerMisc.h b/Engine/Plugins/Media/ElectraPlayer/Source/ElectraPlayerRuntime/Public/ElectraPlayerMisc.h
--- a/Engine/Plugins/Media/ElectraPlayer/Source/ElectraPlayerRuntime/Public/ElectraPlayerMisc.h
+++ b/Engine/Plugins/Media/ElectraPlayer/Source/ElectraPlayerRuntime/Public/ElectraPlayerMisc.h
@@ -14,6 +14,5 @@ ELECTRAPLAYERRUNTIME_API DECLARE_LOG_CATEGORY_EXTERN(LogElectraPlayerStats, Log,
 CSV_DECLARE_CATEGORY_MODULE_EXTERN(ELECTRAPLAYERRUNTIME_API, MediaStreaming);
 CSV_DECLARE_CATEGORY_MODULE_EXTERN(ELECTRAPLAYERRUNTIME_API, ElectraPlayer);
 
-/** UE4 Profiler */
+/** UE Profiler */
 DECLARE_STATS_GROUP(TEXT("ElectraPlayer"), STATGROUP_ElectraPlayer, STATCAT_Advanced);
-
diff --git a/Engine/Plugins/Media/ElectraPlayer/Source/ElectraPlayerRuntime/Public/IElectraPlayerDecoderResourceManager.h b/Engine/Plugins/Media/ElectraPlayer/Source/ElectraPlayerRuntime/Public/IElectraPlayerDecoderResourceManager.h
--- a/Engine/Plugins/Media/ElectraPlayer/Source/ElectraPlayerRuntime/Public/IElectraPlayerDecoderResourceManager.h
+++ b/Engine/Plugins/Media/ElectraPlayer/Source/ElectraPlayerRuntime/Public/IElectraPlayerDecoderResourceManager.h
@@ -8,19 +8,12 @@
 #include "ParameterDictionary.h"
 #include "MediaVideoDecoderOutput.h"
 
-class IElectraPlayerAdapterDelegate;
 class IElectraDecoderVideoOutput;
 
-namespace Electra
-{
-class IVideoDecoderResourceDelegate;
-}
-
 class FElectraPlayerDecoderResourceManager
 {
 public:
 	ELECTRAPLAYERRUNTIME_API static TSharedPtr<IElectraDecoderResourceDelegate, ESPMode::ThreadSafe> GetDelegate();
-	ELECTRAPLAYERRUNTIME_API static TSharedPtr<Electra::IVideoDecoderResourceDelegate, ESPMode::ThreadSafe> CreatePlatformVideoDecoderResourceDelegate(TSharedPtr<IElectraPlayerAdapterDelegate, ESPMode::ThreadSafe> InElectraPlayerAdapterDelegate);
 	ELECTRAPLAYERRUNTIME_API static bool SetupRenderBufferFromDecoderOutput(FString& OutErrorMessage, Electra::IMediaRenderer::IBuffer* InOutBufferToSetup, TSharedPtr<Electra::FParamDict, ESPMode::ThreadSafe> InOutBufferPropertes, TSharedPtr<IElectraDecoderVideoOutput, ESPMode::ThreadSafe> InDecoderOutput, IElectraDecoderResourceDelegateBase::IDecoderPlatformResource* InPlatformSpecificResource);
 	class FVideo
 	{
diff --git a/Engine/Plugins/Media/ElectraPlayer/Source/ElectraPlayerRuntime/Public/IElectraPlayerInterface.h b/Engine/Plugins/Media/ElectraPlayer/Source/ElectraPlayerRuntime/Public/IElectraPlayerInterface.h
--- a/Engine/Plugins/Media/ElectraPlayer/Source/ElectraPlayerRuntime/Public/IElectraPlayerInterface.h
+++ b/Engine/Plugins/Media/ElectraPlayer/Source/ElectraPlayerRuntime/Public/IElectraPlayerInterface.h
@@ -1,375 +1,135 @@
 // Copyright Epic Games, Inc. All Rights Reserved.
 #pragma once
 
-#include "CoreMinimal.h"
-#include "MediaDecoderOutput.h"
+#include "CoreTypes.h"
+#include "Templates/SharedPointer.h"
+#include "Misc/Guid.h"
+#include "IMediaPlayer.h"
+#include "IMediaEventSink.h"
+#include "IMediaOptions.h"
 #include "IAnalyticsProviderET.h"
-#include "ParameterDictionary.h"
-#include "Math/RangeSet.h"
+#include "UObject/Object.h"
+#include "UObject/StrongObjectPtrTemplates.h"
+#include "UObject/WeakObjectPtr.h"
 
-#include "ElectraPlayerMisc.h"
-#include "ElectraPlayerResourceDelegate.h"
-#include "IElectraPlayerDataCache.h"
-#include "MediaStreamMetadata.h"
 
 #define UE_API ELECTRAPLAYERRUNTIME_API
 
-class FVideoDecoderOutput;
-using FVideoDecoderOutputPtr = TSharedPtr<FVideoDecoderOutput, ESPMode::ThreadSafe>;
-class IAudioDecoderOutput;
-using IAudioDecoderOutputPtr = TSharedPtr<IAudioDecoderOutput, ESPMode::ThreadSafe>;
-class IMetaDataDecoderOutput;
-using IMetaDataDecoderOutputPtr = TSharedPtr<IMetaDataDecoderOutput, ESPMode::ThreadSafe>;
-class ISubtitleDecoderOutput;
-using ISubtitleDecoderOutputPtr = TSharedPtr<ISubtitleDecoderOutput, ESPMode::ThreadSafe>;
-
-// ---------------------------------------------------------------------------------------------
-
-DECLARE_MULTICAST_DELEGATE_TwoParams(FElectraPlayerSendAnalyticMetricsDelegate, const TSharedPtr<IAnalyticsProviderET>& /*AnalyticsProvider*/, const FGuid& /*PlayerGuid*/);
-DECLARE_MULTICAST_DELEGATE_OneParam(FElectraPlayerSendAnalyticMetricsPerMinuteDelegate, const TSharedPtr<IAnalyticsProviderET>& /*AnalyticsProvider*/);
-DECLARE_MULTICAST_DELEGATE_TwoParams(FElectraPlayerReportVideoStreamingErrorDelegate, const FGuid& /*PlayerGuid*/, const FString& /*LastError*/);
-DECLARE_MULTICAST_DELEGATE_FourParams(FElectraPlayerReportSubtitlesMetricsDelegate, const FGuid& /*PlayerGuid*/, const FString& /*URL*/, double /*ResponseTime*/, const FString& /*LastError*/);
-
 // ---------------------------------------------------------------------------------------------
 
-class IElectraPlayerAdapterDelegate
+DECLARE_MULTICAST_DELEGATE_TwoParams(FElectraPlayerSendAnalyticMetricsDelegate, const TSharedPtr<IAnalyticsProviderET>& /*InAnalyticsProvider*/, const FGuid& /*InPlayerGuid*/);
+DECLARE_MULTICAST_DELEGATE_OneParam(FElectraPlayerSendAnalyticMetricsPerMinuteDelegate, const TSharedPtr<IAnalyticsProviderET>& /*InAnalyticsProvider*/);
+DECLARE_MULTICAST_DELEGATE_TwoParams(FElectraPlayerReportVideoStreamingErrorDelegate, const FGuid& /*InPlayerGuid*/, const FString& /*InLastError*/);
+DECLARE_MULTICAST_DELEGATE_FourParams(FElectraPlayerReportSubtitlesMetricsDelegate, const FGuid& /*InPlayerGuid*/, const FString& /*InURL*/, double /*InResponseTime*/, const FString& /*InLastError*/);
+
+
+/**
+ * This class is used to get safe access to an IMediaOptions interface.
+ * When passing IMediaOptions through media framework as a POD pointer there is the risk
+ * that it is actually from a derived class like UMediaSource which is subject to GC.
+ * Thus it is not safe to keep the POD IMediaOptions around.
+ * This class is intended to be derived and instantiated from and stored as a TSharedPtr
+ * in a derived UMediaSource class.
+ * Then, as the media player is opened and the POD IMediaOptions is passed along, a
+ * GetMediaOption("GetSafeMediaOptions") query will be made on it to get this instance.
+ * If future access to the IMediaOptions is needed it will be made through this class
+ * by first locking, getting and using the IMediaOptions pointer returned here if it
+ * not null and unlocking afterwards.
+ */
+class IElectraSafeMediaOptionInterface : public IMediaOptions::FDataContainer
 {
 public:
-	virtual ~IElectraPlayerAdapterDelegate() {}
-
-	enum class EOptionType
-	{
-		MaxVerticalStreamResolution = 0,
-		MaxBandwidthForStreaming,
-		PlayListData,
-		LicenseKeyData,
-		PlaystartPosFromSeekPositions,
-		CustomAnalyticsMetric
-	};
-	virtual Electra::FVariantValue QueryOptions(EOptionType Type, const Electra::FVariantValue& Param = Electra::FVariantValue()) = 0;
+	virtual ~IElectraSafeMediaOptionInterface() = default;
+	virtual void Lock() = 0;
+	virtual void Unlock() = 0;
+	virtual IMediaOptions* GetMediaOptionInterface() = 0;
 
-	enum class EBlobResultType
+	class FScopedLock
 	{
-		Success,
-		TimedOut,
-		HttpFailure
-	};
-	virtual void BlobReceived(const TSharedPtr<TArray<uint8>, ESPMode::ThreadSafe>& InBlobData, EBlobResultType InResultType, int32 InResultCode, const Electra::FParamDict* InExtraInfo) = 0;
-
-	enum class EPlayerEvent
-	{
-		MediaBuffering = 0,
-		MediaBufferingComplete,
-		MediaClosed,
-		MediaConnecting,
-		MediaOpened,
-		MediaOpenFailed,
-		PlaybackEndReached,
-		PlaybackResumed,
-		PlaybackSuspended,
-		SeekCompleted,
-		TracksChanged,
-		MetadataChanged,
+	public:
+		FScopedLock(const TSharedPtr<IElectraSafeMediaOptionInterface, ESPMode::ThreadSafe>& InSafeMediaOptionInterface)
+			: SafeMediaOptionInterface(InSafeMediaOptionInterface)
+		{
+			if (SafeMediaOptionInterface.IsValid())
+			{
+				SafeMediaOptionInterface->Lock();
+			}
+		}
+		~FScopedLock()
+		{
+			if (SafeMediaOptionInterface.IsValid())
+			{
+				SafeMediaOptionInterface->Unlock();
+			}
+		}
 
-		Internal_Start,
-		Internal_PurgeVideoSamplesHint = Internal_Start,
-		Internal_VideoSamplesAvailable,
-		Internal_VideoSamplesUnavailable,
-		Internal_AudioSamplesAvailable,
-		Internal_AudioSamplesUnavailable
+	private:
+		TSharedPtr<IElectraSafeMediaOptionInterface, ESPMode::ThreadSafe> SafeMediaOptionInterface;
 	};
-	virtual void SendMediaEvent(EPlayerEvent Event) = 0;
-	virtual void OnVideoFlush() = 0;
-	virtual void OnAudioFlush() = 0;
-	virtual void OnSubtitleFlush() = 0;
-	virtual void PresentVideoFrame(const FVideoDecoderOutputPtr& InVideoFrame) = 0;
-	virtual void PresentAudioFrame(const IAudioDecoderOutputPtr& InAudioFrame) = 0;
-	virtual void PresentSubtitleSample(const ISubtitleDecoderOutputPtr& InSubtitleSample) = 0;
-	virtual void PresentMetadataSample(const IMetaDataDecoderOutputPtr& InMetadataSample) = 0;
-	virtual bool CanReceiveVideoSamples(int32 NumFrames) = 0;
-	virtual bool CanReceiveAudioSamples(int32 NumFrames) = 0;
-
-	virtual FString GetVideoAdapterName() const = 0;
-
-	virtual TSharedPtr<IElectraPlayerResourceDelegate, ESPMode::ThreadSafe> GetResourceDelegate() const = 0;
 };
 
 
-class IElectraPlayerExternalDataReader : public TSharedFromThis<IElectraPlayerExternalDataReader, ESPMode::ThreadSafe>
+/**
+ * A ready to use implementation of the IElectraSafeMediaOptionInterface
+ */
+class FElectraSafeMediaOptionInterface : public IElectraSafeMediaOptionInterface
 {
 public:
-	virtual ~IElectraPlayerExternalDataReader() = default;
-
-	struct FReadParam
-	{
-		FString URL;
-		int64 AbsoluteFileOffset = 0;
-		int64 NumBytesToRead = 0;
-		void* Custom = nullptr;
-	};
-
-	using FResponseDataPtr = TSharedPtr<TArray<uint8>, ESPMode::ThreadSafe>;
-	DECLARE_DELEGATE_ThreeParams(FElectraPlayerExternalDataReaderOnRequestCompleted, FResponseDataPtr /*ResponseData*/, int64 /*TotalFileSize*/, const FReadParam& FromRequestParams);
-
-	/**
-	 * Called to read data by some external means.
-	 * When done the provided completion delegate must be called.
-	 * Passing a nullptr for the response data will trigger a read error and a subsequent playback error.
-	 * Returning fewer bytes than requested implicitly means the end of the file was reached while reading.
-	 * If a read request for 0 bytes at offset 0 is made you are to return the total size of the file. For this
-	 * request it is permitted to return a nullptr for response data without resulting in an error.
-	 * Read requests may set the number of bytes to read to MAX_INT64 to indicate reading until the end of the file.
-	 * The player has no upfront knowledge about the size of the file so it cannot always adjust the number of
-	 * bytes to be read. Even when returning the size of the file in an earlier read there is no guarantee
-	 * that the number of bytes to be read will be adjusted accordingly.
-	 * Returning a negative value for the total size of the file is treated as an error and will be handled
-	 * like a file-not-found error, even though that should not be possible.
-	 * You are providing this reader class to the player knowing that the URL to play exists and is valid.
-	 * If the offset to read from is given as a negative value this is an indication that you may close the
-	 * file, but such a call is optional. You may close the file after having closed the player.
-	 * You may however open and close the file with every read request if that is more convenient for you.
-	 * The player will wait indefinitely for the response data to be provided. There are no timeouts.
-	 * You *must* call the provided delegate even for failures.
-	 */
-	virtual void ReadDataFromFile(const FReadParam& InReadParam, FElectraPlayerExternalDataReaderOnRequestCompleted OutCompletionDelegate) = 0;
+	FElectraSafeMediaOptionInterface(IMediaOptions* InOwner)
+		: Owner(InOwner)
+		, OwnerObject(InOwner ? InOwner->ToUObject() : nullptr)
+	{ }
+	virtual ~FElectraSafeMediaOptionInterface()
+	{
+		ClearOwner();
+	}
+	void ClearOwner()
+	{
+		FScopeLock lock(&OwnerLock);
+		Owner = nullptr;
+		OwnerObject.Reset();
+	}
+	virtual void Lock() override
+	{
+		OwnerLock.Lock();
+	}
+	virtual void Unlock() override
+	{
+		OwnerLock.Unlock();
+	}
+	virtual IMediaOptions* GetMediaOptionInterface() override
+	{
+		return OwnerObject.IsStale(true, true) ? nullptr : Owner;
+	}
+private:
+	FCriticalSection OwnerLock;
+	IMediaOptions* Owner = nullptr;
+	FWeakObjectPtr OwnerObject;
 };
 
 
 
-// Container class to be passed through options as shared pointer to allow passing any non-standard ref-counted entities to the player
-class IOptionPointerValueContainer
-{
-public:
-	virtual ~IOptionPointerValueContainer() {}
-	virtual void* GetPointer() const = 0;
-};
-
-class IElectraPlayerInterface : public TSharedFromThis<IElectraPlayerInterface, ESPMode::ThreadSafe>
+//! Data type for use with media options interface
+class FElectraSeekablePositions : public IMediaOptions::FDataContainer
 {
 public:
-	virtual ~IElectraPlayerInterface() {}
-
-	virtual FString GetUrl() const = 0;
-
-	virtual void SetGuid(const FGuid& Guid) = 0;
-
-	// -------- PlayerAdapter (Plugin/Native) API
-
-	struct FStreamSelectionAttributes
-	{
-		TOptional<FString> Kind;
-		TOptional<FString> Language_RFC4647;
-		TOptional<FString> Codec;
-		TOptional<int32> TrackIndexOverride;
-		void Reset()
-		{
-			Kind.Reset();
-			Language_RFC4647.Reset();
-			Codec.Reset();
-			TrackIndexOverride.Reset();
-		}
-	};
-
-	struct FPlaystartOptions
-	{
-		TOptional<FTimespan>		TimeOffset;
-		FStreamSelectionAttributes	InitialVideoTrackAttributes;
-		FStreamSelectionAttributes	InitialAudioTrackAttributes;
-		FStreamSelectionAttributes	InitialSubtitleTrackAttributes;
-		TOptional<int32>			MaxVerticalStreamResolution;
-		TOptional<int32>			MaxBandwidthForStreaming;
-		bool						bDoNotPreload = false;
-
-		TSharedPtr<IElectraPlayerExternalDataReader, ESPMode::ThreadSafe> ExternalDataReader;
-		TSharedPtr<IElectraPlayerDataCache, ESPMode::ThreadSafe> ExternalDataCache;
-	};
-
-	enum class EOpenType
-	{
-		Media,
-		Blob
-	};
-
-	virtual bool OpenInternal(const FString& Url, const Electra::FParamDict& PlayerOptions, const FPlaystartOptions& InPlaystartOptions, EOpenType InOpenType) = 0;
-	virtual void CloseInternal(bool bKillAfterClose) = 0;
-
-	virtual void Tick(FTimespan DeltaTime, FTimespan Timecode) = 0;
-
-	virtual bool IsKillAfterCloseAllowed() const = 0;
-
-	enum class EPlayerState
-	{
-		Closed = 0,
-		Error,
-		Paused,
-		Playing,
-		Preparing,
-		Stopped
-	};
-
-	enum class EPlayerStatus
-	{
-		None = 0x0,
-		Buffering = 0x1,
-		Connecting = 0x2
-	};
-
-	virtual EPlayerState GetState() const = 0;
-	virtual EPlayerStatus GetStatus() const = 0;
+	FElectraSeekablePositions(const TArray<FTimespan>& InData) : Data(InData) {}
+	virtual ~FElectraSeekablePositions() = default;
 
-	virtual bool IsLooping() const = 0;
-	virtual bool SetLooping(bool bLooping) = 0;
-	virtual int32 GetLoopCount() const = 0;
-
-	virtual FTimespan GetTime() const = 0;
-	virtual FTimespan GetDuration() const = 0;
-
-	virtual bool IsLive() const = 0;
-	virtual FTimespan GetSeekableDuration() const = 0;
-
-	struct FPlaybackRange
-	{
-		TOptional<FTimespan> Start;
-		TOptional<FTimespan> End;
-	};
-	virtual void SetPlaybackRange(const FPlaybackRange& InPlaybackRange) = 0;
-	virtual void GetPlaybackRange(FPlaybackRange& OutPlaybackRange) const = 0;
-
-	enum class ETimeRangeType
-	{
-		/** Total absolute time range as defined by the media. */
-		Absolute,
-
-		/** Current time range of the media, set by media internal means or through API calls. */
-		Current
-	};
-	virtual TRange<FTimespan> GetPlaybackRange(ETimeRangeType InRangeToGet) const = 0;
-
-
-	enum class EPlayRateType
-	{
-		// Playback rate with frames being dropped.
-		Thinned,
-		// Playback rate without dropping frames.
-		Unthinned
-	};
-	virtual TRangeSet<float> GetSupportedRates(EPlayRateType InPlayRateType) const = 0;
-
-	virtual float GetRate() const = 0;
-	virtual bool SetRate(float Rate) = 0;
-
-
-	struct FSeekParam
-	{
-		TOptional<int32> SequenceIndex;
-		TOptional<int32> StartingBitrate;
-	};
-
-	virtual bool Seek(const FTimespan& Time, const FSeekParam& Param) = 0;
-	virtual void SetFrameAccurateSeekMode(bool bEnableFrameAccuracy) = 0;
-
-	struct FAudioTrackFormat
-	{
-		uint32 BitsPerSample;
-		uint32 NumChannels;
-		uint32 SampleRate;
-		FString TypeName;
-	};
-	struct FVideoTrackFormat
-	{
-		FIntPoint Dim;
-		float FrameRate;
-		TRange<float> FrameRates;
-		FString TypeName;
-	};
-	virtual bool GetAudioTrackFormat(int32 TrackIndex, int32 FormatIndex, FAudioTrackFormat& OutFormat) const = 0;
-	virtual bool GetVideoTrackFormat(int32 TrackIndex, int32 FormatIndex, FVideoTrackFormat& OutFormat) const = 0;
-
-	enum class EPlayerTrackType
-	{
-		Audio = 0,
-		Caption,
-		Metadata,
-		Script,
-		Subtitle,
-		Text,
-		Video
-	};
-	virtual int32 GetNumTracks(EPlayerTrackType TrackType) const = 0;
-	virtual int32 GetNumTrackFormats(EPlayerTrackType TrackType, int32 TrackIndex) const = 0;
-	virtual int32 GetSelectedTrack(EPlayerTrackType TrackType) const = 0;
-	virtual FText GetTrackDisplayName(EPlayerTrackType TrackType, int32 TrackIndex) const = 0;
-	virtual int32 GetTrackFormat(EPlayerTrackType TrackType, int32 TrackIndex) const = 0;
-	virtual FString GetTrackLanguage(EPlayerTrackType TrackType, int32 TrackIndex) const = 0;
-	virtual FString GetTrackName(EPlayerTrackType TrackType, int32 TrackIndex) const = 0;
-	virtual bool SelectTrack(EPlayerTrackType TrackType, int32 TrackIndex) = 0;
-
-	struct FVideoStreamFormat
-	{
-		FIntPoint Resolution;
-		double FrameRate;
-		int32 Bitrate;
-	};
-	virtual int32 GetNumVideoStreams(int32 TrackIndex) const = 0;
-	virtual bool GetVideoStreamFormat(FVideoStreamFormat& OutFormat, int32 InTrackIndex, int32 InStreamIndex) const = 0;
-	virtual bool GetActiveVideoStreamFormat(FVideoStreamFormat& OutFormat) const = 0;
-
-	virtual Electra::FVariantValue GetMediaInfo(FName InInfoName) const = 0;
-	virtual TSharedPtr<TMap<FString, TArray<TSharedPtr<Electra::IMediaStreamMetadata::IItem, ESPMode::ThreadSafe>>>, ESPMode::ThreadSafe> GetMediaMetadata() const = 0;
-
-	struct FStreamBufferInfo
-	{
-		struct FTimeValue
-		{
-			FTimespan Time;
-			int64 SequenceIndex = -1;
-		};
-		struct FTimeRange
-		{
-			FTimeValue Start;
-			FTimeValue End;
-		};
-		TArray<FTimeRange> TimeEnqueued;
-		TArray<FTimeRange> TimeAvailable;
-		TArray<FTimeRange> TimeRequested;
-	};
-	virtual bool GetStreamBufferInformation(FStreamBufferInfo& OutBufferInformation, EPlayerTrackType InTrackType) const = 0;
-
-	// Suspends or resumes decoder instances. Not supported on all platforms.
-	virtual void SuspendOrResumeDecoders(bool bSuspend, const Electra::FParamDict& InOptions) = 0;
-
-	enum
-	{
-		ResourceFlags_Decoder = 1 << 0,
-		ResourceFlags_OutputBuffers = 1 << 1,
-
-		ResourceFlags_All = (1 << 2) - 1,
-		ResourceFlags_Any = ResourceFlags_All,
-	};
-
-	class IAsyncResourceReleaseNotifyContainer
-	{
-	public:
-		virtual ~IAsyncResourceReleaseNotifyContainer() {}
-		virtual void Signal(uint32 ResourceFlags) = 0;
-	};
-	virtual void SetAsyncResourceReleaseNotification(IAsyncResourceReleaseNotifyContainer* AsyncResourceReleaseNotification) = 0;
+	TArray<FTimespan> Data;
 };
 
+
 class FElectraPlayerRuntimeFactory
 {
 public:
-	static UE_API IElectraPlayerInterface* CreatePlayer(const TSharedPtr<IElectraPlayerAdapterDelegate, ESPMode::ThreadSafe>& AdapterDelegate,
-												FElectraPlayerSendAnalyticMetricsDelegate& InSendAnalyticMetricsDelegate,
-												FElectraPlayerSendAnalyticMetricsPerMinuteDelegate& InSendAnalyticMetricsPerMinuteDelegate,
-												FElectraPlayerReportVideoStreamingErrorDelegate& InReportVideoStreamingErrorDelegate,
-												FElectraPlayerReportSubtitlesMetricsDelegate& InReportSubtitlesFileMetricsDelegate);
+	static UE_API TSharedPtr<IMediaPlayer, ESPMode::ThreadSafe> CreatePlayer(IMediaEventSink& InEventSink,
+																			FElectraPlayerSendAnalyticMetricsDelegate& InSendAnalyticMetricsDelegate,
+																			FElectraPlayerSendAnalyticMetricsPerMinuteDelegate& InSendAnalyticMetricsPerMinuteDelegate,
+																			FElectraPlayerReportVideoStreamingErrorDelegate& InReportVideoStreamingErrorDelegate,
+																			FElectraPlayerReportSubtitlesMetricsDelegate& InReportSubtitlesFileMetricsDelegate);
 };
 
-class FElectraPlayerPlatform
-{
-public:
-	static UE_API bool StartupPlatformResources(const Electra::FParamDict & Params = Electra::FParamDict());
-};
+// ---------------------------------------------------------------------------------------------
 
 #undef UE_API
diff --git a/Engine/Plugins/Media/ElectraPlayer/Source/ElectraProtron/Private/Player/ElectraProtronPlayerDecoder.cpp b/Engine/Plugins/Media/ElectraPlayer/Source/ElectraProtron/Private/Player/ElectraProtronPlayerDecoder.cpp
--- a/Engine/Plugins/Media/ElectraPlayer/Source/ElectraProtron/Private/Player/ElectraProtronPlayerDecoder.cpp
+++ b/Engine/Plugins/Media/ElectraPlayer/Source/ElectraProtron/Private/Player/ElectraProtronPlayerDecoder.cpp
@@ -15,7 +15,7 @@
 #include "IElectraDecoder.h"
 #include "IElectraDecoderOutputVideo.h"
 #include "IElectraDecoderOutputAudio.h"
-#include "IElectraPlayerPluginModule.h"
+#include "IElectraPlayerRuntimeModule.h"
 #include "ElectraDecodersUtils.h"
 
 #include "MediaSamples.h"
@@ -1132,7 +1132,6 @@ void FElectraProtronPlayer::FImpl::FDecoderThread::DestroyDecoder()
 		FElectraPlayerDecoderResourceManager::GetDelegate()->ReleasePlatformResource(this, PlatformResource);
 		PlatformResource = nullptr;
 	}
-	VideoResourceDelegate.Reset();
 }
 
 
@@ -1153,10 +1152,9 @@ bool FElectraProtronPlayer::FImpl::FDecoderThread::CreateDecoder()
 	IElectraCodecFactoryModule* FactoryModule = static_cast<IElectraCodecFactoryModule*>(FModuleManager::Get().GetModule(TEXT("ElectraCodecFactory")));
 	check(FactoryModule);
 	// And since we piggyback on the Electra Player decoder output handling and conversion we need the plugin as well.
-	IElectraPlayerPluginModule* PluginModule = static_cast<IElectraPlayerPluginModule*>(FModuleManager::Get().GetModule(TEXT("ElectraPlayerPlugin")));
+	IElectraPlayerRuntimeModule* PluginModule = static_cast<IElectraPlayerRuntimeModule*>(FModuleManager::Get().GetModule(TEXT("ElectraPlayerRuntime")));
 	check(PluginModule && PluginModule->IsInitialized());
 
-
 	CurrentCodecSpecificData.Empty();
 	CurrentColorimetry.Reset();
 	CurrentHDR.Reset();
@@ -1191,9 +1189,6 @@ bool FElectraProtronPlayer::FImpl::FDecoderThread::CreateDecoder()
 			DecoderCfgParams.Add(TEXT("fps_n"), FVariant((int64)0));
 			DecoderCfgParams.Add(TEXT("fps_d"), FVariant((uint32)1));
 		}
-		VideoResourceDelegate = PluginModule->CreatePlatformVideoDecoderResourceDelegate();
-		// Set the video resource delegate with the decoder configuration options as a plain pointer.
-		DecoderCfgParams.Emplace(TEXT("VideoResourceDelegate"), FVariant((uint64)VideoResourceDelegate.Get()));
 	}
 	else if (DecoderTypeIndex == CodecTypeIndex(ElectraProtronUtils::FCodecInfo::EType::Audio))
 	{
@@ -1246,7 +1241,7 @@ bool FElectraProtronPlayer::FImpl::FDecoderThread::CreateDecoder()
 	//DecoderCreateCfg.Emplace(TEXT("platformresource"), FVariant((uint64)PlatformResource));
 
 	// Try to create the decoder now.
-	DecoderInstance = DecoderFactory->CreateDecoderForFormat(DecoderFormat, DecoderCfgParams, ElectraDecoderResourceDelegate);
+	DecoderInstance = DecoderFactory->CreateDecoderForFormat(DecoderFormat, DecoderCfgParams);
 	if (!DecoderInstance.IsValid())
 	{
 		ElectraDecoderResourceDelegate->ReleasePlatformResource(this, PlatformResource);
diff --git a/Engine/Plugins/Media/ElectraPlayer/Source/ElectraProtron/Private/Player/ElectraProtronPlayerImpl.h b/Engine/Plugins/Media/ElectraPlayer/Source/ElectraProtron/Private/Player/ElectraProtronPlayerImpl.h
--- a/Engine/Plugins/Media/ElectraPlayer/Source/ElectraProtron/Private/Player/ElectraProtronPlayerImpl.h
+++ b/Engine/Plugins/Media/ElectraPlayer/Source/ElectraProtron/Private/Player/ElectraProtronPlayerImpl.h
@@ -16,7 +16,6 @@
 #include "Utilities/MP4Boxes/MP4Track.h"
 #include "TrackFormatInfo.h"
 #include "IElectraDecoder.h"
-#include "IElectraDecoderResourceDelegateBase.h"
 #include "IElectraPlayerDecoderResourceManager.h"
 #include "Utils/MPEG/ElectraUtilsMPEGVideo.h"
 #include "Utils/AudioChannelMapper.h"
@@ -689,7 +688,6 @@ private:
 		TSharedPtr<IElectraDecoderBitstreamProcessor, ESPMode::ThreadSafe> DecoderBitstreamProcessor;
 
 		TMap<FString, FVariant> CurrentCodecSpecificData;
-		TSharedPtr<Electra::IVideoDecoderResourceDelegate, ESPMode::ThreadSafe> VideoResourceDelegate;
 		IElectraDecoderResourceDelegateBase::IDecoderPlatformResource* PlatformResource = nullptr;
 		TUniquePtr<FInDecoder> CurrentInputSample;
 		TSharedPtr<IElectraDecoderOutput, ESPMode::ThreadSafe> CurrentDecoderOutput;
diff --git a/Engine/Plugins/Media/ElectraUtil/Source/ElectraBase/Public/MediaVideoDecoderOutputAndroid.h b/Engine/Plugins/Media/ElectraUtil/Source/ElectraBase/Public/MediaVideoDecoderOutputAndroid.h
--- a/Engine/Plugins/Media/ElectraUtil/Source/ElectraBase/Public/MediaVideoDecoderOutputAndroid.h
+++ b/Engine/Plugins/Media/ElectraUtil/Source/ElectraBase/Public/MediaVideoDecoderOutputAndroid.h
@@ -13,12 +13,9 @@ public:
 	enum class EOutputType
 	{
 		Unknown = 0,
-		DirectToSurfaceAsView,	//!< Output surface is assumed to be GUI element
 		DirectToSurfaceAsQueue,	//!< Output surface is assumed to be queue for texture output
 	};
 
 	virtual EOutputType GetOutputType() const = 0;
-
-	virtual void ReleaseToSurface() const = 0;
 };
 
diff --git a/Engine/Plugins/Media/ElectraUtil/Source/ElectraSamples/Private/ElectraTextureSampleAndroid.cpp b/Engine/Plugins/Media/ElectraUtil/Source/ElectraSamples/Private/ElectraTextureSampleAndroid.cpp
--- a/Engine/Plugins/Media/ElectraUtil/Source/ElectraSamples/Private/ElectraTextureSampleAndroid.cpp
+++ b/Engine/Plugins/Media/ElectraUtil/Source/ElectraSamples/Private/ElectraTextureSampleAndroid.cpp
@@ -517,8 +517,6 @@ void FElectraTextureSample::Initialize(FVideoDecoderOutput* InVideoDecoderOutput
 	IElectraTextureSampleBase::Initialize(InVideoDecoderOutput);
 	VideoDecoderOutputAndroid = static_cast<FVideoDecoderOutputAndroid*>(InVideoDecoderOutput);
 
-	ensure(VideoDecoderOutputAndroid->GetOutputType() == FVideoDecoderOutputAndroid::EOutputType::DirectToSurfaceAsQueue);
-
 	if (Support->UseGpuOutputPath())
 	{
 		bQueuedForConversion = false;
diff --git a/Engine/Plugins/Media/ElectraUtil/Source/ElectraSamples/Public/Linux/ElectraTextureSample.h b/Engine/Plugins/Media/ElectraUtil/Source/ElectraSamples/Public/Linux/ElectraTextureSample.h
--- a/Engine/Plugins/Media/ElectraUtil/Source/ElectraSamples/Public/Linux/ElectraTextureSample.h
+++ b/Engine/Plugins/Media/ElectraUtil/Source/ElectraSamples/Public/Linux/ElectraTextureSample.h
@@ -57,4 +57,16 @@ private:
 using FElectraTextureSamplePtr = TSharedPtr<FElectraTextureSampleLinux, ESPMode::ThreadSafe>;
 using FElectraTextureSampleRef = TSharedRef<FElectraTextureSampleLinux, ESPMode::ThreadSafe>;
 using FElectraTextureSample = FElectraTextureSampleLinux;
-using FElectraTextureSamplePool = TMediaObjectPool<FElectraTextureSampleLinux>;
+
+class ELECTRASAMPLES_API FElectraTextureSamplePool : public TMediaObjectPool<FElectraTextureSample, FElectraTextureSamplePool>
+{
+	using TextureSample = FElectraTextureSample;
+public:
+	FElectraTextureSamplePool()
+		: TMediaObjectPool<TextureSample, FElectraTextureSamplePool>(this)
+	{ }
+	~FElectraTextureSamplePool()
+	{ }
+	TextureSample *Alloc() const
+	{ return new TextureSample(); }
+};
diff --git a/Engine/Plugins/Media/ElectraUtil/Source/ElectraSamples/Public/Windows/ElectraTextureSample.h b/Engine/Plugins/Media/ElectraUtil/Source/ElectraSamples/Public/Windows/ElectraTextureSample.h
--- a/Engine/Plugins/Media/ElectraUtil/Source/ElectraSamples/Public/Windows/ElectraTextureSample.h
+++ b/Engine/Plugins/Media/ElectraUtil/Source/ElectraSamples/Public/Windows/ElectraTextureSample.h
@@ -77,4 +77,16 @@ private:
 
 using FElectraTextureSamplePtr = TSharedPtr<FElectraTextureSample, ESPMode::ThreadSafe>;
 using FElectraTextureSampleRef = TSharedRef<FElectraTextureSample, ESPMode::ThreadSafe>;
-using FElectraTextureSamplePool = TMediaObjectPool<FElectraTextureSample>;
+
+class ELECTRASAMPLES_API FElectraTextureSamplePool : public TMediaObjectPool<FElectraTextureSample, FElectraTextureSamplePool>
+{
+	using TextureSample = FElectraTextureSample;
+public:
+	FElectraTextureSamplePool()
+		: TMediaObjectPool<TextureSample, FElectraTextureSamplePool>(this)
+	{ }
+	~FElectraTextureSamplePool()
+	{ }
+	TextureSample *Alloc() const
+	{ return new TextureSample(); }
+};
diff --git a/Engine/Plugins/Media/FlacDecoderElectra/Source/FlacDecoderElectra/Private/FlacDecoder/ElectraMediaFlacDecoder.cpp b/Engine/Plugins/Media/FlacDecoderElectra/Source/FlacDecoderElectra/Private/FlacDecoder/ElectraMediaFlacDecoder.cpp
--- a/Engine/Plugins/Media/FlacDecoderElectra/Source/FlacDecoderElectra/Private/FlacDecoder/ElectraMediaFlacDecoder.cpp
+++ b/Engine/Plugins/Media/FlacDecoderElectra/Source/FlacDecoderElectra/Private/FlacDecoder/ElectraMediaFlacDecoder.cpp
@@ -12,7 +12,6 @@
 #include "IElectraDecoderFeaturesAndOptions.h"
 #include "IElectraDecoder.h"
 #include "IElectraDecoderOutputAudio.h"
-#include "IElectraDecoderResourceDelegate.h"
 #include "ElectraDecodersUtils.h"
 #include "Utils/MPEG/ElectraUtilsMP4.h"		// for parsing the 'dfLa' box
 #include "Utils/ElectraBitstreamProcessorDefault.h"
@@ -30,31 +29,6 @@
 /*********************************************************************************************************************/
 /*********************************************************************************************************************/
 
-class FElectraDecoderDefaultAudioOutputFormatFlac_Common : public IElectraDecoderDefaultAudioOutputFormat
-{
-public:
-	virtual ~FElectraDecoderDefaultAudioOutputFormatFlac_Common()
-	{ }
-
-	int32 GetNumChannels() const override
-	{
-		return NumChannels;
-	}
-	int32 GetSampleRate() const override
-	{
-		return SampleRate;
-	}
-	int32 GetNumFrames() const override
-	{
-		return NumFrames;
-	}
-
-	int32 NumChannels = 0;
-	int32 SampleRate = 0;
-	int32 NumFrames = 0;
-};
-
-
 class FElectraAudioDecoderOutputFlac_Common : public IElectraDecoderAudioOutput
 {
 public:
@@ -127,7 +101,7 @@ public:
 	{
 	}
 
-	FElectraFlacDecoder(const TMap<FString, FVariant>& InOptions, TSharedPtr<IElectraDecoderResourceDelegate, ESPMode::ThreadSafe> InResourceDelegate);
+	FElectraFlacDecoder(const TMap<FString, FVariant>& InOptions);
 
 	virtual ~FElectraFlacDecoder();
 
@@ -144,8 +118,6 @@ public:
 	IElectraDecoder::ECSDCompatibility IsCompatibleWith(const TMap<FString, FVariant>& CSDAndAdditionalOptions) override;
 	bool ResetToCleanStart() override;
 
-	TSharedPtr<IElectraDecoderDefaultOutputFormat, ESPMode::ThreadSafe> GetDefaultOutputFormatFromCSD(const TMap<FString, FVariant>& CSDAndAdditionalOptions) override;
-
 	EDecoderError DecodeAccessUnit(const FInputAccessUnit& InInputAccessUnit, const TMap<FString, FVariant>& InAdditionalOptions) override;
 	EDecoderError SendEndOfData() override;
 	EDecoderError Flush() override;
@@ -271,9 +243,9 @@ public:
 		FElectraFlacDecoder::GetConfigurationOptions(OutOptions);
 	}
 
-	TSharedPtr<IElectraDecoder, ESPMode::ThreadSafe> CreateDecoderForFormat(const FString& InCodecFormat, const TMap<FString, FVariant>& InOptions, TSharedPtr<IElectraDecoderResourceDelegate, ESPMode::ThreadSafe> InResourceDelegate) override
+	TSharedPtr<IElectraDecoder, ESPMode::ThreadSafe> CreateDecoderForFormat(const FString& InCodecFormat, const TMap<FString, FVariant>& InOptions) override
 	{
-		return MakeShared<FElectraFlacDecoder, ESPMode::ThreadSafe>(InOptions, InResourceDelegate);
+		return MakeShared<FElectraFlacDecoder, ESPMode::ThreadSafe>(InOptions);
 	}
 
 	static TSharedPtr<FElectraCommonAudioFlacDecoderFactory, ESPMode::ThreadSafe> Self;
@@ -315,7 +287,7 @@ TSharedPtr<IElectraCodecFactory, ESPMode::ThreadSafe> FElectraMediaFlacDecoder::
 /*********************************************************************************************************************/
 /*********************************************************************************************************************/
 
-FElectraFlacDecoder::FElectraFlacDecoder(const TMap<FString, FVariant>& InOptions, TSharedPtr<IElectraDecoderResourceDelegate, ESPMode::ThreadSafe> InResourceDelegate)
+FElectraFlacDecoder::FElectraFlacDecoder(const TMap<FString, FVariant>& InOptions)
 {
 	Codec4CC = (uint32)ElectraDecodersUtil::GetVariantValueSafeU64(InOptions, TEXT("codec_4cc"), 0);
 }
@@ -497,11 +469,6 @@ bool FElectraFlacDecoder::ResetToCleanStart()
 	return true;
 }
 
-TSharedPtr<IElectraDecoderDefaultOutputFormat, ESPMode::ThreadSafe> FElectraFlacDecoder::GetDefaultOutputFormatFromCSD(const TMap<FString, FVariant>& CSDAndAdditionalOptions)
-{
-	return nullptr;
-}
-
 bool FElectraFlacDecoder::SetupChannelMap()
 {
 	if (OutputChannelMap.Num())
diff --git a/Engine/Plugins/Media/HAPDecoderElectra/Source/HAPDecoderElectra/Private/HAPDecoder/ElectraMediaHAPDecoder.cpp b/Engine/Plugins/Media/HAPDecoderElectra/Source/HAPDecoderElectra/Private/HAPDecoder/ElectraMediaHAPDecoder.cpp
--- a/Engine/Plugins/Media/HAPDecoderElectra/Source/HAPDecoderElectra/Private/HAPDecoder/ElectraMediaHAPDecoder.cpp
+++ b/Engine/Plugins/Media/HAPDecoderElectra/Source/HAPDecoderElectra/Private/HAPDecoder/ElectraMediaHAPDecoder.cpp
@@ -12,9 +12,10 @@
 #include "IElectraDecoderFeaturesAndOptions.h"
 #include "IElectraDecoder.h"
 #include "IElectraDecoderOutputVideo.h"
-#include "IElectraDecoderResourceDelegate.h"
 #include "ElectraDecodersUtils.h"
 #include "Utils/ElectraBitstreamProcessorDefault.h"
+
+#include "ElectraDecodersPlatformResources.h"
 #include COMPILED_PLATFORM_HEADER(ElectraDecoderGPUBufferHelpers.h)
 
 #include "hap.h"
@@ -30,15 +31,6 @@
 class FVideoDecoderHAPElectra;
 
 
-class FDecoderDefaultVideoOutputFormatHAPElectra : public IElectraDecoderDefaultVideoOutputFormat
-{
-public:
-	virtual ~FDecoderDefaultVideoOutputFormatHAPElectra()
-	{ }
-
-};
-
-
 class FVideoDecoderOutputHAPElectra : public IElectraDecoderVideoOutput, public IElectraDecoderVideoOutputImageBuffers
 {
 public:
@@ -112,10 +104,6 @@ public:
 		}
 		return nullptr;
 	}
-	IElectraDecoderVideoOutputTransferHandle* GetTransferHandle() const override
-	{
-		return nullptr;
-	}
 	IElectraDecoderVideoOutput::EImageCopyResult CopyPlatformImage(IElectraDecoderVideoOutputCopyResources* InCopyResources) const override
 	{
 		return IElectraDecoderVideoOutput::EImageCopyResult::NotSupported;
@@ -245,7 +233,7 @@ public:
 #endif
 	} Buffer[Buffer_Max];
 #if ELECTRA_MEDIAGPUBUFFER_DX12
-	TSharedPtr<IElectraDecoderResourceDelegateBase::IAsyncConsecutiveTaskSync> Buffer_TaskSync;
+	TSharedPtr<FElectraDecodersPlatformResources::FAsyncConsecutiveTaskSync, ESPMode::ThreadSafe> Buffer_TaskSync;
 #endif
 };
 
@@ -259,7 +247,7 @@ public:
 		OutOptions.Emplace(IElectraDecoderFeature::IsAdaptive, FVariant(true));
 	}
 
-	FVideoDecoderHAPElectra(const TMap<FString, FVariant>& InOptions, TSharedPtr<IElectraDecoderResourceDelegate, ESPMode::ThreadSafe> InResourceDelegate);
+	FVideoDecoderHAPElectra(const TMap<FString, FVariant>& InOptions);
 
 	virtual ~FVideoDecoderHAPElectra();
 
@@ -276,8 +264,6 @@ public:
 	IElectraDecoder::ECSDCompatibility IsCompatibleWith(const TMap<FString, FVariant>& CSDAndAdditionalOptions) override;
 	bool ResetToCleanStart() override;
 
-	TSharedPtr<IElectraDecoderDefaultOutputFormat, ESPMode::ThreadSafe> GetDefaultOutputFormatFromCSD(const TMap<FString, FVariant>& CSDAndAdditionalOptions) override;
-
 	EDecoderError DecodeAccessUnit(const FInputAccessUnit& InInputAccessUnit, const TMap<FString, FVariant>& InAdditionalOptions) override;
 	EDecoderError SendEndOfData() override;
 	EDecoderError Flush() override;
@@ -321,8 +307,6 @@ private:
 	TSharedPtr<FVideoDecoderOutputHAPElectra, ESPMode::ThreadSafe> CurrentOutput;
 	bool bFlushPending = false;
 
-	TWeakPtr<IElectraDecoderResourceDelegate, ESPMode::ThreadSafe> ResourceDelegate;
-
 	uint32 MaxOutputBuffers;
 	struct FBufferInstanceInfo
 	{
@@ -331,8 +315,11 @@ private:
 #if ELECTRA_MEDIAGPUBUFFER_DX12
 		mutable TSharedPtr<FElectraMediaDecoderOutputBufferPool_DX12> D3D12ResourcePool;
 #endif
-	} InstanceBufferData[FVideoDecoderOutputHAPElectra::Buffer_Max];
-	mutable TSharedPtr<IElectraDecoderResourceDelegateBase::IAsyncConsecutiveTaskSync> TaskSync;
+	};
+	FBufferInstanceInfo InstanceBufferData[FVideoDecoderOutputHAPElectra::Buffer_Max];
+#if ELECTRA_MEDIAGPUBUFFER_DX12
+	mutable TSharedPtr<FElectraDecodersPlatformResources::FAsyncConsecutiveTaskSync, ESPMode::ThreadSafe> TaskSync;
+#endif
 };
 
 
@@ -408,9 +395,9 @@ public:
 		FVideoDecoderHAPElectra::GetConfigurationOptions(OutOptions);
 	}
 
-	TSharedPtr<IElectraDecoder, ESPMode::ThreadSafe> CreateDecoderForFormat(const FString& InCodecFormat, const TMap<FString, FVariant>& InOptions, TSharedPtr<IElectraDecoderResourceDelegate, ESPMode::ThreadSafe> InResourceDelegate) override
+	TSharedPtr<IElectraDecoder, ESPMode::ThreadSafe> CreateDecoderForFormat(const FString& InCodecFormat, const TMap<FString, FVariant>& InOptions) override
 	{
-		return MakeShared<FVideoDecoderHAPElectra, ESPMode::ThreadSafe>(InOptions, InResourceDelegate);
+		return MakeShared<FVideoDecoderHAPElectra, ESPMode::ThreadSafe>(InOptions);
 	}
 
 	static TSharedPtr<FHAPVideoDecoderElectraFactory, ESPMode::ThreadSafe> Self;
@@ -446,10 +433,8 @@ void FElectraMediaHAPDecoder::Shutdown()
 /*********************************************************************************************************************/
 /*********************************************************************************************************************/
 
-FVideoDecoderHAPElectra::FVideoDecoderHAPElectra(const TMap<FString, FVariant>& InOptions, TSharedPtr<IElectraDecoderResourceDelegate, ESPMode::ThreadSafe> InResourceDelegate)
+FVideoDecoderHAPElectra::FVideoDecoderHAPElectra(const TMap<FString, FVariant>& InOptions)
 {
-	ResourceDelegate = InResourceDelegate;
-
 	DisplayWidth = (int32)ElectraDecodersUtil::GetVariantValueSafeI64(InOptions, TEXT("width"), 0);
 	DisplayHeight = (int32)ElectraDecodersUtil::GetVariantValueSafeI64(InOptions, TEXT("height"), 0);
 	// All HAP formats have 4x4 pixel blocks.
@@ -538,11 +523,6 @@ bool FVideoDecoderHAPElectra::ResetToCleanStart()
 	return true;
 }
 
-TSharedPtr<IElectraDecoderDefaultOutputFormat, ESPMode::ThreadSafe> FVideoDecoderHAPElectra::GetDefaultOutputFormatFromCSD(const TMap<FString, FVariant>& CSDAndAdditionalOptions)
-{
-	return nullptr;
-}
-
 static bool ConvertHapTextureFormat(HapTextureFormat HapFmt, EElectraDecoderPlatformPixelFormat& Format, EElectraDecoderPlatformPixelEncoding& Encoding)
 {
 	switch (HapFmt)
@@ -643,12 +623,8 @@ IElectraDecoder::EDecoderError FVideoDecoderHAPElectra::DecodeAccessUnit(const F
 		int32 PlatformDeviceVersion = 0;
 		bool bUseGPUBuffers = false;
 #if ELECTRA_MEDIAGPUBUFFER_DX12
-		auto PinnedResourceDelegate = ResourceDelegate.Pin();
-		if (PinnedResourceDelegate.IsValid())
-		{
-			PinnedResourceDelegate->GetD3DDevice(&PlatformDevice, &PlatformDeviceVersion);
-			bUseGPUBuffers = (PlatformDevice && PlatformDeviceVersion >= 12000);
-		}
+		FElectraDecodersPlatformResources::GetD3DDeviceAndVersion(&PlatformDevice, &PlatformDeviceVersion);
+		bUseGPUBuffers = (PlatformDevice && PlatformDeviceVersion >= 12000);
 #endif
 
 		TSharedPtr<FVideoDecoderOutputHAPElectra, ESPMode::ThreadSafe> NewOutput = MakeShared<FVideoDecoderOutputHAPElectra>();
@@ -700,7 +676,7 @@ IElectraDecoder::EDecoderError FVideoDecoderHAPElectra::DecodeAccessUnit(const F
 				// Create a tasksync instance so we can have our own async jobs run in consecutive order
 				if (!TaskSync.IsValid())
 				{
-					TaskSync = PinnedResourceDelegate->CreateAsyncConsecutiveTaskSync();
+					TaskSync = FElectraDecodersPlatformResources::CreateAsyncConsecutiveTaskSync();
 				}
 
 				InstanceBufferData[BufIdx].D3D12ResourcePool->AllocateOutputDataAsBuffer(BufferInfo.GPUBuffer, ResourcePitch);
@@ -745,7 +721,7 @@ IElectraDecoder::EDecoderError FVideoDecoderHAPElectra::DecodeAccessUnit(const F
 			{
 				if (Result == HapResult_No_Error)
 				{
-					if (PinnedResourceDelegate->RunCodeAsync([BufferResource = BufferInfo.GPUBuffer.Resource, BufferFence = BufferInfo.GPUBuffer.Fence, BufferFenceValue = BufferInfo.GPUBuffer.FenceValue, ResourcePitch, BufferAddr, BufferPitch, DecodedPitch = DecodedHeight / BlockSizeY, TempBuffer]()
+					if (FElectraDecodersPlatformResources::RunCodeAsync([BufferResource = BufferInfo.GPUBuffer.Resource, BufferFence = BufferInfo.GPUBuffer.Fence, BufferFenceValue = BufferInfo.GPUBuffer.FenceValue, ResourcePitch, BufferAddr, BufferPitch, DecodedPitch = DecodedHeight / BlockSizeY, TempBuffer]()
 						{
 							uint8* DX12BufferAddr;
 							HRESULT Res = BufferResource->Map(0, nullptr, (void**)&DX12BufferAddr);
@@ -756,7 +732,7 @@ IElectraDecoder::EDecoderError FVideoDecoderHAPElectra::DecodeAccessUnit(const F
 							BufferResource->Unmap(0, nullptr);
 
 							BufferFence->Signal(BufferFenceValue);
-						}, TaskSync.Get()))
+						}, TaskSync))
 					{
 						// note: we share one task sync for all buffers -> this will lead any external task to be executed after ALL tasks we trigger have been done (aka: after the "last" buffer is ready)
 						NewOutput->Buffer_TaskSync = TaskSync;
diff --git a/Engine/Plugins/Media/OpusDecoderElectra/Source/OpusDecoderElectra/Private/OpusDecoder/ElectraMediaOpusDecoder.cpp b/Engine/Plugins/Media/OpusDecoderElectra/Source/OpusDecoderElectra/Private/OpusDecoder/ElectraMediaOpusDecoder.cpp
--- a/Engine/Plugins/Media/OpusDecoderElectra/Source/OpusDecoderElectra/Private/OpusDecoder/ElectraMediaOpusDecoder.cpp
+++ b/Engine/Plugins/Media/OpusDecoderElectra/Source/OpusDecoderElectra/Private/OpusDecoder/ElectraMediaOpusDecoder.cpp
@@ -12,7 +12,6 @@
 #include "IElectraDecoderFeaturesAndOptions.h"
 #include "IElectraDecoder.h"
 #include "IElectraDecoderOutputAudio.h"
-#include "IElectraDecoderResourceDelegate.h"
 #include "ElectraDecodersUtils.h"
 #include "Utils/ElectraBitstreamProcessorDefault.h"
 #include "Utils/MPEG/ElectraUtilsMP4.h"		// for parsing the 'dOps' box
@@ -29,31 +28,6 @@
 /*********************************************************************************************************************/
 /*********************************************************************************************************************/
 
-class FElectraDecoderDefaultAudioOutputFormatOpus_Common : public IElectraDecoderDefaultAudioOutputFormat
-{
-public:
-	virtual ~FElectraDecoderDefaultAudioOutputFormatOpus_Common()
-	{ }
-
-	int32 GetNumChannels() const override
-	{
-		return NumChannels;
-	}
-	int32 GetSampleRate() const override
-	{
-		return SampleRate;
-	}
-	int32 GetNumFrames() const override
-	{
-		return NumFrames;
-	}
-
-	int32 NumChannels = 0;
-	int32 SampleRate = 0;
-	int32 NumFrames = 0;
-};
-
-
 class FElectraAudioDecoderOutputOpus_Common : public IElectraDecoderAudioOutput
 {
 public:
@@ -127,7 +101,7 @@ public:
 	{
 	}
 
-	FElectraOpusDecoder(const TMap<FString, FVariant>& InOptions, TSharedPtr<IElectraDecoderResourceDelegate, ESPMode::ThreadSafe> InResourceDelegate);
+	FElectraOpusDecoder(const TMap<FString, FVariant>& InOptions);
 
 	virtual ~FElectraOpusDecoder();
 
@@ -144,8 +118,6 @@ public:
 	IElectraDecoder::ECSDCompatibility IsCompatibleWith(const TMap<FString, FVariant>& CSDAndAdditionalOptions) override;
 	bool ResetToCleanStart() override;
 
-	TSharedPtr<IElectraDecoderDefaultOutputFormat, ESPMode::ThreadSafe> GetDefaultOutputFormatFromCSD(const TMap<FString, FVariant>& CSDAndAdditionalOptions) override;
-
 	EDecoderError DecodeAccessUnit(const FInputAccessUnit& InInputAccessUnit, const TMap<FString, FVariant>& InAdditionalOptions) override;
 	EDecoderError SendEndOfData() override;
 	EDecoderError Flush() override;
@@ -280,9 +252,9 @@ public:
 		FElectraOpusDecoder::GetConfigurationOptions(OutOptions);
 	}
 
-	TSharedPtr<IElectraDecoder, ESPMode::ThreadSafe> CreateDecoderForFormat(const FString& InCodecFormat, const TMap<FString, FVariant>& InOptions, TSharedPtr<IElectraDecoderResourceDelegate, ESPMode::ThreadSafe> InResourceDelegate) override
+	TSharedPtr<IElectraDecoder, ESPMode::ThreadSafe> CreateDecoderForFormat(const FString& InCodecFormat, const TMap<FString, FVariant>& InOptions) override
 	{
-		return MakeShared<FElectraOpusDecoder, ESPMode::ThreadSafe>(InOptions, InResourceDelegate);
+		return MakeShared<FElectraOpusDecoder, ESPMode::ThreadSafe>(InOptions);
 	}
 
 	static TSharedPtr<FElectraCommonAudioOpusDecoderFactory, ESPMode::ThreadSafe> Self;
@@ -324,7 +296,7 @@ TSharedPtr<IElectraCodecFactory, ESPMode::ThreadSafe> FElectraMediaOpusDecoder::
 /*********************************************************************************************************************/
 /*********************************************************************************************************************/
 
-FElectraOpusDecoder::FElectraOpusDecoder(const TMap<FString, FVariant>& InOptions, TSharedPtr<IElectraDecoderResourceDelegate, ESPMode::ThreadSafe> InResourceDelegate)
+FElectraOpusDecoder::FElectraOpusDecoder(const TMap<FString, FVariant>& InOptions)
 {
 	Codec4CC = (uint32)ElectraDecodersUtil::GetVariantValueSafeU64(InOptions, TEXT("codec_4cc"), 0);
 }
@@ -525,11 +497,6 @@ bool FElectraOpusDecoder::ResetToCleanStart()
 	return true;
 }
 
-TSharedPtr<IElectraDecoderDefaultOutputFormat, ESPMode::ThreadSafe> FElectraOpusDecoder::GetDefaultOutputFormatFromCSD(const TMap<FString, FVariant>& CSDAndAdditionalOptions)
-{
-	return nullptr;
-}
-
 bool FElectraOpusDecoder::SetupChannelMap()
 {
 	if (OutputChannelMap.Num())
