Changing Electra player delegates to be thread-safe

diff --git a/Engine/Plugins/Media/ElectraPlayer/Source/ElectraPlayerRuntime/Public/IElectraPlayerInterface.h b/Engine/Plugins/Media/ElectraPlayer/Source/ElectraPlayerRuntime/Public/IElectraPlayerInterface.h
--- a/Engine/Plugins/Media/ElectraPlayer/Source/ElectraPlayerRuntime/Public/IElectraPlayerInterface.h
+++ b/Engine/Plugins/Media/ElectraPlayer/Source/ElectraPlayerRuntime/Public/IElectraPlayerInterface.h
@@ -17,10 +17,10 @@
 
 // ---------------------------------------------------------------------------------------------
 
-DECLARE_MULTICAST_DELEGATE_TwoParams(FElectraPlayerSendAnalyticMetricsDelegate, const TSharedPtr<IAnalyticsProviderET>& /*InAnalyticsProvider*/, const FGuid& /*InPlayerGuid*/);
-DECLARE_MULTICAST_DELEGATE_OneParam(FElectraPlayerSendAnalyticMetricsPerMinuteDelegate, const TSharedPtr<IAnalyticsProviderET>& /*InAnalyticsProvider*/);
-DECLARE_MULTICAST_DELEGATE_TwoParams(FElectraPlayerReportVideoStreamingErrorDelegate, const FGuid& /*InPlayerGuid*/, const FString& /*InLastError*/);
-DECLARE_MULTICAST_DELEGATE_FourParams(FElectraPlayerReportSubtitlesMetricsDelegate, const FGuid& /*InPlayerGuid*/, const FString& /*InURL*/, double /*InResponseTime*/, const FString& /*InLastError*/);
+DECLARE_TS_MULTICAST_DELEGATE_TwoParams(FElectraPlayerSendAnalyticMetricsDelegate, const TSharedPtr<IAnalyticsProviderET>& /*InAnalyticsProvider*/, const FGuid& /*InPlayerGuid*/);
+DECLARE_TS_MULTICAST_DELEGATE_OneParam(FElectraPlayerSendAnalyticMetricsPerMinuteDelegate, const TSharedPtr<IAnalyticsProviderET>& /*InAnalyticsProvider*/);
+DECLARE_TS_MULTICAST_DELEGATE_TwoParams(FElectraPlayerReportVideoStreamingErrorDelegate, const FGuid& /*InPlayerGuid*/, const FString& /*InLastError*/);
+DECLARE_TS_MULTICAST_DELEGATE_FourParams(FElectraPlayerReportSubtitlesMetricsDelegate, const FGuid& /*InPlayerGuid*/, const FString& /*InURL*/, double /*InResponseTime*/, const FString& /*InLastError*/);
 
 
 /**
