ElectraPlayer: Added current active and buffering bitrates of video and audio to GetMediaInfo()

#rnx

diff --git a/Engine/Plugins/Media/ElectraPlayer/Source/ElectraPlayerRuntime/Private/Runtime/Demuxer/ParserISO14496-12.cpp b/Engine/Plugins/Media/ElectraPlayer/Source/ElectraPlayerRuntime/Private/Runtime/Demuxer/ParserISO14496-12.cpp
--- a/Engine/Plugins/Media/ElectraPlayer/Source/ElectraPlayerRuntime/Private/Runtime/Demuxer/ParserISO14496-12.cpp
+++ b/Engine/Plugins/Media/ElectraPlayer/Source/ElectraPlayerRuntime/Private/Runtime/Demuxer/ParserISO14496-12.cpp
@@ -7904,8 +7904,8 @@ namespace Electra
 				// Expose the timecode in the track's metadata.
 				const FVariantValue StartTimecodeValue(InTrack->StartTimecode->ToString());
 				const FVariantValue StartTimecodeFrameRate(FrameRate.ToPrettyText().ToString());
-				InTrack->CodecInformation.GetExtras().Set(MetaDataKeyStartTimecodeValue, StartTimecodeValue);
-				InTrack->CodecInformation.GetExtras().Set(MetaDataKeyStartTimecodeFrameRate, StartTimecodeFrameRate);
+				InTrack->CodecInformation.GetExtras().Set(MediaInfoStartTimecodeValue, StartTimecodeValue);
+				InTrack->CodecInformation.GetExtras().Set(MediaInfoStartTimecodeFrameRate, StartTimecodeFrameRate);
 				return UEMEDIA_ERROR_OK;
 			}
 		}
@@ -8693,8 +8693,8 @@ namespace Electra
 						{
 							Track->StartTimecode = TimecodeTrack->StartTimecode;
 							FParamDict& TrackExtras = Track->CodecInformation.GetExtras();
-							TrackExtras.SetValueFrom(MetaDataKeyStartTimecodeValue, TimecodeTrack->CodecInformation.GetExtras());
-							TrackExtras.SetValueFrom(MetaDataKeyStartTimecodeFrameRate, TimecodeTrack->CodecInformation.GetExtras());
+							TrackExtras.SetValueFrom(MediaInfoStartTimecodeValue, TimecodeTrack->CodecInformation.GetExtras());
+							TrackExtras.SetValueFrom(MediaInfoStartTimecodeFrameRate, TimecodeTrack->CodecInformation.GetExtras());
 						}
 					}
 				}
@@ -8726,8 +8726,8 @@ namespace Electra
 
 				if (!IsTimecodeTrackReferenced(ParsedTrackInfo, Track->GetID()))
 				{
-					InPlayerSession->GetMutableOptions().SetValueFrom(MetaDataKeyStartTimecodeValue, Track->CodecInformation.GetExtras());
-					InPlayerSession->GetMutableOptions().SetValueFrom(MetaDataKeyStartTimecodeFrameRate, Track->CodecInformation.GetExtras());
+					InPlayerSession->GetMediaInfoDictionary().SetValueFrom(MediaInfoStartTimecodeValue, Track->CodecInformation.GetExtras());
+					InPlayerSession->GetMediaInfoDictionary().SetValueFrom(MediaInfoStartTimecodeFrameRate, Track->CodecInformation.GetExtras());
 					break;
 				}
 			}
diff --git a/Engine/Plugins/Media/ElectraPlayer/Source/ElectraPlayerRuntime/Private/Runtime/Player/AdaptivePlayerMetadataKeynames.h b/Engine/Plugins/Media/ElectraPlayer/Source/ElectraPlayerRuntime/Private/Runtime/Player/AdaptivePlayerMetadataKeynames.h
--- a/Engine/Plugins/Media/ElectraPlayer/Source/ElectraPlayerRuntime/Private/Runtime/Player/AdaptivePlayerMetadataKeynames.h
+++ b/Engine/Plugins/Media/ElectraPlayer/Source/ElectraPlayerRuntime/Private/Runtime/Player/AdaptivePlayerMetadataKeynames.h
@@ -6,8 +6,19 @@
 
 namespace Electra
 {
-	const FName MetaDataKeyStartTimecodeValue(TEXT("StartTimecodeValue"));
-	const FName MetaDataKeyStartTimecodeFrameRate(TEXT("StartTimecodeFrameRate"));
+	const FName MediaInfoStartTimecodeValue(TEXT("StartTimecodeValue"));				// String
+	const FName MediaInfoStartTimecodeFrameRate(TEXT("StartTimecodeFrameRate"));		// String
+
+	/** Bitrate of video currently being displayed. */
+	const FLazyName MediaInfoVideoBitrateCurrent(TEXT("VideoBitrateCurrent"));			// int64, bits per second
+	/** Bitrate of video currently being buffered ahead of time. */
+	const FLazyName MediaInfoVideoBitrateBuffering(TEXT("VideoBitrateBuffering"));		// int64, bits per second
+
+	/** Bitrate of audio currently being heard. */
+	const FLazyName MediaInfoAudioBitrateCurrent(TEXT("AudioBitrateCurrent"));			// int64, bits per second
+	/** Bitrate of audio currently being buffered ahead of time. */
+	const FLazyName MediaInfoAudioBitrateBuffering(TEXT("AudioBitrateBuffering"));		// int64, bits per second
+
 } // namespace Electra
 
 
diff --git a/Engine/Plugins/Media/ElectraPlayer/Source/ElectraPlayerRuntime/Private/Runtime/Player/AdaptiveStreamingPlayer.cpp b/Engine/Plugins/Media/ElectraPlayer/Source/ElectraPlayerRuntime/Private/Runtime/Player/AdaptiveStreamingPlayer.cpp
--- a/Engine/Plugins/Media/ElectraPlayer/Source/ElectraPlayerRuntime/Private/Runtime/Player/AdaptiveStreamingPlayer.cpp
+++ b/Engine/Plugins/Media/ElectraPlayer/Source/ElectraPlayerRuntime/Private/Runtime/Player/AdaptiveStreamingPlayer.cpp
@@ -9,6 +9,7 @@
 #include "Player/AdaptiveStreamingPlayerABR.h"
 #include "Player/PlayerLicenseKey.h"
 #include "Player/AdaptivePlayerOptionKeynames.h"
+#include "Player/AdaptivePlayerMetadataKeynames.h"
 #include "Utilities/Utilities.h"
 #include "Utilities/BCP47-Helpers.h"
 #include "Player/DRM/DRMManager.h"
@@ -268,7 +269,7 @@ void FAdaptiveStreamingPlayer::ModifyOptions(const FParamDict& InOptionsToSetOrC
 
 FVariantValue FAdaptiveStreamingPlayer::GetMediaInfo(FName InKey) const
 {
-	return PlayerOptions.GetValue(InKey);	// thread safe
+	return MediaInfoDictionary.GetValue(InKey);	// thread safe
 }
 
 void FAdaptiveStreamingPlayer::SetStaticResourceProviderCallback(const TSharedPtr<IAdaptiveStreamingPlayerResourceProvider, ESPMode::ThreadSafe>& InStaticResourceProvider)
@@ -1233,6 +1234,10 @@ IAdaptiveStreamingPlayerAEMSHandler* FAdaptiveStreamingPlayer::GetAEMSEventHandl
 {
 	return AEMSEventHandler;
 }
+FParamDictTS& FAdaptiveStreamingPlayer::GetMediaInfoDictionary()
+{
+	return MediaInfoDictionary;
+}
 FParamDictTS& FAdaptiveStreamingPlayer::GetMutableOptions()
 {
 	return PlayerOptions;
@@ -1618,8 +1623,9 @@ bool FAdaptiveStreamingPlayer::InternalHandleThreadMessages()
 						{
 							bool bDrastic = CurrentVideoStreamBitrate.Bitrate && SegmentQualityLevel < CurrentVideoStreamBitrate.QualityLevel-1;
 							DispatchEvent(FMetricEvent::ReportVideoQualityChange(SegmentBitrate, CurrentVideoStreamBitrate.Bitrate, bDrastic));
-							CurrentVideoStreamBitrate.Bitrate      = SegmentBitrate;
+							CurrentVideoStreamBitrate.Bitrate = SegmentBitrate;
 							CurrentVideoStreamBitrate.QualityLevel = SegmentQualityLevel;
+							GetMediaInfoDictionary().Set(MediaInfoVideoBitrateBuffering.Resolve(), FVariantValue((int64)SegmentBitrate));
 						}
 					}
 					// Audio bitrate change?
@@ -1631,8 +1637,9 @@ bool FAdaptiveStreamingPlayer::InternalHandleThreadMessages()
 						{
 							bool bDrastic = CurrentAudioStreamBitrate.Bitrate && SegmentQualityLevel < CurrentAudioStreamBitrate.QualityLevel-1;
 							DispatchEvent(FMetricEvent::ReportAudioQualityChange(SegmentBitrate, CurrentAudioStreamBitrate.Bitrate, bDrastic));
-							CurrentAudioStreamBitrate.Bitrate      = SegmentBitrate;
+							CurrentAudioStreamBitrate.Bitrate = SegmentBitrate;
 							CurrentAudioStreamBitrate.QualityLevel = SegmentQualityLevel;
+							GetMediaInfoDictionary().Set(MediaInfoAudioBitrateBuffering.Resolve(), FVariantValue((int64)SegmentBitrate));
 						}
 					}
 				}
diff --git a/Engine/Plugins/Media/ElectraPlayer/Source/ElectraPlayerRuntime/Private/Runtime/Player/AdaptiveStreamingPlayer.h b/Engine/Plugins/Media/ElectraPlayer/Source/ElectraPlayerRuntime/Private/Runtime/Player/AdaptiveStreamingPlayer.h
--- a/Engine/Plugins/Media/ElectraPlayer/Source/ElectraPlayerRuntime/Private/Runtime/Player/AdaptiveStreamingPlayer.h
+++ b/Engine/Plugins/Media/ElectraPlayer/Source/ElectraPlayerRuntime/Private/Runtime/Player/AdaptiveStreamingPlayer.h
@@ -15,8 +15,6 @@
 #include "AdaptiveStreamingPlayerResourceRequest.h"
 #include "IElectraPlayerDataCache.h"
 
-class IOptionPointerValueContainer;
-
 namespace Electra
 {
 class IMediaRenderer;
diff --git a/Engine/Plugins/Media/ElectraPlayer/Source/ElectraPlayerRuntime/Private/Runtime/Player/AdaptiveStreamingPlayerDecoders.cpp b/Engine/Plugins/Media/ElectraPlayer/Source/ElectraPlayerRuntime/Private/Runtime/Player/AdaptiveStreamingPlayerDecoders.cpp
--- a/Engine/Plugins/Media/ElectraPlayer/Source/ElectraPlayerRuntime/Private/Runtime/Player/AdaptiveStreamingPlayerDecoders.cpp
+++ b/Engine/Plugins/Media/ElectraPlayer/Source/ElectraPlayerRuntime/Private/Runtime/Player/AdaptiveStreamingPlayerDecoders.cpp
@@ -3,6 +3,7 @@
 #include "PlayerCore.h"
 #include "Player/AdaptiveStreamingPlayerInternal.h"
 #include "Player/AdaptivePlayerOptionKeynames.h"
+#include "Player/AdaptivePlayerMetadataKeynames.h"
 #include "Decoder/SubtitleDecoder.h"
 #include "Decoder/AudioDecoder.h"
 #include "Utilities/Utilities.h"
@@ -259,6 +260,7 @@ int32 FAdaptiveStreamingPlayer::CreateDecoder(EStreamType type)
 				{
 					VideoDecoder.CurrentCodecInfo = AccessUnit->AUCodecData->ParsedInfo;
 					VideoDecoder.LastSentAUCodecData = AccessUnit->AUCodecData;
+					GetMediaInfoDictionary().Set(MediaInfoVideoBitrateCurrent.Resolve(), FVariantValue((int64)VideoDecoder.CurrentCodecInfo.GetBitrate()));
 					DispatchEvent(FMetricEvent::ReportCodecFormatChange(VideoDecoder.CurrentCodecInfo));
 				}
 				FAccessUnit::Release(AccessUnit);
@@ -324,6 +326,7 @@ int32 FAdaptiveStreamingPlayer::CreateDecoder(EStreamType type)
 				{
 					AudioDecoder.CurrentCodecInfo = AccessUnit->AUCodecData->ParsedInfo;
 					AudioDecoder.LastSentAUCodecData = AccessUnit->AUCodecData;
+					GetMediaInfoDictionary().Set(MediaInfoAudioBitrateCurrent.Resolve(), FVariantValue((int64)AudioDecoder.CurrentCodecInfo.GetBitrate()));
 					DispatchEvent(FMetricEvent::ReportCodecFormatChange(AudioDecoder.CurrentCodecInfo));
 				}
 				FAccessUnit::Release(AccessUnit);
@@ -745,15 +748,32 @@ void FAdaptiveStreamingPlayer::FeedDecoder(EStreamType Type, IAccessUnitBufferIn
 					}
 				}
 
-				// If there is any pertinent format change, emit an event.
-				if (AccessUnit->AUCodecData.IsValid() && LastSentAUCodecData && LastSentAUCodecData->IsValid() && *LastSentAUCodecData != AccessUnit->AUCodecData)
+				// Check for format changes.
+				if (AccessUnit->AUCodecData.IsValid() && LastSentAUCodecData && LastSentAUCodecData->IsValid())
 				{
-					bool bFormatChanged = AccessUnit->AUCodecData->ParsedInfo.GetBitrate() != (*LastSentAUCodecData)->ParsedInfo.GetBitrate() || !AccessUnit->AUCodecData->ParsedInfo.Equals((*LastSentAUCodecData)->ParsedInfo);
-					if (bFormatChanged)
+					const bool bBitrateChanged = AccessUnit->AUCodecData->ParsedInfo.GetBitrate() != (*LastSentAUCodecData)->ParsedInfo.GetBitrate();
+					// If there is any pertinent format change, emit an event.
+					if (*LastSentAUCodecData != AccessUnit->AUCodecData)
 					{
-						DispatchEvent(FMetricEvent::ReportCodecFormatChange(AccessUnit->AUCodecData->ParsedInfo));
+						bool bFormatChanged = bBitrateChanged || !AccessUnit->AUCodecData->ParsedInfo.Equals((*LastSentAUCodecData)->ParsedInfo);
+						if (bFormatChanged)
+						{
+							DispatchEvent(FMetricEvent::ReportCodecFormatChange(AccessUnit->AUCodecData->ParsedInfo));
+						}
+						*LastSentAUCodecData = AccessUnit->AUCodecData;
+					}
+					// Record the currently active bitrate.
+					if (bBitrateChanged)
+					{
+						if (Type == EStreamType::Video)
+						{
+							GetMediaInfoDictionary().Set(MediaInfoVideoBitrateCurrent.Resolve(), FVariantValue((int64)((*LastSentAUCodecData)->ParsedInfo.GetBitrate())));
+						}
+						else if (Type == EStreamType::Audio)
+						{
+							GetMediaInfoDictionary().Set(MediaInfoAudioBitrateCurrent.Resolve(), FVariantValue((int64)((*LastSentAUCodecData)->ParsedInfo.GetBitrate())));
+						}
 					}
-					*LastSentAUCodecData = AccessUnit->AUCodecData;
 				}
 
 				// Check for presence of encoder latency and store it.
diff --git a/Engine/Plugins/Media/ElectraPlayer/Source/ElectraPlayerRuntime/Private/Runtime/Player/AdaptiveStreamingPlayerInternal.h b/Engine/Plugins/Media/ElectraPlayer/Source/ElectraPlayerRuntime/Private/Runtime/Player/AdaptiveStreamingPlayerInternal.h
--- a/Engine/Plugins/Media/ElectraPlayer/Source/ElectraPlayerRuntime/Private/Runtime/Player/AdaptiveStreamingPlayerInternal.h
+++ b/Engine/Plugins/Media/ElectraPlayer/Source/ElectraPlayerRuntime/Private/Runtime/Player/AdaptiveStreamingPlayerInternal.h
@@ -1206,6 +1206,7 @@ private:
 	TSharedPtrTS<IPlayerEntityCache> GetEntityCache() override;
 	TSharedPtrTS<IHTTPResponseCache> GetHTTPResponseCache() override;
 	IAdaptiveStreamingPlayerAEMSHandler* GetAEMSEventHandler() override;
+	FParamDictTS& GetMediaInfoDictionary() override;
 	FParamDictTS& GetMutableOptions() override;
 	bool HaveOptionValue(const FName& InOption) override;
 	const FVariantValue GetOptionValue(const FName& InOption) override;
@@ -2276,6 +2277,7 @@ private:
 	FWorkerThreadMessages												WorkerThread;
 	IAdaptiveStreamingPlayer::FCreateParam::EWorkerThreads				UseSharedWorkerThreads;
 
+	FParamDictTS														MediaInfoDictionary;
 	FParamDictTS														PlayerOptions;
 	FPlaybackState														PlaybackState;
 	ISynchronizedUTCTime*												SynchronizedUTCTime;
diff --git a/Engine/Plugins/Media/ElectraPlayer/Source/ElectraPlayerRuntime/Private/Runtime/Player/AdaptiveStreamingPlayerManifest.cpp b/Engine/Plugins/Media/ElectraPlayer/Source/ElectraPlayerRuntime/Private/Runtime/Player/AdaptiveStreamingPlayerManifest.cpp
--- a/Engine/Plugins/Media/ElectraPlayer/Source/ElectraPlayerRuntime/Private/Runtime/Player/AdaptiveStreamingPlayerManifest.cpp
+++ b/Engine/Plugins/Media/ElectraPlayer/Source/ElectraPlayerRuntime/Private/Runtime/Player/AdaptiveStreamingPlayerManifest.cpp
@@ -492,7 +492,7 @@ IPlayerSessionServices::ECustomPropertyResult FAdaptiveStreamingPlayer::Validate
 			TSharedRef<TJsonWriter<TCHAR, TCondensedJsonPrintPolicy<TCHAR>>> JsonWriter = TJsonWriterFactory<TCHAR, TCondensedJsonPrintPolicy<TCHAR>>::Create(&PropertyJSON, 0);
 			if (FJsonSerializer::Serialize(RootObject, JsonWriter))
 			{
-				PlayerOptions.Set(AsName, FVariantValue(PropertyJSON));
+				GetMediaInfoDictionary().Set(AsName, FVariantValue(PropertyJSON));
 				return IPlayerSessionServices::ECustomPropertyResult::Accept;
 			}
 		}
diff --git a/Engine/Plugins/Media/ElectraPlayer/Source/ElectraPlayerRuntime/Private/Runtime/Player/HLS/PlaylistHLS.cpp b/Engine/Plugins/Media/ElectraPlayer/Source/ElectraPlayerRuntime/Private/Runtime/Player/HLS/PlaylistHLS.cpp
--- a/Engine/Plugins/Media/ElectraPlayer/Source/ElectraPlayerRuntime/Private/Runtime/Player/HLS/PlaylistHLS.cpp
+++ b/Engine/Plugins/Media/ElectraPlayer/Source/ElectraPlayerRuntime/Private/Runtime/Player/HLS/PlaylistHLS.cpp
@@ -579,6 +579,7 @@ FErrorDetail FActiveHLSPlaylist::CreateTrackMetadata(IPlayerSessionServices* InP
 					if (GroupStreamInfs[ns].ParsedCodecs[nc].IsCodec(InType))
 					{
 						sm.CodecInformation = GroupStreamInfs[ns].ParsedCodecs[nc];
+						sm.CodecInformation.SetBitrate(sm.Bandwidth);
 						break;
 					}
 				}
@@ -632,6 +633,7 @@ FErrorDetail FActiveHLSPlaylist::CreateTrackMetadata(IPlayerSessionServices* InP
 									sm.QualityIndex = GroupStreamInfs[vi].QualityIndex;
 									sm.ID.Empty();	// do not set an ID here to indicate this is a rendition
 									sm.CodecInformation = rg->Renditions[k].ParsedCodecFromStreamInf;
+									sm.CodecInformation.SetBitrate(sm.Bandwidth);
 									if (sm.Bandwidth > tm.Meta.HighestBandwidth)
 									{
 										tm.Meta.HighestBandwidth = sm.Bandwidth;
@@ -694,6 +696,7 @@ FErrorDetail FActiveHLSPlaylist::CreateTrackMetadata(IPlayerSessionServices* InP
 					sm.Bandwidth = tm.Meta.HighestBandwidth;
 					sm.QualityIndex = 0;
 					sm.CodecInformation = tm.Meta.HighestBandwidthCodec;
+					sm.CodecInformation.SetBitrate(sm.Bandwidth);
 				}
 			}
 		}
@@ -758,6 +761,7 @@ FErrorDetail FActiveHLSPlaylist::CreateTrackMetadata(IPlayerSessionServices* InP
 						sm.Bandwidth = tm.Meta.HighestBandwidth;
 						sm.QualityIndex = 0;
 						sm.CodecInformation = tm.Meta.HighestBandwidthCodec;
+						sm.CodecInformation.SetBitrate(sm.Bandwidth);
 					}
 					else
 					{
@@ -772,6 +776,7 @@ FErrorDetail FActiveHLSPlaylist::CreateTrackMetadata(IPlayerSessionServices* InP
 							sm.Bandwidth = InPathway->StreamInfs[InPathway->AudioOnlyVariantGroups[tm.AudioVariantGroupIndex].StreamInfIndices[nv]].Bandwidth;
 							sm.QualityIndex = InPathway->StreamInfs[InPathway->AudioOnlyVariantGroups[tm.AudioVariantGroupIndex].StreamInfIndices[nv]].QualityIndex;
 							sm.CodecInformation = InPathway->StreamInfs[InPathway->AudioOnlyVariantGroups[tm.AudioVariantGroupIndex].StreamInfIndices[nv]].ParsedCodecs[0];
+							sm.CodecInformation.SetBitrate(sm.Bandwidth);
 							if (sm.Bandwidth > tm.Meta.HighestBandwidth)
 							{
 								tm.Meta.HighestBandwidth = sm.Bandwidth;
@@ -814,6 +819,7 @@ FErrorDetail FActiveHLSPlaylist::CreateTrackMetadata(IPlayerSessionServices* InP
 				sm.Bandwidth = tm.Meta.HighestBandwidth;
 				sm.QualityIndex = 0;
 				sm.CodecInformation = tm.Meta.HighestBandwidthCodec;
+				sm.CodecInformation.SetBitrate(sm.Bandwidth);
 			}
 		}
 	}
diff --git a/Engine/Plugins/Media/ElectraPlayer/Source/ElectraPlayerRuntime/Private/Runtime/Player/PlayerSessionServices.h b/Engine/Plugins/Media/ElectraPlayer/Source/ElectraPlayerRuntime/Private/Runtime/Player/PlayerSessionServices.h
--- a/Engine/Plugins/Media/ElectraPlayer/Source/ElectraPlayerRuntime/Private/Runtime/Player/PlayerSessionServices.h
+++ b/Engine/Plugins/Media/ElectraPlayer/Source/ElectraPlayerRuntime/Private/Runtime/Player/PlayerSessionServices.h
@@ -133,6 +133,12 @@ namespace Electra
 		 */
 		virtual IAdaptiveStreamingPlayerAEMSHandler* GetAEMSEventHandler() = 0;
 
+		/**
+		 * Returns the dictionary used in queries via GetMediaInfo().
+		 * This is used to expose metadata through the IMediaPlayer interface.
+		 */
+		virtual FParamDictTS& GetMediaInfoDictionary() = 0;
+
 		/**
 		 * Returns the mutable player option dictionary. Other than initial options this serves as an interface to
 		 * pass values between internal sub systems.
