ElectraProtron: Fixed an array out of bounds crash when the file path could not be parsed as a URL

#rnx

diff --git a/Engine/Plugins/Media/ElectraPlayer/Source/ElectraProtronFactory/Private/ElectraProtronFactoryModule.cpp b/Engine/Plugins/Media/ElectraPlayer/Source/ElectraProtronFactory/Private/ElectraProtronFactoryModule.cpp
--- a/Engine/Plugins/Media/ElectraPlayer/Source/ElectraProtronFactory/Private/ElectraProtronFactoryModule.cpp
+++ b/Engine/Plugins/Media/ElectraPlayer/Source/ElectraProtronFactory/Private/ElectraProtronFactoryModule.cpp
@@ -75,7 +75,7 @@ public:
 		// Check for known extensions
 		TArray<FString> PathComponents;
 		UrlParser.GetPathComponents(PathComponents);
-		FString LowerCaseExtension = FPaths::GetExtension(PathComponents.Last().ToLower());
+		FString LowerCaseExtension = PathComponents.Num() ? FPaths::GetExtension(PathComponents.Last().ToLower()) : FString();
 		// If the extension is known, we are confident that we can play this.
 		// At this point there is no information provided on the codecs used in the media,
 		// so we cannot check for this.
