Reworked ElectraSamples, ElectraUtils and the decoder output of Electra Player

diff --git a/Engine/Plugins/Media/ElectraCodecs/Source/ElectraDecoders/Private/Utils/MPEG/ElectraUtilsMPEGAudio.cpp b/Engine/Plugins/Media/ElectraCodecs/Source/ElectraDecoders/Private/Utils/MPEG/ElectraUtilsMPEGAudio.cpp
--- a/Engine/Plugins/Media/ElectraCodecs/Source/ElectraDecoders/Private/Utils/MPEG/ElectraUtilsMPEGAudio.cpp
+++ b/Engine/Plugins/Media/ElectraCodecs/Source/ElectraDecoders/Private/Utils/MPEG/ElectraUtilsMPEGAudio.cpp
@@ -1,20 +1,18 @@
 // Copyright Epic Games, Inc. All Rights Reserved.
 
 #include "Utils/MPEG/ElectraUtilsMPEGAudio.h"
-#include "Utils/ElectraBitstreamReader.h"
+#include "Utilities/ElectraBitstream.h"
 
 namespace ElectraDecodersUtil
 {
 	namespace MPEG
 	{
-
-
 		bool FESDescriptor::Parse(const TConstArrayView<uint8>& InESDS)
 		{
-			class FMP4BitReader : public FElectraBitstreamReader
+			class FMP4BitReader : public Electra::FBitstreamReader
 			{
 			public:
-				FMP4BitReader(const void* pData, int64 nData) : FElectraBitstreamReader(pData, nData)
+				FMP4BitReader(const void* pData, int64 nData) : FBitstreamReader(pData, nData)
 				{
 				}
 				int32 ReadMP4Length()
@@ -114,32 +112,62 @@ namespace ElectraDecodersUtil
 
 
 
-
-
-
-
-
-
-
-
 		namespace AACUtils
 		{
+			static uint32 GIndexToSampleRate[16] =
+			{
+				96000, 88200, 64000, 48000, 44100, 32000, 24000, 22050, 16000, 12000, 11025, 8000, 7350, 0, 0, 0
+			};
+
 			int32 GetNumberOfChannelsFromChannelConfiguration(uint32 InChannelConfiguration)
 			{
 				static const uint8 NumChannelsForConfig[16] = { 0, 1, 2, 3, 4, 5, 6, 7, 0, 0, 0, 7, 8, 0, 8, 0 };
 				return InChannelConfiguration < 16 ? NumChannelsForConfig[InChannelConfiguration] : 0;
 			}
+
+			int32 GetSampleRateFromFrequenceIndex(uint32 InSamplingFrequencyIndex)
+			{
+				return InSamplingFrequencyIndex < 13 ? GIndexToSampleRate[InSamplingFrequencyIndex] : 0;
+			}
+
+			bool ParseADTSHeader(ADTSheader& OutHeader, const TConstArrayView<uint8>& InHeaderData)
+			{
+				if (InHeaderData.Num() < 7)
+				{
+					return false;
+				}
+
+				Electra::FBitstreamReader bs(InHeaderData.GetData(), InHeaderData.Num());
+				OutHeader.SyncWord = bs.GetBits(12);
+				OutHeader.Version = bs.GetBits(1);
+				OutHeader.Layer = bs.GetBits(2);
+				OutHeader.ProtectionAbsent = bs.GetBits(1);
+				OutHeader.Profile = bs.GetBits(2);
+				OutHeader.SamplingFrequencyIndex = bs.GetBits(4);
+				OutHeader.PrivateStream = bs.GetBits(1);
+				OutHeader.ChannelConfiguration = bs.GetBits(3);
+				OutHeader.Originality = bs.GetBits(1);
+				OutHeader.Home = bs.GetBits(1);
+				OutHeader.Copyrighted = bs.GetBits(1);
+				OutHeader.CopyrightStart = bs.GetBits(1);
+				OutHeader.FrameLength = bs.GetBits(13);
+				OutHeader.BufferFullness = bs.GetBits(11);
+				OutHeader.NumFrames = bs.GetBits(2);
+				OutHeader.HeaderSize = 7;
+				if (InHeaderData.Num() >= 9 && !OutHeader.ProtectionAbsent)
+				{
+					OutHeader.CRCifPresent = bs.GetBits(16);
+					OutHeader.HeaderSize += 2;
+				}
+				return OutHeader.SyncWord == 0xfff && OutHeader.Layer == 0 && OutHeader.SamplingFrequencyIndex < 13;
+			}
+
 		}
 
 
 		namespace AACParseHelper
 		{
-			static uint32 GIndexToSampleRate[16] =
-			{
-				96000, 88200, 64000, 48000, 44100, 32000, 24000, 22050, 16000, 12000, 11025, 8000, 0, 0, 0, 0
-			};
-
-			static uint32 GetAudioObjectType(FElectraBitstreamReader& Bitstream)
+			static uint32 GetAudioObjectType(Electra::FBitstreamReader& Bitstream)
 			{
 				uint32 ObjectType = Bitstream.GetBits(5);
 				if (ObjectType == 31)
@@ -148,7 +176,7 @@ namespace ElectraDecodersUtil
 				}
 				return ObjectType;
 			}
-			static uint32 GetSamplingRateIndex(FElectraBitstreamReader& Bitstream, uint32& Rate)
+			static uint32 GetSamplingRateIndex(Electra::FBitstreamReader& Bitstream, uint32& Rate)
 			{
 				uint32 SamplingFrequencyIndex = Bitstream.GetBits(4);
 				if (SamplingFrequencyIndex == 15)
@@ -157,11 +185,11 @@ namespace ElectraDecodersUtil
 				}
 				else
 				{
-					Rate = GIndexToSampleRate[SamplingFrequencyIndex];
+					Rate = AACUtils::GIndexToSampleRate[SamplingFrequencyIndex];
 				}
 				return SamplingFrequencyIndex;
 			}
-			static void GetGASpecificConfig(FElectraBitstreamReader& Bitstream, uint32 SamplingFrequencyIndex, uint32 ChannelConfiguration, uint32 AudioObjectType)
+			static void GetGASpecificConfig(Electra::FBitstreamReader& Bitstream, uint32 SamplingFrequencyIndex, uint32 ChannelConfiguration, uint32 AudioObjectType)
 			{
 				int32 FrameLengthFlag = Bitstream.GetBits(1);
 				int32 DependsOnCoreDecoder = Bitstream.GetBits(1);
@@ -249,7 +277,7 @@ namespace ElectraDecodersUtil
 			CodecSpecificData.SetNumUninitialized(Size);
 			FMemory::Memcpy(CodecSpecificData.GetData(), Data, Size);
 
-			FElectraBitstreamReader bsp(Data, Size);
+			Electra::FBitstreamReader bsp(Data, Size);
 			SBRSignal = -1;
 			PSSignal = -1;
 			SamplingRate = 0;
@@ -327,37 +355,31 @@ namespace ElectraDecodersUtil
 
 			int32 GetVersionId(uint32 InFrameHeader)
 			{
-				check(HasValidSync(InFrameHeader));
 				return (int32) ((InFrameHeader >> 19) & 3U);
 			}
 
 			int32 GetLayerIndex(uint32 InFrameHeader)
 			{
-				check(HasValidSync(InFrameHeader));
 				return (int32) ((InFrameHeader >> 17) & 3U);
 			}
 
 			int32 GetBitrateIndex(uint32 InFrameHeader)
 			{
-				check(HasValidSync(InFrameHeader));
 				return (int32) ((InFrameHeader >> 12) & 15U);
 			}
 
 			int32 GetSamplingRateIndex(uint32 InFrameHeader)
 			{
-				check(HasValidSync(InFrameHeader));
 				return (int32) ((InFrameHeader >> 10) & 3U);
 			}
 
 			int32 GetChannelMode(uint32 InFrameHeader)
 			{
-				check(HasValidSync(InFrameHeader));
 				return (int32) ((InFrameHeader >> 6) & 3U);
 			}
 
 			int32 GetNumPaddingBytes(uint32 InFrameHeader)
 			{
-				check(HasValidSync(InFrameHeader));
 				return (int32) ((InFrameHeader >> 9) & 1U);
 			}
 
@@ -366,13 +388,21 @@ namespace ElectraDecodersUtil
 				switch(GetVersionId(InFrameHeader))
 				{
 					default:
+					{
 						return 0;
+					}
 					case 0:
+					{
 						return 3;
+					}
 					case 2:
+					{
 						return 2;
+					}
 					case 3:
+					{
 						return 1;
+					}
 				}
 			}
 
@@ -382,13 +412,21 @@ namespace ElectraDecodersUtil
 				switch(GetLayerIndex(InFrameHeader))
 				{
 					default:
+					{
 						return 0;
+					}
 					case 1:
+					{
 						return 3;
+					}
 					case 2:
+					{
 						return 2;
+					}
 					case 3:
+					{
 						return 1;
+					}
 				}
 			}
 
@@ -477,74 +515,3 @@ namespace ElectraDecodersUtil
 
 	} // namespace MPEG
 } // namespace ElectraDecodersUtil
-
-#if 0
-
-struct ADTSheader
-{
-	// Refer to http://wiki.multimedia.cx/index.php?title=MPEG-4_Audio
-	unsigned mSyncWord : 12;		//!< All set
-	unsigned mVersion : 1;			//!< 0=MPEG-4, 1=MPEG-2
-	unsigned mLayer : 2;			//!< 0
-	unsigned mProtAbsent : 1;		//!< 0 if CRC present, 1 if CRC absent
-	unsigned mProfile : 2;			//!< MPEG-4 audio object type -1
-	unsigned mSmpFrqIndex : 4;		//!< MPEG-4 sampling frequency index (15 is forbidden)
-	unsigned mPrivateStream : 1;
-	unsigned mChannelCfg : 3;		//!< MPEG-4 channel configuration
-	unsigned mOriginality : 1;
-	unsigned mHome : 1;
-	unsigned mCopyrighted : 1;
-	unsigned mCopyrightStart : 1;
-	unsigned mFrameLength : 13;
-	unsigned mBufferFullness : 11;
-	unsigned mNumFrames : 2;
-	unsigned mCRCifPresent : 16;
-};
-
-static bool GetADTSHeader(ADTSheader& header, const void* pData, int64 nDataSize)
-{
-	if (nDataSize < 7)
-	{
-		return false;
-	}
-
-	FElectraBitstreamReader bs(pData, nDataSize);
-	header.mSyncWord = bs.GetBits(12);
-	header.mVersion = bs.GetBits(1);
-	header.mLayer = bs.GetBits(2);
-	header.mProtAbsent = bs.GetBits(1);
-	header.mProfile = bs.GetBits(2);
-	header.mSmpFrqIndex = bs.GetBits(4);
-	header.mPrivateStream = bs.GetBits(1);
-	header.mChannelCfg = bs.GetBits(3);
-	header.mOriginality = bs.GetBits(1);
-	header.mHome = bs.GetBits(1);
-	header.mCopyrighted = bs.GetBits(1);
-	header.mCopyrightStart = bs.GetBits(1);
-	header.mFrameLength = bs.GetBits(13);
-	header.mBufferFullness = bs.GetBits(11);
-	header.mNumFrames = bs.GetBits(2);
-	if (nDataSize >= 9 && !header.mProtAbsent)
-	{
-		header.mCRCifPresent = bs.GetBits(16);
-	}
-	return header.mSyncWord == 0xfff && header.mLayer == 0;
-}
-
-/*
-According to the Audio Coding Wiki, which is no longer available as of June, 2010, the ADIF format actually is just one header at the beginning of the AAC file.
-And the rest of the data are consecutive raw data blocks.
-This file format is meant for simple local storing purposes unlike ADTS or LATM which are meant for streaming AAC.
-The ADIF header is made up of the following Field name/Field size in bits/Comment (separated by semicolons):
-	adif_id/32/Always: "ADIF";
-	copyright_id_present/1/[none];
-	copyright_id/72/only if copyright_id_present == 1;
-	original_copy/1/[none];
-	home/1/[none];
-	bitstream_type/1/0: CBR, 1: VBR; bitrate/23/for CBR: bitrate, for VBR: peak bitrate, 0 means unknown;
-	num_program_config_elements/4/[none].
-	The next 2 fields come (num_program_config_elements+1) times: buffer_fullness/20/only if bitstream_type == 0; program_config_element/VAR/[none].
-*/
-
-
-#endif
diff --git a/Engine/Plugins/Media/ElectraCodecs/Source/ElectraDecoders/Private/Utils/MPEG/ElectraUtilsMPEGVideo.cpp b/Engine/Plugins/Media/ElectraCodecs/Source/ElectraDecoders/Private/Utils/MPEG/ElectraUtilsMPEGVideo.cpp
--- a/Engine/Plugins/Media/ElectraCodecs/Source/ElectraDecoders/Private/Utils/MPEG/ElectraUtilsMPEGVideo.cpp
+++ b/Engine/Plugins/Media/ElectraCodecs/Source/ElectraDecoders/Private/Utils/MPEG/ElectraUtilsMPEGVideo.cpp
@@ -1,10 +1,10 @@
 // Copyright Epic Games, Inc. All Rights Reserved.
 
 #include "Utils/MPEG/ElectraUtilsMPEGVideo.h"
-#include "Utils/ElectraBitstreamReader.h"
+#include "Utilities/ElectraBitstream.h"
+#include "Utilities/UtilitiesMP4.h"
 #include "Utils/MPEG/ElectraUtilsMPEGVideo_H264.h"
 #include "Utils/MPEG/ElectraUtilsMPEGVideo_H265.h"
-#include "Utils/MPEG/ElectraUtilsMP4.h"
 #include "ElectraDecodersUtils.h"
 
 namespace ElectraDecodersUtil
@@ -64,7 +64,7 @@ namespace ElectraDecodersUtil
 
 			FScopedDataPtr<uint8> RBSP(FMemory::Malloc((uint32) InBitstreamLength));
 			int32 RBSPsize = EBSPtoRBSP(RBSP, static_cast<const uint8*>(InBitstream), (int32)InBitstreamLength);
-			FElectraBitstreamReader BitReader(RBSP, RBSPsize);
+			Electra::FBitstreamReader BitReader(RBSP, RBSPsize);
 
 			while(BitReader.GetRemainingByteLength() > 0)
 			{
@@ -153,7 +153,7 @@ namespace ElectraDecodersUtil
 		{
 			if (InMessage.PayloadType == FSEIMessage::EPayloadType::PT_mastering_display_colour_volume && InMessage.Message.Num() >= 24)
 			{
-				FElectraBitstreamReader BitReader(InMessage.Message.GetData(), InMessage.Message.Num());
+				Electra::FBitstreamReader BitReader(InMessage.Message.GetData(), InMessage.Message.Num());
 				for(int32 i=0; i<3; ++i)
 				{
 					OutMDCV.display_primaries_x[i] = BitReader.GetBits(16);
@@ -174,7 +174,7 @@ namespace ElectraDecodersUtil
 			{
 				return false;
 			}
-			FMP4AtomReader boxReader(InMDCVBox.GetData(), InMDCVBox.Num());
+			Electra::UtilitiesMP4::FMP4AtomReader boxReader(InMDCVBox);
 			for(int32 i=0; i<3; ++i)
 			{
 				if (!boxReader.Read(OutMDCV.display_primaries_x[i]) || !boxReader.Read(OutMDCV.display_primaries_y[i]))
@@ -194,7 +194,7 @@ namespace ElectraDecodersUtil
 		}
 		TArray<uint8> ELECTRADECODERS_API BuildMDCVBox(const FSEImastering_display_colour_volume& InFromDisplayColorVolume)
 		{
-			FElectraBitstreamWriter bw;
+			Electra::FBitstreamWriter bw;
 			// The order in the MDCV box is G,B,R. Since we assume that in this structure the order is RGB we swap it.
 			bw.PutBits((uint16)InFromDisplayColorVolume.display_primaries_x[1], 16);
 			bw.PutBits((uint16)InFromDisplayColorVolume.display_primaries_y[1], 16);
@@ -216,7 +216,7 @@ namespace ElectraDecodersUtil
 		{
 			if (InMessage.PayloadType == FSEIMessage::EPayloadType::PT_content_light_level_info && InMessage.Message.Num() >= 4)
 			{
-				FElectraBitstreamReader BitReader(InMessage.Message.GetData(), InMessage.Message.Num());
+				Electra::FBitstreamReader BitReader(InMessage.Message.GetData(), InMessage.Message.Num());
 				OutCLLI.max_content_light_level = BitReader.GetBits(16);
 				OutCLLI.max_pic_average_light_level = BitReader.GetBits(16);
 				return true;
@@ -244,7 +244,7 @@ namespace ElectraDecodersUtil
 			{
 				return false;
 			}
-			ElectraDecodersUtil::FMP4AtomReader boxReader(InCLLIBox.GetData(), InCLLIBox.Num());
+			Electra::UtilitiesMP4::FMP4AtomReader boxReader(InCLLIBox);
 			OutCLLI.max_content_light_level = 0;			// MaxCLL
 			OutCLLI.max_pic_average_light_level = 0;		// MaxFALL
 			if (!boxReader.Read(OutCLLI.max_content_light_level) || !boxReader.Read(OutCLLI.max_pic_average_light_level))
@@ -255,7 +255,7 @@ namespace ElectraDecodersUtil
 		}
 		TArray<uint8> BuildCLLIBox(const FSEIcontent_light_level_info& InFromContentLightLevel)
 		{
-			FElectraBitstreamWriter bw;
+			Electra::FBitstreamWriter bw;
 			bw.PutBits((uint16)InFromContentLightLevel.max_content_light_level, 16);
 			bw.PutBits((uint16)InFromContentLightLevel.max_pic_average_light_level, 16);
 			TArray<uint8> Box;
@@ -270,7 +270,7 @@ namespace ElectraDecodersUtil
 		{
 			if (InMessage.PayloadType == FSEIMessage::EPayloadType::PT_alternative_transfer_characteristics && InMessage.Message.Num() >= 1)
 			{
-				FElectraBitstreamReader BitReader(InMessage.Message.GetData(), InMessage.Message.Num());
+				Electra::FBitstreamReader BitReader(InMessage.Message.GetData(), InMessage.Message.Num());
 				OutATC.preferred_transfer_characteristics = BitReader.GetBits(8);
 				return true;
 			}
@@ -281,17 +281,17 @@ namespace ElectraDecodersUtil
 		{
 			if (InCOLRBox.Num())
 			{
-				FMP4AtomReader boxReader(InCOLRBox.GetData(), InCOLRBox.Num());
+				Electra::UtilitiesMP4::FMP4AtomReader boxReader(InCOLRBox);
 				uint32 Type;
 				if (boxReader.Read(Type))
 				{
-					if (Type == ElectraDecodersUtil::MakeMP4Atom('n','c','l','x') || Type == ElectraDecodersUtil::MakeMP4Atom('n', 'c', 'l', 'c'))
+					if (Type == Electra::UtilitiesMP4::MakeBoxAtom('n','c','l','x') || Type == Electra::UtilitiesMP4::MakeBoxAtom('n', 'c', 'l', 'c'))
 					{
 						uint16 colour_primaries, transfer_characteristics, matrix_coeffs;
 						uint8 video_full_range_flag = 0;
 						if (boxReader.Read(colour_primaries) && boxReader.Read(transfer_characteristics) && boxReader.Read(matrix_coeffs))
 						{
-							if (Type == ElectraDecodersUtil::MakeMP4Atom('n','c','l','x') && !boxReader.Read(video_full_range_flag))
+							if (Type == Electra::UtilitiesMP4::MakeBoxAtom('n','c','l','x') && !boxReader.Read(video_full_range_flag))
 							{
 								return false;
 							}
@@ -310,8 +310,8 @@ namespace ElectraDecodersUtil
 
 		TArray<uint8> ELECTRADECODERS_API BuildCOLRBox(const FCommonColorimetry& InFromColorimetry)
 		{
-			FElectraBitstreamWriter bw;
-			bw.PutBits(ElectraDecodersUtil::MakeMP4Atom('n','c','l','x'), 32);
+			Electra::FBitstreamWriter bw;
+			bw.PutBits(Electra::UtilitiesMP4::MakeBoxAtom('n','c','l','x'), 32);
 			bw.PutBits((uint16)InFromColorimetry.colour_primaries, 16);
 			bw.PutBits((uint16)InFromColorimetry.transfer_characteristics, 16);
 			bw.PutBits((uint16)InFromColorimetry.matrix_coeffs, 16);
diff --git a/Engine/Plugins/Media/ElectraCodecs/Source/ElectraDecoders/Private/Utils/MPEG/ElectraUtilsMPEGVideo_H264.cpp b/Engine/Plugins/Media/ElectraCodecs/Source/ElectraDecoders/Private/Utils/MPEG/ElectraUtilsMPEGVideo_H264.cpp
--- a/Engine/Plugins/Media/ElectraCodecs/Source/ElectraDecoders/Private/Utils/MPEG/ElectraUtilsMPEGVideo_H264.cpp
+++ b/Engine/Plugins/Media/ElectraCodecs/Source/ElectraDecoders/Private/Utils/MPEG/ElectraUtilsMPEGVideo_H264.cpp
@@ -1,7 +1,7 @@
 // Copyright Epic Games, Inc. All Rights Reserved.
 
 #include "Utils/MPEG/ElectraUtilsMPEGVideo_H264.h"
-#include "Utils/ElectraBitstreamReader.h"
+#include "Utilities/ElectraBitstream.h"
 #include "ElectraDecodersUtils.h"
 #include "ElectraDecodersModule.h"
 
@@ -139,11 +139,11 @@ namespace ElectraDecodersUtil
 				return true;
 			}
 
-			bool ParseScalingLists(const FSequenceParameterSet* InSPS, const FPictureParameterSet* InPPS, uint8 scaling_list_present_flag[12], uint8 ScalingList4x4[6][16], uint8 ScalingList8x8[6][64], FBitstreamReader& br)
+			bool ParseScalingLists(const FSequenceParameterSet* InSPS, const FPictureParameterSet* InPPS, uint8 scaling_list_present_flag[12], uint8 ScalingList4x4[6][16], uint8 ScalingList8x8[6][64], FBitstreamReaderH264& br)
 			{
 				struct sl
 				{
-					static bool scaling_list(uint8* ScalingList, int32 sizeOfScalingList, const uint8* DefaultScalingMatrix, FBitstreamReader& br)
+					static bool scaling_list(uint8* ScalingList, int32 sizeOfScalingList, const uint8* DefaultScalingMatrix, FBitstreamReaderH264& br)
 					{
 						const uint8* Scanninglist = sizeOfScalingList == 16 ? ScanOrderZZ4 : ScanOrderZZ8;
 						int32 lastScale = 8;
@@ -176,7 +176,7 @@ namespace ElectraDecodersUtil
 						}
 						return true;
 					}
-					static bool parse_scaling_list(uint8* ScalingList, int32 sizeOfScalingList, const uint8* DefaultScalingMatrix, const uint8* FallbackScalingMatrix, FBitstreamReader& br)
+					static bool parse_scaling_list(uint8* ScalingList, int32 sizeOfScalingList, const uint8* DefaultScalingMatrix, const uint8* FallbackScalingMatrix, FBitstreamReaderH264& br)
 					{
 						if (br.GetBits(1))
 						{
@@ -532,7 +532,7 @@ namespace ElectraDecodersUtil
 					}
 
 				TUniquePtr<FRBSP> RBSP(MakeRBSP(InBitstream, InBitstreamLenInBytes));
-				FBitstreamReader br(RBSP->Data.GetData(), RBSP->Data.Num());
+				FBitstreamReaderH264 br(RBSP->Data.GetData(), RBSP->Data.Num());
 
 				// Verify that this is in fact an SPS NALU.
 				if ((br.GetBits(8) & 0x1f) != 0x7)
@@ -792,7 +792,7 @@ namespace ElectraDecodersUtil
 					}
 
 				TUniquePtr<FRBSP> RBSP(MakeRBSP(InBitstream, InBitstreamLenInBytes));
-				FBitstreamReader br(RBSP->Data.GetData(), RBSP->Data.Num());
+				FBitstreamReaderH264 br(RBSP->Data.GetData(), RBSP->Data.Num());
 
 				// Verify that this is in fact a PPS NALU.
 				if ((br.GetBits(8) & 0x1f) != 0x8)
@@ -912,7 +912,7 @@ namespace ElectraDecodersUtil
 			}
 
 
-			bool ParseSliceHeader(TUniquePtr<FRBSP>& OutRBSP, FBitstreamReader& OutRBSPReader, FSliceHeader& OutSlice, const TMap<uint32, FSequenceParameterSet>& InSequenceParameterSets, const TMap<uint32, FPictureParameterSet>& InPictureParameterSets, const uint8* InBitstream, uint64 InBitstreamLenInBytes)
+			bool ParseSliceHeader(TUniquePtr<FRBSP>& OutRBSP, FBitstreamReaderH264& OutRBSPReader, FSliceHeader& OutSlice, const TMap<uint32, FSequenceParameterSet>& InSequenceParameterSets, const TMap<uint32, FPictureParameterSet>& InPictureParameterSets, const uint8* InBitstream, uint64 InBitstreamLenInBytes)
 			{
 				#define RANGE_CHECK_FAILURE(expr, elem)															\
 					if (!(expr))																				\
@@ -923,7 +923,7 @@ namespace ElectraDecodersUtil
 
 				OutRBSP = MakeRBSP(InBitstream, InBitstreamLenInBytes);
 				OutRBSPReader.SetData(OutRBSP->Data.GetData(), OutRBSP->Data.Num());
-				FBitstreamReader& br(OutRBSPReader);
+				FBitstreamReaderH264& br(OutRBSPReader);
 
 				// Verify that this is in fact a (supported) slice.
 				uint32 nut = br.GetBits(8);
@@ -2279,7 +2279,7 @@ namespace ElectraDecodersUtil
 				}
 
 				RawData = InDCR;
-				FElectraByteReader ByteReader(InDCR.GetData(), InDCR.Num());
+				Electra::FElectraByteReader ByteReader(InDCR.GetData(), InDCR.Num());
 
 				if (!ByteReader.ReadByte(ConfigurationVersion) || ConfigurationVersion != 1)
 				{
@@ -2484,7 +2484,7 @@ namespace ElectraDecodersUtil
 
 				// Create the raw configuration record now.
 				ConfigurationVersion = 1;
-				FElectraBitstreamWriter wr;
+				Electra::FBitstreamWriter wr;
 				wr.PutBits(ConfigurationVersion, 8);
 				wr.PutBits(AVCProfileIndication, 8);
 				wr.PutBits(ProfileCompatibility, 8);
diff --git a/Engine/Plugins/Media/ElectraCodecs/Source/ElectraDecoders/Private/Utils/MPEG/ElectraUtilsMPEGVideo_H265.cpp b/Engine/Plugins/Media/ElectraCodecs/Source/ElectraDecoders/Private/Utils/MPEG/ElectraUtilsMPEGVideo_H265.cpp
--- a/Engine/Plugins/Media/ElectraCodecs/Source/ElectraDecoders/Private/Utils/MPEG/ElectraUtilsMPEGVideo_H265.cpp
+++ b/Engine/Plugins/Media/ElectraCodecs/Source/ElectraDecoders/Private/Utils/MPEG/ElectraUtilsMPEGVideo_H265.cpp
@@ -1,7 +1,7 @@
 // Copyright Epic Games, Inc. All Rights Reserved.
 
 #include "Utils/MPEG/ElectraUtilsMPEGVideo_H265.h"
-#include "Utils/ElectraBitstreamReader.h"
+#include "Utilities/ElectraBitstream.h"
 #include "ElectraDecodersUtils.h"
 #include "ElectraDecodersModule.h"
 
@@ -343,7 +343,7 @@ namespace ElectraDecodersUtil
 			}
 
 
-			bool FHRDParameters::Parse(FBitstreamReader& br, bool commonInfPresentFlag, uint32 maxNumSubLayersMinus1)
+			bool FHRDParameters::Parse(FBitstreamReaderH265& br, bool commonInfPresentFlag, uint32 maxNumSubLayersMinus1)
 			{
 				#define RANGE_CHECK_FAILURE(expr, elem)															\
 					if (!(expr))																				\
@@ -431,7 +431,7 @@ namespace ElectraDecodersUtil
 				#undef RANGE_CHECK_FAILURE
 			}
 
-			bool FVUIParameters::Parse(FBitstreamReader& br, uint32 In_sps_max_sub_layers_minus1)
+			bool FVUIParameters::Parse(FBitstreamReaderH265& br, uint32 In_sps_max_sub_layers_minus1)
 			{
 				#define RANGE_CHECK_FAILURE(expr, elem)																\
 					if (!(expr))																					\
@@ -541,7 +541,7 @@ namespace ElectraDecodersUtil
 				#undef RANGE_CHECK_FAILURE
 			}
 
-			bool FScalingListData::Parse(FBitstreamReader& br)
+			bool FScalingListData::Parse(FBitstreamReaderH265& br)
 			{
 				for(uint32 sizeId=0; sizeId<4; ++sizeId)
 				{
@@ -641,7 +641,7 @@ namespace ElectraDecodersUtil
 				}
 			}
 
-			bool FProfileTierLevel::Parse(FBitstreamReader& br, bool profilePresentFlag, uint32 maxNumSubLayersMinus1)
+			bool FProfileTierLevel::Parse(FBitstreamReaderH265& br, bool profilePresentFlag, uint32 maxNumSubLayersMinus1)
 			{
 				#define RANGE_CHECK_FAILURE(expr, elem)																\
 					if (!(expr))																					\
@@ -659,7 +659,7 @@ namespace ElectraDecodersUtil
 					{
 						general_profile_compatibility_flag[i] = br.GetBits(1);
 					}
-					general_constraint_indicator_flags = FBitstreamReader(br).GetBits64(48);	// elaborate PeekBits(48)
+					general_constraint_indicator_flags = FBitstreamReaderH265(br).GetBits64(48);	// elaborate PeekBits(48)
 					general_progressive_source_flag = br.GetBits(1);
 					general_interlaced_source_flag = br.GetBits(1);
 					general_non_packed_constraint_flag = br.GetBits(1);
@@ -827,7 +827,7 @@ namespace ElectraDecodersUtil
 				#undef RANGE_CHECK_FAILURE
 			}
 
-			bool FStRefPicSet::Parse(FBitstreamReader& br, int32 stRpsIdx, const TArray<FStRefPicSet>& InReferencePicSetList, int32 InReferencePicSetListNum)
+			bool FStRefPicSet::Parse(FBitstreamReaderH265& br, int32 stRpsIdx, const TArray<FStRefPicSet>& InReferencePicSetList, int32 InReferencePicSetListNum)
 			{
 				#define RANGE_CHECK_FAILURE(expr, elem)																\
 					if (!(expr))																					\
@@ -944,7 +944,7 @@ namespace ElectraDecodersUtil
 						return false;																	\
 					}
 				TUniquePtr<FRBSP> RBSP(MakeRBSP(InBitstream, InBitstreamLenInBytes));
-				FBitstreamReader br(RBSP->Data, RBSP->Size);
+				FBitstreamReaderH265 br(RBSP->Data, RBSP->Size);
 
 				// Verify that this is in fact a VPS NALU.
 				if ((br.GetBits(16) >> 9) != 32)
@@ -1093,7 +1093,7 @@ namespace ElectraDecodersUtil
 					}
 
 				TUniquePtr<FRBSP> RBSP(MakeRBSP(InBitstream, InBitstreamLenInBytes));
-				FBitstreamReader br(RBSP->Data, RBSP->Size);
+				FBitstreamReaderH265 br(RBSP->Data, RBSP->Size);
 
 				// Verify that this is in fact an SPS NALU.
 				if ((br.GetBits(16) >> 9) != 33)
@@ -1340,7 +1340,7 @@ namespace ElectraDecodersUtil
 					}
 
 				TUniquePtr<FRBSP> RBSP(MakeRBSP(InBitstream, InBitstreamLenInBytes));
-				FBitstreamReader br(RBSP->Data, RBSP->Size);
+				FBitstreamReaderH265 br(RBSP->Data, RBSP->Size);
 
 				// Verify that this is in fact an PPS NALU.
 				if ((br.GetBits(16) >> 9) != 34)
@@ -1502,13 +1502,13 @@ namespace ElectraDecodersUtil
 				#undef RANGE_CHECK_FAILURE
 			}
 
-			bool FPictureParameterSet::FPPSRangeExtension::Parse(FBitstreamReader& br)
+			bool FPictureParameterSet::FPPSRangeExtension::Parse(FBitstreamReaderH265& br)
 			{
 				UE_LOG(LogElectraDecoders, Error, TEXT("Range extension is not supported"));
 				return false;
 			}
 
-			bool FPictureParameterSet::FPPSMultilayerExtension::Parse(FBitstreamReader& br)
+			bool FPictureParameterSet::FPPSMultilayerExtension::Parse(FBitstreamReaderH265& br)
 			{
 				UE_LOG(LogElectraDecoders, Error, TEXT("Multilayer extension is not supported"));
 				return false;
@@ -1579,19 +1579,19 @@ namespace ElectraDecodersUtil
 				*/
 			}
 
-			bool FPictureParameterSet::FPPS3DExtension::Parse(FBitstreamReader& br)
+			bool FPictureParameterSet::FPPS3DExtension::Parse(FBitstreamReaderH265& br)
 			{
 				UE_LOG(LogElectraDecoders, Error, TEXT("3D extension is not supported"));
 				return false;
 			}
 
-			bool FPictureParameterSet::FPPSSCCExtension::Parse(FBitstreamReader& br)
+			bool FPictureParameterSet::FPPSSCCExtension::Parse(FBitstreamReaderH265& br)
 			{
 				UE_LOG(LogElectraDecoders, Error, TEXT("SCC extension is not supported"));
 				return false;
 			}
 
-			bool ParseSliceHeader(TUniquePtr<FRBSP>& OutRBSP, FBitstreamReader& OutRBSPReader, FSliceSegmentHeader& OutSlice, const TMap<uint32, FVideoParameterSet>& InVideoParameterSets, const TMap<uint32, FSequenceParameterSet>& InSequenceParameterSets, const TMap<uint32, FPictureParameterSet>& InPictureParameterSets, const uint8* InBitstream, uint64 InBitstreamLenInBytes)
+			bool ParseSliceHeader(TUniquePtr<FRBSP>& OutRBSP, FBitstreamReaderH265& OutRBSPReader, FSliceSegmentHeader& OutSlice, const TMap<uint32, FVideoParameterSet>& InVideoParameterSets, const TMap<uint32, FSequenceParameterSet>& InSequenceParameterSets, const TMap<uint32, FPictureParameterSet>& InPictureParameterSets, const uint8* InBitstream, uint64 InBitstreamLenInBytes)
 			{
 				#define RANGE_CHECK_FAILURE(expr, elem)																	\
 					if (!(expr))																						\
@@ -1602,7 +1602,7 @@ namespace ElectraDecodersUtil
 
 				OutRBSP = MakeRBSP(InBitstream, InBitstreamLenInBytes);
 				OutRBSPReader.SetData(OutRBSP->Data, OutRBSP->Size);
-				FBitstreamReader& br(OutRBSPReader);
+				FBitstreamReaderH265& br(OutRBSPReader);
 
 				if (br.GetBits(1))	// forbidden_zero_bit
 				{
@@ -2591,7 +2591,7 @@ namespace ElectraDecodersUtil
 				{
 					RawData = InDCR;
 
-					FElectraBitstreamReader BitReader(InDCR.GetData(), InDCR.Num());
+					Electra::FBitstreamReader BitReader(InDCR.GetData(), InDCR.Num());
 
 					ConfigurationVersion = (uint8)BitReader.GetBits(8);
 					if (ConfigurationVersion != 1)
@@ -2844,7 +2844,7 @@ namespace ElectraDecodersUtil
 
 				// Create the raw configuration record now.
 				ConfigurationVersion = 1;
-				FElectraBitstreamWriter wr;
+				Electra::FBitstreamWriter wr;
 				wr.PutBits(ConfigurationVersion, 8);
 				wr.PutBits(GeneralProfileSpace, 2);
 				wr.PutBits(GeneralTierFlag, 1);
diff --git a/Engine/Plugins/Media/ElectraCodecs/Source/ElectraDecoders/Private/Windows/WindowsPlatformElectraDecoders.cpp b/Engine/Plugins/Media/ElectraCodecs/Source/ElectraDecoders/Private/Windows/WindowsPlatformElectraDecoders.cpp
--- a/Engine/Plugins/Media/ElectraCodecs/Source/ElectraDecoders/Private/Windows/WindowsPlatformElectraDecoders.cpp
+++ b/Engine/Plugins/Media/ElectraCodecs/Source/ElectraDecoders/Private/Windows/WindowsPlatformElectraDecoders.cpp
@@ -1,7 +1,6 @@
 // Copyright Epic Games, Inc. All Rights Reserved.
 
 #include "WindowsPlatformElectraDecoders.h"
-#include "WindowsElectraDecodersPlatformResources.h"
 #include "IElectraCodecRegistry.h"
 #include "ElectraDecodersModule.h"
 
@@ -11,27 +10,17 @@
 
 #include "HAL/PlatformProcess.h"
 #include "RHI.h"
-#include "Async/Async.h"
 
 #include "Microsoft/AllowMicrosoftPlatformTypes.h"
 THIRD_PARTY_INCLUDES_START
-#include "mftransform.h"
 #include "mfapi.h"
-#include "mferror.h"
-#include "mfidl.h"
 THIRD_PARTY_INCLUDES_END
 #include "Microsoft/HideMicrosoftPlatformTypes.h"
 
-DECLARE_CYCLE_STAT(TEXT("Electra decoder AsyncJob"), STAT_ElectraDecoderAsyncJob, STATGROUP_Media);
-
-
 namespace PlatformElectraDecodersWindows
 {
 	static TArray<TSharedPtr<IElectraCodecFactory, ESPMode::ThreadSafe>> Factories;
 	static bool bMFInitialized = false;
-
-	static FCriticalSection AsyncJobAccessCS;
-	static FGraphEventRef RunCodeAsyncEvent;
 }
 
 
@@ -103,52 +92,3 @@ void FPlatformElectraDecodersWindows::RegisterWithCodecFactory(IElectraCodecRegi
 #endif
 	}
 }
-
-
-bool FElectraDecodersPlatformResourcesWindows::GetD3DDeviceAndVersion(void** OutD3DDevice, int32* OutD3DVersion)
-{
-	if (!GDynamicRHI || !OutD3DDevice || !OutD3DVersion)
-	{
-		return false;
-	}
-	auto RHIType = RHIGetInterfaceType();
-	if (RHIType != ERHIInterfaceType::D3D11 && RHIType != ERHIInterfaceType::D3D12)
-	{
-		return false;
-	}
-	*OutD3DDevice = GDynamicRHI->RHIGetNativeDevice();
-	*OutD3DVersion = RHIType == ERHIInterfaceType::D3D11 ? 11000 : 12000;
-	return true;
-}
-
-
-
-class FElectraDecodersPlatformResourcesWindows::FAsyncConsecutiveTaskSync
-{
-public:
-	FAsyncConsecutiveTaskSync() = default;
-	~FAsyncConsecutiveTaskSync() = default;
-	FGraphEventRef GraphEvent;
-};
-
-TSharedPtr<FElectraDecodersPlatformResourcesWindows::FAsyncConsecutiveTaskSync, ESPMode::ThreadSafe> FElectraDecodersPlatformResourcesWindows::CreateAsyncConsecutiveTaskSync()
-{
-	return MakeShared<FElectraDecodersPlatformResourcesWindows::FAsyncConsecutiveTaskSync, ESPMode::ThreadSafe>();
-}
-
-bool FElectraDecodersPlatformResourcesWindows::RunCodeAsync(TFunction<void()>&& InCodeToRun, TSharedPtr<FElectraDecodersPlatformResourcesWindows::FAsyncConsecutiveTaskSync, ESPMode::ThreadSafe> InTaskSync)
-{
-	FScopeLock Lock(&PlatformElectraDecodersWindows::AsyncJobAccessCS);
-
-	// Execute code async
-	// (We assume this to be code copying buffer data around. Hence we allow only ONE at any time to not clog up
-	//  the buses even more and delay the copy process further)
-	FGraphEventArray Events;
-	auto& Event = InTaskSync ? InTaskSync->GraphEvent : PlatformElectraDecodersWindows::RunCodeAsyncEvent;
-	if (Event.IsValid())
-	{
-		Events.Add(Event);
-	}
-	Event = FFunctionGraphTask::CreateAndDispatchWhenReady(MoveTemp(InCodeToRun), GET_STATID(STAT_ElectraDecoderAsyncJob), &Events);
-	return true;
-}
diff --git a/Engine/Plugins/Media/ElectraCodecs/Source/ElectraDecoders/Public/Android/AndroidElectraDecodersPlatformResources.h b/Engine/Plugins/Media/ElectraCodecs/Source/ElectraDecoders/Public/Android/AndroidElectraDecodersPlatformResources.h
--- a/Engine/Plugins/Media/ElectraCodecs/Source/ElectraDecoders/Public/Android/AndroidElectraDecodersPlatformResources.h
+++ b/Engine/Plugins/Media/ElectraCodecs/Source/ElectraDecoders/Public/Android/AndroidElectraDecodersPlatformResources.h
@@ -3,7 +3,11 @@
 
 #include "CoreTypes.h"
 #include "Templates/SharedPointer.h"
-#include "Misc/Optional.h"
+
+#include "ParameterDictionary.h"
+
+class IElectraDecoderVideoOutput;
+class FElectraTextureSample;
 
 #include "Android/AndroidPlatform.h"
 #include "Android/AndroidJava.h"
@@ -43,6 +47,13 @@ public:
 
 	static ELECTRADECODERS_API IDecoderResourceInstance* CreateResourceInstance(int64 InResourceID);
 	static ELECTRADECODERS_API void ReleaseResourceInstance(IDecoderResourceInstance* InInstanceToRelease);
+
+
+	struct IDecoderPlatformResource;
+	static ELECTRADECODERS_API IDecoderPlatformResource* CreatePlatformVideoResource(void* InOwnerHandle, const TMap<FString, FVariant> InOptions);
+	static ELECTRADECODERS_API void ReleasePlatformVideoResource(void* InOwnerHandle, IDecoderPlatformResource* InHandleToRelease);
+
+	static ELECTRADECODERS_API bool SetupOutputTextureSample(FString& OutErrorMessage, TSharedPtr<FElectraTextureSample, ESPMode::ThreadSafe> InOutTextureSampleToSetup, TSharedPtr<IElectraDecoderVideoOutput, ESPMode::ThreadSafe> InDecoderOutput, TSharedPtr<Electra::FParamDict, ESPMode::ThreadSafe> InOutBufferPropertes, IDecoderPlatformResource* InPlatformResource);
 };
 
 using FElectraDecodersPlatformResources = FElectraDecodersPlatformResourcesAndroid;
diff --git a/Engine/Plugins/Media/ElectraCodecs/Source/ElectraDecoders/Public/Apple/AppleElectraDecoderPlatformOutputHandleTypes.h b/Engine/Plugins/Media/ElectraCodecs/Source/ElectraDecoders/Public/Apple/AppleElectraDecoderPlatformOutputHandleTypes.h
--- a/Engine/Plugins/Media/ElectraCodecs/Source/ElectraDecoders/Public/Apple/AppleElectraDecoderPlatformOutputHandleTypes.h
+++ b/Engine/Plugins/Media/ElectraCodecs/Source/ElectraDecoders/Public/Apple/AppleElectraDecoderPlatformOutputHandleTypes.h
@@ -3,42 +3,14 @@
 #pragma once
 
 #include "CoreMinimal.h"
+#include "PixelFormat.h"
+#include "IElectraTextureSampleEncoding.h"
 
 enum class EElectraDecoderPlatformOutputHandleType
 {
 	ImageBuffers			// IElectraDecoderVideoOutputImageBuffers interface
 };
 
-enum class EElectraDecoderPlatformPixelFormat
-{
-	INVALID = 0,
-	R8G8B8A8,
-	A8R8G8B8,
-	B8G8R8A8,
-	R16G16B16A16,
-	A16B16G16R16,
-	A32B32G32R32F,
-	A2B10G10R10,
-	DXT1,
-	DXT5,
-	BC4,
-	NV12,
-	P010,
-};
-
-enum class EElectraDecoderPlatformPixelEncoding
-{
-	Native = 0,		//!< Pixel formats native representation
-	RGB,			//!< Interpret as RGB
-	RGBA,			//!< Interpret as RGBA
-	YCbCr,			//!< Interpret as YCbCR
-	YCbCr_Alpha,	//!< Interpret as YCbCR with alpha
-	YCoCg,			//!< Interpret as scaled YCoCg
-	YCoCg_Alpha,	//!< Interpret as scaled YCoCg with trailing BC4 alpha data
-	CbY0CrY1,		//!< Interpret as CbY0CrY1
-	Y0CbY1Cr,		//!< Interpret as Y0CbY1Cr
-	ARGB_BigEndian,	//!< Interpret as ARGB, big endian
-};
 
 class FElectraDecoderOutputSync final
 {
@@ -57,8 +29,8 @@ public:
 	virtual int32 GetNumberOfBuffers() const = 0;
 
 	// Returns the n'th image buffer data address.
-	virtual TSharedPtr<TArray<uint8>, ESPMode::ThreadSafe> GetBufferDataByIndex(int32 InBufferIndex) const = 0;
-	
+	virtual TSharedPtr<TArray64<uint8>, ESPMode::ThreadSafe> GetBufferDataByIndex(int32 InBufferIndex) const = 0;
+
 	// Returns the n'th CFImageBuffer ("texture") reference.
 	virtual void* GetBufferTextureByIndex(int32 InBufferIndex) const = 0;
 
@@ -66,10 +38,10 @@ public:
 	virtual bool GetBufferTextureSyncByIndex(int32 InBufferIndex, FElectraDecoderOutputSync& SyncObject) const = 0;
 
 	// Returns the n'th image buffer format
-	virtual EElectraDecoderPlatformPixelFormat GetBufferFormatByIndex(int32 InBufferIndex) const = 0;
+	virtual EPixelFormat GetBufferFormatByIndex(int32 InBufferIndex) const = 0;
 
 	// Returns the n'th image buffer encoding
-	virtual EElectraDecoderPlatformPixelEncoding GetBufferEncodingByIndex(int32 InBufferIndex) const = 0;
+	virtual EElectraTextureSamplePixelEncoding GetBufferEncodingByIndex(int32 InBufferIndex) const = 0;
 
 	// Returns the n'th image buffer pitch
 	virtual int32 GetBufferPitchByIndex(int32 InBufferIndex) const = 0;
diff --git a/Engine/Plugins/Media/ElectraCodecs/Source/ElectraDecoders/Public/Apple/AppleElectraDecodersPlatformResources.h b/Engine/Plugins/Media/ElectraCodecs/Source/ElectraDecoders/Public/Apple/AppleElectraDecodersPlatformResources.h
--- a/Engine/Plugins/Media/ElectraCodecs/Source/ElectraDecoders/Public/Apple/AppleElectraDecodersPlatformResources.h
+++ b/Engine/Plugins/Media/ElectraCodecs/Source/ElectraDecoders/Public/Apple/AppleElectraDecodersPlatformResources.h
@@ -2,12 +2,27 @@
 #pragma once
 
 #include "CoreTypes.h"
+#include "Templates/SharedPointer.h"
 
+#if PLATFORM_MAC || PLATFORM_IOS || PLATFORM_TVOS
+
+#include "ParameterDictionary.h"
+
+class IElectraDecoderVideoOutput;
+class FElectraTextureSample;
 
 class FElectraDecodersPlatformResourcesApple
 {
 public:
     static ELECTRADECODERS_API bool GetMetalDevice(void** OutMetalDevice);
+
+    struct IDecoderPlatformResource;
+	static ELECTRADECODERS_API IDecoderPlatformResource* CreatePlatformVideoResource(void* InOwnerHandle, const TMap<FString, FVariant> InOptions);
+	static ELECTRADECODERS_API void ReleasePlatformVideoResource(void* InOwnerHandle, IDecoderPlatformResource* InHandleToRelease);
+
+	static ELECTRADECODERS_API bool SetupOutputTextureSample(FString& OutErrorMessage, TSharedPtr<FElectraTextureSample, ESPMode::ThreadSafe> InOutTextureSampleToSetup, TSharedPtr<IElectraDecoderVideoOutput, ESPMode::ThreadSafe> InDecoderOutput, TSharedPtr<Electra::FParamDict, ESPMode::ThreadSafe> InOutBufferPropertes, IDecoderPlatformResource* InPlatformResource);
 };
 
 using FElectraDecodersPlatformResources = FElectraDecodersPlatformResourcesApple;
+
+#endif
diff --git a/Engine/Plugins/Media/ElectraCodecs/Source/ElectraDecoders/Public/Linux/LinuxElectraDecoderPlatformOutputHandleTypes.h b/Engine/Plugins/Media/ElectraCodecs/Source/ElectraDecoders/Public/Linux/LinuxElectraDecoderPlatformOutputHandleTypes.h
--- a/Engine/Plugins/Media/ElectraCodecs/Source/ElectraDecoders/Public/Linux/LinuxElectraDecoderPlatformOutputHandleTypes.h
+++ b/Engine/Plugins/Media/ElectraCodecs/Source/ElectraDecoders/Public/Linux/LinuxElectraDecoderPlatformOutputHandleTypes.h
@@ -3,6 +3,8 @@
 #pragma once
 
 #include "CoreMinimal.h"
+#include "PixelFormat.h"
+#include "IElectraTextureSampleEncoding.h"
 
 enum class EElectraDecoderPlatformOutputHandleType
 {
@@ -10,36 +12,6 @@ enum class EElectraDecoderPlatformOutputHandleType
 	ImageBuffers						// IElectraDecoderVideoOutputImageBuffers interface
 };
 
-enum class EElectraDecoderPlatformPixelFormat
-{
-	INVALID = 0,
-	R8G8B8A8,
-	A8R8G8B8,
-	B8G8R8A8,
-	R16G16B16A16,
-	A16B16G16R16,
-	A32B32G32R32F,
-	A2B10G10R10,
-	DXT1,
-	DXT5,
-	BC4,
-	NV12,
-	P010,
-};
-
-enum class EElectraDecoderPlatformPixelEncoding
-{
-	Native = 0,		//!< Pixel formats native representation
-	RGB,			//!< Interpret as RGB
-	RGBA,			//!< Interpret as RGBA
-	YCbCr,			//!< Interpret as YCbCR
-	YCbCr_Alpha,	//!< Interpret as YCbCR with alpha
-	YCoCg,			//!< Interpret as scaled YCoCg
-	YCoCg_Alpha,	//!< Interpret as scaled YCoCg with trailing BC4 alpha data
-	CbY0CrY1,		//!< Interpret as CbY0CrY1
-	Y0CbY1Cr,		//!< Interpret as Y0CbY1Cr
-	ARGB_BigEndian,	//!< Interpret as ARGB, big endian
-};
 
 class FElectraDecoderOutputSync final
 {
@@ -58,7 +30,7 @@ public:
 	virtual int32 GetNumberOfBuffers() const = 0;
 
 	// Returns the n'th image buffer as CPU data buffer.
-	virtual TSharedPtr<TArray<uint8>, ESPMode::ThreadSafe> GetBufferDataByIndex(int32 InBufferIndex) const = 0;
+	virtual TSharedPtr<TArray64<uint8>, ESPMode::ThreadSafe> GetBufferDataByIndex(int32 InBufferIndex) const = 0;
 
 	// Returns the n'th image buffer as GPU texture resource.
 	virtual void* GetBufferTextureByIndex(int32 InBufferIndex) const = 0;
@@ -67,10 +39,10 @@ public:
 	virtual bool GetBufferTextureSyncByIndex(int32 InBufferIndex, FElectraDecoderOutputSync& SyncObject) const = 0;
 
 	// Returns the n'th image buffer format
-	virtual EElectraDecoderPlatformPixelFormat GetBufferFormatByIndex(int32 InBufferIndex) const = 0;
+	virtual EPixelFormat GetBufferFormatByIndex(int32 InBufferIndex) const = 0;
 
 	// Returns the n'th image buffer encoding
-	virtual EElectraDecoderPlatformPixelEncoding GetBufferEncodingByIndex(int32 InBufferIndex) const = 0;
+	virtual EElectraTextureSamplePixelEncoding GetBufferEncodingByIndex(int32 InBufferIndex) const = 0;
 
 	// Returns the n'th image buffer pitch
 	virtual int32 GetBufferPitchByIndex(int32 InBufferIndex) const = 0;
diff --git a/Engine/Plugins/Media/ElectraCodecs/Source/ElectraDecoders/Public/Linux/LinuxElectraDecodersPlatformResources.h b/Engine/Plugins/Media/ElectraCodecs/Source/ElectraDecoders/Public/Linux/LinuxElectraDecodersPlatformResources.h
--- a/Engine/Plugins/Media/ElectraCodecs/Source/ElectraDecoders/Public/Linux/LinuxElectraDecodersPlatformResources.h
+++ b/Engine/Plugins/Media/ElectraCodecs/Source/ElectraDecoders/Public/Linux/LinuxElectraDecodersPlatformResources.h
@@ -2,12 +2,21 @@
 #pragma once
 
 #include "CoreTypes.h"
+#include "Templates/SharedPointer.h"
 
+#include "ParameterDictionary.h"
+
+class IElectraDecoderVideoOutput;
+class FElectraTextureSample;
 
 class FElectraDecodersPlatformResourcesLinux
 {
 public:
-//	static ELECTRADECODERS_API bool NotUsedRightNow();
+	struct IDecoderPlatformResource;
+	static ELECTRADECODERS_API IDecoderPlatformResource* CreatePlatformVideoResource(void* InOwnerHandle, const TMap<FString, FVariant> InOptions);
+	static ELECTRADECODERS_API void ReleasePlatformVideoResource(void* InOwnerHandle, IDecoderPlatformResource* InHandleToRelease);
+
+	static ELECTRADECODERS_API bool SetupOutputTextureSample(FString& OutErrorMessage, TSharedPtr<FElectraTextureSample, ESPMode::ThreadSafe> InOutTextureSampleToSetup, TSharedPtr<IElectraDecoderVideoOutput, ESPMode::ThreadSafe> InDecoderOutput, TSharedPtr<Electra::FParamDict, ESPMode::ThreadSafe> InOutBufferPropertes, IDecoderPlatformResource* InPlatformResource);
 };
 
 using FElectraDecodersPlatformResources = FElectraDecodersPlatformResourcesLinux;
diff --git a/Engine/Plugins/Media/ElectraCodecs/Source/ElectraDecoders/Public/Utils/MPEG/ElectraUtilsMPEGAudio.h b/Engine/Plugins/Media/ElectraCodecs/Source/ElectraDecoders/Public/Utils/MPEG/ElectraUtilsMPEGAudio.h
--- a/Engine/Plugins/Media/ElectraCodecs/Source/ElectraDecoders/Public/Utils/MPEG/ElectraUtilsMPEGAudio.h
+++ b/Engine/Plugins/Media/ElectraCodecs/Source/ElectraDecoders/Public/Utils/MPEG/ElectraUtilsMPEGAudio.h
@@ -109,7 +109,32 @@ namespace ElectraDecodersUtil
 
 		namespace AACUtils
 		{
+			struct ADTSheader
+			{
+				uint16 SyncWord = 0;
+				uint8 Version = 0;
+				uint8 Layer = 0;
+				uint8 ProtectionAbsent = 0;
+				uint8 Profile = 0;
+				uint8 SamplingFrequencyIndex = 0;
+				uint8 PrivateStream = 0;
+				uint8 ChannelConfiguration = 0;
+				uint8 Originality = 0;
+				uint8 Home = 0;
+				uint8 Copyrighted = 0;
+				uint8 CopyrightStart = 0;
+				uint16 FrameLength = 0;
+				uint16 BufferFullness = 0;
+				uint8 NumFrames = 0;
+				uint16 CRCifPresent = 0;
+				//
+				uint16 HeaderSize = 0;
+			};
+
 			int32 ELECTRADECODERS_API GetNumberOfChannelsFromChannelConfiguration(uint32 InChannelConfiguration);
+			int32 ELECTRADECODERS_API GetSampleRateFromFrequenceIndex(uint32 InSamplingFrequencyIndex);
+
+			bool ELECTRADECODERS_API ParseADTSHeader(ADTSheader& OutHeader, const TConstArrayView<uint8>& InHeaderData);
 		}
 
 
diff --git a/Engine/Plugins/Media/ElectraCodecs/Source/ElectraDecoders/Public/Utils/MPEG/ElectraUtilsMPEGVideo_H264.h b/Engine/Plugins/Media/ElectraCodecs/Source/ElectraDecoders/Public/Utils/MPEG/ElectraUtilsMPEGVideo_H264.h
--- a/Engine/Plugins/Media/ElectraCodecs/Source/ElectraDecoders/Public/Utils/MPEG/ElectraUtilsMPEGVideo_H264.h
+++ b/Engine/Plugins/Media/ElectraCodecs/Source/ElectraDecoders/Public/Utils/MPEG/ElectraUtilsMPEGVideo_H264.h
@@ -4,7 +4,7 @@
 
 #include <CoreMinimal.h>
 #include <Containers/Array.h>
-#include "Utils/ElectraBitstreamReader.h"
+#include "Utilities/ElectraBitstream.h"
 #include "ElectraDecodersUtils.h"
 
 #define UE_API ELECTRADECODERS_API
@@ -31,7 +31,7 @@ namespace ElectraDecodersUtil
 				{
 					Data.SetNumUninitialized((int32)InDataSize);
 				}
-				void TrimToBitstream(const FElectraBitstreamReader& InBSR)
+				void TrimToBitstream(const Electra::FBitstreamReader& InBSR)
 				{
 					if (InBSR.IsByteAligned())
 					{
@@ -50,12 +50,12 @@ namespace ElectraDecodersUtil
 			uint64 ELECTRADECODERS_API EBSPtoRBSP(TArray<uint64>& OutRemovedInBufPos, uint8* OutBuf, const uint8* InBuf, uint64 InNumBytes);
 			TUniquePtr<FRBSP> ELECTRADECODERS_API MakeRBSP(const uint8* InData, uint64 InDataSize);
 
-			class FBitstreamReader : public FElectraBitstreamReader
+			class FBitstreamReaderH264 : public Electra::FBitstreamReader
 			{
 			public:
-				FBitstreamReader() = default;
-				FBitstreamReader(const uint8* InData, uint64 InDataSize)
-					: FElectraBitstreamReader(InData, InDataSize)
+				FBitstreamReaderH264() = default;
+				FBitstreamReaderH264(const uint8* InData, uint64 InDataSize)
+					: Electra::FBitstreamReader(InData, InDataSize)
 				{ }
 
 				uint32 ue_v()
@@ -593,7 +593,7 @@ namespace ElectraDecodersUtil
 
 			bool ELECTRADECODERS_API ParseSequenceParameterSet(TMap<uint32, FSequenceParameterSet>& InOutSequenceParameterSets, const uint8* InBitstream, uint64 InBitstreamLenInBytes);
 			bool ELECTRADECODERS_API ParsePictureParameterSet(TMap<uint32, FPictureParameterSet>& InOutPictureParameterSets, const TMap<uint32, FSequenceParameterSet>& InSequenceParameterSets, const uint8* InBitstream, uint64 InBitstreamLenInBytes);
-			bool ELECTRADECODERS_API ParseSliceHeader(TUniquePtr<FRBSP>& OutRBSP, FBitstreamReader& OutRBSPReader, FSliceHeader& OutSlice, const TMap<uint32, FSequenceParameterSet>& InSequenceParameterSets, const TMap<uint32, FPictureParameterSet>& InPictureParameterSets, const uint8* InBitstream, uint64 InBitstreamLenInBytes);
+			bool ELECTRADECODERS_API ParseSliceHeader(TUniquePtr<FRBSP>& OutRBSP, FBitstreamReaderH264& OutRBSPReader, FSliceHeader& OutSlice, const TMap<uint32, FSequenceParameterSet>& InSequenceParameterSets, const TMap<uint32, FPictureParameterSet>& InPictureParameterSets, const uint8* InBitstream, uint64 InBitstreamLenInBytes);
 
 			bool ELECTRADECODERS_API ParseSequenceParameterSet(FSequenceParameterSet& OutSequenceParameterSet, const uint8* InBitstream, uint64 InBitstreamLenInBytes);
 			inline bool ParseSequenceParameterSet(FSequenceParameterSet& OutSequenceParameterSet, const TArray<uint8>& InData)
diff --git a/Engine/Plugins/Media/ElectraCodecs/Source/ElectraDecoders/Public/Utils/MPEG/ElectraUtilsMPEGVideo_H265.h b/Engine/Plugins/Media/ElectraCodecs/Source/ElectraDecoders/Public/Utils/MPEG/ElectraUtilsMPEGVideo_H265.h
--- a/Engine/Plugins/Media/ElectraCodecs/Source/ElectraDecoders/Public/Utils/MPEG/ElectraUtilsMPEGVideo_H265.h
+++ b/Engine/Plugins/Media/ElectraCodecs/Source/ElectraDecoders/Public/Utils/MPEG/ElectraUtilsMPEGVideo_H265.h
@@ -4,7 +4,7 @@
 
 #include <CoreMinimal.h>
 #include <Containers/Array.h>
-#include "Utils/ElectraBitstreamReader.h"
+#include "Utilities/ElectraBitstream.h"
 #include "ElectraDecodersUtils.h"
 
 #define UE_API ELECTRADECODERS_API
@@ -44,12 +44,12 @@ namespace ElectraDecodersUtil
 			uint64 ELECTRADECODERS_API EBSPtoRBSP(TArray<uint64>& OutRemovedInBufPos, uint8* OutBuf, const uint8* InBuf, uint64 InNumBytes);
 			TUniquePtr<FRBSP> ELECTRADECODERS_API MakeRBSP(const uint8* InData, uint64 InDataSize);
 
-			class FBitstreamReader : public FElectraBitstreamReader
+			class FBitstreamReaderH265 : public Electra::FBitstreamReader
 			{
 			public:
-				FBitstreamReader() = default;
-				FBitstreamReader(const uint8* InData, uint64 InDataSize)
-					: FElectraBitstreamReader(InData, InDataSize)
+				FBitstreamReaderH265() = default;
+				FBitstreamReaderH265(const uint8* InData, uint64 InDataSize)
+					: Electra::FBitstreamReader(InData, InDataSize)
 				{ }
 
 				uint32 ue_v()
@@ -122,7 +122,7 @@ namespace ElectraDecodersUtil
 				{
 					Reset();
 				}
-				UE_API bool Parse(FBitstreamReader& br, bool profilePresentFlag, uint32 maxNumSubLayersMinus1);
+				UE_API bool Parse(FBitstreamReaderH265& br, bool profilePresentFlag, uint32 maxNumSubLayersMinus1);
 				void Reset()
 				{
 					general_profile_space = 0;
@@ -235,7 +235,7 @@ namespace ElectraDecodersUtil
 
 			struct FHRDParameters
 			{
-				UE_API bool Parse(FBitstreamReader& br, bool commonInfPresentFlag, uint32 maxNumSubLayersMinus1);
+				UE_API bool Parse(FBitstreamReaderH265& br, bool commonInfPresentFlag, uint32 maxNumSubLayersMinus1);
 				void Reset()
 				{
 					nal_hrd_parameters_present_flag = 0;
@@ -299,7 +299,7 @@ namespace ElectraDecodersUtil
 				{
 					Reset();
 				}
-				UE_API bool Parse(FBitstreamReader& br, uint32 In_sps_max_sub_layers_minus1);
+				UE_API bool Parse(FBitstreamReaderH265& br, uint32 In_sps_max_sub_layers_minus1);
 				void Reset()
 				{
 					aspect_ratio_info_present_flag = 0;
@@ -387,7 +387,7 @@ namespace ElectraDecodersUtil
 
 			struct FScalingListData
 			{
-				UE_API bool Parse(FBitstreamReader& br);
+				UE_API bool Parse(FBitstreamReaderH265& br);
 				UE_API void SetDefaults();
 				void Reset()
 				{
@@ -407,7 +407,7 @@ namespace ElectraDecodersUtil
 
 			struct FStRefPicSet
 			{
-				UE_API bool Parse(FBitstreamReader& br, int32 stRpsIdx, const TArray<FStRefPicSet>& InReferencePicSetList, int32 InReferencePicSetListNum);
+				UE_API bool Parse(FBitstreamReaderH265& br, int32 stRpsIdx, const TArray<FStRefPicSet>& InReferencePicSetList, int32 InReferencePicSetListNum);
 				void Reset()
 				{
 					inter_ref_pic_set_prediction_flag = 0;
@@ -925,7 +925,7 @@ namespace ElectraDecodersUtil
 					{
 						Reset();
 					}
-					bool Parse(FBitstreamReader& br);
+					bool Parse(FBitstreamReaderH265& br);
 					void Reset()
 					{
 						log2_max_transform_skip_block_size_minus2 = 0;
@@ -955,7 +955,7 @@ namespace ElectraDecodersUtil
 					{
 						Reset();
 					}
-					bool Parse(FBitstreamReader& br);
+					bool Parse(FBitstreamReaderH265& br);
 					void Reset()
 					{
 						poc_reset_info_present_flag = 0;
@@ -1027,7 +1027,7 @@ namespace ElectraDecodersUtil
 
 				struct FPPS3DExtension
 				{
-					bool Parse(FBitstreamReader& br);
+					bool Parse(FBitstreamReaderH265& br);
 					void Reset()
 					{
 					}
@@ -1039,7 +1039,7 @@ namespace ElectraDecodersUtil
 					{
 						Reset();
 					}
-					bool Parse(FBitstreamReader& br);
+					bool Parse(FBitstreamReaderH265& br);
 					void Reset()
 					{
 						pps_curr_pic_ref_enabled_flag = 0;
@@ -1363,7 +1363,7 @@ namespace ElectraDecodersUtil
 			bool ELECTRADECODERS_API ParseVideoParameterSet(TMap<uint32, FVideoParameterSet>& InOutVideoParameterSets, const uint8* InBitstream, uint64 InBitstreamLenInBytes);
 			bool ELECTRADECODERS_API ParseSequenceParameterSet(TMap<uint32, FSequenceParameterSet>& InOutSequenceParameterSets, const uint8* InBitstream, uint64 InBitstreamLenInBytes);
 			bool ELECTRADECODERS_API ParsePictureParameterSet(TMap<uint32, FPictureParameterSet>& InOutPictureParameterSets, const TMap<uint32, FSequenceParameterSet>& InSequenceParameterSets, const uint8* InBitstream, uint64 InBitstreamLenInBytes);
-			bool ELECTRADECODERS_API ParseSliceHeader(TUniquePtr<FRBSP>& OutRBSP, FBitstreamReader& OutRBSPReader, FSliceSegmentHeader& OutSlice, const TMap<uint32, FVideoParameterSet>& InVideoParameterSets, const TMap<uint32, FSequenceParameterSet>& InSequenceParameterSets, const TMap<uint32, FPictureParameterSet>& InPictureParameterSets, const uint8* InBitstream, uint64 InBitstreamLenInBytes);
+			bool ELECTRADECODERS_API ParseSliceHeader(TUniquePtr<FRBSP>& OutRBSP, FBitstreamReaderH265& OutRBSPReader, FSliceSegmentHeader& OutSlice, const TMap<uint32, FVideoParameterSet>& InVideoParameterSets, const TMap<uint32, FSequenceParameterSet>& InSequenceParameterSets, const TMap<uint32, FPictureParameterSet>& InPictureParameterSets, const uint8* InBitstream, uint64 InBitstreamLenInBytes);
 
 			bool ELECTRADECODERS_API ParseVideoParameterSet(FVideoParameterSet& OutSequenceParameterSet, const uint8* InBitstream, uint64 InBitstreamLenInBytes);
 			inline bool ParseVideoParameterSet(FVideoParameterSet& OutSequenceParameterSet, const TArray<uint8>& InData)
diff --git a/Engine/Plugins/Media/ElectraCodecs/Source/ElectraDecoders/Public/Windows/WindowsElectraDecoderPlatformOutputHandleTypes.h b/Engine/Plugins/Media/ElectraCodecs/Source/ElectraDecoders/Public/Windows/WindowsElectraDecoderPlatformOutputHandleTypes.h
--- a/Engine/Plugins/Media/ElectraCodecs/Source/ElectraDecoders/Public/Windows/WindowsElectraDecoderPlatformOutputHandleTypes.h
+++ b/Engine/Plugins/Media/ElectraCodecs/Source/ElectraDecoders/Public/Windows/WindowsElectraDecoderPlatformOutputHandleTypes.h
@@ -4,6 +4,8 @@
 
 #include "CoreMinimal.h"
 #include "ElectraDecodersPlatformResources.h"
+#include "PixelFormat.h"
+#include "IElectraTextureSampleEncoding.h"
 
 #include "Windows/AllowWindowsPlatformTypes.h"
 THIRD_PARTY_INCLUDES_START
@@ -20,37 +22,6 @@ enum class EElectraDecoderPlatformOutputHandleType
 	ImageBuffers			// IElectraDecoderVideoOutputImageBuffers interface
 };
 
-enum class EElectraDecoderPlatformPixelFormat
-{
-	INVALID = 0,
-	R8G8B8A8,
-	A8R8G8B8,
-	B8G8R8A8,
-	R16G16B16A16,
-	A16B16G16R16,
-	A32B32G32R32F,
-	A2B10G10R10,
-	DXT1,
-	DXT5,
-	BC4,
-	NV12,
-	P010,
-};
-
-enum class EElectraDecoderPlatformPixelEncoding
-{
-	Native = 0,		//!< Pixel formats native representation
-	RGB,			//!< Interpret as RGB
-	RGBA,			//!< Interpret as RGBA
-	YCbCr,			//!< Interpret as YCbCR
-	YCbCr_Alpha,	//!< Interpret as YCbCR with alpha
-	YCoCg,			//!< Interpret as scaled YCoCg
-	YCoCg_Alpha,	//!< Interpret as scaled YCoCg with trailing BC4 alpha data
-	CbY0CrY1,		//!< Interpret as CbY0CrY1
-	Y0CbY1Cr,		//!< Interpret as Y0CbY1Cr
-	ARGB_BigEndian,	//!< Interpret as ARGB, big endian
-};
-
 class FElectraDecoderOutputSync final
 {
 public:
@@ -83,7 +54,7 @@ public:
 	virtual int32 GetNumberOfBuffers() const = 0;
 
 	// Returns the n'th image buffer as CPU data buffer.
-	virtual TSharedPtr<TArray<uint8>, ESPMode::ThreadSafe> GetBufferDataByIndex(int32 InBufferIndex) const = 0;
+	virtual TSharedPtr<TArray64<uint8>, ESPMode::ThreadSafe> GetBufferDataByIndex(int32 InBufferIndex) const = 0;
 
 	// Returns the n'th image buffer as GPU texture resource.
 	virtual void* GetBufferTextureByIndex(int32 InBufferIndex) const = 0;
@@ -92,10 +63,10 @@ public:
 	virtual bool GetBufferTextureSyncByIndex(int32 InBufferIndex, FElectraDecoderOutputSync& SyncObject) const = 0;
 
 	// Returns the n'th image buffer format
-	virtual EElectraDecoderPlatformPixelFormat GetBufferFormatByIndex(int32 InBufferIndex) const = 0;
+	virtual EPixelFormat GetBufferFormatByIndex(int32 InBufferIndex) const = 0;
 
 	// Returns the n'th image buffer encoding
-	virtual EElectraDecoderPlatformPixelEncoding GetBufferEncodingByIndex(int32 InBufferIndex) const = 0;
+	virtual EElectraTextureSamplePixelEncoding GetBufferEncodingByIndex(int32 InBufferIndex) const = 0;
 
 	// Returns the n'th image buffer pitch
 	virtual int32 GetBufferPitchByIndex(int32 InBufferIndex) const = 0;
diff --git a/Engine/Plugins/Media/ElectraCodecs/Source/ElectraDecoders/Public/Windows/WindowsElectraDecodersPlatformResources.h b/Engine/Plugins/Media/ElectraCodecs/Source/ElectraDecoders/Public/Windows/WindowsElectraDecodersPlatformResources.h
--- a/Engine/Plugins/Media/ElectraCodecs/Source/ElectraDecoders/Public/Windows/WindowsElectraDecodersPlatformResources.h
+++ b/Engine/Plugins/Media/ElectraCodecs/Source/ElectraDecoders/Public/Windows/WindowsElectraDecodersPlatformResources.h
@@ -4,6 +4,10 @@
 #include "CoreTypes.h"
 #include "Templates/SharedPointer.h"
 
+#include "ParameterDictionary.h"
+
+class IElectraDecoderVideoOutput;
+class FElectraTextureSample;
 
 class FElectraDecodersPlatformResourcesWindows
 {
@@ -13,6 +17,12 @@ public:
 	class FAsyncConsecutiveTaskSync;
 	static ELECTRADECODERS_API TSharedPtr<FAsyncConsecutiveTaskSync, ESPMode::ThreadSafe> CreateAsyncConsecutiveTaskSync();
 	static ELECTRADECODERS_API bool RunCodeAsync(TFunction<void()>&& InCodeToRun, TSharedPtr<FAsyncConsecutiveTaskSync, ESPMode::ThreadSafe> InTaskSync);
+
+	struct IDecoderPlatformResource;
+	static ELECTRADECODERS_API IDecoderPlatformResource* CreatePlatformVideoResource(void* InOwnerHandle, const TMap<FString, FVariant> InOptions);
+	static ELECTRADECODERS_API void ReleasePlatformVideoResource(void* InOwnerHandle, IDecoderPlatformResource* InHandleToRelease);
+
+	static ELECTRADECODERS_API bool SetupOutputTextureSample(FString& OutErrorMessage, TSharedPtr<FElectraTextureSample, ESPMode::ThreadSafe> InOutTextureSampleToSetup, TSharedPtr<IElectraDecoderVideoOutput, ESPMode::ThreadSafe> InDecoderOutput, TSharedPtr<Electra::FParamDict, ESPMode::ThreadSafe> InOutBufferPropertes, FElectraDecodersPlatformResourcesWindows::IDecoderPlatformResource* InPlatformResource);
 };
 
 using FElectraDecodersPlatformResources = FElectraDecodersPlatformResourcesWindows;
diff --git a/Engine/Plugins/Media/ElectraPlayer/Source/ElectraPlayerFactory/Private/ElectraPlayerFactoryModule.cpp b/Engine/Plugins/Media/ElectraPlayer/Source/ElectraPlayerFactory/Private/ElectraPlayerFactoryModule.cpp
--- a/Engine/Plugins/Media/ElectraPlayer/Source/ElectraPlayerFactory/Private/ElectraPlayerFactoryModule.cpp
+++ b/Engine/Plugins/Media/ElectraPlayer/Source/ElectraPlayerFactory/Private/ElectraPlayerFactoryModule.cpp
@@ -184,6 +184,7 @@ public:
 		SupportedFileExtensions.Add(TEXT("webm"));
 		SupportedFileExtensions.Add(TEXT("mp3"));
 		SupportedFileExtensions.Add(TEXT("mpa"));
+		SupportedFileExtensions.Add(TEXT("aac"));
 #endif
 
 		// register player factory
diff --git a/Engine/Plugins/Media/ElectraPlayer/Source/ElectraPlayerRuntime/ElectraPlayerRuntime.Build.cs b/Engine/Plugins/Media/ElectraPlayer/Source/ElectraPlayerRuntime/ElectraPlayerRuntime.Build.cs
--- a/Engine/Plugins/Media/ElectraPlayer/Source/ElectraPlayerRuntime/ElectraPlayerRuntime.Build.cs
+++ b/Engine/Plugins/Media/ElectraPlayer/Source/ElectraPlayerRuntime/ElectraPlayerRuntime.Build.cs
@@ -78,7 +78,6 @@ namespace UnrealBuildTool.Rules
 				PublicSystemLibraries.Add("mfuuid.lib");
 
 				PublicIncludePaths.Add("$(ModuleDir)/Public/Windows");
-				PrivateIncludePaths.Add("ElectraPlayerRuntime/Private/Runtime/Decoder/Windows");
 			}
 			else if (Target.IsInPlatformGroup(UnrealPlatformGroup.Apple))
 			{
@@ -94,19 +93,14 @@ namespace UnrealBuildTool.Rules
 				});
 
 				PublicIncludePaths.Add("$(ModuleDir)/Public/Apple");
-
-				PrivateIncludePaths.Add("ElectraPlayerRuntime/Private/Runtime/Decoder/Apple");
 			}
 			else if (Target.IsInPlatformGroup(UnrealPlatformGroup.Android) )
 			{
 				PublicIncludePaths.Add("$(ModuleDir)/Public/Android");
-
-				PrivateIncludePaths.Add("ElectraPlayerRuntime/Private/Runtime/Decoder/Android");
 			}
 			else if (Target.IsInPlatformGroup(UnrealPlatformGroup.Unix) )
 			{
 				PublicIncludePaths.Add("$(ModuleDir)/Public/Linux");
-				PrivateIncludePaths.Add("ElectraPlayerRuntime/Private/Runtime/Decoder/Linux");
 
 				AddEngineThirdPartyPrivateStaticDependencies(Target, "libav");
 			}
diff --git a/Engine/Plugins/Media/ElectraPlayer/Source/ElectraPlayerRuntime/Private/ElectraPlayer.cpp b/Engine/Plugins/Media/ElectraPlayer/Source/ElectraPlayerRuntime/Private/ElectraPlayer.cpp
--- a/Engine/Plugins/Media/ElectraPlayer/Source/ElectraPlayerRuntime/Private/ElectraPlayer.cpp
+++ b/Engine/Plugins/Media/ElectraPlayer/Source/ElectraPlayerRuntime/Private/ElectraPlayer.cpp
@@ -12,8 +12,6 @@
 #include "Player/AdaptiveStreamingPlayer.h"
 #include "Player/AdaptivePlayerOptionKeynames.h"
 #include "Utilities/Utilities.h"
-#include "Renderer/RendererVideo.h"
-#include "Renderer/RendererAudio.h"
 #include "MediaSubtitleDecoderOutput.h"
 #include "MediaMetaDataDecoderOutput.h"
 #include "IElectraSubtitleSample.h"
@@ -366,10 +364,12 @@ FElectraPlayer::FElectraPlayer(IMediaEventSink& InEventSink,
 {
 	CSV_EVENT(ElectraPlayer, TEXT("Player Creation"));
 
+	OutputAudioSamplePool = MakeShareable(new FElectraAudioSamplePool);
 	OutputTexturePool = MakeShareable(new FElectraTextureSamplePool);
 #if PLATFORM_ANDROID
 	OutputTexturePoolDecoderID = FElectraDecodersPlatformResources::RegisterOutputTexturePool(OutputTexturePool);
 #endif
+
 	EmptyMediaSamples.Reset(new FMediaSamples);
 	CurrentPlaybackRange = TRange<FTimespan>::Empty();
 
@@ -657,15 +657,62 @@ bool FElectraPlayer::OpenInternal(const FString& InUrl, const IMediaOptions* InO
 	// Create a new empty player structure. This contains the actual player instance, its associated renderers and sample queues.
 	TSharedPtr<FInternalPlayerImpl, ESPMode::ThreadSafe> NewPlayer = MakeShared<FInternalPlayerImpl, ESPMode::ThreadSafe>();
 
-	// Create the renderers so we can pass them to the internal player.
-	// They get a pointer to ourselves which they will call On[Video|Audio]Decoded() and On[Video|Audio]Flush() on.
-	NewPlayer->RendererVideo = MakeShared<FElectraRendererVideo, ESPMode::ThreadSafe>(SharedThis(this));
-	NewPlayer->RendererAudio = MakeShared<FElectraRendererAudio, ESPMode::ThreadSafe>(SharedThis(this));
+	// Create the output handlers for audio and video.
+	NewPlayer->AudioOutputHandler = MakeShared<FOutputHandlerAudio, ESPMode::ThreadSafe>();
+	NewPlayer->AudioOutputHandler->SetOutputAudioSamplePool(OutputAudioSamplePool);
+	NewPlayer->AudioOutputHandler->CanOutputQueueReceiveDelegate().BindLambda([Player=NewPlayer.ToWeakPtr(), This=this](bool& bOutCanReceive, int32 InNumSamples)
+	{
+		bOutCanReceive = false;
+		if (auto Pl = Player.Pin())
+		{
+			bOutCanReceive = This->CanPresentAudioFrames(InNumSamples);
+		}
+	});
+	NewPlayer->AudioOutputHandler->OutputQueueReceiveSampleDelegate().BindLambda([Player=NewPlayer.ToWeakPtr(), This=this](FElectraAudioSamplePtr InSample)
+	{
+		if (auto Pl = Player.Pin())
+		{
+			This->OnAudioDecoded(InSample);
+		}
+	});
+	NewPlayer->AudioOutputHandler->OutputQueueFlushSamplesDelegate().BindLambda([Player=NewPlayer.ToWeakPtr(), This=this]()
+	{
+		if (auto Pl = Player.Pin())
+		{
+			This->OnAudioFlush();
+		}
+	});
+
+	NewPlayer->VideoOutputHandler = MakeShared<FOutputHandlerVideo, ESPMode::ThreadSafe>();
+	NewPlayer->VideoOutputHandler->SetOutputTexturePool(OutputTexturePool);
+	NewPlayer->VideoOutputHandler->CanOutputQueueReceiveDelegate().BindLambda([Player=NewPlayer.ToWeakPtr(), This=this](bool& bOutCanReceive, int32 InNumSamples)
+	{
+		bOutCanReceive = false;
+		if (auto Pl = Player.Pin())
+		{
+			bOutCanReceive = This->CanPresentVideoFrames(InNumSamples);
+		}
+	});
+	NewPlayer->VideoOutputHandler->OutputQueueReceiveSampleDelegate().BindLambda([Player=NewPlayer.ToWeakPtr(), This=this](FElectraTextureSamplePtr InSample)
+	{
+		if (auto Pl = Player.Pin())
+		{
+			This->OnVideoDecoded(InSample);
+		}
+	});
+	NewPlayer->VideoOutputHandler->OutputQueueFlushSamplesDelegate().BindLambda([Player=NewPlayer.ToWeakPtr(), This=this]()
+	{
+		if (auto Pl = Player.Pin())
+		{
+			This->OnVideoFlush();
+		}
+	});
+
 
 	// Create the internal player and register ourselves as metrics receiver and static resource provider.
 	IAdaptiveStreamingPlayer::FCreateParam CreateParams;
-	CreateParams.VideoRenderer = NewPlayer->RendererVideo;
-	CreateParams.AudioRenderer = NewPlayer->RendererAudio;
+	CreateParams.VideoOutputHandler = NewPlayer->VideoOutputHandler;
+	CreateParams.AudioOutputHandler = NewPlayer->AudioOutputHandler;
 	CreateParams.ExternalPlayerGUID = PlayerGuid;
 	FString WorkerThreadOption = PlayerOptions.GetValue(Electra::OptionKeyWorkerThreads).SafeGetFString(TEXT("shared"));
 	CreateParams.WorkerThreads = WorkerThreadOption.Equals(TEXT("worker"), ESearchCase::IgnoreCase) ? IAdaptiveStreamingPlayer::FCreateParam::EWorkerThreads::DedicatedWorker :
@@ -753,14 +800,14 @@ void FElectraPlayer::CloseInternal()
 	bFrameAccurateSeeking.Reset();
 	bEnableLooping.Reset();
 
-	// Detach ourselves from the renderers.
-	if (Player->RendererVideo.IsValid())
+	// Detach the output handlers so we do not receive any new output.
+	if (Player->AudioOutputHandler.IsValid())
 	{
-		Player->RendererVideo->DetachPlayer();
+		Player->AudioOutputHandler->DetachPlayer();
 	}
-	if (Player->RendererAudio.IsValid())
+	if (Player->VideoOutputHandler.IsValid())
 	{
-		Player->RendererAudio->DetachPlayer();
+		Player->VideoOutputHandler->DetachPlayer();
 	}
 	// Detach ourselves from receiving any player events.
 	if (Player->AdaptivePlayer.IsValid())
@@ -810,8 +857,8 @@ void FElectraPlayer::FInternalPlayerImpl::DoCloseAsync(TSharedPtr<FInternalPlaye
 		double TimeCloseBegan = FPlatformTime::Seconds();
 		InPlayer->AdaptivePlayer->Stop();
 		InPlayer->AdaptivePlayer.Reset();
-		InPlayer->RendererVideo.Reset();
-		InPlayer->RendererAudio.Reset();
+		InPlayer->AudioOutputHandler.Reset();
+		InPlayer->VideoOutputHandler.Reset();
 		InPlayer->MediaSamplesToDelete.Reset();
 		double TimeCloseEnded = FPlatformTime::Seconds();
 		UE_LOG(LogElectraPlayer, Verbose, TEXT("[%u] DoCloseAsync() finished after %.3f msec!"), InPlayerID, (TimeCloseEnded - TimeCloseBegan) * 1000.0);
@@ -962,25 +1009,18 @@ bool FElectraPlayer::CanPresentVideoFrames(uint64 InNumFrames)
 	return MediaSamples.IsValid() && MediaSamples->CanReceiveVideoSamples(InNumFrames);
 }
 
-void FElectraPlayer::OnVideoDecoded(const FVideoDecoderOutputPtr& InDecoderOutput, bool bInDoNotRender)
+void FElectraPlayer::OnVideoDecoded(FElectraTextureSamplePtr InSample)
 {
-	//UE_LOG(LogElectraPlayer, VeryVerbose, TEXT("[%u] OnVideoDecoded(%d; %#.4f, %d,%d)"), PlayerUniqueID, bInDoNotRender, InDecoderOutput->GetTime().Time.GetTotalSeconds(), (int32)InDecoderOutput->GetTime().SequenceIndex, (int32)(InDecoderOutput->GetTime().SequenceIndex >> 32));
+	//UE_LOG(LogElectraPlayer, VeryVerbose, TEXT("[%u] OnVideoDecoded(%#.4f, %d,%d)"), PlayerUniqueID, InSample->GetTime().GetTime().GetTotalSeconds(), InSample->GetTime().GetSequenceIndex(), InSample->GetTime().GetLoopIndex());
 
 	TSharedPtr<FInternalPlayerImpl, ESPMode::ThreadSafe> Player = CurrentPlayer;
-	if (bInDoNotRender == false && Player.IsValid() && InDecoderOutput.IsValid() && PlayerState.State != EMediaState::Closed)
+	if (Player.IsValid() && PlayerState.State != EMediaState::Closed)
 	{
-		LastPresentedFrameDimension = InDecoderOutput->GetOutputDim();
-
-		auto TexturePool = OutputTexturePool;
-		if (TexturePool.IsValid())
+		LastPresentedFrameDimension = InSample->GetOutputDim();
+		FScopeLock SampleLock(&MediaSamplesLock);
+		if (MediaSamples.IsValid())
 		{
-			FScopeLock SampleLock(&MediaSamplesLock);
-			if (MediaSamples.IsValid())
-			{
-				FElectraTextureSampleRef TextureSample = TexturePool->AcquireShared();
-				TextureSample->Initialize(InDecoderOutput.Get());
-				MediaSamples->AddVideo(TextureSample);
-			}
+			MediaSamples->AddVideo(InSample.ToSharedRef());
 		}
 	}
 }
@@ -1009,19 +1049,17 @@ bool FElectraPlayer::CanPresentAudioFrames(uint64 InNumFrames)
 	return MediaSamples.IsValid() && MediaSamples->CanReceiveAudioSamples(InNumFrames);
 }
 
-void FElectraPlayer::OnAudioDecoded(const IAudioDecoderOutputPtr& InDecoderOutput)
+void FElectraPlayer::OnAudioDecoded(FElectraAudioSamplePtr InSample)
 {
-	//UE_LOG(LogElectraPlayer, VeryVerbose, TEXT("[%u] OnAudioDecoded(%#.4f, %d,%d)"), PlayerUniqueID, InDecoderOutput->GetTime().Time.GetTotalSeconds(), (int32)InDecoderOutput->GetTime().SequenceIndex, (int32)(InDecoderOutput->GetTime().SequenceIndex >> 32));
+	//UE_LOG(LogElectraPlayer, VeryVerbose, TEXT("[%u] OnAudioDecoded(%#.4f, %d,%d)"), PlayerUniqueID, InSample->GetTime().Time.GetTotalSeconds(), (int32)InSample->GetTime().SequenceIndex, (int32)(InSample->GetTime().SequenceIndex >> 32));
 
 	TSharedPtr<FInternalPlayerImpl, ESPMode::ThreadSafe> Player = CurrentPlayer;
-	if (Player.IsValid() && InDecoderOutput.IsValid() && PlayerState.State != EMediaState::Closed)
+	if (Player.IsValid() && InSample.IsValid() && PlayerState.State != EMediaState::Closed)
 	{
 		FScopeLock SampleLock(&MediaSamplesLock);
 		if (MediaSamples.IsValid())
 		{
-			TSharedRef<FElectraPlayerAudioSample, ESPMode::ThreadSafe> AudioSample = OutputAudioPool.AcquireShared();
-			AudioSample->Initialize(InDecoderOutput);
-			MediaSamples->AddAudio(AudioSample);
+			MediaSamples->AddAudio(InSample.ToSharedRef());
 		}
 	}
 }
diff --git a/Engine/Plugins/Media/ElectraPlayer/Source/ElectraPlayerRuntime/Private/ElectraPlayer.h b/Engine/Plugins/Media/ElectraPlayer/Source/ElectraPlayerRuntime/Private/ElectraPlayer.h
--- a/Engine/Plugins/Media/ElectraPlayer/Source/ElectraPlayerRuntime/Private/ElectraPlayer.h
+++ b/Engine/Plugins/Media/ElectraPlayer/Source/ElectraPlayerRuntime/Private/ElectraPlayer.h
@@ -26,8 +26,9 @@
 
 #include "IAnalyticsProviderET.h"
 
+#include "OutputHandler.h"
 #include "ElectraTextureSample.h"
-#include "ElectraPlayerAudioSample.h"
+#include "IElectraAudioSample.h"
 
 #include "Player/AdaptiveStreamingPlayer.h"
 #include "MediaStreamMetadata.h"
@@ -38,16 +39,11 @@
 
 
 
-class FVideoDecoderOutput;
-using FVideoDecoderOutputPtr = TSharedPtr<FVideoDecoderOutput, ESPMode::ThreadSafe>;
-class IAudioDecoderOutput;
-using IAudioDecoderOutputPtr = TSharedPtr<IAudioDecoderOutput, ESPMode::ThreadSafe>;
 class IMetaDataDecoderOutput;
 using IMetaDataDecoderOutputPtr = TSharedPtr<IMetaDataDecoderOutput, ESPMode::ThreadSafe>;
 class ISubtitleDecoderOutput;
 using ISubtitleDecoderOutputPtr = TSharedPtr<ISubtitleDecoderOutput, ESPMode::ThreadSafe>;
 
-class FElectraRendererVideo;
 class FElectraRendererAudio;
 
 using namespace Electra;
@@ -178,9 +174,9 @@ public:
 	//-------------------------------------------------------------------------
 	// Implementation methods
 	//
-	void OnVideoDecoded(const FVideoDecoderOutputPtr& InDecoderOutput, bool bInDoNotRender);
+	void OnVideoDecoded(FElectraTextureSamplePtr InSample);
 	void OnVideoFlush();
-	void OnAudioDecoded(const IAudioDecoderOutputPtr& InDecoderOutput);
+	void OnAudioDecoded(FElectraAudioSamplePtr InSample);
 	void OnAudioFlush();
 	void OnSubtitleDecoded(ISubtitleDecoderOutputPtr InDecoderOutput);
 	void OnSubtitleFlush();
@@ -307,9 +303,9 @@ private:
 		/** The media player itself **/
 		TSharedPtr<IAdaptiveStreamingPlayer, ESPMode::ThreadSafe> AdaptivePlayer;
 
-		/** Renderers to use **/
-		TSharedPtr<FElectraRendererVideo, ESPMode::ThreadSafe> RendererVideo;
-		TSharedPtr<FElectraRendererAudio, ESPMode::ThreadSafe> RendererAudio;
+		/** Output buffer handlers */
+		TSharedPtr<FOutputHandlerVideo, ESPMode::ThreadSafe> VideoOutputHandler;
+		TSharedPtr<FOutputHandlerAudio, ESPMode::ThreadSafe> AudioOutputHandler;
 
 		/** Media samples to delete when closing */
 		TUniquePtr<FMediaSamples> MediaSamplesToDelete;
@@ -745,7 +741,7 @@ private:
 	/** Output sample pools */
 	TSharedPtr<FElectraTextureSamplePool, ESPMode::ThreadSafe> OutputTexturePool;
 	int64 OutputTexturePoolDecoderID = 0;
-	FElectraPlayerAudioSamplePool OutputAudioPool;
+	TSharedPtr<FElectraAudioSamplePool, ESPMode::ThreadSafe> OutputAudioSamplePool;
 
 	/** Application system event handler */
 	TSharedPtr<Electra::FApplicationTerminationHandler, ESPMode::ThreadSafe> AppTerminationHandler;
diff --git a/Engine/Plugins/Media/ElectraPlayer/Source/ElectraPlayerRuntime/Private/Runtime/Decoder/AudioDecoder.cpp b/Engine/Plugins/Media/ElectraPlayer/Source/ElectraPlayerRuntime/Private/Runtime/Decoder/AudioDecoder.cpp
--- a/Engine/Plugins/Media/ElectraPlayer/Source/ElectraPlayerRuntime/Private/Runtime/Decoder/AudioDecoder.cpp
+++ b/Engine/Plugins/Media/ElectraPlayer/Source/ElectraPlayerRuntime/Private/Runtime/Decoder/AudioDecoder.cpp
@@ -10,7 +10,7 @@
 
 #include "StreamAccessUnitBuffer.h"
 #include "Decoder/AudioDecoder.h"
-#include "Renderer/RendererBase.h"
+#include "OutputHandler.h"
 #include "Player/PlayerSessionServices.h"
 #include "Utilities/Utilities.h"
 #include "Utilities/StringHelpers.h"
@@ -32,7 +32,6 @@
 #include "IElectraDecoderOutputAudio.h"
 #include "IElectraDecoderFeaturesAndOptions.h"
 #include "ElectraDecodersUtils.h"
-#include "ElectraDecoderResourceManager.h"
 
 /***************************************************************************************************************************************************/
 
@@ -55,7 +54,7 @@ public:
 	void Open(TSharedPtrTS<FAccessUnit::CodecData> InCodecData) override;
 	void Close() override;
 	void DrainForCodecChange() override;
-	void SetRenderer(TSharedPtr<IMediaRenderer, ESPMode::ThreadSafe> InRenderer) override;
+	void SetOutputHandler(TSharedPtr<FOutputHandlerAudio, ESPMode::ThreadSafe> InOutputHandler) override;
 	void SuspendOrResumeDecoder(bool bSuspend, const FParamDict& InOptions) override;
 	void AUdataPushAU(FAccessUnit* AccessUnit) override;
 	void AUdataPushEOD() override;
@@ -133,6 +132,16 @@ private:
 
 	void NotifyReadyBufferListener(bool bHaveOutput);
 
+	bool CanReceiveOutputFrame()
+	{
+		bool bCanRcv = true;
+		if (OutputHandler)
+		{
+			bCanRcv = OutputHandler->CanReceiveOutputSample();
+		}
+		return bCanRcv;
+	}
+
 	EAUChangeFlags GetAndPrepareInputAU();
 	void PrepareAU(TSharedPtrTS<FDecoderInput> AU);
 	IElectraDecoder::ECSDCompatibility IsCompatibleWith();
@@ -180,8 +189,8 @@ private:
 	TSharedPtrTS<FFGBGNotificationHandlers>									FGBGHandlers;
 	int32																	ApplicationSuspendCount = 0;
 
-	TSharedPtr<IMediaRenderer, ESPMode::ThreadSafe>							Renderer;
-	int32																	MaxDecodeBufferSize = 0;
+	TSharedPtr<FOutputHandlerAudio, ESPMode::ThreadSafe>					OutputHandler;
+	int32																	MaxOutputBuffers = 0;
 
 	FCriticalSection														ListenerMutex;
 	IAccessUnitBufferListener*												InputBufferListener = nullptr;
@@ -194,10 +203,7 @@ private:
 	bool																	bIsAdaptiveDecoder = false;
 	bool																	bMustBeSuspendedInBackground = false;
 
-	IMediaRenderer::IBuffer*												CurrentOutputBuffer = nullptr;
-	FParamDict																EmptyOptions;
-	FParamDict																OutputBufferSampleProperties;
-	FParamDict																EOSBufferProperties;
+	FElectraAudioSamplePtr													CurrentOutputBuffer;
 	FAudioChannelMapper														ChannelMapper;
 };
 ENUM_CLASS_FLAGS(FAudioDecoderImpl::EAUChangeFlags);
@@ -265,12 +271,6 @@ TSharedPtr<IElectraCodecFactory, ESPMode::ThreadSafe> FAudioDecoderImpl::GetDeco
 FAudioDecoderImpl::FAudioDecoderImpl()
 	: FMediaThread("ElectraPlayer::Audio decoder")
 {
-	EOSBufferProperties.Set(RenderOptionKeys::NumChannels, FVariantValue((int64)0));
-	EOSBufferProperties.Set(RenderOptionKeys::UsedByteSize, FVariantValue((int64)0));
-	EOSBufferProperties.Set(RenderOptionKeys::SampleRate, FVariantValue((int64) 0));
-	EOSBufferProperties.Set(RenderOptionKeys::Duration, FVariantValue(FTimeValue::GetZero()));
-	EOSBufferProperties.Set(RenderOptionKeys::PTS, FVariantValue(FTimeValue::GetInvalid()));
-	EOSBufferProperties.Set(RenderOptionKeys::EOSFlag, FVariantValue(true));
 }
 
 FAudioDecoderImpl::~FAudioDecoderImpl()
@@ -299,9 +299,9 @@ void FAudioDecoderImpl::DrainForCodecChange()
 	bDrainForCodecChange = true;
 }
 
-void FAudioDecoderImpl::SetRenderer(TSharedPtr<IMediaRenderer, ESPMode::ThreadSafe> InRenderer)
+void FAudioDecoderImpl::SetOutputHandler(TSharedPtr<FOutputHandlerAudio, ESPMode::ThreadSafe> InOutputHandler)
 {
-	Renderer = InRenderer;
+	OutputHandler = InOutputHandler;
 }
 
 void FAudioDecoderImpl::SuspendOrResumeDecoder(bool bSuspend, const FParamDict& InOptions)
@@ -364,28 +364,27 @@ void FAudioDecoderImpl::StopThread()
 
 void FAudioDecoderImpl::CreateDecoderOutputPool()
 {
-	check(Renderer);
-	FParamDict poolOpts;
-	const int32 kMaxChannels = 16;
-	const int32 kMaxSamplesPerBlock = 2048;
-	uint32 frameSize = sizeof(float) * kMaxChannels * kMaxSamplesPerBlock;
-	poolOpts.Set(RenderOptionKeys::MaxBufferSize, FVariantValue((int64) frameSize));
-	poolOpts.Set(RenderOptionKeys::NumBuffers, FVariantValue((int64) 8));
-	poolOpts.Set(RenderOptionKeys::SamplesPerBlock, FVariantValue((int64) kMaxSamplesPerBlock));
-	poolOpts.Set(RenderOptionKeys::MaxChannels, FVariantValue((int64) kMaxChannels));
-	if (Renderer->CreateBufferPool(poolOpts) == UEMEDIA_ERROR_OK)
-	{
-		MaxDecodeBufferSize = (int32) Renderer->GetBufferPoolProperties().GetValue(RenderOptionKeys::MaxBuffers).GetInt64();
-	}
-	else
+	if (OutputHandler)
 	{
-		PostError(0, TEXT("Failed to create sample pool"), ERRCODE_AUDIO_INTERNAL_COULD_NOT_CREATE_SAMPLE_POOL);
+		FParamDict poolOpts;
+		// Note: Get the default value of 8 from some config option?
+		int64 NumOutputFrames = 8;
+		poolOpts.Set(OutputHandlerOptionKeys::NumBuffers, FVariantValue(NumOutputFrames));
+		if (OutputHandler->PreparePool(poolOpts))
+		{
+			MaxOutputBuffers = (int32) OutputHandler->GetPoolProperties().GetValue(OutputHandlerOptionKeys::NumBuffers).GetInt64();
+			return;
+		}
 	}
+	PostError(0, TEXT("Failed to create sample pool"), ERRCODE_AUDIO_INTERNAL_COULD_NOT_CREATE_SAMPLE_POOL);
 }
 
 void FAudioDecoderImpl::DestroyDecoderOutputPool()
 {
-	Renderer->ReleaseBufferPool();
+	if (OutputHandler)
+	{
+		OutputHandler->ClosePool();
+	}
 }
 
 void FAudioDecoderImpl::NotifyReadyBufferListener(bool bHaveOutput)
@@ -393,7 +392,7 @@ void FAudioDecoderImpl::NotifyReadyBufferListener(bool bHaveOutput)
 	if (ReadyBufferListener)
 	{
 		IDecoderOutputBufferListener::FDecodeReadyStats stats;
-		stats.OutputBufferPoolSize = MaxDecodeBufferSize;
+		stats.OutputBufferPoolSize = MaxOutputBuffers;
 		if ((stats.NumElementsInDecoder = InDecoderInput.Num()) != 0)
 		{
 			stats.InDecoderTimeRangePTS.Start = InDecoderInput[0]->AccessUnit->PTS;
@@ -406,7 +405,7 @@ void FAudioDecoderImpl::NotifyReadyBufferListener(bool bHaveOutput)
 		}
 		stats.InDecoderTimeRangePTS.End.SetSequenceIndex(stats.InDecoderTimeRangePTS.Start.GetSequenceIndex());
 		stats.bOutputStalled = !bHaveOutput;
-		stats.bEODreached = NextAccessUnits.ReachedEOD() && CurrentOutputBuffer == nullptr;
+		stats.bEODreached = NextAccessUnits.ReachedEOD() && !CurrentOutputBuffer;
 		ListenerMutex.Lock();
 		if (ReadyBufferListener)
 		{
@@ -538,10 +537,10 @@ void FAudioDecoderImpl::InternalDecoderDestroy()
 
 void FAudioDecoderImpl::ReturnUnusedOutputBuffer()
 {
-	if (CurrentOutputBuffer)
+	if (CurrentOutputBuffer && OutputHandler)
 	{
-		Renderer->ReturnBuffer(CurrentOutputBuffer, false, EmptyOptions);
-		CurrentOutputBuffer = nullptr;
+		OutputHandler->ReturnOutputSample(CurrentOutputBuffer, FOutputHandlerAudio::EReturnSampleType::DontSendToQueue);
+		CurrentOutputBuffer.Reset();
 	}
 }
 
@@ -718,164 +717,151 @@ IElectraDecoder::EOutputStatus FAudioDecoderImpl::HandleOutput()
 			{
 				continue;
 			}
+			SCOPE_CYCLE_COUNTER(STAT_ElectraPlayer_AudioConvertOutput);
+			CSV_SCOPED_TIMING_STAT(ElectraPlayer, AudioConvertOutput);
+
+			// Check if the output pipeline can accept the output we want to send to it.
+			if (!CanReceiveOutputFrame())
+			{
+				NotifyReadyBufferListener(false);
+				return IElectraDecoder::EOutputStatus::TryAgainLater;
+			}
 
 			// Need a new output buffer?
-			if (CurrentOutputBuffer == nullptr && Renderer.IsValid())
+			if (!CurrentOutputBuffer && OutputHandler)
 			{
-				SCOPE_CYCLE_COUNTER(STAT_ElectraPlayer_AudioConvertOutput);
-				CSV_SCOPED_TIMING_STAT(ElectraPlayer, AudioConvertOutput);
-				UEMediaError bufResult = Renderer->AcquireBuffer(CurrentOutputBuffer, 0, EmptyOptions);
-				check(bufResult == UEMEDIA_ERROR_OK || bufResult == UEMEDIA_ERROR_INSUFFICIENT_DATA);
-				if (bufResult != UEMEDIA_ERROR_OK && bufResult != UEMEDIA_ERROR_INSUFFICIENT_DATA)
+				IOutputHandlerBase::EBufferResult bufResult = OutputHandler->ObtainOutputSample(CurrentOutputBuffer);
+				if (bufResult != IOutputHandlerBase::EBufferResult::Ok && bufResult != IOutputHandlerBase::EBufferResult::NoBuffer)
 				{
-					PostError(0, TEXT("Failed to acquire sample buffer"), ERRCODE_AUDIO_INTERNAL_COULD_NOT_GET_SAMPLE_BUFFER, bufResult);
+					PostError(0, TEXT("Failed to acquire sample buffer"), ERRCODE_AUDIO_INTERNAL_COULD_NOT_GET_SAMPLE_BUFFER);
 					return IElectraDecoder::EOutputStatus::Error;
 				}
 			}
+			// Didn't get a buffer? This should not really happen since the output pipeline said it could accept a frame.
+			if (!CurrentOutputBuffer)
+			{
+				NotifyReadyBufferListener(false);
+				return IElectraDecoder::EOutputStatus::TryAgainLater;
+			}
 
-			bool bHaveAvailSmpBlk = CurrentOutputBuffer != nullptr;
+			NotifyReadyBufferListener(true);
 
-			// Check if the renderer can accept the output we will want to send to it.
-			// If it can't right now we treat this as if we do not have an available output buffer.
-			if (Renderer.IsValid() && !Renderer->CanReceiveOutputFrames(1))
+			TSharedPtr<IElectraDecoderAudioOutput, ESPMode::ThreadSafe> Output = StaticCastSharedPtr<IElectraDecoderAudioOutput>(DecoderInstance->GetOutput());
+			// No output even though it was advertised? This may happen when a decoder has an internal state change.
+			if (!Output.IsValid())
 			{
-				bHaveAvailSmpBlk = false;
+				break;
 			}
-
-			NotifyReadyBufferListener(bHaveAvailSmpBlk);
-			if (bHaveAvailSmpBlk)
+			if (Output->GetType() != IElectraDecoderOutput::EType::Audio || !Output->IsInterleaved() || Output->GetSampleFormat() != IElectraDecoderAudioOutput::ESampleFormat::Float)
 			{
-				SCOPE_CYCLE_COUNTER(STAT_ElectraPlayer_AudioConvertOutput);
-				CSV_SCOPED_TIMING_STAT(ElectraPlayer, AudioConvertOutput);
-
-				TSharedPtr<IElectraDecoderAudioOutput, ESPMode::ThreadSafe> Output = StaticCastSharedPtr<IElectraDecoderAudioOutput>(DecoderInstance->GetOutput());
-				// No output even though it was advertised? This may happen when a decoder has an internal state change.
-				if (!Output.IsValid())
-				{
-					break;
-				}
-				if (Output->GetType() != IElectraDecoderOutput::EType::Audio || !Output->IsInterleaved() || Output->GetSampleFormat() != IElectraDecoderAudioOutput::ESampleFormat::Float)
-				{
-					PostError(0, TEXT("Could not get decoded output due to decoded format being unsupported"), ERRCODE_AUDIO_INTERNAL_UNSUPPORTED_OUTPUT_FORMAT);
-					return IElectraDecoder::EOutputStatus::Error;
-				}
+				PostError(0, TEXT("Could not get decoded output due to decoded format being unsupported"), ERRCODE_AUDIO_INTERNAL_UNSUPPORTED_OUTPUT_FORMAT);
+				return IElectraDecoder::EOutputStatus::Error;
+			}
 
-				// Locate the input AU info that should correspond to this output.
-				TSharedPtrTS<FDecoderInput> MatchingInput;
-				if (InDecoderInput.Num())
-				{
-					// Try the frontmost entry. It tends to be that one.
-					if (InDecoderInput[0]->PTS == Output->GetUserValue())
-					{
-						MatchingInput = InDecoderInput[0];
-						InDecoderInput.RemoveAt(0);
-					}
-					else
-					{
-						/*
-							Not the first element. This can happen if the audio decoder aggregates output from multiple inputs.
-						*/
-						for (int32 i = 0; i < InDecoderInput.Num(); ++i)
-						{
-							if (InDecoderInput[i]->PTS == Output->GetUserValue())
-							{
-								MatchingInput = InDecoderInput[i];
-								InDecoderInput.RemoveAt(i);
-								break;
-							}
-							else if (InDecoderInput[i]->PTS + DecoderConfig::RemovalOfOldDecoderInputThresholdHNS < (int64)Output->GetUserValue())
-							{
-								InDecoderInput.RemoveAt(i);
-								--i;
-							}
-						}
-					}
-				}
-				if (!MatchingInput.IsValid())
+			// Locate the input AU info that should correspond to this output.
+			TSharedPtrTS<FDecoderInput> MatchingInput;
+			if (InDecoderInput.Num())
+			{
+				// Try the frontmost entry. It tends to be that one.
+				if (InDecoderInput[0]->PTS == Output->GetUserValue())
 				{
-					PostError(0, TEXT("There is no pending decoder input for the decoded output!"), ERRCODE_AUDIO_INTERNAL_FAILED_TO_CONVERT_OUTPUT);
-					return IElectraDecoder::EOutputStatus::Error;
+					MatchingInput = InDecoderInput[0];
+					InDecoderInput.RemoveAt(0);
 				}
-
-				const float* PCMBuffer = reinterpret_cast<const float*>(Output->GetData(0));
-				const int32 SamplingRate = Output->GetSampleRate();
-				const int32 NumberOfChannels = Output->GetNumChannels();
-				const int32 NumBytesPerFrame = Output->GetBytesPerFrame();
-				const int32 NumBytesPerSample = Output->GetBytesPerSample();
-				int32 NumSamplesProduced = Output->GetNumFrames();
-				int32 NumDecodedBytes = NumSamplesProduced * NumBytesPerFrame;
-
-				CurrentOutputFormat.SampleRate = SamplingRate;
-				CurrentOutputFormat.NumChannels = NumberOfChannels;
-				if (!ChannelMapper.IsInitialized())
+				else
 				{
-					if (NumberOfChannels)
+					// Not the first element. This can happen if the audio decoder aggregates output from multiple inputs.
+					for(int32 i=0; i<InDecoderInput.Num(); ++i)
 					{
-						if (ChannelMapper.Initialize(Output))
+						if (InDecoderInput[i]->PTS == Output->GetUserValue())
 						{
-							// Ok.
+							MatchingInput = InDecoderInput[i];
+							InDecoderInput.RemoveAt(i);
+							break;
 						}
-						else
+						else if (InDecoderInput[i]->PTS + DecoderConfig::RemovalOfOldDecoderInputThresholdHNS < (int64)Output->GetUserValue())
 						{
-							PostError(0, TEXT("Failed to initialize audio channel mapper"), ERRCODE_AUDIO_INTERNAL_UNSUPPORTED_NUMBER_OF_CHANNELS);
-							return IElectraDecoder::EOutputStatus::Error;
+							InDecoderInput.RemoveAt(i);
+							--i;
 						}
 					}
-					else
-					{
-						PostError(0, TEXT("Bad number (0) of decoded audio channels"), ERRCODE_AUDIO_INTERNAL_UNSUPPORTED_NUMBER_OF_CHANNELS);
-						return IElectraDecoder::EOutputStatus::Error;
-					}
 				}
+			}
+			if (!MatchingInput.IsValid())
+			{
+				PostError(0, TEXT("There is no pending decoder input for the decoded output!"), ERRCODE_AUDIO_INTERNAL_FAILED_TO_CONVERT_OUTPUT);
+				return IElectraDecoder::EOutputStatus::Error;
+			}
 
-				int32 ByteOffsetToFirstSample = 0;
-				if (MatchingInput->StartOverlapDuration.GetAsHNS() || MatchingInput->EndOverlapDuration.GetAsHNS())
+			const float* PCMBuffer = reinterpret_cast<const float*>(Output->GetData(0));
+			const int32 SamplingRate = Output->GetSampleRate();
+			const int32 NumberOfChannels = Output->GetNumChannels();
+			const int32 NumBytesPerFrame = Output->GetBytesPerFrame();
+			const int32 NumBytesPerSample = Output->GetBytesPerSample();
+			int32 NumSamplesProduced = Output->GetNumFrames();
+			int32 NumDecodedBytes = NumSamplesProduced * NumBytesPerFrame;
+
+			CurrentOutputFormat.SampleRate = SamplingRate;
+			CurrentOutputFormat.NumChannels = NumberOfChannels;
+			if (!ChannelMapper.IsInitialized())
+			{
+				if (NumberOfChannels)
 				{
-					int32 SkipStartSampleNum = (int32) (MatchingInput->StartOverlapDuration.GetAsHNS() * SamplingRate / 10000000);
-					int32 SkipEndSampleNum = (int32) (MatchingInput->EndOverlapDuration.GetAsHNS() * SamplingRate / 10000000);
-
-					if (SkipStartSampleNum + SkipEndSampleNum < NumSamplesProduced)
+					if (ChannelMapper.Initialize(Output))
 					{
-						ByteOffsetToFirstSample = SkipStartSampleNum * NumBytesPerFrame;
-						NumSamplesProduced -= SkipStartSampleNum;
-						NumSamplesProduced -= SkipEndSampleNum;
-						NumDecodedBytes -= (SkipStartSampleNum + SkipEndSampleNum) * NumBytesPerFrame;
+						// Ok.
 					}
 					else
 					{
-						NumSamplesProduced = 0;
-						NumDecodedBytes = 0;
+						PostError(0, TEXT("Failed to initialize audio channel mapper"), ERRCODE_AUDIO_INTERNAL_UNSUPPORTED_NUMBER_OF_CHANNELS);
+						return IElectraDecoder::EOutputStatus::Error;
 					}
 				}
-
-				int32 NumDecodedSamplesPerChannel = NumDecodedBytes / NumBytesPerFrame;
-				if (NumDecodedSamplesPerChannel)
+				else
 				{
-					int32 CurrentRenderOutputBufferSize = (int32)CurrentOutputBuffer->GetBufferProperties().GetValue(RenderOptionKeys::AllocatedSize).GetInt64();
-					void* CurrentRenderOutputBufferAddress = CurrentOutputBuffer->GetBufferProperties().GetValue(RenderOptionKeys::AllocatedAddress).GetPointer();
-					ChannelMapper.MapChannels(CurrentRenderOutputBufferAddress, CurrentRenderOutputBufferSize, AdvancePointer(PCMBuffer, ByteOffsetToFirstSample), NumDecodedBytes, NumDecodedSamplesPerChannel);
-
-					FTimeValue Duration;
-					Duration.SetFromND(NumSamplesProduced, SamplingRate);
+					PostError(0, TEXT("Bad number (0) of decoded audio channels"), ERRCODE_AUDIO_INTERNAL_UNSUPPORTED_NUMBER_OF_CHANNELS);
+					return IElectraDecoder::EOutputStatus::Error;
+				}
+			}
 
-					OutputBufferSampleProperties.Set(RenderOptionKeys::NumChannels, FVariantValue((int64)ChannelMapper.GetNumTargetChannels()));
-					OutputBufferSampleProperties.Set(RenderOptionKeys::UsedByteSize, FVariantValue((int64)((int64)NumDecodedSamplesPerChannel * ChannelMapper.GetNumTargetChannels() * NumBytesPerSample)));
-					OutputBufferSampleProperties.Set(RenderOptionKeys::SampleRate, FVariantValue((int64) SamplingRate));
-					OutputBufferSampleProperties.Set(RenderOptionKeys::Duration, FVariantValue(Duration));
-					OutputBufferSampleProperties.Set(RenderOptionKeys::PTS, FVariantValue(MatchingInput->AdjustedPTS));
+			int32 ByteOffsetToFirstSample = 0;
+			if (MatchingInput->StartOverlapDuration.GetAsHNS() || MatchingInput->EndOverlapDuration.GetAsHNS())
+			{
+				int32 SkipStartSampleNum = (int32) (MatchingInput->StartOverlapDuration.GetAsHNS() * SamplingRate / 10000000);
+				int32 SkipEndSampleNum = (int32) (MatchingInput->EndOverlapDuration.GetAsHNS() * SamplingRate / 10000000);
 
-					Renderer->ReturnBuffer(CurrentOutputBuffer, true, OutputBufferSampleProperties);
-					CurrentOutputBuffer = nullptr;
+				if (SkipStartSampleNum + SkipEndSampleNum < NumSamplesProduced)
+				{
+					ByteOffsetToFirstSample = SkipStartSampleNum * NumBytesPerFrame;
+					NumSamplesProduced -= SkipStartSampleNum;
+					NumSamplesProduced -= SkipEndSampleNum;
+					NumDecodedBytes -= (SkipStartSampleNum + SkipEndSampleNum) * NumBytesPerFrame;
 				}
 				else
 				{
-					ReturnUnusedOutputBuffer();
+					NumSamplesProduced = 0;
+					NumDecodedBytes = 0;
 				}
 			}
+
+			int32 NumDecodedSamplesPerChannel = NumDecodedBytes / NumBytesPerFrame;
+			if (NumDecodedSamplesPerChannel)
+			{
+				if (CurrentOutputBuffer->AllocateFor(EMediaAudioSampleFormat::Float, ChannelMapper.GetNumTargetChannels(), NumDecodedSamplesPerChannel))
+				{
+					FMediaTimeStamp PTS(MatchingInput->AdjustedPTS.GetAsTimespan(), MatchingInput->AdjustedPTS.GetSequenceIndex());
+					FTimeValue Duration;
+					Duration.SetFromND(NumSamplesProduced, SamplingRate);
+					CurrentOutputBuffer->SetParameters(SamplingRate, PTS, Duration.GetAsTimespan());
+					ChannelMapper.MapChannels(CurrentOutputBuffer->GetWritableBuffer(), CurrentOutputBuffer->GetAllocatedSize(), AdvancePointer(PCMBuffer, ByteOffsetToFirstSample), NumDecodedBytes, NumDecodedSamplesPerChannel);
+				}
+
+				OutputHandler->ReturnOutputSample(CurrentOutputBuffer, FOutputHandlerAudio::EReturnSampleType::SendToQueue);
+				CurrentOutputBuffer.Reset();
+			}
 			else
 			{
-				// No available buffer. Sleep for a bit. Can't sleep on a signal since we have to check two: abort and flush
-				FMediaRunnable::SleepMilliseconds(10);
+				ReturnUnusedOutputBuffer();
 			}
 		}
 	}
@@ -963,80 +949,70 @@ bool FAudioDecoderImpl::HandleDecoding()
 
 bool FAudioDecoderImpl::SendSilenceOrEOS(ESendMode InSendMode)
 {
-	check(CurrentAccessUnit.IsValid() || InSendMode == ESendMode::SendEOS);
+	// We do not currently have a use case for sending an EOS upstream.
+	if (InSendMode == ESendMode::SendEOS)
+	{
+		return true;
+	}
+	check(CurrentAccessUnit.IsValid());
 	check(bIsDecoderClean);
 
 	// Get output unless flushing or terminating
 	while(!TerminateThreadSignal.IsSignaled() && !FlushDecoderSignal.IsSignaled())
 	{
-		// Need a new output buffer?
-		if (CurrentOutputBuffer == nullptr && Renderer.IsValid())
+		// Check if the output pipeline can accept the output we want to send to it.
+		bool bCanRcv = CanReceiveOutputFrame();
+		if (bCanRcv)
 		{
-			SCOPE_CYCLE_COUNTER(STAT_ElectraPlayer_AudioConvertOutput);
-			CSV_SCOPED_TIMING_STAT(ElectraPlayer, AudioConvertOutput);
-			UEMediaError bufResult = Renderer->AcquireBuffer(CurrentOutputBuffer, 0, EmptyOptions);
-			check(bufResult == UEMEDIA_ERROR_OK || bufResult == UEMEDIA_ERROR_INSUFFICIENT_DATA);
-			if (bufResult != UEMEDIA_ERROR_OK && bufResult != UEMEDIA_ERROR_INSUFFICIENT_DATA)
+			// Need a new output buffer?
+			if (!CurrentOutputBuffer && OutputHandler)
 			{
-				return PostError(0, TEXT("Failed to acquire sample buffer"), ERRCODE_AUDIO_INTERNAL_COULD_NOT_GET_SAMPLE_BUFFER, bufResult);
+				IOutputHandlerBase::EBufferResult bufResult = OutputHandler->ObtainOutputSample(CurrentOutputBuffer);
+				if (bufResult != IOutputHandlerBase::EBufferResult::Ok && bufResult != IOutputHandlerBase::EBufferResult::NoBuffer)
+				{
+					return PostError(0, TEXT("Failed to acquire sample buffer"), ERRCODE_AUDIO_INTERNAL_COULD_NOT_GET_SAMPLE_BUFFER);
+				}
 			}
-		}
-
-		bool bHaveAvailSmpBlk = CurrentOutputBuffer != nullptr;
-
-		// Check if the renderer can accept the output we will want to send to it.
-		// If it can't right now we treat this as if we do not have an available output buffer.
-		if (Renderer.IsValid() && !Renderer->CanReceiveOutputFrames(1))
-		{
-			bHaveAvailSmpBlk = false;
-		}
-
-		NotifyReadyBufferListener(bHaveAvailSmpBlk);
-		if (bHaveAvailSmpBlk)
-		{
-			SCOPE_CYCLE_COUNTER(STAT_ElectraPlayer_AudioConvertOutput);
-			CSV_SCOPED_TIMING_STAT(ElectraPlayer, AudioConvertOutput);
-
-			if (InSendMode == ESendMode::SendSilence)
+			// Didn't get a buffer? This should not really happen since the output pipeline said it could accept a frame.
+			if (CurrentOutputBuffer)
 			{
-				// Clear to silence
-				int64 MaxBufferSize = CurrentOutputBuffer->GetBufferProperties().GetValue(RenderOptionKeys::AllocatedSize).GetInt64();
-				FMemory::Memzero(CurrentOutputBuffer->GetBufferProperties().GetValue(RenderOptionKeys::AllocatedAddress).GetPointer(), MaxBufferSize);
+				NotifyReadyBufferListener(true);
 
-				int32 NumChannels = CurrentOutputFormat.NumChannels ? CurrentOutputFormat.NumChannels : 2;
-				int32 SampleRate = CurrentOutputFormat.SampleRate ? CurrentOutputFormat.SampleRate : 48000;
-				int32 SamplesPerBlock = CurrentAccessUnit->AdjustedDuration.IsValid() ? CurrentAccessUnit->AdjustedDuration.GetAsHNS() * SampleRate / 10000000 : 0;
-				int32 EmptySize = NumChannels * sizeof(float) * SamplesPerBlock;
+				SCOPE_CYCLE_COUNTER(STAT_ElectraPlayer_AudioConvertOutput);
+				CSV_SCOPED_TIMING_STAT(ElectraPlayer, AudioConvertOutput);
 
-				if (SamplesPerBlock)
+				uint32 NumChannels = CurrentOutputFormat.NumChannels ? CurrentOutputFormat.NumChannels : 2;
+				uint32 SampleRate = CurrentOutputFormat.SampleRate ? CurrentOutputFormat.SampleRate : 48000;
+				uint32 SamplesPerBlock = CurrentAccessUnit->AdjustedDuration.IsValid() ? CurrentAccessUnit->AdjustedDuration.GetAsHNS() * SampleRate / 10000000 : 0;
+				uint32 EmptySize = NumChannels * sizeof(float) * SamplesPerBlock;
+				if (EmptySize)
 				{
-					FTimeValue Duration;
-					Duration.SetFromND(SamplesPerBlock, (uint32) SampleRate);
-
-					OutputBufferSampleProperties.Set(RenderOptionKeys::NumChannels, FVariantValue((int64)NumChannels));
-					OutputBufferSampleProperties.Set(RenderOptionKeys::UsedByteSize, FVariantValue((int64)(EmptySize < MaxBufferSize ? EmptySize : MaxBufferSize)));
-					OutputBufferSampleProperties.Set(RenderOptionKeys::SampleRate, FVariantValue((int64)SampleRate));
-					OutputBufferSampleProperties.Set(RenderOptionKeys::Duration, FVariantValue(Duration));
-					OutputBufferSampleProperties.Set(RenderOptionKeys::PTS, FVariantValue(CurrentAccessUnit->AdjustedPTS));
-
-					Renderer->ReturnBuffer(CurrentOutputBuffer, true, OutputBufferSampleProperties);
-					CurrentOutputBuffer = nullptr;
+					if (CurrentOutputBuffer->AllocateFor(EMediaAudioSampleFormat::Float, NumChannels, SamplesPerBlock))
+					{
+						FMediaTimeStamp PTS(CurrentAccessUnit->AdjustedPTS.GetAsTimespan(), CurrentAccessUnit->AdjustedPTS.GetSequenceIndex());
+						FTimeValue Duration;
+						Duration.SetFromND(SamplesPerBlock, SampleRate);
+						CurrentOutputBuffer->SetParameters(SampleRate, PTS, Duration.GetAsTimespan());
+						FMemory::Memzero(CurrentOutputBuffer->GetWritableBuffer(), CurrentOutputBuffer->GetAllocatedSize());
+					}
+					OutputHandler->ReturnOutputSample(CurrentOutputBuffer, FOutputHandlerAudio::EReturnSampleType::SendToQueue);
+					CurrentOutputBuffer.Reset();
 				}
 				else
 				{
 					ReturnUnusedOutputBuffer();
 				}
+				return true;
 			}
 			else
 			{
-				Renderer->ReturnBuffer(CurrentOutputBuffer, false, EOSBufferProperties);
-				CurrentOutputBuffer = nullptr;
+				NotifyReadyBufferListener(false);
+				FMediaRunnable::SleepMilliseconds(10);
 			}
-			return true;
 		}
 		else
 		{
-			// No available buffer. Sleep for a bit. Can't sleep on a signal since we have to check two: abort and flush
+			NotifyReadyBufferListener(false);
 			FMediaRunnable::SleepMilliseconds(10);
 		}
 	}
@@ -1129,7 +1105,7 @@ void FAudioDecoderImpl::WorkerThread()
 	ApplicationSuspendConfirmedSignal.Reset();
 
 	bError = false;
-	CurrentOutputBuffer = nullptr;
+	CurrentOutputBuffer.Reset();
 	bIsFirstAccessUnit = true;
 	bInDummyDecodeMode = false;
 	bIsAdaptiveDecoder = false;
@@ -1232,7 +1208,10 @@ void FAudioDecoderImpl::WorkerThread()
 			// Handle decoding of either regular or dummy data.
 			if (CurrentDecodingState == EDecodingState::NormalDecoding)
 			{
-				HandleOutput();
+				if (HandleOutput() == IElectraDecoder::EOutputStatus::TryAgainLater)
+				{
+					FMediaRunnable::SleepMilliseconds(10);
+				}
 				HandleDecoding();
 				// If this access unit requires us to drain the decoder we do it now.
 				if (bDrainAfterDecode)
diff --git a/Engine/Plugins/Media/ElectraPlayer/Source/ElectraPlayerRuntime/Private/Runtime/Decoder/AudioDecoder.h b/Engine/Plugins/Media/ElectraPlayer/Source/ElectraPlayerRuntime/Private/Runtime/Decoder/AudioDecoder.h
--- a/Engine/Plugins/Media/ElectraPlayer/Source/ElectraPlayerRuntime/Private/Runtime/Decoder/AudioDecoder.h
+++ b/Engine/Plugins/Media/ElectraPlayer/Source/ElectraPlayerRuntime/Private/Runtime/Decoder/AudioDecoder.h
@@ -8,10 +8,10 @@
 #include "Decoder/DecoderBase.h"
 
 
-
 namespace Electra
 {
 	class IPlayerSessionServices;
+	class FOutputHandlerAudio;
 
 
 	class IAudioDecoder : public IDecoderBase, public IAccessUnitBufferInterface, public IDecoderAUBufferDiags, public IDecoderReadyBufferDiags
@@ -23,6 +23,8 @@ namespace Electra
 
 		virtual ~IAudioDecoder() = default;
 
+		virtual void SetOutputHandler(TSharedPtr<FOutputHandlerAudio, ESPMode::ThreadSafe> InOutputHandler) = 0;
+
 		virtual void SetPlayerSessionServices(IPlayerSessionServices* SessionServices) = 0;
 
 		virtual void Open(TSharedPtrTS<FAccessUnit::CodecData> InCodecData) = 0;
@@ -33,7 +35,6 @@ namespace Electra
 		//-------------------------------------------------------------------------
 		// Methods from IDecoderBase
 		//
-		virtual void SetRenderer(TSharedPtr<IMediaRenderer, ESPMode::ThreadSafe> InRenderer) = 0;
 		virtual void SuspendOrResumeDecoder(bool bSuspend, const FParamDict& InOptions) = 0;
 
 		//-------------------------------------------------------------------------
diff --git a/Engine/Plugins/Media/ElectraPlayer/Source/ElectraPlayerRuntime/Private/Runtime/Decoder/DecoderBase.h b/Engine/Plugins/Media/ElectraPlayer/Source/ElectraPlayerRuntime/Private/Runtime/Decoder/DecoderBase.h
--- a/Engine/Plugins/Media/ElectraPlayer/Source/ElectraPlayerRuntime/Private/Runtime/Decoder/DecoderBase.h
+++ b/Engine/Plugins/Media/ElectraPlayer/Source/ElectraPlayerRuntime/Private/Runtime/Decoder/DecoderBase.h
@@ -4,7 +4,6 @@
 
 #include "PlayerCore.h"
 #include "Player/PlayerSessionServices.h"
-#include "Renderer/RendererBase.h"
 #include "ParameterDictionary.h"
 
 namespace Electra
@@ -36,7 +35,6 @@ namespace Electra
 	{
 	public:
 		virtual ~IDecoderBase() = default;
-		virtual void SetRenderer(TSharedPtr<IMediaRenderer, ESPMode::ThreadSafe> InRenderer) = 0;
 		//virtual void SuspendOrResumeDecoder(bool bSuspend, const FParamDict& InOptions) = 0;
 	};
 
@@ -126,4 +124,3 @@ namespace Electra
 
 
 } // namespace Electra
-
