Reworked ElectraSamples, ElectraUtils and the decoder output of Electra Player

diff --git a/Engine/Plugins/Media/ElectraPlayer/Source/ElectraPlayerRuntime/Private/Runtime/Decoder/VideoDecoder.cpp b/Engine/Plugins/Media/ElectraPlayer/Source/ElectraPlayerRuntime/Private/Runtime/Decoder/VideoDecoder.cpp
--- a/Engine/Plugins/Media/ElectraPlayer/Source/ElectraPlayerRuntime/Private/Runtime/Decoder/VideoDecoder.cpp
+++ b/Engine/Plugins/Media/ElectraPlayer/Source/ElectraPlayerRuntime/Private/Runtime/Decoder/VideoDecoder.cpp
@@ -10,20 +10,22 @@
 
 #include "StreamAccessUnitBuffer.h"
 #include "Decoder/VideoDecoder.h"
-#include "Decoder/VideoDecoderHelpers.h"
-#include "Renderer/RendererBase.h"
+#include "OutputHandler.h"
 #include "Player/PlayerSessionServices.h"
 #include "Utilities/Utilities.h"
 #include "Utilities/StringHelpers.h"
 #include "SynchronizedClock.h"
 
+#include "MediaDecoderOutput.h"
+#include "Utils/VideoDecoderHelpers.h"
+
 #include "IElectraCodecFactoryModule.h"
 #include "IElectraCodecFactory.h"
 #include "IElectraDecoder.h"
 #include "IElectraDecoderOutputVideo.h"
 #include "IElectraDecoderFeaturesAndOptions.h"
 #include "ElectraDecodersUtils.h"
-#include "ElectraDecoderResourceManager.h"
+#include "ElectraDecodersPlatformResources.h"
 
 // Error codes must be in the 1000-1999 range. 1-999 is reserved for the decoder implementation.
 #define ERRCODE_VIDEO_INTERNAL_COULD_NOT_CREATE_DECODER			1001
@@ -56,7 +58,7 @@ public:
 	bool Reopen(TSharedPtrTS<FAccessUnit::CodecData> InCodecData, const FParamDict& InAdditionalOptions, const FStreamCodecInformation* InMaxStreamConfiguration) override;
 	void Close() override;
 	void DrainForCodecChange() override;
-	void SetRenderer(TSharedPtr<IMediaRenderer, ESPMode::ThreadSafe> InRenderer) override;
+	void SetOutputHandler(TSharedPtr<FOutputHandlerVideo, ESPMode::ThreadSafe> InOutputHandler) override;
 	void SuspendOrResumeDecoder(bool bSuspend, const FParamDict& InOptions) override;
 	void AUdataPushAU(FAccessUnit* AccessUnit) override;
 	void AUdataPushEOD() override;
@@ -132,6 +134,16 @@ private:
 
 	void NotifyReadyBufferListener(bool bHaveOutput);
 
+	bool CanReceiveOutputFrame()
+	{
+		bool bCanRcv = true;
+		if (OutputHandler)
+		{
+			bCanRcv = OutputHandler->CanReceiveOutputSample();
+		}
+		return bCanRcv;
+	}
+
 	EAUChangeFlags GetAndPrepareInputAU();
 	EAUChangeFlags PrepareAU(TSharedPtrTS<FDecoderInput> AU);
 	IElectraDecoder::ECSDCompatibility IsCompatibleWith();
@@ -192,7 +204,7 @@ private:
 	TSharedPtrTS<FFGBGNotificationHandlers>									FGBGHandlers;
 	int32																	ApplicationSuspendCount = 0;
 
-	TSharedPtr<IMediaRenderer, ESPMode::ThreadSafe>							Renderer;
+	TSharedPtr<FOutputHandlerVideo, ESPMode::ThreadSafe>					OutputHandler;
 	int32																	MaxOutputBuffers = 0;
 
 	FCriticalSection														ListenerMutex;
@@ -205,7 +217,7 @@ private:
 	TMap<FString, FVariant>													DecoderFactoryAddtlCfg;
 	FString																	DecoderFormat;
 
-	IElectraDecoderResourceDelegateBase::IDecoderPlatformResource*			PlatformResource = nullptr;
+	FElectraDecodersPlatformResources::IDecoderPlatformResource*			DecoderPlatformResource = nullptr;
 
 	TMap<FString, FVariant>													DecoderConfigOptions;
 	TSharedPtr<IElectraDecoder, ESPMode::ThreadSafe>						DecoderInstance;
@@ -220,8 +232,7 @@ private:
 	TOptional<MPEG::FColorimetryHelper>										CurrentColorimetry;
 	TOptional<MPEG::FHDRHelper>												CurrentHDR;
 
-	IMediaRenderer::IBuffer*												CurrentOutputBuffer = nullptr;
-	FParamDict																EmptyOptions;
+	FElectraTextureSamplePtr												CurrentOutputBuffer;
 	FParamDict																DummyBufferSampleProperties;
 };
 ENUM_CLASS_FLAGS(FVideoDecoderImpl::EAUChangeFlags);
@@ -373,9 +384,9 @@ void FVideoDecoderImpl::DrainForCodecChange()
 	bDrainForCodecChange = true;
 }
 
-void FVideoDecoderImpl::SetRenderer(TSharedPtr<IMediaRenderer, ESPMode::ThreadSafe> InRenderer)
+void FVideoDecoderImpl::SetOutputHandler(TSharedPtr<FOutputHandlerVideo, ESPMode::ThreadSafe> InOutputHandler)
 {
-	Renderer = InRenderer;
+	OutputHandler = MoveTemp(InOutputHandler);
 }
 
 void FVideoDecoderImpl::SuspendOrResumeDecoder(bool bSuspend, const FParamDict& InOptions)
@@ -438,26 +449,28 @@ void FVideoDecoderImpl::StopThread()
 
 void FVideoDecoderImpl::CreateDecoderOutputPool()
 {
-	FParamDict poolOpts;
-	check(Renderer);
-
-// TODO/FIXME: get the default value of 8 from some config option?
-	int64 NumOutputFrames = ElectraDecodersUtil::GetVariantValueSafeI64(DecoderConfigOptions, IElectraDecoderFeature::MinimumNumberOfOutputFrames, 8);
-	poolOpts.Set(RenderOptionKeys::NumBuffers, FVariantValue(NumOutputFrames));
-	if (Renderer->CreateBufferPool(poolOpts) == UEMEDIA_ERROR_OK)
+	if (OutputHandler)
 	{
-		MaxOutputBuffers = (int32) Renderer->GetBufferPoolProperties().GetValue(RenderOptionKeys::MaxBuffers).GetInt64();
-		DecoderFactoryAddtlCfg.Add(TEXT("max_output_buffers"), FVariant((uint32)MaxOutputBuffers));
-	}
-	else
-	{
-		PostError(0, TEXT("Failed to create sample pool"), ERRCODE_VIDEO_INTERNAL_COULD_NOT_CREATE_SAMPLE_POOL);
+		FParamDict poolOpts;
+		// Note: Get the default value of 8 from some config option?
+		int64 NumOutputFrames = ElectraDecodersUtil::GetVariantValueSafeI64(DecoderConfigOptions, IElectraDecoderFeature::MinimumNumberOfOutputFrames, 8);
+		poolOpts.Set(OutputHandlerOptionKeys::NumBuffers, FVariantValue(NumOutputFrames));
+		if (OutputHandler->PreparePool(poolOpts))
+		{
+			MaxOutputBuffers = (int32) OutputHandler->GetPoolProperties().GetValue(OutputHandlerOptionKeys::NumBuffers).GetInt64();
+			DecoderFactoryAddtlCfg.Add(TEXT("max_output_buffers"), FVariant((uint32)MaxOutputBuffers));
+			return;
+		}
 	}
+	PostError(0, TEXT("Failed to create sample pool"), ERRCODE_VIDEO_INTERNAL_COULD_NOT_CREATE_SAMPLE_POOL);
 }
 
 void FVideoDecoderImpl::DestroyDecoderOutputPool()
 {
-	Renderer->ReleaseBufferPool();
+	if (OutputHandler)
+	{
+		OutputHandler->ClosePool();
+	}
 }
 
 void FVideoDecoderImpl::NotifyReadyBufferListener(bool bHaveOutput)
@@ -478,7 +491,7 @@ void FVideoDecoderImpl::NotifyReadyBufferListener(bool bHaveOutput)
 		}
 		stats.InDecoderTimeRangePTS.End.SetSequenceIndex(stats.InDecoderTimeRangePTS.Start.GetSequenceIndex());
 		stats.bOutputStalled = !bHaveOutput;
-		stats.bEODreached = NextAccessUnits.ReachedEOD() && CurrentOutputBuffer == nullptr;
+		stats.bEODreached = NextAccessUnits.ReachedEOD() && !CurrentOutputBuffer;
 		ListenerMutex.Lock();
 		if (ReadyBufferListener)
 		{
@@ -557,7 +570,7 @@ bool FVideoDecoderImpl::InternalDecoderCreate()
 
 	// Create platform specifics.
 	TMap<FString, FVariant> PlatformSpecificCfg(DecoderFactoryAddtlCfg);
-	PlatformResource = FPlatformElectraDecoderResourceManager::GetDelegate()->CreatePlatformResource(this, IElectraDecoderResourceDelegateBase::EDecoderPlatformResourceType::Video, PlatformSpecificCfg);
+	DecoderPlatformResource = FElectraDecodersPlatformResources::CreatePlatformVideoResource(this, PlatformSpecificCfg);
 
 	// Add in video decoder special options passed from the application.
 	TMap<FString, FVariant> DecoderCreateCfg(DecoderFactoryAddtlCfg);
@@ -620,10 +633,10 @@ void FVideoDecoderImpl::InternalDecoderDestroy()
 		DecoderInstance->Close();
 		DecoderInstance.Reset();
 	}
-	if (PlatformResource)
+	if (DecoderPlatformResource)
 	{
-		FPlatformElectraDecoderResourceManager::GetDelegate()->ReleasePlatformResource(this, PlatformResource);
-		PlatformResource = nullptr;
+		FElectraDecodersPlatformResources::ReleasePlatformVideoResource(this, DecoderPlatformResource);
+		DecoderPlatformResource = nullptr;
 	}
 	bIsAdaptiveDecoder = false;
 	bSupportsDroppingOutput = false;
@@ -635,10 +648,10 @@ void FVideoDecoderImpl::InternalDecoderDestroy()
 
 void FVideoDecoderImpl::ReturnUnusedOutputBuffer()
 {
-	if (CurrentOutputBuffer)
+	if (CurrentOutputBuffer && OutputHandler)
 	{
-		Renderer->ReturnBuffer(CurrentOutputBuffer, false, EmptyOptions);
-		CurrentOutputBuffer = nullptr;
+		OutputHandler->ReturnOutputSample(CurrentOutputBuffer, FOutputHandlerVideo::EReturnSampleType::DontSendToQueue);
+		CurrentOutputBuffer.Reset();
 	}
 }
 
@@ -870,8 +883,8 @@ IElectraDecoder::EOutputStatus FVideoDecoderImpl::HandleOutput()
 				continue;
 			}
 
-			// Check if the renderer can accept the output we want to send to it.
-			if (Renderer.IsValid() && !Renderer->CanReceiveOutputFrames(1))
+			// Check if the output pipeline can accept the output we want to send to it.
+			if (!CanReceiveOutputFrame())
 			{
 				NotifyReadyBufferListener(false);
 				return IElectraDecoder::EOutputStatus::TryAgainLater;
@@ -895,19 +908,18 @@ IElectraDecoder::EOutputStatus FVideoDecoderImpl::HandleOutput()
 			}
 
 			// Need a new output buffer?
-			if (CurrentOutputBuffer == nullptr && Renderer.IsValid())
+			if (!CurrentOutputBuffer && OutputHandler)
 			{
 				SCOPE_CYCLE_COUNTER(STAT_ElectraPlayer_VideoConvertOutput);
 				CSV_SCOPED_TIMING_STAT(ElectraPlayer, VideoConvertOutput);
-				UEMediaError bufResult = Renderer->AcquireBuffer(CurrentOutputBuffer, 0, EmptyOptions);
-				check(bufResult == UEMEDIA_ERROR_OK || bufResult == UEMEDIA_ERROR_INSUFFICIENT_DATA);
-				if (bufResult != UEMEDIA_ERROR_OK && bufResult != UEMEDIA_ERROR_INSUFFICIENT_DATA)
+				IOutputHandlerBase::EBufferResult bufResult = OutputHandler->ObtainOutputSample(CurrentOutputBuffer);
+				if (bufResult != IOutputHandlerBase::EBufferResult::Ok && bufResult != IOutputHandlerBase::EBufferResult::NoBuffer)
 				{
-					PostError(0, TEXT("Failed to acquire sample buffer"), ERRCODE_VIDEO_INTERNAL_COULD_NOT_GET_SAMPLE_BUFFER, bufResult);
+					PostError(0, TEXT("Failed to acquire sample buffer"), ERRCODE_VIDEO_INTERNAL_COULD_NOT_GET_SAMPLE_BUFFER);
 					return IElectraDecoder::EOutputStatus::Error;
 				}
 			}
-			// Didn't get a buffer? This should not really happen since the renderer said it could accept a frame.
+			// Didn't get a buffer? This should not really happen since the output pipeline said it could accept a frame.
 			if (!CurrentOutputBuffer)
 			{
 				NotifyReadyBufferListener(false);
@@ -976,8 +988,8 @@ IElectraDecoder::EOutputStatus FVideoDecoderImpl::HandleOutput()
 				{
 					// Create the platform specific decoder output.
 					TSharedPtr<FParamDict, ESPMode::ThreadSafe> BufferProperties(new FParamDict);
-					BufferProperties->Set(RenderOptionKeys::PTS, FVariantValue(MatchingInput->AdjustedPTS));
-					BufferProperties->Set(RenderOptionKeys::Duration, FVariantValue(MatchingInput->AdjustedDuration));
+					BufferProperties->Set(OutputHandlerOptionKeys::PTS, FVariantValue(MatchingInput->AdjustedPTS));
+					BufferProperties->Set(OutputHandlerOptionKeys::Duration, FVariantValue(MatchingInput->AdjustedDuration));
 
 					// Set properties from the bitstream messages.
 					if (DecoderBitstreamProcessor.IsValid())
@@ -1051,7 +1063,7 @@ IElectraDecoder::EOutputStatus FVideoDecoderImpl::HandleOutput()
 					}
 
 					FString DecoderOutputErrorMsg;
-					if (!FPlatformElectraDecoderResourceManager::SetupRenderBufferFromDecoderOutput(DecoderOutputErrorMsg, CurrentOutputBuffer, BufferProperties, CurrentDecoderOutput, PlatformResource))
+					if (!FElectraDecodersPlatformResources::SetupOutputTextureSample(DecoderOutputErrorMsg, CurrentOutputBuffer, CurrentDecoderOutput, BufferProperties, DecoderPlatformResource))
 					{
 						if (DecoderOutputErrorMsg.IsEmpty())
 						{
@@ -1063,8 +1075,9 @@ IElectraDecoder::EOutputStatus FVideoDecoderImpl::HandleOutput()
 						}
 						return IElectraDecoder::EOutputStatus::Error;
 					}
-					Renderer->ReturnBuffer(CurrentOutputBuffer, bUseOutput, *BufferProperties);
-					CurrentOutputBuffer = nullptr;
+
+					OutputHandler->ReturnOutputSample(CurrentOutputBuffer, bUseOutput ? FOutputHandlerVideo::EReturnSampleType::SendToQueue : FOutputHandlerVideo::EReturnSampleType::DontSendToQueue);
+					CurrentOutputBuffer.Reset();
 
 					if (bIsStartOfSequence && NumInitialSkippedFrames)
 					{
@@ -1279,8 +1292,8 @@ bool FVideoDecoderImpl::HandleDummyDecoding()
 	// Get output unless flushing or terminating
 	while(!TerminateThreadSignal.IsSignaled() && !FlushDecoderSignal.IsSignaled())
 	{
-		// Check if the renderer can accept the output we want to send to it.
-		if (Renderer.IsValid() && !Renderer->CanReceiveOutputFrames(1))
+		// Check if the output pipeline can accept the output we want to send to it.
+		if (!CanReceiveOutputFrame())
 		{
 			NotifyReadyBufferListener(false);
 			FMediaRunnable::SleepMilliseconds(5);
@@ -1288,15 +1301,14 @@ bool FVideoDecoderImpl::HandleDummyDecoding()
 		}
 
 		// Need a new output buffer?
-		if (CurrentOutputBuffer == nullptr && Renderer.IsValid())
+		if (!CurrentOutputBuffer && OutputHandler)
 		{
 			SCOPE_CYCLE_COUNTER(STAT_ElectraPlayer_VideoConvertOutput);
 			CSV_SCOPED_TIMING_STAT(ElectraPlayer, VideoConvertOutput);
-			UEMediaError bufResult = Renderer->AcquireBuffer(CurrentOutputBuffer, 0, EmptyOptions);
-			check(bufResult == UEMEDIA_ERROR_OK || bufResult == UEMEDIA_ERROR_INSUFFICIENT_DATA);
-			if (bufResult != UEMEDIA_ERROR_OK && bufResult != UEMEDIA_ERROR_INSUFFICIENT_DATA)
+			IOutputHandlerBase::EBufferResult bufResult = OutputHandler->ObtainOutputSample(CurrentOutputBuffer);
+			if (bufResult != IOutputHandlerBase::EBufferResult::Ok && bufResult != IOutputHandlerBase::EBufferResult::NoBuffer)
 			{
-				return PostError(0, TEXT("Failed to acquire sample buffer"), ERRCODE_VIDEO_INTERNAL_COULD_NOT_GET_SAMPLE_BUFFER, bufResult);
+				return PostError(0, TEXT("Failed to acquire sample buffer"), ERRCODE_VIDEO_INTERNAL_COULD_NOT_GET_SAMPLE_BUFFER);
 			}
 		}
 		// Didn't get a buffer?
@@ -1309,15 +1321,14 @@ bool FVideoDecoderImpl::HandleDummyDecoding()
 
 		NotifyReadyBufferListener(true);
 
-		DummyBufferSampleProperties.Set(RenderOptionKeys::Duration, FVariantValue(CurrentAccessUnit->AdjustedDuration));
-		DummyBufferSampleProperties.Set(RenderOptionKeys::PTS, FVariantValue(CurrentAccessUnit->AdjustedPTS));
-		DummyBufferSampleProperties.Set(RenderOptionKeys::DummyBufferFlag, FVariantValue(true));
-		Renderer->ReturnBuffer(CurrentOutputBuffer, true, DummyBufferSampleProperties);
-		CurrentOutputBuffer = nullptr;
+		CurrentOutputBuffer->SetTime(FMediaTimeStamp(CurrentAccessUnit->AdjustedPTS.GetAsTimespan(), CurrentAccessUnit->AdjustedPTS.GetSequenceIndex()));
+		CurrentOutputBuffer->SetDuration(CurrentAccessUnit->AdjustedDuration.GetAsTimespan());
+		OutputHandler->ReturnOutputSample(CurrentOutputBuffer, FOutputHandlerVideo::EReturnSampleType::DummySample);
+		CurrentOutputBuffer.Reset();
 		// We must not drain the source buffer too quickly. While our counterpart code in the audio decoder actually
 		// produces a usable sample containing silence, we cannot create a usable dummy frame because we have to
 		// keep the last good frame on screen. Our sample we have just returned will not actually be sent into the
-		// media sample queue and thus any next call to `Renderer->CanReceiveOutputFrames(1)` above will always
+		// media sample queue and thus any next call to `CanReceiveOutputFrame()` above will always
 		// return `true` because the sample queue will not be full, and as a result we race and take new source
 		// samples from the buffer so quickly, that the buffer will underrun.
 		// To prevent this we put ourselves to sleep for a while. Not the entire sample duration though, but for
@@ -1432,7 +1443,6 @@ void FVideoDecoderImpl::WorkerThread()
 	ApplicationSuspendConfirmedSignal.Reset();
 
 	bError = false;
-	CurrentOutputBuffer = nullptr;
 	bIsFirstAccessUnit = true;
 	bInDummyDecodeMode = false;
 	bIsAdaptiveDecoder = false;
diff --git a/Engine/Plugins/Media/ElectraPlayer/Source/ElectraPlayerRuntime/Private/Runtime/Decoder/VideoDecoder.h b/Engine/Plugins/Media/ElectraPlayer/Source/ElectraPlayerRuntime/Private/Runtime/Decoder/VideoDecoder.h
--- a/Engine/Plugins/Media/ElectraPlayer/Source/ElectraPlayerRuntime/Private/Runtime/Decoder/VideoDecoder.h
+++ b/Engine/Plugins/Media/ElectraPlayer/Source/ElectraPlayerRuntime/Private/Runtime/Decoder/VideoDecoder.h
@@ -8,10 +8,10 @@
 #include "Decoder/DecoderBase.h"
 
 
-
 namespace Electra
 {
 	class IPlayerSessionServices;
+	class FOutputHandlerVideo;
 
 
 	class IVideoDecoder : public IDecoderBase, public IAccessUnitBufferInterface, public IDecoderAUBufferDiags, public IDecoderReadyBufferDiags
@@ -23,6 +23,8 @@ namespace Electra
 
 		virtual ~IVideoDecoder() = default;
 
+		virtual void SetOutputHandler(TSharedPtr<FOutputHandlerVideo, ESPMode::ThreadSafe> InOutputHandler) = 0;
+
 		virtual void SetPlayerSessionServices(IPlayerSessionServices* SessionServices) = 0;
 
 		virtual void Open(TSharedPtrTS<FAccessUnit::CodecData> InCodecData, FParamDict&& InAdditionalOptions, const FStreamCodecInformation* InMaxStreamConfiguration) = 0;
@@ -33,7 +35,6 @@ namespace Electra
 		//-------------------------------------------------------------------------
 		// Methods from IDecoderBase
 		//
-		virtual void SetRenderer(TSharedPtr<IMediaRenderer, ESPMode::ThreadSafe> InRenderer) = 0;
 		virtual void SuspendOrResumeDecoder(bool bSuspend, const FParamDict& InOptions) = 0;
 
 		//-------------------------------------------------------------------------
@@ -61,4 +62,3 @@ namespace Electra
 	};
 
 } // namespace Electra
-
diff --git a/Engine/Plugins/Media/ElectraPlayer/Source/ElectraPlayerRuntime/Private/Runtime/Demuxer/ParserISO13818-1.cpp b/Engine/Plugins/Media/ElectraPlayer/Source/ElectraPlayerRuntime/Private/Runtime/Demuxer/ParserISO13818-1.cpp
--- a/Engine/Plugins/Media/ElectraPlayer/Source/ElectraPlayerRuntime/Private/Runtime/Demuxer/ParserISO13818-1.cpp
+++ b/Engine/Plugins/Media/ElectraPlayer/Source/ElectraPlayerRuntime/Private/Runtime/Demuxer/ParserISO13818-1.cpp
@@ -1,8 +1,8 @@
 // Copyright Epic Games, Inc. All Rights Reserved.
 
 #include "Demuxer/ParserISO13818-1.h"
-#include "Utils/ElectraBitstreamReader.h"
-#include "Utils/MPEG/ElectraUtilsMP4.h"
+#include "Utilities/ElectraBitstream.h"
+#include "Utilities/UtilitiesMP4.h"
 #include "Utils/MPEG/ElectraUtilsMPEGAudio.h"
 #include "Utils/MPEG/ElectraUtilsMPEGVideo_H264.h"
 #include "Utils/MPEG/ElectraUtilsMPEGVideo_H265.h"
@@ -162,7 +162,7 @@ private:
 		{
 			PacketDataBuffer = MakeSharedTS<TArray<uint8>>();
 		}
-		bool ProcessPayload(TArray<FPayload>& OutPayloadResults, ElectraDecodersUtil::FElectraBitstreamReader& InOutBR, const FCurrent& InCurrent);
+		bool ProcessPayload(TArray<FPayload>& OutPayloadResults, FBitstreamReader& InOutBR, const FCurrent& InCurrent);
 		void FinishCurrentPESPacket(TSharedPtrTS<FPESData>& OutPESPacket);
 	private:
 		void ExtractValidSections(TArray<TSharedPtrTS<TArray<uint8>>>& OutDataSections);
@@ -229,18 +229,18 @@ private:
 	EParseState ParseNextPacket(IPlayerSessionServices* InPlayerSession, IGenericDataReader* InDataReader);
 
 	ETableResult HandlePSITable(const FPayload& InTablePayload);
-	bool ProcessPAT(const FPSITable& InTableInfo, ElectraDecodersUtil::FElectraBitstreamReader& InOutBR);
-	bool ProcessPMT(const FPSITable& InTableInfo, ElectraDecodersUtil::FElectraBitstreamReader& InOutBR);
-	void ProcessSDT(const FPSITable& InTableInfo, ElectraDecodersUtil::FElectraBitstreamReader& InOutBR);
-	void ProcessDescriptor(ElectraDecodersUtil::FElectraBitstreamReader& InOutBR, int32 InStreamType, FStreamCodecInformation& OutCodecInfo);
+	bool ProcessPAT(const FPSITable& InTableInfo, FBitstreamReader& InOutBR);
+	bool ProcessPMT(const FPSITable& InTableInfo, FBitstreamReader& InOutBR);
+	void ProcessSDT(const FPSITable& InTableInfo, FBitstreamReader& InOutBR);
+	void ProcessDescriptor(FBitstreamReader& InOutBR, int32 InStreamType, FStreamCodecInformation& OutCodecInfo);
 
 	void DeselectAllPESStreams();
 	void ActivateUserStreamSelection();
 
-	EPESPacketResult ParseADTSAAC(TArray<FESPacket>& OutPackets, ElectraDecodersUtil::FElectraBitstreamReader& InOutBR, const FDTSPTS& InDTSPTS, FResidualPESData* InResidualData, bool bFlushResiduals);
-	EPESPacketResult ParseMPEGAudio(TArray<FESPacket>& OutPackets, ElectraDecodersUtil::FElectraBitstreamReader& InOutBR, const FDTSPTS& InDTSPTS, FResidualPESData* InResidualData, bool bFlushResiduals);
-	EPESPacketResult ParseAVC(TArray<FESPacket>& OutPackets, ElectraDecodersUtil::FElectraBitstreamReader& InOutBR, const FDTSPTS& InDTSPTS, FResidualPESData* InResidualData, bool bFlushResiduals);
-	EPESPacketResult ParseHEVC(TArray<FESPacket>& OutPackets, ElectraDecodersUtil::FElectraBitstreamReader& InOutBR, const FDTSPTS& InDTSPTS, FResidualPESData* InResidualData, bool bFlushResiduals);
+	EPESPacketResult ParseADTSAAC(TArray<FESPacket>& OutPackets, FBitstreamReader& InOutBR, const FDTSPTS& InDTSPTS, FResidualPESData* InResidualData, bool bFlushResiduals);
+	EPESPacketResult ParseMPEGAudio(TArray<FESPacket>& OutPackets, FBitstreamReader& InOutBR, const FDTSPTS& InDTSPTS, FResidualPESData* InResidualData, bool bFlushResiduals);
+	EPESPacketResult ParseAVC(TArray<FESPacket>& OutPackets, FBitstreamReader& InOutBR, const FDTSPTS& InDTSPTS, FResidualPESData* InResidualData, bool bFlushResiduals);
+	EPESPacketResult ParseHEVC(TArray<FESPacket>& OutPackets, FBitstreamReader& InOutBR, const FDTSPTS& InDTSPTS, FResidualPESData* InResidualData, bool bFlushResiduals);
 
 	TUniquePtr<FStaticInitSegReader> InitSegReader;
 
@@ -370,7 +370,7 @@ IParserISO13818_1::EParseState FParserISO13818_1::ParseNextPacket(IPlayerSession
 			NextParseState = EParseState::Failed;
 			return NextParseState;
 		}
-		ElectraDecodersUtil::FElectraBitstreamReader br(Current.DataBlock, 188, 1);
+		FBitstreamReader br(Current.DataBlock, 188, 1);
 		const uint32 transport_error_indicator = br.GetBits(1);
 		const uint32 payload_unit_start_indicator = br.GetBits(1);
 		const uint32 transport_priority = br.GetBits(1);
@@ -418,7 +418,7 @@ IParserISO13818_1::EParseState FParserISO13818_1::ParseNextPacket(IPlayerSession
 			int32 adaptation_field_length_remaining = (int32) adaptation_field_length;
 			if (adaptation_field_length)
 			{
-				ElectraDecodersUtil::FElectraBitstreamReader abr1(br);
+				FBitstreamReader abr1(br);
 
 				Current.AdaptationFieldFirstByte = (int32) abr1.PeekBits(8);
 
@@ -973,7 +973,7 @@ void FParserISO13818_1::FPIDStream::ExtractValidPESPackets(TArray<TSharedPtrTS<F
 }
 
 
-bool FParserISO13818_1::FPIDStream::ProcessPayload(TArray<FPayload>& OutPayloadResults, ElectraDecodersUtil::FElectraBitstreamReader& InOutBR, const FCurrent& InCurrent)
+bool FParserISO13818_1::FPIDStream::ProcessPayload(TArray<FPayload>& OutPayloadResults, FBitstreamReader& InOutBR, const FCurrent& InCurrent)
 {
 	check(PacketDataBuffer.IsValid());
 	if (!PacketDataBuffer.IsValid())
@@ -1115,7 +1115,7 @@ FParserISO13818_1::ETableResult FParserISO13818_1::HandlePSITable(const FPayload
 	check(InTablePayload.Type == EPayloadType::PSI);
 	check(InTablePayload.Data.IsValid());
 
-	ElectraDecodersUtil::FElectraBitstreamReader br(InTablePayload.Data->GetData(), InTablePayload.Data->Num());
+	FBitstreamReader br(InTablePayload.Data->GetData(), InTablePayload.Data->Num());
 
 	uint8 table_id = (uint8) br.GetBits(8);
 	if (table_id == 0xff)
@@ -1147,7 +1147,7 @@ FParserISO13818_1::ETableResult FParserISO13818_1::HandlePSITable(const FPayload
 		return ETableResult::Continue;
 	}
 
-	auto GetStandardTable = [](FPSITable& InOutTable, ElectraDecodersUtil::FElectraBitstreamReader& InOutBR)
+	auto GetStandardTable = [](FPSITable& InOutTable, FBitstreamReader& InOutBR)
 	{
 		InOutTable.TransportOrProgramStreamId = (int32) InOutBR.GetBits(16);
 		InOutBR.SkipBits(2);	// reserved
@@ -1231,7 +1231,7 @@ FParserISO13818_1::ETableResult FParserISO13818_1::HandlePSITable(const FPayload
 }
 
 
-bool FParserISO13818_1::ProcessPAT(const FPSITable& InTableInfo, ElectraDecodersUtil::FElectraBitstreamReader& InOutBR)
+bool FParserISO13818_1::ProcessPAT(const FPSITable& InTableInfo, FBitstreamReader& InOutBR)
 {
 	const uint32 kCRC32_ByteSize = 4;
 	const uint32 kTableEntryByteSize = 4;
@@ -1280,7 +1280,7 @@ bool FParserISO13818_1::ProcessPAT(const FPSITable& InTableInfo, ElectraDecoders
 	return false;
 }
 
-bool FParserISO13818_1::ProcessPMT(const FPSITable& InTableInfo, ElectraDecodersUtil::FElectraBitstreamReader& InOutBR)
+bool FParserISO13818_1::ProcessPMT(const FPSITable& InTableInfo, FBitstreamReader& InOutBR)
 {
 	TSharedPtrTS<FProgramMap> NewProgramMap = MakeSharedTS<FProgramMap>();
 	NewProgramMap->VersionNumber = InTableInfo.VersionNumber;
@@ -1293,7 +1293,7 @@ bool FParserISO13818_1::ProcessPMT(const FPSITable& InTableInfo, ElectraDecoders
 	uint32 program_info_length = InOutBR.GetBits(12) & 0x3ff;	// only 10 bits used, upper 2 bits must be zero, we ignore them.
 	if (program_info_length)
 	{
-		ElectraDecodersUtil::FElectraBitstreamReader br(InOutBR.GetRemainingData(), program_info_length);
+		FBitstreamReader br(InOutBR.GetRemainingData(), program_info_length);
 		InOutBR.SkipBytes(program_info_length);
 		while(br.GetRemainingByteLength())
 		{
@@ -1315,7 +1315,7 @@ bool FParserISO13818_1::ProcessPMT(const FPSITable& InTableInfo, ElectraDecoders
 		// Descriptors?
 		if (ES_info_length)
 		{
-			ElectraDecodersUtil::FElectraBitstreamReader br(InOutBR.GetRemainingData(), ES_info_length);
+			FBitstreamReader br(InOutBR.GetRemainingData(), ES_info_length);
 			InOutBR.SkipBytes(ES_info_length);
 			while(br.GetRemainingByteLength())
 			{
@@ -1433,7 +1433,7 @@ bool FParserISO13818_1::ProcessPMT(const FPSITable& InTableInfo, ElectraDecoders
 }
 
 // ETSI EN 300 468 - 5.2.3 Service Description Table
-void FParserISO13818_1::ProcessSDT(const FPSITable& InTableInfo, ElectraDecodersUtil::FElectraBitstreamReader& InOutBR)
+void FParserISO13818_1::ProcessSDT(const FPSITable& InTableInfo, FBitstreamReader& InOutBR)
 {
 	uint32 original_network_id = InOutBR.GetBits(16);
 	InOutBR.SkipBits(8);	// reserved for future use
@@ -1450,7 +1450,7 @@ void FParserISO13818_1::ProcessSDT(const FPSITable& InTableInfo, ElectraDecoders
 		// Descriptors?
 		if (descriptors_length)
 		{
-			ElectraDecodersUtil::FElectraBitstreamReader br(InOutBR.GetRemainingData(), descriptors_length);
+			FBitstreamReader br(InOutBR.GetRemainingData(), descriptors_length);
 			InOutBR.SkipBytes(descriptors_length);
 			while(br.GetRemainingByteLength())
 			{
@@ -1459,7 +1459,7 @@ void FParserISO13818_1::ProcessSDT(const FPSITable& InTableInfo, ElectraDecoders
 				// 6.2.33 Service descriptor ?
 				if (descriptor_tag == 0x48)
 				{
-					ElectraDecodersUtil::FElectraBitstreamReader dr(br.GetRemainingData(), descriptor_length);
+					FBitstreamReader dr(br.GetRemainingData(), descriptor_length);
 					uint32 service_type = dr.GetBits(8);
 					// no use for this right now
 					(void)service_type;
@@ -1476,11 +1476,11 @@ void FParserISO13818_1::ProcessSDT(const FPSITable& InTableInfo, ElectraDecoders
 	}
 }
 
-void FParserISO13818_1::ProcessDescriptor(ElectraDecodersUtil::FElectraBitstreamReader& InOutBR, int32 InStreamType, FStreamCodecInformation& OutCodecInfo)
+void FParserISO13818_1::ProcessDescriptor(FBitstreamReader& InOutBR, int32 InStreamType, FStreamCodecInformation& OutCodecInfo)
 {
 	uint32 descriptor_tag = InOutBR.GetBits(8);
 	uint32 descriptor_length = InOutBR.GetBits(8);
-	ElectraDecodersUtil::FElectraBitstreamReader br(InOutBR.GetRemainingData(), descriptor_length);
+	FBitstreamReader br(InOutBR.GetRemainingData(), descriptor_length);
 	InOutBR.SkipBytes(descriptor_length);
 	switch(descriptor_tag)
 	{
@@ -1601,7 +1601,7 @@ IParserISO13818_1::EPESPacketResult FParserISO13818_1::ParsePESPacket(TArray<IPa
 	}
 
 	FDTSPTS DTSPTS;
-	ElectraDecodersUtil::FElectraBitstreamReader br;
+	FBitstreamReader br;
 	uint8 stream_id = 0xbe;
 	bool bHandleStreamId = false;
 	if (InPESPacket->PacketData.IsValid())
@@ -1644,7 +1644,7 @@ IParserISO13818_1::EPESPacketResult FParserISO13818_1::ParsePESPacket(TArray<IPa
 			const uint32 PES_CRC_flag = br.GetBits(1);
 			const uint32 PES_extension_flag = br.GetBits(1);
 			const uint32 PES_header_data_length = br.GetBits(8);
-			ElectraDecodersUtil::FElectraBitstreamReader opt(br);
+			FBitstreamReader opt(br);
 			br.SkipBytes(PES_header_data_length);
 			if (PTS_DTS_flags == 2)
 			{
@@ -1824,9 +1824,9 @@ IParserISO13818_1::EPESPacketResult FParserISO13818_1::ParsePESPacket(TArray<IPa
 	}
 }
 
-FParserISO13818_1::EPESPacketResult FParserISO13818_1::ParseADTSAAC(TArray<FESPacket>& OutPackets, ElectraDecodersUtil::FElectraBitstreamReader& InOutBR, const FDTSPTS& InDTSPTS, FResidualPESData* InResidualData, bool bFlushResiduals)
+FParserISO13818_1::EPESPacketResult FParserISO13818_1::ParseADTSAAC(TArray<FESPacket>& OutPackets, FBitstreamReader& InOutBR, const FDTSPTS& InDTSPTS, FResidualPESData* InResidualData, bool bFlushResiduals)
 {
-	ElectraDecodersUtil::FElectraBitstreamReader br(InOutBR);
+	FBitstreamReader br(InOutBR);
 	FDTSPTS DTSPTS(InDTSPTS);
 
 	// Are we dealing with residuals?
@@ -1925,9 +1925,9 @@ FParserISO13818_1::EPESPacketResult FParserISO13818_1::ParseADTSAAC(TArray<FESPa
 	return EPESPacketResult::Ok;
 }
 
-FParserISO13818_1::EPESPacketResult FParserISO13818_1::ParseMPEGAudio(TArray<FESPacket>& OutPackets, ElectraDecodersUtil::FElectraBitstreamReader& InOutBR, const FDTSPTS& InDTSPTS, FResidualPESData* InResidualData, bool bFlushResiduals)
+FParserISO13818_1::EPESPacketResult FParserISO13818_1::ParseMPEGAudio(TArray<FESPacket>& OutPackets, FBitstreamReader& InOutBR, const FDTSPTS& InDTSPTS, FResidualPESData* InResidualData, bool bFlushResiduals)
 {
-	ElectraDecodersUtil::FElectraBitstreamReader br(InOutBR);
+	FBitstreamReader br(InOutBR);
 	FDTSPTS DTSPTS(InDTSPTS);
 
 	// Are we dealing with residuals?
@@ -2023,9 +2023,9 @@ FParserISO13818_1::EPESPacketResult FParserISO13818_1::ParseMPEGAudio(TArray<FES
 	return EPESPacketResult::Ok;
 }
 
-FParserISO13818_1::EPESPacketResult FParserISO13818_1::ParseAVC(TArray<FESPacket>& OutPackets, ElectraDecodersUtil::FElectraBitstreamReader& InOutBR, const FDTSPTS& InDTSPTS, FResidualPESData* InResidualData, bool bFlushResiduals)
+FParserISO13818_1::EPESPacketResult FParserISO13818_1::ParseAVC(TArray<FESPacket>& OutPackets, FBitstreamReader& InOutBR, const FDTSPTS& InDTSPTS, FResidualPESData* InResidualData, bool bFlushResiduals)
 {
-	ElectraDecodersUtil::FElectraBitstreamReader br(InOutBR);
+	FBitstreamReader br(InOutBR);
 	FDTSPTS DTSPTS(InDTSPTS);
 
 	// Are we dealing with residuals?
@@ -2168,9 +2168,9 @@ FParserISO13818_1::EPESPacketResult FParserISO13818_1::ParseAVC(TArray<FESPacket
 	return EPESPacketResult::Ok;
 }
 
-FParserISO13818_1::EPESPacketResult FParserISO13818_1::ParseHEVC(TArray<FESPacket>& OutPackets, ElectraDecodersUtil::FElectraBitstreamReader& InOutBR, const FDTSPTS& InDTSPTS, FResidualPESData* InResidualData, bool bFlushResiduals)
+FParserISO13818_1::EPESPacketResult FParserISO13818_1::ParseHEVC(TArray<FESPacket>& OutPackets, FBitstreamReader& InOutBR, const FDTSPTS& InDTSPTS, FResidualPESData* InResidualData, bool bFlushResiduals)
 {
-	ElectraDecodersUtil::FElectraBitstreamReader br(InOutBR);
+	FBitstreamReader br(InOutBR);
 	FDTSPTS DTSPTS(InDTSPTS);
 
 	// Are we dealing with residuals?
diff --git a/Engine/Plugins/Media/ElectraPlayer/Source/ElectraPlayerRuntime/Private/Runtime/Player/AdaptivePlayerOptionKeynames.h b/Engine/Plugins/Media/ElectraPlayer/Source/ElectraPlayerRuntime/Private/Runtime/Player/AdaptivePlayerOptionKeynames.h
--- a/Engine/Plugins/Media/ElectraPlayer/Source/ElectraPlayerRuntime/Private/Runtime/Player/AdaptivePlayerOptionKeynames.h
+++ b/Engine/Plugins/Media/ElectraPlayer/Source/ElectraPlayerRuntime/Private/Runtime/Player/AdaptivePlayerOptionKeynames.h
@@ -36,9 +36,6 @@ namespace Electra
 	//! (bool) true to perform frame accurate seeks (slow as decoding and discarding data from a preceeding keyframe is required)
 	const FName OptionKeyFrameAccurateSeek(TEXT("frame_accurate_seeking"));
 
-	//! (bool) false (default) to emit the first decoded video frame while prerolling so it can be displayed while scrubbing, true to not display it right away.
-	const FName OptionKeyHoldBackFirstVideoFrame(TEXT("hold_back_first_frame"));
-
 	//! (bool) true to not truncate the media segment access units at the end of the presentation. Must only be used without a set playback range end!
 	const FName OptionKeyDoNotTruncateAtPresentationEnd(TEXT("do_not_truncate_at_presentation_end"));
 
diff --git a/Engine/Plugins/Media/ElectraPlayer/Source/ElectraPlayerRuntime/Private/Runtime/Player/AdaptiveStreamingPlayer.cpp b/Engine/Plugins/Media/ElectraPlayer/Source/ElectraPlayerRuntime/Private/Runtime/Player/AdaptiveStreamingPlayer.cpp
--- a/Engine/Plugins/Media/ElectraPlayer/Source/ElectraPlayerRuntime/Private/Runtime/Player/AdaptiveStreamingPlayer.cpp
+++ b/Engine/Plugins/Media/ElectraPlayer/Source/ElectraPlayerRuntime/Private/Runtime/Player/AdaptiveStreamingPlayer.cpp
@@ -49,28 +49,22 @@ FAdaptiveStreamingPlayer::FAdaptiveStreamingPlayer(const IAdaptiveStreamingPlaye
 {
 	Electra::AddActivePlayerInstance();
 
-	if (InCreateParameters.VideoRenderer)
-	{
-		VideoRender.Renderer = CreateWrappedRenderer(InCreateParameters.VideoRenderer, EStreamType::Video);
-	}
-	if (InCreateParameters.AudioRenderer)
-	{
-		AudioRender.Renderer = CreateWrappedRenderer(InCreateParameters.AudioRenderer, EStreamType::Audio);
-	}
-	ExternalPlayerGUID		= InCreateParameters.ExternalPlayerGUID;
-	UseSharedWorkerThreads  = InCreateParameters.WorkerThreads;
-	ManifestType			= EMediaFormatType::Unknown;
-	CurrentState   		    = EPlayerState::eState_Idle;
-	PipelineState  		    = EPipelineState::ePipeline_Stopped;
-	DecoderState   		    = EDecoderState::eDecoder_Paused;
-	StreamState			    = EStreamState::eStream_Running;
-	RenderRateScale			= 1.0;
-	StreamReaderHandler     = nullptr;
-	bStreamingHasStarted	= false;
-	RebufferCause   	    = ERebufferCause::None;
-	Manifest  				= nullptr;
-	LastBufferingState 		= EPlayerState::eState_Buffering;
-	bIsClosing 				= false;
+	VideoOutput.OutputHandler = InCreateParameters.VideoOutputHandler;
+	AudioOutput.OutputHandler = InCreateParameters.AudioOutputHandler;
+	ExternalPlayerGUID = InCreateParameters.ExternalPlayerGUID;
+	UseSharedWorkerThreads = InCreateParameters.WorkerThreads;
+	ManifestType = EMediaFormatType::Unknown;
+	CurrentState = EPlayerState::eState_Idle;
+	PipelineState = EPipelineState::ePipeline_Stopped;
+	DecoderState = EDecoderState::eDecoder_Paused;
+	StreamState = EStreamState::eStream_Running;
+	RenderRateScale = 1.0;
+	StreamReaderHandler = nullptr;
+	bStreamingHasStarted = false;
+	RebufferCause = ERebufferCause::None;
+	Manifest = nullptr;
+	LastBufferingState = EPlayerState::eState_Buffering;
+	bIsClosing = false;
 
 	RebufferDetectedAtPlayPos.SetToInvalid();
 
@@ -1017,21 +1011,21 @@ void FAdaptiveStreamingPlayer::QueryStreamBufferInfo(FStreamBufferInfo& OutStrea
 				};
 
 				// Ask the renderers what they know they have released into the media sample queue.
-				TArray<IAdaptiveStreamingWrappedRenderer::FEnqueuedSampleInfo> Enqueued;
-				if (InStreamType == EStreamType::Video && VideoRender.Renderer.IsValid())
+				TArray<IOutputHandlerBase::FEnqueuedSampleInfo> Enqueued;
+				if (InStreamType == EStreamType::Video && VideoOutput.OutputHandler.IsValid())
 				{
-					VideoRender.Renderer->GetNumEnqueuedSamples(&Enqueued);
+					VideoOutput.OutputHandler->GetEnqueuedSampleInfo(Enqueued);
 				}
-				else if (InStreamType == EStreamType::Audio && AudioRender.Renderer.IsValid())
+				else if (InStreamType == EStreamType::Audio && AudioOutput.OutputHandler.IsValid())
 				{
-					AudioRender.Renderer->GetNumEnqueuedSamples(&Enqueued);
+					AudioOutput.OutputHandler->GetEnqueuedSampleInfo(Enqueued);
 				}
 				if (Enqueued.Num())
 				{
 					// TBD: Do we need to break them down into separate ranges based on the sequence index?
 					FTimeRange tr;
-					tr.Start = Enqueued[0].PTS;
-					tr.End = Enqueued.Last().PTS + Enqueued.Last().Duration;
+					tr.Start.SetFromTimespan(Enqueued[0].Timestamp.GetTime());
+					tr.End.SetFromTimespan(Enqueued.Last().Timestamp.GetTime() + Enqueued.Last().Duration);
 					OutStreamBufferInfo.TimeEnqueued.Emplace(MoveTemp(tr));
 				}
 
@@ -1082,16 +1076,13 @@ void FAdaptiveStreamingPlayer::StartRendering()
 {
 	RenderClock->Start();
 
-	if (VideoRender.Renderer)
+	if (VideoOutput.OutputHandler)
 	{
-		FParamDict noOptions;
-		VideoRender.Renderer->StartRendering(noOptions);
+		VideoOutput.OutputHandler->StartOutput();
 	}
-
-	if (AudioRender.Renderer)
+	if (AudioOutput.OutputHandler)
 	{
-		FParamDict noOptions;
-		AudioRender.Renderer->StartRendering(noOptions);
+		AudioOutput.OutputHandler->StartOutput();
 	}
 }
 
@@ -1103,15 +1094,13 @@ void FAdaptiveStreamingPlayer::StopRendering()
 {
 	RenderClock->Stop();
 
-	if (VideoRender.Renderer)
+	if (VideoOutput.OutputHandler)
 	{
-		FParamDict noOptions;
-		VideoRender.Renderer->StopRendering(noOptions);
+		VideoOutput.OutputHandler->StopOutput();
 	}
-	if (AudioRender.Renderer)
+	if (AudioOutput.OutputHandler)
 	{
-		FParamDict noOptions;
-		AudioRender.Renderer->StopRendering(noOptions);
+		AudioOutput.OutputHandler->StopOutput();
 	}
 }
 
@@ -1128,18 +1117,14 @@ int32 FAdaptiveStreamingPlayer::CreateRenderers()
 	const bool bCurrentlyPaused = PlaybackState.GetIsPaused();
 
 	// Set the render clock with the renderes.
-	if (VideoRender.Renderer)
+	if (VideoOutput.OutputHandler)
 	{
-		VideoRender.Renderer->SetRenderClock(RenderClock);
-		VideoRender.Renderer->SetPlaybackRate(CurrentRate, IntendedRate, bCurrentlyPaused);
-		// Hold back all video frames during preroll or emit the first frame for scrubbing?
-		bool bHoldBackFirstVideoFrame = PlayerOptions.GetValue(OptionKeyHoldBackFirstVideoFrame).SafeGetBool(false);
-		VideoRender.Renderer->EnableHoldbackOfFirstRenderableVideoFrame(bHoldBackFirstVideoFrame);
+		VideoOutput.OutputHandler->SetClock(RenderClock);
 	}
-	if (AudioRender.Renderer)
+	if (AudioOutput.OutputHandler)
 	{
-		AudioRender.Renderer->SetRenderClock(RenderClock);
-		AudioRender.Renderer->SetPlaybackRate(CurrentRate, IntendedRate, bCurrentlyPaused);
+		AudioOutput.OutputHandler->SetClock(RenderClock);
+		AudioOutput.OutputHandler->SetPlaybackRate(CurrentRate, IntendedRate, bCurrentlyPaused);
 	}
 
 	return 0;
@@ -1155,8 +1140,8 @@ void FAdaptiveStreamingPlayer::DestroyRenderers()
 	// Decoders must be gone already
 	check(!VideoDecoder.Decoder);
 	check(!AudioDecoder.Decoder);
-	AudioRender.Close();
-	VideoRender.Close();
+	AudioOutput.Close();
+	VideoOutput.Close();
 }
 
 
@@ -1790,13 +1775,9 @@ void FAdaptiveStreamingPlayer::HandlePlayStateChanges()
 	const double CurrentRate = PlaybackState.GetCurrentPlayRate();
 	const double IntendedRate = PlaybackState.GetDesiredPlayRate();
 	const bool bCurrentlyPaused = PlaybackState.GetIsPaused();
-	if (VideoRender.Renderer)
+	if (AudioOutput.OutputHandler)
 	{
-		VideoRender.Renderer->SetPlaybackRate(CurrentRate, IntendedRate, bCurrentlyPaused);
-	}
-	if (AudioRender.Renderer)
-	{
-		AudioRender.Renderer->SetPlaybackRate(CurrentRate, IntendedRate, bCurrentlyPaused);
+		AudioOutput.OutputHandler->SetPlaybackRate(CurrentRate, IntendedRate, bCurrentlyPaused);
 	}
 }
 
@@ -2980,16 +2961,8 @@ void FAdaptiveStreamingPlayer::InternalHandlePendingFirstSegmentRequest(const FT
 					DispatchEvent(FMetricEvent::ReportJumpInPlayPosition(NewPlayPos, CurrentPlayPos, Metrics::ETimeJumpReason::UserSeek));
 				}
 
-				if (AudioRender.Renderer)
-				{
-					RenderClock->SetCurrentTime(IMediaRenderClock::ERendererType::Audio, NewPlayPos);
-					AudioRender.Renderer->SetNextApproximatePresentationTime(NewPlayPos);
-				}
-				if (VideoRender.Renderer)
-				{
-					RenderClock->SetCurrentTime(IMediaRenderClock::ERendererType::Video, NewPlayPos);
-					VideoRender.Renderer->SetNextApproximatePresentationTime(NewPlayPos);
-				}
+				RenderClock->SetCurrentTime(IMediaRenderClock::ERendererType::Audio, NewPlayPos);
+				RenderClock->SetCurrentTime(IMediaRenderClock::ERendererType::Video, NewPlayPos);
 
 				PlaybackState.SetIsBuffering(true);
 				PlaybackState.SetIsSeeking(true);
@@ -3572,7 +3545,7 @@ void FAdaptiveStreamingPlayer::InternalDeselectStream(EStreamType StreamType)
 			bIsVideoDeselected = true;
 			SelectedStreamAttributesVid.Deselect();
 			VideoDecoder.Flush();
-			VideoRender.Flush(false);
+			VideoOutput.Flush(false);
 			break;
 		}
 		case EStreamType::Audio:
@@ -3580,7 +3553,7 @@ void FAdaptiveStreamingPlayer::InternalDeselectStream(EStreamType StreamType)
 			bIsAudioDeselected = true;
 			SelectedStreamAttributesAud.Deselect();
 			AudioDecoder.Flush();
-			AudioRender.Flush();
+			AudioOutput.Flush();
 			break;
 		}
 		case EStreamType::Subtitle:
@@ -4498,8 +4471,8 @@ void FAdaptiveStreamingPlayer::InternalStop(bool bHoldCurrentFrame)
 	// Flush the renderers once before flushing the decoders.
 	// In case the media samples hold references to the decoder (because of shared frame resources) it could be a possibility
 	// that the decoder cannot flush as long as those references are "out there".
-	AudioRender.Flush();
-	VideoRender.Flush(bHoldCurrentFrame);
+	AudioOutput.Flush();
+	VideoOutput.Flush(bHoldCurrentFrame);
 
 	// Flush the decoders.
 	AudioDecoder.Flush();
@@ -4509,8 +4482,8 @@ void FAdaptiveStreamingPlayer::InternalStop(bool bHoldCurrentFrame)
 	ClearAllDecoderEODs();
 
 	// Flush the renderers again, this time to discard everything the decoders may have emitted while being flushed.
-	AudioRender.Flush();
-	VideoRender.Flush(bHoldCurrentFrame);
+	AudioOutput.Flush();
+	VideoOutput.Flush(bHoldCurrentFrame);
 
 	// Flush dynamic events from the AEMS handler.
 	AEMSEventHandler->FlushDynamic();
@@ -4848,55 +4821,55 @@ void FAdaptiveStreamingPlayer::ABRGetStreamBufferStats(IAdaptiveStreamSelector::
 	switch(ForStream)
 	{
 		case EStreamType::Video:
+		{
 			OutBufferStats.PlayableContentDuration = VideoBufferStats.StreamBuffer.PlayableDuration;
 			OutBufferStats.bReachedEnd = VideoBufferStats.StreamBuffer.bEndOfData;
 			OutBufferStats.bEndOfTrack = VideoBufferStats.StreamBuffer.bEndOfTrack;
 			break;
+		}
 		case EStreamType::Audio:
+		{
 			OutBufferStats.PlayableContentDuration = AudioBufferStats.StreamBuffer.PlayableDuration;
 			OutBufferStats.bReachedEnd = AudioBufferStats.StreamBuffer.bEndOfData;
 			OutBufferStats.bEndOfTrack = AudioBufferStats.StreamBuffer.bEndOfTrack;
 			break;
+		}
 		case EStreamType::Subtitle:
+		{
 			OutBufferStats.PlayableContentDuration = TextBufferStats.StreamBuffer.PlayableDuration;
 			OutBufferStats.bReachedEnd = TextBufferStats.StreamBuffer.bEndOfData;
 			OutBufferStats.bEndOfTrack = TextBufferStats.StreamBuffer.bEndOfTrack;
 			break;
+		}
 		default:
+		{
 			break;
+		}
 	}
 	DiagnosticsCriticalSection.Unlock();
-	if (ForStream == EStreamType::Video && VideoRender.Renderer.IsValid())
+	if (ForStream == EStreamType::Video && VideoOutput.OutputHandler.IsValid())
 	{
-		OutBufferStats.PlayableContentDuration += VideoRender.Renderer->GetEnqueuedSampleDuration();
+		OutBufferStats.PlayableContentDuration += FTimeValue().SetFromTimespan(VideoOutput.OutputHandler->GetEnqueuedSampleDuration());
 	}
-	else if (ForStream == EStreamType::Audio && AudioRender.Renderer.IsValid())
+	else if (ForStream == EStreamType::Audio && AudioOutput.OutputHandler.IsValid())
 	{
-		OutBufferStats.PlayableContentDuration += AudioRender.Renderer->GetEnqueuedSampleDuration();
+		OutBufferStats.PlayableContentDuration += FTimeValue().SetFromTimespan(AudioOutput.OutputHandler->GetEnqueuedSampleDuration());
 	}
 }
 void FAdaptiveStreamingPlayer::ABRSetRenderRateScale(double InRenderRateScale)
 {
 	RenderRateScale = InRenderRateScale;
-	if (VideoRender.Renderer.IsValid())
+	if (AudioOutput.OutputHandler.IsValid())
 	{
-		VideoRender.Renderer->SetPlayRateScale(InRenderRateScale);
-	}
-	if (AudioRender.Renderer.IsValid())
-	{
-		AudioRender.Renderer->SetPlayRateScale(InRenderRateScale);
+		AudioOutput.OutputHandler->SetPlayRateScale(InRenderRateScale);
 	}
 }
 FTimeRange FAdaptiveStreamingPlayer::ABRGetSupportedRenderRateScale()
 {
 	FTimeRange Range;
-	if (AudioRender.Renderer.IsValid())
-	{
-		Range = AudioRender.Renderer->GetSupportedRenderRateScale();
-	}
-	else if (VideoRender.Renderer.IsValid())
+	if (AudioOutput.OutputHandler.IsValid())
 	{
-		Range = VideoRender.Renderer->GetSupportedRenderRateScale();
+		Range = AudioOutput.OutputHandler->GetSupportedRenderRateScale();
 	}
 	return Range;
 }
@@ -4941,14 +4914,14 @@ void FAdaptiveStreamingPlayer::DebugPrint(void* pPlayer, void (*pPrintFN)(void*
 		pD = &VideoBufferStats.StreamBuffer;
 		pPrintFN(pPlayer, "Video buffer  : EOT %d; EOS %d; %3u AUs; %8lld bytes in; %#7.4fs", pD->bEndOfTrack, pD->bEndOfData, (uint32)pD->NumCurrentAccessUnits, (long long int)pD->CurrentMemInUse, pD->PlayableDuration.GetAsSeconds());
 		pPrintFN(pPlayer, "Video decoder : %2u in decoder, %zu total; EOD in %d; EOD out %d; %s", (uint32)VideoBufferStats.DecoderOutputBuffer.NumElementsInDecoder, (uint32)VideoBufferStats.DecoderOutputBuffer.OutputBufferPoolSize, VideoBufferStats.DecoderInputBuffer.bEODReached, VideoBufferStats.DecoderOutputBuffer.bEODreached, VideoBufferStats.DecoderOutputBuffer.bOutputStalled?"stalled":"not stalled");
-		pPrintFN(pPlayer, "Video renderer: %#7.4fs enqueued", VideoRender.Renderer.IsValid() ? VideoRender.Renderer->GetEnqueuedSampleDuration().GetAsSeconds() : 0.0);
+		pPrintFN(pPlayer, "Video renderer: %#7.4fs enqueued", VideoOutput.OutputHandler.IsValid() ? VideoOutput.OutputHandler->GetEnqueuedSampleDuration().GetTotalSeconds() : 0.0);
 	}
 	if(1)
 	{
 		pD = &AudioBufferStats.StreamBuffer;
 		pPrintFN(pPlayer, "Audio buffer  : EOT %d; EOS %d; %3u AUs; %8lld bytes in; %#7.4fs", pD->bEndOfTrack, pD->bEndOfData, (uint32)pD->NumCurrentAccessUnits, (long long int)pD->CurrentMemInUse, pD->PlayableDuration.GetAsSeconds());
 		pPrintFN(pPlayer, "Audio decoder : %2u in decoder, %2u total; EOD in %d; EOD out %d; %s", (uint32)AudioBufferStats.DecoderOutputBuffer.NumElementsInDecoder, (uint32)AudioBufferStats.DecoderOutputBuffer.OutputBufferPoolSize, AudioBufferStats.DecoderInputBuffer.bEODReached, AudioBufferStats.DecoderOutputBuffer.bEODreached, AudioBufferStats.DecoderOutputBuffer.bOutputStalled?"stalled":"not stalled");
-		pPrintFN(pPlayer, "Audio renderer: %#7.4fs enqueued", AudioRender.Renderer.IsValid() ? AudioRender.Renderer->GetEnqueuedSampleDuration().GetAsSeconds() : 0.0);
+		pPrintFN(pPlayer, "Audio renderer: %#7.4fs enqueued", AudioOutput.OutputHandler.IsValid() ? AudioOutput.OutputHandler->GetEnqueuedSampleDuration().GetTotalSeconds() : 0.0);
 	}
 	if(1)
 	{
diff --git a/Engine/Plugins/Media/ElectraPlayer/Source/ElectraPlayerRuntime/Private/Runtime/Player/AdaptiveStreamingPlayer.h b/Engine/Plugins/Media/ElectraPlayer/Source/ElectraPlayerRuntime/Private/Runtime/Player/AdaptiveStreamingPlayer.h
--- a/Engine/Plugins/Media/ElectraPlayer/Source/ElectraPlayerRuntime/Private/Runtime/Player/AdaptiveStreamingPlayer.h
+++ b/Engine/Plugins/Media/ElectraPlayer/Source/ElectraPlayerRuntime/Private/Runtime/Player/AdaptiveStreamingPlayer.h
@@ -17,7 +17,8 @@
 
 namespace Electra
 {
-class IMediaRenderer;
+class FOutputHandlerVideo;
+class FOutputHandlerAudio;
 
 
 /**
@@ -28,8 +29,9 @@ class IAdaptiveStreamingPlayer : public TMediaNoncopyable<IAdaptiveStreamingPlay
 public:
 	struct FCreateParam
 	{
-		TSharedPtr<IMediaRenderer, ESPMode::ThreadSafe> VideoRenderer;
-		TSharedPtr<IMediaRenderer, ESPMode::ThreadSafe> AudioRenderer;
+		TSharedPtr<FOutputHandlerVideo, ESPMode::ThreadSafe> VideoOutputHandler;
+		TSharedPtr<FOutputHandlerAudio, ESPMode::ThreadSafe> AudioOutputHandler;
+
 		FGuid ExternalPlayerGUID;
 		enum class EWorkerThreads
 		{
diff --git a/Engine/Plugins/Media/ElectraPlayer/Source/ElectraPlayerRuntime/Private/Runtime/Player/AdaptiveStreamingPlayerDecoders.cpp b/Engine/Plugins/Media/ElectraPlayer/Source/ElectraPlayerRuntime/Private/Runtime/Player/AdaptiveStreamingPlayerDecoders.cpp
--- a/Engine/Plugins/Media/ElectraPlayer/Source/ElectraPlayerRuntime/Private/Runtime/Player/AdaptiveStreamingPlayerDecoders.cpp
+++ b/Engine/Plugins/Media/ElectraPlayer/Source/ElectraPlayerRuntime/Private/Runtime/Player/AdaptiveStreamingPlayerDecoders.cpp
@@ -288,8 +288,8 @@ int32 FAdaptiveStreamingPlayer::CreateDecoder(EStreamType type)
 					// Attach video decoder buffer monitor.
 					VideoDecoder.Decoder->SetAUInputBufferListener(&VideoDecoder);
 					VideoDecoder.Decoder->SetReadyBufferListener(&VideoDecoder);
-					// Have the video decoder send its output to the video renderer
-					VideoDecoder.Decoder->SetRenderer(VideoRender.Renderer);
+					// Where the video decoder shall send its output to
+					VideoDecoder.Decoder->SetOutputHandler(VideoOutput.OutputHandler);
 					// Open the decoder after having set all listeners.
 					VideoDecoder.Decoder->Open(VideoDecoder.LastSentAUCodecData, PlayerOptions.GetDictionary(), &HighestStream);
 				}
@@ -340,8 +340,8 @@ int32 FAdaptiveStreamingPlayer::CreateDecoder(EStreamType type)
 					// Attach buffer monitors.
 					AudioDecoder.Decoder->SetAUInputBufferListener(&AudioDecoder);
 					AudioDecoder.Decoder->SetReadyBufferListener(&AudioDecoder);
-					// Have to audio decoder send its output to the audio renderer
-					AudioDecoder.Decoder->SetRenderer(AudioRender.Renderer);
+					// Where the audio decoder shall send its output to
+					AudioDecoder.Decoder->SetOutputHandler(AudioOutput.OutputHandler);
 					// Open the decoder after having set all listeners.
 					AudioDecoder.Decoder->Open(AudioDecoder.LastSentAUCodecData);
 				}
@@ -634,13 +634,17 @@ void FAdaptiveStreamingPlayer::FeedDecoder(EStreamType Type, IAccessUnitBufferIn
 			if (!bEODSet && !bEOTSet && FromMultistreamBuffer->Num() == 0 && PushedDuration >= MinPushedDuration)
 			{
 				int32 NumEnqueuedSamples = 0;
-				if (Type == EStreamType::Video && VideoRender.Renderer.IsValid())
+				if (Type == EStreamType::Video && VideoOutput.OutputHandler.IsValid())
 				{
-					NumEnqueuedSamples = VideoRender.Renderer->GetNumEnqueuedSamples(nullptr);
+					TArray<IOutputHandlerBase::FEnqueuedSampleInfo> sampleInfos;
+					VideoOutput.OutputHandler->GetEnqueuedSampleInfo(sampleInfos);
+					NumEnqueuedSamples = sampleInfos.Num();
 				}
-				else if (Type == EStreamType::Audio && AudioRender.Renderer.IsValid())
+				else if (Type == EStreamType::Audio && AudioOutput.OutputHandler.IsValid())
 				{
-					NumEnqueuedSamples = AudioRender.Renderer->GetNumEnqueuedSamples(nullptr);
+					TArray<IOutputHandlerBase::FEnqueuedSampleInfo> sampleInfos;
+					AudioOutput.OutputHandler->GetEnqueuedSampleInfo(sampleInfos);
+					NumEnqueuedSamples = sampleInfos.Num();
 				}
 				if (NumEnqueuedSamples <= 1)
 				{
diff --git a/Engine/Plugins/Media/ElectraPlayer/Source/ElectraPlayerRuntime/Private/Runtime/Player/AdaptiveStreamingPlayerInternal.h b/Engine/Plugins/Media/ElectraPlayer/Source/ElectraPlayerRuntime/Private/Runtime/Player/AdaptiveStreamingPlayerInternal.h
--- a/Engine/Plugins/Media/ElectraPlayer/Source/ElectraPlayerRuntime/Private/Runtime/Player/AdaptiveStreamingPlayerInternal.h
+++ b/Engine/Plugins/Media/ElectraPlayer/Source/ElectraPlayerRuntime/Private/Runtime/Player/AdaptiveStreamingPlayerInternal.h
@@ -10,12 +10,12 @@
 #include "HTTP/HTTPResponseCache.h"
 #include "Player/PlayerSessionServices.h"
 #include "SynchronizedClock.h"
+#include "PlayerClock.h"
 
 #include "Demuxer/ParserISO14496-12.h"
 #include "Decoder/SubtitleDecoder.h"
 #include "Decoder/AudioDecoder.h"
 #include "Decoder/VideoDecoder.h"
-
 #include "Player/AdaptiveStreamingPlayerResourceRequest.h"
 #include "Player/PlayerStreamReader.h"
 #include "Player/PlaylistReader.h"
@@ -24,6 +24,7 @@
 #include "Player/AdaptiveStreamingPlayerABR.h"
 #include "Player/AdaptivePlayerOptionKeynames.h"
 #include "Player/ContentSteeringHandler.h"
+#include "OutputHandler.h"
 
 #include "Templates/UniqueObj.h"
 
@@ -65,29 +66,6 @@ namespace Electra
 class FAdaptiveStreamingPlayer;
 
 
-class IAdaptiveStreamingWrappedRenderer : public IMediaRenderer
-{
-public:
-	struct FEnqueuedSampleInfo
-	{
-		FTimeValue PTS;
-		FTimeValue Duration;
-	};
-
-	virtual ~IAdaptiveStreamingWrappedRenderer() = default;
-	virtual FTimeValue GetEnqueuedSampleDuration() = 0;
-	virtual int32 GetNumEnqueuedSamples(TArray<FEnqueuedSampleInfo>* OutOptionalSampleInfos) = 0;
-
-	virtual void AlwaysEmitSamplesWhenPaused(bool bEmitAlways) = 0;
-	virtual void SetPlaybackRate(double InCurrentPlaybackRate, double InIntendedPlaybackRate, bool bInCurrentlyPaused) = 0;
-	virtual void EnableHoldbackOfFirstRenderableVideoFrame(bool bEnableHoldback) = 0;
-
-	virtual FTimeRange GetSupportedRenderRateScale() = 0;
-	virtual void SetPlayRateScale(double InNewScale) = 0;
-	virtual double GetPlayRateScale() = 0;
-};
-
-
 class FMediaRenderClock : public IMediaRenderClock
 {
 public:
@@ -1322,38 +1300,36 @@ private:
 		TrackswitchUnderrun
 	};
 
-	struct FVideoRenderer
+	struct FVideoOutput
 	{
 		void Close()
 		{
-			Renderer.Reset();
+			OutputHandler.Reset();
 		}
 		void Flush(bool bHoldCurrentFrame)
 		{
-			if (Renderer.IsValid())
+			if (OutputHandler.IsValid())
 			{
-				FParamDict NoOptions;
-				Renderer->Flush(NoOptions);
+				OutputHandler->Flush();
 			}
 		}
-		TSharedPtr<IAdaptiveStreamingWrappedRenderer, ESPMode::ThreadSafe>	Renderer;
+		TSharedPtr<FOutputHandlerVideo, ESPMode::ThreadSafe> OutputHandler;
 	};
 
-	struct FAudioRenderer
+	struct FAudioOutput
 	{
 		void Close()
 		{
-			Renderer.Reset();
+			OutputHandler.Reset();
 		}
 		void Flush()
 		{
-			if (Renderer.IsValid())
+			if (OutputHandler.IsValid())
 			{
-				FParamDict options;
-				Renderer->Flush(options);
+				OutputHandler->Flush();
 			}
 		}
-		TSharedPtr<IAdaptiveStreamingWrappedRenderer, ESPMode::ThreadSafe>	Renderer;
+		TSharedPtr<FOutputHandlerAudio, ESPMode::ThreadSafe> OutputHandler;
 	};
 
 	struct FVideoDecoder : public IAccessUnitBufferListener, public IDecoderOutputBufferListener
@@ -2261,7 +2237,6 @@ private:
 
 	void UpdateDataAvailabilityState(Metrics::FDataAvailabilityChange& DataAvailabilityState, Metrics::FDataAvailabilityChange::EAvailability NewAvailability);
 
-	TSharedPtr<IAdaptiveStreamingWrappedRenderer, ESPMode::ThreadSafe> CreateWrappedRenderer(TSharedPtr<IMediaRenderer, ESPMode::ThreadSafe> RendererToWrap, EStreamType InType);
 	void StartRendering();
 	void StopRendering();
 
@@ -2386,8 +2361,8 @@ private:
 	uint32																CurrentPlaybackSequenceID[4]; // 0=video, 1=audio, 2=subtitiles, 3=UNSUPPORTED
 	FTimeRange															CurrentSegmentDownloadTimeRange[4];
 
-	FVideoRenderer														VideoRender;
-	FAudioRenderer														AudioRender;
+	FVideoOutput														VideoOutput;
+	FAudioOutput														AudioOutput;
 	FVideoDecoder														VideoDecoder;
 	FAudioDecoder														AudioDecoder;
 	FSubtitleDecoder													SubtitleDecoder;
diff --git a/Engine/Plugins/Media/ElectraPlayer/Source/ElectraPlayerRuntime/Private/Runtime/Player/AdaptiveStreamingPlayerManifest.cpp b/Engine/Plugins/Media/ElectraPlayer/Source/ElectraPlayerRuntime/Private/Runtime/Player/AdaptiveStreamingPlayerManifest.cpp
--- a/Engine/Plugins/Media/ElectraPlayer/Source/ElectraPlayerRuntime/Private/Runtime/Player/AdaptiveStreamingPlayerManifest.cpp
+++ b/Engine/Plugins/Media/ElectraPlayer/Source/ElectraPlayerRuntime/Private/Runtime/Player/AdaptiveStreamingPlayerManifest.cpp
@@ -28,6 +28,7 @@ static const FString MIMETypeDASH(TEXT("application/dash+xml"));
 static const FString MIMETypeMKV(TEXT("video/x-matroska"));
 static const FString MIMETypeMKA(TEXT("audio/x-matroska"));
 static const FString MIMETypeMPEGAudio(TEXT("audio/mpeg"));
+static const FString MIMETypeMPEGAudio4(TEXT("audio/aac"));
 
 
 
@@ -83,6 +84,7 @@ FString GetMIMETypeForURL(const FString& URL)
 		static const FString kTextWEBM(TEXT("webm"));
 		static const FString kTextMPA(TEXT("mpa"));
 		static const FString kTextMP3(TEXT("mp3"));
+		static const FString kTextAAC(TEXT("aac"));
 		if (LowerCaseExtension == kTextMP4 || LowerCaseExtension == kTextMP4V)
 		{
 			MimeType = MIMETypeMP4;
@@ -115,6 +117,10 @@ FString GetMIMETypeForURL(const FString& URL)
 		{
 			MimeType = MIMETypeMPEGAudio;
 		}
+		else if (LowerCaseExtension == kTextAAC)
+		{
+			MimeType = MIMETypeMPEGAudio4;
+		}
 	}
 
 	return MimeType;
@@ -238,7 +244,7 @@ void FAdaptiveStreamingPlayer::InternalLoadManifest(const FString& InURL, const
 				ManifestReader = IPlaylistReaderMKV::Create(this);
 				ManifestType = EMediaFormatType::MKV;
 			}
-			else if (mimeType == Playlist::MIMETypeMPEGAudio)
+			else if (mimeType == Playlist::MIMETypeMPEGAudio || mimeType == Playlist::MIMETypeMPEGAudio4)
 			{
 				ManifestReader = IPlaylistReaderMPEGAudio::Create(this);
 				ManifestType = EMediaFormatType::MPEGAudio;
@@ -266,10 +272,6 @@ void FAdaptiveStreamingPlayer::InternalLoadManifest(const FString& InURL, const
 			PostError(err);
 		}
 	}
-	else
-	{
-// TODO: Error: Not idle
-	}
 }
 
 
@@ -345,7 +347,7 @@ bool FAdaptiveStreamingPlayer::SelectManifest()
 			PlayerConfig.RebufferMinTimeAvailBeforePlayback = Utils::Min(minBufTimeMPD, PlayerConfig.RebufferMinTimeAvailBeforePlayback);
 
 			// For an mp4 or mkv stream we can now get rid of the manifest reader. It is no longer needed and we don't need to have it linger.
-			if (ManifestType == EMediaFormatType::ISOBMFF || ManifestType == EMediaFormatType::MKV)
+			if (ManifestType == EMediaFormatType::ISOBMFF || ManifestType == EMediaFormatType::MKV || ManifestType == EMediaFormatType::MPEGAudio)
 			{
 				InternalCloseManifestReader();
 			}
diff --git a/Engine/Plugins/Media/ElectraPlayer/Source/ElectraPlayerRuntime/Private/Runtime/Player/StreamSegmentReaderCommon.cpp b/Engine/Plugins/Media/ElectraPlayer/Source/ElectraPlayerRuntime/Private/Runtime/Player/StreamSegmentReaderCommon.cpp
--- a/Engine/Plugins/Media/ElectraPlayer/Source/ElectraPlayerRuntime/Private/Runtime/Player/StreamSegmentReaderCommon.cpp
+++ b/Engine/Plugins/Media/ElectraPlayer/Source/ElectraPlayerRuntime/Private/Runtime/Player/StreamSegmentReaderCommon.cpp
@@ -23,11 +23,11 @@
 #include "Utilities/TimeUtilities.h"
 #include "Utilities/UtilsMP4.h"
 #include "Utilities/UtilsMPEG.h"
+#include "Utilities/ElectraBitstream.h"
 #include "Utils/Google/ElectraUtilsVPxVideo.h"
 #include "Utils/MPEG/ElectraUtilsMPEGVideo_H264.h"
 #include "Utils/MPEG/ElectraUtilsMPEGVideo_H265.h"
 #include "Utils/MPEG/ElectraUtilsMPEGAudio.h"
-#include "Utils/ElectraBitstreamReader.h"
 #include "Misc/TVariant.h"
 #include "BufferedDataReader.h"
 
@@ -777,10 +777,10 @@ void FStreamSegmentReaderCommon::FStreamHandler::HandleRequest()
 			};
 			// ISO/IEC 14496-12 ?
 			const uint32 boxName = GetUINT32BE(ProbeBytes + 4);
-			if (boxName == UtilsMP4::Make4CC('s','t','y','p') || boxName == UtilsMP4::Make4CC('m','o','o','f') || boxName == UtilsMP4::Make4CC('s','i','d','x')
-				|| boxName == UtilsMP4::Make4CC('f','t','y','p') || boxName == UtilsMP4::Make4CC('m','o','o','v') || boxName == UtilsMP4::Make4CC('e','m','s','g')
-				|| boxName == UtilsMP4::Make4CC('f','r','e','e') || boxName == UtilsMP4::Make4CC('s','k','i','p') || boxName == UtilsMP4::Make4CC('s','s','i','x')
-				|| boxName == UtilsMP4::Make4CC('p','r','f','t') || boxName == UtilsMP4::Make4CC('u','d','t','a'))
+			if (boxName == Electra::UtilitiesMP4::MakeBoxAtom('s','t','y','p') || boxName == Electra::UtilitiesMP4::MakeBoxAtom('m','o','o','f') || boxName == Electra::UtilitiesMP4::MakeBoxAtom('s','i','d','x')
+				|| boxName == Electra::UtilitiesMP4::MakeBoxAtom('f','t','y','p') || boxName == Electra::UtilitiesMP4::MakeBoxAtom('m','o','o','v') || boxName == Electra::UtilitiesMP4::MakeBoxAtom('e','m','s','g')
+				|| boxName == Electra::UtilitiesMP4::MakeBoxAtom('f','r','e','e') || boxName == Electra::UtilitiesMP4::MakeBoxAtom('s','k','i','p') || boxName == Electra::UtilitiesMP4::MakeBoxAtom('s','s','i','x')
+				|| boxName == Electra::UtilitiesMP4::MakeBoxAtom('p','r','f','t') || boxName == Electra::UtilitiesMP4::MakeBoxAtom('u','d','t','a'))
 			{
 				ContainerFormat = FStreamSegmentRequestCommon::EContainerFormat::ISO_14496_12;
 			}
@@ -1078,8 +1078,8 @@ bool FStreamSegmentReaderCommon::FStreamHandler::FetchInitSegment(TSharedPtrTS<F
 		};
 		// ISO/IEC 14496-12 ?
 		const uint32 boxName = InitDataSize >= 8 ? GetUINT32BE(InitDataPtr + 4) : 0;
-		if (boxName == UtilsMP4::Make4CC('f','t','y','p') || boxName == UtilsMP4::Make4CC('f','r','e','e') || boxName == UtilsMP4::Make4CC('s','k','i','p')
-			|| boxName == UtilsMP4::Make4CC('p','d','i','n') || boxName == UtilsMP4::Make4CC('u','d','t','a'))
+		if (boxName == Electra::UtilitiesMP4::MakeBoxAtom('f','t','y','p') || boxName == Electra::UtilitiesMP4::MakeBoxAtom('f','r','e','e') || boxName == Electra::UtilitiesMP4::MakeBoxAtom('s','k','i','p')
+			|| boxName == Electra::UtilitiesMP4::MakeBoxAtom('p','d','i','n') || boxName == Electra::UtilitiesMP4::MakeBoxAtom('u','d','t','a'))
 		{
 			StaticSegmentDataReader.SetSourceBuffer(ResponseBuffer);
 			ContainerFormat = FStreamSegmentRequestCommon::EContainerFormat::ISO_14496_12;
@@ -1115,8 +1115,8 @@ bool FStreamSegmentReaderCommon::FStreamHandler::FetchInitSegment(TSharedPtrTS<F
 		// For self-initializing media we need to locate the 'moov' box.
 		UtilsMP4::FMP4RootBoxLocator BoxLocator;
 		TArray<UtilsMP4::FMP4RootBoxLocator::FBoxInfo> Boxes;
-		const TArray<uint32> FirstBox { UtilsMP4::Make4CC('f','t','y','p') };
-		const TArray<uint32> MoovBox { UtilsMP4::Make4CC('m','o','o','v') };
+		const TArray<uint32> FirstBox { Electra::UtilitiesMP4::MakeBoxAtom('f','t','y','p') };
+		const TArray<uint32> MoovBox { Electra::UtilitiesMP4::MakeBoxAtom('m','o','o','v') };
 		bool bSuccess = BoxLocator.LocateRootBoxes(Boxes, PlayerSessionService->GetHTTPManager(), InitURLToUse.Url.URL, FirstBox, MoovBox, MoovBox, UtilsMP4::FMP4RootBoxLocator::FCancellationCheckDelegate::CreateLambda([&]()
 		{
 			return HasReadBeenAborted();
@@ -1137,7 +1137,7 @@ bool FStreamSegmentReaderCommon::FStreamHandler::FetchInitSegment(TSharedPtrTS<F
 			SetError(FString::Printf(TEXT("Self-initializing media 'moov' scan error: %s"),	*ErrMsg), INTERNAL_SEG_ERROR_INIT_SEGMENT_DOWNLOAD_ERROR);
 			return false;
 		}
-		const UtilsMP4::FMP4RootBoxLocator::FBoxInfo* Moov = Boxes.FindByPredicate([InType=UtilsMP4::Make4CC('m','o','o','v')](const UtilsMP4::FMP4RootBoxLocator::FBoxInfo& InBox){return InBox.Type == InType;});
+		const UtilsMP4::FMP4RootBoxLocator::FBoxInfo* Moov = Boxes.FindByPredicate([InType=Electra::UtilitiesMP4::MakeBoxAtom('m','o','o','v')](const UtilsMP4::FMP4RootBoxLocator::FBoxInfo& InBox){return InBox.Type == InType;});
 		if (Moov)
 		{
 			StaticSegmentDataReader.SetSourceBuffer(Moov->DataBuffer);
@@ -3155,7 +3155,7 @@ FStreamSegmentReaderCommon::FStreamHandler::EHandleResult FStreamSegmentReaderCo
 						return EHandleResult::Failed;
 					}
 
-					ElectraDecodersUtil::FElectraBitstreamReader br(ProbeBytes, kNumProbeBytesNeeded);
+					FBitstreamReader br(ProbeBytes, kNumProbeBytesNeeded);
 					br.SkipBits(12);	// sync word, already checked for.
 					const uint32 mpeg_version = br.GetBits(1);
 					const uint32 layer = br.GetBits(2);
diff --git a/Engine/Plugins/Media/ElectraPlayer/Source/ElectraPlayerRuntime/Private/Runtime/Player/mp4/PlaylistReaderMP4.cpp b/Engine/Plugins/Media/ElectraPlayer/Source/ElectraPlayerRuntime/Private/Runtime/Player/mp4/PlaylistReaderMP4.cpp
--- a/Engine/Plugins/Media/ElectraPlayer/Source/ElectraPlayerRuntime/Private/Runtime/Player/mp4/PlaylistReaderMP4.cpp
+++ b/Engine/Plugins/Media/ElectraPlayer/Source/ElectraPlayerRuntime/Private/Runtime/Player/mp4/PlaylistReaderMP4.cpp
@@ -227,9 +227,9 @@ void FPlaylistReaderMP4::WorkerThread()
 	SCOPE_CYCLE_COUNTER(STAT_ElectraPlayer_MP4_PlaylistWorker);
 	CSV_SCOPED_TIMING_STAT(ElectraPlayer, MP4_PlaylistWorker);
 
-	const TArray<uint32> StopAfterBoxes { UtilsMP4::Make4CC('m','o','o','v') };
-	const TArray<uint32> ReadBoxes { UtilsMP4::Make4CC('f','t','y','p'), UtilsMP4::Make4CC('m','o','o','v') };
-	const TArray<uint32> FirstBoxes { UtilsMP4::Make4CC('f','t','y','p'), UtilsMP4::Make4CC('s','t','y','p'), UtilsMP4::Make4CC('s','i','d','x'), UtilsMP4::Make4CC('f','r','e','e'), UtilsMP4::Make4CC('s','k','i','p') };
+	const TArray<uint32> StopAfterBoxes { Electra::UtilitiesMP4::MakeBoxAtom('m','o','o','v') };
+	const TArray<uint32> ReadBoxes { Electra::UtilitiesMP4::MakeBoxAtom('f','t','y','p'), Electra::UtilitiesMP4::MakeBoxAtom('m','o','o','v') };
+	const TArray<uint32> FirstBoxes { Electra::UtilitiesMP4::MakeBoxAtom('f','t','y','p'), Electra::UtilitiesMP4::MakeBoxAtom('s','t','y','p'), Electra::UtilitiesMP4::MakeBoxAtom('s','i','d','x'), Electra::UtilitiesMP4::MakeBoxAtom('f','r','e','e'), Electra::UtilitiesMP4::MakeBoxAtom('s','k','i','p') };
 	UtilsMP4::FMP4RootBoxLocator BoxLocator;
 	TArray<UtilsMP4::FMP4RootBoxLocator::FBoxInfo> BoxInfos;
 	bool bGotBoxes = BoxLocator.LocateRootBoxes(BoxInfos, PlayerSessionServices->GetHTTPManager(), MainPlaylistURL, FirstBoxes, StopAfterBoxes, ReadBoxes, UtilsMP4::FMP4RootBoxLocator::FCancellationCheckDelegate::CreateLambda([&]()
@@ -261,10 +261,10 @@ void FPlaylistReaderMP4::WorkerThread()
 	else if (bGotBoxes && !bAbort && !bHasErrored)
 	{
 		// Do we have the `ftyp` and `moov` boxes?
-		if (BoxInfos.ContainsByPredicate([InType=UtilsMP4::Make4CC('f','t','y','p')](const UtilsMP4::FMP4RootBoxLocator::FBoxInfo& InBox){return InBox.Type == InType;}) &&
-			BoxInfos.ContainsByPredicate([InType=UtilsMP4::Make4CC('m','o','o','v')](const UtilsMP4::FMP4RootBoxLocator::FBoxInfo& InBox){return InBox.Type == InType;}))
+		if (BoxInfos.ContainsByPredicate([InType=Electra::UtilitiesMP4::MakeBoxAtom('f','t','y','p')](const UtilsMP4::FMP4RootBoxLocator::FBoxInfo& InBox){return InBox.Type == InType;}) &&
+			BoxInfos.ContainsByPredicate([InType=Electra::UtilitiesMP4::MakeBoxAtom('m','o','o','v')](const UtilsMP4::FMP4RootBoxLocator::FBoxInfo& InBox){return InBox.Type == InType;}))
 		{
-			if (BoxInfos.ContainsByPredicate([InType=UtilsMP4::Make4CC('m','d','a','t')](const UtilsMP4::FMP4RootBoxLocator::FBoxInfo& InBox){return InBox.Type == InType;}))
+			if (BoxInfos.ContainsByPredicate([InType=Electra::UtilitiesMP4::MakeBoxAtom('m','d','a','t')](const UtilsMP4::FMP4RootBoxLocator::FBoxInfo& InBox){return InBox.Type == InType;}))
 			{
 				LogMessage(IInfoLog::ELevel::Info, FString::Printf(TEXT("The mp4 at \"%s\" is not fast-startable. Consider moving the 'moov' box in front of the 'mdat' for faster startup times."), *ConnectionInfo.EffectiveURL));
 			}
@@ -272,14 +272,14 @@ void FPlaylistReaderMP4::WorkerThread()
 			MP4Parser = IParserISO14496_12::CreateParser();
 
 			// First parse the ftyp on which some elements in the moov depend.
-			ParseBuffer = BoxInfos.FindByPredicate([InType=UtilsMP4::Make4CC('f','t','y','p')](const UtilsMP4::FMP4RootBoxLocator::FBoxInfo& InBox){return InBox.Type == InType;})->DataBuffer;
+			ParseBuffer = BoxInfos.FindByPredicate([InType=Electra::UtilitiesMP4::MakeBoxAtom('f','t','y','p')](const UtilsMP4::FMP4RootBoxLocator::FBoxInfo& InBox){return InBox.Type == InType;})->DataBuffer;
 			ParsePos = 0;
 			ParseBufferSize = ParseBuffer->Num();
 			UEMediaError parseError = MP4Parser->ParseHeader(this, this, PlayerSessionServices, nullptr);
 			// If successful parse the moov
 			if (parseError == UEMEDIA_ERROR_OK || parseError == UEMEDIA_ERROR_END_OF_STREAM)
 			{
-				ParseBuffer = BoxInfos.FindByPredicate([InType=UtilsMP4::Make4CC('m','o','o','v')](const UtilsMP4::FMP4RootBoxLocator::FBoxInfo& InBox){return InBox.Type == InType;})->DataBuffer;
+				ParseBuffer = BoxInfos.FindByPredicate([InType=Electra::UtilitiesMP4::MakeBoxAtom('m','o','o','v')](const UtilsMP4::FMP4RootBoxLocator::FBoxInfo& InBox){return InBox.Type == InType;})->DataBuffer;
 				ParsePos = 0;
 				ParseBufferSize = ParseBuffer->Num();
 				parseError = MP4Parser->ParseHeader(this, this, PlayerSessionServices, nullptr);
diff --git a/Engine/Plugins/Media/ElectraPlayer/Source/ElectraPlayerRuntime/Private/Runtime/Player/mpegaudio/ManifestMPEGAudio.cpp b/Engine/Plugins/Media/ElectraPlayer/Source/ElectraPlayerRuntime/Private/Runtime/Player/mpegaudio/ManifestMPEGAudio.cpp
--- a/Engine/Plugins/Media/ElectraPlayer/Source/ElectraPlayerRuntime/Private/Runtime/Player/mpegaudio/ManifestMPEGAudio.cpp
+++ b/Engine/Plugins/Media/ElectraPlayer/Source/ElectraPlayerRuntime/Private/Runtime/Player/mpegaudio/ManifestMPEGAudio.cpp
@@ -818,6 +818,7 @@ IManifest::FResult FManifestMPEGAudioInternal::FTimelineAssetMPEGAudio::GetStart
 			req->MPEGHeaderMask = MPEGInfoHeader.MPEGHeaderMask;
 			req->MPEGHeaderExpectedValue = MPEGInfoHeader.MPEGHeaderExpectedValue;
 			req->CBRFrameSize = MPEGInfoHeader.CBRFrameSize;
+			req->bIsAAC = MPEGInfoHeader.Type == FMPEGAudioInfoHeader::EType::ISO_14496_3;
 			req->bIsVBR = MPEGInfoHeader.bIsVBR;
 			req->bIsLive = true;
 			req->CodecInfo = MPEGInfoHeader.CodecInfo;
@@ -960,6 +961,7 @@ IManifest::FResult FManifestMPEGAudioInternal::FTimelineAssetMPEGAudio::GetStart
 				req->MPEGHeaderMask = MPEGInfoHeader.MPEGHeaderMask;
 				req->MPEGHeaderExpectedValue = MPEGInfoHeader.MPEGHeaderExpectedValue;
 				req->CBRFrameSize = MPEGInfoHeader.CBRFrameSize;
+				req->bIsAAC = MPEGInfoHeader.Type == FMPEGAudioInfoHeader::EType::ISO_14496_3;
 				req->bIsVBR = MPEGInfoHeader.bIsVBR;
 				req->bIsLive = false;
 				req->Duration = kSegmentDuration;
@@ -1101,7 +1103,7 @@ FErrorDetail FManifestMPEGAudioInternal::FRepresentationMPEGAudio::CreateFrom(co
 	{
 		Name = TEXT("Unknown");
 	}
-	
+
 	// Language?
 	if (InHeader.ID3v2.IsValid() && InHeader.ID3v2->HaveTag(Utils::Make4CC('T','L','A','N')))
 	{
diff --git a/Engine/Plugins/Media/ElectraPlayer/Source/ElectraPlayerRuntime/Private/Runtime/Player/mpegaudio/ManifestMPEGAudio.h b/Engine/Plugins/Media/ElectraPlayer/Source/ElectraPlayerRuntime/Private/Runtime/Player/mpegaudio/ManifestMPEGAudio.h
--- a/Engine/Plugins/Media/ElectraPlayer/Source/ElectraPlayerRuntime/Private/Runtime/Player/mpegaudio/ManifestMPEGAudio.h
+++ b/Engine/Plugins/Media/ElectraPlayer/Source/ElectraPlayerRuntime/Private/Runtime/Player/mpegaudio/ManifestMPEGAudio.h
@@ -19,7 +19,14 @@ namespace Electra
 
 struct FMPEGAudioInfoHeader
 {
+	enum class EType
+	{
+		Unknown,
+		ISO_11172_3,		// MPEG1/2 layer 1,2 and 3
+		ISO_14496_3			// AAC
+	};
 	TSharedPtrTS<MPEG::FID3V2Metadata> ID3v2;
+	EType Type = EType::Unknown;
 	uint32 MPEGHeaderMask = 0;
 	uint32 MPEGHeaderExpectedValue = 0;
 	int32 MPEGVersion = 0;
diff --git a/Engine/Plugins/Media/ElectraPlayer/Source/ElectraPlayerRuntime/Private/Runtime/Player/mpegaudio/PlaylistReaderMPEGAudio.cpp b/Engine/Plugins/Media/ElectraPlayer/Source/ElectraPlayerRuntime/Private/Runtime/Player/mpegaudio/PlaylistReaderMPEGAudio.cpp
--- a/Engine/Plugins/Media/ElectraPlayer/Source/ElectraPlayerRuntime/Private/Runtime/Player/mpegaudio/PlaylistReaderMPEGAudio.cpp
+++ b/Engine/Plugins/Media/ElectraPlayer/Source/ElectraPlayerRuntime/Private/Runtime/Player/mpegaudio/PlaylistReaderMPEGAudio.cpp
@@ -19,8 +19,8 @@
 #include "ElectraPlayerPrivate.h"
 #include "Player/mpegaudio/ManifestMPEGAudio.h"
 #include "Player/mpegaudio/OptionKeynamesMPEGAudio.h"
+#include "Utilities/ElectraBitstream.h"
 #include "Utils/MPEG/ElectraUtilsMPEGAudio.h"
-#include "Utils/ElectraBitstreamReader.h"
 
 #define ERRCODE_MPEGAUDIO_INVALID_FILE		1
 #define ERRCODE_MPEGAUDIO_DOWNLOAD_ERROR	2
@@ -303,7 +303,7 @@ void FPlaylistReaderMPEGAudio::FRequest::FindSyncMarkers()
 			{
 				continue;
 			}
-			// Check for validity
+			// Check for ISO 11172-3 header
 			if (Base+2 < End && (Base[1] & 0xe0) == 0xe0		// sync marker (11 1-bits)
 				&& ((Base[1] >> 3) & 3) >= 2					// audio version 1 or 2 (2.5 not supported)
 				&& ((Base[1] >> 1) & 3) != 0					// layer index 1, 2 or 3
@@ -313,6 +313,15 @@ void FPlaylistReaderMPEGAudio::FRequest::FindSyncMarkers()
 				SyncMarkerOffsets.Add(Base - Buffer.GetData());
 				Base += 2;
 			}
+			// Check for ISO 14496-3 ADTS header
+			else if (Base+2 < End && (Base[1] & 0xf0) == 0xf0		// sync marker (12 1-bits)
+				&& ((Base[1] & 0x8) == 0)							// MPEG version must be MPEG-4
+				&& ((Base[1] & 0x6) == 0)							// layer must be 0
+				&& ((Base[2] >> 2) & 7) < 13)						// sampling frequency index not a reserved value
+			{
+				SyncMarkerOffsets.Add(Base - Buffer.GetData());
+				Base += 2;
+			}
 		}
 		CurrentSyncMarkerCheckPos = Buffer.Num() - 1;
 	}
@@ -376,8 +385,6 @@ bool FPlaylistReaderMPEGAudio::FRequest::Validate(FMPEGAudioInfoHeader& InOutInf
 		return false;
 	}
 
-	InOutInfoHeader.MPEGHeaderMask = 0xfffe0c00;
-
 	const uint8* const HeaderData = Buffer.GetData();
 
 	// Check if there is an ID3 header.
@@ -408,7 +415,7 @@ bool FPlaylistReaderMPEGAudio::FRequest::Validate(FMPEGAudioInfoHeader& InOutInf
 			MPEG::FID3V2Metadata::FItem MLLTBlob;
 			InOutInfoHeader.ID3v2->GetTag(MLLTBlob, Utils::Make4CC('M','L','L','T'));
 			TArray<uint8> MLLT = MLLTBlob.Value.GetValue<TArray<uint8>>();
-			ElectraDecodersUtil::FElectraBitstreamReader br(MLLT.GetData(), MLLT.Num());
+			FBitstreamReader br(MLLT.GetData(), MLLT.Num());
 			InOutInfoHeader.MLLT = MakeSharedTS<FMPEGAudioInfoHeader::FMLLT>();
 			InOutInfoHeader.MLLT->FramesBetweenReferences = br.GetBits(16);
 			uint32 BytesBetweenReference = br.GetBits(24);
@@ -479,7 +486,7 @@ bool FPlaylistReaderMPEGAudio::FRequest::Validate(FMPEGAudioInfoHeader& InOutInf
 				const int32 Version = ElectraDecodersUtil::MPEG::UtilsMPEG123::GetVersion(HeaderValue);
 				const int32 NumChannels = ElectraDecodersUtil::MPEG::UtilsMPEG123::GetChannelCount(HeaderValue);
 				InfoHeaderOffset = SyncMarkerOffsets[nSyncMarker];
-				InOutInfoHeader.MPEGHeaderExpectedValue = HeaderValue & InOutInfoHeader.MPEGHeaderMask;
+				InOutInfoHeader.MPEGHeaderExpectedValue = HeaderValue & 0xfffe0c00;
 				InOutInfoHeader.MPEGVersion = Version;
 				InOutInfoHeader.MPEGLayer = Layer;
 				InOutInfoHeader.SampleRate = ElectraDecodersUtil::MPEG::UtilsMPEG123::GetSamplingRate(HeaderValue);
@@ -529,7 +536,7 @@ bool FPlaylistReaderMPEGAudio::FRequest::Validate(FMPEGAudioInfoHeader& InOutInf
 					// Additional info added with version 3.90
 					if (VersionMajor > 3 || (VersionMajor == 3 && VersionMinor >= 90))
 					{
-						ElectraDecodersUtil::FElectraBitstreamReader br(HeaderData + Off, 208, 9, 0);
+						FBitstreamReader br(HeaderData + Off, 208, 9, 0);
 						uint32 InfoTagRevision = br.GetBits(4);			// 15=reserved
 						uint32 VBRMethod = br.GetBits(4);				// 0=unknown, 1=CBR, 2=ABR, 3-6=VBR method 1-4, 8=CBR 2 pass, 9=ABR 2 pass, 15=reserved
 						uint32 LowPassFreq = br.GetBits(8) * 100;		// 0=unknown
@@ -590,7 +597,7 @@ bool FPlaylistReaderMPEGAudio::FRequest::Validate(FMPEGAudioInfoHeader& InOutInf
 			// Test VBRI header
 			else if (IsVBRIHeader(HeaderValue, HeaderData, Off, HeaderName))
 			{
-				ElectraDecodersUtil::FElectraBitstreamReader br(HeaderData + Off + 4, Utils::Max(0, ProbeSize-Off));
+				FBitstreamReader br(HeaderData + Off + 4, Utils::Max(0, ProbeSize-Off));
 				int32 VersionId = (int32)br.GetBits(16);
 				int32 Delay = (int32)br.GetBits(16);
 				int32 QualityIndicator = (int32)br.GetBits(16);
@@ -650,54 +657,65 @@ bool FPlaylistReaderMPEGAudio::FRequest::Validate(FMPEGAudioInfoHeader& InOutInf
 		}
 	}
 
+
+
 	// Iterate all the headers that resemble an mpeg audio header and count those that are identical
 	// except for slight differences (like padding, private, copyright bit etc.)
-	TMap<uint32, int32> HeaderMap;
+	TMap<uint32, int32> HeaderMap1, HeaderMap2;
 	for(int32 i=0; i<SyncMarkerOffsets.Num()-2; ++i)
 	{
 		int32 Off = SyncMarkerOffsets[i];
 		uint32 HeaderValue = GetUINT32BE(HeaderData + Off);
-		HeaderValue &= InOutInfoHeader.MPEGHeaderMask;
-		++HeaderMap.FindOrAdd(HeaderValue, 0);
+		++HeaderMap1.FindOrAdd(HeaderValue & 0xfffe0c00, 0);
+		++HeaderMap2.FindOrAdd(HeaderValue & 0xfffefd00, 0);
 	}
-	HeaderMap.ValueSort([](const int32& e1, const int32& e2) { return e1 > e2; });
-
-	uint32 ExpectedHeaderValue = HeaderMap.Num() ? HeaderMap.CreateConstIterator().Key() : 0xfffffffU;
-	// Possible mismatch between the header value found in the info packet above and the majority vote?
-	if (InOutInfoHeader.MPEGHeaderExpectedValue && InOutInfoHeader.MPEGHeaderExpectedValue != ExpectedHeaderValue)
+	HeaderMap1.ValueSort([](const int32& e1, const int32& e2) { return e1 > e2; });
+	HeaderMap2.ValueSort([](const int32& e1, const int32& e2) { return e1 > e2; });
+
+	int32 FirstHeaderOffset1 = InfoHeaderOffset;
+	int32 FirstHeaderOffset2 = InfoHeaderOffset;
+	bool bProbe11172_3 = false;
+	bool bProbe14496_3 = false;
+	ElectraDecodersUtil::MPEG::AACUtils::ADTSheader FirstADTSHeader;
+	// Check ISO 11172-3
+	if (HeaderMap1.Num())
 	{
-		OutWarnings.Emplace(FString::Printf(TEXT("Mismatching MPEG header bitmask")));
-	}
-	InOutInfoHeader.MPEGHeaderExpectedValue = ExpectedHeaderValue;
-
-	// Locate the first matching header.
-	int32 FirstHeaderOffset = InfoHeaderOffset;
-	for(int32 i=0; FirstHeaderOffset<0 && i<SyncMarkerOffsets.Num()-2; ++i)
-	{
-		const int32 Off = SyncMarkerOffsets[i];
-		if ((GetUINT32BE(HeaderData + Off) & InOutInfoHeader.MPEGHeaderMask) == InOutInfoHeader.MPEGHeaderExpectedValue)
+		InOutInfoHeader.MPEGHeaderMask = 0xfffe0c00;
+		if (!InOutInfoHeader.MPEGHeaderExpectedValue)
 		{
-			FirstHeaderOffset = Off;
-			break;
+			InOutInfoHeader.MPEGHeaderExpectedValue = HeaderMap1.CreateConstIterator().Key() & InOutInfoHeader.MPEGHeaderMask;
 		}
-	}
-	// Perform a quick sanity check that the blocks are really where they should be.
-	bool bProbeOk = true;
-	if (FirstHeaderOffset >= 0)
-	{
+		// Locate the first matching header.
+		FirstHeaderOffset1 = InfoHeaderOffset;
+		for(int32 i=0; FirstHeaderOffset1<0 && i<SyncMarkerOffsets.Num()-2; ++i)
+		{
+			const int32 Off = SyncMarkerOffsets[i];
+			if ((GetUINT32BE(HeaderData + Off) & InOutInfoHeader.MPEGHeaderMask) == InOutInfoHeader.MPEGHeaderExpectedValue)
+			{
+				FirstHeaderOffset1 = Off;
+				break;
+			}
+		}
+
 		const int32 LastHeaderOffset = SyncMarkerOffsets.Num() >= 2 ? SyncMarkerOffsets[SyncMarkerOffsets.Num()-2] : 0;
+		bProbe11172_3 = true;
+		int32 Off = FirstHeaderOffset1;
 		// Probing 10 headers should be fine
-		int32 Off = FirstHeaderOffset;
 		for(int32 i=0; i<10; ++i)
 		{
 			const uint32 HeaderValue = GetUINT32BE(HeaderData + Off);
 			int32 FrameSize = ElectraDecodersUtil::MPEG::UtilsMPEG123::GetFrameSize(HeaderValue);
-			if (Off + FrameSize + 4 <= LastHeaderOffset)
+			if (FrameSize <= 0)
+			{
+				bProbe11172_3 = false;
+				break;
+			}
+			else if (Off + FrameSize + 4 <= LastHeaderOffset)
 			{
 				const uint32 HeaderValue2 = GetUINT32BE(HeaderData + Off + FrameSize);
-				if ((HeaderValue2 & InOutInfoHeader.MPEGHeaderMask) != InOutInfoHeader.MPEGHeaderExpectedValue)
+				if ((HeaderValue2 & InOutInfoHeader.MPEGHeaderMask) != (HeaderValue & InOutInfoHeader.MPEGHeaderMask))
 				{
-					bProbeOk = false;
+					bProbe11172_3 = false;
 					break;
 				}
 				Off += FrameSize;
@@ -708,23 +726,69 @@ bool FPlaylistReaderMPEGAudio::FRequest::Validate(FMPEGAudioInfoHeader& InOutInf
 			}
 		}
 	}
-	// Fill in values that haven't been set up by any of the special headers yet.
-	if (bProbeOk && FirstHeaderOffset >= 0)
+
+	// Check for ISO 14496-3 if 11172-3 was no success.
+	if (!bProbe11172_3 && HeaderMap2.Num())
 	{
-		const uint32 HeaderValue = GetUINT32BE(HeaderData + FirstHeaderOffset);
-		InOutInfoHeader.MPEGVersion = InOutInfoHeader.MPEGVersion ? InOutInfoHeader.MPEGVersion : ElectraDecodersUtil::MPEG::UtilsMPEG123::GetVersion(HeaderValue);
-		InOutInfoHeader.MPEGLayer = InOutInfoHeader.MPEGLayer ? InOutInfoHeader.MPEGLayer : ElectraDecodersUtil::MPEG::UtilsMPEG123::GetLayer(HeaderValue);
-		InOutInfoHeader.SampleRate = InOutInfoHeader.SampleRate ? InOutInfoHeader.SampleRate : ElectraDecodersUtil::MPEG::UtilsMPEG123::GetSamplingRate(HeaderValue);
-		InOutInfoHeader.NumChannels = InOutInfoHeader.NumChannels ? InOutInfoHeader.NumChannels : ElectraDecodersUtil::MPEG::UtilsMPEG123::GetChannelCount(HeaderValue);
-		InOutInfoHeader.Bitrate = InOutInfoHeader.Bitrate ? InOutInfoHeader.Bitrate : ElectraDecodersUtil::MPEG::UtilsMPEG123::GetBitrate(HeaderValue);
-		InOutInfoHeader.CBRFrameSize = InOutInfoHeader.bIsVBR ? 0 : InOutInfoHeader.CBRFrameSize ? InOutInfoHeader.CBRFrameSize : ElectraDecodersUtil::MPEG::UtilsMPEG123::GetFrameSize(HeaderValue, 1);
-		InOutInfoHeader.SamplesPerFrame = InOutInfoHeader.SamplesPerFrame ? InOutInfoHeader.SamplesPerFrame : ElectraDecodersUtil::MPEG::UtilsMPEG123::GetSamplesPerFrame(HeaderValue);
-		InOutInfoHeader.FirstDataByte = InOutInfoHeader.FirstDataByte >= 0 ? InOutInfoHeader.FirstDataByte : FirstHeaderOffset;
-		InOutInfoHeader.LastDataByte = InOutInfoHeader.LastDataByte >= 0 ? InOutInfoHeader.LastDataByte : -1;
+		InOutInfoHeader.MPEGHeaderMask = 0xfffefd00;
+		InOutInfoHeader.MPEGHeaderExpectedValue = HeaderMap2.CreateConstIterator().Key() & InOutInfoHeader.MPEGHeaderMask;
+		// Locate the first matching header.
+		FirstHeaderOffset2 = InfoHeaderOffset;
+		for(int32 i=0; FirstHeaderOffset2<0 && i<SyncMarkerOffsets.Num()-2; ++i)
+		{
+			const int32 Off = SyncMarkerOffsets[i];
+			if ((GetUINT32BE(HeaderData + Off) & InOutInfoHeader.MPEGHeaderMask) == InOutInfoHeader.MPEGHeaderExpectedValue)
+			{
+				FirstHeaderOffset2 = Off;
+				break;
+			}
+		}
+
+		const int32 LastHeaderOffset = SyncMarkerOffsets.Num() >= 2 ? SyncMarkerOffsets[SyncMarkerOffsets.Num()-2] : 0;
+		bProbe14496_3 = true;
+		int32 Off = FirstHeaderOffset2;
+		// Probing 10 headers should be fine
+		for(int32 i=0; i<10 && LastHeaderOffset > Off; ++i)
+		{
+			ElectraDecodersUtil::MPEG::AACUtils::ADTSheader adts;
+			if (!ElectraDecodersUtil::MPEG::AACUtils::ParseADTSHeader(adts, MakeConstArrayView(HeaderData + Off, LastHeaderOffset - Off)))
+			{
+				bProbe14496_3 = false;
+				break;
+			}
+			else if (i == 0)
+			{
+				FirstADTSHeader = adts;
+			}
+			Off += adts.FrameLength;
+		}
+	}
+
+	if (bProbe11172_3 && bProbe14496_3)
+	{
+		OutError = FString::Printf(TEXT("Data parses as both MPEG 1 or 2 and AAC. Cannot determine precise format!"));
+		return false;
 	}
 
-	if (bProbeOk)
+
+	if (bProbe11172_3)
 	{
+		InOutInfoHeader.Type = FMPEGAudioInfoHeader::EType::ISO_11172_3;
+		// Fill in values that haven't been set up by any of the special headers yet.
+		if (FirstHeaderOffset1 >= 0)
+		{
+			const uint32 HeaderValue = GetUINT32BE(HeaderData + FirstHeaderOffset1);
+			InOutInfoHeader.MPEGVersion = InOutInfoHeader.MPEGVersion ? InOutInfoHeader.MPEGVersion : ElectraDecodersUtil::MPEG::UtilsMPEG123::GetVersion(HeaderValue);
+			InOutInfoHeader.MPEGLayer = InOutInfoHeader.MPEGLayer ? InOutInfoHeader.MPEGLayer : ElectraDecodersUtil::MPEG::UtilsMPEG123::GetLayer(HeaderValue);
+			InOutInfoHeader.SampleRate = InOutInfoHeader.SampleRate ? InOutInfoHeader.SampleRate : ElectraDecodersUtil::MPEG::UtilsMPEG123::GetSamplingRate(HeaderValue);
+			InOutInfoHeader.NumChannels = InOutInfoHeader.NumChannels ? InOutInfoHeader.NumChannels : ElectraDecodersUtil::MPEG::UtilsMPEG123::GetChannelCount(HeaderValue);
+			InOutInfoHeader.Bitrate = InOutInfoHeader.Bitrate ? InOutInfoHeader.Bitrate : ElectraDecodersUtil::MPEG::UtilsMPEG123::GetBitrate(HeaderValue);
+			InOutInfoHeader.CBRFrameSize = InOutInfoHeader.bIsVBR ? 0 : InOutInfoHeader.CBRFrameSize ? InOutInfoHeader.CBRFrameSize : ElectraDecodersUtil::MPEG::UtilsMPEG123::GetFrameSize(HeaderValue, 1);
+			InOutInfoHeader.SamplesPerFrame = InOutInfoHeader.SamplesPerFrame ? InOutInfoHeader.SamplesPerFrame : ElectraDecodersUtil::MPEG::UtilsMPEG123::GetSamplesPerFrame(HeaderValue);
+			InOutInfoHeader.FirstDataByte = InOutInfoHeader.FirstDataByte >= 0 ? InOutInfoHeader.FirstDataByte : FirstHeaderOffset1;
+			InOutInfoHeader.LastDataByte = InOutInfoHeader.LastDataByte >= 0 ? InOutInfoHeader.LastDataByte : -1;
+		}
+
 		InOutInfoHeader.CodecInfo.SetStreamType(EStreamType::Audio);
 		InOutInfoHeader.CodecInfo.SetMimeType(TEXT("audio/mpeg"));
 		InOutInfoHeader.CodecInfo.SetCodec(FStreamCodecInformation::ECodec::Audio4CC);
@@ -735,7 +799,38 @@ bool FPlaylistReaderMPEGAudio::FRequest::Validate(FMPEGAudioInfoHeader& InOutInf
 		InOutInfoHeader.CodecInfo.SetSamplingRate(InOutInfoHeader.SampleRate);
 		InOutInfoHeader.CodecInfo.SetNumberOfChannels(InOutInfoHeader.NumChannels);
 		InOutInfoHeader.CodecInfo.SetBitrate(InOutInfoHeader.AverageBitrate ? InOutInfoHeader.AverageBitrate : InOutInfoHeader.Bitrate ? InOutInfoHeader.Bitrate : 0);
+	}
+	else if (bProbe14496_3)
+	{
+		InOutInfoHeader.Type = FMPEGAudioInfoHeader::EType::ISO_14496_3;
+		InOutInfoHeader.bIsVBR = true;
+		InOutInfoHeader.SampleRate = ElectraDecodersUtil::MPEG::AACUtils::GetSampleRateFromFrequenceIndex(FirstADTSHeader.SamplingFrequencyIndex);
+		InOutInfoHeader.NumChannels = ElectraDecodersUtil::MPEG::AACUtils::GetNumberOfChannelsFromChannelConfiguration(FirstADTSHeader.ChannelConfiguration);
+		InOutInfoHeader.AverageBitrate = 64000;		// we need some value to guesstimate the duration....
+		InOutInfoHeader.SamplesPerFrame = 1024;
+		InOutInfoHeader.FirstDataByte = InOutInfoHeader.FirstDataByte >= 0 ? InOutInfoHeader.FirstDataByte : FirstHeaderOffset2;
+		InOutInfoHeader.LastDataByte = InOutInfoHeader.LastDataByte >= 0 ? InOutInfoHeader.LastDataByte : -1;
+
+		InOutInfoHeader.CodecInfo.SetStreamType(EStreamType::Audio);
+		InOutInfoHeader.CodecInfo.SetMimeType(TEXT("audio/aac"));
+		InOutInfoHeader.CodecInfo.SetCodec(FStreamCodecInformation::ECodec::AAC);
+		InOutInfoHeader.CodecInfo.SetCodecSpecifierRFC6381(TEXT("mp4a.40.2"));
+		InOutInfoHeader.CodecInfo.SetSamplingRate(InOutInfoHeader.SampleRate);
+		InOutInfoHeader.CodecInfo.SetNumberOfChannels(InOutInfoHeader.NumChannels);
+		InOutInfoHeader.CodecInfo.SetChannelConfiguration(FirstADTSHeader.ChannelConfiguration);
+		InOutInfoHeader.CodecInfo.SetBitrate(InOutInfoHeader.AverageBitrate ? InOutInfoHeader.AverageBitrate : InOutInfoHeader.Bitrate ? InOutInfoHeader.Bitrate : 0);
 
+		TArray<uint8> CSD;
+		CSD.SetNumUninitialized(2);
+		uint32 csd = (uint32)(FirstADTSHeader.Profile + 1) << 11;
+		csd |= (uint32)FirstADTSHeader.SamplingFrequencyIndex << 7;
+		csd |= (uint32)FirstADTSHeader.ChannelConfiguration << 3;
+		CSD[0] = (uint8)(csd >> 8);
+		CSD[1] = (uint8)(csd & 255);
+		InOutInfoHeader.CodecInfo.SetCodecSpecificData(CSD);
+	}
+	if (bProbe11172_3 || bProbe14496_3)
+	{
 		if (InOutInfoHeader.bIsLive)
 		{
 			InOutInfoHeader.EstimatedDuration.SetToPositiveInfinity();
@@ -794,7 +889,7 @@ bool FPlaylistReaderMPEGAudio::FRequest::Validate(FMPEGAudioInfoHeader& InOutInf
 		InOutInfoHeader.HTTPResponseHeaders = Headers;
 	}
 
-	return bProbeOk;
+	return bProbe11172_3 || bProbe14496_3;
 }
 
 
@@ -821,6 +916,7 @@ void FPlaylistReaderMPEGAudio::WorkerThread()
 		Filename.MidInline(7);	// remove "file://"
 		TSharedPtr<FArchive, ESPMode::ThreadSafe> Archive = MakeShareable(IFileManager::Get().CreateFileReader(*Filename));
 		ConnectionInfo.EffectiveURL = Filename;
+		// Set a content mime type. Does not correctly identify the type yet (might be `audio/aac`)
 		ConnectionInfo.ContentType = TEXT("audio/mpeg");
 		if (Archive.IsValid())
 		{
diff --git a/Engine/Plugins/Media/ElectraPlayer/Source/ElectraPlayerRuntime/Private/Runtime/Player/mpegaudio/StreamReaderMPEGAudio.cpp b/Engine/Plugins/Media/ElectraPlayer/Source/ElectraPlayerRuntime/Private/Runtime/Player/mpegaudio/StreamReaderMPEGAudio.cpp
--- a/Engine/Plugins/Media/ElectraPlayer/Source/ElectraPlayerRuntime/Private/Runtime/Player/mpegaudio/StreamReaderMPEGAudio.cpp
+++ b/Engine/Plugins/Media/ElectraPlayer/Source/ElectraPlayerRuntime/Private/Runtime/Player/mpegaudio/StreamReaderMPEGAudio.cpp
@@ -470,7 +470,7 @@ void FStreamReaderMPEGAudio::HandleRequest()
 				{
 					continue;
 				}
-				// Check for validity
+				// Check for validity. This works for MPEG1 and AAC audio.
 				uint32 hdr = GetUINT32BE(Base);
 				if ((hdr & FrameSyncMask) == FrameSyncValue)
 				{
@@ -504,22 +504,57 @@ void FStreamReaderMPEGAudio::HandleRequest()
 			{
 				const uint8* CheckBaseAddr = BufferBaseData + SyncMarkerOffsets[i];
 				bool bRunOk = true;
-				uint32 HeaderValue = GetUINT32BE(CheckBaseAddr);
-				int32 StartFrameSize = ElectraDecodersUtil::MPEG::UtilsMPEG123::GetFrameSize(HeaderValue);
-				for(int32 j=i+1; j<MaxCheckFrames; ++j)
+				int32 StartFrameSize = 0;
+				if (!Request->bIsAAC)
 				{
-					int32 FrameSize = ElectraDecodersUtil::MPEG::UtilsMPEG123::GetFrameSize(HeaderValue);
-					CheckBaseAddr += FrameSize;
-					if (CheckBaseAddr >= MaxCheckAddr)
+					uint32 HeaderValue = GetUINT32BE(CheckBaseAddr);
+					StartFrameSize = ElectraDecodersUtil::MPEG::UtilsMPEG123::GetFrameSize(HeaderValue);
+					for(int32 j=i+1; j<MaxCheckFrames; ++j)
+					{
+						int32 FrameSize = ElectraDecodersUtil::MPEG::UtilsMPEG123::GetFrameSize(HeaderValue);
+						CheckBaseAddr += FrameSize;
+						if (CheckBaseAddr >= MaxCheckAddr)
+						{
+							break;
+						}
+						HeaderValue = GetUINT32BE(CheckBaseAddr);
+						if ((HeaderValue & FrameSyncMask) != FrameSyncValue)
+						{
+							bRunOk = false;
+							break;
+						}
+					}
+				}
+				else
+				{
+					ElectraDecodersUtil::MPEG::AACUtils::ADTSheader adts;
+					if (CheckBaseAddr >= MaxCheckAddr - 12)
 					{
 						break;
 					}
-					HeaderValue = GetUINT32BE(CheckBaseAddr);
-					if ((HeaderValue & FrameSyncMask) != FrameSyncValue)
+					if (ElectraDecodersUtil::MPEG::AACUtils::ParseADTSHeader(adts, MakeConstArrayView(CheckBaseAddr, MaxCheckAddr - CheckBaseAddr)))
+					{
+						StartFrameSize = adts.FrameLength;
+						for(int32 j=i+1; j<MaxCheckFrames; ++j)
+						{
+							CheckBaseAddr += adts.FrameLength;
+							if (CheckBaseAddr >= MaxCheckAddr - 12)
+							{
+								break;
+							}
+							if (!ElectraDecodersUtil::MPEG::AACUtils::ParseADTSHeader(adts, MakeConstArrayView(CheckBaseAddr, MaxCheckAddr - CheckBaseAddr)))
+							{
+								bRunOk = false;
+								break;
+							}
+						}
+					}
+					else
 					{
 						bRunOk = false;
 						break;
 					}
+					CheckBaseAddr += adts.FrameLength;
 				}
 				if (bRunOk)
 				{
@@ -543,17 +578,19 @@ void FStreamReaderMPEGAudio::HandleRequest()
 	NextAUPTS.SetSequenceIndex(Request->TimestampSequenceIndex);
 	int32 NextAUBufferOffset = FirstAUBufferOffset;
 	int32 NextAUFrameSize = FirstAUFrameSize;
+	int32 NextHeaderSkipBytes = 0;
 	bool bIsFirstInSequence = true;
 	bool bReadPastLastPTS = false;
 	FTimeValue AUDuration;
 	TSharedPtrTS<FAccessUnit::CodecData> AUCodecData = MakeSharedTS<FAccessUnit::CodecData>();
 	AUCodecData->ParsedInfo = Request->CodecInfo;
+	AUCodecData->CodecSpecificData = Request->CodecInfo.GetCodecSpecificData();
 
 	TSharedPtrTS<FBufferSourceInfo> BufferSourceInfo = MakeSharedTS<FBufferSourceInfo>();
 	BufferSourceInfo->PeriodID = ds.MediaAssetID;
 	BufferSourceInfo->PeriodAdaptationSetID = ds.MediaAssetID + TEXT(".") + ds.AdaptationSetID;
 	BCP47::ParseRFC5646Tag(BufferSourceInfo->LanguageTag, FString(TEXT("und")));
-	BufferSourceInfo->Codec = TEXT("mp4a.6b");
+	BufferSourceInfo->Codec = Request->CodecInfo.GetCodecSpecifierRFC6381();
 	BufferSourceInfo->HardIndex = 0;
 	BufferSourceInfo->PlaybackSequenceID = Request->GetPlaybackSequenceID();
 
@@ -562,9 +599,10 @@ void FStreamReaderMPEGAudio::HandleRequest()
 	// amount of memory being moved around.
 	int32 NumLiveStreamBytesToRemove = Request->GetBitrate() * 4 / 8;
 
+	const int32 kNumNextHeaderBytes = Request->bIsAAC ? 7 : 4;
+	ElectraDecodersUtil::MPEG::AACUtils::ADTSheader adts;
 	while(!bDone && !HasErrored() && !HasBeenAborted() && !bTerminate)
 	{
-		check(NextAUFrameSize);
 		if (!NextAUFrameSize)
 		{
 			break;
@@ -577,8 +615,8 @@ void FStreamReaderMPEGAudio::HandleRequest()
 			NextAUBufferOffset = 0;
 		}
 
-		// Wait until we get the next AU's data plus the following 4 bytes that are the header of the following frame.
-		const int64 TotalNumNeeded = NextAUBufferOffset + NextAUFrameSize + 4;
+		// Wait until we get the next AU's data plus the following bytes that are the header of the following frame.
+		const int64 TotalNumNeeded = NextAUBufferOffset + NextAUFrameSize + kNumNextHeaderBytes;
 		if (ReceiveBuffer->WaitUntilSizeAvailable(TotalNumNeeded, 1000 * 20))
 		{
 			FScopeLock Lock(ReceiveBuffer->GetLock());
@@ -596,8 +634,7 @@ void FStreamReaderMPEGAudio::HandleRequest()
 			{
 				BufferBaseData += NextAUBufferOffset;
 				uint32 HeaderValue = GetUINT32BE(BufferBaseData);
-				// Safety check.
-				check((HeaderValue & FrameSyncMask) == FrameSyncValue);
+				// Safety check. Works for MPEG1 and AAC audio.
 				if ((HeaderValue & FrameSyncMask) != FrameSyncValue)
 				{
 					ParsingErrorMessage = FString::Printf(TEXT("Frame sync marker not found. Corrupt file?"));
@@ -606,11 +643,38 @@ void FStreamReaderMPEGAudio::HandleRequest()
 					bHasErrored = true;
 				}
 
+				if (Request->bIsAAC)
+				{
+					if (ElectraDecodersUtil::MPEG::AACUtils::ParseADTSHeader(adts, MakeConstArrayView(BufferBaseData, NextAUFrameSize)))
+					{
+						NextHeaderSkipBytes = adts.HeaderSize;
+					}
+					else
+					{
+						ParsingErrorMessage = FString::Printf(TEXT("Not a valid ADTS frame. Corrupt file?"));
+						LogMessage(IInfoLog::ELevel::Error, ParsingErrorMessage);
+						bDone = true;
+						bHasErrored = true;
+					}
+				}
+
 				if (!AUDuration.IsValid())
 				{
-					int32 FrameSize = ElectraDecodersUtil::MPEG::UtilsMPEG123::GetSamplesPerFrame(HeaderValue);
-					int32 SampleRate = ElectraDecodersUtil::MPEG::UtilsMPEG123::GetSamplingRate(HeaderValue);
-					AUDuration.SetFromND(FrameSize, (uint32)SampleRate);
+					if (!Request->bIsAAC)
+					{
+						int32 FrameSize = ElectraDecodersUtil::MPEG::UtilsMPEG123::GetSamplesPerFrame(HeaderValue);
+						int32 SampleRate = ElectraDecodersUtil::MPEG::UtilsMPEG123::GetSamplingRate(HeaderValue);
+						AUDuration.SetFromND(FrameSize, (uint32)SampleRate);
+					}
+					else
+					{
+						int32 FrameSize = 1024;
+						int32 SampleRate = ElectraDecodersUtil::MPEG::AACUtils::GetSampleRateFromFrequenceIndex(adts.SamplingFrequencyIndex);
+						if (SampleRate)
+						{
+							AUDuration.SetFromND(FrameSize, (uint32)SampleRate);
+						}
+					}
 				}
 
 				// Keep track of how many bytes we have consumed so far. This correlates the Live stream metadata to this time.
@@ -667,7 +731,7 @@ void FStreamReaderMPEGAudio::HandleRequest()
 					AccessUnit->EarliestPTS = Request->EarliestPTS;
 					AccessUnit->LatestPTS = Request->LastPTS;
 
-					AccessUnit->AUSize = (uint32) NextAUFrameSize;
+					AccessUnit->AUSize = (uint32) NextAUFrameSize - NextHeaderSkipBytes;
 					AccessUnit->AUCodecData = AUCodecData;
 					// Set the sequence index member and update all timestamps with it as well.
 					AccessUnit->SequenceIndex = Request->TimestampSequenceIndex;
@@ -686,7 +750,7 @@ void FStreamReaderMPEGAudio::HandleRequest()
 
 					bIsFirstInSequence = false;
 
-					FMemory::Memcpy(AccessUnit->AUData, BufferBaseData, NextAUFrameSize);
+					FMemory::Memcpy(AccessUnit->AUData, BufferBaseData + NextHeaderSkipBytes, NextAUFrameSize - NextHeaderSkipBytes);
 
 					// Unlock the receive buffer now so the reader won't be blocked.
 					Lock.Unlock();
@@ -749,12 +813,23 @@ void FStreamReaderMPEGAudio::HandleRequest()
 				NextAUBufferOffset += NextAUFrameSize;
 				NumGot -= NextAUFrameSize;
 				NextAUFrameSize = 0;
-				// Did we also get the next 4 bytes?
-				if (NumGot >= 4)
+				// Did we also get the next header bytes?
+				if (NumGot >= kNumNextHeaderBytes)
 				{
 					FScopeLock Lock2(ReceiveBuffer->GetLock());
-					HeaderValue = GetUINT32BE(ReceiveBuffer->GetLinearReadData() + NextAUBufferOffset + NextAUFrameSize);
-					NextAUFrameSize = ElectraDecodersUtil::MPEG::UtilsMPEG123::GetFrameSize(HeaderValue);
+					if (!Request->bIsAAC)
+					{
+						HeaderValue = GetUINT32BE(ReceiveBuffer->GetLinearReadData() + NextAUBufferOffset + NextAUFrameSize);
+						NextAUFrameSize = ElectraDecodersUtil::MPEG::UtilsMPEG123::GetFrameSize(HeaderValue);
+					}
+					else
+					{
+						if (ElectraDecodersUtil::MPEG::AACUtils::ParseADTSHeader(adts, MakeConstArrayView(ReceiveBuffer->GetLinearReadData() + NextAUBufferOffset + NextAUFrameSize, kNumNextHeaderBytes)))
+						{
+							NextAUFrameSize = adts.FrameLength;
+							NextHeaderSkipBytes = adts.HeaderSize;
+						}
+					}
 				}
 				else
 				{
diff --git a/Engine/Plugins/Media/ElectraPlayer/Source/ElectraPlayerRuntime/Private/Runtime/Player/mpegaudio/StreamReaderMPEGAudio.h b/Engine/Plugins/Media/ElectraPlayer/Source/ElectraPlayerRuntime/Private/Runtime/Player/mpegaudio/StreamReaderMPEGAudio.h
--- a/Engine/Plugins/Media/ElectraPlayer/Source/ElectraPlayerRuntime/Private/Runtime/Player/mpegaudio/StreamReaderMPEGAudio.h
+++ b/Engine/Plugins/Media/ElectraPlayer/Source/ElectraPlayerRuntime/Private/Runtime/Player/mpegaudio/StreamReaderMPEGAudio.h
@@ -70,6 +70,7 @@ public:
 	uint32 MPEGHeaderMask = 0;
 	uint32 MPEGHeaderExpectedValue = 0;
 	int32 CBRFrameSize = 0;
+	bool bIsAAC = false;
 	bool bIsVBR = false;
 	bool bIsLive = false;
 	bool bIsEOSRequest = false;
diff --git a/Engine/Plugins/Media/ElectraPlayer/Source/ElectraPlayerRuntime/Private/Runtime/PlayerRuntimeGlobal.cpp b/Engine/Plugins/Media/ElectraPlayer/Source/ElectraPlayerRuntime/Private/Runtime/PlayerRuntimeGlobal.cpp
--- a/Engine/Plugins/Media/ElectraPlayer/Source/ElectraPlayerRuntime/Private/Runtime/PlayerRuntimeGlobal.cpp
+++ b/Engine/Plugins/Media/ElectraPlayer/Source/ElectraPlayerRuntime/Private/Runtime/PlayerRuntimeGlobal.cpp
@@ -95,7 +95,14 @@ namespace Electra
 	 */
 	bool Startup(const Configuration& InConfiguration)
 	{
-		FMediaRunnable::Startup();
+		// Load the modules we depend on. They may have been loaded already, but we do it explicitly here to ensure that
+		// they will not be unloaded on shutdown before this module here, otherwise there could be crashes.
+		FModuleManager::Get().LoadModule(TEXT("ElectraBase"));
+		FModuleManager::Get().LoadModule(TEXT("ElectraSamples"));
+		FModuleManager::Get().LoadModule(TEXT("ElectraCodecFactory"));
+		FModuleManager::Get().LoadModule(TEXT("ElectraHTTPStream"));
+		FModuleManager::Get().LoadModule(TEXT("ElectraSubtitles"));
+		FModuleManager::Get().LoadModule(TEXT("ElectraCDM"));
 
 		if (!ApplicationTerminatingDelegate.IsValid())
 		{
@@ -110,15 +117,6 @@ namespace Electra
 			ApplicationResumeDelegate = FCoreDelegates::ApplicationHasEnteredForegroundDelegate.AddStatic(&HandleApplicationHasEnteredForeground);
 		}
 
-		// Load the modules we depend on. They may have been loaded already, but we do it explicitly here to ensure that
-		// they will not be unloaded on shutdown before this module here, otherwise there could be crashes.
-		FModuleManager::Get().LoadModule(TEXT("ElectraBase"));
-		FModuleManager::Get().LoadModule(TEXT("ElectraSamples"));
-		FModuleManager::Get().LoadModule(TEXT("ElectraCodecFactory"));
-		FModuleManager::Get().LoadModule(TEXT("ElectraHTTPStream"));
-		FModuleManager::Get().LoadModule(TEXT("ElectraSubtitles"));
-		FModuleManager::Get().LoadModule(TEXT("ElectraCDM"));
-
 		EnabledAnalyticsEvents = InConfiguration.EnabledAnalyticsEvents;
 
 #if ELECTRA_PRECREATE_HTTP_MANAGER
@@ -156,8 +154,6 @@ namespace Electra
 		{
 			FCoreDelegates::GetApplicationWillTerminateDelegate().Remove(ApplicationTerminatingDelegate);
 		}
-
-		FMediaRunnable::Shutdown();
 	}
 
 
