ElectraPlayer: Not flushing outputs when closing

#rnx

diff --git a/Engine/Plugins/Media/ElectraPlayer/Source/ElectraPlayerRuntime/Private/Runtime/OutputHandlerAudio.cpp b/Engine/Plugins/Media/ElectraPlayer/Source/ElectraPlayerRuntime/Private/Runtime/OutputHandlerAudio.cpp
--- a/Engine/Plugins/Media/ElectraPlayer/Source/ElectraPlayerRuntime/Private/Runtime/OutputHandlerAudio.cpp
+++ b/Engine/Plugins/Media/ElectraPlayer/Source/ElectraPlayerRuntime/Private/Runtime/OutputHandlerAudio.cpp
@@ -203,7 +203,10 @@ namespace Electra
 		PendingOutputSamplesLock.Unlock();
 
 		// Notify the registered sample flush delegate.
-		OutputQueueFlushSamplesDelegate().ExecuteIfBound();
+		if (!bIsDetached)
+		{
+			OutputQueueFlushSamplesDelegate().ExecuteIfBound();
+		}
 
 		// Anything that has not been returned by now we drop from the sample info stats.
 		EnqueuedSampleInfosLock.Lock();
diff --git a/Engine/Plugins/Media/ElectraPlayer/Source/ElectraPlayerRuntime/Private/Runtime/OutputHandlerVideo.cpp b/Engine/Plugins/Media/ElectraPlayer/Source/ElectraPlayerRuntime/Private/Runtime/OutputHandlerVideo.cpp
--- a/Engine/Plugins/Media/ElectraPlayer/Source/ElectraPlayerRuntime/Private/Runtime/OutputHandlerVideo.cpp
+++ b/Engine/Plugins/Media/ElectraPlayer/Source/ElectraPlayerRuntime/Private/Runtime/OutputHandlerVideo.cpp
@@ -68,7 +68,10 @@ namespace Electra
 		PendingOutputSamplesLock.Unlock();
 
 		// Notify the registered sample flush delegate.
-		OutputQueueFlushSamplesDelegate().ExecuteIfBound();
+		if (!bIsDetached)
+		{
+			OutputQueueFlushSamplesDelegate().ExecuteIfBound();
+		}
 
 		// Anything that has not been returned by now we drop from the sample info stats.
 		EnqueuedSampleInfosLock.Lock();
