Fixed static code analysis warning

#rnx

diff --git a/Engine/Plugins/Media/ElectraPlayer/Source/ElectraPlayerRuntime/Private/Runtime/Player/mp4/ManifestMP4.cpp b/Engine/Plugins/Media/ElectraPlayer/Source/ElectraPlayerRuntime/Private/Runtime/Player/mp4/ManifestMP4.cpp
--- a/Engine/Plugins/Media/ElectraPlayer/Source/ElectraPlayerRuntime/Private/Runtime/Player/mp4/ManifestMP4.cpp
+++ b/Engine/Plugins/Media/ElectraPlayer/Source/ElectraPlayerRuntime/Private/Runtime/Player/mp4/ManifestMP4.cpp
@@ -897,7 +897,7 @@ void FManifestMP4Internal::FTimelineAssetMP4::LimitSegmentDownloadSize(TSharedPt
 				}
 				int64 CurrentTrackOffset = LastTrackOffset;
 				bool bDurationLimitReached = TrackDurationLimit > 0 && TrackDur > TrackDurationLimit;
-				bool bSizeLimitReached = (MaxSegmentByteCount > 0) && (CurrentTrackOffset - StartOffset >= MaxSegmentByteCount);
+				bool bSizeLimitReached = /*(MaxSegmentByteCount > 0) &&*/ (CurrentTrackOffset - StartOffset >= MaxSegmentByteCount);
 				if (CurrentTrackOffset >= EndOffset || bDurationLimitReached || bSizeLimitReached)
 				{
 					// Limit reached.
