ElectraPlayer: Additional check in output handler in case the media sample queues are configured larger than the handler buffer size

#rnx

diff --git a/Engine/Plugins/Media/ElectraPlayer/Source/ElectraPlayerRuntime/Private/Runtime/OutputHandlerAudio.cpp b/Engine/Plugins/Media/ElectraPlayer/Source/ElectraPlayerRuntime/Private/Runtime/OutputHandlerAudio.cpp
--- a/Engine/Plugins/Media/ElectraPlayer/Source/ElectraPlayerRuntime/Private/Runtime/OutputHandlerAudio.cpp
+++ b/Engine/Plugins/Media/ElectraPlayer/Source/ElectraPlayerRuntime/Private/Runtime/OutputHandlerAudio.cpp
@@ -262,6 +262,10 @@ namespace Electra
 		PendingOutputSamplesLock.Lock();
 		int32 NumPendingOutputSamples = PendingOutputSamples.Num();
 		PendingOutputSamplesLock.Unlock();
+		if (1 + NumPendingOutputSamples > NumOutputSamples)
+		{
+			return false;
+		}
 		CanOutputQueueReceiveDelegate().ExecuteIfBound(bCanRcv, 1 + NumPendingOutputSamples);
 		return bCanRcv;
 	}
diff --git a/Engine/Plugins/Media/ElectraPlayer/Source/ElectraPlayerRuntime/Private/Runtime/OutputHandlerVideo.cpp b/Engine/Plugins/Media/ElectraPlayer/Source/ElectraPlayerRuntime/Private/Runtime/OutputHandlerVideo.cpp
--- a/Engine/Plugins/Media/ElectraPlayer/Source/ElectraPlayerRuntime/Private/Runtime/OutputHandlerVideo.cpp
+++ b/Engine/Plugins/Media/ElectraPlayer/Source/ElectraPlayerRuntime/Private/Runtime/OutputHandlerVideo.cpp
@@ -130,6 +130,10 @@ namespace Electra
 		PendingOutputSamplesLock.Lock();
 		int32 NumPendingOutputSamples = PendingOutputSamples.Num();
 		PendingOutputSamplesLock.Unlock();
+		if (1 + NumPendingOutputSamples > NumOutputSamples)
+		{
+			return false;
+		}
 		CanOutputQueueReceiveDelegate().ExecuteIfBound(bCanRcv, 1 + NumPendingOutputSamples);
 		return bCanRcv;
 	}
