ElectraPlayer: Handling of 'meta' boxes that are not a full box in an isobmff file.

#rnx

diff --git a/Engine/Plugins/Media/ElectraPlayer/Source/ElectraPlayerRuntime/Private/Runtime/Demuxer/ParserISO14496-12.cpp b/Engine/Plugins/Media/ElectraPlayer/Source/ElectraPlayerRuntime/Private/Runtime/Demuxer/ParserISO14496-12.cpp
--- a/Engine/Plugins/Media/ElectraPlayer/Source/ElectraPlayerRuntime/Private/Runtime/Demuxer/ParserISO14496-12.cpp
+++ b/Engine/Plugins/Media/ElectraPlayer/Source/ElectraPlayerRuntime/Private/Runtime/Demuxer/ParserISO14496-12.cpp
@@ -282,7 +282,7 @@ namespace Electra
 	class FDataBufferReader : public IGenericDataReader
 	{
 	public:
-		FDataBufferReader(const TArray<uint8>& InDataBufferToReadFrom) : DataBufferRef(InDataBufferToReadFrom)
+		FDataBufferReader(const TArray<uint8>& InDataBufferToReadFrom, int64 InBaseOffset=0) : DataBufferRef(InDataBufferToReadFrom), BaseOffset(InBaseOffset)
 		{}
 		virtual ~FDataBufferReader() = default;
 		int64 ReadData(void* IntoBuffer, int64 NumBytesToRead, int64 InFromOffset) override
@@ -304,13 +304,14 @@ namespace Electra
 		bool HasReadBeenAborted() const override
 		{ return false;	}
 		int64 GetCurrentOffset() const override
-		{ return CurrentOffset;	}
+		{ return CurrentOffset + BaseOffset; }
 		int64 GetTotalSize() const override
 		{ check(!"this should not be called"); return -1; }
 
 	private:
 		const TArray<uint8>& DataBufferRef;
 		int64 CurrentOffset = 0;
+		int64 BaseOffset = 0;
 	};
 
 
@@ -347,6 +348,17 @@ namespace Electra
 			return BoxReader;
 		}
 
+		void PushReader(FMP4BoxReader* InNewReader)
+		{
+			BoxReaders.Push(BoxReader);
+			BoxReader = InNewReader;
+		}
+
+		void PopReader()
+		{
+			BoxReader = BoxReaders.Pop();
+		}
+
 		FMP4Box* GetCurrentFtypBox() const
 		{
 			return CurrentFtypBox;
@@ -479,6 +491,7 @@ namespace Electra
 		const FMP4ParseInfo* OptionalMoovParseInfo = nullptr;
 
 		FMP4Box* RootBox = nullptr;					//!< A root box that is not an actual file box but a container representing the file itself.
+		TArray<FMP4BoxReader*> BoxReaders;
 		FMP4BoxReader* BoxReader = nullptr;			//!< Instance of the box reader we were given. This is not ours so we must not delete it!
 		int32 BoxNestingLevel = 0;					//!< Box tree depth
 
@@ -5402,24 +5415,55 @@ namespace Electra
 		{
 			UEMediaError Error = UEMEDIA_ERROR_OK;
 
-			// Parsing of the 'meta' box depends on the major brand. In QuickTime it is not a full box, in ISO it is.
-			const FMP4BoxFTYP* ftypBox = static_cast<const FMP4BoxFTYP*>(ParseInfo->GetCurrentStypBox() ? ParseInfo->GetCurrentStypBox() : ParseInfo->GetCurrentFtypBox());
-			if (ftypBox && ftypBox->GetBrand(0) != FMP4Box::kBrand_QT)
+			// The 'meta' box is a full box under ISO and a basic box in QuickTime.
+			// Some tools write a QT 'meta' box even for ISO which is then missing the version number and flags!
+			// To work around this issue we read the box here as a whole and locate the mandatory 'hdlr' box
+			// that has to be the first box and based on its position handle this as a basic or a full box.
+			TArray<uint8> BoxData;
+			int32 BytesRemaining = BoxSize - (ParseInfo->Reader()->GetCurrentReadOffset() - StartOffset);
+			if (BytesRemaining)
 			{
-				RETURN_IF_ERROR(FMP4BoxFull::ReadAndParseAttributes(ParseInfo));
+				BoxData.AddUninitialized(BytesRemaining);
+				RETURN_IF_ERROR(ParseInfo->Reader()->ReadBytes(BoxData.GetData(), BytesRemaining));
 			}
 
-			int64 NowAt = ParseInfo->Reader()->GetCurrentReadOffset();
-			int32 BytesRemaining = (int32)(BoxSize - (NowAt - StartOffset));
+			bool bIsFullBox = false;
+			if (BytesRemaining >= 12)
+			{
+				bIsFullBox = BoxData[8]=='h' && BoxData[9]=='d' && BoxData[10]=='l' && BoxData[11]=='r';
+			}
+
+			FDataBufferReader BufferReader(BoxData, StartOffset);
+			FMP4BoxReader Reader(&BufferReader, nullptr);
+			struct FAutoPopReader
+			{
+				FAutoPopReader(FMP4ParseInfo* InParseInfo, FMP4BoxReader* InReader) : PInfo(InParseInfo)
+				{ PInfo->PushReader(InReader); }
+				~FAutoPopReader()
+				{ Detach(); }
+				void Detach()
+				{
+					if (PInfo)
+					{
+						PInfo->PopReader();
+						PInfo = nullptr;
+					}
+				}
+				FMP4ParseInfo* PInfo;
+			};
+			FAutoPopReader BoxReader(ParseInfo, &Reader);
+			if (bIsFullBox)
+			{
+				RETURN_IF_ERROR(FMP4BoxFull::ReadAndParseAttributes(ParseInfo));
+				BytesRemaining -= 4;
+			}
 			if (BytesRemaining > 0)
 			{
 				bIsLeafBox = false;
-				RETURN_IF_ERROR(ParseInfo->ReadAndParseNextBox(this, NowAt + BytesRemaining));
-				BytesRemaining = (int32)(BoxSize - (ParseInfo->Reader()->GetCurrentReadOffset() - StartOffset));
-				check(BytesRemaining >= 0);
-				if (BytesRemaining > 0)
+				UEMediaError Err = ParseInfo->ReadAndParseNextBox(this, ParseInfo->Reader()->GetCurrentReadOffset() + BytesRemaining);
+				if (Err != UEMEDIA_ERROR_OK && Err != UEMEDIA_ERROR_END_OF_STREAM)
 				{
-					RETURN_IF_ERROR(ParseInfo->Reader()->ReadBytes(nullptr, BytesRemaining));
+					return Err;
 				}
 			}
 			Handler = static_cast<const FMP4BoxHDLR*>(FindBox(kBox_hdlr, 0));
