ElectraPlayer: Fixed a signedness warning

#rnx

diff --git a/Engine/Plugins/Media/ElectraPlayer/Source/ElectraPlayerRuntime/Private/Runtime/Demuxer/ParserISO14496-12.cpp b/Engine/Plugins/Media/ElectraPlayer/Source/ElectraPlayerRuntime/Private/Runtime/Demuxer/ParserISO14496-12.cpp
--- a/Engine/Plugins/Media/ElectraPlayer/Source/ElectraPlayerRuntime/Private/Runtime/Demuxer/ParserISO14496-12.cpp
+++ b/Engine/Plugins/Media/ElectraPlayer/Source/ElectraPlayerRuntime/Private/Runtime/Demuxer/ParserISO14496-12.cpp
@@ -5420,7 +5420,7 @@ namespace Electra
 			// To work around this issue we read the box here as a whole and locate the mandatory 'hdlr' box
 			// that has to be the first box and based on its position handle this as a basic or a full box.
 			TArray<uint8> BoxData;
-			int32 BytesRemaining = BoxSize - (ParseInfo->Reader()->GetCurrentReadOffset() - StartOffset);
+			uint32 BytesRemaining = (uint32)(BoxSize - (ParseInfo->Reader()->GetCurrentReadOffset() - StartOffset));
 			if (BytesRemaining)
 			{
 				BoxData.AddUninitialized(BytesRemaining);
