Electra: Added support for CMCD v1 to DASH and HLS

#rnx

diff --git a/Engine/Plugins/Media/ElectraPlayer/Source/ElectraPlayerRuntime/Private/ElectraPlayer.cpp b/Engine/Plugins/Media/ElectraPlayer/Source/ElectraPlayerRuntime/Private/ElectraPlayer.cpp
--- a/Engine/Plugins/Media/ElectraPlayer/Source/ElectraPlayerRuntime/Private/ElectraPlayer.cpp
+++ b/Engine/Plugins/Media/ElectraPlayer/Source/ElectraPlayerRuntime/Private/ElectraPlayer.cpp
@@ -3,6 +3,7 @@
 #include "ElectraPlayer.h"
 #include "ElectraPlayerPrivate.h"
 #include "Misc/Optional.h"
+#include "Misc/SecureHash.h"
 
 #include "MediaSamples.h"
 #include "MediaPlayerOptions.h"
@@ -12,6 +13,7 @@
 #include "Player/AdaptiveStreamingPlayer.h"
 #include "Player/AdaptivePlayerOptionKeynames.h"
 #include "Utilities/Utilities.h"
+#include "Utilities/URLParser.h"
 #include "MediaSubtitleDecoderOutput.h"
 #include "MediaMetaDataDecoderOutput.h"
 #include "IElectraSubtitleSample.h"
@@ -41,6 +43,8 @@ DECLARE_CYCLE_STAT(TEXT("FElectraPlayer::QueryOptions"), STAT_ElectraPlayer_Elec
 namespace ElectraMediaOptions
 {
 	static const FName GetSafeMediaOptions(TEXT("GetSafeMediaOptions"));
+	static const FName ContentID(TEXT("content_id"));
+	static const FName ElectraCMCDConfig(TEXT("ElectraCMCDConfig"));
 	static const FName ElectraNoPreloading(TEXT("ElectraNoPreloading"));
 	static const FName PlaylistProperties(TEXT("playlist_properties"));
 	static const FName ElectraInitialBitrate(TEXT("ElectraInitialBitrate"));
@@ -60,7 +64,6 @@ namespace ElectraMediaOptions
 	static const FName ElectraGetLicenseKeyData(TEXT("ElectraGetLicenseKeyData"));
 	static const FName ElectraGetPlaystartPosFromSeekPositions(TEXT("ElectraGetPlaystartPosFromSeekPositions"));
 
-	static const FName OptionKeyParseTimecodeInfo(TEXT("parse_timecode_info"));
 
 
 	enum class EOptionType
@@ -466,6 +469,8 @@ bool FElectraPlayer::OpenInternal(const FString& InUrl, const IMediaOptions* InO
 	CloseInternal();
 	ClearToDefaultState();
 
+	FGuid SessionID = FGuid::NewGuid();
+
 	PlayerUniqueID = ++NextPlayerUniqueID;
 	MediaSamplesLock.Lock();
 	MediaSamples.Reset(new FMediaSamples);
@@ -534,6 +539,23 @@ bool FElectraPlayer::OpenInternal(const FString& InUrl, const IMediaOptions* InO
 	Electra::FParamDict PlayerOptions;
 	PlayerOptions.Set(Electra::OptionKeyOutputTexturePoolID, FVariantValue(OutputTexturePoolDecoderID));
 
+	// Set the session ID to be the same we will use for analytics here (see below)
+	PlayerOptions.Set(Electra::OptionKeySessionID, Electra::FVariantValue(SessionID.ToString(EGuidFormats::DigitsWithHyphensLower)));
+	// Ask for content ID. If none is provided we construct one from the URL.
+	FString ContentID = InOptions->GetMediaOption(ElectraMediaOptions::ContentID, FString());
+	if (ContentID.IsEmpty())
+	{
+		FString UrlToHash(InUrl.TrimStartAndEnd());
+		Electra::FURL_RFC3986 up;
+		if (up.Parse(UrlToHash))
+		{
+			UrlToHash = up.GetHost() + up.GetPath(false, false);
+		}
+		ContentID = FMD5::HashAnsiString(*UrlToHash);
+	}
+	// Set the content ID.
+	PlayerOptions.Set(Electra::OptionKeyContentID, Electra::FVariantValue(ContentID));
+
 	for(auto &CodecOption : ElectraMediaOptions::CodecOptions)
 	{
 		FString Value = InOptions->GetMediaOption(CodecOption, FString());
@@ -546,25 +568,31 @@ bool FElectraPlayer::OpenInternal(const FString& InUrl, const IMediaOptions* InO
 	FString PlaylistProperties = InOptions->GetMediaOption(ElectraMediaOptions::PlaylistProperties, FString());
 	if (PlaylistProperties.Len())
 	{
-		PlayerOptions.Set(ElectraMediaOptions::PlaylistProperties, Electra::FVariantValue(PlaylistProperties));
+		PlayerOptions.Set(Electra::OptionKeyPlaylistProperties, Electra::FVariantValue(PlaylistProperties));
+	}
+	// CMCD configuration
+	FString CMCDConfiguration = InOptions->GetMediaOption(ElectraMediaOptions::ElectraCMCDConfig, FString());
+	if (CMCDConfiguration.Len())
+	{
+		PlayerOptions.Set(Electra::OptionKeyCMCDConfiguration, Electra::FVariantValue(CMCDConfiguration));
 	}
 	// Timecode parsing?
 	if (InPlayerOptions && InPlayerOptions->InternalCustomOptions.Find(MediaPlayerOptionValues::ParseTimecodeInfo()))
 	{
-		PlayerOptions.Set(ElectraMediaOptions::OptionKeyParseTimecodeInfo, FVariantValue());
+		PlayerOptions.Set(Electra::OptionKeyParseTimecodeInfo, FVariantValue());
 	}
 
 	// Check for one-time initialization options that can't be changed during playback.
 	int64 InitialStreamBitrate = InOptions->GetMediaOption(ElectraMediaOptions::ElectraInitialBitrate, (int64)-1);
 	if (InitialStreamBitrate > 0)
 	{
-		PlayerOptions.Set(TEXT("initial_bitrate"), Electra::FVariantValue(InitialStreamBitrate));
+		PlayerOptions.Set(Electra::OptionKeyInitialBitrate, Electra::FVariantValue(InitialStreamBitrate));
 		UE_LOG(LogElectraPlayer, Verbose, TEXT("[%u] IMediaPlayer::Open: Using initial bitrate of %d bits/second"), PlayerUniqueID, (int32)InitialStreamBitrate);
 	}
 	FString MediaMimeType = InOptions->GetMediaOption(ElectraMediaOptions::Mimetype, FString());
 	if (MediaMimeType.Len())
 	{
-		PlayerOptions.Set(TEXT("mime_type"), Electra::FVariantValue(MediaMimeType));
+		PlayerOptions.Set(Electra::OptionKeyMimeType, Electra::FVariantValue(MediaMimeType));
 		UE_LOG(LogElectraPlayer, Verbose, TEXT("[%u] IMediaPlayer::Open: Setting media mime type to \"%s\""), PlayerUniqueID, *MediaMimeType);
 	}
 	int64 MaxVerticalHeight = InOptions->GetMediaOption(ElectraMediaOptions::MaxElectraVerticalResolution, (int64)-1);
@@ -582,7 +610,7 @@ bool FElectraPlayer::OpenInternal(const FString& InUrl, const IMediaOptions* InO
 	double LiveEdgeDistanceForNormalPresentation = InOptions->GetMediaOption(ElectraMediaOptions::ElectraLivePresentationOffset, (double)-1.0);
 	if (LiveEdgeDistanceForNormalPresentation > 0.0)
 	{
-		PlayerOptions.Set(TEXT("seekable_range_live_end_offset"), Electra::FVariantValue(Electra::FTimeValue().SetFromSeconds(LiveEdgeDistanceForNormalPresentation)));
+		PlayerOptions.Set(Electra::OptionKeyLiveSeekableEndOffset, Electra::FVariantValue(Electra::FTimeValue().SetFromSeconds(LiveEdgeDistanceForNormalPresentation)));
 		UE_LOG(LogElectraPlayer, Verbose, TEXT("[%u] IMediaPlayer::Open: Setting distance to live edge for normal presentations to %.3f seconds"), PlayerUniqueID, LiveEdgeDistanceForNormalPresentation);
 	}
 	double LiveEdgeDistanceForAudioOnlyPresentation = InOptions->GetMediaOption(ElectraMediaOptions::ElectraLiveAudioPresentationOffset, (double)-1.0);
@@ -600,7 +628,7 @@ bool FElectraPlayer::OpenInternal(const FString& InUrl, const IMediaOptions* InO
 	bool bThrowErrorWhenRebuffering = InOptions->GetMediaOption(ElectraMediaOptions::ElectraThrowErrorWhenRebuffering, (bool)false);
 	if (bThrowErrorWhenRebuffering)
 	{
-		PlayerOptions.Set(TEXT("throw_error_when_rebuffering"), Electra::FVariantValue(bThrowErrorWhenRebuffering));
+		PlayerOptions.Set(Electra::OptionThrowErrorWhenRebuffering, Electra::FVariantValue(bThrowErrorWhenRebuffering));
 		UE_LOG(LogElectraPlayer, Verbose, TEXT("[%u] IMediaPlayer::Open: Throw playback error when rebuffering"), PlayerUniqueID);
 	}
 	FString CDNHTTPStatusDenyStream = InOptions->GetMediaOption(ElectraMediaOptions::ElectraGetDenyStreamCode, FString());
@@ -619,7 +647,7 @@ bool FElectraPlayer::OpenInternal(const FString& InUrl, const IMediaOptions* InO
 	// Certain optimization settings apply for dedicated environments.
 	if (Environment == MediaPlayerOptionValues::Environment_Preview() || Environment == MediaPlayerOptionValues::Environment_Sequencer())
 	{
-		PlayerOptions.Set(TEXT("worker_threads"), Electra::FVariantValue(FString(TEXT("worker"))));
+		PlayerOptions.Set(Electra::OptionKeyWorkerThreads, Electra::FVariantValue(FString(TEXT("worker"))));
 	}
 
 	// Check for options that can be changed during playback and apply them at startup already.
@@ -643,7 +671,7 @@ bool FElectraPlayer::OpenInternal(const FString& InUrl, const IMediaOptions* InO
 	QueuedAnalyticEvents.Empty();
 	NumQueuedAnalyticEvents = 0;
 	// Create a guid string for the analytics. We do this here and not in the constructor in case the same instance is used over again.
-	AnalyticsInstanceGuid = FGuid::NewGuid().ToString(EGuidFormats::Digits);
+	AnalyticsInstanceGuid = SessionID.ToString(EGuidFormats::Digits);
 	UpdateAnalyticsCustomValues();
 	// Start statistics with a clean slate.
 	Statistics.Reset();
diff --git a/Engine/Plugins/Media/ElectraPlayer/Source/ElectraPlayerRuntime/Private/Runtime/Platforms/Android/PlayerRuntime_Android.cpp b/Engine/Plugins/Media/ElectraPlayer/Source/ElectraPlayerRuntime/Private/Runtime/Platforms/Android/PlayerRuntime_Android.cpp
--- a/Engine/Plugins/Media/ElectraPlayer/Source/ElectraPlayerRuntime/Private/Runtime/Platforms/Android/PlayerRuntime_Android.cpp
+++ b/Engine/Plugins/Media/ElectraPlayer/Source/ElectraPlayerRuntime/Private/Runtime/Platforms/Android/PlayerRuntime_Android.cpp
@@ -1,14 +1,25 @@
 // Copyright Epic Games, Inc. All Rights Reserved.
 
 #include "SynchronizedClock.h"
+#include "PlayerPlatform.h"
 
 namespace Electra
 {
 
-int64 MEDIAutcTime::CurrentMSec()
-{
-	FTimespan localTime(FDateTime::UtcNow().GetTicks());
-	return (int64)localTime.GetTotalMilliseconds();
-};
+	namespace Platform
+	{
+		FString GetPlatformID()
+		{
+			return FString(TEXT("And"));
+		}
+	}
+
+
+
+	int64 MEDIAutcTime::CurrentMSec()
+	{
+		FTimespan localTime(FDateTime::UtcNow().GetTicks());
+		return (int64)localTime.GetTotalMilliseconds();
+	}
 
 }
diff --git a/Engine/Plugins/Media/ElectraPlayer/Source/ElectraPlayerRuntime/Private/Runtime/Platforms/Apple/PlayerRuntime_Apple.cpp b/Engine/Plugins/Media/ElectraPlayer/Source/ElectraPlayerRuntime/Private/Runtime/Platforms/Apple/PlayerRuntime_Apple.cpp
--- a/Engine/Plugins/Media/ElectraPlayer/Source/ElectraPlayerRuntime/Private/Runtime/Platforms/Apple/PlayerRuntime_Apple.cpp
+++ b/Engine/Plugins/Media/ElectraPlayer/Source/ElectraPlayerRuntime/Private/Runtime/Platforms/Apple/PlayerRuntime_Apple.cpp
@@ -1,14 +1,28 @@
 // Copyright Epic Games, Inc. All Rights Reserved.
 
 #include "SynchronizedClock.h"
+#include "PlayerPlatform.h"
 
 namespace Electra
 {
 
-int64 MEDIAutcTime::CurrentMSec()
-{
-	FTimespan localTime(FDateTime::UtcNow().GetTicks());
-	return (int64)localTime.GetTotalMilliseconds();
-};
+	namespace Platform
+	{
+		FString GetPlatformID()
+		{
+		#if PLATFORM_MAC
+			return FString(TEXT("Mac"));
+		#else
+			return FString(TEXT("iOS"));
+		#endif
+		}
+	}
+
+
+	int64 MEDIAutcTime::CurrentMSec()
+	{
+		FTimespan localTime(FDateTime::UtcNow().GetTicks());
+		return (int64)localTime.GetTotalMilliseconds();
+	}
 
 }
diff --git a/Engine/Plugins/Media/ElectraPlayer/Source/ElectraPlayerRuntime/Private/Runtime/Platforms/Linux/PlayerRuntime_Linux.cpp b/Engine/Plugins/Media/ElectraPlayer/Source/ElectraPlayerRuntime/Private/Runtime/Platforms/Linux/PlayerRuntime_Linux.cpp
--- a/Engine/Plugins/Media/ElectraPlayer/Source/ElectraPlayerRuntime/Private/Runtime/Platforms/Linux/PlayerRuntime_Linux.cpp
+++ b/Engine/Plugins/Media/ElectraPlayer/Source/ElectraPlayerRuntime/Private/Runtime/Platforms/Linux/PlayerRuntime_Linux.cpp
@@ -1,14 +1,22 @@
 // Copyright Epic Games, Inc. All Rights Reserved.
 
 #include "SynchronizedClock.h"
+#include "PlayerPlatform.h"
 
 namespace Electra
 {
+	namespace Platform
+	{
+		FString GetPlatformID()
+		{
+			return FString(TEXT("Linux"));
+		}
+	}
 
-int64 MEDIAutcTime::CurrentMSec()
-{
-	FTimespan localTime(FDateTime::UtcNow().GetTicks());
-	return (int64)localTime.GetTotalMilliseconds();
-};
+	int64 MEDIAutcTime::CurrentMSec()
+	{
+		FTimespan localTime(FDateTime::UtcNow().GetTicks());
+		return (int64)localTime.GetTotalMilliseconds();
+	}
 
 }
diff --git a/Engine/Plugins/Media/ElectraPlayer/Source/ElectraPlayerRuntime/Private/Runtime/Platforms/Windows/PlayerRuntime_PC.cpp b/Engine/Plugins/Media/ElectraPlayer/Source/ElectraPlayerRuntime/Private/Runtime/Platforms/Windows/PlayerRuntime_PC.cpp
--- a/Engine/Plugins/Media/ElectraPlayer/Source/ElectraPlayerRuntime/Private/Runtime/Platforms/Windows/PlayerRuntime_PC.cpp
+++ b/Engine/Plugins/Media/ElectraPlayer/Source/ElectraPlayerRuntime/Private/Runtime/Platforms/Windows/PlayerRuntime_PC.cpp
@@ -1,14 +1,23 @@
 // Copyright Epic Games, Inc. All Rights Reserved.
 
 #include "SynchronizedClock.h"
+#include "PlayerPlatform.h"
 
 namespace Electra
 {
 
-int64 MEDIAutcTime::CurrentMSec()
-{
-	FTimespan localTime(FDateTime::UtcNow().GetTicks());
-	return (int64)localTime.GetTotalMilliseconds();
-};
+	namespace Platform
+	{
+		FString GetPlatformID()
+		{
+			return FString(TEXT("Win"));
+		}
+	}
+
+	int64 MEDIAutcTime::CurrentMSec()
+	{
+		FTimespan localTime(FDateTime::UtcNow().GetTicks());
+		return (int64)localTime.GetTotalMilliseconds();
+	}
 
 }
diff --git a/Engine/Plugins/Media/ElectraPlayer/Source/ElectraPlayerRuntime/Private/Runtime/Player/ABRRules/ABRInternal.h b/Engine/Plugins/Media/ElectraPlayer/Source/ElectraPlayerRuntime/Private/Runtime/Player/ABRRules/ABRInternal.h
--- a/Engine/Plugins/Media/ElectraPlayer/Source/ElectraPlayerRuntime/Private/Runtime/Player/ABRRules/ABRInternal.h
+++ b/Engine/Plugins/Media/ElectraPlayer/Source/ElectraPlayerRuntime/Private/Runtime/Player/ABRRules/ABRInternal.h
@@ -41,7 +41,7 @@ namespace Electra
 		bool											bLowLatencyEnabled = false;
 	};
 
-	class IABRInfoInterface : public IAdaptiveStreamSelector::IPlayerLiveControl
+	class IABRInfoInterface : public IPlayerABRLiveControl
 	{
 	public:
 		virtual ~IABRInfoInterface() = default;
diff --git a/Engine/Plugins/Media/ElectraPlayer/Source/ElectraPlayerRuntime/Private/Runtime/Player/ABRRules/ABRLiveStream.cpp b/Engine/Plugins/Media/ElectraPlayer/Source/ElectraPlayerRuntime/Private/Runtime/Player/ABRRules/ABRLiveStream.cpp
--- a/Engine/Plugins/Media/ElectraPlayer/Source/ElectraPlayerRuntime/Private/Runtime/Player/ABRRules/ABRLiveStream.cpp
+++ b/Engine/Plugins/Media/ElectraPlayer/Source/ElectraPlayerRuntime/Private/Runtime/Player/ABRRules/ABRLiveStream.cpp
@@ -277,7 +277,7 @@ private:
 
 	double GetPlayablePlayerDuration(bool& bEOS, EStreamType InStreamType)
 	{
-		IAdaptiveStreamSelector::IPlayerLiveControl::FABRBufferStats bs;
+		IPlayerABRLiveControl::FABRBufferStats bs;
 		Info->ABRGetStreamBufferStats(bs, InStreamType);
 		bEOS = bs.bReachedEnd || bs.bEndOfTrack;
 		return bs.PlayableContentDuration.GetAsSeconds();
@@ -959,7 +959,7 @@ IAdaptiveStreamSelector::ESegmentAction FABRLiveStream::EvaluateForError(TArray<
 					// For low-latency streams trigger a clock resync as a 404 may be indicative of clock drift.
 					if (WorkVars->bIsLowLatencyEnabled)
 					{
-						Info->ABRTriggerClockSync(IAdaptiveStreamSelector::IPlayerLiveControl::EClockSyncType::Required);
+						Info->ABRTriggerClockSync(IPlayerABRLiveControl::EClockSyncType::Required);
 					}
 					else
 					{
diff --git a/Engine/Plugins/Media/ElectraPlayer/Source/ElectraPlayerRuntime/Private/Runtime/Player/ABRRules/ABRVoDPlus.cpp b/Engine/Plugins/Media/ElectraPlayer/Source/ElectraPlayerRuntime/Private/Runtime/Player/ABRRules/ABRVoDPlus.cpp
--- a/Engine/Plugins/Media/ElectraPlayer/Source/ElectraPlayerRuntime/Private/Runtime/Player/ABRRules/ABRVoDPlus.cpp
+++ b/Engine/Plugins/Media/ElectraPlayer/Source/ElectraPlayerRuntime/Private/Runtime/Player/ABRRules/ABRVoDPlus.cpp
@@ -161,7 +161,7 @@ private:
 
 	double GetPlayablePlayerDuration(bool& bEOS, EStreamType InStreamType)
 	{
-		IAdaptiveStreamSelector::IPlayerLiveControl::FABRBufferStats bs;
+		IPlayerABRLiveControl::FABRBufferStats bs;
 		Info->ABRGetStreamBufferStats(bs, InStreamType);
 		bEOS = bs.bReachedEnd || bs.bEndOfTrack;
 		return bs.PlayableContentDuration.GetAsSeconds();
diff --git a/Engine/Plugins/Media/ElectraPlayer/Source/ElectraPlayerRuntime/Private/Runtime/Player/AdaptivePlayerOptionKeynames.h b/Engine/Plugins/Media/ElectraPlayer/Source/ElectraPlayerRuntime/Private/Runtime/Player/AdaptivePlayerOptionKeynames.h
--- a/Engine/Plugins/Media/ElectraPlayer/Source/ElectraPlayerRuntime/Private/Runtime/Player/AdaptivePlayerOptionKeynames.h
+++ b/Engine/Plugins/Media/ElectraPlayer/Source/ElectraPlayerRuntime/Private/Runtime/Player/AdaptivePlayerOptionKeynames.h
@@ -6,6 +6,12 @@
 
 namespace Electra
 {
+	//! (FString) An ID uniquely identifying the content. Must not exceed 64 characters.
+	const FName OptionKeyContentID(TEXT("content_id"));
+
+	//! (FString) A session ID, in UUID format. Limited to 64 characters if not UUID format.
+	const FName OptionKeySessionID(TEXT("session_id"));
+
 	//! (int64) ID of the platform specific output texture pool registered with the decoder platform resources.
 	const FName OptionKeyOutputTexturePoolID(TEXT("output_texturepool_id"));
 
@@ -55,6 +61,8 @@ namespace Electra
 	const FName OptionKeyParseTimecodeInfo(TEXT("parse_timecode_info"));
 
 	const FName OptionKeyPlaylistProperties(TEXT("playlist_properties"));
+
+	const FName OptionKeyCMCDConfiguration(TEXT("cmcd_configuration"));
 } // namespace Electra
 
 
diff --git a/Engine/Plugins/Media/ElectraPlayer/Source/ElectraPlayerRuntime/Private/Runtime/Player/AdaptiveStreamingPlayer.cpp b/Engine/Plugins/Media/ElectraPlayer/Source/ElectraPlayerRuntime/Private/Runtime/Player/AdaptiveStreamingPlayer.cpp
--- a/Engine/Plugins/Media/ElectraPlayer/Source/ElectraPlayerRuntime/Private/Runtime/Player/AdaptiveStreamingPlayer.cpp
+++ b/Engine/Plugins/Media/ElectraPlayer/Source/ElectraPlayerRuntime/Private/Runtime/Player/AdaptiveStreamingPlayer.cpp
@@ -1195,6 +1195,11 @@ TSharedPtrTS<IAdaptiveStreamSelector> FAdaptiveStreamingPlayer::GetStreamSelecto
 	return StreamSelector;
 }
 
+IPlayerABRLiveControl* FAdaptiveStreamingPlayer::GetPlayerARBControl()
+{
+	return this;
+}
+
 TSharedPtrTS<IPlaylistReader> FAdaptiveStreamingPlayer::GetManifestReader()
 {
 	return ManifestReader;
@@ -1205,6 +1210,11 @@ TSharedPtrTS<FContentSteeringHandler> FAdaptiveStreamingPlayer::GetContentSteeri
 	return ContentSteeringHandler;
 }
 
+TSharedPtrTS<FCMCDHandler> FAdaptiveStreamingPlayer::GetCMCDHandler()
+{
+	return CMCDHandler;
+}
+
 TSharedPtrTS<IPlayerEntityCache> FAdaptiveStreamingPlayer::GetEntityCache()
 {
 	return EntityCache;
@@ -1430,6 +1440,11 @@ void FAdaptiveStreamingPlayer::InternalHandleOnce()
 		{
 			ContentSteeringHandler->PeriodicHandle();
 		}
+		// Handle CMCD
+		if (CMCDHandler.IsValid())
+		{
+			CMCDHandler->PeriodicHandle();
+		}
 	}
 }
 
@@ -4504,12 +4519,17 @@ void FAdaptiveStreamingPlayer::InternalInitialize()
 	// Create an HTTP response cache.
 	HttpResponseCache = IHTTPResponseCache::Create(GetOptionValue(OptionKeyResponseCacheMaxByteSize).SafeGetInt64(0), (int32)GetOptionValue(OptionKeyResponseCacheMaxEntries).SafeGetInt64(8192), nullptr);
 
+	// Create the ABR stream selector.
+	StreamSelector = IAdaptiveStreamSelector::Create(this);
+	AddMetricsReceiver(StreamSelector.Get());
+
 	// Create the content steering handler.
 	ContentSteeringHandler = MakeSharedTS<FContentSteeringHandler>(this);
 
-	// Create the ABR stream selector.
-	StreamSelector = IAdaptiveStreamSelector::Create(this, this);
-	AddMetricsReceiver(StreamSelector.Get());
+	// Create the CMCD handler.
+	CMCDHandler = MakeSharedTS<FCMCDHandler>(this);
+	CMCDHandler->Initialize(GetOptionValue(OptionKeyContentID).SafeGetFString(), GetOptionValue(OptionKeySessionID).SafeGetFString(), GetOptionValue(OptionKeyCMCDConfiguration).SafeGetFString(FString()));
+	AddMetricsReceiver(CMCDHandler.Get());
 
 	// Create renderers.
 	CreateRenderers();
@@ -4594,14 +4614,16 @@ void FAdaptiveStreamingPlayer::InternalClose()
 	RemoveMetricsReceiver(StreamSelector.Get());
 	StreamSelector.Reset();
 	ContentSteeringHandler.Reset();
+	RemoveMetricsReceiver(CMCDHandler.Get());
+	CMCDHandler.Reset();
 
 	HttpManager.Reset();
 	EntityCache.Reset();
 	HttpResponseCache.Reset();
 
 	// Reset remaining internal state
-	CurrentState   	= EPlayerState::eState_Idle;
-	StreamState		= EStreamState::eStream_Running;
+	CurrentState = EPlayerState::eState_Idle;
+	StreamState = EStreamState::eStream_Running;
 	RebufferCause = ERebufferCause::None;
 	LastBufferingState = EPlayerState::eState_Buffering;
 	RebufferDetectedAtPlayPos.SetToInvalid();
@@ -4814,7 +4836,7 @@ FTimeValue FAdaptiveStreamingPlayer::ABRGetPlaySpeed() const
 {
 	return FTimeValue(RenderClock->IsRunning() ? 1.0 : 0.0);
 }
-void FAdaptiveStreamingPlayer::ABRGetStreamBufferStats(IAdaptiveStreamSelector::IPlayerLiveControl::FABRBufferStats& OutBufferStats, EStreamType ForStream)
+void FAdaptiveStreamingPlayer::ABRGetStreamBufferStats(IPlayerABRLiveControl::FABRBufferStats& OutBufferStats, EStreamType ForStream)
 {
 	UpdateDiagnostics();
 	DiagnosticsCriticalSection.Lock();
@@ -4879,11 +4901,11 @@ double FAdaptiveStreamingPlayer::ABRGetRenderRateScale() const
 	return RenderRateScale;
 }
 
-void FAdaptiveStreamingPlayer::ABRTriggerClockSync(IAdaptiveStreamSelector::IPlayerLiveControl::EClockSyncType InClockSyncType)
+void FAdaptiveStreamingPlayer::ABRTriggerClockSync(IPlayerABRLiveControl::EClockSyncType InClockSyncType)
 {
 	if (Manifest.IsValid())
 	{
-		Manifest->TriggerClockSync(InClockSyncType == IAdaptiveStreamSelector::IPlayerLiveControl::EClockSyncType::Required ?
+		Manifest->TriggerClockSync(InClockSyncType == IPlayerABRLiveControl::EClockSyncType::Required ?
 									IManifest::EClockSyncType::Required : IManifest::EClockSyncType::Recommended);
 	}
 }
diff --git a/Engine/Plugins/Media/ElectraPlayer/Source/ElectraPlayerRuntime/Private/Runtime/Player/AdaptiveStreamingPlayerABR.cpp b/Engine/Plugins/Media/ElectraPlayer/Source/ElectraPlayerRuntime/Private/Runtime/Player/AdaptiveStreamingPlayerABR.cpp
--- a/Engine/Plugins/Media/ElectraPlayer/Source/ElectraPlayerRuntime/Private/Runtime/Player/AdaptiveStreamingPlayerABR.cpp
+++ b/Engine/Plugins/Media/ElectraPlayer/Source/ElectraPlayerRuntime/Private/Runtime/Player/AdaptiveStreamingPlayerABR.cpp
@@ -6,10 +6,10 @@
 #include "Player/PlaybackTimeline.h"
 #include "Player/AdaptivePlayerOptionKeynames.h"
 #include "Player/ABRRules/ABRInternal.h"
-//#include "Player/ABRRules/ABRLeakyBucket.h"
 #include "Player/ABRRules/ABRVoDPlus.h"
 #include "Player/ABRRules/ABRLiveStream.h"
 #include "Player/ABRRules/ABRFixedStream.h"
+#include "Player/CMCDHandler.h"
 
 #include "StreamAccessUnitBuffer.h"
 #include "Utilities/Utilities.h"
@@ -22,7 +22,7 @@ namespace Electra
 	class FAdaptiveStreamSelector : public IAdaptiveStreamSelector, public IABRInfoInterface
 	{
 	public:
-		FAdaptiveStreamSelector(IPlayerSessionServices* PlayerSessionServices, IAdaptiveStreamSelector::IPlayerLiveControl* InPlayerLiveControl);
+		FAdaptiveStreamSelector(IPlayerSessionServices* PlayerSessionServices, IPlayerABRLiveControl* InPlayerLiveControl);
 		virtual ~FAdaptiveStreamSelector();
 		void SetFormatType(EMediaFormatType FormatType) override;
 		void SetBandwidthCeiling(int32 HighestManifestBitrate) override;
@@ -113,7 +113,7 @@ namespace Electra
 		{ return PlayerLiveControl->ABRGetLatency(); }
 		FTimeValue ABRGetPlaySpeed() const override
 		{ return PlayerLiveControl->ABRGetPlaySpeed(); }
-		void ABRGetStreamBufferStats(IAdaptiveStreamSelector::IPlayerLiveControl::FABRBufferStats& OutBufferStats, EStreamType ForStream) override
+		void ABRGetStreamBufferStats(IPlayerABRLiveControl::FABRBufferStats& OutBufferStats, EStreamType ForStream) override
 		{ PlayerLiveControl->ABRGetStreamBufferStats(OutBufferStats, ForStream); }
 		FTimeRange ABRGetSupportedRenderRateScale() override
 		{ return PlayerLiveControl->ABRGetSupportedRenderRateScale(); }
@@ -121,14 +121,14 @@ namespace Electra
 		{ PlayerLiveControl->ABRSetRenderRateScale(InRenderRateScale); }
 		double ABRGetRenderRateScale() const override
 		{ return PlayerLiveControl->ABRGetRenderRateScale(); }
-		void ABRTriggerClockSync(IAdaptiveStreamSelector::IPlayerLiveControl::EClockSyncType InClockSyncType) override
+		void ABRTriggerClockSync(IPlayerABRLiveControl::EClockSyncType InClockSyncType) override
 		{ PlayerLiveControl->ABRTriggerClockSync(InClockSyncType); }
 		void ABRTriggerPlaylistRefresh() override
 		{ PlayerLiveControl->ABRTriggerPlaylistRefresh(); }
 
 
 		IPlayerSessionServices*								PlayerSessionServices;
-		IAdaptiveStreamSelector::IPlayerLiveControl*		PlayerLiveControl;
+		IPlayerABRLiveControl*								PlayerLiveControl;
 
 		FCriticalSection									AccessMutex;
 		TSharedPtrTS<IManifest::IPlayPeriod>				CurrentPlayPeriodVideo;
@@ -155,13 +155,12 @@ namespace Electra
 	 * Create an instance of this class
 	 *
 	 * @param InPlayerSessionServices
-	 * @param InPlayerLiveControl
 	 *
 	 * @return
 	 */
-	TSharedPtrTS<IAdaptiveStreamSelector> IAdaptiveStreamSelector::Create(IPlayerSessionServices* InPlayerSessionServices, IPlayerLiveControl* InPlayerLiveControl)
+	TSharedPtrTS<IAdaptiveStreamSelector> IAdaptiveStreamSelector::Create(IPlayerSessionServices* InPlayerSessionServices)
 	{
-		return TSharedPtrTS<IAdaptiveStreamSelector>(new FAdaptiveStreamSelector(InPlayerSessionServices, InPlayerLiveControl));
+		return TSharedPtrTS<IAdaptiveStreamSelector>(new FAdaptiveStreamSelector(InPlayerSessionServices, InPlayerSessionServices->GetPlayerARBControl()));
 	}
 
 	//-----------------------------------------------------------------------------
@@ -171,7 +170,7 @@ namespace Electra
 	 * @param InPlayerSessionServices
 	 * @param InPlayerLiveControl
 	 */
-	FAdaptiveStreamSelector::FAdaptiveStreamSelector(IPlayerSessionServices* InPlayerSessionServices, IAdaptiveStreamSelector::IPlayerLiveControl* InPlayerLiveControl)
+	FAdaptiveStreamSelector::FAdaptiveStreamSelector(IPlayerSessionServices* InPlayerSessionServices, IPlayerABRLiveControl* InPlayerLiveControl)
 		: PlayerSessionServices(InPlayerSessionServices)
 		, PlayerLiveControl(InPlayerLiveControl)
 	{
@@ -258,7 +257,9 @@ namespace Electra
 							{
 								return a->Bitrate < b->Bitrate;
 							});
-
+							
+							// Set the lowest and highest bitrate with the CMCD handler
+							PlayerSessionServices->GetCMCDHandler()->SetPlayableBitrateRange(EStreamType::Video, StreamInformationVideo.Num() ? StreamInformationVideo[0]->Bitrate : 0, StreamInformationVideo.Num() ? StreamInformationVideo.Last()->Bitrate : 0);
 							break;
 						}
 					}
@@ -304,6 +305,8 @@ namespace Electra
 								return a->Bitrate < b->Bitrate;
 							});
 
+							// Set the lowest and highest bitrate with the CMCD handler
+							PlayerSessionServices->GetCMCDHandler()->SetPlayableBitrateRange(EStreamType::Audio, StreamInformationAudio.Num() ? StreamInformationAudio[0]->Bitrate : 0, StreamInformationAudio.Num() ? StreamInformationAudio.Last()->Bitrate : 0);
 							break;
 						}
 					}
diff --git a/Engine/Plugins/Media/ElectraPlayer/Source/ElectraPlayerRuntime/Private/Runtime/Player/AdaptiveStreamingPlayerABR.h b/Engine/Plugins/Media/ElectraPlayer/Source/ElectraPlayerRuntime/Private/Runtime/Player/AdaptiveStreamingPlayerABR.h
--- a/Engine/Plugins/Media/ElectraPlayer/Source/ElectraPlayerRuntime/Private/Runtime/Player/AdaptiveStreamingPlayerABR.h
+++ b/Engine/Plugins/Media/ElectraPlayer/Source/ElectraPlayerRuntime/Private/Runtime/Player/AdaptiveStreamingPlayerABR.h
@@ -14,68 +14,67 @@
 namespace Electra
 {
 
-
-class IAdaptiveStreamSelector : public IAdaptiveStreamingPlayerMetrics
+class IPlayerABRLiveControl
 {
 public:
-	class IPlayerLiveControl
-	{
-	public:
-		virtual ~IPlayerLiveControl() = default;
-
-		// Provides the current playback position.
-		virtual FTimeValue ABRGetPlayPosition() const = 0;
+	virtual ~IPlayerABRLiveControl() = default;
 
-		// Provides the current full media timeline, not adjusted for seeking.
-		virtual FTimeRange ABRGetTimeline() const = 0;
+	// Provides the current playback position.
+	virtual FTimeValue ABRGetPlayPosition() const = 0;
 
-		// Provides the current wallclock time (server real time).
-		virtual FTimeValue ABRGetWallclockTime() const = 0;
+	// Provides the current full media timeline, not adjusted for seeking.
+	virtual FTimeRange ABRGetTimeline() const = 0;
 
-		// Returns true if this is a Live presentation, false for VoD.
-		virtual bool ABRIsLive() const = 0;
+	// Provides the current wallclock time (server real time).
+	virtual FTimeValue ABRGetWallclockTime() const = 0;
 
-		// Returns true if playback should be on the Live edge. false if the user intentionally seeked away from the edge.
-		virtual bool ABRShouldPlayOnLiveEdge() const = 0;
+	// Returns true if this is a Live presentation, false for VoD.
+	virtual bool ABRIsLive() const = 0;
 
-		// Provides the current low latency descriptor if low latency playback is active. Returns nullptr if not.
-		virtual TSharedPtrTS<const FLowLatencyDescriptor> ABRGetLowLatencyDescriptor() const = 0;
+	// Returns true if playback should be on the Live edge. false if the user intentionally seeked away from the edge.
+	virtual bool ABRShouldPlayOnLiveEdge() const = 0;
 
-		// Provides the desired latency to the live edge for both low-latency and normal live playback.
-		virtual FTimeValue ABRGetDesiredLiveEdgeLatency() const = 0;
+	// Provides the current low latency descriptor if low latency playback is active. Returns nullptr if not.
+	virtual TSharedPtrTS<const FLowLatencyDescriptor> ABRGetLowLatencyDescriptor() const = 0;
 
-		// Provides the currently measured latency to the live edge.
-		virtual FTimeValue ABRGetLatency() const = 0;
+	// Provides the desired latency to the live edge for both low-latency and normal live playback.
+	virtual FTimeValue ABRGetDesiredLiveEdgeLatency() const = 0;
 
-		// Provides the current playback speed.
-		virtual FTimeValue ABRGetPlaySpeed() const = 0;
+	// Provides the currently measured latency to the live edge.
+	virtual FTimeValue ABRGetLatency() const = 0;
 
-		// Returns the current stream access unit buffer stats.
-		struct FABRBufferStats
-		{
-			FTimeValue PlayableContentDuration;
-			bool bReachedEnd = false;
-			bool bEndOfTrack = false;
-		};
-		virtual void ABRGetStreamBufferStats(FABRBufferStats& OutBufferStats, EStreamType ForStream) = 0;
+	// Provides the current playback speed.
+	virtual FTimeValue ABRGetPlaySpeed() const = 0;
 
-		virtual FTimeRange ABRGetSupportedRenderRateScale() = 0;
-		virtual void ABRSetRenderRateScale(double InRenderRateScale) = 0;
-		virtual double ABRGetRenderRateScale() const = 0;
+	// Returns the current stream access unit buffer stats.
+	struct FABRBufferStats
+	{
+		FTimeValue PlayableContentDuration;
+		bool bReachedEnd = false;
+		bool bEndOfTrack = false;
+	};
+	virtual void ABRGetStreamBufferStats(FABRBufferStats& OutBufferStats, EStreamType ForStream) = 0;
 
-		enum class EClockSyncType
-		{
-			Recommended,
-			Required
-		};
-		virtual void ABRTriggerClockSync(EClockSyncType InClockSyncType) = 0;
+	virtual FTimeRange ABRGetSupportedRenderRateScale() = 0;
+	virtual void ABRSetRenderRateScale(double InRenderRateScale) = 0;
+	virtual double ABRGetRenderRateScale() const = 0;
 
-		virtual void ABRTriggerPlaylistRefresh() = 0;
+	enum class EClockSyncType
+	{
+		Recommended,
+		Required
 	};
+	virtual void ABRTriggerClockSync(EClockSyncType InClockSyncType) = 0;
+
+	virtual void ABRTriggerPlaylistRefresh() = 0;
+};
 
 
+class IAdaptiveStreamSelector : public IAdaptiveStreamingPlayerMetrics
+{
+public:
 	//! Create an instance of this class
-	static TSharedPtrTS<IAdaptiveStreamSelector> Create(IPlayerSessionServices* PlayerSessionServices, IPlayerLiveControl* PlayerLiveControl);
+	static TSharedPtrTS<IAdaptiveStreamSelector> Create(IPlayerSessionServices* PlayerSessionServices);
 	virtual ~IAdaptiveStreamSelector() = default;
 
 	virtual void SetFormatType(EMediaFormatType FormatType) = 0;
diff --git a/Engine/Plugins/Media/ElectraPlayer/Source/ElectraPlayerRuntime/Private/Runtime/Player/AdaptiveStreamingPlayerInternal.h b/Engine/Plugins/Media/ElectraPlayer/Source/ElectraPlayerRuntime/Private/Runtime/Player/AdaptiveStreamingPlayerInternal.h
--- a/Engine/Plugins/Media/ElectraPlayer/Source/ElectraPlayerRuntime/Private/Runtime/Player/AdaptiveStreamingPlayerInternal.h
+++ b/Engine/Plugins/Media/ElectraPlayer/Source/ElectraPlayerRuntime/Private/Runtime/Player/AdaptiveStreamingPlayerInternal.h
@@ -24,6 +24,7 @@
 #include "Player/AdaptiveStreamingPlayerABR.h"
 #include "Player/AdaptivePlayerOptionKeynames.h"
 #include "Player/ContentSteeringHandler.h"
+#include "Player/CMCDHandler.h"
 #include "OutputHandler.h"
 
 #include "Templates/UniqueObj.h"
@@ -1066,7 +1067,7 @@ class FAdaptiveStreamingPlayer : public IAdaptiveStreamingPlayer
 							   , public IAccessUnitMemoryProvider
 							   , public IPlayerStreamFilter
 							   , public TSharedFromThis<FAdaptiveStreamingPlayer, ESPMode::ThreadSafe>
-							   , public IAdaptiveStreamSelector::IPlayerLiveControl
+							   , public IPlayerABRLiveControl
 {
 public:
 	class TDeleter
@@ -1177,10 +1178,12 @@ private:
 	TSharedPtr<IAdaptiveStreamingPlayerResourceProvider, ESPMode::ThreadSafe> GetStaticResourceProvider() override;
 	TSharedPtrTS<IElectraHttpManager> GetHTTPManager() override;
 	TSharedPtrTS<IAdaptiveStreamSelector> GetStreamSelector() override;
+	IPlayerABRLiveControl* GetPlayerARBControl() override;
 	IPlayerStreamFilter* GetStreamFilter() override;
 	const FCodecSelectionPriorities& GetCodecSelectionPriorities(EStreamType ForStream) override;
 	TSharedPtrTS<IPlaylistReader> GetManifestReader() override;
 	TSharedPtrTS<FContentSteeringHandler> GetContentSteeringHandler() override;
+	TSharedPtrTS<FCMCDHandler> GetCMCDHandler() override;
 	TSharedPtrTS<IPlayerEntityCache> GetEntityCache() override;
 	TSharedPtrTS<IHTTPResponseCache> GetHTTPResponseCache() override;
 	IAdaptiveStreamingPlayerAEMSHandler* GetAEMSEventHandler() override;
@@ -1197,7 +1200,7 @@ private:
 	bool CanDecodeSubtitle(const FString& MimeType, const FString& Codec) const override;
 
 
-	// Methods from IAdaptiveStreamSelector::IPlayerLiveControl
+	// Methods from IPlayerABRLiveControl
 	FTimeValue ABRGetPlayPosition() const override;
 	FTimeRange ABRGetTimeline() const override;
 	FTimeValue ABRGetWallclockTime() const override;
@@ -1207,11 +1210,11 @@ private:
 	FTimeValue ABRGetDesiredLiveEdgeLatency() const override;
 	FTimeValue ABRGetLatency() const override;
 	FTimeValue ABRGetPlaySpeed() const override;
-	void ABRGetStreamBufferStats(IAdaptiveStreamSelector::IPlayerLiveControl::FABRBufferStats& OutBufferStats, EStreamType ForStream) override;
+	void ABRGetStreamBufferStats(IPlayerABRLiveControl::FABRBufferStats& OutBufferStats, EStreamType ForStream) override;
 	FTimeRange ABRGetSupportedRenderRateScale() override;
 	void ABRSetRenderRateScale(double InRenderRateScale) override;
 	double ABRGetRenderRateScale() const override;
-	void ABRTriggerClockSync(IAdaptiveStreamSelector::IPlayerLiveControl::EClockSyncType InClockSyncType) override;
+	void ABRTriggerClockSync(IPlayerABRLiveControl::EClockSyncType InClockSyncType) override;
 	void ABRTriggerPlaylistRefresh() override;
 
 
@@ -2324,6 +2327,7 @@ private:
 	bool																bIsClosing;
 
 	TSharedPtrTS<FContentSteeringHandler>								ContentSteeringHandler;
+	TSharedPtrTS<FCMCDHandler>											CMCDHandler;
 
 	TSharedPtrTS<IAdaptiveStreamSelector>								StreamSelector;
 	int32																BitrateCeiling;
diff --git a/Engine/Plugins/Media/ElectraPlayer/Source/ElectraPlayerRuntime/Private/Runtime/Player/AdaptiveStreamingPlayerManifest.cpp b/Engine/Plugins/Media/ElectraPlayer/Source/ElectraPlayerRuntime/Private/Runtime/Player/AdaptiveStreamingPlayerManifest.cpp
--- a/Engine/Plugins/Media/ElectraPlayer/Source/ElectraPlayerRuntime/Private/Runtime/Player/AdaptiveStreamingPlayerManifest.cpp
+++ b/Engine/Plugins/Media/ElectraPlayer/Source/ElectraPlayerRuntime/Private/Runtime/Player/AdaptiveStreamingPlayerManifest.cpp
@@ -228,6 +228,7 @@ void FAdaptiveStreamingPlayer::InternalLoadManifest(const FString& InURL, const
 			{
 				ManifestReader = IPlaylistHandlerHLS::Create(this);
 				ManifestType = EMediaFormatType::HLS;
+				CMCDHandler->SetStreamingFormat(FCMCDHandler::EStreamingFormat::HLS);
 			}
 			else if (mimeType == Playlist::MIMETypeMP4 || mimeType == Playlist::MIMETypeMP4A || mimeType == Playlist::MIMETypeQuickTime)
 			{
@@ -238,6 +239,7 @@ void FAdaptiveStreamingPlayer::InternalLoadManifest(const FString& InURL, const
 			{
 				ManifestReader = IPlaylistReaderDASH::Create(this);
 				ManifestType = EMediaFormatType::DASH;
+				CMCDHandler->SetStreamingFormat(FCMCDHandler::EStreamingFormat::DASH);
 			}
 			else if (mimeType == Playlist::MIMETypeMKV || mimeType == Playlist::MIMETypeMKA)
 			{
@@ -361,6 +363,28 @@ bool FAdaptiveStreamingPlayer::SelectManifest()
 				HttpResponseCache->Disable();
 			}
 
+			// Update the CMCD handler
+			switch(ManifestType)
+			{
+				case EMediaFormatType::HLS:
+				{
+					CMCDHandler->SetStreamingFormat(FCMCDHandler::EStreamingFormat::HLS);
+					CMCDHandler->SetStreamType(NewPresentation->GetPresentationType() == IManifest::EType::OnDemand ? FCMCDHandler::EStreamType::VOD : FCMCDHandler::EStreamType::Live);
+					break;
+				}
+				case EMediaFormatType::DASH:
+				{
+					CMCDHandler->SetStreamingFormat(NewPresentation->GetLowLatencyDescriptor().IsValid() ? FCMCDHandler::EStreamingFormat::LowLatencyDASH : FCMCDHandler::EStreamingFormat::DASH);
+					CMCDHandler->SetStreamType(NewPresentation->GetPresentationType() == IManifest::EType::OnDemand ? FCMCDHandler::EStreamType::VOD : FCMCDHandler::EStreamType::Live);
+					break;
+				}
+				default:
+				{
+					CMCDHandler->SetStreamingFormat(FCMCDHandler::EStreamingFormat::Undefined);
+					break;
+				}
+			}
+
 			return true;
 		}
 		else
diff --git a/Engine/Plugins/Media/ElectraPlayer/Source/ElectraPlayerRuntime/Private/Runtime/Player/CMCDHandler.cpp b/Engine/Plugins/Media/ElectraPlayer/Source/ElectraPlayerRuntime/Private/Runtime/Player/CMCDHandler.cpp
new file mode 100644
--- /dev/null
+++ b/Engine/Plugins/Media/ElectraPlayer/Source/ElectraPlayerRuntime/Private/Runtime/Player/CMCDHandler.cpp
@@ -0,0 +1,940 @@
+// Copyright Epic Games, Inc. All Rights Reserved.
+
+#include "CMCDHandler.h"
+#include "Player/PlayerSessionServices.h"
+#include "Player/AdaptiveStreamingPlayerABR.h"
+#include "ElectraPlayerMisc.h"
+#include "PlayerPlatform.h"
+
+#include "Dom/JsonObject.h"
+#include "Serialization/JsonReader.h"
+#include "Serialization/JsonSerializer.h"
+
+#include "HAL/IConsoleManager.h"
+static TAutoConsoleVariable<int32> CVarElectraCMCDMetrics(
+	TEXT("Electra.CMCDMetrics"),
+	1,
+	TEXT("Controls CMCD metrics.\n")
+	TEXT(" 0: disable all CMCD metrics; 1: use default settings; 2: always enable CMCD metrics"),
+	ECVF_Default);
+
+
+namespace Electra
+{
+// Standard config only enables playlist provided parameters.
+static const TCHAR* const StandardCMCDConfig = TEXT(R"##({
+"cmcd":[]
+})##");
+
+// Default values are: `headers` mode, `request` type, all keys enabled, version `1` and playlist parameters to be applied.
+// Values that are default do not need to be included in this config.
+static const TCHAR* const DefaultEnabledCMCDConfig = TEXT(R"##({
+"cmcd":[{}]
+})##");
+
+
+
+FCMCDHandler::FCMCDHandler(IPlayerSessionServices* InPlayerSessionServices)
+	: PlayerSessionServices(InPlayerSessionServices)
+{
+}
+
+FCMCDHandler::~FCMCDHandler()
+{
+}
+
+bool FCMCDHandler::IsEnabled()
+{
+	return bInitOk && StreamingFormat != EStreamingFormat::Undefined && !ReportElements.IsEmpty();
+}
+
+void FCMCDHandler::PeriodicHandle()
+{
+}
+
+void FCMCDHandler::Initialize(const FString& InContentID, const FString& InSessionID, const FString& InConfiguration)
+{
+	FString Configuration(InConfiguration);
+
+	ContentID = InContentID;
+	SessionID = InSessionID;
+	PlatformID = Electra::Platform::GetPlatformID();
+
+	int32 CMCDMetrics = CVarElectraCMCDMetrics.GetValueOnAnyThread();
+	// Disable CMCD entirely?
+	if (CMCDMetrics == 0)
+	{
+		bInitOk = false;
+		return;
+	}
+
+	// If the configuration is empty or cannot be parsed as a JSON then CMCD will be disabled for this session.
+	if (Configuration.IsEmpty())
+	{
+		// Using default CMCD means to look for playlist provided parameters only.
+		if (CMCDMetrics == 1)
+		{
+			Configuration = FString(StandardCMCDConfig);
+		}
+		// Always enabling CMCD without a supplied configuration means we are using an internal catch-all config.
+		else if (CMCDMetrics == 2)
+		{
+			Configuration = FString(DefaultEnabledCMCDConfig);
+		}
+	}
+	bInitOk = InitializeInternal(EInitFrom::GlobalConfig, Configuration);
+}
+
+bool FCMCDHandler::InitializeInternal(EInitFrom InInitFrom, const FString& InConfiguration)
+{
+#define V1MAPKEY(K) { TEXT(#K), EV1Keys::k_##K }
+	TMap<FString, EV1Keys> AllV1Keys { V1MAPKEY(br), V1MAPKEY(bl), V1MAPKEY(bs), V1MAPKEY(d), V1MAPKEY(cid), V1MAPKEY(dl), V1MAPKEY(mtp),V1MAPKEY(nor), V1MAPKEY(nrr),
+									   V1MAPKEY(ot), V1MAPKEY(pr), V1MAPKEY(rtp), V1MAPKEY(sf), V1MAPKEY(sid), V1MAPKEY(st), V1MAPKEY(su), V1MAPKEY(tb), V1MAPKEY(v) };
+#undef V1MAPKEY
+
+	TSharedPtr<FJsonObject> Config;
+	TSharedRef<TJsonReader<>> Reader = TJsonReaderFactory<>::Create(InConfiguration);
+	if (!FJsonSerializer::Deserialize(Reader, Config))
+	{
+		if (InInitFrom == EInitFrom::GlobalConfig)
+		{
+			PlayerSessionServices->PostLog(Facility::EFacility::Player, IInfoLog::ELevel::Warning, TEXT("Failed to parse CMCD configuration, handler will be deactivated."));
+		}
+		else
+		{
+			PlayerSessionServices->PostLog(Facility::EFacility::Player, IInfoLog::ELevel::Warning, TEXT("Failed to parse CMCD configuration, changes will not be applied."));
+		}
+		return false;
+	}
+	// Does the configuration provide the content or session ID?
+	if (InInitFrom == EInitFrom::GlobalConfig)
+	{
+		FString CfgContentID;
+		if (Config->TryGetStringField(Key_ContentId(), CfgContentID))
+		{
+			ContentID = CfgContentID;
+		}
+		FString CfgSessionID;
+		if (Config->TryGetStringField(Key_SessionId(), CfgSessionID))
+		{
+			SessionID = CfgSessionID;
+		}
+		if (ContentID.IsEmpty() || SessionID.IsEmpty())
+		{
+			PlayerSessionServices->PostLog(Facility::EFacility::Player, IInfoLog::ELevel::Warning, TEXT("Missing content or session ID, handler will be deactivated."));
+			return false;
+		}
+
+		// See if there is a key defining the name to use for the custom platform identifier.
+		if (!(Config->TryGetStringField(TEXT("platformKey"), PlatformKey)))
+		{
+			PlatformKey = TEXT("platform-id");
+		}
+
+		// Does the config allow or prohibit the use of parameters conveyed in the playlist?
+		bool bAllowPlaylistParams = false;
+		if (Config->TryGetBoolField(TEXT("applyPlaylistParams"), bAllowPlaylistParams))
+		{
+			bApplyPlaylistParams = bAllowPlaylistParams;
+		}
+		// See if there is a key we may get user data from the playlist for CMCD configuration.
+		if (!Config->TryGetStringField(TEXT("importFromPlaylistKey"), PlaylistParamImportKey))
+		{
+			PlaylistParamImportKey = TEXT("EPIC-CMCD");
+		}
+	}
+
+	// Get the `cmcd` array, if it exists. The configuration might also come in later from the playlist.
+	const TArray<TSharedPtr<FJsonValue>>* CMCDArray = nullptr;
+	if ((Config->TryGetArrayField(Key_CmcdArray(), CMCDArray) && CMCDArray))
+	{
+		auto GetVar = [](const TSharedPtr<FJsonObject>& InObj, FStringView InName, EJson InType) -> TSharedPtr<FJsonValue>
+		{
+			return InObj->HasField(InName) ? InObj->GetField(InName, InType) : nullptr;
+		};
+
+		// Iterate it
+		TArray<FReportElement> NewReportElements;
+		for(auto& It : *CMCDArray)
+		{
+			const TSharedPtr<FJsonObject>* Obj = nullptr;
+			if (!It.IsValid() || !It->TryGetObject(Obj))
+			{
+				continue;
+			}
+
+			// New report element.
+			FReportElement re;
+			TSharedPtr<FJsonValue> JVersion = GetVar(*Obj, Key_Version(), EJson::Number);
+			TSharedPtr<FJsonValue> JType = GetVar(*Obj, Key_Type(), EJson::String);
+			TSharedPtr<FJsonValue> JMode = GetVar(*Obj, Key_Mode(), EJson::String);
+			TSharedPtr<FJsonValue> JKeys = GetVar(*Obj, Key_Keys(), EJson::String);
+			TSharedPtr<FJsonValue> JCDNs = GetVar(*Obj, Key_CDNs(), EJson::String);
+			TSharedPtr<FJsonValue> JInclude  = GetVar(*Obj, Key_IncludeInRequests(), EJson::String);
+			TSharedPtr<FJsonValue> JContentID = GetVar(*Obj, Key_ContentId(), EJson::String);
+			TSharedPtr<FJsonValue> JSessionID = GetVar(*Obj, Key_SessionId(), EJson::String);
+
+			re.Version = JVersion.IsValid() ? (int32)JVersion->AsNumber() : 1;
+			// If not a supported version continue with the next
+			if (re.Version != 1)
+			{
+				continue;
+			}
+			FString Type = JType.IsValid() ? JType->AsString() : TEXT("request");
+			// With version 1 the type can only be `request`
+			if (!Type.Equals(TEXT("request")))
+			{
+				continue;
+			}
+			re.Type = FReportElement::EType::Request;
+			// Mode
+			FString Mode = JMode.IsValid() ? JMode->AsString() : TEXT("headers");
+			if (!Mode.Equals(TEXT("headers")) && !Mode.Equals(TEXT("query")))
+			{
+				continue;
+			}
+			re.Mode = Mode.Equals(TEXT("headers")) ? FReportElement::EMode::Headers : FReportElement::EMode::Query;
+
+			// Content and session ID to use with this entry. Use global ones if not specified for this entry.
+			re.ContentID = JContentID.IsValid() ? JContentID->AsString() : ContentID;
+			re.SessionID = JSessionID.IsValid() ? JSessionID->AsString() : SessionID;
+
+			// Enabled keywords. If none are given, or only the "*" wildcard, enabled all of them.
+			if (!JKeys.IsValid() || JKeys->AsString().Equals(TEXT("*")))
+			{
+				for(auto& En : AllV1Keys)
+				{
+					re.EnableKey(En.Value);
+				}
+			}
+			else
+			{
+				TArray<FString> EnabledKeywords;
+				JKeys->AsString().ParseIntoArrayWS(EnabledKeywords, TEXT(","), true);
+				for(auto& En : EnabledKeywords)
+				{
+					if (AllV1Keys.Contains(En))
+					{
+						re.EnableKey(AllV1Keys[En]);
+					}
+				}
+			}
+
+			// CDNs.
+			//  - If absent or empty, all CDNs match.
+			//  - A "." matches the default CDN
+			//       - which is either indicated by a "." (the HLS default pathway)
+			//       - or an empty string (no @serviceLocation in (non-DVB !!) DASH)
+			// ! Note that if CDNs other than "." are listed then the playlist MUST use pathways (HLS) or @serviceLocation (DASH)!
+			// ! Note that in DVB DASH and absent @serviceLocation in the MPD automatically makes it the same as the URL!
+			if (JCDNs.IsValid() && !JCDNs->AsString().IsEmpty())
+			{
+				JCDNs->AsString().ParseIntoArrayWS(re.CDNs, TEXT(","));
+			}
+
+			// Included in requests filter.
+			if (JInclude.IsValid())
+			{
+				TArray<FString> incReqs;
+				JInclude->AsString().ParseIntoArrayWS(incReqs, TEXT(","));
+				for(auto& ir : incReqs)
+				{
+					if (ir.Equals(IncludeInRequests_Segment()))
+					{
+						re.EnableRequestType(ERequestType::InitSegment);
+						re.EnableRequestType(ERequestType::IndexSegment);
+						re.EnableRequestType(ERequestType::Segment);
+					}
+					else if (ir.Equals(IncludeInRequests_FirstPlaylist()))
+					{
+						re.EnableRequestType(ERequestType::FirstPlaylist);
+					}
+					else if (ir.Equals(IncludeInRequests_Playlist()))
+					{
+						re.EnableRequestType(ERequestType::FirstPlaylist);
+						re.EnableRequestType(ERequestType::Playlist);
+						re.EnableRequestType(ERequestType::PlaylistUpdate);
+					}
+					else if (ir.Equals(IncludeInRequests_Steering()))
+					{
+						re.EnableRequestType(ERequestType::Steering);
+					}
+					else if (ir.Equals(IncludeInRequests_Other()))
+					{
+						re.EnableRequestType(ERequestType::Other);
+					}
+					else if (ir.Equals(IncludeInRequests_DRM()))
+					{
+						re.EnableRequestType(ERequestType::DRM);
+					}
+				}
+			}
+			else
+			{
+				// Enable all request types by default.
+				re.EnableRequestType(ERequestType::InitSegment);
+				re.EnableRequestType(ERequestType::IndexSegment);
+				re.EnableRequestType(ERequestType::Segment);
+				re.EnableRequestType(ERequestType::FirstPlaylist);
+				re.EnableRequestType(ERequestType::Playlist);
+				re.EnableRequestType(ERequestType::PlaylistUpdate);
+				re.EnableRequestType(ERequestType::Steering);
+				re.EnableRequestType(ERequestType::Other);
+				re.EnableRequestType(ERequestType::DRM);
+			}
+
+			// Handle configuration items that are only permitted in the global configuration.
+			if (InInitFrom == EInitFrom::GlobalConfig)
+			{
+				TSharedPtr<FJsonValue> JHosts = GetVar(*Obj, TEXT("hosts"), EJson::String);
+				TSharedPtr<FJsonValue> JDisable = GetVar(*Obj, TEXT("disable"), EJson::String);
+				// Disabled keywords.
+				if (JDisable.IsValid())
+				{
+					TArray<FString> DisabledKeywords;
+					JDisable->AsString().ParseIntoArrayWS(DisabledKeywords, TEXT(","), true);
+					for(auto& Dis : DisabledKeywords)
+					{
+						if (AllV1Keys.Contains(Dis))
+						{
+							re.DisableKey(AllV1Keys[Dis]);
+						}
+					}
+				}
+
+				// Hosts. If empty or only the "*" wildcard we set it as empty which is a match-all later on.
+				if (JHosts.IsValid() && !JHosts->AsString().Equals(TEXT("*")))
+				{
+					JHosts->AsString().ParseIntoArrayWS(re.Hosts, TEXT(","));
+					for(int32 i=0; i<re.Hosts.Num(); ++i)
+					{
+						// Check if the host contains one or more wildcards in it.
+						if (re.Hosts[i].Contains(TEXT("*")))
+						{
+							// Replace '.' with the escape backslash since we need to treat it as a literal.
+							FString Host(re.Hosts[i]);
+							re.Hosts.RemoveAt(i);
+							--i;
+							Host.ReplaceInline(TEXT("."), TEXT("\\."));
+							// Replace the '*' wildcards with ".*" used with a regex.
+							Host.ReplaceInline(TEXT("*"), TEXT(".*"));
+							re.HostRegexPatterns.Emplace(MakeUnique<FRegexPattern>(Host));
+						}
+					}
+				}
+			}
+
+			// Add to the report elements list.
+			NewReportElements.Emplace(MoveTemp(re));
+		}
+
+		// Where to insert the new elements
+		if (InInitFrom == EInitFrom::GlobalConfig)
+		{
+			ReportElements = MoveTemp(NewReportElements);
+		}
+		else
+		{
+			ReportElements.Insert(MoveTemp(NewReportElements), 0);
+		}
+	}
+	return true;
+}
+
+
+bool FCMCDHandler::UseParametersFromPlaylist()
+{
+	return bApplyPlaylistParams;
+}
+
+FString FCMCDHandler::UsePlaylistParametersFromKey()
+{
+	return PlaylistParamImportKey;
+}
+
+void FCMCDHandler::UpdateParameters(const FString& InUpdatedConfiguration)
+{
+	InitializeInternal(EInitFrom::Playlist, InUpdatedConfiguration);
+}
+
+
+void FCMCDHandler::SetStreamingFormat(EStreamingFormat InStreamingFormat)
+{
+	StreamingFormat = InStreamingFormat;
+}
+
+void FCMCDHandler::SetStreamType(EStreamType InStreamType)
+{
+	StreamType = InStreamType;
+}
+
+void FCMCDHandler::SetPlayableBitrateRange(Electra::EStreamType InStreamType, int32 InLowestBitrate, int32 InHighestBitrate)
+{
+	FBandwidths bw;
+	bw.Lowest = InLowestBitrate;
+	bw.Highest = InHighestBitrate;
+	Bandwidths.Add(InStreamType, MoveTemp(bw));
+}
+
+
+FString FCMCDHandler::RemoveParamsFromURL(FString InURL)
+{
+	// We need to remove a CMCD parameter from the MPD URL.
+	FURL_RFC3986 ParsedURL;
+	if (ParsedURL.Parse(InURL))
+	{
+		TArray<FURL_RFC3986::FQueryParam> qps;
+		ParsedURL.GetQueryParams(qps, false, false);
+		qps.RemoveAll([](FURL_RFC3986::FQueryParam& e) { return e.Name.Equals(TEXT("CMCD")); });
+		ParsedURL.SetQueryParams(qps);
+		return ParsedURL.Get();
+	}
+	return InURL;
+}
+
+
+void FCMCDHandler::SetupRequestObject(ERequestType InRequestType, EObjectType InObject /* 'ot' */, FString& InOutRequestURL, TArray<HTTP::FHTTPHeader>& InOutRequestHeaders, const FString& InCDNId, const FRequestObjectInfo& InObjectInfo)
+{
+	// If CMCD is not enabled, just return.
+	if (!bInitOk || ReportElements.IsEmpty())
+	{
+		return;
+	}
+	// Bail out if unsupported stream format.
+	if (StreamingFormat == EStreamingFormat::Undefined)
+	{
+		return;
+	}
+	// Get the host the request is to be made to. If that fails, just leave.
+	FURL_RFC3986 RequestHost;
+	if (!RequestHost.Parse(InOutRequestURL))
+	{
+		return;
+	}
+	FString Host(RequestHost.GetHost());
+	// Handle requests first
+	for(auto& Req : ReportElements)
+	{
+		if (Req.Type != FReportElement::EType::Request)
+		{
+			continue;
+		}
+
+		// Is the request type a match? If not skip to the next.
+		if (!Req.IsRequestTypeEnabled(InRequestType))
+		{
+			continue;
+		}
+
+		// Filter for CDN? No CDNs means any CDN matches.
+		if (Req.CDNs.Num())
+		{
+			bool bCDNMatch = false;
+			for(auto& cdn : Req.CDNs)
+			{
+				// If the CDN filter is "." then this allows for an empty CDN ID or the "."
+				// Otherwise the CDN must be a case-insensitive match.
+				if ((cdn.Equals(TEXT(".")) && (InCDNId.IsEmpty() || InCDNId.Equals(TEXT(".")))) ||
+					(cdn.Equals(InCDNId, ESearchCase::IgnoreCase)))
+				{
+					bCDNMatch = true;
+					break;
+				}
+			}
+			if (!bCDNMatch)
+			{
+				continue;
+			}
+		}
+
+		// Match for host?
+		bool bMatch = (Req.Hosts.IsEmpty() && Req.HostRegexPatterns.IsEmpty()) || Req.Hosts.Contains(Host);
+		if (!bMatch && !Req.HostRegexPatterns.IsEmpty())
+		{
+			for(int32 i=0; i<Req.HostRegexPatterns.Num(); ++i)
+			{
+				FRegexMatcher HostNameMatcher(*Req.HostRegexPatterns[i], Host);
+				if (HostNameMatcher.FindNext())
+				{
+					bMatch = true;
+					break;
+				}
+			}
+		}
+		if (bMatch)
+		{
+			// Hostname matches, apply this report
+			ApplyReport(RequestHost, Req, InRequestType, InObject, InOutRequestURL, InOutRequestHeaders, InObjectInfo);
+			// Leave the loop. The first matching request wins.
+			break;
+		}
+	}
+// TBD: Extended v2 handling, check events
+}
+
+void FCMCDHandler::ApplyReport(FURL_RFC3986& InOutRequestHost, const FReportElement& InParams, ERequestType InRequestType, EObjectType InObject, FString& InOutRequestURL, TArray<HTTP::FHTTPHeader>& InOutRequestHeaders, const FRequestObjectInfo& InObjectInfo)
+{
+	FReqHeaders Headers;
+	switch(InObject)
+	{
+		case EObjectType::ManifestOrPlaylist:
+		{
+			ApplyManifestOrPlaylist(Headers, InParams, InObjectInfo);
+			break;
+		}
+		case EObjectType::InitSegment:
+		{
+			ApplyInitSegment(Headers, InParams, InObjectInfo);
+			break;
+		}
+		case EObjectType::VideoOnly:
+		case EObjectType::AudioOnly:
+		case EObjectType::MuxedAudioAndVideo:
+		case EObjectType::CaptionOrSubtitle:
+		case EObjectType::TimedTextTrack:
+		{
+			ApplyMediaSegment(Headers, InObject, InParams, InObjectInfo, InOutRequestHost);
+			break;
+		}
+		case EObjectType::Other:
+		{
+			ApplyGeneric(Headers, InParams, InObjectInfo);
+			break;
+		}
+	}
+
+
+	auto MakeList = [](const TArray<FString>& InList) -> FString
+	{
+		FString Result;
+		for(int32 i=0,iMax=InList.Num(); i<iMax; ++i)
+		{
+			if (i)
+			{
+				Result.Append(TEXT(","));
+			}
+			Result.Append(InList[i]);
+		}
+		return Result;
+	};
+
+	// Add as headers or as query parameters?
+	if (InParams.Mode == FReportElement::EMode::Query)
+	{
+		TArray<FString> AllHeaders(Headers.Request);
+		AllHeaders.Append(Headers.Object);
+		AllHeaders.Append(Headers.Status);
+		AllHeaders.Append(Headers.Session);
+		if (AllHeaders.Num())
+		{
+			AllHeaders.Sort();
+			TArray<FURL_RFC3986::FQueryParam> qps;
+			if (!FURL_RFC3986::UrlEncode(qps.Emplace_GetRef().Value, MakeList(AllHeaders), FString()))
+			{
+				return;
+			}
+			qps[0].Name = TEXT("CMCD");
+			InOutRequestHost.AddOrUpdateQueryParams(qps);
+			InOutRequestURL = InOutRequestHost.Get();
+		}
+	}
+	else
+	{
+		HTTP::FHTTPHeader Hdr;
+		if (Headers.Request.Num())
+		{
+			Headers.Request.Sort();
+			Hdr.Value = MakeList(Headers.Request);
+			Hdr.Header = TEXT("CMCD-Request");
+			InOutRequestHeaders.Emplace(MoveTemp(Hdr));
+		}
+		if (Headers.Object.Num())
+		{
+			Headers.Object.Sort();
+			Hdr.Value = MakeList(Headers.Object);
+			Hdr.Header = TEXT("CMCD-Object");
+			InOutRequestHeaders.Emplace(MoveTemp(Hdr));
+		}
+		if (Headers.Status.Num())
+		{
+			Headers.Status.Sort();
+			Hdr.Value = MakeList(Headers.Status);
+			Hdr.Header = TEXT("CMCD-Status");
+			InOutRequestHeaders.Emplace(MoveTemp(Hdr));
+		}
+		if (Headers.Session.Num())
+		{
+			Headers.Session.Sort();
+			Hdr.Value = MakeList(Headers.Session);
+			Hdr.Header = TEXT("CMCD-Session");
+			InOutRequestHeaders.Emplace(MoveTemp(Hdr));
+		}
+	}
+}
+
+
+void FCMCDHandler::ApplyManifestOrPlaylist(FReqHeaders& OutHeaders, const FReportElement& InParams, const FRequestObjectInfo& InObjectInfo)
+{
+	// cid, ot, pr, sf, sid, st, su, v
+	Add_cid(OutHeaders, InParams);
+	Add_ot(OutHeaders, InParams, GetOT(EObjectType::ManifestOrPlaylist));
+	double pr = PlayerSessionServices->GetPlayerARBControl()->ABRGetPlaySpeed().GetAsSeconds();
+	Add_pr(OutHeaders, InParams, pr);
+	Add_sf(OutHeaders, InParams);
+	Add_sid(OutHeaders, InParams);
+	Add_st(OutHeaders, InParams);
+	Add_su(OutHeaders, InParams);
+	Add_v(OutHeaders, InParams);
+	Add_platform(OutHeaders, InParams);
+}
+
+void FCMCDHandler::ApplyGeneric(FReqHeaders& OutHeaders, const FReportElement& InParams, const FRequestObjectInfo& InObjectInfo)
+{
+	// cid, ot, pr, sf, sid, st, su, v
+	Add_cid(OutHeaders, InParams);
+	Add_ot(OutHeaders, InParams, GetOT(EObjectType::Other));
+	double pr = PlayerSessionServices->GetPlayerARBControl()->ABRGetPlaySpeed().GetAsSeconds();
+	Add_pr(OutHeaders, InParams, pr);
+	Add_sf(OutHeaders, InParams);
+	Add_sid(OutHeaders, InParams);
+	Add_st(OutHeaders, InParams);
+	Add_su(OutHeaders, InParams);
+	Add_v(OutHeaders, InParams);
+	Add_platform(OutHeaders, InParams);
+}
+
+void FCMCDHandler::ApplyInitSegment(FReqHeaders& OutHeaders, const FReportElement& InParams, const FRequestObjectInfo& InObjectInfo)
+{
+	// br, cid, ot, pr, sf, sid, st, su, v
+	Add_br(OutHeaders, InParams, InObjectInfo);
+	Add_cid(OutHeaders, InParams);
+	Add_ot(OutHeaders, InParams, GetOT(EObjectType::InitSegment));
+	double pr = PlayerSessionServices->GetPlayerARBControl()->ABRGetPlaySpeed().GetAsSeconds();
+	Add_pr(OutHeaders, InParams, pr);
+	Add_sf(OutHeaders, InParams);
+	Add_sid(OutHeaders, InParams);
+	Add_st(OutHeaders, InParams);
+	Add_su(OutHeaders, InParams);
+	Add_v(OutHeaders, InParams);
+	Add_platform(OutHeaders, InParams);
+}
+
+void FCMCDHandler::ApplyMediaSegment(FReqHeaders& OutHeaders, EObjectType InObject, const FReportElement& InParams, const FRequestObjectInfo& InObjectInfo, const FURL_RFC3986& InRequestHost)
+{
+	// br, bl, bs, cid, d, dl, mtp, nor, nrr, ot, pr, rtp, sf, sid, st, su, tb, v
+	int32 bl = -1;
+	// Get the current buffer level.
+	if (InObject == EObjectType::VideoOnly || InObject == EObjectType::AudioOnly || InObject == EObjectType::MuxedAudioAndVideo)
+	{
+		IPlayerABRLiveControl::FABRBufferStats BufferStats;
+		PlayerSessionServices->GetPlayerARBControl()->ABRGetStreamBufferStats(BufferStats, InObject == EObjectType::AudioOnly ? Electra::EStreamType::Audio : Electra::EStreamType::Video);
+		bl = (((int32)BufferStats.PlayableContentDuration.GetAsMilliseconds() + 50) / 100) * 100;
+		Add_bl(OutHeaders, InParams, bl);
+	}
+	Add_br(OutHeaders, InParams, InObjectInfo);
+	Add_bs(OutHeaders, InObject, InParams);
+	Add_cid(OutHeaders, InParams);
+	Add_d(OutHeaders, InParams, InObjectInfo);
+	// Do not add `dl` for now. We do not want a CDN to use this to throttle delivery.
+	int32 mtp = (int32)(((PlayerSessionServices->GetStreamSelector()->GetAverageBandwidth() + 50000) / 100000) * 100);
+	Add_mtp(OutHeaders, InParams, mtp);
+	Add_ot(OutHeaders, InParams, GetOT(InObject));
+	double pr = PlayerSessionServices->GetPlayerARBControl()->ABRGetPlaySpeed().GetAsSeconds();
+	Add_pr(OutHeaders, InParams, pr);
+	Add_nornrr(OutHeaders, InRequestHost, InParams, InObjectInfo);
+	// Do not add `rtp` for now. We do not want a CDN to use this to throttle delivery.
+	Add_sf(OutHeaders, InParams);
+	Add_sid(OutHeaders, InParams);
+	Add_st(OutHeaders, InParams);
+	Add_su(OutHeaders, InParams);
+	Add_tb(OutHeaders, InParams, InObjectInfo);
+	Add_v(OutHeaders, InParams);
+	Add_platform(OutHeaders, InParams);
+}
+
+const TCHAR* FCMCDHandler::GetOT(EObjectType InObject)
+{
+	switch(InObject)
+	{
+		case EObjectType::ManifestOrPlaylist:		return TEXT("m");
+		case EObjectType::AudioOnly:				return TEXT("a");
+		case EObjectType::VideoOnly:				return TEXT("v");
+		case EObjectType::MuxedAudioAndVideo:		return TEXT("av");
+		case EObjectType::InitSegment:				return TEXT("i");
+		case EObjectType::CaptionOrSubtitle:		return TEXT("c");
+		case EObjectType::TimedTextTrack:			return TEXT("tt");
+		case EObjectType::CryptoKeyLicenseOrCert:	return TEXT("k");
+		case EObjectType::Other:					return TEXT("o");
+		default:									return TEXT("");
+	}
+}
+
+
+void FCMCDHandler::Add_br(FReqHeaders& OutHeaders, const FReportElement& InParams, const FRequestObjectInfo& InObjectInfo)
+{
+	// Note: in the work-in-progress spec for v2 this must not be sent for `ot` other than `a`, `v`, `av` and `o`
+	if (InParams.IsKeyEnabled(EV1Keys::k_br) && InObjectInfo.EncodedBitrate.Get(0) > 0)
+	{
+		OutHeaders.Object.Emplace(FString::Printf(TEXT("br=%d"), InObjectInfo.EncodedBitrate.GetValue() / 1000));
+	}
+}
+
+void FCMCDHandler::Add_bl(FReqHeaders& OutHeaders, const FReportElement& InParams, int32 InBufferLength)
+{
+	if (InBufferLength >= 0 && InParams.IsKeyEnabled(EV1Keys::k_bl))
+	{
+		OutHeaders.Request.Emplace(FString::Printf(TEXT("bl=%d"), InBufferLength));
+	}
+}
+
+void FCMCDHandler::Add_bs(FReqHeaders& OutHeaders, EObjectType InObject, const FReportElement& InParams)
+{
+	bool bWasStarved = false;
+	if (InObject == EObjectType::VideoOnly || InObject == EObjectType::MuxedAudioAndVideo)
+	{
+		bWasStarved = Vars.bVideoStarved;
+		Vars.bVideoStarved = false;
+	}
+	else if (InObject == EObjectType::AudioOnly)
+	{
+		bWasStarved = Vars.bAudioStarved;
+		Vars.bAudioStarved = false;
+	}
+	if (bWasStarved && InParams.IsKeyEnabled(EV1Keys::k_bs))
+	{
+		OutHeaders.Status.Emplace(TEXT("bs"));
+	}
+}
+
+void FCMCDHandler::Add_cid(FReqHeaders& OutHeaders, const FReportElement& InParams)
+{
+	if (InParams.IsKeyEnabled(EV1Keys::k_cid) && InParams.ContentID.Len())
+	{
+		OutHeaders.Session.Emplace(FString::Printf(TEXT("cid=\"%s\""), *InParams.ContentID));
+	}
+}
+
+void FCMCDHandler::Add_d(FReqHeaders& OutHeaders, const FReportElement& InParams, const FRequestObjectInfo& InObjectInfo)
+{
+	if (InObjectInfo.ObjectDuration.IsSet() && InParams.IsKeyEnabled(EV1Keys::k_d))
+	{
+		OutHeaders.Object.Emplace(FString::Printf(TEXT("d=%d"), InObjectInfo.ObjectDuration.GetValue()));
+	}
+}
+
+void FCMCDHandler::Add_mtp(FReqHeaders& OutHeaders, const FReportElement& InParams, int32 InMeasuredThroughput)
+{
+	if (InMeasuredThroughput > 0 && InParams.IsKeyEnabled(EV1Keys::k_mtp))
+	{
+		OutHeaders.Request.Emplace(FString::Printf(TEXT("mtp=%d"), InMeasuredThroughput));
+	}
+}
+
+void FCMCDHandler::Add_nornrr(FReqHeaders& OutHeaders, const FURL_RFC3986& InRequestHost, const FReportElement& InParams, const FRequestObjectInfo& InObjectInfo)
+{
+	if (InParams.Version == 1 && InObjectInfo.NextObjectRequest.Num())
+	{
+		// Next segment is a range within the current segment?
+		if (InObjectInfo.NextObjectRequest[0].URL.IsEmpty() && !InObjectInfo.NextObjectRequest[0].Range.IsEmpty() && InParams.IsKeyEnabled(EV1Keys::k_nrr))
+		{
+			OutHeaders.Request.Emplace(FString::Printf(TEXT("nrr=%s"), *InObjectInfo.NextObjectRequest[0].Range));
+		}
+		// Next segment is a full segment?
+		else if (!InObjectInfo.NextObjectRequest[0].URL.IsEmpty() && InObjectInfo.NextObjectRequest[0].Range.IsEmpty() && InParams.IsKeyEnabled(EV1Keys::k_nor))
+		{
+			FString Rel, EscRel;
+			if (!InRequestHost.MakeRelativePath(Rel, InObjectInfo.NextObjectRequest[0].URL) ||
+				!FURL_RFC3986::UrlEncode(EscRel, Rel, FString()))
+			{
+				return;
+			}
+			if (EscRel.Len())
+			{
+				OutHeaders.Request.Emplace(FString::Printf(TEXT("nor=\"%s\""), *EscRel));
+			}
+		}
+		// Next segment is a part of a segment?
+		else if (!InObjectInfo.NextObjectRequest[0].URL.IsEmpty() && !InObjectInfo.NextObjectRequest[0].Range.IsEmpty() && InParams.IsKeyEnabled(EV1Keys::k_nor) && InParams.IsKeyEnabled(EV1Keys::k_nrr))
+		{
+			FString Rel, EscRel;
+			if (!InRequestHost.MakeRelativePath(Rel, InObjectInfo.NextObjectRequest[0].URL) ||
+				!FURL_RFC3986::UrlEncode(EscRel, Rel, FString()))
+			{
+				return;
+			}
+			if (EscRel.Len())
+			{
+				OutHeaders.Request.Emplace(FString::Printf(TEXT("nor=\"%s\""), *EscRel));
+			}
+			OutHeaders.Request.Emplace(FString::Printf(TEXT("nrr=%s"), *InObjectInfo.NextObjectRequest[0].Range));
+		}
+	}
+}
+
+void FCMCDHandler::Add_ot(FReqHeaders& OutHeaders, const FReportElement& InParams, const TCHAR* InOT)
+{
+	if (InParams.IsKeyEnabled(EV1Keys::k_ot))
+	{
+		OutHeaders.Object.Emplace(FString::Printf(TEXT("ot=%s"), InOT));
+	}
+}
+
+void FCMCDHandler::Add_pr(FReqHeaders& OutHeaders, const FReportElement& InParams, double InPlayRate)
+{
+	if (InPlayRate != 1.0 && InParams.IsKeyEnabled(EV1Keys::k_pr))
+	{
+		OutHeaders.Session.Emplace(FString::Printf(TEXT("pr=%s"), *FString::SanitizeFloat(InPlayRate, 0)));
+	}
+}
+
+void FCMCDHandler::Add_sf(FReqHeaders& OutHeaders, const FReportElement& InParams)
+{
+	if (InParams.IsKeyEnabled(EV1Keys::k_sf))
+	{
+		switch(StreamingFormat)
+		{
+			case EStreamingFormat::HLS:
+			{
+				OutHeaders.Session.Emplace(TEXT("sf=h"));
+				break;
+			}
+			case EStreamingFormat::DASH:
+			{
+				OutHeaders.Session.Emplace(TEXT("sf=d"));
+				break;
+			}
+			case EStreamingFormat::LowLatencyDASH:
+			{
+				if (InParams.Version == 1)
+				{
+					OutHeaders.Session.Emplace(TEXT("sf=d"));
+				}
+				else //if (InParams.Version == 2)
+				{
+					OutHeaders.Session.Emplace(TEXT("sf=ld"));
+				}
+				break;
+			}
+			default:
+			{
+				unimplemented();
+				break;
+			}
+		}
+	}
+}
+
+void FCMCDHandler::Add_sid(FReqHeaders& OutHeaders, const FReportElement& InParams)
+{
+	if (InParams.IsKeyEnabled(EV1Keys::k_sid) && InParams.SessionID.Len())
+	{
+		OutHeaders.Session.Emplace(FString::Printf(TEXT("sid=\"%s\""), *InParams.SessionID));
+	}
+}
+
+void FCMCDHandler::Add_st(FReqHeaders& OutHeaders, const FReportElement& InParams)
+{
+	if (InParams.IsKeyEnabled(EV1Keys::k_st))
+	{
+		if (StreamType == EStreamType::VOD)
+		{
+			OutHeaders.Session.Emplace(TEXT("st=v"));
+		}
+		else if (StreamType == EStreamType::Live)
+		{
+			OutHeaders.Session.Emplace(TEXT("st=l"));
+		}
+	}
+}
+
+void FCMCDHandler::Add_su(FReqHeaders& OutHeaders, const FReportElement& InParams)
+{
+	if (Vars.bIsStartup && InParams.IsKeyEnabled(EV1Keys::k_su))
+	{
+		OutHeaders.Request.Emplace(TEXT("su"));
+	}
+}
+
+void FCMCDHandler::Add_tb(FReqHeaders& OutHeaders, const FReportElement& InParams, const FRequestObjectInfo& InObjectInfo)
+{
+	if (InParams.IsKeyEnabled(EV1Keys::k_tb) && InObjectInfo.MediaStreamType.IsSet() && Bandwidths.Contains(InObjectInfo.MediaStreamType.GetValue()))
+	{
+		if (int32 bw = Bandwidths[InObjectInfo.MediaStreamType.GetValue()].Highest)
+		{
+			OutHeaders.Object.Emplace(FString::Printf(TEXT("tb=%d"), bw / 1000));
+		}
+	}
+}
+
+void FCMCDHandler::Add_v(FReqHeaders& OutHeaders, const FReportElement& InParams)
+{
+	if (InParams.Version > 1 && InParams.IsKeyEnabled(EV1Keys::k_v))
+	{
+		OutHeaders.Session.Emplace(FString::Printf(TEXT("v=%d"), InParams.Version));
+	}
+}
+
+void FCMCDHandler::Add_platform(FReqHeaders& OutHeaders, const FReportElement& InParams)
+{
+	if (PlatformKey.Len() && PlatformID.Len())
+	{
+		OutHeaders.Session.Emplace(FString::Printf(TEXT("%s=\"%s\""), *PlatformKey, *PlatformID));
+	}
+}
+
+
+
+void FCMCDHandler::ReportOpenSource(const FString& InURL) { }
+void FCMCDHandler::ReportReceivedMainPlaylist(const FString& InEffectiveURL) { }
+void FCMCDHandler::ReportReceivedPlaylists() { }
+void FCMCDHandler::ReportTracksChanged() { }
+void FCMCDHandler::ReportPlaylistDownload(const Metrics::FPlaylistDownloadStats& InPlaylistDownloadStats) { }
+void FCMCDHandler::ReportCleanStart() { }
+void FCMCDHandler::ReportBufferingStart(Metrics::EBufferingReason InBufferingReason)
+{
+	if (InBufferingReason == Metrics::EBufferingReason::Seeking)
+	{
+		Vars.bIsStartup = true;
+	}
+	else if (InBufferingReason == Metrics::EBufferingReason::Rebuffering)
+	{
+		Vars.bIsStartup = true;
+
+		IPlayerABRLiveControl::FABRBufferStats bsV, bsA;
+		PlayerSessionServices->GetPlayerARBControl()->ABRGetStreamBufferStats(bsV, Electra::EStreamType::Video);
+		PlayerSessionServices->GetPlayerARBControl()->ABRGetStreamBufferStats(bsA, Electra::EStreamType::Audio);
+		Vars.bVideoStarved = bsV.PlayableContentDuration.GetAsMilliseconds() < 100;
+		Vars.bAudioStarved = bsA.PlayableContentDuration.GetAsMilliseconds() < 100;
+	}
+}
+
+void FCMCDHandler::ReportBufferingEnd(Metrics::EBufferingReason InBufferingReason)
+{
+	// After any type of buffering this is not the startup phase any more.
+	Vars.bIsStartup = false;
+}
+
+void FCMCDHandler::ReportBandwidth(int64 InEffectiveBps, int64 InThroughputBps, double InLatencyInSeconds) { }
+void FCMCDHandler::ReportBufferUtilization(const Metrics::FBufferStats& InBufferStats) { }
+void FCMCDHandler::ReportSegmentDownload(const Metrics::FSegmentDownloadStats& InSegmentDownloadStats) { }
+void FCMCDHandler::ReportLicenseKey(const Metrics::FLicenseKeyStats& InLicenseKeyStats) { }
+void FCMCDHandler::ReportVideoQualityChange(int32 InNewBitrate, int32 InPreviousBitrate, bool bInIsDrasticDownswitch) { }
+void FCMCDHandler::ReportAudioQualityChange(int32 InNewBitrate, int32 InPreviousBitrate, bool bInIsDrasticDownswitch) { }
+void FCMCDHandler::ReportDataAvailabilityChange(const Metrics::FDataAvailabilityChange& InDataAvailability) { }
+void FCMCDHandler::ReportDecodingFormatChange(const FStreamCodecInformation& InNewDecodingFormat) { }
+void FCMCDHandler::ReportPrerollStart() { }
+void FCMCDHandler::ReportPrerollEnd() { }
+void FCMCDHandler::ReportPlaybackStart() { }
+void FCMCDHandler::ReportPlaybackPaused() { }
+void FCMCDHandler::ReportPlaybackResumed() { }
+void FCMCDHandler::ReportPlaybackEnded() { }
+void FCMCDHandler::ReportJumpInPlayPosition(const FTimeValue& InToNewTime, const FTimeValue& InFromTime, Metrics::ETimeJumpReason InTimejumpReason) { }
+void FCMCDHandler::ReportPlaybackStopped() { }
+void FCMCDHandler::ReportSeekCompleted() { }
+void FCMCDHandler::ReportMediaMetadataChanged(TSharedPtrTS<UtilsMP4::FMetadataParser> InMetadata) { }
+void FCMCDHandler::ReportError(const FString& InErrorReason) { }
+void FCMCDHandler::ReportLogMessage(IInfoLog::ELevel InLogLevel, const FString& InLogMessage, int64 InPlayerWallclockMilliseconds) { }
+void FCMCDHandler::ReportDroppedVideoFrame() { }
+void FCMCDHandler::ReportDroppedAudioFrame() { }
+
+} // namespace Electra
diff --git a/Engine/Plugins/Media/ElectraPlayer/Source/ElectraPlayerRuntime/Private/Runtime/Player/CMCDHandler.h b/Engine/Plugins/Media/ElectraPlayer/Source/ElectraPlayerRuntime/Private/Runtime/Player/CMCDHandler.h
new file mode 100644
--- /dev/null
+++ b/Engine/Plugins/Media/ElectraPlayer/Source/ElectraPlayerRuntime/Private/Runtime/Player/CMCDHandler.h
@@ -0,0 +1,430 @@
+// Copyright Epic Games, Inc. All Rights Reserved.
+
+#pragma once
+
+#include "CoreTypes.h"
+#include "HAL/CriticalSection.h"
+#include "Containers/Array.h"
+#include "Containers/UnrealString.h"
+#include "Templates/SharedPointer.h"
+#include "Internationalization/Regex.h"
+#include "HTTP/HTTPManager.h"
+#include "MediaURLType.h"
+#include "ParameterDictionary.h"
+#include "Utilities/URLParser.h"
+#include "StreamTypes.h"
+#include "Player/AdaptiveStreamingPlayerMetrics.h"
+
+
+namespace Electra
+{
+
+class IPlayerSessionServices;
+
+
+
+/**
+ * Common Media Client Data (CMCD) Handler
+ *
+ * This handler implements v1 of CTA-5004.
+ *
+ * Version 2 is still under development and changes often.
+ * We recognize elements here but do not process them yet.
+ *
+ */
+class FCMCDHandler : public IAdaptiveStreamingPlayerMetrics
+{
+public:
+	FCMCDHandler(IPlayerSessionServices* InPlayerSessionServices);
+	virtual ~FCMCDHandler();
+
+	enum class EStreamingFormat
+	{
+		// CMCD v1 defines HLS, DASH and SmoothStreaming, while CMCD v2 also adds
+		// low latency flavors. Upon requesting the main playlist it is not clear however if the
+		// stream will be low latency or not, so the `ot=m` type can't reflect the low latency
+		// version without externally provided means.
+		Undefined,				// not known, used as a placeholder to omit the object type in requests
+		HLS,
+		DASH,
+		SmoothStreaming,		// not supported
+		// v2
+		HESP,					// not supported
+		LowLatencyDASH,
+		LowLatencyLowDelayDASH,
+		LowLatencyHLS,
+
+		Other
+	};
+
+	enum class EStreamType
+	{
+		Undefined,
+		VOD,
+		Live
+	};
+
+	enum class EObjectType
+	{
+		Undefined,				// not known, used as a placeholder to omit the object type in requests
+		ManifestOrPlaylist,		// 'm'
+		AudioOnly,				// 'a'
+		VideoOnly,				// 'v'
+		MuxedAudioAndVideo,		// 'av'
+		InitSegment,			// 'i'		(probably also to be used for the index segment)
+		CaptionOrSubtitle,		// 'c'
+		TimedTextTrack,			// 'tt'
+		CryptoKeyLicenseOrCert,	// 'k'
+		Other					// 'o'
+	};
+
+	enum class ERequestType
+	{
+		// filtered by `@includeInRequests`: "segment"
+		InitSegment,
+		IndexSegment,
+		Segment,
+		// filtered by `@includeInRequests`: "playlist"
+		FirstPlaylist,
+		Playlist,
+		PlaylistUpdate,
+		// filtered by `@includeInRequests`: "steering"
+		Steering,
+		// filtered by `@includeInRequests`: "other"
+		Other,					// Callback, Chaining, Fallback, Sideload, XLink
+		// filtered by `@includeInRequests`: "drm"
+		DRM
+	};
+
+	static const TCHAR* IncludeInRequests_Segment() { return TEXT("segment"); }
+	static const TCHAR* IncludeInRequests_FirstPlaylist() { return TEXT("firstplaylist"); }
+	static const TCHAR* IncludeInRequests_Playlist() { return TEXT("playlist"); }
+	static const TCHAR* IncludeInRequests_Steering() { return TEXT("steering"); }
+	static const TCHAR* IncludeInRequests_Other() { return TEXT("other"); }
+	static const TCHAR* IncludeInRequests_DRM() { return TEXT("drm"); }
+
+	static const TCHAR* Key_CmcdArray() { return TEXT("cmcd"); }
+	static const TCHAR* Key_ContentId() { return TEXT("contentID"); }
+	static const TCHAR* Key_SessionId() { return TEXT("sessionID"); }
+	static const TCHAR* Key_Type() { return TEXT("type"); }
+	static const TCHAR* Key_Version() { return TEXT("v"); }
+	static const TCHAR* Key_IncludeInRequests() { return TEXT("include"); }
+	static const TCHAR* Key_CDNs() { return TEXT("cdns"); }
+	static const TCHAR* Key_Mode() { return TEXT("mode"); }
+	static const TCHAR* Key_Keys() { return TEXT("keys"); }
+
+
+	/**
+	 * Returns true if CMCD is enabled for this session, false if not.
+	 */
+	bool IsEnabled();
+
+	/**
+	 *
+	 */
+	void PeriodicHandle();
+
+	/**
+	 * Initialize the CMCD handler after creating it.
+	 */
+	void Initialize(const FString& InContentID, const FString& InSessionID, const FString& InConfiguration);
+
+
+	/**
+	 * Returns `true` when the initial configuration, either implicitly or by stream format defaults,
+	 * allows to use parameters provided by the playlist.
+	 * In that case construct a configuration structure from the parameters in the playlist and update
+	 * the configuration.
+	 */
+	bool UseParametersFromPlaylist();
+
+	/**
+	 * If `UseParametersFromPlaylist()` returns `true` and the streaming format does not have a standardized
+	 * way of conveying CMCD parameters, see if the playlist has a user data element of the name returned by
+	 * this function and if so, assume it contains the CMCD configuration.
+	 */
+	FString UsePlaylistParametersFromKey();
+
+	/**
+	 * Update the initial configuration with this new config.
+	 */
+	void UpdateParameters(const FString& InUpdatedConfiguration);
+
+
+	/**
+	 * Set the streaming format. This should be done at least once before or after calling `Initialize()`.
+	 * If the precise sub-format is only known after loading the playlist, call this method again with the
+	 * final determined format.
+	 */
+	void SetStreamingFormat(EStreamingFormat InStreamingFormat);
+
+	/**
+	 * Set the stream type. This is not bound to change during playback (but it could, if Live turns into VoD)
+	 * and it is simpler to set it once instead of providing it in every `FRequestObjectInfo` (see below).
+	 */
+	void SetStreamType(EStreamType InStreamType);
+
+	/**
+	 * Sets the lowest and highest bitrates of a media type this platform is allowed to play.
+	 * Used for the `tb` (in v1; and later for v2 `tpb`, `lb`, `tab` and `lab`)
+	 */
+	void SetPlayableBitrateRange(Electra::EStreamType InStreamType, int32 InLowestBitrate, int32 InHighestBitrate);
+
+
+	/*
+		Request object information that can be obtained through the metrics listener and
+		do not need to be provided in `FRequestObjectInfo`.
+
+		v1
+			- Buffer length (`bl`)
+			- Buffer starvation (`bs`)
+			- Measured throughput (`mtp`)
+			- Playback rate (`pr`)
+			- Startup (`su`)
+			- Top encoded bitrate (`tb`)
+
+		v2
+			still under development. Fields change periodically.
+			"Response" mode has been removed, so all fields related to it but still present in the document
+			are not relevant unless that mode gets reinstated.
+
+			- Playhead bitrate (`pb`)
+			- Buffer starvation duration (`bsd`)
+			- Live stream latency (`ltc`)
+			- Backgrounded (`bg`)
+			- Sequence Number (`sn`)
+			- State (`sta`)
+			- Timestamp (`ts`)
+			- Top playable bitrate (`tpb`)
+			- Lowest encoded bitrate (`lb`)
+			- Top aggregated encoded bitrate (`tab`)
+			- Lowest aggregated encoded bitrate (`lab`)
+			- Playhead time (`pt`)
+			- Player Error Code (`ec`)
+			- Media Start Delay (`msd`)
+			- Event (`e`)
+			- Dropped Frames (`df`)
+			- Non rendered (`nr`)
+	*/
+
+	struct FRequestObjectInfo
+	{
+		struct FNextObjectRequest
+		{
+			FString URL;
+			FString Range;
+		};
+
+		// Internal media stream type for index and init segment where the `ot` is only `i` and not `a`,`v` or `av`.
+		TOptional<Electra::EStreamType> MediaStreamType;
+
+		// v1
+		TOptional<int32> EncodedBitrate;					// `br`
+		TOptional<int32> ObjectDuration;					// `d`
+		TOptional<int32> Deadline;							// `dl`
+		TArray<FNextObjectRequest> NextObjectRequest;		// `nor` and `nrr`
+		TOptional<int32> RequestedMaximumThroughput;		// `rtp`
+		/*
+			set via SetPlayableBitrateRange()
+		TOptional<int32> TopBitrate;						// `tb`
+		*/
+
+		// v2 (not yet supported)
+		/*
+		TOptional<int32> AggregateEncodedBitrate;			// `ab`
+		TOptional<int32> TargetBufferLength;				// `tbl`
+		*/
+	};
+
+
+
+	/**
+	 * Modifies the request headers or URL of the given request to include the CMCD metrics.
+	 */
+	void SetupRequestObject(ERequestType InRequestType, EObjectType InObject /* 'ot' */, FString& InOutRequestURL, TArray<HTTP::FHTTPHeader>& InOutRequestHeaders, const FString& InCDNId, const FRequestObjectInfo& InObjectInfo);
+
+	/**
+	 * Removes a `CMCD` query parameter from the URL.
+	 */
+	FString RemoveParamsFromURL(FString InURL);
+
+private:
+	void ReportOpenSource(const FString& InURL) override;
+	void ReportReceivedMainPlaylist(const FString& InEffectiveURL) override;
+	void ReportReceivedPlaylists() override;
+	void ReportTracksChanged() override;
+	void ReportPlaylistDownload(const Metrics::FPlaylistDownloadStats& InPlaylistDownloadStats) override;
+	void ReportCleanStart() override;
+	void ReportBufferingStart(Metrics::EBufferingReason InBufferingReason) override;
+	void ReportBufferingEnd(Metrics::EBufferingReason InBufferingReason) override;
+	void ReportBandwidth(int64 InEffectiveBps, int64 InThroughputBps, double InLatencyInSeconds) override;
+	void ReportBufferUtilization(const Metrics::FBufferStats& InBufferStats) override;
+	void ReportSegmentDownload(const Metrics::FSegmentDownloadStats& InSegmentDownloadStats) override;
+	void ReportLicenseKey(const Metrics::FLicenseKeyStats& InLicenseKeyStats) override;
+	void ReportVideoQualityChange(int32 InNewBitrate, int32 InPreviousBitrate, bool bInIsDrasticDownswitch) override;
+	void ReportAudioQualityChange(int32 InNewBitrate, int32 InPreviousBitrate, bool bInIsDrasticDownswitch) override;
+	void ReportDataAvailabilityChange(const Metrics::FDataAvailabilityChange& InDataAvailability) override;
+	void ReportDecodingFormatChange(const FStreamCodecInformation& InNewDecodingFormat) override;
+	void ReportPrerollStart() override;
+	void ReportPrerollEnd() override;
+	void ReportPlaybackStart() override;
+	void ReportPlaybackPaused() override;
+	void ReportPlaybackResumed() override;
+	void ReportPlaybackEnded() override;
+	void ReportJumpInPlayPosition(const FTimeValue& InToNewTime, const FTimeValue& InFromTime, Metrics::ETimeJumpReason InTimejumpReason) override;
+	void ReportPlaybackStopped() override;
+	void ReportSeekCompleted() override;
+	void ReportMediaMetadataChanged(TSharedPtrTS<UtilsMP4::FMetadataParser> InMetadata) override;
+	void ReportError(const FString& InErrorReason) override;
+	void ReportLogMessage(IInfoLog::ELevel InLogLevel, const FString& InLogMessage, int64 InPlayerWallclockMilliseconds) override;
+	void ReportDroppedVideoFrame() override;
+	void ReportDroppedAudioFrame() override;
+
+	enum class EV1Keys : uint8
+	{
+		k_br,
+		k_bl,
+		k_bs,
+		k_d,
+		k_cid,
+		k_dl,
+		k_mtp,
+		k_nor,
+		k_nrr,
+		k_ot,
+		k_pr,
+		k_rtp,
+		k_sf,
+		k_sid,
+		k_st,
+		k_su,
+		k_tb,
+		k_v
+	};
+
+
+	struct FReportElement
+	{
+		enum class EType
+		{
+			Undefined,
+			Request,
+			Event
+		};
+		enum EMode
+		{
+			Undefined,
+			Headers,
+			Query
+		};
+		EType Type = EType::Undefined;
+		int32 Version = 0;
+		TArray<FString> Hosts;
+		TArray<FString> CDNs;
+		TArray<TUniquePtr<FRegexPattern>> HostRegexPatterns;
+		FString ContentID;
+		FString SessionID;
+		EMode Mode = EMode::Undefined;
+		int32 EnabledKeys = 0;
+		int32 EnabledRequestTypes = 0;
+		void EnableKey(EV1Keys InKey)
+		{
+			EnabledKeys |= 1 << (int32)InKey;
+		}
+		void DisableKey(EV1Keys InKey)
+		{
+			EnabledKeys &= ~(1 << (int32)InKey);
+		}
+		bool IsKeyEnabled(EV1Keys InKey) const
+		{
+			return !!(EnabledKeys & (1 << (int32)InKey));
+		}
+
+		void EnableRequestType(ERequestType InType)
+		{
+			EnabledRequestTypes |= 1 << (int32)InType;
+		}
+		void DisableRequestType(ERequestType InType)
+		{
+			EnabledRequestTypes &= ~(1 << (int32)InType);
+		}
+		bool IsRequestTypeEnabled(ERequestType InType) const
+		{
+			return !!(EnabledRequestTypes & (1 << (int32)InType));
+		}
+	};
+
+	struct FBandwidths
+	{
+		int32 Lowest = 0;
+		int32 Highest = 0;
+	};
+
+	IPlayerSessionServices* PlayerSessionServices = nullptr;
+	FString ContentID;
+	FString SessionID;
+	FString PlatformID;
+	FString PlaylistParamImportKey;
+	bool bApplyPlaylistParams = true;
+	bool bInitOk = false;
+
+	EStreamingFormat StreamingFormat = EStreamingFormat::Undefined;
+	EStreamType StreamType = EStreamType::Undefined;
+	TArray<FReportElement> ReportElements;
+	FString PlatformKey;
+	TMap<Electra::EStreamType, FBandwidths> Bandwidths;
+
+	struct FVars
+	{
+		bool bIsStartup = true;
+		bool bVideoStarved = false;
+		bool bAudioStarved = false;
+	};
+	FVars Vars;
+
+	struct FReqHeaders
+	{
+		TArray<FString> Request;
+		TArray<FString> Object;
+		TArray<FString> Status;
+		TArray<FString> Session;
+	};
+
+	const TCHAR* GetOT(EObjectType InObject);
+
+	enum class EInitFrom
+	{
+		GlobalConfig,
+		Playlist
+	};
+	bool InitializeInternal(EInitFrom InInitFrom, const FString& InConfiguration);
+
+	void ApplyManifestOrPlaylist(FReqHeaders& OutHeaders, const FReportElement& InParams, const FRequestObjectInfo& InObjectInfo);
+	void ApplyGeneric(FReqHeaders& OutHeaders, const FReportElement& InParams, const FRequestObjectInfo& InObjectInfo);
+	void ApplyInitSegment(FReqHeaders& OutHeaders, const FReportElement& InParams, const FRequestObjectInfo& InObjectInfo);
+	void ApplyMediaSegment(FReqHeaders& OutHeaders, EObjectType InObject, const FReportElement& InParams, const FRequestObjectInfo& InObjectInfo, const FURL_RFC3986& InRequestHost);
+
+	void ApplyReport(FURL_RFC3986& InOutRequestHost, const FReportElement& InParams, ERequestType InRequestType, EObjectType InObject, FString& InOutRequestURL, TArray<HTTP::FHTTPHeader>& InOutRequestHeaders, const FRequestObjectInfo& InObjectInfo);
+
+	void Add_br(FReqHeaders& OutHeaders, const FReportElement& InParams, const FRequestObjectInfo& InObjectInfo);
+	void Add_bl(FReqHeaders& OutHeaders, const FReportElement& InParams, int32 InBufferLength);
+	void Add_bs(FReqHeaders& OutHeaders, EObjectType InObject, const FReportElement& InParams);
+	void Add_cid(FReqHeaders& OutHeaders, const FReportElement& InParams);
+	void Add_d(FReqHeaders& OutHeaders, const FReportElement& InParams, const FRequestObjectInfo& InObjectInfo);
+//	void Add_dl(FReqHeaders& OutHeaders, const FReportElement& InParams, const FRequestObjectInfo& InObjectInfo);
+	void Add_mtp(FReqHeaders& OutHeaders, const FReportElement& InParams, int32 InMeasuredThroughput);
+	void Add_nornrr(FReqHeaders& OutHeaders, const FURL_RFC3986& InRequestHost, const FReportElement& InParams, const FRequestObjectInfo& InObjectInfo);
+	void Add_ot(FReqHeaders& OutHeaders, const FReportElement& InParams, const TCHAR* InOT);
+	void Add_pr(FReqHeaders& OutHeaders, const FReportElement& InParams, double InPlayRate);
+//	void Add_rtp(FReqHeaders& OutHeaders, const FReportElement& InParams, const FRequestObjectInfo& InObjectInfo);
+	void Add_sf(FReqHeaders& OutHeaders, const FReportElement& InParams);
+	void Add_sid(FReqHeaders& OutHeaders, const FReportElement& InParams);
+	void Add_st(FReqHeaders& OutHeaders, const FReportElement& InParams);
+	void Add_su(FReqHeaders& OutHeaders, const FReportElement& InParams);
+	void Add_tb(FReqHeaders& OutHeaders, const FReportElement& InParams, const FRequestObjectInfo& InObjectInfo);
+	void Add_v(FReqHeaders& OutHeaders, const FReportElement& InParams);
+	void Add_platform(FReqHeaders& OutHeaders, const FReportElement& InParams);
+
+};
+
+} // namespace Electra
diff --git a/Engine/Plugins/Media/ElectraPlayer/Source/ElectraPlayerRuntime/Private/Runtime/Player/DASH/MPDElementsDASH.h b/Engine/Plugins/Media/ElectraPlayer/Source/ElectraPlayerRuntime/Private/Runtime/Player/DASH/MPDElementsDASH.h
--- a/Engine/Plugins/Media/ElectraPlayer/Source/ElectraPlayerRuntime/Private/Runtime/Player/DASH/MPDElementsDASH.h
+++ b/Engine/Plugins/Media/ElectraPlayer/Source/ElectraPlayerRuntime/Private/Runtime/Player/DASH/MPDElementsDASH.h
@@ -96,6 +96,8 @@ public:
 		ServiceDescription,
 		Resync,
 		ContentSteering,
+		ClientDataReporting,
+		CMCDParameters,
 		Other
 	};
 
@@ -1234,6 +1236,68 @@ private:
 
 
 
+/**
+ * DASH-IF proposed CMCD reporting type
+ */
+class FDashMPD_CMCDParametersType : public IDashMPDElement
+{
+	typedef IDashMPDElement Super;
+public:
+	FDashMPD_CMCDParametersType(const TCHAR* Name, const TCHAR* Data) : Super(Name, Data)
+	{
+		Mode = TEXT("query");
+		// Add default value. It will be removed when the @includeInRequests attribute is encountered.
+		// This is valid even for the base UrlQueryInfo which _only_ applies to segments.
+		IncludeInRequests.Add(TEXT("segment"));
+	}
+	virtual ~FDashMPD_CMCDParametersType() = default;
+	EType GetElementType() const override
+	{ return EType::CMCDParameters; }
+
+	int32 GetVersion() const							{ return Version; }
+	const TArray<FString>& GetIncludeInRequests() const	{ return IncludeInRequests; }
+	const TArray<FString>& GetKeys() const				{ return Keys; }
+	const FString& GetMode() const						{ return Mode; }
+	const FString& GetContentID() const					{ return ContentID; }
+	const FString& GetSessionID() const					{ return SessionID; }
+	const FString& GetSchemeIdUri() const				{ return SchemeIdUri; }
+private:
+	bool ProcessAttribute(FManifestParserDASH* Builder, const TCHAR* AttributeName, const TCHAR* AttributeValue) override;
+	// Attributes
+	TArray<FString> Keys;
+	TArray<FString> IncludeInRequests;
+	FString Mode;
+	FString SchemeIdUri;
+	FString ContentID;
+	FString SessionID;
+	int32 Version = 1;
+};
+
+class FDashMPD_ClientDataReportingType : public IDashMPDElement
+{
+	typedef IDashMPDElement Super;
+public:
+	FDashMPD_ClientDataReportingType(const TCHAR* Name, const TCHAR* Data) : Super(Name, Data)
+	{ }
+	virtual ~FDashMPD_ClientDataReportingType() = default;
+	EType GetElementType() const override
+	{ return EType::ClientDataReporting; }
+
+	const TArray<FString>& GetServiceLocations() const	{ return ServiceLocations; }
+	const TArray<FString>& GetAdaptationSets() const	{ return AdaptationSets; }
+
+	const TArray<TSharedPtrTS<FDashMPD_CMCDParametersType>>& GetCMCDParameters_Elements() const	{ return CMCDParameters_Elements; }
+private:
+	bool ProcessElement(FManifestParserDASH* Builder, const TCHAR* ElementName, const TCHAR* ElementData, int32 XmlFileLineNumber) override;
+	bool ProcessAttribute(FManifestParserDASH* Builder, const TCHAR* AttributeName, const TCHAR* AttributeValue) override;
+	TArray<TSharedPtrTS<FDashMPD_CMCDParametersType>> CMCDParameters_Elements;
+	// Attributes
+	TArray<FString> ServiceLocations;
+	TArray<FString> AdaptationSets;
+};
+
+
+
 /**
  * K.4.2.1 - ServiceDescriptionType
  */
@@ -1247,11 +1311,13 @@ public:
 	EType GetElementType() const override
 	{ return EType::ServiceDescription; }
 
-	const TArray<TSharedPtrTS<FDashMPD_DescriptorType>>& GetScopes() const						{ return Scopes; }
-	const TArray<TSharedPtrTS<FDashMPD_LatencyType>>& GetLatencies() const						{ return Latencies; }
-	const TArray<TSharedPtrTS<FDashMPD_PlaybackRateType>>& GetPlaybackRates() const				{ return PlaybackRates; }
-	const TArray<TSharedPtrTS<FDashMPD_OperatingQualityType>>& GetOperatingQualities() const	{ return OperatingQualities; }
-	const TArray<TSharedPtrTS<FDashMPD_OperatingBandwidthType>>& GetOperatingBandwidths() const	{ return OperatingBandwidths; }
+	const TArray<TSharedPtrTS<FDashMPD_DescriptorType>>& GetScopes() const							{ return Scopes; }
+	const TArray<TSharedPtrTS<FDashMPD_LatencyType>>& GetLatencies() const							{ return Latencies; }
+	const TArray<TSharedPtrTS<FDashMPD_PlaybackRateType>>& GetPlaybackRates() const					{ return PlaybackRates; }
+	const TArray<TSharedPtrTS<FDashMPD_OperatingQualityType>>& GetOperatingQualities() const		{ return OperatingQualities; }
+	const TArray<TSharedPtrTS<FDashMPD_OperatingBandwidthType>>& GetOperatingBandwidths() const		{ return OperatingBandwidths; }
+	const TArray<TSharedPtrTS<FDashMPD_ContentSteeringType>>& GetContentSteerings() const			{ return ContentSteerings; }
+	const TArray<TSharedPtrTS<FDashMPD_ClientDataReportingType>>& GetClientDataReportings() const	{ return ClientDataReportings; }
 
 	uint32 GetID() const { return (uint32) ID; }
 
@@ -1265,6 +1331,8 @@ private:
 	TArray<TSharedPtrTS<FDashMPD_PlaybackRateType>> PlaybackRates;
 	TArray<TSharedPtrTS<FDashMPD_OperatingQualityType>> OperatingQualities;
 	TArray<TSharedPtrTS<FDashMPD_OperatingBandwidthType>> OperatingBandwidths;
+	TArray<TSharedPtrTS<FDashMPD_ContentSteeringType>> ContentSteerings;
+	TArray<TSharedPtrTS<FDashMPD_ClientDataReportingType>> ClientDataReportings;
 
 	// Attributes
 	uint64 ID = 0;
diff --git a/Engine/Plugins/Media/ElectraPlayer/Source/ElectraPlayerRuntime/Private/Runtime/Player/DASH/ManifestBuilderDASH.cpp b/Engine/Plugins/Media/ElectraPlayer/Source/ElectraPlayerRuntime/Private/Runtime/Player/DASH/ManifestBuilderDASH.cpp
--- a/Engine/Plugins/Media/ElectraPlayer/Source/ElectraPlayerRuntime/Private/Runtime/Player/DASH/ManifestBuilderDASH.cpp
+++ b/Engine/Plugins/Media/ElectraPlayer/Source/ElectraPlayerRuntime/Private/Runtime/Player/DASH/ManifestBuilderDASH.cpp
@@ -253,18 +253,9 @@ namespace DASHUrlHelpers
 					{
 						nxt.CDN = bu->GetServiceLocation();
 						// An empty serviceLocation on an absolute URL uses the URL as the serviceLocation according to DVB-DASH.
-						if (nxt.CDN.IsEmpty())
+						if (nxt.CDN.IsEmpty() && bIsDVBDash)
 						{
-							if (bIsDVBDash)
-							{
-								nxt.CDN = nxt.URL;
-							}
-							else
-							{
-								// For non-DVB DASH we use the host as serviceLocation, mostly because it is a shorter string
-								// used in debugging...
-								nxt.CDN = nxt.Parsed.GetHost();
-							}
+							nxt.CDN = nxt.URL;
 						}
 						if (bu->GetDVBPriority().IsSet() && bu->GetDVBPriority().Value() < 0)
 						{
diff --git a/Engine/Plugins/Media/ElectraPlayer/Source/ElectraPlayerRuntime/Private/Runtime/Player/DASH/ManifestDASH.cpp b/Engine/Plugins/Media/ElectraPlayer/Source/ElectraPlayerRuntime/Private/Runtime/Player/DASH/ManifestDASH.cpp
--- a/Engine/Plugins/Media/ElectraPlayer/Source/ElectraPlayerRuntime/Private/Runtime/Player/DASH/ManifestDASH.cpp
+++ b/Engine/Plugins/Media/ElectraPlayer/Source/ElectraPlayerRuntime/Private/Runtime/Player/DASH/ManifestDASH.cpp
@@ -22,6 +22,7 @@
 #include "Player/AdaptivePlayerOptionKeynames.h"
 #include "Player/DRM/DRMManager.h"
 #include "Player/ContentSteeringHandler.h"
+#include "Player/CMCDHandler.h"
 
 DECLARE_CYCLE_STAT(TEXT("FRepresentation::FindSegment"), STAT_ElectraPlayer_DASH_FindSegment, STATGROUP_ElectraPlayer);
 
@@ -33,6 +34,9 @@ namespace Electra
 #define ERRCODE_DASH_DRM_ERROR							2000
 
 
+#define MAX_NEXT_SEGMENTS_FOR_CMCD	3
+
+
 namespace DashUtils
 {
 	#define GETPLAYEROPTION(Type, Getter)																						\
@@ -1556,22 +1560,30 @@ IManifest::FResult FDASHPlayPeriod::GetNextOrRetrySegment(TSharedPtrTS<IStreamSe
 	switch(Current->GetType())
 	{
 		case EStreamType::Video:
+		{
 			ActiveAdaptationSetIDByType = ActiveVideoAdaptationSetID;
 			ActiveRepresentationIDByType = ActiveVideoRepresentationID;
 			ActiveQualityIndex = ActiveVideoQualityIndex;
 			break;
+		}
 		case EStreamType::Audio:
+		{
 			ActiveAdaptationSetIDByType = ActiveAudioAdaptationSetID;
 			ActiveRepresentationIDByType = ActiveAudioRepresentationID;
 			ActiveQualityIndex = ActiveAudioQualityIndex;
 			break;
+		}
 		case EStreamType::Subtitle:
+		{
 			ActiveAdaptationSetIDByType = ActiveSubtitleAdaptationSetID;
 			ActiveRepresentationIDByType = ActiveSubtitleRepresentationID;
 			ActiveQualityIndex = ActiveSubtitleQualityIndex;
 			break;
+		}
 		default:
+		{
 			break;
+		}
 	}
 	TSharedPtrTS<FManifestDASHInternal::FAdaptationSet> Adapt = Period->GetAdaptationSetByUniqueID(ActiveAdaptationSetIDByType);
 	TSharedPtrTS<FManifestDASHInternal::FRepresentation> Repr = Adapt.IsValid()? Adapt->GetRepresentationByUniqueID(ActiveRepresentationIDByType) : nullptr;
@@ -2019,7 +2031,7 @@ void FDASHPlayPeriod::HandleRepresentationInitSegmentLoading(const TArray<FInitS
 			if (!PlayerSessionServices->GetEntityCache()->GetCachedEntity(CachedItem, is.InitSegmentInfo.InitializationURL.Url.URL, is.InitSegmentInfo.InitializationURL.Range))
 			{
 				is.LoadRequest = MakeSharedTS<FMPDLoadRequestDASH>();
-				is.LoadRequest->LoadType = FMPDLoadRequestDASH::ELoadType::Segment;
+				is.LoadRequest->LoadType = FMPDLoadRequestDASH::ELoadType::InitSegment;
 				is.LoadRequest->Url.URL = is.InitSegmentInfo.InitializationURL.Url.URL;
 				is.LoadRequest->Range = is.InitSegmentInfo.InitializationURL.Range;
 				if (is.InitSegmentInfo.InitializationURL.CustomHeader.Len())
@@ -2269,12 +2281,9 @@ FManifestDASHInternal::FRepresentation::ESearchResult FManifestDASHInternal::FRe
 			DASHUrlHelpers::FDASHMediaURL RootDocument;
 			RootDocument.URL = MPDRepresentation->GetDocumentURL();
 			// The MPD has no @serviceLocation or the DVB properties. According to DVB-DASH the @serviceLocation is the URL itself and the attributes default to 1.
-			RootDocument.CDN = RootDocument.URL;
-			if (!InPlayerSessionServices->GetContentSteeringHandler()->IsDVBDash())
+			if (InPlayerSessionServices->GetContentSteeringHandler()->IsDVBDash())
 			{
-				FURL_RFC3986 RootDocURL;
-				RootDocURL.Parse(RootDocument.URL);
-				RootDocument.CDN = RootDocURL.GetHost();
+				RootDocument.CDN = RootDocument.URL;
 			}
 			RootDocument.DVBpriority = 1;
 			RootDocument.DVBweight = 1;
@@ -2306,7 +2315,7 @@ FManifestDASHInternal::FRepresentation::ESearchResult FManifestDASHInternal::FRe
 				{
 					// Create the request.
 					TSharedPtrTS<FMPDLoadRequestDASH> LoadReq = MakeSharedTS<FMPDLoadRequestDASH>();
-					LoadReq->LoadType = FMPDLoadRequestDASH::ELoadType::Segment;
+					LoadReq->LoadType = FMPDLoadRequestDASH::ELoadType::IndexSegment;
 					LoadReq->Url.URL = ChosenIndex.MediaURL.URL;
 					LoadReq->Range = RepresentationIndexSegmentURLRange;
 					if (RequestHeader.Len())
@@ -2318,6 +2327,7 @@ FManifestDASHInternal::FRepresentation::ESearchResult FManifestDASHInternal::FRe
 					LoadReq->SegmentStreamType = InSearchOptions.StreamType;
 					LoadReq->SegmentQualityIndex = InSearchOptions.QualityIndex;
 					LoadReq->SegmentQualityIndexMax = InSearchOptions.MaxQualityIndex;
+					LoadReq->Bitrate = Bitrate;
 					LoadReq->CompleteCallback.BindThreadSafeSP(AsShared(), &FManifestDASHInternal::FRepresentation::SegmentIndexDownloadComplete);
 					OutRemoteElementLoadRequests.Emplace(LoadReq);
 					PendingSegmentIndexLoadRequest = MoveTemp(LoadReq);
@@ -2332,7 +2342,7 @@ FManifestDASHInternal::FRepresentation::ESearchResult FManifestDASHInternal::FRe
 						LoadReq->NumRemainingInChain = 1;
 					}
 
-					LoadReq->LoadType = FMPDLoadRequestDASH::ELoadType::Segment;
+					LoadReq->LoadType = FMPDLoadRequestDASH::ELoadType::IndexSegment;
 					LoadReq->Url.URL = ChosenIndex.MediaURL.URL;;
 					LoadReq->Range = RepresentationIndexSegmentURLRange;
 					if (RequestHeader.Len())
@@ -2344,6 +2354,7 @@ FManifestDASHInternal::FRepresentation::ESearchResult FManifestDASHInternal::FRe
 					LoadReq->SegmentStreamType = InSearchOptions.StreamType;
 					LoadReq->SegmentQualityIndex = InSearchOptions.QualityIndex;
 					LoadReq->SegmentQualityIndexMax = InSearchOptions.MaxQualityIndex;
+					LoadReq->Bitrate = Bitrate;
 					LoadReq->CompleteCallback.BindThreadSafeSP(AsShared(), &FManifestDASHInternal::FRepresentation::SegmentIndexDownloadComplete);
 					OutRemoteElementLoadRequests.Emplace(LoadReq);
 					PendingSegmentIndexLoadRequest = MoveTemp(LoadReq);
@@ -2412,12 +2423,9 @@ bool FManifestDASHInternal::FRepresentation::PrepareDownloadURLs(IPlayerSessionS
 	DASHUrlHelpers::FDASHMediaURL RootDocument;
 	RootDocument.URL = MPDRepresentation->GetDocumentURL();
 	// The MPD has no @serviceLocation or the DVB properties. According to DVB-DASH the @serviceLocation is the URL itself and the attributes default to 1.
-	RootDocument.CDN = RootDocument.URL;
-	if (!InPlayerSessionServices->GetContentSteeringHandler()->IsDVBDash())
+	if (InPlayerSessionServices->GetContentSteeringHandler()->IsDVBDash())
 	{
-		FURL_RFC3986 RootDocURL;
-		RootDocURL.Parse(RootDocument.URL);
-		RootDocument.CDN = RootDocURL.GetHost();
+		RootDocument.CDN = RootDocument.URL;
 	}
 	RootDocument.DVBpriority = 1;
 	RootDocument.DVBweight = 1;
@@ -2518,8 +2526,8 @@ bool FManifestDASHInternal::FRepresentation::PrepareDownloadURLs(IPlayerSessionS
 	}
 
 	// Substitute template parameters for the media and init segments.
-	FString MediaTemplateURL = ApplyTemplateStrings(MediaTemplate, InOutSegmentInfo);
-	FString InitTemplateURL = ApplyTemplateStrings(InitializationTemplate, InOutSegmentInfo);
+	FString MediaTemplateURL = ApplyTemplateStrings(MediaTemplate, InOutSegmentInfo.Number, InOutSegmentInfo.Time, InOutSegmentInfo.SubIndex);
+	FString InitTemplateURL = ApplyTemplateStrings(InitializationTemplate, InOutSegmentInfo.Number, InOutSegmentInfo.Time, InOutSegmentInfo.SubIndex);
 
 	// The URL query might need to be changed. Look for the UrlQuery properties.
 	TArray<TSharedPtrTS<FDashMPD_UrlQueryInfoType>> UrlQueries;
@@ -2530,12 +2538,9 @@ bool FManifestDASHInternal::FRepresentation::PrepareDownloadURLs(IPlayerSessionS
 	DASHUrlHelpers::FDASHMediaURL RootDocument;
 	RootDocument.URL = MPDRepresentation->GetDocumentURL();
 	// The MPD has no @serviceLocation or the DVB properties. According to DVB-DASH the @serviceLocation is the URL itself and the attributes default to 1.
-	RootDocument.CDN = RootDocument.URL;
-	if (!InPlayerSessionServices->GetContentSteeringHandler()->IsDVBDash())
+	if (InPlayerSessionServices->GetContentSteeringHandler()->IsDVBDash())
 	{
-		FURL_RFC3986 RootDocURL;
-		RootDocURL.Parse(RootDocument.URL);
-		RootDocument.CDN = RootDocURL.GetHost();
+		RootDocument.CDN = RootDocument.URL;
 	}
 	RootDocument.DVBpriority = 1;
 	RootDocument.DVBweight = 1;
@@ -2557,6 +2562,22 @@ bool FManifestDASHInternal::FRepresentation::PrepareDownloadURLs(IPlayerSessionS
 	InOutSegmentInfo.MediaURL.CustomHeader = ChosenMedia.AdditionalParams.GetValue(DASHUrlHelpers::SteerOption_AnnexIRequestHeader).SafeGetFString();
 	InOutSegmentInfo.MediaURL.SteeringID = ChosenMedia.SteeringID;
 
+	// If there are future segments to be requested to be conveyed via CMCD `nor`, prepare their URLs now.
+	for(int32 nf=0; nf<InOutSegmentInfo.NextMediaCMCDNumberTimes.Num(); ++nf)
+	{
+		FString NextMediaURL = ApplyTemplateStrings(MediaTemplate, InOutSegmentInfo.NextMediaCMCDNumberTimes[nf].Number, InOutSegmentInfo.NextMediaCMCDNumberTimes[nf].Time, 0);
+
+		TArray<FContentSteeringHandler::FCandidateURL> NextSteeringCandidates;
+		PrepareSteeringCandidates(NextSteeringCandidates, InPlayerSessionServices, RootDocument.URL, CandidateURLs, UrlQueries, NextMediaURL);
+		FString DummyMessage;
+		FContentSteeringHandler::FSelectedCandidateURL NextChosenMedia = InPlayerSessionServices->GetContentSteeringHandler()->SelectBestCandidateFrom(DummyMessage, FContentSteeringHandler::ESelectFor::Segment, NextSteeringCandidates);
+		if (NextChosenMedia.MediaURL.URL.IsEmpty())
+		{
+			break;
+		}
+		InOutSegmentInfo.NextMediaURLS.Emplace_GetRef().Url.URL = NextChosenMedia.MediaURL.URL;
+	}
+
 	// Now do the same for the initialization segment
 	if (!InitTemplateURL.IsEmpty())
 	{
@@ -2589,7 +2610,7 @@ bool FManifestDASHInternal::FRepresentation::PrepareDownloadURLs(IPlayerSessionS
 	return true;
 }
 
-FString FManifestDASHInternal::FRepresentation::ApplyTemplateStrings(FString TemplateURL, const FSegmentInformationCommon& InSegmentInfo)
+FString FManifestDASHInternal::FRepresentation::ApplyTemplateStrings(FString TemplateURL, int64 InNumber, int64 InTime, int64 InSubIndex)
 {
 	auto PrintWithWidth = [](int64 Value, int32 Width) -> FString
 	{
@@ -2646,7 +2667,7 @@ FString FManifestDASHInternal::FRepresentation::ApplyTemplateStrings(FString Tem
 				// $Number$ ?
 				else if (token.StartsWith(TEXT("Number")))
 				{
-					NewURL.Append(PrintWithWidth(InSegmentInfo.Number, GetFormatWidth(token.Mid(6))));
+					NewURL.Append(PrintWithWidth(InNumber, GetFormatWidth(token.Mid(6))));
 				}
 				// $Bandwidth$ ?
 				else if (token.StartsWith(TEXT("Bandwidth")))
@@ -2656,12 +2677,12 @@ FString FManifestDASHInternal::FRepresentation::ApplyTemplateStrings(FString Tem
 				// $Time$ ?
 				else if (token.StartsWith(TEXT("Time")))
 				{
-					NewURL.Append(PrintWithWidth(InSegmentInfo.Time, GetFormatWidth(token.Mid(4))));
+					NewURL.Append(PrintWithWidth(InTime, GetFormatWidth(token.Mid(4))));
 				}
 				// $SubNumber$ ?
 				else if (token.StartsWith(TEXT("SubNumber")))
 				{
-					NewURL.Append(PrintWithWidth(InSegmentInfo.SubIndex, GetFormatWidth(token.Mid(9))));
+					NewURL.Append(PrintWithWidth(InSubIndex, GetFormatWidth(token.Mid(9))));
 				}
 				else
 				{
@@ -3071,6 +3092,27 @@ FManifestDASHInternal::FRepresentation::ESearchResult FManifestDASHInternal::FRe
 		OutSegmentInfo.ATO = ATO;
 		OutSegmentInfo.bIsLastInPeriod = CurrentT + CurrentD >= MediaLocalEndTime;
 		OutSegmentInfo.bFrameAccuracyRequired = InSearchOptions.bFrameAccurateSearch;
+		// Look at the next n segment ranges for CMCD `nrr`
+		if (!OutSegmentInfo.bIsLastInPeriod && InPlayerSessionServices->GetCMCDHandler()->IsEnabled())
+		{
+			int64 cbo = OutSegmentInfo.FirstByteOffset + OutSegmentInfo.NumberOfBytes;
+			int64 nlt = CurrentT + CurrentD;
+			for(int32 jj=CurrentN+1; jj<EndNumber && nlt<MediaLocalEndTime; ++jj)
+			{
+				int64 nss = Sidx->GetEntry(jj).Size;
+				int64 nsd = Sidx->GetEntry(jj).SubSegmentDuration;
+				int64 NextSegStart = cbo;
+				int64 NextSegEnd = cbo + nss - 1;
+				OutSegmentInfo.NextMediaURLS.Emplace_GetRef().Range = FString::Printf(TEXT("%lld-%lld"), (long long int)NextSegStart, (long long int)NextSegEnd);
+				if (OutSegmentInfo.NextMediaURLS.Num() >= MAX_NEXT_SEGMENTS_FOR_CMCD)
+				{
+					break;
+				}
+				cbo += nss;
+				nlt += nsd;
+			}
+		}
+
 		CollectInbandEventStreams(InPlayerSessionServices, OutSegmentInfo);
 		SetupProducerReferenceTimeInfo(InPlayerSessionServices, OutSegmentInfo);
 		if (!PrepareDownloadURLs(InPlayerSessionServices, OutSegmentInfo, SegmentBase))
@@ -3314,12 +3356,31 @@ FManifestDASHInternal::FRepresentation::ESearchResult FManifestDASHInternal::FRe
 	OutSegmentInfo.Duration = SegmentDuration;
 	OutSegmentInfo.Number = StartNumber + SegmentNum;
 	OutSegmentInfo.MediaLocalFirstAUTime = OutSegmentInfo.MediaLocalFirstPTS = MediaLocalSearchTime + PTO;
-	OutSegmentInfo.MediaLocalLastAUTime = MediaLocalEndTime + PTO;
+	if ((OutSegmentInfo.MediaLocalLastAUTime = MediaLocalEndTime + PTO) < 0)
+	{
+		OutSegmentInfo.MediaLocalLastAUTime = TNumericLimits<int64>::Max();
+	}
 	OutSegmentInfo.Timescale = MPDTimescale;
 	OutSegmentInfo.bMayBeMissing = SegmentNum + 1 >= MaxSegmentsInPeriod;
 	OutSegmentInfo.bIsLastInPeriod = OutSegmentInfo.bMayBeMissing && InSearchOptions.bHasFollowingPeriod;
 	OutSegmentInfo.bFrameAccuracyRequired = InSearchOptions.bFrameAccurateSearch;
 	OutSegmentInfo.ATO = ATO;
+
+	// Look at the next n segment numbers for CMCD `nro` construction
+	if (!OutSegmentInfo.bIsLastInPeriod && InPlayerSessionServices->GetCMCDHandler()->IsEnabled())
+	{
+		for(uint32 jj=SegmentNum+1; jj<MaxSegmentsInPeriod; ++jj)
+		{
+			auto& nxt = OutSegmentInfo.NextMediaCMCDNumberTimes.Emplace_GetRef();
+			nxt.Number = StartNumber + jj;
+			nxt.Time = PTO + EPTdelta + jj * (int64)SegmentDuration;
+			if (OutSegmentInfo.NextMediaCMCDNumberTimes.Num() >= MAX_NEXT_SEGMENTS_FOR_CMCD)
+			{
+				break;
+			}
+		}
+	}
+
 	CollectInbandEventStreams(InPlayerSessionServices, OutSegmentInfo);
 	SetupProducerReferenceTimeInfo(InPlayerSessionServices, OutSegmentInfo);
 	if (!PrepareDownloadURLs(InPlayerSessionServices, OutSegmentInfo, SegmentTemplate))
@@ -3391,9 +3452,11 @@ FManifestDASHInternal::FRepresentation::ESearchResult FManifestDASHInternal::FRe
 	int64 CurrentN = Selements[0].HaveN ? Selements[0].N : StartNumber;
 	int32 CurrentR = Selements[0].HaveR ? Selements[0].R : 0;
 	int64 CurrentD = Selements[0].D;
+	int32 nIndex = 0;
 	bool bIsCurrentlyAGap = false;
 
 	bool bFound = false;
+	bool bFirstContentMissing = false;
 	// Search for the segment. It is possible that already the first segment has a larger T value than we are searching for.
 	if (CurrentT > MediaLocalSearchTime)
 	{
@@ -3408,13 +3471,14 @@ FManifestDASHInternal::FRepresentation::ESearchResult FManifestDASHInternal::FRe
 			LogMessage(InPlayerSessionServices, IInfoLog::ELevel::Warning, FString::Printf(TEXT("Representation \"%s\" <SegmentTimeline> starts with %.3f seconds of missing content that will be skipped over and might lead to playback issues."), *MPDRepresentation->GetID(), MissingContentDuration));
 		}
 		bFound = true;
+		bFirstContentMissing = true;
 	}
 	else
 	{
 		int64 PreviousD = CurrentD;
 		int64 PreviousN = CurrentN - 1;			// start with -1 so we can test if the current N is the previous+1 !
 		int64 PreviousT = CurrentT - CurrentD;	// same for T
-		for(int32 nIndex=0; !bFound && nIndex<Selements.Num(); ++nIndex)
+		for(nIndex=0; nIndex<Selements.Num(); ++nIndex)
 		{
 			if (!Selements[nIndex].HaveD)
 			{
@@ -3626,9 +3690,15 @@ FManifestDASHInternal::FRepresentation::ESearchResult FManifestDASHInternal::FRe
 					--CurrentN;
 					CurrentD = PreviousD;
 					CurrentT = PreviousT;
+					CurrentR = 0;
 					bFound = true;
 				}
 			}
+
+			if (bFound)
+			{
+				break;
+			}
 		}
 	}
 
@@ -3652,6 +3722,52 @@ FManifestDASHInternal::FRepresentation::ESearchResult FManifestDASHInternal::FRe
 			LogMessage(InPlayerSessionServices, IInfoLog::ELevel::Warning, FString::Printf(TEXT("Representation \"%s\" <SegmentTimeline> gap encountered for needed 't' value of %lld. Replacing with an empty filler segment."), *MPDRepresentation->GetID(), (long long int)CurrentT));
 		}
 		OutSegmentInfo.ATO = ATO;
+
+		// Look at the next n segment numbers for CMCD `nro` construction
+		if (!bFirstContentMissing && !OutSegmentInfo.bIsLastInPeriod && InPlayerSessionServices->GetCMCDHandler()->IsEnabled())
+		{
+			bool bAllDone = false;
+			while(!bAllDone)
+			{
+				while(--CurrentR >= 0)
+				{
+					CurrentT += CurrentD;
+					++CurrentN;
+					if (CurrentT >= MediaLocalEndTime)
+					{
+						bAllDone = true;
+						break;
+					}
+					auto& nxt = OutSegmentInfo.NextMediaCMCDNumberTimes.Emplace_GetRef();
+					nxt.Number = CurrentN;
+					nxt.Time = CurrentT;
+					if (OutSegmentInfo.NextMediaCMCDNumberTimes.Num() >= MAX_NEXT_SEGMENTS_FOR_CMCD)
+					{
+						bAllDone = true;
+						break;
+					}
+				}
+				if (bAllDone)
+				{
+					break;
+				}
+
+				if (++nIndex >= Selements.Num())
+				{
+					bAllDone = true;
+					break;
+				}
+				CurrentD = Selements[nIndex].D;
+				CurrentN = Selements[nIndex].HaveN ? Selements[nIndex].N - 1 : CurrentN;		// -1 since n is incremented right at the start of the loop above
+				CurrentR = (Selements[nIndex].HaveR ? Selements[nIndex].R : 0) + 1;				// +1 due to the predecrement at the start of the loop
+				CurrentT = Selements[nIndex].HaveT ? Selements[nIndex].T - CurrentD : CurrentT;	// subtract one d from the new time as it is added at the start of the loop
+				if (CurrentD == 0)
+				{
+					break;
+				}
+			}
+		}
+
 		CollectInbandEventStreams(InPlayerSessionServices, OutSegmentInfo);
 		SetupProducerReferenceTimeInfo(InPlayerSessionServices, OutSegmentInfo);
 		if (!PrepareDownloadURLs(InPlayerSessionServices, OutSegmentInfo, SegmentTemplate))
@@ -3886,7 +4002,7 @@ void FManifestDASHInternal::FRepresentation::SegmentIndexDownloadComplete(TShare
 			if (LoadRequest->NumRemainingInChain)
 			{
 				TSharedPtrTS<FMPDLoadRequestDASH> NextChainedReq = MakeSharedTS<FMPDLoadRequestDASH>();
-				NextChainedReq->LoadType = FMPDLoadRequestDASH::ELoadType::Segment;
+				NextChainedReq->LoadType = FMPDLoadRequestDASH::ELoadType::InitSegment;
 				NextChainedReq->Url = LoadRequest->Url;
 				NextChainedReq->Range = LoadRequest->Range2;
 				NextChainedReq->Headers = LoadRequest->Headers;
diff --git a/Engine/Plugins/Media/ElectraPlayer/Source/ElectraPlayerRuntime/Private/Runtime/Player/DASH/ManifestParserDASH.cpp b/Engine/Plugins/Media/ElectraPlayer/Source/ElectraPlayerRuntime/Private/Runtime/Player/DASH/ManifestParserDASH.cpp
--- a/Engine/Plugins/Media/ElectraPlayer/Source/ElectraPlayerRuntime/Private/Runtime/Player/DASH/ManifestParserDASH.cpp
+++ b/Engine/Plugins/Media/ElectraPlayer/Source/ElectraPlayerRuntime/Private/Runtime/Player/DASH/ManifestParserDASH.cpp
@@ -104,6 +104,8 @@ DEF_ELEMENT(OperatingBandwidth);
 DEF_ELEMENT(ServiceDescription);
 DEF_ELEMENT(Resync);
 DEF_ELEMENT(ContentSteering);
+DEF_ELEMENT(ClientDataReporting);
+DEF_ELEMENT(CMCDParameters);
 
 const TCHAR* const Element_up_UrlQueryInfo = TEXT("up:UrlQueryInfo");
 const TCHAR* const Element_up_ExtUrlQueryInfo = TEXT("up:ExtUrlQueryInfo");
@@ -237,6 +239,14 @@ DEF_ATTR(defaultServiceLocation);
 DEF_ATTR(queryBeforeStart);
 DEF_ATTR(clientRequirement);
 DEF_ATTR(proxyServerURL);
+DEF_ATTR(serviceLocations);
+DEF_ATTR(adaptationSets);
+DEF_ATTR(version);
+DEF_ATTR(mode);
+DEF_ATTR(keys);
+DEF_ATTR(contentID);
+DEF_ATTR(sessionID);
+
 
 // DVB DASH attributes have a colon in them.
 const TCHAR* const Attr_dvb_priority = TEXT("dvb:priority");
@@ -805,7 +815,7 @@ bool IDashMPDElement::ProcessAttribute(FManifestParserDASH* Builder, const TCHAR
 
 	// Since descriptors are largely application specific we cannot know if their elements are meaningful or not, so we don't show a warning.
 	// Similary for anything we don't really understand we also should not emit warnings.
-	if (GetElementType() == EType::Descriptor || GetElementType() == EType::ContentSteering || GetElementType() == EType::Other)
+	if (GetElementType() == EType::Descriptor || GetElementType() == EType::ContentSteering || GetElementType() == EType::ClientDataReporting || GetElementType() == EType::Other)
 	{
 		bEmitWarning = false;
 	}
@@ -1005,6 +1015,7 @@ int32 FDashMPD_MPDType::ReplaceChildElementWithRemoteEntities(TSharedPtrTS<IDash
 					// having another xlink on it. Also note that this element was brought in remotely.
 					for(int32 j=0; j<NewRootEntities.Periods.Num(); ++j)
 					{
+						NewRootEntities.Periods[j]->SetParentElement(AsShared());
 						NewRootEntities.Periods[j]->GetXLink().LastResolveID = NewResolveID;
 						NewRootEntities.Periods[j]->GetXLink().bWasRemoteEntity = true;
 					}
@@ -1030,6 +1041,7 @@ int32 FDashMPD_MPDType::ReplaceChildElementWithRemoteEntities(TSharedPtrTS<IDash
 					// having another xlink on it. Also note that this element was brought in remotely.
 					for(int32 j=0; j<NewRootEntities.InitializationSets.Num(); ++j)
 					{
+						NewRootEntities.InitializationSets[j]->SetParentElement(AsShared());
 						NewRootEntities.InitializationSets[j]->GetXLink().LastResolveID = NewResolveID;
 						NewRootEntities.InitializationSets[j]->GetXLink().bWasRemoteEntity = true;
 					}
@@ -1284,6 +1296,7 @@ int32 FDashMPD_DescriptorType::ReplaceChildElementWithRemoteEntities(TSharedPtrT
 					// having another xlink on it. Also note that this element was brought in remotely.
 					for(int32 j=0; j<NewRootEntities.UrlQueries.Num(); ++j)
 					{
+						NewRootEntities.UrlQueries[j]->SetParentElement(AsShared());
 						NewRootEntities.UrlQueries[j]->GetXLink().LastResolveID = NewResolveID;
 						NewRootEntities.UrlQueries[j]->GetXLink().bWasRemoteEntity = true;
 						WellKnownDescriptors.Insert(NewRootEntities.UrlQueries[j], i+j);
@@ -1485,6 +1498,7 @@ int32 FDashMPD_PeriodType::ReplaceChildElementWithRemoteEntities(TSharedPtrTS<ID
 					// having another xlink on it. Also note that this element was brought in remotely.
 					for(int32 j=0; j<NewRootEntities.AdaptationSets.Num(); ++j)
 					{
+						NewRootEntities.AdaptationSets[j]->SetParentElement(AsShared());
 						NewRootEntities.AdaptationSets[j]->GetXLink().LastResolveID = NewResolveID;
 						NewRootEntities.AdaptationSets[j]->GetXLink().bWasRemoteEntity = true;
 					}
@@ -1709,6 +1723,7 @@ int32 FDashMPD_RepresentationBaseType::ReplaceChildElementWithRemoteEntities(TSh
 					// having another xlink on it. Also note that this element was brought in remotely.
 					for(int32 j=0; j<NewRootEntities.EventStreams.Num(); ++j)
 					{
+						NewRootEntities.EventStreams[j]->SetParentElement(AsShared());
 						NewRootEntities.EventStreams[j]->GetXLink().LastResolveID = NewResolveID;
 						NewRootEntities.EventStreams[j]->GetXLink().bWasRemoteEntity = true;
 					}
@@ -1947,6 +1962,7 @@ int32 FDashMPD_AdaptationSetType::ReplaceChildElementWithRemoteEntities(TSharedP
 				// having another xlink on it. Also note that this element was brought in remotely.
 				for(int32 j=0; j<NewRootEntities.SegmentLists.Num(); ++j)
 				{
+					NewRootEntities.SegmentLists[j]->SetParentElement(AsShared());
 					NewRootEntities.SegmentLists[j]->GetXLink().LastResolveID = NewResolveID;
 					NewRootEntities.SegmentLists[j]->GetXLink().bWasRemoteEntity = true;
 				}
@@ -2188,6 +2204,72 @@ bool FDashMPD_ContentSteeringType::ProcessAttribute(FManifestParserDASH* Builder
 
 /*********************************************************************************************************************/
 
+bool FDashMPD_ClientDataReportingType::ProcessElement(FManifestParserDASH* Builder, const TCHAR* ElementName, const TCHAR* ElementData, int32 XmlFileLineNumber)
+{
+	if (StringHelpers::StringEquals(ElementName, ELEMENT(CMCDParameters)))
+	{
+		CMCDParameters_Elements.Emplace(MakeSharedTS<FDashMPD_CMCDParametersType>(ElementName, ElementData));
+		return Builder->PushNewElement(CMCDParameters_Elements.Last(), XmlFileLineNumber);
+	}
+	else
+	{
+		return Super::ProcessElement(Builder, ElementName, ElementData, XmlFileLineNumber);
+	}
+}
+bool FDashMPD_ClientDataReportingType::ProcessAttribute(FManifestParserDASH* Builder, const TCHAR* AttributeName, const TCHAR* AttributeValue)
+{
+	if (AttrEquals(serviceLocations))
+	{
+		return ParseStringArray(ServiceLocations, TEXT(" "), Builder, AttributeName, AttributeValue, GetXMLLineNumber());
+	}
+	else if (AttrEquals(adaptationSets))
+	{
+		return ParseStringArray(AdaptationSets, TEXT(" "), Builder, AttributeName, AttributeValue, GetXMLLineNumber());
+	}
+	else
+	{
+		return Super::ProcessAttribute(Builder, AttributeName, AttributeValue);
+	}
+}
+
+bool FDashMPD_CMCDParametersType::ProcessAttribute(FManifestParserDASH* Builder, const TCHAR* AttributeName, const TCHAR* AttributeValue)
+{
+	if (AttrEquals(version))
+	{
+		return ParseIntAttribute(Version, Builder, AttributeName, AttributeValue, GetXMLLineNumber());
+	}
+	else if (AttrEquals(includeInRequests))
+	{
+		return ParseStringArray(IncludeInRequests, TEXT(" "), Builder, AttributeName, AttributeValue, GetXMLLineNumber());
+	}
+	else if (AttrEquals(keys))
+	{
+		return ParseStringArray(Keys, TEXT(" "), Builder, AttributeName, AttributeValue, GetXMLLineNumber());
+	}
+	else if (AttrEquals(mode))
+	{
+		return ParseStringAttribute(Mode, Builder, AttributeName, AttributeValue, GetXMLLineNumber());
+	}
+	else if (AttrEquals(contentID))
+	{
+		return ParseStringAttribute(ContentID, Builder, AttributeName, AttributeValue, GetXMLLineNumber());
+	}
+	else if (AttrEquals(sessionID))
+	{
+		return ParseStringAttribute(SessionID, Builder, AttributeName, AttributeValue, GetXMLLineNumber());
+	}
+	else if (AttrEquals(schemeIdUri))
+	{
+		return ParseStringAttribute(SchemeIdUri, Builder, AttributeName, AttributeValue, GetXMLLineNumber());
+	}
+	else
+	{
+		return Super::ProcessAttribute(Builder, AttributeName, AttributeValue);
+	}
+}
+
+/*********************************************************************************************************************/
+
 bool FDashMPD_ContentPopularityRateType::FPRElement::ProcessAttribute(FManifestParserDASH* Builder, const TCHAR* AttributeName, const TCHAR* AttributeValue)
 {
 	if (AttrEquals(popularityRate))
@@ -2582,6 +2664,17 @@ bool FDashMPD_ServiceDescriptionType::ProcessElement(FManifestParserDASH* Builde
 		OperatingBandwidths.Emplace(MakeSharedTS<FDashMPD_OperatingBandwidthType>(ElementName, ElementData));
 		return Builder->PushNewElement(OperatingBandwidths.Last(), XmlFileLineNumber);
 	}
+	// Content steering and client data reportings are now provided in the service description.
+	else if (StringHelpers::StringEquals(ElementName, ELEMENT(ContentSteering)))
+	{
+		ContentSteerings.Emplace(MakeSharedTS<FDashMPD_ContentSteeringType>(ElementName, ElementData));
+		return Builder->PushNewElement(ContentSteerings.Last(), XmlFileLineNumber);
+	}
+	else if (StringHelpers::StringEquals(ElementName, ELEMENT(ClientDataReporting)))
+	{
+		ClientDataReportings.Emplace(MakeSharedTS<FDashMPD_ClientDataReportingType>(ElementName, ElementData));
+		return Builder->PushNewElement(ClientDataReportings.Last(), XmlFileLineNumber);
+	}
 	else
 	{
 		return Super::ProcessElement(Builder, ElementName, ElementData, XmlFileLineNumber);
@@ -2728,6 +2821,7 @@ int32 FDashMPD_RepresentationType::ReplaceChildElementWithRemoteEntities(TShared
 				// having another xlink on it. Also note that this element was brought in remotely.
 				for(int32 j=0; j<NewRootEntities.SegmentLists.Num(); ++j)
 				{
+					NewRootEntities.SegmentLists[j]->SetParentElement(AsShared());
 					NewRootEntities.SegmentLists[j]->GetXLink().LastResolveID = NewResolveID;
 					NewRootEntities.SegmentLists[j]->GetXLink().bWasRemoteEntity = true;
 				}
diff --git a/Engine/Plugins/Media/ElectraPlayer/Source/ElectraPlayerRuntime/Private/Runtime/Player/DASH/PlaylistReaderDASH.cpp b/Engine/Plugins/Media/ElectraPlayer/Source/ElectraPlayerRuntime/Private/Runtime/Player/DASH/PlaylistReaderDASH.cpp
--- a/Engine/Plugins/Media/ElectraPlayer/Source/ElectraPlayerRuntime/Private/Runtime/Player/DASH/PlaylistReaderDASH.cpp
+++ b/Engine/Plugins/Media/ElectraPlayer/Source/ElectraPlayerRuntime/Private/Runtime/Player/DASH/PlaylistReaderDASH.cpp
@@ -21,6 +21,7 @@
 #include "Player/DASH/PlayerEventDASH.h"
 #include "Player/DASH/PlayerEventDASH_Internal.h"
 #include "Player/ContentSteeringHandler.h"
+#include "Player/CMCDHandler.h"
 #include "ElectraHTTPStream.h"
 #include "Misc/DateTime.h"
 
@@ -103,6 +104,7 @@ private:
 	void EnqueueInitialXLinkRequests();
 
 	void ManifestDownloadCompleted(FResourceLoadRequestPtr Request, bool bSuccess);
+	void UpdateCMCDParamsFromServiceDescriptions(const TArray<TSharedPtrTS<FDashMPD_ServiceDescriptionType>>& InServiceDescriptions);
 	void FinishManifestBuildingAfterTimesync(bool bGotTheTime);
 	void PrepareNextTimeSync();
 	void ManifestUpdateDownloadCompleted(FResourceLoadRequestPtr Request, bool bSuccess);
@@ -468,7 +470,7 @@ void FPlaylistReaderDASH::CheckForMPDUpdate()
 			FURL_RFC3986 RootDocURL;
 			RootDocURL.Parse(RootDocument.URL);
 			// The MPD has no @serviceLocation or the DVB properties. According to DVB-DASH the @serviceLocation is the URL itself and the attributes default to 1.
-			RootDocument.CDN = MPDRoot->IsDVBDash() ? RootDocument.URL : RootDocURL.GetHost();
+			RootDocument.CDN = MPDRoot->IsDVBDash() ? RootDocument.URL : FString();
 			RootDocument.DVBpriority = 1;
 			RootDocument.DVBweight = 1;
 			// Build a list we can pass to the steering handler.
@@ -508,17 +510,9 @@ void FPlaylistReaderDASH::CheckForMPDUpdate()
 				sc.AdditionalParams.Set(DASHUrlHelpers::SteerOption_AnnexIRequestHeader, FVariantValue(AnnexIRequestHeader));
 
 				// If there is no service location then as with DVB-DASH the URL is to be used for that.
-				if (sc.MediaURL.CDN.IsEmpty())
+				if (sc.MediaURL.CDN.IsEmpty() && MPDRoot->IsDVBDash())
 				{
-					if (!MPDRoot->IsDVBDash())
-					{
-						FURL_RFC3986 tmp;tmp.Parse(sc.MediaURL.URL);
-						sc.MediaURL.CDN = tmp.GetHost();
-				}
-					else
-					{
-						sc.MediaURL.CDN = sc.MediaURL.URL;
-					}
+					sc.MediaURL.CDN = sc.MediaURL.URL;
 				}
 			}
 
@@ -895,7 +889,49 @@ void FPlaylistReaderDASH::ExecutePendingRequests(const FTimeValue& TimeNow)
 				Request->Request->URL(Request->Url.URL).Verb(Request->Verb).Range(Request->Range).Headers(Request->Headers).Object(Request);
 				// Set up timeouts and also accepted encodings depending on the type of request.
 				SetupRequestTimeouts(Request);
-				Request->Request->AllowStaticQuery(IAdaptiveStreamingPlayerResourceRequest::EPlaybackResourceType::Playlist);
+				FString CDNId(Request->Url.CDN);
+				switch(Request->GetLoadType())
+				{
+					case FMPDLoadRequestDASH::ELoadType::MPD:
+					{
+						Request->Request->AllowStaticQuery(IAdaptiveStreamingPlayerResourceRequest::EPlaybackResourceType::Playlist);
+						PlayerSessionServices->GetCMCDHandler()->SetupRequestObject(FCMCDHandler::ERequestType::FirstPlaylist, FCMCDHandler::EObjectType::ManifestOrPlaylist, Request->Request->GetRequest()->Parameters.URL, Request->Request->GetRequest()->Parameters.RequestHeaders, CDNId, FCMCDHandler::FRequestObjectInfo());
+						break;
+					}
+					case FMPDLoadRequestDASH::ELoadType::MPDUpdate:
+					{
+						PlayerSessionServices->GetCMCDHandler()->SetupRequestObject(FCMCDHandler::ERequestType::PlaylistUpdate, FCMCDHandler::EObjectType::ManifestOrPlaylist, Request->Request->GetRequest()->Parameters.URL, Request->Request->GetRequest()->Parameters.RequestHeaders, CDNId, FCMCDHandler::FRequestObjectInfo());
+						break;
+					}
+					case FMPDLoadRequestDASH::ELoadType::XLink_Period:
+					case FMPDLoadRequestDASH::ELoadType::XLink_AdaptationSet:
+					case FMPDLoadRequestDASH::ELoadType::XLink_EventStream:
+					case FMPDLoadRequestDASH::ELoadType::XLink_SegmentList:
+					case FMPDLoadRequestDASH::ELoadType::XLink_URLQuery:
+					case FMPDLoadRequestDASH::ELoadType::XLink_InitializationSet:
+					{
+						PlayerSessionServices->GetCMCDHandler()->SetupRequestObject(FCMCDHandler::ERequestType::Other, FCMCDHandler::EObjectType::Other, Request->Request->GetRequest()->Parameters.URL, Request->Request->GetRequest()->Parameters.RequestHeaders, CDNId, FCMCDHandler::FRequestObjectInfo());
+						break;
+					}
+					case FMPDLoadRequestDASH::ELoadType::Steering:
+					{
+						PlayerSessionServices->GetCMCDHandler()->SetupRequestObject(FCMCDHandler::ERequestType::Steering, FCMCDHandler::EObjectType::Other, Request->Request->GetRequest()->Parameters.URL, Request->Request->GetRequest()->Parameters.RequestHeaders, CDNId, FCMCDHandler::FRequestObjectInfo());
+						break;
+					}
+					case FMPDLoadRequestDASH::ELoadType::IndexSegment:
+					case FMPDLoadRequestDASH::ELoadType::InitSegment:
+					{
+						FCMCDHandler::FRequestObjectInfo ri;
+						ri.MediaStreamType = Request->SegmentStreamType;
+						ri.EncodedBitrate = Request->Bitrate;
+						PlayerSessionServices->GetCMCDHandler()->SetupRequestObject(FCMCDHandler::ERequestType::InitSegment, FCMCDHandler::EObjectType::InitSegment, Request->Request->GetRequest()->Parameters.URL, Request->Request->GetRequest()->Parameters.RequestHeaders, CDNId, ri);
+						break;
+					}
+					default:
+					{
+						break;
+					}
+				}
 
 				Request->Request->Callback().BindThreadSafeSP(AsShared(), &FPlaylistReaderDASH::OnHTTPResourceRequestComplete);
 				Request->Request->StartGet(PlayerSessionServices);
@@ -922,6 +958,8 @@ void FPlaylistReaderDASH::SetupRequestTimeouts(FResourceLoadRequestPtr InRequest
 			break;
 		}
 		case FMPDLoadRequestDASH::ELoadType::Segment:
+		case FMPDLoadRequestDASH::ELoadType::IndexSegment:
+		case FMPDLoadRequestDASH::ELoadType::InitSegment:
 		{
 			InRequest->Request->ConnectionTimeout(FTimeValue(FTimeValue::MillisecondsToHNS(1000 * 5)));
 			InRequest->Request->NoDataTimeout(FTimeValue(FTimeValue::MillisecondsToHNS(1000 * 2)));
@@ -973,6 +1011,8 @@ int32 FPlaylistReaderDASH::CheckForRetry(FResourceLoadRequestPtr InRequest, int3
 			return -1;
 		}
 		case FMPDLoadRequestDASH::ELoadType::Segment:
+		case FMPDLoadRequestDASH::ELoadType::IndexSegment:
+		case FMPDLoadRequestDASH::ELoadType::InitSegment:
 		{
 			if (Error < 100 && InRequest->Attempt < 2)
 			{
@@ -1164,7 +1204,6 @@ void FPlaylistReaderDASH::ManifestDownloadCompleted(FResourceLoadRequestPtr Requ
 		}
 
 		FString ETag;
-		FString EffectiveURL = ConnInfo->EffectiveURL;
 		TArray<FElectraHTTPStreamHeader> Headers;
 		InitialTimesync.HttpRequestStart = ConnInfo->RequestStartTime;
 		InitialTimesync.UtcTimeThen = MEDIAutcTime::Current();
@@ -1188,6 +1227,7 @@ void FPlaylistReaderDASH::ManifestDownloadCompleted(FResourceLoadRequestPtr Requ
 		LastErrorDetail = GetXMLResponseString(XML, Request);
 		if (LastErrorDetail.IsOK())
 		{
+			FString EffectiveURL = PlayerSessionServices->GetCMCDHandler()->RemoveParamsFromURL(ConnInfo->EffectiveURL);
 			Builder.Reset(IManifestBuilderDASH::Create(PlayerSessionServices));
 			TSharedPtrTS<FManifestDASHInternal> NewManifest;
 			LastErrorDetail = Builder->BuildFromMPD(NewManifest, XML.GetCharArray().GetData(), EffectiveURL, ETag, Headers);
@@ -1195,6 +1235,9 @@ void FPlaylistReaderDASH::ManifestDownloadCompleted(FResourceLoadRequestPtr Requ
 			{
 				if (NewManifest.IsValid())
 				{
+					// Update the CMCD parameters from the root level service descriptions.
+					UpdateCMCDParamsFromServiceDescriptions(NewManifest->GetMPDRoot()->GetServiceDescriptions());
+
 					// Parse the URL fragment into its components. For a DASH URL the fragment is constructed like a query string with & delimited key/value pairs.
 					TArray<FURL_RFC3986::FQueryParam> URLFragmentComponents;
 					FURL_RFC3986::GetQueryParams(URLFragmentComponents, Fragment, false);	// The fragment is already URL escaped, so no need to do it again.
@@ -1223,6 +1266,98 @@ void FPlaylistReaderDASH::ManifestDownloadCompleted(FResourceLoadRequestPtr Requ
 	}
 }
 
+void FPlaylistReaderDASH::UpdateCMCDParamsFromServiceDescriptions(const TArray<TSharedPtrTS<FDashMPD_ServiceDescriptionType>>& InServiceDescriptions)
+{
+	// Check if the CMCD handler allows for configuration from the MPD.
+	if (!PlayerSessionServices->GetCMCDHandler()->UseParametersFromPlaylist())
+	{
+		return;
+	}
+	TArray<FString> cmcdConfigs;
+	for(int32 i=0; i<InServiceDescriptions.Num(); ++i)
+	{
+		const TArray<TSharedPtrTS<FDashMPD_ClientDataReportingType>>& cdrt = InServiceDescriptions[i]->GetClientDataReportings();
+		for(int32 j=0; j<cdrt.Num(); ++j)
+		{
+			const TArray<FString>& ServiceLocations = cdrt[j]->GetServiceLocations();
+			if (cdrt[j]->GetCMCDParameters_Elements().Num())
+			{
+				// This may be an array, but the spec allows for 1 entry at most.
+				const TSharedPtrTS<FDashMPD_CMCDParametersType>& pt = cdrt[j]->GetCMCDParameters_Elements()[0];
+				// The keys are mandatory, so if they are not present we ignore this entry.
+				if (pt->GetVersion() == 1 && !pt->GetKeys().IsEmpty())
+				{
+					FString cfg(FString::Printf(TEXT("{\"%s\":\"request\",\"%s\":1"), FCMCDHandler::Key_Type(), FCMCDHandler::Key_Version()));
+					cfg.Append(FString::Printf(TEXT(",\"%s\":\"%s\""), FCMCDHandler::Key_Mode(), *pt->GetMode()));
+					if (pt->GetContentID().Len())
+					{
+						cfg.Append(FString::Printf(TEXT(",\"%s\":\"%s\""), FCMCDHandler::Key_ContentId(), *pt->GetContentID()));
+					}
+					if (pt->GetSessionID().Len())
+					{
+						cfg.Append(FString::Printf(TEXT(",\"%s\":\"%s\""), FCMCDHandler::Key_SessionId(), *pt->GetSessionID()));
+					}
+					auto MakeList = [](const TArray<FString>& InList) -> FString
+					{
+						FString r;
+						for(int32 i=0; i<InList.Num(); ++i)
+						{
+							if (i)
+							{
+								r.Append(TEXT(" "));
+							}
+							r.Append(InList[i]);
+						}
+						return r;
+					};
+					cfg.Append(FString::Printf(TEXT(",\"%s\":\"%s\""), FCMCDHandler::Key_Keys(), *MakeList(pt->GetKeys())));
+					if (ServiceLocations.Num())
+					{
+						cfg.Append(FString::Printf(TEXT(",\"%s\":\"%s\""), FCMCDHandler::Key_CDNs(), *MakeList(ServiceLocations)));
+					}
+					TArray<FString> IncludeIn;
+					for(auto& InR : pt->GetIncludeInRequests())
+					{
+						if (InR.Equals(TEXT("segments")))
+						{
+							IncludeIn.Emplace(FCMCDHandler::IncludeInRequests_Segment());
+						}
+						else if (InR.Equals(TEXT("mpd")))
+						{
+							IncludeIn.Emplace(FCMCDHandler::IncludeInRequests_Playlist());
+						}
+						else if (InR.Equals(TEXT("steering")))
+						{
+							IncludeIn.Emplace(FCMCDHandler::IncludeInRequests_Steering());
+						}
+					}
+					if (IncludeIn.Num())
+					{
+						cfg.Append(FString::Printf(TEXT(",\"%s\":\"%s\""), FCMCDHandler::Key_IncludeInRequests(), *MakeList(IncludeIn)));
+					}
+					cfg.Append(TEXT("}"));
+					cmcdConfigs.Emplace(MoveTemp(cfg));
+				}
+			}
+		}
+	}
+	// Build the overall config.
+	if (cmcdConfigs.Num())
+	{
+		FString cfg(FString::Printf(TEXT("{\"%s\":["), FCMCDHandler::Key_CmcdArray()));
+		for(int32 i=0; i<cmcdConfigs.Num(); ++i)
+		{
+			if (i)
+			{
+				cfg.Append(TEXT(","));
+			}
+			cfg.Append(cmcdConfigs[i]);
+		}
+		cfg.Append(TEXT("]}"));
+		PlayerSessionServices->GetCMCDHandler()->UpdateParameters(cfg);
+	}
+}
+
 
 void FPlaylistReaderDASH::PrepareNextTimeSync()
 {
diff --git a/Engine/Plugins/Media/ElectraPlayer/Source/ElectraPlayerRuntime/Private/Runtime/Player/DASH/PlaylistReaderDASH_Internal.h b/Engine/Plugins/Media/ElectraPlayer/Source/ElectraPlayerRuntime/Private/Runtime/Player/DASH/PlaylistReaderDASH_Internal.h
--- a/Engine/Plugins/Media/ElectraPlayer/Source/ElectraPlayerRuntime/Private/Runtime/Player/DASH/PlaylistReaderDASH_Internal.h
+++ b/Engine/Plugins/Media/ElectraPlayer/Source/ElectraPlayerRuntime/Private/Runtime/Player/DASH/PlaylistReaderDASH_Internal.h
@@ -37,6 +37,8 @@ struct FMPDLoadRequestDASH : public IHTTPResourceRequestObject
 		XLink_InitializationSet,
 		Callback,
 		Segment,
+		IndexSegment,
+		InitSegment,
 		TimeSync,
 		Sideload,
 		Steering
@@ -55,6 +57,8 @@ struct FMPDLoadRequestDASH : public IHTTPResourceRequestObject
 			case ELoadType::XLink_InitializationSet:return TEXT("remote InitializationSet");
 			case ELoadType::Callback:				return TEXT("Callback");
 			case ELoadType::Segment:				return TEXT("Segment");
+			case ELoadType::IndexSegment:			return TEXT("IndexSegment");
+			case ELoadType::InitSegment:			return TEXT("InitSegment");
 			case ELoadType::TimeSync:				return TEXT("Time sync");
 			case ELoadType::Sideload:				return TEXT("Sideload");
 			default:								return TEXT("<unknown>");
@@ -87,6 +91,7 @@ struct FMPDLoadRequestDASH : public IHTTPResourceRequestObject
 	EStreamType SegmentStreamType = EStreamType::Unsupported;
 	int32 SegmentQualityIndex = 0;
 	int32 SegmentQualityIndexMax = 0;
+	int32 Bitrate = 0;
 
 	IPlayerSessionServices* PlayerSessionServices = nullptr;
 	TSharedPtrTS<FHTTPResourceRequest> Request;
@@ -245,7 +250,7 @@ public:
 		void PrepareSteeringCandidates(TArray<FContentSteeringHandler::FCandidateURL>& OutSteeringCandidates, IPlayerSessionServices* InPlayerSessionServices, const FString& InDocumentURL, const TArray<DASHUrlHelpers::FDASHMediaURL>& InCandidateURLs, const TArray<TSharedPtrTS<FDashMPD_UrlQueryInfoType>>& InUrlQueries, const FString& InTemplateURL);
 		bool PrepareDownloadURLs(IPlayerSessionServices* PlayerSessionServices, FSegmentInformationCommon& InOutSegmentInfo, const TArray<TSharedPtrTS<FDashMPD_SegmentBaseType>>& SegmentBase);
 		bool PrepareDownloadURLs(IPlayerSessionServices* PlayerSessionServices, FSegmentInformationCommon& InOutSegmentInfo, const TArray<TSharedPtrTS<FDashMPD_SegmentTemplateType>>& SegmentTemplate);
-		FString ApplyTemplateStrings(FString TemplateURL, const FSegmentInformationCommon& InSegmentInfo);
+		FString ApplyTemplateStrings(FString TemplateURL, int64 InNumber, int64 InTime, int64 InSubIndex);
 		void CollectInbandEventStreams(IPlayerSessionServices* PlayerSessionServices, FSegmentInformationCommon& InOutSegmentInfo);
 		void SetupProducerReferenceTimeInfo(IPlayerSessionServices* PlayerSessionServices, FSegmentInformationCommon& InOutSegmentInfo);
 
diff --git a/Engine/Plugins/Media/ElectraPlayer/Source/ElectraPlayerRuntime/Private/Runtime/Player/HLS/PlaylistHLS.cpp b/Engine/Plugins/Media/ElectraPlayer/Source/ElectraPlayerRuntime/Private/Runtime/Player/HLS/PlaylistHLS.cpp
--- a/Engine/Plugins/Media/ElectraPlayer/Source/ElectraPlayerRuntime/Private/Runtime/Player/HLS/PlaylistHLS.cpp
+++ b/Engine/Plugins/Media/ElectraPlayer/Source/ElectraPlayerRuntime/Private/Runtime/Player/HLS/PlaylistHLS.cpp
@@ -1930,6 +1930,21 @@ FActiveHLSPlaylist::FTimelineMediaAsset::ESegSearchResult FActiveHLSPlaylist::FT
 	OutSegment->Segment.bIsLastInPeriod = (mp->bHasEndList || mp->PlaylistType == FPlaylistParserHLS::EPlaylistType::VOD) && SelectedSegmentIndex+1 >= Segments.Num();
 	OutSegment->bIsGapSegment = OutSegment->Segment.bMayBeMissing = !!Seg.bGap;
 
+	// Add next segments for CMCD?
+	if (!OutSegment->Segment.bIsLastInPeriod)
+	{
+		const int32 kMaxNextSegmentsForCMCD = 3;
+		while(++SelectedSegmentIndex < Segments.Num() && OutSegment->Segment.NextMediaURLS.Num() < kMaxNextSegmentsForCMCD)
+		{
+			const FMediaSegmentHLS& NextSeg(Segments[SelectedSegmentIndex]);
+			FSegmentInformationCommon::FURL& Nxt = OutSegment->Segment.NextMediaURLS.Emplace_GetRef();
+			FURL_RFC3986 NxtUrl(mp->ParsedURL);
+			NxtUrl.ResolveWith(NextSeg.URL);
+			Nxt.Url.URL = NxtUrl.Get(true, true);
+			Nxt.Range = NextSeg.ByteRange.GetForHTTP();
+		}
+	}
+
 	return ESegSearchResult::Found;
 }
 
diff --git a/Engine/Plugins/Media/ElectraPlayer/Source/ElectraPlayerRuntime/Private/Runtime/Player/HLS/PlaylistHLS.h b/Engine/Plugins/Media/ElectraPlayer/Source/ElectraPlayerRuntime/Private/Runtime/Player/HLS/PlaylistHLS.h
--- a/Engine/Plugins/Media/ElectraPlayer/Source/ElectraPlayerRuntime/Private/Runtime/Player/HLS/PlaylistHLS.h
+++ b/Engine/Plugins/Media/ElectraPlayer/Source/ElectraPlayerRuntime/Private/Runtime/Player/HLS/PlaylistHLS.h
@@ -480,6 +480,7 @@ struct FLoadRequestHLSPlaylist
 	int32 Attempt = 1;
 	bool bIsPrimaryPlaylist = false;
 	bool bIsPreStartSteering = false;
+	bool bAddedCMCDParameters = false;
 
 	FMediaPlaylistInformationHLS PlaylistInfo;
 
diff --git a/Engine/Plugins/Media/ElectraPlayer/Source/ElectraPlayerRuntime/Private/Runtime/Player/HLS/PlaylistHandlerHLS.cpp b/Engine/Plugins/Media/ElectraPlayer/Source/ElectraPlayerRuntime/Private/Runtime/Player/HLS/PlaylistHandlerHLS.cpp
--- a/Engine/Plugins/Media/ElectraPlayer/Source/ElectraPlayerRuntime/Private/Runtime/Player/HLS/PlaylistHandlerHLS.cpp
+++ b/Engine/Plugins/Media/ElectraPlayer/Source/ElectraPlayerRuntime/Private/Runtime/Player/HLS/PlaylistHandlerHLS.cpp
@@ -6,12 +6,14 @@
 #include "Player/AdaptiveStreamingPlayerABR.h"
 #include "Player/AdaptiveStreamingPlayerResourceRequest.h"
 #include "Player/ContentSteeringHandler.h"
+#include "Player/CMCDHandler.h"
 #include "Utilities/Utilities.h"
 #include "Utilities/URLParser.h"
 #include "Utilities/TimeUtilities.h"
 #include "ElectraPlayerMisc.h"
 #include "StreamTypes.h"
 #include "Misc/SecureHash.h"
+#include "Misc/Base64.h"
 #include "Internationalization/Regex.h"
 
 namespace Electra
@@ -238,6 +240,34 @@ void FPlaylistHandlerHLS::HandleOnce()
 	{
 		if (!PendingLoadRequests[i]->ExecuteAtUTC.IsValid() || PendingLoadRequests[i]->ExecuteAtUTC <= Now)
 		{
+			IElectraHttpManager::FParams& Params(PendingLoadRequests[i]->ResourceRequest->GetRequest()->Parameters);
+			FCMCDHandler::FRequestObjectInfo ri;
+			FString RequestedURL(Params.URL);
+			FString PathwayID(PendingLoadRequests[i]->PlaylistInfo.PathwayID);
+			switch(PendingLoadRequests[i]->LoadType)
+			{
+				case FLoadRequestHLSPlaylist::ELoadType::Main:
+				{
+					PlayerSessionServices->GetCMCDHandler()->SetupRequestObject(FCMCDHandler::ERequestType::FirstPlaylist, FCMCDHandler::EObjectType::ManifestOrPlaylist, Params.URL, Params.RequestHeaders, PathwayID, ri);
+					break;
+				}
+				case FLoadRequestHLSPlaylist::ELoadType::InitialVariant:
+				{
+					PlayerSessionServices->GetCMCDHandler()->SetupRequestObject(FCMCDHandler::ERequestType::Playlist, FCMCDHandler::EObjectType::ManifestOrPlaylist, Params.URL, Params.RequestHeaders, PathwayID, ri);
+					break;
+				}
+				case FLoadRequestHLSPlaylist::ELoadType::Variant:
+				{
+					PlayerSessionServices->GetCMCDHandler()->SetupRequestObject(FCMCDHandler::ERequestType::PlaylistUpdate, FCMCDHandler::EObjectType::ManifestOrPlaylist, Params.URL, Params.RequestHeaders, PathwayID, ri);
+					break;
+				}
+				case FLoadRequestHLSPlaylist::ELoadType::Steering:
+				{
+					PlayerSessionServices->GetCMCDHandler()->SetupRequestObject(FCMCDHandler::ERequestType::Steering, FCMCDHandler::EObjectType::Other, Params.URL, Params.RequestHeaders, PathwayID, ri);
+					break;
+				}
+			}
+			PendingLoadRequests[i]->bAddedCMCDParameters = !RequestedURL.Equals(Params.URL);
 			PendingLoadRequests[i]->ResourceRequest->StartGet(PlayerSessionServices);
 			RunningLoadRequests.Add(PendingLoadRequests[i]);
 			PendingLoadRequests.RemoveAt(i);
@@ -263,6 +293,16 @@ void FPlaylistHandlerHLS::HandleOnce()
 		{
 			continue;
 		}
+		// Did we add CMCD query parameters?
+		if (FinishedLoadRequests[i]->bAddedCMCDParameters)
+		{
+			if (FinishedLoadRequests[i]->ResourceRequest->GetRequest())
+			{
+				IElectraHttpManager::FParams& Params(FinishedLoadRequests[i]->ResourceRequest->GetRequest()->Parameters);
+				Params.URL = PlayerSessionServices->GetCMCDHandler()->RemoveParamsFromURL(Params.URL);
+			}
+			const_cast<HTTP::FConnectionInfo*>(FinishedLoadRequests[i]->ResourceRequest->GetConnectionInfo())->EffectiveURL = PlayerSessionServices->GetCMCDHandler()->RemoveParamsFromURL(FinishedLoadRequests[i]->ResourceRequest->GetConnectionInfo()->EffectiveURL);
+		}
 
 		auto MakeVariantMediaPlaylist = [&](const TSharedPtr<FLoadRequestHLSPlaylist, ESPMode::ThreadSafe>& FinishedReq, bool bGotUpdate) -> void
 		{
@@ -2022,12 +2062,24 @@ bool FPlaylistHandlerHLS::BuildMultiVariantPlaylist(TSharedPtr<FMultiVariantPlay
 						return false;
 					}
 
-					// Send this session data property up to registered listeners.
-					IPlayerSessionServices::ECustomPropertyResult Result = PlayerSessionServices->ValidateMainPlaylistCustomProperty(TEXT("hls"), InPlaylist->GetURL(), InPlaylist->GetResponseHeaders(), Prop);
-					if (Result == IPlayerSessionServices::ECustomPropertyResult::Reject)
+					// Check if the DATA-ID is one to configure CMCD with.
+					if (PlayerSessionServices->GetCMCDHandler()->UseParametersFromPlaylist() && PlayerSessionServices->GetCMCDHandler()->UsePlaylistParametersFromKey().Equals(Prop.Tag))
 					{
-						LastError = PostError(FString::Printf(TEXT("Playlist has been rejected by application due to bad EXT-X-SESSION-DATA")), HLS::ERRCODE_PLAYLIST_PARSING_FAILED);
-						return false;
+						FString CMCDConfig;
+						if (FBase64::Decode(Prop.Value, CMCDConfig))
+						{
+							PlayerSessionServices->GetCMCDHandler()->UpdateParameters(CMCDConfig);
+						}
+					}
+					else
+					{
+						// Send this session data property up to registered listeners.
+						IPlayerSessionServices::ECustomPropertyResult Result = PlayerSessionServices->ValidateMainPlaylistCustomProperty(TEXT("hls"), InPlaylist->GetURL(), InPlaylist->GetResponseHeaders(), Prop);
+						if (Result == IPlayerSessionServices::ECustomPropertyResult::Reject)
+						{
+							LastError = PostError(FString::Printf(TEXT("Playlist has been rejected by application due to bad EXT-X-SESSION-DATA")), HLS::ERRCODE_PLAYLIST_PARSING_FAILED);
+							return false;
+						}
 					}
 				}
 			}
diff --git a/Engine/Plugins/Media/ElectraPlayer/Source/ElectraPlayerRuntime/Private/Runtime/Player/PlayerSessionServices.h b/Engine/Plugins/Media/ElectraPlayer/Source/ElectraPlayerRuntime/Private/Runtime/Player/PlayerSessionServices.h
--- a/Engine/Plugins/Media/ElectraPlayer/Source/ElectraPlayerRuntime/Private/Runtime/Player/PlayerSessionServices.h
+++ b/Engine/Plugins/Media/ElectraPlayer/Source/ElectraPlayerRuntime/Private/Runtime/Player/PlayerSessionServices.h
@@ -14,6 +14,7 @@ namespace Electra
 	class ISynchronizedUTCTime;
 	class IElectraHttpManager;
 	class IAdaptiveStreamSelector;
+	class IPlayerABRLiveControl;
 	class IPlayerStreamFilter;
 	struct FAccessUnitBufferInfo;
 	class IAdaptiveStreamingPlayerResourceProvider;
@@ -22,6 +23,7 @@ namespace Electra
 	class IAdaptiveStreamingPlayerAEMSHandler;
 	class FDRMManager;
 	class FContentSteeringHandler;
+	class FCMCDHandler;
 	class IHTTPResponseCache;
 
 
@@ -96,6 +98,11 @@ namespace Electra
 		 */
 		virtual TSharedPtrTS<IAdaptiveStreamSelector> GetStreamSelector() = 0;
 
+		/**
+		 * Returns the player ABR control interface.
+		 */
+		virtual IPlayerABRLiveControl* GetPlayerARBControl() = 0;
+
 		/**
 		 * Returns the stream filter interface used by playlist readers to determine whether or not a stream
 		 * can be used on the platform.
@@ -128,6 +135,11 @@ namespace Electra
 		 */
 		virtual TSharedPtrTS<FContentSteeringHandler> GetContentSteeringHandler() = 0;
 
+		/**
+		 * Returns the CMCD handler for this player instance.
+		 */
+		virtual TSharedPtrTS<FCMCDHandler> GetCMCDHandler() = 0;
+
 		/**
 		 * Returns the "Application Event or Metadata Streams" (AEMS) handler.
 		 */
diff --git a/Engine/Plugins/Media/ElectraPlayer/Source/ElectraPlayerRuntime/Private/Runtime/Player/StreamSegmentReaderCommon.cpp b/Engine/Plugins/Media/ElectraPlayer/Source/ElectraPlayerRuntime/Private/Runtime/Player/StreamSegmentReaderCommon.cpp
--- a/Engine/Plugins/Media/ElectraPlayer/Source/ElectraPlayerRuntime/Private/Runtime/Player/StreamSegmentReaderCommon.cpp
+++ b/Engine/Plugins/Media/ElectraPlayer/Source/ElectraPlayerRuntime/Private/Runtime/Player/StreamSegmentReaderCommon.cpp
@@ -19,6 +19,7 @@
 #include "Player/DASH/PlayerEventDASH.h"
 #include "Player/DASH/PlayerEventDASH_Internal.h"
 #include "Player/DRM/DRMManager.h"
+#include "Player/CMCDHandler.h"
 #include "Utilities/Utilities.h"
 #include "Utilities/TimeUtilities.h"
 #include "Utilities/UtilsMP4.h"
@@ -957,19 +958,25 @@ bool FStreamSegmentReaderCommon::FStreamHandler::FetchInitSegment(TSharedPtrTS<F
 	// Not cached yet, need to fetch.
 	CurrentConnectionInfo = HTTP::FConnectionInfo();
 	FGenericDataReader StaticSegmentDataReader;
+	TArray<HTTP::FHTTPHeader> ReqHeaders;
+	if (InitURLToUse.CustomHeader.Len() && InRequest->StreamingProtocol == FStreamSegmentRequestCommon::EStreamingProtocol::DASH)
+	{
+		ReqHeaders.Emplace(HTTP::FHTTPHeader({DASH::HTTPHeaderOptionName, InitURLToUse.CustomHeader}));
+	}
 	// Self-initializing media?
 	if (!bFetchFromSelfInitializing)
 	{
-		TArray<HTTP::FHTTPHeader> ReqHeaders;
-		if (InitURLToUse.CustomHeader.Len() && InRequest->StreamingProtocol == FStreamSegmentRequestCommon::EStreamingProtocol::DASH)
-		{
-			ReqHeaders.Emplace(HTTP::FHTTPHeader({DASH::HTTPHeaderOptionName, InitURLToUse.CustomHeader}));
-		}
-
 		TSharedPtrTS<FHTTPResourceRequest> rr = MakeSharedTS<FHTTPResourceRequest>();
 		TSharedPtr<FHTTPResourceRequestCompletionSignal, ESPMode::ThreadSafe> rrsig = FHTTPResourceRequestCompletionSignal::Create();
 
-		rr->Verb(TEXT("GET")).URL(InitURLToUse.Url.URL).Range(InitURLToUse.Range).Headers(ReqHeaders).AcceptEncoding(TEXT("identity"))
+		FString InitSegmentUrl(InitURLToUse.Url.URL);
+		// Set up information for CMCD.
+		FCMCDHandler::FRequestObjectInfo ri;
+		ri.MediaStreamType = InRequest->StreamType;
+		ri.EncodedBitrate = InRequest->GetBitrate();
+		PlayerSessionService->GetCMCDHandler()->SetupRequestObject(FCMCDHandler::ERequestType::InitSegment, FCMCDHandler::EObjectType::InitSegment, InitSegmentUrl, ReqHeaders, InitURLToUse.Url.CDN, ri);
+
+		rr->Verb(TEXT("GET")).URL(InitSegmentUrl).Range(InitURLToUse.Range).Headers(ReqHeaders).AcceptEncoding(TEXT("identity"))
 			.ConnectionTimeout(FTimeValue().SetFromMilliseconds(5000)).NoDataTimeout(FTimeValue().SetFromMilliseconds(2000))
 			.StreamTypeAndQuality(InRequest->StreamType, InRequest->QualityIndex, InRequest->MaxQualityIndex)
 			.CompletionSignal(rrsig)
@@ -987,6 +994,14 @@ bool FStreamSegmentReaderCommon::FStreamHandler::FetchInitSegment(TSharedPtrTS<F
 			return true;
 		}
 		CurrentConnectionInfo = *rr->GetConnectionInfo();
+		// If CMCD added query parameters to the URL then we need to remove them from the effective URL.
+		if (!InitSegmentUrl.Equals(InitURLToUse.Url.URL))
+		{
+			// Keep the initial URL though as it was since it was (a) not modified and (b) it is what the init segment gets cached under.
+			// We do not want to use the redirected effective URL for the cache.
+			CurrentConnectionInfo.EffectiveURL = PlayerSessionService->GetCMCDHandler()->RemoveParamsFromURL(CurrentConnectionInfo.EffectiveURL);
+		}
+
 		TSharedPtrTS<FWaitableBuffer> ResponseBuffer = rr->GetResponseBuffer();
 
 		bool bSuccessful = !rr->GetError() && ResponseBuffer.IsValid();
@@ -1117,7 +1132,14 @@ bool FStreamSegmentReaderCommon::FStreamHandler::FetchInitSegment(TSharedPtrTS<F
 		TArray<UtilsMP4::FMP4RootBoxLocator::FBoxInfo> Boxes;
 		const TArray<uint32> FirstBox { Electra::UtilitiesMP4::MakeBoxAtom('f','t','y','p') };
 		const TArray<uint32> MoovBox { Electra::UtilitiesMP4::MakeBoxAtom('m','o','o','v') };
-		bool bSuccess = BoxLocator.LocateRootBoxes(Boxes, PlayerSessionService->GetHTTPManager(), InitURLToUse.Url.URL, FirstBox, MoovBox, MoovBox, UtilsMP4::FMP4RootBoxLocator::FCancellationCheckDelegate::CreateLambda([&]()
+		FString InitSegmentUrl(InitURLToUse.Url.URL);
+
+		// Set up information for CMCD.
+		FCMCDHandler::FRequestObjectInfo ri;
+		ri.MediaStreamType = InRequest->StreamType;
+		ri.EncodedBitrate = InRequest->GetBitrate();
+		PlayerSessionService->GetCMCDHandler()->SetupRequestObject(FCMCDHandler::ERequestType::InitSegment, FCMCDHandler::EObjectType::InitSegment, InitSegmentUrl, ReqHeaders, InitURLToUse.Url.CDN, ri);
+		bool bSuccess = BoxLocator.LocateRootBoxes(Boxes, PlayerSessionService->GetHTTPManager(), InitSegmentUrl, ReqHeaders, FirstBox, MoovBox, MoovBox, UtilsMP4::FMP4RootBoxLocator::FCancellationCheckDelegate::CreateLambda([&]()
 		{
 			return HasReadBeenAborted();
 		}));
@@ -1126,6 +1148,13 @@ bool FStreamSegmentReaderCommon::FStreamHandler::FetchInitSegment(TSharedPtrTS<F
 			return true;
 		}
 		CurrentConnectionInfo = BoxLocator.GetConnectionInfo();
+		// If CMCD added query parameters to the URL then we need to remove them from the effective URL.
+		if (!InitSegmentUrl.Equals(InitURLToUse.Url.URL))
+		{
+			// Keep the initial URL though as it was since it was (a) not modified and (b) it is what the init segment gets cached under.
+			// We do not want to use the redirected effective URL for the cache.
+			CurrentConnectionInfo.EffectiveURL = PlayerSessionService->GetCMCDHandler()->RemoveParamsFromURL(CurrentConnectionInfo.EffectiveURL);
+		}
 		SetupSegmentDownloadStatsFromConnectionInfo(CurrentConnectionInfo);
 		if (!bSuccess)
 		{
@@ -1356,6 +1385,25 @@ FStreamSegmentReaderCommon::FStreamHandler::EMediaSegmentTriggerResult FStreamSe
 	HTTPRequest->Parameters.ConnectTimeout = PlayerSessionService->GetOptionValue(DASH::OptionKeyMediaSegmentConnectTimeout).SafeGetTimeValue(FTimeValue().SetFromMilliseconds(1000 * 4));
 	HTTPRequest->Parameters.NoDataTimeout = PlayerSessionService->GetOptionValue(DASH::OptionKeyMediaSegmentNoDataTimeout).SafeGetTimeValue(FTimeValue().SetFromMilliseconds(1000 * 4));
 
+	// Set up information for CMCD.
+	FTimeValue SegmentDuration(InRequest->Segment.Duration, InRequest->Segment.Timescale, 0);
+	FCMCDHandler::FRequestObjectInfo ri;
+	ri.MediaStreamType = InRequest->StreamType;
+	ri.EncodedBitrate = InRequest->GetBitrate();
+	ri.ObjectDuration = (int32) SegmentDuration.GetAsMilliseconds();
+	for(auto& NextMediaUrl : InRequest->Segment.NextMediaURLS)
+	{
+		auto& nxt = ri.NextObjectRequest.Emplace_GetRef();
+		nxt.URL = NextMediaUrl.Url.URL;
+		nxt.Range = NextMediaUrl.Range;
+	}
+	FCMCDHandler::EObjectType CMCDot =
+		InRequest->bIsMultiplex && !InRequest->bIgnoreVideo && !InRequest->bIgnoreAudio ? FCMCDHandler::EObjectType::MuxedAudioAndVideo :
+		InRequest->StreamType == EStreamType::Video ? FCMCDHandler::EObjectType::VideoOnly :
+		InRequest->StreamType == EStreamType::Audio ? FCMCDHandler::EObjectType::AudioOnly : FCMCDHandler::EObjectType::CaptionOrSubtitle;
+	PlayerSessionService->GetCMCDHandler()->SetupRequestObject(FCMCDHandler::ERequestType::Segment, CMCDot, HTTPRequest->Parameters.URL, HTTPRequest->Parameters.RequestHeaders, InRequest->Segment.MediaURL.Url.CDN, ri);
+	bAddedCMCDQueryToMediaURL = !HTTPRequest->Parameters.URL.Equals(InRequest->Segment.MediaURL.Url.URL);
+
 	ProgressReportCount = 0;
 	PlayerSessionService->GetHTTPManager()->AddRequest(HTTPRequest, false);
 
@@ -1381,6 +1429,10 @@ void FStreamSegmentReaderCommon::FStreamHandler::HandleCommonMediaEnd(EHandleRes
 	{
 		PlayerSessionService->GetHTTPManager()->RemoveRequest(HTTPRequest, false);
 		CurrentConnectionInfo = HTTPRequest->ConnectionInfo;
+		if (bAddedCMCDQueryToMediaURL)
+		{
+			CurrentConnectionInfo.EffectiveURL = PlayerSessionService->GetCMCDHandler()->RemoveParamsFromURL(CurrentConnectionInfo.EffectiveURL);
+		}
 		HTTPRequest.Reset();
 	}
 
@@ -1941,12 +1993,12 @@ FStreamSegmentReaderCommon::FStreamHandler::EHandleResult FStreamSegmentReaderCo
 					if (CurrentlyActiveTrackData->MediaLocalLastAUTime != TNumericLimits<int64>::Max())
 					{
 						AccessUnit->LatestPTS.SetFromND(CurrentlyActiveTrackData->MediaLocalLastAUTime - CurrentlyActiveTrackData->PTO, TrackTimescale);
+						AccessUnit->LatestPTS += TimeOffset;
 					}
 					else
 					{
 						AccessUnit->LatestPTS.SetToPositiveInfinity();
 					}
-					AccessUnit->LatestPTS += TimeOffset;
 					AccessUnit->LatestPTS.SetSequenceIndex(InRequest->TimestampSequenceIndex);
 
 					AccessUnit->ProducerReferenceTime = CurrentlyActiveTrackData->ProducerTime.Base + FTimeValue(AUDTS - CurrentlyActiveTrackData->ProducerTime.Media, TrackTimescale);
diff --git a/Engine/Plugins/Media/ElectraPlayer/Source/ElectraPlayerRuntime/Private/Runtime/Player/StreamSegmentReaderCommon.h b/Engine/Plugins/Media/ElectraPlayer/Source/ElectraPlayerRuntime/Private/Runtime/Player/StreamSegmentReaderCommon.h
--- a/Engine/Plugins/Media/ElectraPlayer/Source/ElectraPlayerRuntime/Private/Runtime/Player/StreamSegmentReaderCommon.h
+++ b/Engine/Plugins/Media/ElectraPlayer/Source/ElectraPlayerRuntime/Private/Runtime/Player/StreamSegmentReaderCommon.h
@@ -59,8 +59,15 @@ struct FSegmentInformationCommon
 		ISO13818_1,
 		Matroska
 	};
+	struct FCMCDNextSegNumberTimes
+	{
+		int64 Number = -1;
+		int64 Time = -1;
+	};
 	FURL InitializationURL;
 	FURL MediaURL;
+	TArray<FURL> NextMediaURLS;
+	TArray<FCMCDNextSegNumberTimes> NextMediaCMCDNumberTimes;
 	FTimeValue ATO;
 	int64 Time = 0;								// Time value T in timescale units
 	int64 PTO = 0;								// PresentationTimeOffset
@@ -594,6 +601,7 @@ private:
 		bool bAbortedByABR = false;
 		bool bAllowEarlyEmitting = false;
 		bool bFillRemainingDuration = false;
+		bool bAddedCMCDQueryToMediaURL = false;
 
 		IPlayerSessionServices* PlayerSessionService = nullptr;
 		FReadBuffer ReadBuffer;
diff --git a/Engine/Plugins/Media/ElectraPlayer/Source/ElectraPlayerRuntime/Private/Runtime/Player/mp4/PlaylistReaderMP4.cpp b/Engine/Plugins/Media/ElectraPlayer/Source/ElectraPlayerRuntime/Private/Runtime/Player/mp4/PlaylistReaderMP4.cpp
--- a/Engine/Plugins/Media/ElectraPlayer/Source/ElectraPlayerRuntime/Private/Runtime/Player/mp4/PlaylistReaderMP4.cpp
+++ b/Engine/Plugins/Media/ElectraPlayer/Source/ElectraPlayerRuntime/Private/Runtime/Player/mp4/PlaylistReaderMP4.cpp
@@ -232,7 +232,8 @@ void FPlaylistReaderMP4::WorkerThread()
 	const TArray<uint32> FirstBoxes { Electra::UtilitiesMP4::MakeBoxAtom('f','t','y','p'), Electra::UtilitiesMP4::MakeBoxAtom('s','t','y','p'), Electra::UtilitiesMP4::MakeBoxAtom('s','i','d','x'), Electra::UtilitiesMP4::MakeBoxAtom('f','r','e','e'), Electra::UtilitiesMP4::MakeBoxAtom('s','k','i','p') };
 	UtilsMP4::FMP4RootBoxLocator BoxLocator;
 	TArray<UtilsMP4::FMP4RootBoxLocator::FBoxInfo> BoxInfos;
-	bool bGotBoxes = BoxLocator.LocateRootBoxes(BoxInfos, PlayerSessionServices->GetHTTPManager(), MainPlaylistURL, FirstBoxes, StopAfterBoxes, ReadBoxes, UtilsMP4::FMP4RootBoxLocator::FCancellationCheckDelegate::CreateLambda([&]()
+	TArray<HTTP::FHTTPHeader> ReqHeaders;
+	bool bGotBoxes = BoxLocator.LocateRootBoxes(BoxInfos, PlayerSessionServices->GetHTTPManager(), MainPlaylistURL, ReqHeaders, FirstBoxes, StopAfterBoxes, ReadBoxes, UtilsMP4::FMP4RootBoxLocator::FCancellationCheckDelegate::CreateLambda([&]()
 	{
 		return bAbort;
 	}));
diff --git a/Engine/Plugins/Media/ElectraPlayer/Source/ElectraPlayerRuntime/Private/Runtime/PlayerPlatform.h b/Engine/Plugins/Media/ElectraPlayer/Source/ElectraPlayerRuntime/Private/Runtime/PlayerPlatform.h
new file mode 100644
--- /dev/null
+++ b/Engine/Plugins/Media/ElectraPlayer/Source/ElectraPlayerRuntime/Private/Runtime/PlayerPlatform.h
@@ -0,0 +1,14 @@
+// Copyright Epic Games, Inc. All Rights Reserved.
+
+#pragma once
+
+#include "CoreTypes.h"
+#include "Containers/UnrealString.h"
+
+namespace Electra
+{
+	namespace Platform
+	{
+		FString GetPlatformID();
+	}
+} // namespace Electra
diff --git a/Engine/Plugins/Media/ElectraPlayer/Source/ElectraPlayerRuntime/Private/Runtime/PlayerRuntimeGlobal.cpp b/Engine/Plugins/Media/ElectraPlayer/Source/ElectraPlayerRuntime/Private/Runtime/PlayerRuntimeGlobal.cpp
--- a/Engine/Plugins/Media/ElectraPlayer/Source/ElectraPlayerRuntime/Private/Runtime/PlayerRuntimeGlobal.cpp
+++ b/Engine/Plugins/Media/ElectraPlayer/Source/ElectraPlayerRuntime/Private/Runtime/PlayerRuntimeGlobal.cpp
@@ -2,6 +2,7 @@
 
 #include "PlayerCore.h"
 #include "PlayerRuntimeGlobal.h"
+#include "PlayerPlatform.h"
 #include "Misc/CoreDelegates.h"
 #include "Modules/ModuleManager.h"
 
@@ -126,6 +127,9 @@ namespace Electra
 			return false;
 		}
 #endif
+
+		// Dummy get the platform ID once so any potential statics are created right away.
+		Platform::GetPlatformID();
 		return true;
 	}
 
@@ -194,7 +198,7 @@ namespace Electra
 		const bool* bEventEnabled = EnabledAnalyticsEvents.Find(AnalyticsEventName);
 		return bEventEnabled && *bEventEnabled;
 	}
-	
+
 	class PendingTaskCounter
 	{
 	public:
diff --git a/Engine/Plugins/Media/ElectraPlayer/Source/ElectraPlayerRuntime/Private/Runtime/Utilities/UtilsMP4.cpp b/Engine/Plugins/Media/ElectraPlayer/Source/ElectraPlayerRuntime/Private/Runtime/Utilities/UtilsMP4.cpp
--- a/Engine/Plugins/Media/ElectraPlayer/Source/ElectraPlayerRuntime/Private/Runtime/Utilities/UtilsMP4.cpp
+++ b/Engine/Plugins/Media/ElectraPlayer/Source/ElectraPlayerRuntime/Private/Runtime/Utilities/UtilsMP4.cpp
@@ -552,7 +552,7 @@ UtilsMP4::FMP4RootBoxLocator::~FMP4RootBoxLocator()
 {
 }
 
-bool UtilsMP4::FMP4RootBoxLocator::LocateRootBoxes(TArray<FBoxInfo>& OutBoxInfos, const TSharedPtrTS<IElectraHttpManager>& InHTTPManager, const FString& InURL, const TArray<uint32>& InFirstBoxes, const TArray<uint32>& InStopAfterBoxes, const TArray<uint32>& InReadDataOfBoxes, FCancellationCheckDelegate InCheckCancellationDelegate)
+bool UtilsMP4::FMP4RootBoxLocator::LocateRootBoxes(TArray<FBoxInfo>& OutBoxInfos, const TSharedPtrTS<IElectraHttpManager>& InHTTPManager, const FString& InURL, const TArray<HTTP::FHTTPHeader>& InRequestHeaders, const TArray<uint32>& InFirstBoxes, const TArray<uint32>& InStopAfterBoxes, const TArray<uint32>& InReadDataOfBoxes, FCancellationCheckDelegate InCheckCancellationDelegate)
 {
 	FMediaEvent ReadCompleted;
 	FString URL(InURL);
@@ -608,6 +608,8 @@ bool UtilsMP4::FMP4RootBoxLocator::LocateRootBoxes(TArray<FBoxInfo>& OutBoxInfos
 		Req->Parameters.Range.SetEndIncluding(LastByte);
 		Req->Parameters.ConnectTimeout = FTimeValue().SetFromMilliseconds(1000 * 8);
 		Req->Parameters.NoDataTimeout = FTimeValue().SetFromMilliseconds(1000 * 6);
+		Req->Parameters.AcceptEncoding.Set(TEXT("identity"));
+		Req->Parameters.RequestHeaders = InRequestHeaders;
 		OutReceiveBuffer = MakeSharedTS<FWaitableBuffer>();
 		OutReceiveBuffer->Reserve(InNumBytes);
 		Req->ReceiveBuffer = OutReceiveBuffer;
diff --git a/Engine/Plugins/Media/ElectraPlayer/Source/ElectraPlayerRuntime/Private/Runtime/Utilities/UtilsMP4.h b/Engine/Plugins/Media/ElectraPlayer/Source/ElectraPlayerRuntime/Private/Runtime/Utilities/UtilsMP4.h
--- a/Engine/Plugins/Media/ElectraPlayer/Source/ElectraPlayerRuntime/Private/Runtime/Utilities/UtilsMP4.h
+++ b/Engine/Plugins/Media/ElectraPlayer/Source/ElectraPlayerRuntime/Private/Runtime/Utilities/UtilsMP4.h
@@ -96,7 +96,7 @@ public:
 		FMP4RootBoxLocator() = default;
 		~FMP4RootBoxLocator();
 		DECLARE_DELEGATE_RetVal(bool, FCancellationCheckDelegate);
-		bool LocateRootBoxes(TArray<FBoxInfo>& OutBoxInfos, const TSharedPtrTS<IElectraHttpManager>& InHTTPManager, const FString& InURL, const TArray<uint32>& InFirstBoxes, const TArray<uint32>& InStopAfterBoxes, const TArray<uint32>& InReadDataOfBoxes, FCancellationCheckDelegate InCheckCancellationDelegate);
+		bool LocateRootBoxes(TArray<FBoxInfo>& OutBoxInfos, const TSharedPtrTS<IElectraHttpManager>& InHTTPManager, const FString& InURL, const TArray<HTTP::FHTTPHeader>& InRequestHeaders, const TArray<uint32>& InFirstBoxes, const TArray<uint32>& InStopAfterBoxes, const TArray<uint32>& InReadDataOfBoxes, FCancellationCheckDelegate InCheckCancellationDelegate);
 
 		int64 GetFileSize() const
 		{ return FileSize; }
@@ -116,10 +116,10 @@ public:
 
 	class FMP4ChunkLoader
 	{
-	public:		
+	public:
 		FMP4ChunkLoader() = default;
 		~FMP4ChunkLoader() = default;
-		
+
 		DECLARE_DELEGATE_RetVal(bool, FCancellationCheckDelegate);
 		TSharedPtrTS<FWaitableBuffer> LoadChunk(const int64 InOffset, const int64 InSize, const TSharedPtrTS<IElectraHttpManager>& InHTTPManager, const TSharedPtrTS<IHTTPResponseCache>& InHttpResponseCache, const FString& InURL, FCancellationCheckDelegate InCheckCancellationDelegate);
 
diff --git a/Engine/Plugins/Media/ElectraUtil/Source/ElectraBase/Private/URLParser.cpp b/Engine/Plugins/Media/ElectraUtil/Source/ElectraBase/Private/URLParser.cpp
--- a/Engine/Plugins/Media/ElectraUtil/Source/ElectraBase/Private/URLParser.cpp
+++ b/Engine/Plugins/Media/ElectraUtil/Source/ElectraBase/Private/URLParser.cpp
@@ -997,5 +997,75 @@ namespace Electra
 	}
 
 
+	bool FURL_RFC3986::MakeRelativePath(FString& OutRelativePath, const FString& InAbsoluteURLToMakeRelativeToThis) const
+	{
+		FURL_RFC3986 Other;
+		if (!Other.Parse(InAbsoluteURLToMakeRelativeToThis))
+		{
+			return false;
+		}
+		// Scheme, user info, host and port have to match.
+		if (!Scheme.Equals(Other.Scheme) || !UserInfo.Equals(Other.UserInfo) || !Host.Equals(Other.Host) || !Port.Equals(Other.Port))
+		{
+			return false;
+		}
+		// Both paths must not end on a slash
+		if (Path.EndsWith(URLComponents::Slash) || Other.Path.EndsWith(URLComponents::Slash))
+		{
+			return false;
+		}
+		// Split the paths into arrays
+		TArray<FString> ThisPaths, OtherPaths;
+		GetPathComponents(ThisPaths, Path);
+		GetPathComponents(OtherPaths, Other.Path);
+		while(ThisPaths.Num() && OtherPaths.Num() && ThisPaths[0].Equals(OtherPaths[0]))
+		{
+			ThisPaths.RemoveAt(0);
+			OtherPaths.RemoveAt(0);
+		}
+
+		auto AssemblePaths = [](const TArray<FString>& InPaths) -> FString
+		{
+			FString p;
+			for(int32 i=0,iMax=InPaths.Num(); i<iMax; ++i)
+			{
+				if (i)
+				{
+					p.Append(URLComponents::Slash);
+				}
+				p.Append(InPaths[i]);
+			}
+			return p;
+		};
+
+		FString RelPath;
+		// If this path is now only the filename, the other is relative to here.
+		if (ThisPaths.Num() == 1)
+		{
+			RelPath = TEXT("./");
+		}
+		else
+		{
+			for(int32 i=1; i<ThisPaths.Num(); ++i)
+			{
+				RelPath.Append(TEXT("../"));
+			}
+		}
+		RelPath.Append(AssemblePaths(OtherPaths));
+
+		if (Other.Query.Len())
+		{
+			RelPath += URLComponents::QuestionMark;
+			RelPath += Other.Query;
+		}
+		if (Other.Fragment.Len())
+		{
+			RelPath += URLComponents::HashTag;
+			UrlEncode(RelPath, Other.Fragment, URLComponents::KeepCharsFragment);
+		}
+		OutRelativePath = MoveTemp(RelPath);
+		return true;
+	}
+
 
 } // namespace Electra
diff --git a/Engine/Plugins/Media/ElectraUtil/Source/ElectraBase/Public/Utilities/URLParser.h b/Engine/Plugins/Media/ElectraUtil/Source/ElectraBase/Public/Utilities/URLParser.h
--- a/Engine/Plugins/Media/ElectraUtil/Source/ElectraBase/Public/Utilities/URLParser.h
+++ b/Engine/Plugins/Media/ElectraUtil/Source/ElectraBase/Public/Utilities/URLParser.h
@@ -47,6 +47,9 @@ namespace Electra
 		// Resolves this URL (which should be relative) against the specified URL.
 		UE_API FURL_RFC3986& ResolveAgainst(const FString& InParentURL);
 
+		// Attempts to make an absolute URL relative to this one.
+		UE_API bool MakeRelativePath(FString& OutRelativePath, const FString& InAbsoluteURLToMakeRelativeToThis) const;
+
 		UE_API void SetScheme(const FString& InNewScheme);
 		UE_API void SetHost(const FString& InNewHost);
 		UE_API void SetPort(const FString& InNewPort);
