Electra: combined multiline raw string into a single line

#rnx

diff --git a/Engine/Plugins/Media/ElectraPlayer/Source/ElectraPlayerRuntime/Private/Runtime/Player/CMCDHandler.cpp b/Engine/Plugins/Media/ElectraPlayer/Source/ElectraPlayerRuntime/Private/Runtime/Player/CMCDHandler.cpp
--- a/Engine/Plugins/Media/ElectraPlayer/Source/ElectraPlayerRuntime/Private/Runtime/Player/CMCDHandler.cpp
+++ b/Engine/Plugins/Media/ElectraPlayer/Source/ElectraPlayerRuntime/Private/Runtime/Player/CMCDHandler.cpp
@@ -22,15 +22,11 @@ static TAutoConsoleVariable<int32> CVarElectraCMCDMetrics(
 namespace Electra
 {
 // Standard config only enables playlist provided parameters.
-static const TCHAR* const StandardCMCDConfig = TEXT(R"##({
-"cmcd":[]
-})##");
+static const TCHAR* const StandardCMCDConfig = TEXT("{\"cmcd\":[]}");
 
 // Default values are: `headers` mode, `request` type, all keys enabled, version `1` and playlist parameters to be applied.
 // Values that are default do not need to be included in this config.
-static const TCHAR* const DefaultEnabledCMCDConfig = TEXT(R"##({
-"cmcd":[{}]
-})##");
+static const TCHAR* const DefaultEnabledCMCDConfig = TEXT("{\"cmcd\":[{}]}");
 
 
 
