ElectraPlayer: Using a VoD asset for a synchronized event now allows it to loop when provided via DASH and the loop state is provided initially in the player option struct

#rnx

diff --git a/Engine/Plugins/Media/ElectraPlayer/Source/ElectraPlayerRuntime/Private/ElectraPlayer.cpp b/Engine/Plugins/Media/ElectraPlayer/Source/ElectraPlayerRuntime/Private/ElectraPlayer.cpp
--- a/Engine/Plugins/Media/ElectraPlayer/Source/ElectraPlayerRuntime/Private/ElectraPlayer.cpp
+++ b/Engine/Plugins/Media/ElectraPlayer/Source/ElectraPlayerRuntime/Private/ElectraPlayer.cpp
@@ -496,6 +496,11 @@ bool FElectraPlayer::OpenInternal(const FString& InUrl, const IMediaOptions* InO
 	FName Environment;
 	if (InPlayerOptions)
 	{
+		if (InPlayerOptions->Loop != EMediaPlayerOptionBooleanOverride::UseMediaPlayerSetting)
+		{
+			bEnableLooping = InPlayerOptions->Loop == EMediaPlayerOptionBooleanOverride::Enabled;
+			UE_LOG(LogElectraPlayer, Log, TEXT("[%u] IMediaPlayer::Open: Setting initial loop state to %s"), PlayerUniqueID, bEnableLooping ? TEXT("enabled") : TEXT("disabled"));
+		}
 		if (InPlayerOptions->SeekTimeType != EMediaPlayerOptionSeekTimeType::Ignored)
 		{
 			PlaystartOptions.TimeOffset = InPlayerOptions->SeekTime;
@@ -3405,7 +3410,7 @@ bool FElectraPlayer::SetLooping(bool bInLooping)
 	IAdaptiveStreamingPlayer::FLoopParam loop;
 	loop.bEnableLooping = bEnableLooping.GetValue();
 	Player->AdaptivePlayer->SetLooping(loop);
-	UE_LOG(LogElectraPlayer, Verbose, TEXT("[%u] IMediaPlayer::SetLooping(%s)"), PlayerUniqueID, bEnableLooping.GetValue()?TEXT("true"):TEXT("false"));
+	UE_LOG(LogElectraPlayer, Log, TEXT("[%u] IMediaPlayer::SetLooping(%s)"), PlayerUniqueID, bEnableLooping.GetValue()?TEXT("true"):TEXT("false"));
 	return true;
 }
 
diff --git a/Engine/Plugins/Media/ElectraPlayer/Source/ElectraPlayerRuntime/Private/Runtime/Player/AdaptiveStreamingPlayer.cpp b/Engine/Plugins/Media/ElectraPlayer/Source/ElectraPlayerRuntime/Private/Runtime/Player/AdaptiveStreamingPlayer.cpp
--- a/Engine/Plugins/Media/ElectraPlayer/Source/ElectraPlayerRuntime/Private/Runtime/Player/AdaptiveStreamingPlayer.cpp
+++ b/Engine/Plugins/Media/ElectraPlayer/Source/ElectraPlayerRuntime/Private/Runtime/Player/AdaptiveStreamingPlayer.cpp
@@ -2487,8 +2487,9 @@ void FAdaptiveStreamingPlayer::InternalHandlePendingStartRequest(const FTimeValu
 				return;
 			}
 
-			// For sanities sake disable any loop flag that might have been set before loading the playlist if the presentation isn't VoD.
-			if (Manifest->GetPresentationType() != IManifest::EType::OnDemand)
+			// For sanities sake disable any loop flag that might have been set before loading the playlist if the presentation isn't VoD
+			// and is also not a replay event.
+			if (Manifest->GetPresentationType() != IManifest::EType::OnDemand && Manifest->GetReplayType() == IManifest::EReplayEventType::NoReplay)
 			{
 				CurrentLoopParam.bEnableLooping = false;
 			}
@@ -2517,9 +2518,22 @@ void FAdaptiveStreamingPlayer::InternalHandlePendingStartRequest(const FTimeValu
 						{
 							return;
 						}
-
 						check(Seekable.End.IsValid());
-						StartAt.Time = Seekable.End;
+
+						// If this is a looping replay event and our start time is not set then we are to join 'now'.
+						if (CurrentLoopParam.bEnableLooping && Manifest->GetReplayType() != IManifest::EReplayEventType::NoReplay)
+						{
+							StartAt.Time = GetSynchronizedUTCTime()->GetTime();
+							FTimeValue LiveLatency = Manifest->GetDesiredLiveLatency();
+							if (LiveLatency.IsValid())
+							{
+								StartAt.Time -= LiveLatency;
+							}
+						}
+						else
+						{
+							StartAt.Time = Seekable.End;
+						}
 						switch(Manifest->GetLiveEdgePlayMode())
 						{
 							case IManifest::ELiveEdgePlayMode::Default:
@@ -2560,16 +2574,48 @@ void FAdaptiveStreamingPlayer::InternalHandlePendingStartRequest(const FTimeValu
 					case IManifest::FResult::EType::PastEOS:
 					{
 						// If the initial start time is beyond the end of the presentation
-						// but looping is enabled we set this request up to start at the beginning.
+						// but looping is enabled we need to handle this accordingly.
+						bool bEnded = !CurrentLoopParam.bEnableLooping;
 						if (CurrentLoopParam.bEnableLooping)
 						{
-							FTimeRange Seekable = Manifest->GetSeekableTimeRange();
-							PendingStartRequest->SearchType = IManifest::ESearchType::Closest;
-							StartAt.Time = Seekable.Start;
-							ClampStartRequestTime(StartAt.Time);
-							PendingStartRequest->StartAt.Time = StartAt.Time;
+							if (Manifest->GetReplayType() == IManifest::EReplayEventType::NoReplay)
+							{
+								// Not a replay, start at the beginning.
+								FTimeRange Seekable = Manifest->GetSeekableTimeRange();
+								PendingStartRequest->SearchType = IManifest::ESearchType::Closest;
+								StartAt.Time = Seekable.Start;
+								ClampStartRequestTime(StartAt.Time);
+								PendingStartRequest->StartAt.Time = StartAt.Time;
+							}
+							else
+							{
+								// For a looping replay we need to check if the start is beyond the event range and
+								// then adjust the timeline accordingly.
+								FTimeRange TotalRange = Manifest->GetTotalTimeRange();
+								if (TotalRange.IsValid())
+								{
+									const FTimeValue RangeDuration = TotalRange.End - TotalRange.Start;
+									const FTimeValue OverTime = StartAt.Time - TotalRange.Start;
+									int32 NumWraps = OverTime.GetAsHNS() / RangeDuration.GetAsHNS();
+									// Apply a batch of wraps only when there is more than one to avoid inaccuracies
+									// that may result in slightly overshooting the range.
+									if (NumWraps > 1)
+									{
+										Manifest->PrepareForLooping(NumWraps - 1);
+									}
+									// Advance single iterations (zero or one)
+									while(StartAt.Time >= Manifest->GetTotalTimeRange().End)
+									{
+										Manifest->PrepareForLooping(1);
+									}
+								}
+								else
+								{
+									bEnded = true;
+								}
+							}
 						}
-						else
+						if (bEnded)
 						{
 							PendingStartRequest.Reset();
 							InternalSetPlaybackEnded();
@@ -3066,6 +3112,9 @@ void FAdaptiveStreamingPlayer::InternalHandleCompletedSegmentRequests(const FTim
 			{
 				if (Manifest.IsValid())
 				{
+					// We will be looping now. If the manifest needs internal updating, let it know to do so now.
+					Manifest->PrepareForLooping(1);
+
 					FTimeRange Seekable = Manifest->GetSeekableTimeRange();
 
 					PendingStartRequest = MakeSharedTS<FPendingStartRequest>();
@@ -4666,11 +4715,11 @@ void FAdaptiveStreamingPlayer::InternalHandlePendingLoopRequest()
 		// In case this is called that early on that we can't determine the presentation type yet we allow
 		// looping to be set. It won't get evaluated anyway since the Live presentation doesn't end.
 		// The only issue with this is that until a manifest is loaded querying loop enable will return true.
-		if (!Manifest || Manifest->GetPresentationType() == IManifest::EType::OnDemand)
+		if (!Manifest || Manifest->GetPresentationType() == IManifest::EType::OnDemand || Manifest->GetReplayType() == IManifest::EReplayEventType::IsReplay)
 		{
 			bool bStartOver = false;
 			// When looping is to be enabled we need to check if we can do this without having to start over.
-			if (!CurrentLoopParam.bEnableLooping && NewLoopParam->bEnableLooping)
+			if (Manifest && !CurrentLoopParam.bEnableLooping && NewLoopParam->bEnableLooping)
 			{
 				// If there are any completed segment requests the loop state will be checked when all reach
 				// the end. Otherwise, if any of the receive buffers has the end-of-data flag set that
diff --git a/Engine/Plugins/Media/ElectraPlayer/Source/ElectraPlayerRuntime/Private/Runtime/Player/DASH/ManifestBuilderDASH.cpp b/Engine/Plugins/Media/ElectraPlayer/Source/ElectraPlayerRuntime/Private/Runtime/Player/DASH/ManifestBuilderDASH.cpp
--- a/Engine/Plugins/Media/ElectraPlayer/Source/ElectraPlayerRuntime/Private/Runtime/Player/DASH/ManifestBuilderDASH.cpp
+++ b/Engine/Plugins/Media/ElectraPlayer/Source/ElectraPlayerRuntime/Private/Runtime/Player/DASH/ManifestBuilderDASH.cpp
@@ -1193,6 +1193,7 @@ void FManifestDASHInternal::TransformIntoEpicEvent()
 		}
 		if (Params.Num())
 		{
+			OriginalPresentationDuration = GetDuration();
 			EpicEventType = bStaticStart ? EEpicEventType::Static : EEpicEventType::Dynamic;
 			// Get the event start time. This is either a Posix time in seconds since the Epoch (1/1/1970) or the special
 			// word 'now' optionally followed by a value to be added or subtracted from now.
@@ -1248,6 +1249,26 @@ void FManifestDASHInternal::TransformIntoEpicEvent()
 	}
 }
 
+void FManifestDASHInternal::PrepareEpicEventForLooping(int32 InNumLoopsToAdd)
+{
+	check(OriginalPresentationDuration.IsValid());
+	if (OriginalPresentationDuration.IsValid())
+	{
+		FTimeValue ast = GetAnchorTime();
+		FTimeValue AddDur = OriginalPresentationDuration * InNumLoopsToAdd;
+		FTimeValue Start = ast + AddDur;
+		MPDRoot->SetAvailabilityStartTime(Start);
+		if (EpicEventType == EEpicEventType::Dynamic)
+		{
+			MPDRoot->SetPublishTime(Start);
+		}
+		for(auto &Period : GetPeriods())
+		{
+			Period->StartAST += AddDur;
+			Period->EndAST += AddDur;
+		}
+	}
+}
 
 
 FErrorDetail FManifestDASHInternal::Build(IPlayerSessionServices* InPlayerSessionServices, TSharedPtr<FDashMPD_MPDType, ESPMode::ThreadSafe> InMPDRoot, TArray<TWeakPtrTS<IDashMPDElement>> InXLinkElements)
diff --git a/Engine/Plugins/Media/ElectraPlayer/Source/ElectraPlayerRuntime/Private/Runtime/Player/DASH/ManifestDASH.cpp b/Engine/Plugins/Media/ElectraPlayer/Source/ElectraPlayerRuntime/Private/Runtime/Player/DASH/ManifestDASH.cpp
--- a/Engine/Plugins/Media/ElectraPlayer/Source/ElectraPlayerRuntime/Private/Runtime/Player/DASH/ManifestDASH.cpp
+++ b/Engine/Plugins/Media/ElectraPlayer/Source/ElectraPlayerRuntime/Private/Runtime/Player/DASH/ManifestDASH.cpp
@@ -348,6 +348,17 @@ IManifest::EType FManifestDASH::GetPresentationType() const
 	return IManifest::EType::OnDemand;
 }
 
+IManifest::EReplayEventType FManifestDASH::GetReplayType() const
+{
+	TSharedPtrTS<FManifestDASHInternal> Manifest(CurrentManifest);
+	if (Manifest.IsValid())
+	{
+		return Manifest->IsEpicEvent() ? IManifest::EReplayEventType::IsReplay : IManifest::EReplayEventType::NoReplay;
+	}
+	return IManifest::EReplayEventType::NoReplay;
+}
+
+
 TSharedPtrTS<const FLowLatencyDescriptor> FManifestDASH::GetLowLatencyDescriptor() const
 {
 	TSharedPtrTS<FManifestDASHInternal> Manifest(CurrentManifest);
@@ -519,6 +530,16 @@ void FManifestDASH::UpdateDynamicRefetchCounter()
 	}
 }
 
+void FManifestDASH::PrepareForLooping(int32 InNumLoopsToAdd)
+{
+	TSharedPtrTS<FManifestDASHInternal> Manifest = CurrentManifest;
+	if (Manifest.IsValid() && Manifest->IsEpicEvent())
+	{
+		Manifest->PrepareEpicEventForLooping(InNumLoopsToAdd);
+	}
+}
+
+
 void FManifestDASH::TriggerClockSync(IManifest::EClockSyncType InClockSyncType)
 {
 	IPlaylistReaderDASH* Reader = static_cast<IPlaylistReaderDASH*>(PlayerSessionServices->GetManifestReader().Get());
diff --git a/Engine/Plugins/Media/ElectraPlayer/Source/ElectraPlayerRuntime/Private/Runtime/Player/DASH/ManifestDASH.h b/Engine/Plugins/Media/ElectraPlayer/Source/ElectraPlayerRuntime/Private/Runtime/Player/DASH/ManifestDASH.h
--- a/Engine/Plugins/Media/ElectraPlayer/Source/ElectraPlayerRuntime/Private/Runtime/Player/DASH/ManifestDASH.h
+++ b/Engine/Plugins/Media/ElectraPlayer/Source/ElectraPlayerRuntime/Private/Runtime/Player/DASH/ManifestDASH.h
@@ -19,6 +19,7 @@ public:
 
 	virtual ~FManifestDASH();
 	EType GetPresentationType() const override;
+	EReplayEventType GetReplayType() const override;
 	TSharedPtrTS<const FLowLatencyDescriptor> GetLowLatencyDescriptor() const override;
 	FTimeValue CalculateCurrentLiveLatency(const FTimeValue& InCurrentPlaybackPosition, const FTimeValue& InEncoderLatency, bool bViaLatencyElement) const override;
 	FTimeValue GetAnchorTime() const override;
@@ -38,6 +39,7 @@ public:
 	void GetTrackMetadata(TArray<FTrackMetadata>& OutMetadata, EStreamType StreamType) const override;
 	void UpdateRunningMetaData(TSharedPtrTS<UtilsMP4::FMetadataParser> InUpdatedMetaData) override;
 	void UpdateDynamicRefetchCounter() override;
+	void PrepareForLooping(int32 InNumLoopsToAdd) override;
 	void TriggerClockSync(IManifest::EClockSyncType InClockSyncType) override;
 	void TriggerPlaylistRefresh() override;
 	void ReachedStableBuffer() override;
diff --git a/Engine/Plugins/Media/ElectraPlayer/Source/ElectraPlayerRuntime/Private/Runtime/Player/DASH/PlaylistReaderDASH_Internal.h b/Engine/Plugins/Media/ElectraPlayer/Source/ElectraPlayerRuntime/Private/Runtime/Player/DASH/PlaylistReaderDASH_Internal.h
--- a/Engine/Plugins/Media/ElectraPlayer/Source/ElectraPlayerRuntime/Private/Runtime/Player/DASH/PlaylistReaderDASH_Internal.h
+++ b/Engine/Plugins/Media/ElectraPlayer/Source/ElectraPlayerRuntime/Private/Runtime/Player/DASH/PlaylistReaderDASH_Internal.h
@@ -728,6 +728,11 @@ public:
 		return PresentationType;
 	}
 
+	bool IsEpicEvent() const
+	{
+		return EpicEventType != EEpicEventType::None;
+	}
+
 	bool IsDynamicEpicEvent() const
 	{
 		return EpicEventType == EEpicEventType::Dynamic;
@@ -798,6 +803,7 @@ public:
 
 	void InjectEpicTimingSources();
 	void TransformIntoEpicEvent();
+	void PrepareEpicEventForLooping(int32 InNumLoopsToAdd);
 
 	ELECTRA_IMPL_DEFAULT_ERROR_METHODS(DASHManifest);
 
@@ -874,6 +880,7 @@ private:
 	// Type of the presentation.
 	EPresentationType PresentationType;
 	EEpicEventType EpicEventType = EEpicEventType::None;
+	FTimeValue OriginalPresentationDuration;
 
 	TArray<TSharedPtrTS<FPeriod>> Periods;
 
diff --git a/Engine/Plugins/Media/ElectraPlayer/Source/ElectraPlayerRuntime/Private/Runtime/Player/HLS/PlaylistHLS.cpp b/Engine/Plugins/Media/ElectraPlayer/Source/ElectraPlayerRuntime/Private/Runtime/Player/HLS/PlaylistHLS.cpp
--- a/Engine/Plugins/Media/ElectraPlayer/Source/ElectraPlayerRuntime/Private/Runtime/Player/HLS/PlaylistHLS.cpp
+++ b/Engine/Plugins/Media/ElectraPlayer/Source/ElectraPlayerRuntime/Private/Runtime/Player/HLS/PlaylistHLS.cpp
@@ -2238,6 +2238,11 @@ IManifest::EType FActiveHLSPlaylist::GetPresentationType() const
 	return (TimelineMediaAsset->PlaylistType == FPlaylistParserHLS::EPlaylistType::VOD || TimelineMediaAsset->bHasEndList) ? IManifest::EType::OnDemand : IManifest::EType::Live;
 }
 
+IManifest::EReplayEventType FActiveHLSPlaylist::GetReplayType() const
+{
+	return IManifest::EReplayEventType::NoReplay;
+}
+
 TSharedPtrTS<const FLowLatencyDescriptor> FActiveHLSPlaylist::GetLowLatencyDescriptor() const
 {
 	return nullptr;
@@ -2378,6 +2383,10 @@ void FActiveHLSPlaylist::UpdateDynamicRefetchCounter()
 	// Not used with HLS
 }
 
+void FActiveHLSPlaylist::PrepareForLooping(int32 InNumLoopsToAdd)
+{
+}
+
 void FActiveHLSPlaylist::TriggerClockSync(EClockSyncType InClockSyncType)
 {
 	// Not used with HLS
diff --git a/Engine/Plugins/Media/ElectraPlayer/Source/ElectraPlayerRuntime/Private/Runtime/Player/HLS/PlaylistHLS.h b/Engine/Plugins/Media/ElectraPlayer/Source/ElectraPlayerRuntime/Private/Runtime/Player/HLS/PlaylistHLS.h
--- a/Engine/Plugins/Media/ElectraPlayer/Source/ElectraPlayerRuntime/Private/Runtime/Player/HLS/PlaylistHLS.h
+++ b/Engine/Plugins/Media/ElectraPlayer/Source/ElectraPlayerRuntime/Private/Runtime/Player/HLS/PlaylistHLS.h
@@ -529,6 +529,7 @@ public:
 	void CheckForPathwaySwitch();
 
 	EType GetPresentationType() const override;
+	EReplayEventType GetReplayType() const override;
 	TSharedPtrTS<const FLowLatencyDescriptor> GetLowLatencyDescriptor() const override;
 	FTimeValue CalculateCurrentLiveLatency(const FTimeValue& InCurrentPlaybackPosition, const FTimeValue& InEncoderLatency, bool bViaLatencyElement) const override;
 	FTimeValue GetAnchorTime() const override;
@@ -548,6 +549,7 @@ public:
 	void GetTrackMetadata(TArray<FTrackMetadata>& OutMetadata, EStreamType InStreamType) const override;
 	void UpdateRunningMetaData(TSharedPtrTS<UtilsMP4::FMetadataParser> InUpdatedMetaData) override;
 	void UpdateDynamicRefetchCounter() override;
+	void PrepareForLooping(int32 InNumLoopsToAdd) override;
 	void TriggerClockSync(EClockSyncType InClockSyncType) override;
 	void TriggerPlaylistRefresh() override;
 	void ReachedStableBuffer() override;
diff --git a/Engine/Plugins/Media/ElectraPlayer/Source/ElectraPlayerRuntime/Private/Runtime/Player/Manifest.h b/Engine/Plugins/Media/ElectraPlayer/Source/ElectraPlayerRuntime/Private/Runtime/Player/Manifest.h
--- a/Engine/Plugins/Media/ElectraPlayer/Source/ElectraPlayerRuntime/Private/Runtime/Player/Manifest.h
+++ b/Engine/Plugins/Media/ElectraPlayer/Source/ElectraPlayerRuntime/Private/Runtime/Player/Manifest.h
@@ -156,6 +156,12 @@ namespace Electra
 			Live,							//!< A live presentation
 		};
 
+		enum class EReplayEventType
+		{
+			NoReplay,
+			IsReplay
+		};
+
 		enum class ESearchType
 		{
 			Closest,				//!< Find closest match
@@ -223,19 +229,33 @@ namespace Electra
 				switch(s)
 				{
 					case EType::Found:
+					{
 						return TEXT("Found");
+					}
 					case EType::NotFound:
+					{
 						return TEXT("Not found");
+					}
 					case EType::PastEOS:
+					{
 						return TEXT("Behind EOF");
+					}
 					case EType::BeforeStart:
+					{
 						return TEXT("Before start");
+					}
 					case EType::TryAgainLater:
+					{
 						return TEXT("Try again later");
+					}
 					case EType::NotLoaded:
+					{
 						return TEXT("Not loaded");
+					}
 					default:
+					{
 						return TEXT("undefined");
+					}
 				}
 			}
 		protected:
@@ -254,6 +274,9 @@ namespace Electra
 		//! Returns the type of this presentation, either on-demand or live.
 		virtual EType GetPresentationType() const = 0;
 
+		//! Returns the type of replay event.
+		virtual EReplayEventType GetReplayType() const = 0;
+
 		//! Returns the low-latency descriptor, if any. May return nullptr if there is none.
 		virtual TSharedPtrTS<const FLowLatencyDescriptor> GetLowLatencyDescriptor() const = 0;
 
@@ -385,6 +408,13 @@ namespace Electra
 		//! again. This is different to internal seeking for retry purposes where content will not be re-resolved.
 		virtual void UpdateDynamicRefetchCounter() = 0;
 
+		/**
+		 * Prepares the manifest for looping. This does not normally do anything except for when this is a replay,
+		 * in which case the internal timestamps may be modified such that the entire presentation is shifted into
+		 * the future by one presentation duration.
+		 */
+		virtual void PrepareForLooping(int32 InNumLoopsToAdd) = 0;
+
 		enum class EClockSyncType
 		{
 			Recommended,
diff --git a/Engine/Plugins/Media/ElectraPlayer/Source/ElectraPlayerRuntime/Private/Runtime/Player/mkv/ManifestMKV.h b/Engine/Plugins/Media/ElectraPlayer/Source/ElectraPlayerRuntime/Private/Runtime/Player/mkv/ManifestMKV.h
--- a/Engine/Plugins/Media/ElectraPlayer/Source/ElectraPlayerRuntime/Private/Runtime/Player/mkv/ManifestMKV.h
+++ b/Engine/Plugins/Media/ElectraPlayer/Source/ElectraPlayerRuntime/Private/Runtime/Player/mkv/ManifestMKV.h
@@ -30,6 +30,8 @@ public:
 	FErrorDetail Build(TSharedPtrTS<IParserMKV> MP4Parser, const FString& URL, const HTTP::FConnectionInfo& InConnectionInfo);
 
 	EType GetPresentationType() const override;
+	EReplayEventType GetReplayType() const override
+	{ return IManifest::EReplayEventType::NoReplay; }
 	TSharedPtrTS<const FLowLatencyDescriptor> GetLowLatencyDescriptor() const override
 	{ return nullptr; }
 	FTimeValue CalculateCurrentLiveLatency(const FTimeValue& InCurrentPlaybackPosition, const FTimeValue& InEncoderLatency, bool bViaLatencyElement) const override
@@ -64,6 +66,8 @@ public:
 	TRangeSet<double> GetPossiblePlaybackRates(EPlayRateType InForType) const override;
 	TSharedPtrTS<IProducerReferenceTimeInfo> GetProducerReferenceTimeInfo(int64 ID) const override;
 	void UpdateDynamicRefetchCounter() override;
+	void PrepareForLooping(int32 InNumLoopsToAdd) override
+	{ }
 	void TriggerClockSync(EClockSyncType InClockSyncType) override;
 	void TriggerPlaylistRefresh() override;
 	void ReachedStableBuffer() override
diff --git a/Engine/Plugins/Media/ElectraPlayer/Source/ElectraPlayerRuntime/Private/Runtime/Player/mp4/ManifestMP4.h b/Engine/Plugins/Media/ElectraPlayer/Source/ElectraPlayerRuntime/Private/Runtime/Player/mp4/ManifestMP4.h
--- a/Engine/Plugins/Media/ElectraPlayer/Source/ElectraPlayerRuntime/Private/Runtime/Player/mp4/ManifestMP4.h
+++ b/Engine/Plugins/Media/ElectraPlayer/Source/ElectraPlayerRuntime/Private/Runtime/Player/mp4/ManifestMP4.h
@@ -31,6 +31,8 @@ public:
 	FErrorDetail Build(TSharedPtrTS<IParserISO14496_12> MP4Parser, const FString& URL, const HTTP::FConnectionInfo& InConnectionInfo);
 
 	EType GetPresentationType() const override;
+	EReplayEventType GetReplayType() const override
+	{ return IManifest::EReplayEventType::NoReplay; }
 	TSharedPtrTS<const FLowLatencyDescriptor> GetLowLatencyDescriptor() const override
 	{ return nullptr; }
 	FTimeValue CalculateCurrentLiveLatency(const FTimeValue& InCurrentPlaybackPosition, const FTimeValue& InEncoderLatency, bool bViaLatencyElement) const override
@@ -65,6 +67,8 @@ public:
 	TRangeSet<double> GetPossiblePlaybackRates(EPlayRateType InForType) const override;
 	TSharedPtrTS<IProducerReferenceTimeInfo> GetProducerReferenceTimeInfo(int64 ID) const override;
 	void UpdateDynamicRefetchCounter() override;
+	void PrepareForLooping(int32 InNumLoopsToAdd) override
+	{ }
 	void TriggerClockSync(EClockSyncType InClockSyncType) override;
 	void TriggerPlaylistRefresh() override;
 	void ReachedStableBuffer() override
diff --git a/Engine/Plugins/Media/ElectraPlayer/Source/ElectraPlayerRuntime/Private/Runtime/Player/mpegaudio/ManifestMPEGAudio.h b/Engine/Plugins/Media/ElectraPlayer/Source/ElectraPlayerRuntime/Private/Runtime/Player/mpegaudio/ManifestMPEGAudio.h
--- a/Engine/Plugins/Media/ElectraPlayer/Source/ElectraPlayerRuntime/Private/Runtime/Player/mpegaudio/ManifestMPEGAudio.h
+++ b/Engine/Plugins/Media/ElectraPlayer/Source/ElectraPlayerRuntime/Private/Runtime/Player/mpegaudio/ManifestMPEGAudio.h
@@ -79,6 +79,8 @@ public:
 	FErrorDetail Build(const FMPEGAudioInfoHeader& InHeader, const FString& URL);
 
 	EType GetPresentationType() const override;
+	EReplayEventType GetReplayType() const override
+	{ return IManifest::EReplayEventType::NoReplay; }
 	TSharedPtrTS<const FLowLatencyDescriptor> GetLowLatencyDescriptor() const override;
 	FTimeValue CalculateCurrentLiveLatency(const FTimeValue& InCurrentPlaybackPosition, const FTimeValue& InEncoderLatency, bool bViaLatencyElement) const override;
 	FTimeValue GetAnchorTime() const override
@@ -107,6 +109,8 @@ public:
 	TRangeSet<double> GetPossiblePlaybackRates(EPlayRateType InForType) const override;
 	TSharedPtrTS<IProducerReferenceTimeInfo> GetProducerReferenceTimeInfo(int64 ID) const override;
 	void UpdateDynamicRefetchCounter() override;
+	void PrepareForLooping(int32 InNumLoopsToAdd) override
+	{ }
 	void TriggerClockSync(EClockSyncType InClockSyncType) override;
 	void TriggerPlaylistRefresh() override;
 	void ReachedStableBuffer() override
diff --git a/Engine/Plugins/Media/ElectraUtil/Source/ElectraBase/Private/URLParser.cpp b/Engine/Plugins/Media/ElectraUtil/Source/ElectraBase/Private/URLParser.cpp
--- a/Engine/Plugins/Media/ElectraUtil/Source/ElectraBase/Private/URLParser.cpp
+++ b/Engine/Plugins/Media/ElectraUtil/Source/ElectraBase/Private/URLParser.cpp
@@ -338,14 +338,25 @@ namespace Electra
 			ValuePairs[i].FindChar(TCHAR('='), EqPos);
 			FQueryParam qp;
 			bool bOk = true;
+			FString k, v;
+			if (EqPos != INDEX_NONE)
+			{
+				k = ValuePairs[i].Mid(0, EqPos);
+				v = ValuePairs[i].Mid(EqPos+1);
+			}
+			else
+			{
+				k = ValuePairs[i];
+			}
+
 			if (bPerformUrlDecoding)
 			{
-				bOk = UrlDecode(qp.Name, ValuePairs[i].Mid(0, EqPos)) && UrlDecode(qp.Value, ValuePairs[i].Mid(EqPos+1));
+				bOk = UrlDecode(qp.Name, k) && UrlDecode(qp.Value, v);
 			}
 			else
 			{
-				qp.Name = ValuePairs[i].Mid(0, EqPos);
-				qp.Value = ValuePairs[i].Mid(EqPos+1);
+				qp.Name = MoveTemp(k);
+				qp.Value = MoveTemp(v);
 			}
 			if (bOk)
 			{
