ElectraPlayer: Adjusting the maximum Live edge latency in case the media segments have a larger duration that would put the start with its implicit startup latency over the maximum.

#rnx

diff --git a/Engine/Plugins/Media/ElectraPlayer/Source/ElectraPlayerRuntime/Private/Runtime/Player/ABRRules/ABRLiveStream.cpp b/Engine/Plugins/Media/ElectraPlayer/Source/ElectraPlayerRuntime/Private/Runtime/Player/ABRRules/ABRLiveStream.cpp
--- a/Engine/Plugins/Media/ElectraPlayer/Source/ElectraPlayerRuntime/Private/Runtime/Player/ABRRules/ABRLiveStream.cpp
+++ b/Engine/Plugins/Media/ElectraPlayer/Source/ElectraPlayerRuntime/Private/Runtime/Player/ABRRules/ABRLiveStream.cpp
@@ -426,22 +426,14 @@ void FABRLiveStream::PrepareLatencyConfiguration()
 	FScopeLock lock(&Lock);
 	LatencyConfig.Reset();
 	LatencyConfig.TargetLatency = TargetLatency;
-	// If the target latency is greater than our assumed default value we need to adjust our value.
-	if (TargetLatency > LatencyConfig.MaxLatency)
+	// The target latency should be at least half the maximum permitted latency.
+	// When the maximum latency is reached a jump to the Live edge may be performed, so
+	// some wiggle room is needed to avoid the target to be close to the maximum, especially
+	// if the media segment duration is quite large and starting playback may put us further
+	// back and possibly over the permitted maximum.
+	if (TargetLatency * 2.0 > LatencyConfig.MaxLatency)
 	{
-		// If the media timeline is valid we set the latency to the duration of the timeline.
-		// While this may be far too large it will at least prevent falling behind what the timeline indicates.
-		// If there is a latency element present this value will be replaced by the latency element, so this
-		// is a fallback for streams without a latency element (and therefore not low-latency).
-		if (TimelineRange.IsValid())
-		{
-			LatencyConfig.MaxLatency = (TimelineRange.End - TimelineRange.Start).GetAsSeconds();
-		}
-		else
-		{
-			// Arbitrary fallback value.
-			LatencyConfig.MaxLatency = 60.0;
-		}
+		LatencyConfig.MaxLatency = TargetLatency * 2.0;
 	}
 	if (lld.IsValid())
 	{
