ElectraPlayer: Emit warning about mismatched media segment duration only when the duration check was enabled and failed.

#rnx

diff --git a/Engine/Plugins/Media/ElectraPlayer/Source/ElectraPlayerRuntime/Private/Runtime/Player/HLS/PlaylistHLS.cpp b/Engine/Plugins/Media/ElectraPlayer/Source/ElectraPlayerRuntime/Private/Runtime/Player/HLS/PlaylistHLS.cpp
--- a/Engine/Plugins/Media/ElectraPlayer/Source/ElectraPlayerRuntime/Private/Runtime/Player/HLS/PlaylistHLS.cpp
+++ b/Engine/Plugins/Media/ElectraPlayer/Source/ElectraPlayerRuntime/Private/Runtime/Player/HLS/PlaylistHLS.cpp
@@ -3346,7 +3346,7 @@ IManifest::FResult FActiveHLSPlaylist::FPlayPeriod::GetLoopingSegment(TSharedPtr
 void FActiveHLSPlaylist::FPlayPeriod::ValidateDownloadedSegmentDuration(const FStreamSegmentRequestCommon* InRequest)
 {
 	// For successful downloads check if the segment duration was as specified.
-	if (InRequest && InRequest->DownloadStats.bWasSuccessful && !InRequest->DownloadStats.bInsertedFillerData)
+	if (InRequest && InRequest->DownloadStats.bWasSuccessful && !InRequest->DownloadStats.bInsertedFillerData && !InRequest->TimestampVars.Next.bFailed)
 	{
 		// The #EXTINF duration needs to be as precise to the actual media duration as possible.
 		// Otherwise playback errors could occur and switching across renditions may not work correctly.
