ElectraPlayer: Improved support for replay events; added ability to turn a HLS VoD stream into a replay

#rnx

diff --git a/Engine/Plugins/Media/ElectraPlayer/Source/ElectraPlayerRuntime/Private/ElectraPlayer.cpp b/Engine/Plugins/Media/ElectraPlayer/Source/ElectraPlayerRuntime/Private/ElectraPlayer.cpp
--- a/Engine/Plugins/Media/ElectraPlayer/Source/ElectraPlayerRuntime/Private/ElectraPlayer.cpp
+++ b/Engine/Plugins/Media/ElectraPlayer/Source/ElectraPlayerRuntime/Private/ElectraPlayer.cpp
@@ -1191,6 +1191,13 @@ void FElectraPlayer::TriggerFirstSeekIfNecessary()
 		auto Player = CurrentPlayer;
 		if (Player.IsValid())
 		{
+			// If the media is flagged as a replay event then the start time is also implicitly provided and we must not set one.
+			IAdaptiveStreamingPlayer::FEventReplayState evtrs;
+			Player->AdaptivePlayer->GetReplayEventState(evtrs);
+			if (evtrs.bIsReplayEvent)
+			{
+				playParam.Time.SetToInvalid();
+			}
 			Player->AdaptivePlayer->SeekTo(playParam);
 		}
 	}
diff --git a/Engine/Plugins/Media/ElectraPlayer/Source/ElectraPlayerRuntime/Private/Runtime/Player/AdaptivePlayerOptionKeynames.h b/Engine/Plugins/Media/ElectraPlayer/Source/ElectraPlayerRuntime/Private/Runtime/Player/AdaptivePlayerOptionKeynames.h
--- a/Engine/Plugins/Media/ElectraPlayer/Source/ElectraPlayerRuntime/Private/Runtime/Player/AdaptivePlayerOptionKeynames.h
+++ b/Engine/Plugins/Media/ElectraPlayer/Source/ElectraPlayerRuntime/Private/Runtime/Player/AdaptivePlayerOptionKeynames.h
@@ -70,6 +70,14 @@ namespace Electra
 
 	const FName OptionKeyCMCDConfiguration(TEXT("cmcd_configuration"));
 
+
+	namespace CustomOptions
+	{
+		const TCHAR* const Custom_EpicStaticStart = TEXT("EpicStaticStart");
+		const TCHAR* const Custom_EpicDynamicStart = TEXT("EpicDynamicStart");
+		const TCHAR* const Custom_EpicUTCUrl = TEXT("EpicUTCUrl");
+		const TCHAR* const Custom_EpicUTCNow = TEXT("EpicUTCNow");
+	}
 } // namespace Electra
 
 
diff --git a/Engine/Plugins/Media/ElectraPlayer/Source/ElectraPlayerRuntime/Private/Runtime/Player/AdaptiveStreamingPlayer.cpp b/Engine/Plugins/Media/ElectraPlayer/Source/ElectraPlayerRuntime/Private/Runtime/Player/AdaptiveStreamingPlayer.cpp
--- a/Engine/Plugins/Media/ElectraPlayer/Source/ElectraPlayerRuntime/Private/Runtime/Player/AdaptiveStreamingPlayer.cpp
+++ b/Engine/Plugins/Media/ElectraPlayer/Source/ElectraPlayerRuntime/Private/Runtime/Player/AdaptiveStreamingPlayer.cpp
@@ -562,10 +562,6 @@ void FAdaptiveStreamingPlayer::LoadManifest(const FString& InManifestURL)
 	WorkerThread.SendLoadManifestMessage(InManifestURL, FString());
 }
 
-void FAdaptiveStreamingPlayer::LoadBlob(TSharedPtr<FHTTPResourceRequest, ESPMode::ThreadSafe> InBlobLoadRequest)
-{
-	WorkerThread.SendLoadBlobMessage(InBlobLoadRequest);
-}
 
 
 //-----------------------------------------------------------------------------
@@ -766,6 +762,12 @@ void FAdaptiveStreamingPlayer::GetLoopState(FLoopState& OutLoopState) const
 	}
 }
 
+void FAdaptiveStreamingPlayer::GetReplayEventState(FEventReplayState& OutReplayEventState) const
+{
+	OutReplayEventState.bIsReplayEvent = Manifest.IsValid() ? Manifest->GetReplayType() != IManifest::EReplayEventType::NoReplay : EventReplayParams.HaveReplayParams();
+}
+
+
 //-----------------------------------------------------------------------------
 /**
  * Returns track metadata of the currently active play period.
@@ -1500,15 +1502,6 @@ bool FAdaptiveStreamingPlayer::InternalHandleThreadMessages()
 				InternalLoadManifest(ev.URL, ev.MimeType);
 				break;
 			}
-			case FWorkerThreadMessages::FMessage::EType::LoadBlob:
-			{
-				FWorkerThreadMessages::FMessage::FLoadBlob& ev = msg.Data.Get<FWorkerThreadMessages::FMessage::FLoadBlob>();
-				if (ev.BlobLoadRequest.IsValid())
-				{
-					ev.BlobLoadRequest->StartGet(this);
-				}
-				break;
-			}
 			case FWorkerThreadMessages::FMessage::EType::Pause:
 			{
 				bShouldBePaused = true;
diff --git a/Engine/Plugins/Media/ElectraPlayer/Source/ElectraPlayerRuntime/Private/Runtime/Player/AdaptiveStreamingPlayer.h b/Engine/Plugins/Media/ElectraPlayer/Source/ElectraPlayerRuntime/Private/Runtime/Player/AdaptiveStreamingPlayer.h
--- a/Engine/Plugins/Media/ElectraPlayer/Source/ElectraPlayerRuntime/Private/Runtime/Player/AdaptiveStreamingPlayer.h
+++ b/Engine/Plugins/Media/ElectraPlayer/Source/ElectraPlayerRuntime/Private/Runtime/Player/AdaptiveStreamingPlayer.h
@@ -114,9 +114,6 @@ public:
 	//-------------------------------------------------------------------------
 	// Manifest loading functions
 	//
-	//! Issues loading of a binary blob of arbitrary data.
-	virtual void LoadBlob(TSharedPtr<FHTTPResourceRequest, ESPMode::ThreadSafe> InBlobLoadRequest) = 0;
-
 	//! Issue a load and parse of the manifest/main playlist file. Make initial stream selection choice by calling SetInitialStreamAttributes() beforehand.
 	virtual void LoadManifest(const FString& ManifestURL) = 0;
 
@@ -309,6 +306,14 @@ public:
 	//! Returns attributes of the currently selected track.
 	virtual void GetSelectedTrackAttributes(FStreamSelectionAttributes& OutAttributes, EStreamType StreamType) const = 0;
 
+	struct FEventReplayState
+	{
+		FTimeValue ProposedJoinTime;
+		bool bIsReplayEvent = false;
+	};
+	//! Returns if the media is set to be a replay (configured via URL fragment parameters)
+	virtual void GetReplayEventState(FEventReplayState& OutReplayEventState) const = 0;
+
 
 	//-------------------------------------------------------------------------
 	// Manual stream selection functions
diff --git a/Engine/Plugins/Media/ElectraPlayer/Source/ElectraPlayerRuntime/Private/Runtime/Player/AdaptiveStreamingPlayerInternal.h b/Engine/Plugins/Media/ElectraPlayer/Source/ElectraPlayerRuntime/Private/Runtime/Player/AdaptiveStreamingPlayerInternal.h
--- a/Engine/Plugins/Media/ElectraPlayer/Source/ElectraPlayerRuntime/Private/Runtime/Player/AdaptiveStreamingPlayerInternal.h
+++ b/Engine/Plugins/Media/ElectraPlayer/Source/ElectraPlayerRuntime/Private/Runtime/Player/AdaptiveStreamingPlayerInternal.h
@@ -1121,7 +1121,6 @@ public:
 
 	void SetInitialStreamAttributes(EStreamType StreamType, FStreamSelectionAttributes InitialSelection) override;
 	void EnableFrameAccurateSeeking(bool bEnabled) override;
-	void LoadBlob(TSharedPtr<FHTTPResourceRequest, ESPMode::ThreadSafe> InBlobLoadRequest) override;
 	void LoadManifest(const FString& manifestURL) override;
 
 	void SeekTo(const FSeekParam& NewPosition) override;
@@ -1152,6 +1151,7 @@ public:
 	void GetTrackMetadata(TArray<FTrackMetadata>& OutTrackMetadata, EStreamType StreamType) const override;
 	//void GetSelectedTrackMetadata(TOptional<FTrackMetadata>& OutSelectedTrackMetadata, EStreamType StreamType) const override;
 	void GetSelectedTrackAttributes(FStreamSelectionAttributes& OutAttributes, EStreamType StreamType) const override;
+	void GetReplayEventState(FEventReplayState& OutReplayEventState) const override;
 
 	void SetBitrateCeiling(int32 highestSelectableBitrate) override;
 	void SetMaxResolution(int32 MaxWidth, int32 MaxHeight) override;
@@ -1551,7 +1551,6 @@ private:
 				Initialize,
 				ChangeOptions,
 				LoadManifest,
-				LoadBlob,
 				Pause,
 				Resume,
 				Close,
@@ -1679,14 +1678,6 @@ private:
 			Msg.Data.Emplace<FMessage::FLoadManifest>(MoveTemp(lm));
 			TriggerSharedWorkerThread(MoveTemp(Msg));
 		}
-		void SendLoadBlobMessage(const TSharedPtr<FHTTPResourceRequest, ESPMode::ThreadSafe>& InBlobLoadRequest)
-		{
-			FMessage::FLoadBlob lb { InBlobLoadRequest };
-			FMessage Msg;
-			Msg.Type = FMessage::EType::LoadBlob;
-			Msg.Data.Emplace<FMessage::FLoadBlob>(MoveTemp(lb));
-			TriggerSharedWorkerThread(MoveTemp(Msg));
-		}
 		void SendPauseMessage()
 		{
 			FMessage Msg;
@@ -2277,6 +2268,8 @@ private:
 	TSharedPtrTS<IPlaylistReader>										ManifestReader;
 	TSharedPtrTS<FHTTPResourceRequest>									ManifestMimeTypeRequest;
 	FAdaptiveStreamingPlayerPlaylistPropertyHandler						PlaylistPropertyHandler;
+	Playlist::FReplayEventParams										EventReplayParams;
+
 
 	IStreamReader*														StreamReaderHandler;
 	bool																bStreamingHasStarted;
diff --git a/Engine/Plugins/Media/ElectraPlayer/Source/ElectraPlayerRuntime/Private/Runtime/Player/AdaptiveStreamingPlayerManifest.cpp b/Engine/Plugins/Media/ElectraPlayer/Source/ElectraPlayerRuntime/Private/Runtime/Player/AdaptiveStreamingPlayerManifest.cpp
--- a/Engine/Plugins/Media/ElectraPlayer/Source/ElectraPlayerRuntime/Private/Runtime/Player/AdaptiveStreamingPlayerManifest.cpp
+++ b/Engine/Plugins/Media/ElectraPlayer/Source/ElectraPlayerRuntime/Private/Runtime/Player/AdaptiveStreamingPlayerManifest.cpp
@@ -262,7 +262,14 @@ void FAdaptiveStreamingPlayer::InternalLoadManifest(const FString& InURL, const
 
 			if (ManifestReader.IsValid())
 			{
-				ManifestReader->LoadAndParse(URL);
+				if (EventReplayParams.ParseFromURL(URL))
+				{
+					FLoopState ls;
+					PlaybackState.GetLoopState(ls);
+					URL = EventReplayParams.GetCleanedURL();
+					EventReplayParams.SetLoopingAtStart(ls.bIsEnabled);
+				}
+				ManifestReader->LoadAndParse(URL, EventReplayParams);
 			}
 		}
 		else
diff --git a/Engine/Plugins/Media/ElectraPlayer/Source/ElectraPlayerRuntime/Private/Runtime/Player/DASH/ManifestBuilderDASH.cpp b/Engine/Plugins/Media/ElectraPlayer/Source/ElectraPlayerRuntime/Private/Runtime/Player/DASH/ManifestBuilderDASH.cpp
--- a/Engine/Plugins/Media/ElectraPlayer/Source/ElectraPlayerRuntime/Private/Runtime/Player/DASH/ManifestBuilderDASH.cpp
+++ b/Engine/Plugins/Media/ElectraPlayer/Source/ElectraPlayerRuntime/Private/Runtime/Player/DASH/ManifestBuilderDASH.cpp
@@ -58,11 +58,6 @@ namespace Electra
 
 namespace
 {
-	const TCHAR* const Custom_EpicStaticStart = TEXT("EpicStaticStart");
-	const TCHAR* const Custom_EpicDynamicStart = TEXT("EpicDynamicStart");
-	const TCHAR* const Custom_EpicUTCUrl = TEXT("EpicUTCUrl");
-	const TCHAR* const Custom_EpicUTCNow = TEXT("EpicUTCNow");
-
 	const TCHAR* const XLinkActuateOnLoad = TEXT("onLoad");
 	const TCHAR* const XLinkActuateOnRequest = TEXT("onRequest");
 	const TCHAR* const XLinkResolveToZero = TEXT("urn:mpeg:dash:resolve-to-zero:2013");
@@ -1152,16 +1147,17 @@ void FManifestDASHInternal::InjectEpicTimingSources()
 {
 	const TArray<TSharedPtrTS<FDashMPD_DescriptorType>>& UTCTimings = MPDRoot->GetUTCTimings();
 
-	for(int32 i=0,iMax=URLFragmentComponents.Num(); i<iMax; ++i)
+	const TArray<FURL_RFC3986::FQueryParam>& frgs(ReplayEventParams.GetUrlFragmentParams());
+	for(int32 i=0,iMax=frgs.Num(); i<iMax; ++i)
 	{
-		if (URLFragmentComponents[i].Name.Equals(Custom_EpicUTCUrl))
+		if (frgs[i].Name.Equals(CustomOptions::Custom_EpicUTCUrl))
 		{
-			MPDRoot->AddUTCTimingElement(FString(), DASH::Schemes::TimingSources::Scheme_urn_mpeg_dash_utc_httpxsdate2014, FString(), URLFragmentComponents[i].Value);
+			MPDRoot->AddUTCTimingElement(FString(), DASH::Schemes::TimingSources::Scheme_urn_mpeg_dash_utc_httpxsdate2014, FString(), frgs[i].Value);
 			bDidInjectUTCTimingElements = true;
 		}
-		else if (URLFragmentComponents[i].Name.Equals(Custom_EpicUTCNow))
+		else if (frgs[i].Name.Equals(CustomOptions::Custom_EpicUTCNow))
 		{
-			MPDRoot->AddUTCTimingElement(FString(), DASH::Schemes::TimingSources::Scheme_urn_mpeg_dash_utc_direct2014, FString(), URLFragmentComponents[i].Value);
+			MPDRoot->AddUTCTimingElement(FString(), DASH::Schemes::TimingSources::Scheme_urn_mpeg_dash_utc_direct2014, FString(), frgs[i].Value);
 			bDidInjectUTCTimingElements = true;
 		}
 	}
@@ -1175,16 +1171,25 @@ void FManifestDASHInternal::TransformIntoEpicEvent()
 		FString Arg;
 		bool bStaticStart = false;
 		bool bDynamicStart = false;
-		for(int32 i=0,iMax=URLFragmentComponents.Num(); i<iMax; ++i)
+
+		const TArray<FURL_RFC3986::FQueryParam>& frgs(ReplayEventParams.GetUrlFragmentParams());
+		for(int32 i=0,iMax=frgs.Num(); i<iMax; ++i)
 		{
-			bStaticStart = URLFragmentComponents[i].Name.Equals(Custom_EpicStaticStart);
-			bDynamicStart = URLFragmentComponents[i].Name.Equals(Custom_EpicDynamicStart);
+			bStaticStart = frgs[i].Name.Equals(CustomOptions::Custom_EpicStaticStart);
+			bDynamicStart = frgs[i].Name.Equals(CustomOptions::Custom_EpicDynamicStart);
 			if (bStaticStart || bDynamicStart)
 			{
-				Arg = URLFragmentComponents[i].Value;
+				Arg = frgs[i].Value;
 				break;
 			}
 		}
+
+		// Any indication that this is a replay event?
+		if (!bStaticStart && !bDynamicStart)
+		{
+			return;
+		}
+
 		TArray<FString> Params;
 		if (!Arg.IsEmpty())
 		{
@@ -1251,7 +1256,6 @@ void FManifestDASHInternal::TransformIntoEpicEvent()
 
 void FManifestDASHInternal::PrepareEpicEventForLooping(int32 InNumLoopsToAdd)
 {
-	check(OriginalPresentationDuration.IsValid());
 	if (OriginalPresentationDuration.IsValid())
 	{
 		FTimeValue ast = GetAnchorTime();
diff --git a/Engine/Plugins/Media/ElectraPlayer/Source/ElectraPlayerRuntime/Private/Runtime/Player/DASH/PlaylistReaderDASH.cpp b/Engine/Plugins/Media/ElectraPlayer/Source/ElectraPlayerRuntime/Private/Runtime/Player/DASH/PlaylistReaderDASH.cpp
--- a/Engine/Plugins/Media/ElectraPlayer/Source/ElectraPlayerRuntime/Private/Runtime/Player/DASH/PlaylistReaderDASH.cpp
+++ b/Engine/Plugins/Media/ElectraPlayer/Source/ElectraPlayerRuntime/Private/Runtime/Player/DASH/PlaylistReaderDASH.cpp
@@ -58,7 +58,7 @@ public:
 		static FString Type("dash");
 		return Type;
 	}
-	void LoadAndParse(const FString& URL) override;
+	void LoadAndParse(const FString& URL, const Playlist::FReplayEventParams& InReplayEventParams) override;
 	FString GetURL() const override;
 	TSharedPtrTS<IManifest> GetManifest() override;
 	void AddElementLoadRequests(const TArray<TWeakPtrTS<FMPDLoadRequestDASH>>& RemoteElementLoadRequests) override;
@@ -121,7 +121,8 @@ private:
 
 	IPlayerSessionServices*									PlayerSessionServices = nullptr;
 	FString													MPDURL;
-	FString													Fragment;
+	TArray<FURL_RFC3986::FQueryParam>						URLFragmentComponents;
+	Playlist::FReplayEventParams							ReplayEventParams;
 	bool													bIsInitialized = false;
 
 	FCriticalSection										RequestsLock;
@@ -303,9 +304,10 @@ void FPlaylistReaderDASH::LogMessage(IInfoLog::ELevel Level, const FString& Mess
 	}
 }
 
-void FPlaylistReaderDASH::LoadAndParse(const FString& URL)
+void FPlaylistReaderDASH::LoadAndParse(const FString& URL, const Playlist::FReplayEventParams& InReplayEventParams)
 {
 	MPDURL = URL;
+	ReplayEventParams = InReplayEventParams;
 	InternalSetup();
 }
 
@@ -829,7 +831,7 @@ void FPlaylistReaderDASH::InternalSetup()
 	FURL_RFC3986 UrlParser;
 	UrlParser.Parse(MPDURL);
 	PlaylistLoadRequest->Url.URL = UrlParser.Get(true, false);
-	Fragment = UrlParser.GetFragment();
+	FURL_RFC3986::GetQueryParams(URLFragmentComponents, UrlParser.GetFragment(), false);
 	PlaylistLoadRequest->CompleteCallback.BindThreadSafeSP(AsShared(), &FPlaylistReaderDASH::ManifestDownloadCompleted);
 	EnqueueResourceRequest(MoveTemp(PlaylistLoadRequest));
 	bIsInitialized = true;
@@ -1238,10 +1240,8 @@ void FPlaylistReaderDASH::ManifestDownloadCompleted(FResourceLoadRequestPtr Requ
 					// Update the CMCD parameters from the root level service descriptions.
 					UpdateCMCDParamsFromServiceDescriptions(NewManifest->GetMPDRoot()->GetServiceDescriptions());
 
-					// Parse the URL fragment into its components. For a DASH URL the fragment is constructed like a query string with & delimited key/value pairs.
-					TArray<FURL_RFC3986::FQueryParam> URLFragmentComponents;
-					FURL_RFC3986::GetQueryParams(URLFragmentComponents, Fragment, false);	// The fragment is already URL escaped, so no need to do it again.
-					NewManifest->SetURLFragmentComponents(MoveTemp(URLFragmentComponents));
+					NewManifest->SetURLFragmentComponents(URLFragmentComponents);
+					NewManifest->SetReplayEventParams(ReplayEventParams);
 					// Inject fake <UTCTiming> elements that were passed in the fragments.
 					NewManifest->InjectEpicTimingSources();
 				}
diff --git a/Engine/Plugins/Media/ElectraPlayer/Source/ElectraPlayerRuntime/Private/Runtime/Player/DASH/PlaylistReaderDASH.h b/Engine/Plugins/Media/ElectraPlayer/Source/ElectraPlayerRuntime/Private/Runtime/Player/DASH/PlaylistReaderDASH.h
--- a/Engine/Plugins/Media/ElectraPlayer/Source/ElectraPlayerRuntime/Private/Runtime/Player/DASH/PlaylistReaderDASH.h
+++ b/Engine/Plugins/Media/ElectraPlayer/Source/ElectraPlayerRuntime/Private/Runtime/Player/DASH/PlaylistReaderDASH.h
@@ -25,9 +25,10 @@ public:
 	/**
 	 * Loads and parses the MPD.
 	 *
-	 * @param URL     URL of the MPD to load
+	 * @param URL URL of the playlist to load
+	 * @param InReplayEventParams URL fragments pertaining to replay event which have been removed from the source URL
 	 */
-	virtual void LoadAndParse(const FString& URL) = 0;
+	virtual void LoadAndParse(const FString& InURL, const Playlist::FReplayEventParams& InReplayEventParams) = 0;
 
 	/**
 	 * Returns the URL from which the MPD was loaded (or supposed to be loaded).
diff --git a/Engine/Plugins/Media/ElectraPlayer/Source/ElectraPlayerRuntime/Private/Runtime/Player/DASH/PlaylistReaderDASH_Internal.h b/Engine/Plugins/Media/ElectraPlayer/Source/ElectraPlayerRuntime/Private/Runtime/Player/DASH/PlaylistReaderDASH_Internal.h
--- a/Engine/Plugins/Media/ElectraPlayer/Source/ElectraPlayerRuntime/Private/Runtime/Player/DASH/PlaylistReaderDASH_Internal.h
+++ b/Engine/Plugins/Media/ElectraPlayer/Source/ElectraPlayerRuntime/Private/Runtime/Player/DASH/PlaylistReaderDASH_Internal.h
@@ -801,6 +801,16 @@ public:
 		return URLFragmentComponents;
 	}
 
+	void SetReplayEventParams(const Playlist::FReplayEventParams& InEventReplayParams)
+	{
+		ReplayEventParams = InEventReplayParams;
+	}
+
+	const Playlist::FReplayEventParams& GetReplayEventParams()
+	{
+		return ReplayEventParams;
+	}
+
 	void InjectEpicTimingSources();
 	void TransformIntoEpicEvent();
 	void PrepareEpicEventForLooping(int32 InNumLoopsToAdd);
@@ -876,6 +886,7 @@ private:
 
 	// The MPD URL fragment components
 	TArray<FURL_RFC3986::FQueryParam> URLFragmentComponents;
+	Playlist::FReplayEventParams ReplayEventParams;
 
 	// Type of the presentation.
 	EPresentationType PresentationType;
diff --git a/Engine/Plugins/Media/ElectraPlayer/Source/ElectraPlayerRuntime/Private/Runtime/Player/HLS/PlaylistHLS.cpp b/Engine/Plugins/Media/ElectraPlayer/Source/ElectraPlayerRuntime/Private/Runtime/Player/HLS/PlaylistHLS.cpp
--- a/Engine/Plugins/Media/ElectraPlayer/Source/ElectraPlayerRuntime/Private/Runtime/Player/HLS/PlaylistHLS.cpp
+++ b/Engine/Plugins/Media/ElectraPlayer/Source/ElectraPlayerRuntime/Private/Runtime/Player/HLS/PlaylistHLS.cpp
@@ -13,6 +13,7 @@
 #include "Crypto/StreamCryptoAES128.h"
 #include "Misc/Optional.h"
 #include "Misc/SecureHash.h"
+#include "ElectraPlayerPrivate.h"
 
 namespace Electra
 {
@@ -1051,13 +1052,21 @@ FTimeRange FActiveHLSPlaylist::FTimelineMediaAsset::GetSeekableTimeRange() const
 {
 	FTimeRange tr = GetTimeRange();
 	tr.End -= GetDesiredLiveLatency();
-	check(tr.End >= tr.Start);
+	if (tr.End < tr.Start)
+	{
+		tr.End = tr.Start;
+	}
 	return tr;
 }
 
 
 FTimeValue FActiveHLSPlaylist::FTimelineMediaAsset::GetDesiredLiveLatency() const
 {
+	if (ReplayEvent.bIsReplay)
+	{
+		return ReplayEvent.SuggestedPresentationDelay.IsValid() ? ReplayEvent.SuggestedPresentationDelay : TargetDuration * 3;
+	}
+
 	// Called for a static playlist? If so, there is no desired latency.
 	if (PlaylistType == FPlaylistParserHLS::EPlaylistType::VOD || bHasEndList)
 	{
@@ -1100,7 +1109,7 @@ FTimeValue FActiveHLSPlaylist::FTimelineMediaAsset::GetDesiredLiveLatency() cons
 FTimeValue FActiveHLSPlaylist::FTimelineMediaAsset::CalculateCurrentLiveLatency(const FTimeValue& InCurrentPlaybackPosition, const FTimeValue& InEncoderLatency)
 {
 	FTimeValue LiveLatency;
-	if (PlaylistType != FPlaylistParserHLS::EPlaylistType::VOD && !bHasEndList)
+	if ((ReplayEvent.bIsReplay && !ReplayEvent.bIsStaticEvent) || (PlaylistType != FPlaylistParserHLS::EPlaylistType::VOD && !bHasEndList))
 	{
 		FTimeValue Now = PlayerSessionServices->GetSynchronizedUTCTime()->GetTime();
 		// With PDT we have the timeline locked to current wallclock 'Now', so the latency is
@@ -1173,7 +1182,7 @@ FTimeValue FActiveHLSPlaylist::FTimelineMediaAsset::CalculateStartTime(const TSh
 {
 	FTimeValue tv;
 	const FStartTimeHLS& StartTime(InPlaylist->StartTime);
-	if (StartTime.Offset.IsValid())
+	if (!ReplayEvent.bIsReplay && StartTime.Offset.IsValid())
 	{
 		FTimeRange tr = GetTimeRange();
 
@@ -1294,8 +1303,10 @@ void FActiveHLSPlaylist::FTimelineMediaAsset::UpdateWithMediaPlaylist(TSharedPtr
 			bHasProgramDateTime = mp->bHasProgramDateTime;
 			check(!bHasProgramDateTime || (bHasProgramDateTime && FirstProgramDateTime.IsValid()));
 			MultiVariantURLFragmentComponents = InMediaPlaylist->MultiVariantURLFragmentComponents;
+			ReplayEventParams = InMediaPlaylist->ReplayEventParams;
 			// Establish the time offset between the current time and the playlist (zero for VOD)
 			BaseTimeOffset = CalculatePlaylistTimeOffset(mp);
+			TransformIntoReplayEvent();
 			FTimeValue StartAt = CalculateStartTime(mp);
 			// Get the start range from the URL fragment parameters. This overrules any EXT-X-START value.
 			DefaultStartAndEndTime = GetPlaybackRangeFromURL(IManifest::EPlaybackRangeType::TemporaryPlaystartRange);
@@ -1355,33 +1366,57 @@ FTimeRange FActiveHLSPlaylist::FTimelineMediaAsset::GetPlaybackRangeFromURL(EPla
 		return FromTo;
 	}
 
-	const TCHAR* const TimeDelimiter = TEXT(",");
-	TArray<FString> TimeRange;
-	FTimeValue Offset;
-	Time.ParseIntoArray(TimeRange, TimeDelimiter, false);
-	if (TimeRange.Num() && !TimeRange[0].IsEmpty())
+	if (InRangeType == IManifest::EPlaybackRangeType::TemporaryPlaystartRange && Time.Equals(TEXT("posix:now")))
 	{
-		if (RFC2326::ParseNPTTime(Offset, TimeRange[0]))
+		FTimeValue Now = PlayerSessionServices->GetSynchronizedUTCTime()->GetTime();
+		// A static event will not use an updated wallclock NOW, so if 'now' is used we do an init
+		// with the current time.
+		if (ReplayEvent.bIsReplay)
 		{
-			if (!FromTo.Start.IsValid())
+			if (ReplayEvent.bIsStaticEvent)
 			{
-				FromTo.Start = Offset;
+				FromTo.Start = Now;
 			}
 			else
 			{
-				FromTo.Start += Offset;
+				// 'now' is dynamic. The time will continue to flow between here where we set the value and
+				// the moment playback will begin with buffered data.
+				// We do not lock 'now' with the current time but leave it unset. This results in the start
+				// time to be taken from the seekable range's end value which is updating dynamically.
+				FromTo.Start.SetToInvalid();
 			}
 		}
 	}
-	if (TimeRange.Num() > 1 && !TimeRange[1].IsEmpty())
+	else
 	{
-		if (RFC2326::ParseNPTTime(Offset, TimeRange[1]))
+		const TCHAR* const TimeDelimiter = TEXT(",");
+		TArray<FString> TimeRange;
+		FTimeValue Offset;
+		Time.ParseIntoArray(TimeRange, TimeDelimiter, false);
+		if (TimeRange.Num() && !TimeRange[0].IsEmpty())
 		{
-			FromTo.End = Offset;
+			if (RFC2326::ParseNPTTime(Offset, TimeRange[0]))
+			{
+				if (!FromTo.Start.IsValid())
+				{
+					FromTo.Start = Offset;
+				}
+				else
+				{
+					FromTo.Start += Offset;
+				}
+			}
+		}
+		if (TimeRange.Num() > 1 && !TimeRange[1].IsEmpty())
+		{
+			if (RFC2326::ParseNPTTime(Offset, TimeRange[1]))
+			{
+				FromTo.End = Offset;
+			}
 		}
+		FromTo.Start += BaseTimeOffset;
+		FromTo.End += BaseTimeOffset;
 	}
-	FromTo.Start += BaseTimeOffset;
-	FromTo.End += BaseTimeOffset;
 	return FromTo;
 }
 
@@ -2037,11 +2072,124 @@ void FActiveHLSPlaylist::FTimelineMediaAsset::ResetInternalTimeline()
 	InternalMediaTimeline.ResyncNeeded();
 }
 
+void FActiveHLSPlaylist::FTimelineMediaAsset::TransformIntoReplayEvent()
+{
+	// Can only transform VOD into a replay event.
+	if (!(PlaylistType == FPlaylistParserHLS::EPlaylistType::VOD || bHasEndList))
+	{
+		return;
+	}
+
+	FString Arg;
+	bool bStaticStart = false;
+	bool bDynamicStart = false;
+
+	// Get the replay base start time.
+	const TArray<FURL_RFC3986::FQueryParam>& frgs(ReplayEventParams.GetUrlFragmentParams());
+	for(int32 i=0,iMax=frgs.Num(); i<iMax; ++i)
+	{
+		if (frgs[i].Name.Equals(CustomOptions::Custom_EpicUTCNow))
+		{
+			if (!UnixEpoch::ParseFloatString(ReplayEvent.BaseStartTime, frgs[i].Value))
+			{
+				UE_LOG(LogElectraPlayer, Warning, TEXT("Could not parse the custom UTC now value, disabling replay event."));
+				return;
+			}
+			break;
+		}
+		else if (frgs[i].Name.Equals(CustomOptions::Custom_EpicUTCUrl))
+		{
+			UE_LOG(LogElectraPlayer, Warning, TEXT("Fetching UTC time from a custom URL for a replay event is not supported with HLS."));
+		}
+	}
+	// Which type of replay event?
+	for(int32 i=0,iMax=frgs.Num(); i<iMax; ++i)
+	{
+		bStaticStart = frgs[i].Name.Equals(CustomOptions::Custom_EpicStaticStart);
+		bDynamicStart = frgs[i].Name.Equals(CustomOptions::Custom_EpicDynamicStart);
+		if (bStaticStart || bDynamicStart)
+		{
+			Arg = frgs[i].Value;
+			break;
+		}
+	}
+
+	// Any indication that this is a replay event?
+	if (!bStaticStart && !bDynamicStart)
+	{
+		return;
+	}
+	// Cannot use playlists that contain an #EXT-X-PROGRAM-DATE-TIME tag.
+	if (bHasProgramDateTime)
+	{
+		UE_LOG(LogElectraPlayer, Warning, TEXT("HLS stream uses program-date-time, cannot use as replay event."));
+		return;
+	}
+
+	TArray<FString> Params;
+	if (!Arg.IsEmpty())
+	{
+		const TCHAR* const Delimiter = TEXT(",");
+		Arg.ParseIntoArray(Params, Delimiter, true);
+	}
+	if (Params.Num())
+	{
+		ReplayEvent.bIsReplay = true;
+		ReplayEvent.bIsStaticEvent = bStaticStart;
+		if (ReplayEvent.BaseStartTime.IsValid())
+		{
+			PlayerSessionServices->GetSynchronizedUTCTime()->SetTime(ReplayEvent.BaseStartTime);
+		}
+		// Get the event start time. This is either a Posix time in seconds since the Epoch (1/1/1970) or the special
+		// word 'now' optionally followed by a value to be added or subtracted from now.
+		FString Time = Params[0];
+		if (Time.StartsWith(TEXT("now")))
+		{
+			Time.RightChopInline(3);
+		}
+		if (!Time.IsEmpty())
+		{
+			FTimeValue Offset = FTimeValue().SetFromTimeFraction(FTimeFraction().SetFromFloatString(Time));
+			if (Offset.IsValid())
+			{
+				if (ReplayEvent.BaseStartTime.IsValid())
+				{
+					ReplayEvent.BaseStartTime += Offset;
+				}
+				else
+				{
+					ReplayEvent.BaseStartTime = Offset;
+				}
+			}
+		}
+		ReplayEvent.OriginalBaseTimeOffset = BaseTimeOffset;
+		BaseTimeOffset = ReplayEvent.BaseStartTime;
+		if (Params.Num() > 1)
+		{
+			UnixEpoch::ParseFloatString(ReplayEvent.SuggestedPresentationDelay, Params[1]);
+		}
+	}
+}
+
+void FActiveHLSPlaylist::FTimelineMediaAsset::PrepareReplayEventForLooping(int32 InNumLoopsToAdd)
+{
+	if (ReplayEvent.bIsReplay)
+	{
+		BaseTimeOffset += Duration * InNumLoopsToAdd;
+	}
+}
+
+
 FTimeRange FActiveHLSPlaylist::FTimelineMediaAsset::GetTimeRange() const
 {
 	FTimeRange tr;
+	if (ReplayEvent.bIsReplay)
+	{
+		tr.Start = BaseTimeOffset;
+		tr.End = tr.Start + Duration;
+	}
 	// Was this a static asset from the get-go?
-	if (InitialPlaylistType == FPlaylistParserHLS::EPlaylistType::VOD || bInitialHasEndList)
+	else if (InitialPlaylistType == FPlaylistParserHLS::EPlaylistType::VOD || bInitialHasEndList)
 	{
 		// If the #EXT-X-PROGRAM-DATE-TIME tags are in use we need to abide by them
 		// even for VOD content because the fmp4 segments are bound to have an EPT
@@ -2235,12 +2383,16 @@ void FActiveHLSPlaylist::FTimelineMediaAsset::UpdateRunningMetaData(const FStrin
 
 IManifest::EType FActiveHLSPlaylist::GetPresentationType() const
 {
+	if (TimelineMediaAsset->ReplayEvent.bIsReplay)
+	{
+		return TimelineMediaAsset->ReplayEvent.bIsStaticEvent ? IManifest::EType::OnDemand : IManifest::EType::Live;
+	}
 	return (TimelineMediaAsset->PlaylistType == FPlaylistParserHLS::EPlaylistType::VOD || TimelineMediaAsset->bHasEndList) ? IManifest::EType::OnDemand : IManifest::EType::Live;
 }
 
 IManifest::EReplayEventType FActiveHLSPlaylist::GetReplayType() const
 {
-	return IManifest::EReplayEventType::NoReplay;
+	return TimelineMediaAsset->ReplayEvent.bIsReplay ? IManifest::EReplayEventType::IsReplay : IManifest::EReplayEventType::NoReplay;
 }
 
 TSharedPtrTS<const FLowLatencyDescriptor> FActiveHLSPlaylist::GetLowLatencyDescriptor() const
@@ -2385,6 +2537,10 @@ void FActiveHLSPlaylist::UpdateDynamicRefetchCounter()
 
 void FActiveHLSPlaylist::PrepareForLooping(int32 InNumLoopsToAdd)
 {
+	if (TimelineMediaAsset.IsValid())
+	{
+		TimelineMediaAsset->PrepareReplayEventForLooping(InNumLoopsToAdd);
+	}
 }
 
 void FActiveHLSPlaylist::TriggerClockSync(EClockSyncType InClockSyncType)
diff --git a/Engine/Plugins/Media/ElectraPlayer/Source/ElectraPlayerRuntime/Private/Runtime/Player/HLS/PlaylistHLS.h b/Engine/Plugins/Media/ElectraPlayer/Source/ElectraPlayerRuntime/Private/Runtime/Player/HLS/PlaylistHLS.h
--- a/Engine/Plugins/Media/ElectraPlayer/Source/ElectraPlayerRuntime/Private/Runtime/Player/HLS/PlaylistHLS.h
+++ b/Engine/Plugins/Media/ElectraPlayer/Source/ElectraPlayerRuntime/Private/Runtime/Player/HLS/PlaylistHLS.h
@@ -379,6 +379,8 @@ struct FMediaPlaylistAndStateHLS
 	// Media playlist URL as given in the multivariant playlist. This is not necessarily the effective URL.
 	FString URL;
 	TArray<FURL_RFC3986::FQueryParam> MultiVariantURLFragmentComponents;
+	Playlist::FReplayEventParams ReplayEventParams;
+
 	// Current state of the playlist.
 	EPlaylistState PlaylistState = EPlaylistState::NotLoaded;
 	// Time at which the response was received, from the Date HTTP response header
@@ -596,6 +598,8 @@ private:
 		FTimeValue CalculateCurrentLiveLatency(const FTimeValue& InCurrentPlaybackPosition, const FTimeValue& InEncoderLatency);
 		void UpdateTimelineFromMediaSegment(TSharedPtrTS<FStreamSegmentRequestCommon> InSegment);
 		void ResetInternalTimeline();
+		void TransformIntoReplayEvent();
+		void PrepareReplayEventForLooping(int32 InNumLoopsToAdd);
 
 		TSharedPtrTS<FMultiVariantPlaylistHLS::FPathwayStreamInfs> GetCurrentPathway() const
 		{
@@ -659,6 +663,7 @@ private:
 		bool bInitialHasEndList = false;
 		bool bHasProgramDateTime = false;
 		TArray<FURL_RFC3986::FQueryParam> MultiVariantURLFragmentComponents;
+		Playlist::FReplayEventParams ReplayEventParams;
 		FTimeRange DefaultStartAndEndTime;
 
 		// A one-time only established base time to shift the playlist timeline around.
@@ -705,6 +710,16 @@ private:
 		// The currently active pathway
 		TSharedPtrTS<FMultiVariantPlaylistHLS::FPathwayStreamInfs> CurrentPathway;
 		FString CurrentPathwayId;
+
+		struct FReplayEvent
+		{
+			FTimeValue BaseStartTime;
+			FTimeValue OriginalBaseTimeOffset;
+			FTimeValue SuggestedPresentationDelay;
+			bool bIsReplay = false;
+			bool bIsStaticEvent = false;
+		};
+		FReplayEvent ReplayEvent;
 	};
 
 	FCriticalSection RequestedPeriodsLock;
diff --git a/Engine/Plugins/Media/ElectraPlayer/Source/ElectraPlayerRuntime/Private/Runtime/Player/HLS/PlaylistHandlerHLS.cpp b/Engine/Plugins/Media/ElectraPlayer/Source/ElectraPlayerRuntime/Private/Runtime/Player/HLS/PlaylistHandlerHLS.cpp
--- a/Engine/Plugins/Media/ElectraPlayer/Source/ElectraPlayerRuntime/Private/Runtime/Player/HLS/PlaylistHandlerHLS.cpp
+++ b/Engine/Plugins/Media/ElectraPlayer/Source/ElectraPlayerRuntime/Private/Runtime/Player/HLS/PlaylistHandlerHLS.cpp
@@ -32,7 +32,7 @@ public:
 		static const FString PlaylistType(TEXT("hls"));
 		return PlaylistType;
 	}
-	void LoadAndParse(const FString& URL) override;
+	void LoadAndParse(const FString& InURL, const Playlist::FReplayEventParams& InReplayEventParams) override;
 	FString GetURL() const override;
 	TSharedPtrTS<IManifest> GetManifest() override;
 
@@ -110,6 +110,7 @@ private:
 	IPlayerSessionServices* PlayerSessionServices = nullptr;
 	FString MultiVariantPlaylistEffectiveURL;
 	TArray<FURL_RFC3986::FQueryParam> MultiVariantURLFragmentComponents;
+	Playlist::FReplayEventParams ReplayEventParams;
 
 	TArray<TSharedPtr<FLoadRequestHLSPlaylist, ESPMode::ThreadSafe>> PendingLoadRequests;
 	TArray<TSharedPtr<FLoadRequestHLSPlaylist, ESPMode::ThreadSafe>> RunningLoadRequests;
@@ -312,6 +313,7 @@ void FPlaylistHandlerHLS::HandleOnce()
 				MediaPlaylist = MakeShared<FMediaPlaylistAndStateHLS, ESPMode::ThreadSafe>();
 				MediaPlaylist->URL = FinishedReq->ResourceRequest->GetURL();
 				MediaPlaylist->MultiVariantURLFragmentComponents = MultiVariantURLFragmentComponents;
+				MediaPlaylist->ReplayEventParams = ReplayEventParams;
 				MediaPlaylist->PlaylistState = FMediaPlaylistAndStateHLS::EPlaylistState::Invalid;
 				MediaPlaylist->ResponseDateHeaderTime = FinishedReq->ResponseDateHeaderTime;
 				MediaPlaylist->bIsPrimaryPlaylist = FinishedReq->bIsPrimaryPlaylist;
@@ -739,14 +741,13 @@ void FPlaylistHandlerHLS::PerformContentSteeringCloning(bool bIsPreStartFetch)
 }
 
 
-void FPlaylistHandlerHLS::LoadAndParse(const FString& URL)
+void FPlaylistHandlerHLS::LoadAndParse(const FString& InURL, const Playlist::FReplayEventParams& InReplayEventParams)
 {
 	FURL_RFC3986 UrlParser;
-	UrlParser.Parse(URL);
+	UrlParser.Parse(InURL);
 	MultiVariantPlaylistEffectiveURL = UrlParser.Get(true, false);
-	FString URLFragment = UrlParser.GetFragment();
-	// The fragment is already URL escaped, so no need to do it again.
-	FURL_RFC3986::GetQueryParams(MultiVariantURLFragmentComponents, URLFragment, false);
+	ReplayEventParams = InReplayEventParams;
+	FURL_RFC3986::GetQueryParams(MultiVariantURLFragmentComponents, UrlParser.GetFragment(), false);
 	LoadMainPlaylist(MultiVariantPlaylistEffectiveURL);
 }
 
diff --git a/Engine/Plugins/Media/ElectraPlayer/Source/ElectraPlayerRuntime/Private/Runtime/Player/Playlist.h b/Engine/Plugins/Media/ElectraPlayer/Source/ElectraPlayerRuntime/Private/Runtime/Player/Playlist.h
--- a/Engine/Plugins/Media/ElectraPlayer/Source/ElectraPlayerRuntime/Private/Runtime/Player/Playlist.h
+++ b/Engine/Plugins/Media/ElectraPlayer/Source/ElectraPlayerRuntime/Private/Runtime/Player/Playlist.h
@@ -3,11 +3,86 @@
 #pragma once
 
 #include "PlayerCore.h"
+#include "Containers/Array.h"
+#include "Utilities/URLParser.h"
+#include "Player/AdaptivePlayerOptionKeynames.h"
 
 namespace Electra
 {
 	namespace Playlist
 	{
+		class FReplayEventParams
+		{
+		public:
+			void SetLoopingAtStart(bool bInIsLoopEnabledAtStart)
+			{ bIsLoopEnabledAtStart = bInIsLoopEnabledAtStart; }
+			bool IsLoopingAtStartEnabled() const
+			{ return bIsLoopEnabledAtStart; }
+
+			bool HaveReplayParams() const
+			{ return !UrlFragmentParams.IsEmpty(); }
+			
+			// Parse URL fragment parameters
+			bool ParseFromURL(const FString& InURL)
+			{
+				static const TArray<FString> FragmentParamsToParse { CustomOptions::Custom_EpicStaticStart, CustomOptions::Custom_EpicDynamicStart, CustomOptions::Custom_EpicUTCUrl, CustomOptions::Custom_EpicUTCNow };
+				SourceURL = CleanedURL = InURL;
+				FURL_RFC3986 UrlParser;
+				if (UrlParser.Parse(SourceURL))
+				{
+					FString NewFragment;
+					FString URLFragment = UrlParser.GetFragment();
+					TArray<FURL_RFC3986::FQueryParam> UrlFragments;
+					FURL_RFC3986::GetQueryParams(UrlFragments, URLFragment, false);
+					if (UrlFragments.Num())
+					{
+						for(int32 i=0; i<UrlFragments.Num(); ++i)
+						{
+							if (FragmentParamsToParse.Contains(UrlFragments[i].Name))
+							{
+								UrlFragmentParams.Emplace(UrlFragments[i]);
+								UrlFragments.RemoveAt(i);
+								--i;
+							}
+							else
+							{
+								if (!NewFragment.IsEmpty())
+								{
+									NewFragment.Append(TEXT("&"));
+								}
+								NewFragment.Append(UrlFragments[i].Name);
+								if (!UrlFragments[i].Value.IsEmpty())
+								{
+									NewFragment.Append(TEXT("="));
+									NewFragment.Append(UrlFragments[i].Value);
+								}
+							}
+						}
+						CleanedURL = UrlParser.Get(true, false);
+						if (NewFragment.Len())
+						{
+							CleanedURL.Append(TEXT("#"));
+							CleanedURL.Append(NewFragment);
+						}
+						return true;
+					}
+				}
+				return false;
+			}
+			FString GetSourceURL() const
+			{ return SourceURL; }
+			FString GetCleanedURL() const
+			{ return CleanedURL; }
+			const TArray<FURL_RFC3986::FQueryParam>& GetUrlFragmentParams() const
+			{ return UrlFragmentParams; }
+		private:
+			FString SourceURL;
+			FString CleanedURL;
+			TArray<FURL_RFC3986::FQueryParam> UrlFragmentParams;
+			bool bIsLoopEnabledAtStart = false;
+		};
+
+
 		enum class EListType
 		{
 			Main,
diff --git a/Engine/Plugins/Media/ElectraPlayer/Source/ElectraPlayerRuntime/Private/Runtime/Player/PlaylistReader.h b/Engine/Plugins/Media/ElectraPlayer/Source/ElectraPlayerRuntime/Private/Runtime/Player/PlaylistReader.h
--- a/Engine/Plugins/Media/ElectraPlayer/Source/ElectraPlayerRuntime/Private/Runtime/Player/PlaylistReader.h
+++ b/Engine/Plugins/Media/ElectraPlayer/Source/ElectraPlayerRuntime/Private/Runtime/Player/PlaylistReader.h
@@ -40,9 +40,10 @@ namespace Electra
 		/**
 		 * Loads and parses the playlist.
 		 *
-		 * @param URL     URL of the playlist to load
+		 * @param URL URL of the playlist to load
+		 * @param InReplayEventParams URL fragments pertaining to replay event which have been removed from the source URL
 		 */
-		virtual void LoadAndParse(const FString& URL) = 0;
+		virtual void LoadAndParse(const FString& InURL, const Playlist::FReplayEventParams& InReplayEventParams) = 0;
 
 		/**
 		 * Returns the URL from which the playlist was loaded (or supposed to be loaded).
diff --git a/Engine/Plugins/Media/ElectraPlayer/Source/ElectraPlayerRuntime/Private/Runtime/Player/mkv/PlaylistReaderMKV.cpp b/Engine/Plugins/Media/ElectraPlayer/Source/ElectraPlayerRuntime/Private/Runtime/Player/mkv/PlaylistReaderMKV.cpp
--- a/Engine/Plugins/Media/ElectraPlayer/Source/ElectraPlayerRuntime/Private/Runtime/Player/mkv/PlaylistReaderMKV.cpp
+++ b/Engine/Plugins/Media/ElectraPlayer/Source/ElectraPlayerRuntime/Private/Runtime/Player/mkv/PlaylistReaderMKV.cpp
@@ -56,9 +56,10 @@ public:
 	/**
 	 * Loads and parses the playlist.
 	 *
-	 * @param URL     URL of the playlist to load
+	 * @param URL URL of the playlist to load
+	 * @param InReplayEventParams URL fragments pertaining to replay event which have been removed from the source URL
 	 */
-	virtual void LoadAndParse(const FString& URL) override;
+	virtual void LoadAndParse(const FString& InURL, const Playlist::FReplayEventParams& InReplayEventParams) override;
 
 	/**
 	 * Returns the URL from which the playlist was loaded (or supposed to be loaded).
@@ -243,10 +244,10 @@ void FPlaylistReaderMKV::LogMessage(IInfoLog::ELevel Level, const FString& Messa
 	PlayerSessionServices->PostLog(Facility::EFacility::MKVPlaylistReader, Level, Message);
 }
 
-void FPlaylistReaderMKV::LoadAndParse(const FString& URL)
+void FPlaylistReaderMKV::LoadAndParse(const FString& InURL, const Playlist::FReplayEventParams& InReplayEventParams)
 {
 	FURL_RFC3986 UrlParser;
-	UrlParser.Parse(URL);
+	UrlParser.Parse(InURL);
 	PlaylistURL = UrlParser.Get(true, false);
 	URLFragment = UrlParser.GetFragment();
 	StartWorkerThread();
diff --git a/Engine/Plugins/Media/ElectraPlayer/Source/ElectraPlayerRuntime/Private/Runtime/Player/mp4/PlaylistReaderMP4.cpp b/Engine/Plugins/Media/ElectraPlayer/Source/ElectraPlayerRuntime/Private/Runtime/Player/mp4/PlaylistReaderMP4.cpp
--- a/Engine/Plugins/Media/ElectraPlayer/Source/ElectraPlayerRuntime/Private/Runtime/Player/mp4/PlaylistReaderMP4.cpp
+++ b/Engine/Plugins/Media/ElectraPlayer/Source/ElectraPlayerRuntime/Private/Runtime/Player/mp4/PlaylistReaderMP4.cpp
@@ -57,9 +57,10 @@ public:
 	/**
 	 * Loads and parses the playlist.
 	 *
-	 * @param URL     URL of the playlist to load
+	 * @param URL URL of the playlist to load
+	 * @param InReplayEventParams URL fragments pertaining to replay event which have been removed from the source URL
 	 */
-	virtual void LoadAndParse(const FString& URL) override;
+	virtual void LoadAndParse(const FString& InURL, const Playlist::FReplayEventParams& InReplayEventParams) override;
 
 	/**
 	 * Returns the URL from which the playlist was loaded (or supposed to be loaded).
@@ -209,10 +210,10 @@ void FPlaylistReaderMP4::LogMessage(IInfoLog::ELevel Level, const FString& Messa
 	}
 }
 
-void FPlaylistReaderMP4::LoadAndParse(const FString& URL)
+void FPlaylistReaderMP4::LoadAndParse(const FString& InURL, const Playlist::FReplayEventParams& InReplayEventParams)
 {
 	FURL_RFC3986 UrlParser;
-	UrlParser.Parse(URL);
+	UrlParser.Parse(InURL);
 	MainPlaylistURL = UrlParser.Get(true, false);
 	URLFragment = UrlParser.GetFragment();
 
diff --git a/Engine/Plugins/Media/ElectraPlayer/Source/ElectraPlayerRuntime/Private/Runtime/Player/mpegaudio/PlaylistReaderMPEGAudio.cpp b/Engine/Plugins/Media/ElectraPlayer/Source/ElectraPlayerRuntime/Private/Runtime/Player/mpegaudio/PlaylistReaderMPEGAudio.cpp
--- a/Engine/Plugins/Media/ElectraPlayer/Source/ElectraPlayerRuntime/Private/Runtime/Player/mpegaudio/PlaylistReaderMPEGAudio.cpp
+++ b/Engine/Plugins/Media/ElectraPlayer/Source/ElectraPlayerRuntime/Private/Runtime/Player/mpegaudio/PlaylistReaderMPEGAudio.cpp
@@ -58,8 +58,9 @@ public:
 	 * Loads and parses the playlist.
 	 *
 	 * @param URL URL of the playlist to load
+	 * @param InReplayEventParams URL fragments pertaining to replay event which have been removed from the source URL
 	 */
-	void LoadAndParse(const FString& URL) override;
+	void LoadAndParse(const FString& InURL, const Playlist::FReplayEventParams& InReplayEventParams) override;
 
 	/**
 	 * Returns the URL from which the playlist was loaded (or supposed to be loaded).
@@ -252,9 +253,9 @@ void FPlaylistReaderMPEGAudio::LogMessage(IInfoLog::ELevel Level, const FString&
 	}
 }
 
-void FPlaylistReaderMPEGAudio::LoadAndParse(const FString& URL)
+void FPlaylistReaderMPEGAudio::LoadAndParse(const FString& InURL, const Playlist::FReplayEventParams& InReplayEventParams)
 {
-	PlaylistURL = URL;
+	PlaylistURL = InURL;
 	StartWorkerThread();
 }
 
