Electra: Attempting to fix PVS warnings

#rnx

diff --git a/Engine/Plugins/Media/ElectraCodecs/Source/ElectraDecoders/Public/Utils/MPEG/ElectraUtilsMPEGVideo_H264.h b/Engine/Plugins/Media/ElectraCodecs/Source/ElectraDecoders/Public/Utils/MPEG/ElectraUtilsMPEGVideo_H264.h
--- a/Engine/Plugins/Media/ElectraCodecs/Source/ElectraDecoders/Public/Utils/MPEG/ElectraUtilsMPEGVideo_H264.h
+++ b/Engine/Plugins/Media/ElectraCodecs/Source/ElectraDecoders/Public/Utils/MPEG/ElectraUtilsMPEGVideo_H264.h
@@ -61,11 +61,11 @@ namespace ElectraDecodersUtil
 				uint32 ue_v()
 				{
 					int32 lz = -1;
-					for(uint32 b = 0; b == 0; ++lz)
+					for(uint32 b=0; b==0 && GetRemainingBits(); ++lz)
 					{
 						b = GetBits(1);
 					}
-					if (lz)
+					if (lz >= 0 && GetRemainingBits() >= lz)
 					{
 						return ((1 << lz) | GetBits(lz)) - 1;
 					}
diff --git a/Engine/Plugins/Media/ElectraCodecs/Source/ElectraDecoders/Public/Utils/MPEG/ElectraUtilsMPEGVideo_H265.h b/Engine/Plugins/Media/ElectraCodecs/Source/ElectraDecoders/Public/Utils/MPEG/ElectraUtilsMPEGVideo_H265.h
--- a/Engine/Plugins/Media/ElectraCodecs/Source/ElectraDecoders/Public/Utils/MPEG/ElectraUtilsMPEGVideo_H265.h
+++ b/Engine/Plugins/Media/ElectraCodecs/Source/ElectraDecoders/Public/Utils/MPEG/ElectraUtilsMPEGVideo_H265.h
@@ -55,11 +55,11 @@ namespace ElectraDecodersUtil
 				uint32 ue_v()
 				{
 					int32 lz = -1;
-					for(uint32 b = 0; b == 0; ++lz)
+					for(uint32 b=0; b==0 && GetRemainingBits(); ++lz)
 					{
 						b = GetBits(1);
 					}
-					if (lz)
+					if (lz >= 0 && GetRemainingBits() >= lz)
 					{
 						return ((1 << lz) | GetBits(lz)) - 1;
 					}
diff --git a/Engine/Plugins/Media/ElectraPlayer/Source/ElectraPlayerRuntime/Private/Runtime/Player/DASH/ManifestDASH.cpp b/Engine/Plugins/Media/ElectraPlayer/Source/ElectraPlayerRuntime/Private/Runtime/Player/DASH/ManifestDASH.cpp
--- a/Engine/Plugins/Media/ElectraPlayer/Source/ElectraPlayerRuntime/Private/Runtime/Player/DASH/ManifestDASH.cpp
+++ b/Engine/Plugins/Media/ElectraPlayer/Source/ElectraPlayerRuntime/Private/Runtime/Player/DASH/ManifestDASH.cpp
@@ -3950,7 +3950,7 @@ void FManifestDASHInternal::FRepresentation::GetAverageSegmentDuration(FTimeValu
 								DurationSum += Selements[nI].D * RCount;
 								NumElements += RCount;
 							}
-							OutAverageSegmentDuration.SetFromND(DurationSum / NumElements, MPDTimescale);
+							OutAverageSegmentDuration.SetFromND(NumElements ? DurationSum / NumElements : 0, MPDTimescale);
 							return;
 						}
 					}
