ElectraPlayer: Moving media samples into internal player to remove delivery delegates and lambdas and to enable tracking of dropped frames

#rnx

diff --git a/Engine/Plugins/Media/ElectraCodecs/Source/ElectraDecoders/Private/Utils/MPEG/ElectraUtilsMPEGVideo_H264.cpp b/Engine/Plugins/Media/ElectraCodecs/Source/ElectraDecoders/Private/Utils/MPEG/ElectraUtilsMPEGVideo_H264.cpp
--- a/Engine/Plugins/Media/ElectraCodecs/Source/ElectraDecoders/Private/Utils/MPEG/ElectraUtilsMPEGVideo_H264.cpp
+++ b/Engine/Plugins/Media/ElectraCodecs/Source/ElectraDecoders/Private/Utils/MPEG/ElectraUtilsMPEGVideo_H264.cpp
@@ -379,9 +379,19 @@ namespace ElectraDecodersUtil
 			}
 			FFractionalValue FSequenceParameterSet::GetTiming() const
 			{
-				if (vui_parameters_present_flag && timing_info_present_flag && fixed_frame_rate_flag)
+				if (vui_parameters_present_flag && timing_info_present_flag)
 				{
-					return FFractionalValue(time_scale, num_units_in_tick * 2 /* rate per field, we assume progressive frames, so 2 fields */);
+					// Having just the SPS is not necessarily all that is needed to get the framerate as this
+					// may require the pic_struct (from buffering SEI message) and field_pic_flag (from the slice).
+					// Since we do not handle field encoded video we assume field_pic_flag to be 0.
+					// If pic_struct_present_flag is zero then we use the timing info as it is given with a DeltaTfiDivisor of 2.
+					// Otherwise, if pic_struct_present_flag is one then DeltaTfiDivisor is 2 only for pic_struct values of 0, 3 and 4.
+					// Anyway, that is what we assume to handle, so there...
+					const int32 DeltaTfiDivisor = 2;
+					if (/*pic_struct_present_flag == 0 &&*/ time_scale && num_units_in_tick)
+					{
+						return FFractionalValue(time_scale, num_units_in_tick * DeltaTfiDivisor);
+					}
 				}
 				return FFractionalValue();
 			}
diff --git a/Engine/Plugins/Media/ElectraPlayer/Source/ElectraPlayerRuntime/Private/ElectraPlayer.cpp b/Engine/Plugins/Media/ElectraPlayer/Source/ElectraPlayerRuntime/Private/ElectraPlayer.cpp
--- a/Engine/Plugins/Media/ElectraPlayer/Source/ElectraPlayerRuntime/Private/ElectraPlayer.cpp
+++ b/Engine/Plugins/Media/ElectraPlayer/Source/ElectraPlayerRuntime/Private/ElectraPlayer.cpp
@@ -16,9 +16,7 @@
 #include "Utilities/Utilities.h"
 #include "Utilities/URLParser.h"
 #include "Utilities/StringHelpers.h"
-#include "MediaSubtitleDecoderOutput.h"
 #include "MediaMetaDataDecoderOutput.h"
-#include "IElectraSubtitleSample.h"
 #include "IElectraMetadataSample.h"
 #include "IMediaMetadataItem.h"
 
@@ -234,24 +232,6 @@ public:
 	TOptional<FDecoderTimeStamp> TrackBaseTime;
 };
 
-class FElectraSubtitleSample : public IElectraSubtitleSample
-{
-public:
-	FGuid GetGUID() const override						{ return IElectraSubtitleSample::GetSampleTypeGUID(); }
-	FMediaTimeStamp GetTime() const override			{ FDecoderTimeStamp ts = Subtitle->GetTime(); return FMediaTimeStamp(ts.Time, ts.SequenceIndex); }
-	FTimespan GetDuration() const override				{ return Subtitle->GetDuration(); }
-	TOptional<FVector2D> GetPosition() const override	{ return TOptional<FVector2D>(); }
-	EMediaOverlaySampleType GetType() const override	{ return EMediaOverlaySampleType::Subtitle; }
-	FText GetText() const override
-	{
-		FUTF8ToTCHAR cnv((const ANSICHAR*)Subtitle->GetData().GetData(), Subtitle->GetData().Num());
-		FString UTF8Text = FString::ConstructFromPtrSize(cnv.Get(), cnv.Length());
-		return FText::FromString(UTF8Text);
-	}
-	ISubtitleDecoderOutputPtr Subtitle;
-};
-
-
 class FElectraBinarySample : public IElectraBinarySample
 {
 public:
@@ -441,7 +421,6 @@ void FElectraPlayer::ClearToDefaultState()
 	bSubtitleTrackIndexDirty = true;
 	bInitialSeekPerformed = false;
 	bIsFirstBuffering = true;
-	LastPresentedFrameDimension = FIntPoint::ZeroValue;
 	CurrentStreamMetadata.Reset();
 	CurrentlyActiveVideoStreamFormat.Reset();
 	DeferredPlayerEvents.Empty();
@@ -451,7 +430,6 @@ void FElectraPlayer::ClearToDefaultState()
 
 void FElectraPlayer::SendMediaSinkEvent(EMediaEvent InEventToSend)
 {
-	FScopeLock lock(&ParentObjectLock);
 	if (EventSink)
 	{
 		EventSink->ReceiveMediaEvent(InEventToSend);
@@ -471,9 +449,6 @@ bool FElectraPlayer::OpenInternal(const FString& InUrl, const IMediaOptions* InO
 	FGuid SessionID = FGuid::NewGuid();
 
 	PlayerUniqueID = ++NextPlayerUniqueID;
-	MediaSamplesLock.Lock();
-	MediaSamples.Reset(new FMediaSamples);
-	MediaSamplesLock.Unlock();
 
 	// Check that we have a valid option interface.
 	if (InOptions == nullptr)
@@ -483,9 +458,7 @@ bool FElectraPlayer::OpenInternal(const FString& InUrl, const IMediaOptions* InO
 		return false;
 	}
 	// Get the safe option interface to poll for changes during playback.
-	ParentObjectLock.Lock();
 	OptionInterface = StaticCastSharedPtr<IElectraSafeMediaOptionInterface>(InOptions->GetMediaOption(ElectraMediaOptions::GetSafeMediaOptions, TSharedPtr<IElectraSafeMediaOptionInterface, ESPMode::ThreadSafe>()));
-	ParentObjectLock.Unlock();
 	UE_LOG(LogElectraPlayer, Log, TEXT("[%u] IMediaPlayer::Open"), PlayerUniqueID);
 
 	// Create the static resource provider using the safe media option interface.
@@ -677,62 +650,10 @@ bool FElectraPlayer::OpenInternal(const FString& InUrl, const IMediaOptions* InO
 	// Create a new empty player structure. This contains the actual player instance, its associated renderers and sample queues.
 	TSharedPtr<FInternalPlayerImpl, ESPMode::ThreadSafe> NewPlayer = MakeShared<FInternalPlayerImpl, ESPMode::ThreadSafe>();
 
-	// Create the output handlers for audio and video.
-	NewPlayer->AudioOutputHandler = MakeShared<FOutputHandlerAudio, ESPMode::ThreadSafe>();
-	NewPlayer->AudioOutputHandler->SetOutputAudioSamplePool(OutputAudioSamplePool);
-	NewPlayer->AudioOutputHandler->CanOutputQueueReceiveDelegate().BindLambda([Player=NewPlayer.ToWeakPtr(), This=this](bool& bOutCanReceive, int32 InNumSamples)
-	{
-		bOutCanReceive = false;
-		if (auto Pl = Player.Pin())
-		{
-			bOutCanReceive = This->CanPresentAudioFrames(InNumSamples);
-		}
-	});
-	NewPlayer->AudioOutputHandler->OutputQueueReceiveSampleDelegate().BindLambda([Player=NewPlayer.ToWeakPtr(), This=this](FElectraAudioSamplePtr InSample)
-	{
-		if (auto Pl = Player.Pin())
-		{
-			This->OnAudioDecoded(InSample);
-		}
-	});
-	NewPlayer->AudioOutputHandler->OutputQueueFlushSamplesDelegate().BindLambda([Player=NewPlayer.ToWeakPtr(), This=this]()
-	{
-		if (auto Pl = Player.Pin())
-		{
-			This->OnAudioFlush();
-		}
-	});
-
-	NewPlayer->VideoOutputHandler = MakeShared<FOutputHandlerVideo, ESPMode::ThreadSafe>();
-	NewPlayer->VideoOutputHandler->SetOutputTexturePool(OutputTexturePool);
-	NewPlayer->VideoOutputHandler->CanOutputQueueReceiveDelegate().BindLambda([Player=NewPlayer.ToWeakPtr(), This=this](bool& bOutCanReceive, int32 InNumSamples)
-	{
-		bOutCanReceive = false;
-		if (auto Pl = Player.Pin())
-		{
-			bOutCanReceive = This->CanPresentVideoFrames(InNumSamples);
-		}
-	});
-	NewPlayer->VideoOutputHandler->OutputQueueReceiveSampleDelegate().BindLambda([Player=NewPlayer.ToWeakPtr(), This=this](FElectraTextureSamplePtr InSample)
-	{
-		if (auto Pl = Player.Pin())
-		{
-			This->OnVideoDecoded(InSample);
-		}
-	});
-	NewPlayer->VideoOutputHandler->OutputQueueFlushSamplesDelegate().BindLambda([Player=NewPlayer.ToWeakPtr(), This=this]()
-	{
-		if (auto Pl = Player.Pin())
-		{
-			This->OnVideoFlush();
-		}
-	});
-
-
 	// Create the internal player and register ourselves as metrics receiver and static resource provider.
 	IAdaptiveStreamingPlayer::FCreateParam CreateParams;
-	CreateParams.VideoOutputHandler = NewPlayer->VideoOutputHandler;
-	CreateParams.AudioOutputHandler = NewPlayer->AudioOutputHandler;
+	CreateParams.OutputTexturePool = OutputTexturePool;
+	CreateParams.OutputAudioSamplePool = OutputAudioSamplePool;
 	CreateParams.ExternalPlayerGUID = PlayerGuid;
 	FString WorkerThreadOption = PlayerOptions.GetValue(Electra::OptionKeyWorkerThreads).SafeGetFString(TEXT("shared"));
 	CreateParams.WorkerThreads = WorkerThreadOption.Equals(TEXT("worker"), ESearchCase::IgnoreCase) ? IAdaptiveStreamingPlayer::FCreateParam::EWorkerThreads::DedicatedWorker :
@@ -742,12 +663,6 @@ bool FElectraPlayer::OpenInternal(const FString& InUrl, const IMediaOptions* InO
 	NewPlayer->AdaptivePlayer->AddMetricsReceiver(this);
 	NewPlayer->AdaptivePlayer->SetStaticResourceProviderCallback(StaticResourceProvider);
 
-	// Create the subtitle receiver and register it with the player.
-	MediaPlayerSubtitleReceiver = MakeSharedTS<FSubtitleEventReceiver>();
-	MediaPlayerSubtitleReceiver->GetSubtitleReceivedDelegate().BindRaw(this, &FElectraPlayer::OnSubtitleDecoded);
-	MediaPlayerSubtitleReceiver->GetSubtitleFlushDelegate().BindRaw(this, &FElectraPlayer::OnSubtitleFlush);
-	NewPlayer->AdaptivePlayer->AddSubtitleReceiver(MediaPlayerSubtitleReceiver);
-
 	// Create a new media player event receiver and register it to receive all non player internal events as soon as they are received.
 	MediaPlayerEventReceiver = MakeSharedTS<FAEMSEventReceiver>();
 	MediaPlayerEventReceiver->GetEventReceivedDelegate().BindRaw(this, &FElectraPlayer::OnMediaPlayerEventReceived);
@@ -765,14 +680,15 @@ bool FElectraPlayer::OpenInternal(const FString& InUrl, const IMediaOptions* InO
 		NewPlayer->AdaptivePlayer->SetBitrateCeiling(PlaystartOptions.MaxBandwidthForStreaming.GetValue());
 	}
 
-    // Set the player member variable to the new player so we can use our internal configuration methods on the new player.
-    CurrentPlayer = MoveTemp(NewPlayer);
 	// Apply options that may have been set prior to calling Open().
 	// Set these only if they have defined values as to not override what might have been set in the PlayerOptions.
 	if (bFrameAccurateSeeking.IsSet())
 	{
 		NewPlayer->AdaptivePlayer->EnableFrameAccurateSeeking(bFrameAccurateSeeking.GetValue());
 	}
+
+    // Set the player member variable to the new player so we can use our internal configuration methods on the new player.
+    CurrentPlayer = MoveTemp(NewPlayer);
 	if (bEnableLooping.IsSet())
 	{
 		SetLooping(bEnableLooping.GetValue());
@@ -797,9 +713,7 @@ void FElectraPlayer::CloseInternal()
 {
 	LLM_SCOPE(ELLMTag::ElectraPlayer);
 
-	ParentObjectLock.Lock();
 	OptionInterface.Reset();
-	ParentObjectLock.Unlock();
 
 	TSharedPtr<FInternalPlayerImpl, ESPMode::ThreadSafe> Player = CurrentPlayer;
 	if (!Player.IsValid())
@@ -820,15 +734,6 @@ void FElectraPlayer::CloseInternal()
 	bFrameAccurateSeeking.Reset();
 	bEnableLooping.Reset();
 
-	// Detach the output handlers so we do not receive any new output.
-	if (Player->AudioOutputHandler.IsValid())
-	{
-		Player->AudioOutputHandler->DetachPlayer();
-	}
-	if (Player->VideoOutputHandler.IsValid())
-	{
-		Player->VideoOutputHandler->DetachPlayer();
-	}
 	// Detach ourselves from receiving any player events.
 	if (Player->AdaptivePlayer.IsValid())
 	{
@@ -838,13 +743,6 @@ void FElectraPlayer::CloseInternal()
 			Player->AdaptivePlayer->RemoveAEMSReceiver(MediaPlayerEventReceiver, TEXT("*"), TEXT(""), IAdaptiveStreamingPlayerAEMSReceiver::EDispatchMode::OnStart);
 			MediaPlayerEventReceiver.Reset();
 		}
-		if (MediaPlayerSubtitleReceiver.IsValid())
-		{
-			Player->AdaptivePlayer->RemoveSubtitleReceiver(MediaPlayerSubtitleReceiver);
-			MediaPlayerSubtitleReceiver->GetSubtitleReceivedDelegate().Unbind();
-			MediaPlayerSubtitleReceiver->GetSubtitleFlushDelegate().Unbind();
-			MediaPlayerSubtitleReceiver.Reset();
-		}
 		Player->AdaptivePlayer->SetStaticResourceProviderCallback(nullptr);
 		Player->AdaptivePlayer->RemoveMetricsReceiver(this);
 	}
@@ -861,11 +759,6 @@ void FElectraPlayer::CloseInternal()
 	// Clear out the player instance now.
 	CurrentPlayer.Reset();
 
-	// Swap out the media sample queue.
-	MediaSamplesLock.Lock();
-	Player->MediaSamplesToDelete = MoveTemp(MediaSamples);
-	MediaSamplesLock.Unlock();
-
 	// Kick off asynchronous closing now.
 	FInternalPlayerImpl::DoCloseAsync(MoveTemp(Player), PlayerUniqueID, AsyncResourceReleaseNotification);
 }
@@ -877,9 +770,6 @@ void FElectraPlayer::FInternalPlayerImpl::DoCloseAsync(TSharedPtr<FInternalPlaye
 		double TimeCloseBegan = FPlatformTime::Seconds();
 		InPlayer->AdaptivePlayer->Stop();
 		InPlayer->AdaptivePlayer.Reset();
-		InPlayer->AudioOutputHandler.Reset();
-		InPlayer->VideoOutputHandler.Reset();
-		InPlayer->MediaSamplesToDelete.Reset();
 		double TimeCloseEnded = FPlatformTime::Seconds();
 		UE_LOG(LogElectraPlayer, Verbose, TEXT("[%u] DoCloseAsync() finished after %.3f msec!"), InPlayerID, (TimeCloseEnded - TimeCloseBegan) * 1000.0);
 		if (InAsyncResourceReleaseNotification.IsValid())
@@ -1020,127 +910,6 @@ void FElectraPlayer::TickInternal(FTimespan DeltaTime, FTimespan Timecode)
 	}
 }
 
-
-
-bool FElectraPlayer::CanPresentVideoFrames(uint64 InNumFrames)
-{
-	SendMediaSinkEvent(EMediaEvent::Internal_PurgeVideoSamplesHint);
-	FScopeLock SampleLock(&MediaSamplesLock);
-	return MediaSamples.IsValid() && MediaSamples->CanReceiveVideoSamples(InNumFrames);
-}
-
-void FElectraPlayer::OnVideoDecoded(FElectraTextureSamplePtr InSample)
-{
-	//UE_LOG(LogElectraPlayer, VeryVerbose, TEXT("[%u] OnVideoDecoded(%#.4f, %d,%d)"), PlayerUniqueID, InSample->GetTime().GetTime().GetTotalSeconds(), InSample->GetTime().GetSequenceIndex(), InSample->GetTime().GetLoopIndex());
-
-	TSharedPtr<FInternalPlayerImpl, ESPMode::ThreadSafe> Player = CurrentPlayer;
-	if (Player.IsValid() && PlayerState.State != EMediaState::Closed)
-	{
-		LastPresentedFrameDimension = InSample->GetOutputDim();
-		FScopeLock SampleLock(&MediaSamplesLock);
-		if (MediaSamples.IsValid())
-		{
-			MediaSamples->AddVideo(InSample.ToSharedRef());
-		}
-	}
-}
-
-void FElectraPlayer::OnVideoFlush()
-{
-	//UE_LOG(LogElectraPlayer, VeryVerbose, TEXT("[%u] OnVideoFlush()"), PlayerUniqueID);
-
-	TSharedPtr<FInternalPlayerImpl, ESPMode::ThreadSafe> Player = CurrentPlayer;
-	if (Player.IsValid())
-	{
-		FScopeLock SampleLock(&MediaSamplesLock);
-		if (MediaSamples.IsValid())
-		{
-			TRange<FTimespan> AllTime(FTimespan::MinValue(), FTimespan::MaxValue());
-			TSharedPtr<IMediaTextureSample, ESPMode::ThreadSafe> FlushSample;
-			while(MediaSamples->FetchVideo(AllTime, FlushSample))
-			{ }
-		}
-	}
-}
-
-bool FElectraPlayer::CanPresentAudioFrames(uint64 InNumFrames)
-{
-	FScopeLock lock(&MediaSamplesLock);
-	return MediaSamples.IsValid() && MediaSamples->CanReceiveAudioSamples(InNumFrames);
-}
-
-void FElectraPlayer::OnAudioDecoded(FElectraAudioSamplePtr InSample)
-{
-	//UE_LOG(LogElectraPlayer, VeryVerbose, TEXT("[%u] OnAudioDecoded(%#.4f, %d,%d)"), PlayerUniqueID, InSample->GetTime().Time.GetTotalSeconds(), (int32)InSample->GetTime().SequenceIndex, (int32)(InSample->GetTime().SequenceIndex >> 32));
-
-	TSharedPtr<FInternalPlayerImpl, ESPMode::ThreadSafe> Player = CurrentPlayer;
-	if (Player.IsValid() && InSample.IsValid() && PlayerState.State != EMediaState::Closed)
-	{
-		FScopeLock SampleLock(&MediaSamplesLock);
-		if (MediaSamples.IsValid())
-		{
-			MediaSamples->AddAudio(InSample.ToSharedRef());
-		}
-	}
-}
-
-void FElectraPlayer::OnAudioFlush()
-{
-	//UE_LOG(LogElectraPlayer, VeryVerbose, TEXT("[%u] OnAudioFlush()"), PlayerUniqueID);
-
-	TSharedPtr<FInternalPlayerImpl, ESPMode::ThreadSafe> Player = CurrentPlayer;
-	if (Player.IsValid())
-	{
-		FScopeLock SampleLock(&MediaSamplesLock);
-		if (MediaSamples.IsValid())
-		{
-			TRange<FTimespan> AllTime(FTimespan::MinValue(), FTimespan::MaxValue());
-			TSharedPtr<IMediaAudioSample, ESPMode::ThreadSafe> FlushSample;
-			while(MediaSamples->FetchAudio(AllTime, FlushSample))
-			{ }
-		}
-	}
-}
-
-void FElectraPlayer::OnSubtitleDecoded(ISubtitleDecoderOutputPtr InDecoderOutput)
-{
-	//UE_LOG(LogElectraPlayer, VeryVerbose, TEXT("[%u] OnSubtitleDecoded(%#.4f, %d,%d)"), PlayerUniqueID, InDecoderOutput->GetTime().Time.GetTotalSeconds(), (int32)InDecoderOutput->GetTime().SequenceIndex, (int32)(InDecoderOutput->GetTime().SequenceIndex >> 32));
-
-	TSharedPtr<FInternalPlayerImpl, ESPMode::ThreadSafe> Player = CurrentPlayer;
-	if (Player.IsValid() && InDecoderOutput.IsValid() && PlayerState.State != EMediaState::Closed)
-	{
-		FScopeLock SampleLock(&MediaSamplesLock);
-		if (MediaSamples.IsValid())
-		{
-			TSharedRef<FElectraSubtitleSample, ESPMode::ThreadSafe> SubtitleSample = MakeShared<FElectraSubtitleSample, ESPMode::ThreadSafe>();
-			SubtitleSample->Subtitle = InDecoderOutput;
-			MediaSamples->AddSubtitle(SubtitleSample);
-		}
-	}
-}
-
-void FElectraPlayer::OnSubtitleFlush()
-{
-	//UE_LOG(LogElectraPlayer, VeryVerbose, TEXT("[%u] OnSubtitleFlush()"), PlayerUniqueID);
-
-	TSharedPtr<FInternalPlayerImpl, ESPMode::ThreadSafe> Player = CurrentPlayer;
-	if (Player.IsValid())
-	{
-		FScopeLock SampleLock(&MediaSamplesLock);
-		if (MediaSamples.IsValid())
-		{
-			TRange<FTimespan> AllTime(FTimespan::MinValue(), FTimespan::MaxValue());
-			TSharedPtr<IMediaOverlaySample, ESPMode::ThreadSafe> FlushSample;
-			while(MediaSamples->FetchSubtitle(AllTime, FlushSample))
-			{ }
-		}
-	}
-}
-
-
-
-
-
 bool FElectraPlayer::IsLive()
 {
 	auto Player = CurrentPlayer;
@@ -1330,10 +1099,6 @@ void FElectraPlayer::HandleDeferredPlayerEvents()
 				HandlePlayerEventPlaylistDownload(Ev->PlaylistDownloadStats);
 				break;
 			}
-			case FPlayerMetricEventBase::EType::CleanStart:
-			{
-				break;
-			}
 			case FPlayerMetricEventBase::EType::BufferingStart:
 			{
 				FPlayerMetricEvent_BufferingStart* Ev = static_cast<FPlayerMetricEvent_BufferingStart*>(Event.Get());
@@ -1460,12 +1225,14 @@ void FElectraPlayer::HandleDeferredPlayerEvents()
 			}
 			case FPlayerMetricEventBase::EType::DroppedVideoFrame:
 			{
-				HandlePlayerEventDroppedVideoFrame();
+				FPlayerMetricEvent_DroppedFrame* Ev = static_cast<FPlayerMetricEvent_DroppedFrame*>(Event.Get());
+				HandlePlayerEventDroppedVideoFrame(Ev->FrameStats);
 				break;
 			}
 			case FPlayerMetricEventBase::EType::DroppedAudioFrame:
 			{
-				HandlePlayerEventDroppedAudioFrame();
+				FPlayerMetricEvent_DroppedFrame* Ev = static_cast<FPlayerMetricEvent_DroppedFrame*>(Event.Get());
+				HandlePlayerEventDroppedAudioFrame(Ev->FrameStats);
 				break;
 			}
 			default:
@@ -2178,15 +1945,11 @@ void FElectraPlayer::HandlePlayerEventCodecFormatChange(const Electra::FStreamCo
 {
 	if (NewDecodingFormat.IsVideoCodec())
 	{
-		FVideoStreamFormat fmt;
-		fmt.Bitrate = NewDecodingFormat.GetBitrate();
-		fmt.Resolution.X = NewDecodingFormat.GetResolution().Width;
-		fmt.Resolution.Y = NewDecodingFormat.GetResolution().Height;
-		fmt.FrameRate = NewDecodingFormat.GetFrameRate().IsValid() ? NewDecodingFormat.GetFrameRate().GetAsDouble() : 0.0;
-		{
 		FScopeLock lock(&VideoFormatLock);
-		CurrentlyActiveVideoStreamFormat = fmt;
-		}
+		CurrentlyActiveVideoStreamFormat.Bitrate = NewDecodingFormat.GetBitrate();
+		CurrentlyActiveVideoStreamFormat.Resolution.X = NewDecodingFormat.GetResolution().Width;
+		CurrentlyActiveVideoStreamFormat.Resolution.Y = NewDecodingFormat.GetResolution().Height;
+		CurrentlyActiveVideoStreamFormat.FrameRate = NewDecodingFormat.GetFrameRate().IsValid() ? NewDecodingFormat.GetFrameRate().GetAsDouble() : 0.0;
 	}
 }
 
@@ -2503,15 +2266,22 @@ void FElectraPlayer::HandlePlayerEventLogMessage(Electra::IInfoLog::ELevel InLog
 	}
 }
 
-void FElectraPlayer::HandlePlayerEventDroppedVideoFrame()
+void FElectraPlayer::HandlePlayerEventDroppedVideoFrame(const Metrics::FFrameStats& InFrameStats)
 {
+	Statistics.NumTotalAddedVideoFrames = InFrameStats.NumTotal;
+	Statistics.NumTotalDroppedVideoFrames = InFrameStats.NumDropped;
+	Statistics.AvgVideoDropsPerInterval = InFrameStats.AvgDropsPerInterval;
+	Statistics.MinVideoDropsPerInterval = InFrameStats.MinDropsPerInterval;
+	Statistics.MaxVideoDropsPerInterval = InFrameStats.MaxDropsPerInterval;
+	Statistics.VideoDropIntervalMs = InFrameStats.DropIntervalMs;
+	Statistics.VideoNumIntervals = InFrameStats.NumIntervals;
+	Statistics.VideoNumIntervalsWithDrops = InFrameStats.NumIntervalsWithDrops;
 }
 
-void FElectraPlayer::HandlePlayerEventDroppedAudioFrame()
+void FElectraPlayer::HandlePlayerEventDroppedAudioFrame(const Metrics::FFrameStats& InFrameStats)
 {
 }
 
-
 void FElectraPlayer::FStatistics::AddMessageToHistory(FString InMessage)
 {
 	if (MessageHistoryBuffer.Num() >= 20)
@@ -2600,10 +2370,11 @@ void FElectraPlayer::LogStatistics()
 		"Number of audio quality upswitches: %d\n"\
 		"Number of audio quality downswitches: %d\n"\
 		"Number of audio drastic downswitches: %d\n"\
-		"Bytes of video data streamed: %lld\n"\
-		"Bytes of audio data streamed: %lld\n"\
+		"Bytes of video data streamed: %lld from %lld segments\n"\
+		"Bytes of audio data streamed: %lld from %lld segments\n"\
 		"Video quality percentage:\n%s"\
 		"Audio quality percentage:\n%s"\
+		"Video frames dropped: %lld of %lld total, %u avg (%u max) per %u ms interval (drops in %.1f%% of all intervals)\n"\
 		"Currently active playlist URL: %s\n"\
 		"Currently active resolution: %d * %d\n" \
 		"Current state: %s\n" \
@@ -2647,9 +2418,17 @@ void FElectraPlayer::LogStatistics()
 		Statistics.NumAudioQualityDownswitches,
 		Statistics.NumAudioQualityDrasticDownswitches,
 		(long long int)Statistics.NumVideoDatabytesStreamed,
+		(long long int)Statistics.NumVideoSegmentsStreamed,
 		(long long int)Statistics.NumAudioDatabytesStreamed,
+		(long long int)Statistics.NumAudioSegmentsStreamed,
 		*VideoSegsPercentage,
 		*AudioSegsPercentage,
+		(long long int)Statistics.NumTotalDroppedVideoFrames,
+		(long long int)Statistics.NumTotalAddedVideoFrames,
+		Statistics.AvgVideoDropsPerInterval,
+		Statistics.MaxVideoDropsPerInterval,
+		Statistics.VideoDropIntervalMs,
+		Statistics.VideoNumIntervals ? ((int64)Statistics.VideoNumIntervalsWithDrops * 10000 / Statistics.VideoNumIntervals) / 100.0 : 0.0,
 		*SanitizeMessage(Statistics.CurrentlyActivePlaylistURL),
 		Statistics.CurrentlyActiveResolutionWidth,
 		Statistics.CurrentlyActiveResolutionHeight,
@@ -2768,6 +2547,18 @@ void FElectraPlayer::SendAnalyticMetrics(const TSharedPtr<IAnalyticsProviderET>&
 	{
 		ParamArray.Add(FAnalyticsEventAttribute(FString::Printf(TEXT("qp%d"), qbIdx++), (int32) qbIt.Value));
 	}
+	ParamArray.Add(FAnalyticsEventAttribute(TEXT("NumVideoSegmentsFetched"), (uint32) Statistics.NumVideoSegmentsStreamed));
+	ParamArray.Add(FAnalyticsEventAttribute(TEXT("NumAudioSegmentsFetched"), (uint32) Statistics.NumAudioSegmentsStreamed));
+	ParamArray.Add(FAnalyticsEventAttribute(TEXT("NumTotalSegmentsFetched"), (uint32) Statistics.NumVideoSegmentsStreamed + Statistics.NumAudioSegmentsStreamed));
+	ParamArray.Add(FAnalyticsEventAttribute(TEXT("NumVideoFramesAddedTotal"), (uint32) Statistics.NumTotalAddedVideoFrames));
+	ParamArray.Add(FAnalyticsEventAttribute(TEXT("NumVideoFramesDroppedTotal"), (uint32) Statistics.NumTotalDroppedVideoFrames));
+	ParamArray.Add(FAnalyticsEventAttribute(TEXT("VideoDropIntervalMs"), (uint32) Statistics.VideoDropIntervalMs));
+	ParamArray.Add(FAnalyticsEventAttribute(TEXT("AvgVideoDropsPerInterval"), (uint32) Statistics.AvgVideoDropsPerInterval));
+	ParamArray.Add(FAnalyticsEventAttribute(TEXT("MinVideoDropsPerInterval"), (uint32) Statistics.MinVideoDropsPerInterval));
+	ParamArray.Add(FAnalyticsEventAttribute(TEXT("MaxVideoDropsPerInterval"), (uint32) Statistics.MaxVideoDropsPerInterval));
+	ParamArray.Add(FAnalyticsEventAttribute(TEXT("VideoNumIntervals"), (uint32) Statistics.VideoNumIntervals));
+	ParamArray.Add(FAnalyticsEventAttribute(TEXT("VideoNumIntervalsWithDrops"), (uint32) Statistics.VideoNumIntervalsWithDrops));
+	ParamArray.Add(FAnalyticsEventAttribute(TEXT("VideoIntervalsWithDropsPercentage"), (double)(Statistics.VideoNumIntervals ? ((int64)Statistics.VideoNumIntervalsWithDrops * 10000 / Statistics.VideoNumIntervals) / 100.0 : 0.0)));
 	ParamArray.Add(FAnalyticsEventAttribute(TEXT("Num404"), (uint32) Statistics.NumErr404));
 	ParamArray.Add(FAnalyticsEventAttribute(TEXT("Num4xx"), (uint32) Statistics.NumErr4xx));
 	ParamArray.Add(FAnalyticsEventAttribute(TEXT("Num5xx"), (uint32) Statistics.NumErr5xx));
@@ -3100,10 +2891,9 @@ void FElectraPlayer::OnMediaPlayerEventReceived(TSharedPtrTS<IAdaptiveStreamingP
 		// timeline will be fixed.
 		Meta->TrackBaseTime = FDecoderTimeStamp(MediaTimeline.Start.GetAsTimespan(), MediaTimeline.Start.GetSequenceIndex());
 
-		FScopeLock SampleLock(&MediaSamplesLock);
 		TSharedRef<FElectraBinarySample, ESPMode::ThreadSafe> MetaDataSample = MakeShared<FElectraBinarySample, ESPMode::ThreadSafe>();
 		MetaDataSample->Metadata = Meta;
-		MediaSamples->AddMetadata(MetaDataSample);
+		Player->AdaptivePlayer->GetMediaSamples()->AddMetadata(MetaDataSample);
 	}
 }
 
@@ -3131,8 +2921,8 @@ FGuid FElectraPlayer::GetPlayerPluginGUID() const
 
 IMediaSamples& FElectraPlayer::GetSamples()
 {
-	FScopeLock lock(&MediaSamplesLock);
-	return MediaSamples.IsValid() ? *MediaSamples : *EmptyMediaSamples;
+	auto Player = CurrentPlayer;
+	return Player.IsValid() ? *Player->AdaptivePlayer->GetMediaSamples() : *EmptyMediaSamples;
 }
 
 FString FElectraPlayer::GetStats() const
@@ -3245,22 +3035,21 @@ bool FElectraPlayer::QueryCacheState(EMediaCacheState InState, TRangeSet<FTimesp
 		case EMediaCacheState::Loading:
 		case EMediaCacheState::Pending:
 		{
-			// When asked to provide what's already loaded we look at what we have in the sample queue
-			// and add that to the result. These samples have already left the player but are ready
-			// for use.
-			if (InState == EMediaCacheState::Loaded)
+			auto Player = CurrentPlayer;
+			if (Player.IsValid())
 			{
-				TArray<TRange<FMediaTimeStamp>> QueuedRange;
-				FScopeLock SampleLock(&MediaSamplesLock);
-				if (MediaSamples.IsValid() && MediaSamples->PeekVideoSampleTimeRanges(QueuedRange) && QueuedRange.Num())
+				// When asked to provide what's already loaded we look at what we have in the sample queue
+				// and add that to the result. These samples have already left the player but are ready
+				// for use.
+				if (InState == EMediaCacheState::Loaded)
 				{
-					OutTimeRanges.Add(TRange<FTimespan>(QueuedRange[0].GetLowerBoundValue().Time, QueuedRange.Last().GetUpperBoundValue().Time));
+					TArray<TRange<FMediaTimeStamp>> QueuedRange;
+					if (Player->AdaptivePlayer->GetMediaSamples()->PeekVideoSampleTimeRanges(QueuedRange) && QueuedRange.Num())
+					{
+						OutTimeRanges.Add(TRange<FTimespan>(QueuedRange[0].GetLowerBoundValue().Time, QueuedRange.Last().GetUpperBoundValue().Time));
+					}
 				}
-			}
 
-			auto Player = CurrentPlayer;
-			if (Player.IsValid())
-			{
 				IAdaptiveStreamingPlayer::FStreamBufferInfo bi;
 				// Query video first.
 				Player->AdaptivePlayer->QueryStreamBufferInfo(bi, Electra::EStreamType::Video);
@@ -3754,7 +3543,17 @@ bool FElectraPlayer::GetVideoTrackFormat(int32 InTrackIndex, int32 InFormatIndex
 			const Electra::FStreamCodecInformation& ci = Meta->HighestBandwidthCodec;
 			OutFormat.Dim.X = ci.GetResolution().Width;
 			OutFormat.Dim.Y = ci.GetResolution().Height;
-			OutFormat.FrameRate = (float)ci.GetFrameRate().GetAsDouble();
+
+			OutFormat.FrameRate = 0.0f;
+			if (ci.GetFrameRate().IsValid())
+			{
+				OutFormat.FrameRate = (float)ci.GetFrameRate().GetAsDouble();
+			}
+			else
+			{
+				FScopeLock lock(&VideoFormatLock);
+				OutFormat.FrameRate = CurrentlyActiveVideoStreamFormat.FrameRate;
+			}
 			OutFormat.FrameRates = TRange<float>{ OutFormat.FrameRate };
 			OutFormat.TypeName = ci.GetHumanReadableCodecName();
 			return true;
@@ -3897,11 +3696,6 @@ void FElectraPlayer::ReportPlaylistDownload(const Metrics::FPlaylistDownloadStat
 	DeferredPlayerEvents.Enqueue(MakeSharedTS<FPlayerMetricEvent_PlaylistDownload>(InPlaylistDownloadStats));
 }
 
-void FElectraPlayer::ReportCleanStart()
-{
-	/*DeferredPlayerEvents.Enqueue(MakeSharedTS<FPlayerMetricEventBase>(FPlayerMetricEventBase::EType::CleanStart));*/
-}
-
 void FElectraPlayer::ReportBufferingStart(Metrics::EBufferingReason InBufferingReason)
 {
 	DeferredPlayerEvents.Enqueue(MakeSharedTS<FPlayerMetricEvent_BufferingStart>(InBufferingReason));
@@ -4012,14 +3806,14 @@ void FElectraPlayer::ReportLogMessage(IInfoLog::ELevel InLogLevel, const FString
 	DeferredPlayerEvents.Enqueue(MakeSharedTS<FPlayerMetricEvent_LogMessage>(InLogLevel, InLogMessage, InPlayerWallclockMilliseconds));
 }
 
-void FElectraPlayer::ReportDroppedVideoFrame()
+void FElectraPlayer::ReportVideoFrameStats(const Metrics::FFrameStats& InFrameStats)
 {
-	DeferredPlayerEvents.Enqueue(MakeSharedTS<FPlayerMetricEventBase>(FPlayerMetricEventBase::EType::DroppedVideoFrame));
+	DeferredPlayerEvents.Enqueue(MakeSharedTS<FPlayerMetricEvent_DroppedFrame>(FPlayerMetricEventBase::EType::DroppedVideoFrame, InFrameStats));
 }
 
-void FElectraPlayer::ReportDroppedAudioFrame()
+void FElectraPlayer::ReportAudioFrameStats(const Metrics::FFrameStats& InFrameStats)
 {
-	DeferredPlayerEvents.Enqueue(MakeSharedTS<FPlayerMetricEventBase>(FPlayerMetricEventBase::EType::DroppedAudioFrame));
+	DeferredPlayerEvents.Enqueue(MakeSharedTS<FPlayerMetricEvent_DroppedFrame>(FPlayerMetricEventBase::EType::DroppedAudioFrame, InFrameStats));
 }
 
 // ------------------------------------------------------------------------------------------------------------------
diff --git a/Engine/Plugins/Media/ElectraPlayer/Source/ElectraPlayerRuntime/Private/ElectraPlayer.h b/Engine/Plugins/Media/ElectraPlayer/Source/ElectraPlayerRuntime/Private/ElectraPlayer.h
--- a/Engine/Plugins/Media/ElectraPlayer/Source/ElectraPlayerRuntime/Private/ElectraPlayer.h
+++ b/Engine/Plugins/Media/ElectraPlayer/Source/ElectraPlayerRuntime/Private/ElectraPlayer.h
@@ -41,8 +41,6 @@
 
 class IMetaDataDecoderOutput;
 using IMetaDataDecoderOutputPtr = TSharedPtr<IMetaDataDecoderOutput, ESPMode::ThreadSafe>;
-class ISubtitleDecoderOutput;
-using ISubtitleDecoderOutputPtr = TSharedPtr<ISubtitleDecoderOutput, ESPMode::ThreadSafe>;
 
 using namespace Electra;
 
@@ -142,7 +140,6 @@ public:
 	void ReportReceivedPlaylists() override;
 	void ReportTracksChanged() override;
 	void ReportPlaylistDownload(const Metrics::FPlaylistDownloadStats& InPlaylistDownloadStats) override;
-	void ReportCleanStart() override;
 	void ReportBufferingStart(Metrics::EBufferingReason InBufferingReason) override;
 	void ReportBufferingEnd(Metrics::EBufferingReason InBufferingReason) override;
 	void ReportBandwidth(int64 InEffectiveBps, int64 InThroughputBps, double InLatencyInSeconds) override;
@@ -165,21 +162,13 @@ public:
 	void ReportMediaMetadataChanged(TSharedPtrTS<Electra::UtilsMP4::FMetadataParser> InMetadata) override;
 	void ReportError(const FString& InErrorReason) override;
 	void ReportLogMessage(IInfoLog::ELevel InLogLevel, const FString& InLogMessage, int64 InPlayerWallclockMilliseconds) override;
-	void ReportDroppedVideoFrame() override;
-	void ReportDroppedAudioFrame() override;
+	void ReportVideoFrameStats(const Metrics::FFrameStats& InFrameStats) override;
+	void ReportAudioFrameStats(const Metrics::FFrameStats& InFrameStats) override;
 
 
 	//-------------------------------------------------------------------------
 	// Implementation methods
 	//
-	void OnVideoDecoded(FElectraTextureSamplePtr InSample);
-	void OnVideoFlush();
-	void OnAudioDecoded(FElectraAudioSamplePtr InSample);
-	void OnAudioFlush();
-	void OnSubtitleDecoded(ISubtitleDecoderOutputPtr InDecoderOutput);
-	void OnSubtitleFlush();
-	bool CanPresentVideoFrames(uint64 InNumFrames);
-	bool CanPresentAudioFrames(uint64 InNumFrames);
 	void PlatformSuspendOrResumeDecoders(bool bInSuspend, const Electra::FParamDict& InOptions);
 
 	void SendAnalyticMetrics(const TSharedPtr<IAnalyticsProviderET>& InAnalyticsProvider, const FGuid& InPlayerGuid);
@@ -234,17 +223,15 @@ private:
 	void HandlePlayerMediaMetadataChanged(const TSharedPtrTS<Electra::UtilsMP4::FMetadataParser>& InMetadata);
 	void HandlePlayerEventError(const FString& InErrorReason);
 	void HandlePlayerEventLogMessage(IInfoLog::ELevel InLogLevel, const FString& InLogMessage, int64 InPlayerWallclockMilliseconds);
-	void HandlePlayerEventDroppedVideoFrame();
-	void HandlePlayerEventDroppedAudioFrame();
+	void HandlePlayerEventDroppedVideoFrame(const Metrics::FFrameStats& InFrameStats);
+	void HandlePlayerEventDroppedAudioFrame(const Metrics::FFrameStats& InFrameStats);
 	void MediaStateOnPreparingFinished();
 	bool MediaStateOnPlay();
 	bool MediaStateOnPause();
 	void MediaStateOnEndReached();
 	void MediaStateOnSeekFinished();
 
-//+++++++++++++++++++++++++++++++++++++++++++++++++++++++++++++
-//+++++++++++++++++++++++++++++++++++++++++++++++++++++++++++++
-
+	//********************************************************************************************************************
 private:
 	DECLARE_DELEGATE_TwoParams(FOnMediaPlayerEventReceivedDelegate, TSharedPtrTS<IAdaptiveStreamingPlayerAEMSEvent> /*InEvent*/, IAdaptiveStreamingPlayerAEMSReceiver::EDispatchMode /*InDispatchMode*/);
 	class FAEMSEventReceiver : public IAdaptiveStreamingPlayerAEMSReceiver
@@ -259,25 +246,6 @@ private:
 		FOnMediaPlayerEventReceivedDelegate EventReceivedDelegate;
 	};
 
-	DECLARE_DELEGATE_OneParam(FOnMediaPlayerSubtitleReceivedDelegate, ISubtitleDecoderOutputPtr);
-	DECLARE_DELEGATE(FOnMediaPlayerSubtitleFlushDelegate);
-	class FSubtitleEventReceiver : public IAdaptiveStreamingPlayerSubtitleReceiver
-	{
-	public:
-		virtual ~FSubtitleEventReceiver() = default;
-		FOnMediaPlayerSubtitleReceivedDelegate& GetSubtitleReceivedDelegate()
-		{ return SubtitleReceivedDelegate; }
-		FOnMediaPlayerSubtitleFlushDelegate& GetSubtitleFlushDelegate()
-		{ return SubtitleFlushDelegate; }
-	private:
-		virtual void OnMediaPlayerSubtitleReceived(ISubtitleDecoderOutputPtr Subtitle) override
-		{ SubtitleReceivedDelegate.ExecuteIfBound(Subtitle); }
-		virtual void OnMediaPlayerFlushSubtitles() override
-		{ SubtitleFlushDelegate.ExecuteIfBound(); }
-		FOnMediaPlayerSubtitleReceivedDelegate SubtitleReceivedDelegate;
-		FOnMediaPlayerSubtitleFlushDelegate SubtitleFlushDelegate;
-	};
-
 	class FAdaptiveStreamingPlayerResourceProvider : public IAdaptiveStreamingPlayerResourceProvider
 	{
 	public:
@@ -301,13 +269,6 @@ private:
 		/** The media player itself **/
 		TSharedPtr<IAdaptiveStreamingPlayer, ESPMode::ThreadSafe> AdaptivePlayer;
 
-		/** Output buffer handlers */
-		TSharedPtr<FOutputHandlerVideo, ESPMode::ThreadSafe> VideoOutputHandler;
-		TSharedPtr<FOutputHandlerAudio, ESPMode::ThreadSafe> AudioOutputHandler;
-
-		/** Media samples to delete when closing */
-		TUniquePtr<FMediaSamples> MediaSamplesToDelete;
-
 		static void DoCloseAsync(TSharedPtr<FInternalPlayerImpl, ESPMode::ThreadSafe>&& InPlayer, uint32 InPlayerID, TSharedPtr<IAsyncResourceReleaseNotification, ESPMode::ThreadSafe> InAsyncResourceReleaseNotification);
 	};
 
@@ -334,6 +295,12 @@ private:
 
 	struct FVideoStreamFormat
 	{
+		void Reset()
+		{
+			Resolution = FIntPoint::ZeroValue;
+			FrameRate = 0.0;
+			Bitrate = 0;
+		}
 		FIntPoint Resolution {};
 		double FrameRate = 0.0;
 		int32 Bitrate = 0;
@@ -373,7 +340,6 @@ private:
 			ReceivedPlaylists,
 			TracksChanged,
 			PlaylistDownload,
-			CleanStart,
 			BufferingStart,
 			BufferingEnd,
 			Bandwidth,
@@ -498,6 +464,11 @@ private:
 		FString LogMessage;
 		int64 PlayerWallclockMilliseconds;
 	};
+	struct FPlayerMetricEvent_DroppedFrame : public FPlayerMetricEventBase
+	{
+		FPlayerMetricEvent_DroppedFrame(EType InType, const Metrics::FFrameStats& InFrameStats)	: FPlayerMetricEventBase(InType), FrameStats(InFrameStats) {}
+		Metrics::FFrameStats FrameStats;
+	};
 
 
 	/** Analytics event */
@@ -652,6 +623,14 @@ private:
 			NumErrTimeouts = 0;
 			NumErrConnDrops = 0;
 			NumErrOther = 0;
+			NumTotalAddedVideoFrames = 0;
+			NumTotalDroppedVideoFrames = 0;
+			AvgVideoDropsPerInterval = 0;
+			MinVideoDropsPerInterval = 0;
+			MaxVideoDropsPerInterval = 0;
+			VideoDropIntervalMs = 0;
+			VideoNumIntervals = 0;
+			VideoNumIntervalsWithDrops = 0;
 		}
 		void AddMessageToHistory(FString InMessage);
 
@@ -709,12 +688,17 @@ private:
 		uint32					NumErrTimeouts;
 		uint32					NumErrConnDrops;
 		uint32					NumErrOther;
+		uint64					NumTotalAddedVideoFrames;
+		uint64					NumTotalDroppedVideoFrames;
+		uint32					AvgVideoDropsPerInterval;
+		uint32					MinVideoDropsPerInterval;
+		uint32					MaxVideoDropsPerInterval;
+		uint32					VideoDropIntervalMs;
+		uint32					VideoNumIntervals;
+		uint32					VideoNumIntervalsWithDrops;
 	};
 
 
-	/** Critical section protecting some members used by both the game and worker threads. */
-	FCriticalSection ParentObjectLock;
-
 	/** Media event sink */
 	IMediaEventSink* EventSink = nullptr;
 
@@ -724,10 +708,8 @@ private:
 	FElectraPlayerReportVideoStreamingErrorDelegate& ReportVideoStreamingErrorDelegate;
 	FElectraPlayerReportSubtitlesMetricsDelegate& ReportSubtitlesMetricsDelegate;
 
-	/** Media sample queues */
-	mutable FCriticalSection MediaSamplesLock;
+	/** Empty sample queue */
 	TUniquePtr<FMediaSamples> EmptyMediaSamples;
-	TUniquePtr<FMediaSamples> MediaSamples;
 
 	/** Option interface */
 	TWeakPtr<IElectraSafeMediaOptionInterface, ESPMode::ThreadSafe> OptionInterface;
@@ -773,9 +755,8 @@ private:
 
 	TQueue<TSharedPtrTS<FPlayerMetricEventBase>> DeferredPlayerEvents;
 	TQueue<EMediaEvent> DeferredMediaEvents;
-	FCriticalSection VideoFormatLock;
-	TOptional<FVideoStreamFormat> CurrentlyActiveVideoStreamFormat;
-	FIntPoint LastPresentedFrameDimension;
+	mutable FCriticalSection VideoFormatLock;
+	FVideoStreamFormat CurrentlyActiveVideoStreamFormat;
 	int32 NumTracksAudio = 0;
 	int32 NumTracksVideo = 0;
 	int32 NumTracksSubtitle = 0;
@@ -792,7 +773,6 @@ private:
 
 	TSharedPtr<FInternalPlayerImpl, ESPMode::ThreadSafe> CurrentPlayer;
 	TSharedPtr<FAEMSEventReceiver, ESPMode::ThreadSafe> MediaPlayerEventReceiver;
-	TSharedPtr<FSubtitleEventReceiver, ESPMode::ThreadSafe> MediaPlayerSubtitleReceiver;
 	TSharedPtr<IAsyncResourceReleaseNotification, ESPMode::ThreadSafe> AsyncResourceReleaseNotification;
 
 	TRange<FTimespan> CurrentPlaybackRange;
diff --git a/Engine/Plugins/Media/ElectraPlayer/Source/ElectraPlayerRuntime/Private/Runtime/Decoder/SubtitleDecoder.cpp b/Engine/Plugins/Media/ElectraPlayer/Source/ElectraPlayerRuntime/Private/Runtime/Decoder/SubtitleDecoder.cpp
--- a/Engine/Plugins/Media/ElectraPlayer/Source/ElectraPlayerRuntime/Private/Runtime/Decoder/SubtitleDecoder.cpp
+++ b/Engine/Plugins/Media/ElectraPlayer/Source/ElectraPlayerRuntime/Private/Runtime/Decoder/SubtitleDecoder.cpp
@@ -4,6 +4,8 @@
 #include "Features/IModularFeatures.h"
 #include "IElectraSubtitleDecoder.h"
 #include "IElectraSubtitleModule.h"
+#include "IElectraSubtitleSample.h"
+
 
 namespace Electra
 {
@@ -127,23 +129,21 @@ public:
 	FSubtitleDecoder(TSharedPtr<IElectraSubtitleDecoder, ESPMode::ThreadSafe> InPluginDecoder);
 	virtual ~FSubtitleDecoder();
 
-	virtual void SetPlayerSessionServices(IPlayerSessionServices* SessionServices) override;
-
-	virtual void Open(const FParamDict& Options) override;
-	virtual void Close() override;
-	virtual void Start() override;
-	virtual void Stop() override;
-	virtual void UpdatePlaybackPosition(FTimeValue InAbsolutePosition, FTimeValue InLocalPosition) override;
+	void SetPlayerSessionServices(IPlayerSessionServices* SessionServices) override;
+	void SetMediaSamples(TSharedPtr<FMediaSamples, ESPMode::ThreadSafe> InMediaSamples) override;
 
-	virtual FTimeValue GetStreamedDeliveryTimeOffset() override;
+	void Open(const FParamDict& Options) override;
+	void Close() override;
+	void Start() override;
+	void Stop() override;
+	void UpdatePlaybackPosition(FTimeValue InAbsolutePosition, FTimeValue InLocalPosition) override;
 
-	virtual FDecodedSubtitleReceivedDelegate& GetDecodedSubtitleReceiveDelegate() override;
-	virtual FDecodedSubtitleFlushDelegate& GetDecodedSubtitleFlushDelegate() override;
+	FTimeValue GetStreamedDeliveryTimeOffset() override;
 
-	virtual void AUdataPushAU(FAccessUnit* AccessUnit) override;
-	virtual void AUdataPushEOD() override;
-	virtual void AUdataClearEOD() override;
-	virtual void AUdataFlushEverything() override;
+	void AUdataPushAU(FAccessUnit* AccessUnit) override;
+	void AUdataPushEOD() override;
+	void AUdataClearEOD() override;
+	void AUdataFlushEverything() override;
 
 	void SetIsSideloaded(bool bInIsSideloaded)
 	{
@@ -167,14 +167,15 @@ public:
 
 private:
 	void OnDecodedSubtitleReceived(ISubtitleDecoderOutputPtr DecodedSubtitle);
+	void FlushSamples();
+
+	TSharedPtr<FMediaSamples, ESPMode::ThreadSafe> CurrentMediaSamples;
 
 	TSharedPtr<IElectraSubtitleDecoder, ESPMode::ThreadSafe> PluginDecoder;
 	FParamDict PluginDecoderOptions;
 	IPlayerSessionServices* PlayerSessionServices = nullptr;
 	FTimeValue DeliveryOffset;
 	FDelegateHandle ReceiveDelegateHandle;
-	FDecodedSubtitleReceivedDelegate SubtitleReceiverDelegate;
-	FDecodedSubtitleFlushDelegate SubtitleFlushDelegate;
 	bool bIsStarted = false;
 
 	struct FCurrentConfig
@@ -256,6 +257,11 @@ void FSubtitleDecoder::SetPlayerSessionServices(IPlayerSessionServices* InSessio
 	PlayerSessionServices = InSessionServices;
 }
 
+void FSubtitleDecoder::SetMediaSamples(TSharedPtr<FMediaSamples, ESPMode::ThreadSafe> InMediaSamples)
+{
+	CurrentMediaSamples = MoveTemp(InMediaSamples);
+}
+
 void FSubtitleDecoder::Open(const FParamDict& Options)
 {
 	check(PlayerSessionServices);
@@ -293,6 +299,7 @@ void FSubtitleDecoder::Close()
 		PluginDecoder->GetParsedSubtitleReceiveDelegate().Remove(ReceiveDelegateHandle);
 		ReceiveDelegateHandle.Reset();
 	}
+	CurrentMediaSamples.Reset();
 	PluginDecoder.Reset();
 	CurrentConfig.Reset();
 }
@@ -329,25 +336,22 @@ FTimeValue FSubtitleDecoder::GetStreamedDeliveryTimeOffset()
 	return DeliveryOffset;
 }
 
-ISubtitleDecoder::FDecodedSubtitleReceivedDelegate& FSubtitleDecoder::GetDecodedSubtitleReceiveDelegate()
-{
-	return SubtitleReceiverDelegate;
-}
-
-ISubtitleDecoder::FDecodedSubtitleFlushDelegate& FSubtitleDecoder::GetDecodedSubtitleFlushDelegate()
-{
-	return SubtitleFlushDelegate;
-}
-
 void FSubtitleDecoder::OnDecodedSubtitleReceived(ISubtitleDecoderOutputPtr DecodedSubtitle)
 {
 	if (DecodedSubtitle.IsValid())
 	{
-	/*
-		FDecoderTimeStamp ts = DecodedSubtitle->GetTime();
-		DecodedSubtitle->SetTime(ts);
-	*/
-		SubtitleReceiverDelegate.ExecuteIfBound(DecodedSubtitle);
+		auto MediaSamples = CurrentMediaSamples;
+		if (MediaSamples.IsValid())
+		{
+			FDecoderTimeStamp ts = DecodedSubtitle->GetTime();
+			FUTF8ToTCHAR cnv((const ANSICHAR*)DecodedSubtitle->GetData().GetData(), DecodedSubtitle->GetData().Num());
+
+			TSharedRef<FElectraSubtitleSample, ESPMode::ThreadSafe> SubtitleSample = MakeShared<FElectraSubtitleSample, ESPMode::ThreadSafe>();
+			SubtitleSample->Timestamp.SetTimeAndIndexValue(ts.Time, ts.SequenceIndex); 
+			SubtitleSample->Duration = DecodedSubtitle->GetDuration();
+			SubtitleSample->Text = FString::ConstructFromPtrSize(cnv.Get(), cnv.Length());
+			MediaSamples->AddSubtitle(SubtitleSample);
+		}
 	}
 }
 
@@ -360,7 +364,7 @@ bool FSubtitleDecoder::ResetForTrackChange(FAccessUnit* AccessUnit)
 	{
 		PluginDecoder->Flush();
 	}
-	SubtitleFlushDelegate.ExecuteIfBound();
+	FlushSamples();
 
 	// Do we need to close and create a new decoder?
 	if (AccessUnit && AccessUnit->AUCodecData.IsValid())
@@ -371,6 +375,7 @@ bool FSubtitleDecoder::ResetForTrackChange(FAccessUnit* AccessUnit)
 			CurrentConfig.RawCSD != AccessUnit->AUCodecData->RawCSD ||
 			!CurrentConfig.CodecInfo.Equals(AccessUnit->AUCodecData->ParsedInfo))
 		{
+			auto MediaSamples = CurrentMediaSamples;
 			Close();
 			if (AccessUnit && AccessUnit->AUCodecData.IsValid())
 			{
@@ -390,6 +395,7 @@ bool FSubtitleDecoder::ResetForTrackChange(FAccessUnit* AccessUnit)
 					SetSourceID(NewSourceID);
 					SetCSD(AccessUnit->AUCodecData->RawCSD);
 					SetCodecInfo(AccessUnit->AUCodecData->ParsedInfo);
+					CurrentMediaSamples = MoveTemp(MediaSamples);
 					Open(PluginDecoderOptions);
 					if (bWasStarted)
 					{
@@ -465,8 +471,20 @@ void FSubtitleDecoder::AUdataFlushEverything()
 	{
 		PluginDecoder->Flush();
 	}
-	SubtitleFlushDelegate.ExecuteIfBound();
+	FlushSamples();
 }
 
-} // namespace Electra
+void FSubtitleDecoder::FlushSamples()
+{
+	auto MediaSamples = CurrentMediaSamples;
+	if (MediaSamples.IsValid())
+	{
+		TRange<FTimespan> AllTime(FTimespan::MinValue(), FTimespan::MaxValue());
+		TSharedPtr<IMediaOverlaySample, ESPMode::ThreadSafe> FlushSample;
+		while(MediaSamples->FetchSubtitle(AllTime, FlushSample))
+		{ }
+	}
+}
 
+
+} // namespace Electra
diff --git a/Engine/Plugins/Media/ElectraPlayer/Source/ElectraPlayerRuntime/Private/Runtime/Decoder/SubtitleDecoder.h b/Engine/Plugins/Media/ElectraPlayer/Source/ElectraPlayerRuntime/Private/Runtime/Decoder/SubtitleDecoder.h
--- a/Engine/Plugins/Media/ElectraPlayer/Source/ElectraPlayerRuntime/Private/Runtime/Decoder/SubtitleDecoder.h
+++ b/Engine/Plugins/Media/ElectraPlayer/Source/ElectraPlayerRuntime/Private/Runtime/Decoder/SubtitleDecoder.h
@@ -3,13 +3,11 @@
 #pragma once
 
 #include "PlayerCore.h"
+#include "MediaSamples.h"
 
 #include "StreamAccessUnitBuffer.h"
 #include "Decoder/DecoderBase.h"
 
-class ISubtitleDecoderOutput;
-using ISubtitleDecoderOutputPtr = TSharedPtr<ISubtitleDecoderOutput, ESPMode::ThreadSafe>;
-
 
 namespace Electra
 {
@@ -32,6 +30,7 @@ namespace Electra
 		// State methods
 		//
 		virtual void SetPlayerSessionServices(IPlayerSessionServices* SessionServices) = 0;
+		virtual void SetMediaSamples(TSharedPtr<FMediaSamples, ESPMode::ThreadSafe> InMediaSamples) = 0;
 		virtual void Open(const FParamDict& Options) = 0;
 		virtual void Close() = 0;
 		virtual void Start() = 0;
@@ -43,14 +42,6 @@ namespace Electra
 		virtual FTimeValue GetStreamedDeliveryTimeOffset() = 0;
 
 
-		//-------------------------------------------------------------------------
-		// Decoded subtitle receiver
-		//
-		DECLARE_DELEGATE_OneParam(FDecodedSubtitleReceivedDelegate, ISubtitleDecoderOutputPtr /* DecodedSubtitle */);
-		virtual FDecodedSubtitleReceivedDelegate& GetDecodedSubtitleReceiveDelegate() = 0;
-		DECLARE_DELEGATE(FDecodedSubtitleFlushDelegate);
-		virtual FDecodedSubtitleFlushDelegate& GetDecodedSubtitleFlushDelegate() = 0;
-
 		//-------------------------------------------------------------------------
 		// Methods from IAccessUnitBufferInterface
 		//
diff --git a/Engine/Plugins/Media/ElectraPlayer/Source/ElectraPlayerRuntime/Private/Runtime/OutputHandler.h b/Engine/Plugins/Media/ElectraPlayer/Source/ElectraPlayerRuntime/Private/Runtime/OutputHandler.h
--- a/Engine/Plugins/Media/ElectraPlayer/Source/ElectraPlayerRuntime/Private/Runtime/OutputHandler.h
+++ b/Engine/Plugins/Media/ElectraPlayer/Source/ElectraPlayerRuntime/Private/Runtime/OutputHandler.h
@@ -6,6 +6,7 @@
 #include "Templates/SharedPointer.h"
 #include "Templates/PimplPtr.h"
 #include "IMediaTimeSource.h"
+#include "MediaSamples.h"
 
 #include "ElectraTextureSample.h"
 #include "IElectraAudioSample.h"
@@ -13,6 +14,8 @@
 #include "PlayerClock.h"
 #include "ParameterDictionary.h"
 
+#include <atomic>
+
 namespace Electra
 {
 
@@ -21,6 +24,11 @@ namespace Electra
 	public:
 		virtual ~IOutputHandlerBase() = default;
 
+		virtual void SetMediaSamples(TSharedPtr<FMediaSamples, ESPMode::ThreadSafe> InMediaSamples)
+		{
+			CurrentMediaSamples = MoveTemp(InMediaSamples);
+		}
+
 		virtual bool PreparePool(const FParamDict& InParameters) = 0;
 		virtual void ClosePool() = 0;
 		virtual const FParamDict& GetPoolProperties() = 0;
@@ -47,6 +55,9 @@ namespace Electra
 
 		virtual FTimespan GetEnqueuedSampleDuration() = 0;
 		virtual void GetEnqueuedSampleInfo(TArray<FEnqueuedSampleInfo>& OutOptionalSampleInfos) = 0;
+
+	protected:
+		TSharedPtr<FMediaSamples, ESPMode::ThreadSafe> CurrentMediaSamples;
 	};
 
 
@@ -56,20 +67,9 @@ namespace Electra
 		FOutputHandlerVideo();
 		virtual ~FOutputHandlerVideo();
 
-		DECLARE_DELEGATE_TwoParams(FCanOutputQueueReceive, bool&, int32);
-		FCanOutputQueueReceive& CanOutputQueueReceiveDelegate();
-
-		DECLARE_DELEGATE_OneParam(FOutputQueueReceiveSample, FElectraTextureSamplePtr);
-		FOutputQueueReceiveSample& OutputQueueReceiveSampleDelegate();
-
-		DECLARE_DELEGATE(FOutputQueueFlushSamples);
-		FOutputQueueFlushSamples& OutputQueueFlushSamplesDelegate();
-
 		/** Sets the texture pool from where ObtainBuffer() will get an output sample. */
 		void SetOutputTexturePool(TSharedPtr<FElectraTextureSamplePool, ESPMode::ThreadSafe> InOutputTexturePool);
 
-		void DetachPlayer();
-
 		bool PreparePool(const FParamDict& InParameters) override;
 		void ClosePool() override;
 		const FParamDict& GetPoolProperties() override;
@@ -93,14 +93,17 @@ namespace Electra
 
 		void ReturnOutputSample(FElectraTextureSamplePtr InTextureSampleToReturn, EReturnSampleType InSendToOutputQueueType);
 
+		/** Returns the total number of video samples added to the queue. */
+		uint64 GetTotalAddedSamples();
+		/** Returns the total number of samples that were dropped from the queue. */
+		uint64 GetTotalNumDroppedSamples();
+		/** Returns the number of dropped samples since the last time this methods was called. */
+		uint64 GetNewlyNumDroppedSamples(bool bInClear);
+
 	private:
 		void ReleaseSampleToPool(IElectraTextureSampleBase* InTextureSampleToReturn, uint64 InSampleID);
 		void SendAllPendingSamples();
 
-		FCanOutputQueueReceive CanOutputQueueReceiveDlg;
-		FOutputQueueReceiveSample OutputQueueReceiveSampleDlg;
-		FOutputQueueFlushSamples OutputQueueFlushSamplesDlg;
-
 		TSharedPtr<FElectraTextureSamplePool, ESPMode::ThreadSafe> OutputTexturePool;
 		TSharedPtr<IMediaRenderClock, ESPMode::ThreadSafe> Clock;
 
@@ -112,10 +115,12 @@ namespace Electra
 		FCriticalSection PendingOutputSamplesLock;
 		TArray<FPendingSample> PendingOutputSamples;
 		FParamDict PoolProperties;
+		std::atomic<uint64> NumTotalSamplesAdded { 0 };
+		std::atomic<uint64> NumTotalSamplesDropped { 0 };
+		std::atomic<uint64> TotalSamplesDroppedAtLastCall { 0 };
 		int32 NumOutputSamples = 0;
 		bool bSendOutput = false;
 		bool bDidSendFirst = false;
-		bool bIsDetached = false;
 
 		FCriticalSection EnqueuedSampleInfosLock;
 		TArray<FEnqueuedSampleInfo> EnqueuedSampleInfos;
@@ -130,20 +135,9 @@ namespace Electra
 		FOutputHandlerAudio();
 		virtual ~FOutputHandlerAudio();
 
-		DECLARE_DELEGATE_TwoParams(FCanOutputQueueReceive, bool&, int32);
-		FCanOutputQueueReceive& CanOutputQueueReceiveDelegate();
-
-		DECLARE_DELEGATE_OneParam(FOutputQueueReceiveSample, FElectraAudioSamplePtr);
-		FOutputQueueReceiveSample& OutputQueueReceiveSampleDelegate();
-
-		DECLARE_DELEGATE(FOutputQueueFlushSamples);
-		FOutputQueueFlushSamples& OutputQueueFlushSamplesDelegate();
-
 		/** Sets the sample pool from where ObtainBuffer() will get an output sample. */
 		void SetOutputAudioSamplePool(TSharedPtr<FElectraAudioSamplePool, ESPMode::ThreadSafe> InOutputAudioSamplePool);
 
-		void DetachPlayer();
-
 		bool PreparePool(const FParamDict& InParameters) override;
 		void ClosePool() override;
 		const FParamDict& GetPoolProperties() override;
@@ -184,10 +178,6 @@ namespace Electra
 		void ReleaseSampleToPool(FElectraAudioSample* InTextureSampleToReturn, uint64 InSampleID);
 		void SendPendingSamples();
 
-		FCanOutputQueueReceive CanOutputQueueReceiveDlg;
-		FOutputQueueReceiveSample OutputQueueReceiveSampleDlg;
-		FOutputQueueFlushSamples OutputQueueFlushSamplesDlg;
-
 		TSharedPtr<FElectraAudioSamplePool, ESPMode::ThreadSafe> OutputAudioSamplePool;
 		TSharedPtr<IMediaRenderClock, ESPMode::ThreadSafe> Clock;
 
@@ -201,7 +191,6 @@ namespace Electra
 		FParamDict PoolProperties;
 		int32 NumOutputSamples = 0;
 		bool bSendOutput = false;
-		bool bIsDetached = false;
 
 		double CurrentPlaybackRate = 0.0;
 		double IntendedPlaybackRate = 0.0;
diff --git a/Engine/Plugins/Media/ElectraPlayer/Source/ElectraPlayerRuntime/Private/Runtime/OutputHandlerAudio.cpp b/Engine/Plugins/Media/ElectraPlayer/Source/ElectraPlayerRuntime/Private/Runtime/OutputHandlerAudio.cpp
--- a/Engine/Plugins/Media/ElectraPlayer/Source/ElectraPlayerRuntime/Private/Runtime/OutputHandlerAudio.cpp
+++ b/Engine/Plugins/Media/ElectraPlayer/Source/ElectraPlayerRuntime/Private/Runtime/OutputHandlerAudio.cpp
@@ -146,7 +146,6 @@ namespace Electra
 	FOutputHandlerAudio::~FOutputHandlerAudio()
 	{
 		AudioVars.Reset();
-		// TBD: Would there be a reason to explicitly unbind the delegates or reset the output sample pool?
 	}
 
 	bool FOutputHandlerAudio::PreparePool(const FParamDict& InParameters)
@@ -202,10 +201,14 @@ namespace Electra
 		}
 		PendingOutputSamplesLock.Unlock();
 
-		// Notify the registered sample flush delegate.
-		if (!bIsDetached)
+		// Flush any samples in the media queue.
+		auto MediaSamples = CurrentMediaSamples;
+		if (MediaSamples.IsValid())
 		{
-			OutputQueueFlushSamplesDelegate().ExecuteIfBound();
+			TRange<FTimespan> AllTime(FTimespan::MinValue(), FTimespan::MaxValue());
+			TSharedPtr<IMediaAudioSample, ESPMode::ThreadSafe> FlushSample;
+			while(MediaSamples->FetchAudio(AllTime, FlushSample))
+			{ }
 		}
 
 		// Anything that has not been returned by now we drop from the sample info stats.
@@ -232,29 +235,10 @@ namespace Electra
 		OutputAudioSamplePool = MoveTemp(InOutputAudioSamplePool);
 	}
 
-	void FOutputHandlerAudio::DetachPlayer()
-	{
-		bIsDetached = true;
-	}
-
-	FOutputHandlerAudio::FCanOutputQueueReceive& FOutputHandlerAudio::CanOutputQueueReceiveDelegate()
-	{
-		return CanOutputQueueReceiveDlg;
-	}
-
-	FOutputHandlerAudio::FOutputQueueReceiveSample& FOutputHandlerAudio::OutputQueueReceiveSampleDelegate()
-	{
-		return OutputQueueReceiveSampleDlg;
-	}
-
-	FOutputHandlerAudio::FOutputQueueFlushSamples& FOutputHandlerAudio::OutputQueueFlushSamplesDelegate()
-	{
-		return OutputQueueFlushSamplesDlg;
-	}
-
 	bool FOutputHandlerAudio::CanReceiveOutputSample()
 	{
-		if (bIsDetached)
+		auto MediaSamples = CurrentMediaSamples;
+		if (!MediaSamples.IsValid())
 		{
 			return false;
 		}
@@ -266,7 +250,7 @@ namespace Electra
 		{
 			return false;
 		}
-		CanOutputQueueReceiveDelegate().ExecuteIfBound(bCanRcv, 1 + NumPendingOutputSamples);
+		bCanRcv = MediaSamples->CanReceiveAudioSamples(1 + NumPendingOutputSamples);
 		return bCanRcv;
 	}
 
@@ -282,7 +266,8 @@ namespace Electra
 
 	void FOutputHandlerAudio::ReturnOutputSample(FElectraAudioSamplePtr InAudioSampleToReturn, EReturnSampleType InSendToOutputQueueType)
 	{
-		if (InAudioSampleToReturn && InSendToOutputQueueType != EReturnSampleType::DontSendToQueue && !bIsDetached)
+		auto MediaSamples = CurrentMediaSamples;
+		if (InAudioSampleToReturn && InSendToOutputQueueType != EReturnSampleType::DontSendToQueue && MediaSamples.IsValid())
 		{
 			TArray<FElectraAudioSamplePtr> ProcessedSampleBlocks;
 			double Rate = GetPlayRateScale();
@@ -327,11 +312,14 @@ namespace Electra
 
 	void FOutputHandlerAudio::SendPendingSamples()
 	{
+		auto MediaSamples = CurrentMediaSamples;
+		if (!MediaSamples.IsValid())
+		{
+			return;
+		}
 		while(bSendOutput)
 		{
-			bool bCanRcv = false;
-			CanOutputQueueReceiveDelegate().ExecuteIfBound(bCanRcv, 1);
-			if (!bCanRcv)
+			if (!MediaSamples->CanReceiveAudioSamples(1))
 			{
 				return;
 			}
@@ -346,9 +334,9 @@ namespace Electra
 				FPendingSample ps = PendingOutputSamples[0];
 				PendingOutputSamples.RemoveAt(0);
 				PendingOutputSamplesLock.Unlock();
-				if (!bIsDetached && !ps.SampleInfo.bIsDummy)
+				if (!ps.SampleInfo.bIsDummy)
 				{
-					OutputQueueReceiveSampleDelegate().ExecuteIfBound(MoveTemp(ps.Sample));
+					MediaSamples->AddAudio(ps.Sample.ToSharedRef());
 				}
 				ps.Sample.Reset();
 			}
diff --git a/Engine/Plugins/Media/ElectraPlayer/Source/ElectraPlayerRuntime/Private/Runtime/OutputHandlerVideo.cpp b/Engine/Plugins/Media/ElectraPlayer/Source/ElectraPlayerRuntime/Private/Runtime/OutputHandlerVideo.cpp
--- a/Engine/Plugins/Media/ElectraPlayer/Source/ElectraPlayerRuntime/Private/Runtime/OutputHandlerVideo.cpp
+++ b/Engine/Plugins/Media/ElectraPlayer/Source/ElectraPlayerRuntime/Private/Runtime/OutputHandlerVideo.cpp
@@ -11,7 +11,6 @@ namespace Electra
 
 	FOutputHandlerVideo::~FOutputHandlerVideo()
 	{
-		// TBD: Would there be a reason to explicitly unbind the delegates or reset the output texture pool?
 	}
 
 	bool FOutputHandlerVideo::PreparePool(const FParamDict& InParameters)
@@ -67,10 +66,14 @@ namespace Electra
 		}
 		PendingOutputSamplesLock.Unlock();
 
-		// Notify the registered sample flush delegate.
-		if (!bIsDetached)
+		// Flush any samples in the media queue.
+		auto MediaSamples = CurrentMediaSamples;
+		if (MediaSamples.IsValid())
 		{
-			OutputQueueFlushSamplesDelegate().ExecuteIfBound();
+			TRange<FTimespan> AllTime(FTimespan::MinValue(), FTimespan::MaxValue());
+			TSharedPtr<IMediaTextureSample, ESPMode::ThreadSafe> FlushSample;
+			while(MediaSamples->FetchVideo(AllTime, FlushSample))
+			{ }
 		}
 
 		// Anything that has not been returned by now we drop from the sample info stats.
@@ -83,46 +86,54 @@ namespace Electra
 		bDidSendFirst = false;
 	}
 
-	FTimespan FOutputHandlerVideo::GetEnqueuedSampleDuration()
-	{
-		FScopeLock lock(&EnqueuedSampleInfosLock);
-		return EnqueuedSampleDuration;
-	}
 
-	void FOutputHandlerVideo::GetEnqueuedSampleInfo(TArray<FEnqueuedSampleInfo>& OutOptionalSampleInfos)
+	uint64 FOutputHandlerVideo::GetTotalAddedSamples()
 	{
-		FScopeLock lock(&EnqueuedSampleInfosLock);
-		OutOptionalSampleInfos = EnqueuedSampleInfos;
+		return NumTotalSamplesAdded;
 	}
 
-	void FOutputHandlerVideo::SetOutputTexturePool(TSharedPtr<FElectraTextureSamplePool, ESPMode::ThreadSafe> InOutputTexturePool)
+	uint64 FOutputHandlerVideo::GetTotalNumDroppedSamples()
 	{
-		OutputTexturePool = MoveTemp(InOutputTexturePool);
+		auto MediaSamples = CurrentMediaSamples;
+		if (MediaSamples.IsValid())
+		{
+			NumTotalSamplesDropped = MediaSamples->GetNumDroppedVideoSamples(false);
+		}
+		return NumTotalSamplesDropped;
 	}
 
-	void FOutputHandlerVideo::DetachPlayer()
+	uint64 FOutputHandlerVideo::GetNewlyNumDroppedSamples(bool bInClear)
 	{
-		bIsDetached = true;
+		uint64 nd = GetTotalNumDroppedSamples();
+		uint64 Diff = (uint64)(nd - TotalSamplesDroppedAtLastCall);
+		if (bInClear)
+		{
+			TotalSamplesDroppedAtLastCall = nd;
+		}
+		return Diff;
 	}
 
-	FOutputHandlerVideo::FCanOutputQueueReceive& FOutputHandlerVideo::CanOutputQueueReceiveDelegate()
+	FTimespan FOutputHandlerVideo::GetEnqueuedSampleDuration()
 	{
-		return CanOutputQueueReceiveDlg;
+		FScopeLock lock(&EnqueuedSampleInfosLock);
+		return EnqueuedSampleDuration;
 	}
 
-	FOutputHandlerVideo::FOutputQueueReceiveSample& FOutputHandlerVideo::OutputQueueReceiveSampleDelegate()
+	void FOutputHandlerVideo::GetEnqueuedSampleInfo(TArray<FEnqueuedSampleInfo>& OutOptionalSampleInfos)
 	{
-		return OutputQueueReceiveSampleDlg;
+		FScopeLock lock(&EnqueuedSampleInfosLock);
+		OutOptionalSampleInfos = EnqueuedSampleInfos;
 	}
 
-	FOutputHandlerVideo::FOutputQueueFlushSamples& FOutputHandlerVideo::OutputQueueFlushSamplesDelegate()
+	void FOutputHandlerVideo::SetOutputTexturePool(TSharedPtr<FElectraTextureSamplePool, ESPMode::ThreadSafe> InOutputTexturePool)
 	{
-		return OutputQueueFlushSamplesDlg;
+		OutputTexturePool = MoveTemp(InOutputTexturePool);
 	}
 
 	bool FOutputHandlerVideo::CanReceiveOutputSample()
 	{
-		if (bIsDetached)
+		auto MediaSamples = CurrentMediaSamples;
+		if (!MediaSamples.IsValid())
 		{
 			return false;
 		}
@@ -134,7 +145,7 @@ namespace Electra
 		{
 			return false;
 		}
-		CanOutputQueueReceiveDelegate().ExecuteIfBound(bCanRcv, 1 + NumPendingOutputSamples);
+		bCanRcv = MediaSamples->CanReceiveVideoSamples(1 + NumPendingOutputSamples);
 		return bCanRcv;
 	}
 
@@ -150,7 +161,8 @@ namespace Electra
 
 	void FOutputHandlerVideo::ReturnOutputSample(FElectraTextureSamplePtr InTextureSampleToReturn, EReturnSampleType InSendToOutputQueueType)
 	{
-		if (InTextureSampleToReturn && InSendToOutputQueueType != EReturnSampleType::DontSendToQueue && !bIsDetached)
+		auto MediaSamples = CurrentMediaSamples;
+		if (InTextureSampleToReturn && InSendToOutputQueueType != EReturnSampleType::DontSendToQueue && MediaSamples.IsValid())
 		{
 			uint64 SampleID = ++NextSampleID;
 
@@ -177,7 +189,8 @@ namespace Electra
 				// to ensure the pipeline is not running dry.
 				if (!si.bIsDummy)
 				{
-					OutputQueueReceiveSampleDelegate().ExecuteIfBound(InTextureSampleToReturn);
+					MediaSamples->AddVideo(InTextureSampleToReturn.ToSharedRef());
+					++NumTotalSamplesAdded;
 				}
 			}
 			else
@@ -194,15 +207,17 @@ namespace Electra
 
 	void FOutputHandlerVideo::SendAllPendingSamples()
 	{
+		auto MediaSamples = CurrentMediaSamples;
 		PendingOutputSamplesLock.Lock();
 		while(PendingOutputSamples.Num())
 		{
 			FPendingSample ps = PendingOutputSamples[0];
 			PendingOutputSamples.RemoveAt(0);
 			PendingOutputSamplesLock.Unlock();
-			if (!bIsDetached && !ps.SampleInfo.bIsDummy)
+			if (MediaSamples.IsValid() && !ps.SampleInfo.bIsDummy)
 			{
-				OutputQueueReceiveSampleDelegate().ExecuteIfBound(MoveTemp(ps.Sample));
+				MediaSamples->AddVideo(ps.Sample.ToSharedRef());
+				++NumTotalSamplesAdded;
 			}
 			ps.Sample.Reset();
 			PendingOutputSamplesLock.Lock();
diff --git a/Engine/Plugins/Media/ElectraPlayer/Source/ElectraPlayerRuntime/Private/Runtime/Player/AdaptivePlayerOptionKeynames.h b/Engine/Plugins/Media/ElectraPlayer/Source/ElectraPlayerRuntime/Private/Runtime/Player/AdaptivePlayerOptionKeynames.h
--- a/Engine/Plugins/Media/ElectraPlayer/Source/ElectraPlayerRuntime/Private/Runtime/Player/AdaptivePlayerOptionKeynames.h
+++ b/Engine/Plugins/Media/ElectraPlayer/Source/ElectraPlayerRuntime/Private/Runtime/Player/AdaptivePlayerOptionKeynames.h
@@ -70,6 +70,7 @@ namespace Electra
 
 	const FName OptionKeyCMCDConfiguration(TEXT("cmcd_configuration"));
 
+	const FName OptionKeyFrameStatInterval(TEXT("framestat_interval"));
 
 	namespace CustomOptions
 	{
diff --git a/Engine/Plugins/Media/ElectraPlayer/Source/ElectraPlayerRuntime/Private/Runtime/Player/AdaptiveStreamingPlayer.cpp b/Engine/Plugins/Media/ElectraPlayer/Source/ElectraPlayerRuntime/Private/Runtime/Player/AdaptiveStreamingPlayer.cpp
--- a/Engine/Plugins/Media/ElectraPlayer/Source/ElectraPlayerRuntime/Private/Runtime/Player/AdaptiveStreamingPlayer.cpp
+++ b/Engine/Plugins/Media/ElectraPlayer/Source/ElectraPlayerRuntime/Private/Runtime/Player/AdaptiveStreamingPlayer.cpp
@@ -49,8 +49,15 @@ FAdaptiveStreamingPlayer::FAdaptiveStreamingPlayer(const IAdaptiveStreamingPlaye
 {
 	Electra::AddActivePlayerInstance();
 
-	VideoOutput.OutputHandler = InCreateParameters.VideoOutputHandler;
-	AudioOutput.OutputHandler = InCreateParameters.AudioOutputHandler;
+	CurrentMediaSamples = MakeShared<FMediaSamples, ESPMode::ThreadSafe>();
+	EmptyMediaSamples = MakeShared<FMediaSamples, ESPMode::ThreadSafe>();
+	VideoOutput.OutputHandler = MakeShared<FOutputHandlerVideo, ESPMode::ThreadSafe>();
+	VideoOutput.OutputHandler->SetOutputTexturePool(InCreateParameters.OutputTexturePool);
+	VideoOutput.OutputHandler->SetMediaSamples(CurrentMediaSamples);
+	AudioOutput.OutputHandler = MakeShared<FOutputHandlerAudio, ESPMode::ThreadSafe>();
+	AudioOutput.OutputHandler->SetOutputAudioSamplePool(InCreateParameters.OutputAudioSamplePool);
+	AudioOutput.OutputHandler->SetMediaSamples(CurrentMediaSamples);
+
 	ExternalPlayerGUID = InCreateParameters.ExternalPlayerGUID;
 	UseSharedWorkerThreads = InCreateParameters.WorkerThreads;
 	ManifestType = EMediaFormatType::Unknown;
@@ -241,6 +248,74 @@ void FAdaptiveStreamingPlayer::UpdateDiagnostics()
 	DiagnosticsCriticalSection.Unlock();
 }
 
+void FAdaptiveStreamingPlayer::UpdateFrameStats(bool bFinal)
+{
+	bool bSendReport = bFinal;
+	// Track frame stats only while playing.
+	if (CurrentState == EPlayerState::eState_Paused)
+	{
+		if (DroppedFrameStats.LastProcessedTime >= 0 && DroppedFrameStats.PausedDelta < 0)
+		{
+			DroppedFrameStats.PausedDelta = MEDIAutcTime::CurrentMSec() - DroppedFrameStats.LastProcessedTime;
+		}
+	}
+	else if (CurrentState == EPlayerState::eState_Playing)
+	{
+		// Check if any video frames have been dropped.
+		uint64 NumNewDroppedVideoFrames = VideoOutput.PeekNewlyNumDroppedSamples();
+		bSendReport |= NumNewDroppedVideoFrames != 0;
+		if (DroppedFrameStats.ReportingIntervalMs > 0)
+		{
+			// Add an entry to the dropped frame stats.
+			if (NumNewDroppedVideoFrames)
+			{
+				FDroppedFrameStats::FEntry& NewEntry = DroppedFrameStats.Entries.Emplace_GetRef();
+				NewEntry.DropCount = NumNewDroppedVideoFrames;
+			}
+			int64 tNow = MEDIAutcTime::CurrentMSec();
+			// First time init?
+			if (DroppedFrameStats.LastProcessedTime < 0)
+			{
+				DroppedFrameStats.LastProcessedTime = tNow;
+			}
+			// Back to playing from previously paused?
+			if (DroppedFrameStats.PausedDelta > 0)
+			{
+				DroppedFrameStats.LastProcessedTime = tNow - DroppedFrameStats.PausedDelta;
+				DroppedFrameStats.PausedDelta = -1;
+			}
+			int64 tElapsed = tNow - DroppedFrameStats.LastProcessedTime;
+			if (bFinal || tElapsed >= DroppedFrameStats.ReportingIntervalMs)
+			{
+				// Sum up and remove entries up to the end of the current reporting interval.
+				uint64 DropsInInterval = 0;
+				for(auto& e : DroppedFrameStats.Entries)
+				{
+					DropsInInterval += e.DropCount;
+				}
+				DroppedFrameStats.Entries.Empty();
+				if (DropsInInterval)
+				{
+					DroppedFrameStats.AverageDropsPerInterval.AddValue((int64)DropsInInterval);
+					++DroppedFrameStats.NumIntervalsWithDrops;
+				}
+				DroppedFrameStats.LastProcessedTime = tNow;
+				++DroppedFrameStats.NumIntervals;
+				bSendReport = true;
+			}
+		}
+	}
+	if (bSendReport)
+	{
+		DispatchEvent(FMetricEvent::ReportVideoFrameStats(VideoOutput.GetTotalAddedSamples(), VideoOutput.GetTotalNumDroppedSamples(), VideoOutput.GetNewlyNumDroppedSamples(), DroppedFrameStats));
+	}
+	// When this is the final report we stop the interval
+	if (bFinal)
+	{
+		DroppedFrameStats.ReportingIntervalMs = -1;
+	}
+}
+
 
 
 
@@ -303,50 +378,6 @@ void FAdaptiveStreamingPlayer::RemoveAEMSReceiver(TWeakPtrTS<IAdaptiveStreamingP
 }
 
 
-void FAdaptiveStreamingPlayer::AddSubtitleReceiver(TWeakPtrTS<IAdaptiveStreamingPlayerSubtitleReceiver> InReceiver)
-{
-	if (InReceiver.IsValid())
-	{
-		FScopeLock lock(&SubtitleReceiversCriticalSection);
-		SubtitleReceivers.Add(InReceiver);
-	}
-}
-
-void FAdaptiveStreamingPlayer::RemoveSubtitleReceiver(TWeakPtrTS<IAdaptiveStreamingPlayerSubtitleReceiver> InReceiver)
-{
-	if (InReceiver.IsValid())
-	{
-		FScopeLock lock(&SubtitleReceiversCriticalSection);
-		SubtitleReceivers.Remove(InReceiver);
-	}
-}
-
-void FAdaptiveStreamingPlayer::OnDecodedSubtitleReceived(ISubtitleDecoderOutputPtr Subtitle)
-{
-	FScopeLock lock(&SubtitleReceiversCriticalSection);
-	for(auto &Receiver : SubtitleReceivers)
-	{
-		TSharedPtrTS<IAdaptiveStreamingPlayerSubtitleReceiver> Recv = Receiver.Pin();
-		if (Recv.IsValid())
-		{
-			Recv->OnMediaPlayerSubtitleReceived(Subtitle);
-		}
-	}
-}
-
-void FAdaptiveStreamingPlayer::OnFlushSubtitleReceivers()
-{
-	FScopeLock lock(&SubtitleReceiversCriticalSection);
-	for(auto &Receiver : SubtitleReceivers)
-	{
-		TSharedPtrTS<IAdaptiveStreamingPlayerSubtitleReceiver> Recv = Receiver.Pin();
-		if (Recv.IsValid())
-		{
-			Recv->OnMediaPlayerFlushSubtitles();
-		}
-	}
-}
-
 
 
 
@@ -388,9 +419,6 @@ void FAdaptiveStreamingPlayer::FireSyncEvent(TSharedPtrTS<FMetricEvent> Event)
 			case FMetricEvent::EType::TracksChanged:
 				Listeners[i]->ReportTracksChanged();
 				break;
-			case FMetricEvent::EType::CleanStart:
-				Listeners[i]->ReportCleanStart();
-				break;
 			case FMetricEvent::EType::BufferingStart:
 				Listeners[i]->ReportBufferingStart(Event->Param.BufferingReason);
 				break;
@@ -460,6 +488,9 @@ void FAdaptiveStreamingPlayer::FireSyncEvent(TSharedPtrTS<FMetricEvent> Event)
 			case FMetricEvent::EType::LogMessage:
 				Listeners[i]->ReportLogMessage(Event->Param.LogMessage.Level, Event->Param.LogMessage.Message, Event->Param.LogMessage.AtMillis);
 				break;
+			case FMetricEvent::EType::VideoFrameStats:
+				Listeners[i]->ReportVideoFrameStats(Event->Param.FrameStats);
+				break;
 		}
 	}
 	UnlockMetricsReceivers();
@@ -562,6 +593,17 @@ void FAdaptiveStreamingPlayer::LoadManifest(const FString& InManifestURL)
 	WorkerThread.SendLoadManifestMessage(InManifestURL, FString());
 }
 
+//-----------------------------------------------------------------------------
+/**
+ * Returns the pointer to the current media samples. Because this will get
+ * dereferenced and used _as a reference_(!!) it must never be a nullptr.
+ */
+TSharedPtr<FMediaSamples, ESPMode::ThreadSafe> FAdaptiveStreamingPlayer::GetMediaSamples()
+{
+	TSharedPtr<FMediaSamples, ESPMode::ThreadSafe> ms(CurrentMediaSamples);
+	return ms.IsValid() ? ms : EmptyMediaSamples;
+}
+
 
 
 //-----------------------------------------------------------------------------
@@ -1388,9 +1430,10 @@ void FAdaptiveStreamingPlayer::InternalHandleOnce()
 
 		// Handle state changes to match the user request.
 		HandlePlayStateChanges();
-
-		// Update diag buffers
+		// Update diag buffers.
 		UpdateDiagnostics();
+		// Update frame statistics.
+		UpdateFrameStats(false);
 		// Handle seek requests.
 		HandleSeeking();
 		// Check for end of stream.
@@ -4284,6 +4327,7 @@ void FAdaptiveStreamingPlayer::InternalStartAt(const FSeekParam& NewPosition, co
 	VideoBufferStats.Clear();
 	AudioBufferStats.Clear();
 	TextBufferStats.Clear();
+	DroppedFrameStats.Clear();
 
 #if NOTIFY_DATA_AVAILABILITY_AT_START_AND_END
 	// Update data availability states in case this wasn't done yet.
@@ -4354,14 +4398,6 @@ void FAdaptiveStreamingPlayer::InternalStartAt(const FSeekParam& NewPosition, co
 
 	// From this point on any further start will not be the very first one any more.
 	SeekVars.bIsPlayStart = false;
-
-	// Starting with a clean slate. Fire this event synchronously directly to the receivers.
-	// This is intended to release any upper layer blocks of passing decoded output to sinks
-	// that may be in place to prevent receiving decoder output that may be "stale" from the
-	// perspective of the upper layer.
-	// It is IMPERATIVE for the receiver of this event to do NO BLOCKING tasks OR call in
-	// any of the player API!
-	FireSyncEvent(FMetricEvent::ReportCleanStart());
 }
 
 
@@ -4398,6 +4434,8 @@ void FAdaptiveStreamingPlayer::InternalSetPlaybackEnded()
 
 	// Go to pause mode now.
 	InternalPause();
+	// Report the frame stats.
+	UpdateFrameStats(true);
 	DispatchEvent(FMetricEvent::ReportPlaybackEnded());
 	// In case a seek back into the stream will happen we reset the first start state.
 	PrerollVars.bIsVeryFirstStart = true;
@@ -4625,6 +4663,8 @@ void FAdaptiveStreamingPlayer::InternalInitialize()
 		}
 	}
 
+	// Set the interval to aggregate (dropped) frame statistics and report them.
+	DroppedFrameStats.ReportingIntervalMs = GetOptionValue(OptionKeyFrameStatInterval).SafeGetInt64(4000);
 
 	// Initialize decoder limits.
 	// NOTE: This is to be deprecated as soon as possible.
@@ -4640,6 +4680,7 @@ void FAdaptiveStreamingPlayer::InternalClose()
 
 	DestroyDecoders();
 	DestroyRenderers();
+	CurrentMediaSamples.Reset();
 
 	// Do not clear the playback state to allow GetPlayPosition() to query the last play position.
 	// This also means queries like IsBuffering() will return the last state but calling those
@@ -4696,6 +4737,7 @@ void FAdaptiveStreamingPlayer::InternalClose()
 	AudioBufferStats.Clear();
 	TextBufferStats.Clear();
 	DiagnosticsCriticalSection.Unlock();
+	DroppedFrameStats.Clear();
 }
 
 
diff --git a/Engine/Plugins/Media/ElectraPlayer/Source/ElectraPlayerRuntime/Private/Runtime/Player/AdaptiveStreamingPlayer.h b/Engine/Plugins/Media/ElectraPlayer/Source/ElectraPlayerRuntime/Private/Runtime/Player/AdaptiveStreamingPlayer.h
--- a/Engine/Plugins/Media/ElectraPlayer/Source/ElectraPlayerRuntime/Private/Runtime/Player/AdaptiveStreamingPlayer.h
+++ b/Engine/Plugins/Media/ElectraPlayer/Source/ElectraPlayerRuntime/Private/Runtime/Player/AdaptiveStreamingPlayer.h
@@ -11,26 +11,25 @@
 #include "StreamTypes.h"
 #include "AdaptiveStreamingPlayerMetrics.h"
 #include "AdaptiveStreamingPlayerEvents.h"
-#include "AdaptiveStreamingPlayerSubtitles.h"
 #include "AdaptiveStreamingPlayerResourceRequest.h"
-#include "IElectraPlayerDataCache.h"
+
+class FMediaSamples;
+class FElectraTextureSamplePool;
+class FElectraAudioSamplePool;
 
 namespace Electra
 {
-class FOutputHandlerVideo;
-class FOutputHandlerAudio;
-
 
 /**
- *
+ * This class declares the adaptive streaming player.
 **/
 class IAdaptiveStreamingPlayer : public TMediaNoncopyable<IAdaptiveStreamingPlayer>
 {
 public:
 	struct FCreateParam
 	{
-		TSharedPtr<FOutputHandlerVideo, ESPMode::ThreadSafe> VideoOutputHandler;
-		TSharedPtr<FOutputHandlerAudio, ESPMode::ThreadSafe> AudioOutputHandler;
+		TSharedPtr<FElectraTextureSamplePool, ESPMode::ThreadSafe> OutputTexturePool;
+		TSharedPtr<FElectraAudioSamplePool, ESPMode::ThreadSafe> OutputAudioSamplePool;
 
 		FGuid ExternalPlayerGUID;
 		enum class EWorkerThreads
@@ -65,11 +64,6 @@ public:
 	virtual void AddAEMSReceiver(TWeakPtrTS<IAdaptiveStreamingPlayerAEMSReceiver> InReceiver, FString InForSchemeIdUri, FString InForValue, IAdaptiveStreamingPlayerAEMSReceiver::EDispatchMode InDispatchMode) = 0;
 	virtual void RemoveAEMSReceiver(TWeakPtrTS<IAdaptiveStreamingPlayerAEMSReceiver> InReceiver, FString InForSchemeIdUri, FString InForValue, IAdaptiveStreamingPlayerAEMSReceiver::EDispatchMode InDispatchMode) = 0;
 
-	//-------------------------------------------------------------------------
-	// Subtitle receiver
-	//
-	virtual void AddSubtitleReceiver(TWeakPtrTS<IAdaptiveStreamingPlayerSubtitleReceiver> InReceiver) = 0;
-	virtual void RemoveSubtitleReceiver(TWeakPtrTS<IAdaptiveStreamingPlayerSubtitleReceiver> InReceiver) = 0;
 
 	//-------------------------------------------------------------------------
 	// Initializes the player. Options may be passed in to affect behaviour.
@@ -111,6 +105,12 @@ public:
 	virtual void EnableFrameAccurateSeeking(bool bEnabled) = 0;
 
 
+	//-------------------------------------------------------------------------
+	// Interfacing functions
+	//
+	virtual TSharedPtr<FMediaSamples, ESPMode::ThreadSafe> GetMediaSamples() = 0;
+
+
 	//-------------------------------------------------------------------------
 	// Manifest loading functions
 	//
@@ -367,4 +367,3 @@ protected:
 
 
 } // namespace Electra
-
diff --git a/Engine/Plugins/Media/ElectraPlayer/Source/ElectraPlayerRuntime/Private/Runtime/Player/AdaptiveStreamingPlayerABR.cpp b/Engine/Plugins/Media/ElectraPlayer/Source/ElectraPlayerRuntime/Private/Runtime/Player/AdaptiveStreamingPlayerABR.cpp
--- a/Engine/Plugins/Media/ElectraPlayer/Source/ElectraPlayerRuntime/Private/Runtime/Player/AdaptiveStreamingPlayerABR.cpp
+++ b/Engine/Plugins/Media/ElectraPlayer/Source/ElectraPlayerRuntime/Private/Runtime/Player/AdaptiveStreamingPlayerABR.cpp
@@ -53,7 +53,6 @@ namespace Electra
 		void ReportReceivedMainPlaylist(const FString& EffectiveURL) override {}
 		void ReportReceivedPlaylists() override {}
 		void ReportTracksChanged() override {}
-		void ReportCleanStart() override {};
 		void ReportPlaylistDownload(const Metrics::FPlaylistDownloadStats& PlaylistDownloadStats) override {}
 		void ReportBandwidth(int64 EffectiveBps, int64 ThroughputBps, double LatencyInSeconds) override {}
 		void ReportBufferUtilization(const Metrics::FBufferStats& BufferStats) override {}
@@ -75,8 +74,8 @@ namespace Electra
 		void ReportMediaMetadataChanged(TSharedPtrTS<UtilsMP4::FMetadataParser> Metadata) override {}
 		void ReportError(const FString& ErrorReason) override {}
 		void ReportLogMessage(IInfoLog::ELevel InLogLevel, const FString& LogMessage, int64 PlayerWallclockMilliseconds) override {}
-		void ReportDroppedVideoFrame() override {}
-		void ReportDroppedAudioFrame() override {}
+		void ReportVideoFrameStats(const Metrics::FFrameStats& InFrameStats) override {}
+		void ReportAudioFrameStats(const Metrics::FFrameStats& InFrameStats) override {}
 
 	private:
 		bool HaveOptionValue(const FName& InOption) override
diff --git a/Engine/Plugins/Media/ElectraPlayer/Source/ElectraPlayerRuntime/Private/Runtime/Player/AdaptiveStreamingPlayerDecoders.cpp b/Engine/Plugins/Media/ElectraPlayer/Source/ElectraPlayerRuntime/Private/Runtime/Player/AdaptiveStreamingPlayerDecoders.cpp
--- a/Engine/Plugins/Media/ElectraPlayer/Source/ElectraPlayerRuntime/Private/Runtime/Player/AdaptiveStreamingPlayerDecoders.cpp
+++ b/Engine/Plugins/Media/ElectraPlayer/Source/ElectraPlayerRuntime/Private/Runtime/Player/AdaptiveStreamingPlayerDecoders.cpp
@@ -373,18 +373,16 @@ int32 FAdaptiveStreamingPlayer::CreateDecoder(EStreamType type)
 					SubtitleDecoder.LastSentAUCodecData = AccessUnit->AUCodecData;
 					DispatchEvent(FMetricEvent::ReportCodecFormatChange(SubtitleDecoder.CurrentCodecInfo));
 				}
-				FAccessUnit::Release(AccessUnit);
-
 				// Create subtitle decoder for this type of AU.
 				SubtitleDecoder.Decoder = ISubtitleDecoder::Create(AccessUnit);
+				FAccessUnit::Release(AccessUnit);
 				if (SubtitleDecoder.Decoder)
 				{
 					FParamDict Options;
 					Options.Set(DecoderOptionKeys::SendEmptySubtitleDuringGaps, Electra::FVariantValue(true));
 
 					SubtitleDecoder.Decoder->SetPlayerSessionServices(this);
-					SubtitleDecoder.Decoder->GetDecodedSubtitleReceiveDelegate().BindRaw(this, &FAdaptiveStreamingPlayer::OnDecodedSubtitleReceived);
-					SubtitleDecoder.Decoder->GetDecodedSubtitleFlushDelegate().BindRaw(this, &FAdaptiveStreamingPlayer::OnFlushSubtitleReceivers);
+					SubtitleDecoder.Decoder->SetMediaSamples(CurrentMediaSamples);
 					SubtitleDecoder.Decoder->Open(Options);
 					// Start the decoder, if decoding is currently running.
 					if (DecoderState == EDecoderState::eDecoder_Running)
diff --git a/Engine/Plugins/Media/ElectraPlayer/Source/ElectraPlayerRuntime/Private/Runtime/Player/AdaptiveStreamingPlayerInternal.h b/Engine/Plugins/Media/ElectraPlayer/Source/ElectraPlayerRuntime/Private/Runtime/Player/AdaptiveStreamingPlayerInternal.h
--- a/Engine/Plugins/Media/ElectraPlayer/Source/ElectraPlayerRuntime/Private/Runtime/Player/AdaptiveStreamingPlayerInternal.h
+++ b/Engine/Plugins/Media/ElectraPlayer/Source/ElectraPlayerRuntime/Private/Runtime/Player/AdaptiveStreamingPlayerInternal.h
@@ -25,22 +25,20 @@
 #include "Player/AdaptivePlayerOptionKeynames.h"
 #include "Player/ContentSteeringHandler.h"
 #include "Player/CMCDHandler.h"
+#include "Player/ABRRules/ABRStatisticTypes.h"
 #include "OutputHandler.h"
-
-#include "Templates/UniqueObj.h"
-
 #include "Utilities/UtilsMP4.h"
-
 #include "ElectraCDM.h"
-
 #include "InfoLog.h"
 
+#include "Templates/UniqueObj.h"
 #include "Misc/TVariant.h"
 #include "Dom/JsonObject.h"
 #include "Serialization/JsonReader.h"
 #include "Serialization/JsonSerializer.h"
 #include "Policies/CondensedJsonPrintPolicy.h"
 #include "Misc/Base64.h"
+#include "MediaSamples.h"
 
 
 
@@ -182,6 +180,31 @@ private:
 
 
 
+struct FDroppedFrameStats
+{
+	struct FEntry
+	{
+		uint64 DropCount;
+	};
+	void Clear()
+	{
+		Entries.Empty();
+		LastProcessedTime = -1;
+		PausedDelta = -1;
+		NumIntervals = 0;
+		NumIntervalsWithDrops = 0;
+	}
+	TArray<FEntry> Entries;
+	TRollingAverage<int64> AverageDropsPerInterval;
+	int64 LastProcessedTime = -1;
+	int64 PausedDelta = -1;
+	int64 ReportingIntervalMs = -1;
+	uint32 NumIntervals = 0;
+	uint32 NumIntervalsWithDrops = 0;
+};
+
+
+
 /**
  * Interchange structure between the player worker thread and the public API to
  * avoid mutex locks all over the place.
@@ -634,7 +657,6 @@ struct FMetricEvent
 		ReceivedMainPlaylist,
 		ReceivedPlaylists,
 		TracksChanged,
-		CleanStart,
 		BufferingStart,
 		BufferingEnd,
 		Bandwidth,
@@ -658,6 +680,7 @@ struct FMetricEvent
 		LicenseKey,
 		Errored,
 		LogMessage,
+		VideoFrameStats
 	};
 
 	struct FParam
@@ -702,6 +725,7 @@ struct FMetricEvent
 		Metrics::FSegmentDownloadStats		SegmentStats;
 		Metrics::FLicenseKeyStats			LicenseKeyStats;
 		Metrics::FDataAvailabilityChange	DataAvailability;
+		Metrics::FFrameStats				FrameStats;
 		FBandwidth							Bandwidth;
 		FQualityChange						QualityChange;
 		FStreamCodecInformation				CodecFormatChange;
@@ -717,12 +741,6 @@ struct FMetricEvent
 	TWeakPtr<FAdaptiveStreamingPlayer, ESPMode::ThreadSafe>	Player;
 	FMediaEvent	*EventSignal = nullptr;
 
-	static TSharedPtrTS<FMetricEvent> ReportCleanStart()
-	{
-		TSharedPtrTS<FMetricEvent> Evt = MakeSharedTS<FMetricEvent>();
-		Evt->Type = EType::CleanStart;
-		return Evt;
-	}
 	static TSharedPtrTS<FMetricEvent> ReportOpenSource(const FString& URL)
 	{
 		TSharedPtrTS<FMetricEvent> Evt = MakeSharedTS<FMetricEvent>();
@@ -913,6 +931,21 @@ struct FMetricEvent
 		Evt->Param.LogMessage.AtMillis = PlayerWallclockMilliseconds;
 		return Evt;
 	}
+	static TSharedPtrTS<FMetricEvent> ReportVideoFrameStats(uint64 InNumTotalFrames, uint64 InNumTotalDroppedFrames, uint64 InNumNewDroppedFrames, const FDroppedFrameStats& InDropStats)
+	{
+		TSharedPtrTS<FMetricEvent> Evt = MakeSharedTS<FMetricEvent>();
+		Evt->Type = EType::VideoFrameStats;
+		Evt->Param.FrameStats.NumTotal = InNumTotalFrames;
+		Evt->Param.FrameStats.NumDropped = InNumTotalDroppedFrames;
+		Evt->Param.FrameStats.NumNewDrops = (uint32) InNumNewDroppedFrames;
+		Evt->Param.FrameStats.AvgDropsPerInterval = (uint32)InDropStats.AverageDropsPerInterval.GetAverage();
+		Evt->Param.FrameStats.MinDropsPerInterval = (uint32)InDropStats.AverageDropsPerInterval.GetMin();
+		Evt->Param.FrameStats.MaxDropsPerInterval = (uint32)InDropStats.AverageDropsPerInterval.GetMax();
+		Evt->Param.FrameStats.DropIntervalMs = (uint32)(InDropStats.ReportingIntervalMs > 0 ? InDropStats.ReportingIntervalMs : 0);
+		Evt->Param.FrameStats.NumIntervals = InDropStats.NumIntervals;
+		Evt->Param.FrameStats.NumIntervalsWithDrops = InDropStats.NumIntervalsWithDrops;
+		return Evt;
+	}
 
 };
 
@@ -1112,15 +1145,13 @@ public:
 	void AddAEMSReceiver(TWeakPtrTS<IAdaptiveStreamingPlayerAEMSReceiver> InReceiver, FString InForSchemeIdUri, FString InForValue, IAdaptiveStreamingPlayerAEMSReceiver::EDispatchMode InDispatchMode) override;
 	void RemoveAEMSReceiver(TWeakPtrTS<IAdaptiveStreamingPlayerAEMSReceiver> InReceiver, FString InForSchemeIdUri, FString InForValue, IAdaptiveStreamingPlayerAEMSReceiver::EDispatchMode InDispatchMode) override;
 
-	void AddSubtitleReceiver(TWeakPtrTS<IAdaptiveStreamingPlayerSubtitleReceiver> InReceiver) override;
-	void RemoveSubtitleReceiver(TWeakPtrTS<IAdaptiveStreamingPlayerSubtitleReceiver> InReceiver) override;
-
 	void Initialize(const FParamDict& Options) override;
 	void ModifyOptions(const FParamDict& InOptionsToSetOrChange, const FParamDict& InOptionsToClear) override;
 	FVariantValue GetMediaInfo(FName InKey) const override;
 
 	void SetInitialStreamAttributes(EStreamType StreamType, FStreamSelectionAttributes InitialSelection) override;
 	void EnableFrameAccurateSeeking(bool bEnabled) override;
+	TSharedPtr<FMediaSamples, ESPMode::ThreadSafe> GetMediaSamples() override;
 	void LoadManifest(const FString& manifestURL) override;
 
 	void SeekTo(const FSeekParam& NewPosition) override;
@@ -1307,7 +1338,11 @@ private:
 	{
 		void Close()
 		{
-			OutputHandler.Reset();
+			if (OutputHandler.IsValid())
+			{
+				OutputHandler->SetMediaSamples(nullptr);
+				OutputHandler.Reset();
+			}
 		}
 		void Flush(bool bHoldCurrentFrame)
 		{
@@ -1316,6 +1351,22 @@ private:
 				OutputHandler->Flush();
 			}
 		}
+		uint64 GetTotalAddedSamples()
+		{
+			return OutputHandler.IsValid() ? OutputHandler->GetTotalAddedSamples() : (uint64)0;
+		}
+		uint64 GetTotalNumDroppedSamples()
+		{
+			return OutputHandler.IsValid() ? OutputHandler->GetTotalNumDroppedSamples() : (uint64)0;
+		}
+		uint64 GetNewlyNumDroppedSamples()
+		{
+			return OutputHandler.IsValid() ? OutputHandler->GetNewlyNumDroppedSamples(true) : (uint64)0;
+		}
+		uint64 PeekNewlyNumDroppedSamples()
+		{
+			return OutputHandler.IsValid() ? OutputHandler->GetNewlyNumDroppedSamples(false) : (uint64)0;
+		}
 		TSharedPtr<FOutputHandlerVideo, ESPMode::ThreadSafe> OutputHandler;
 	};
 
@@ -1323,7 +1374,11 @@ private:
 	{
 		void Close()
 		{
-			OutputHandler.Reset();
+			if (OutputHandler.IsValid())
+			{
+				OutputHandler->SetMediaSamples(nullptr);
+				OutputHandler.Reset();
+			}
 		}
 		void Flush()
 		{
@@ -1477,8 +1532,6 @@ private:
 			{
 				Stop();
 				Decoder->Close();
-				Decoder->GetDecodedSubtitleReceiveDelegate().Unbind();
-				Decoder->GetDecodedSubtitleFlushDelegate().Unbind();
 				delete Decoder;
 				Decoder = nullptr;
 				bIsRunning = false;
@@ -2104,9 +2157,6 @@ private:
 	void AudioDecoderInputNeeded(const IAccessUnitBufferListener::FBufferStats& currentInputBufferStats);
 	void AudioDecoderOutputReady(const IDecoderOutputBufferListener::FDecodeReadyStats& currentReadyStats);
 
-	void OnDecodedSubtitleReceived(ISubtitleDecoderOutputPtr Subtitle);
-	void OnFlushSubtitleReceivers();
-
 	void StartWorkerThread();
 	void StopWorkerThread();
 
@@ -2148,6 +2198,7 @@ private:
 	void InternalStartoverAtLiveEdge();
 
 	void UpdateDiagnostics();
+	void UpdateFrameStats(bool bFinal);
 	void HandleNewBufferedData();
 	void HandleNewOutputData();
 	void HandleSessionMessage(TSharedPtrTS<IPlayerMessage> SessionMessage);
@@ -2237,6 +2288,8 @@ private:
 	//
 	// Member variables
 	//
+	TSharedPtr<FMediaSamples, ESPMode::ThreadSafe>						CurrentMediaSamples;
+	TSharedPtr<FMediaSamples, ESPMode::ThreadSafe>						EmptyMediaSamples;
 	FGuid																ExternalPlayerGUID;
 
 	TWeakPtr<IAdaptiveStreamingPlayerResourceProvider, ESPMode::ThreadSafe> StaticResourceProvider;
@@ -2363,9 +2416,6 @@ private:
 	FVideoDecoder														VideoDecoder;
 	FAudioDecoder														AudioDecoder;
 	FSubtitleDecoder													SubtitleDecoder;
-	FCriticalSection													SubtitleReceiversCriticalSection;
-	TArray<TWeakPtrTS<IAdaptiveStreamingPlayerSubtitleReceiver>>		SubtitleReceivers;
-
 
 	FCriticalSection													MetricListenerCriticalSection;
 	TArray<IAdaptiveStreamingPlayerMetrics*, TInlineAllocator<4>>		MetricListeners;
@@ -2376,6 +2426,8 @@ private:
 	FBufferStats														TextBufferStats;
 	FErrorDetail														LastErrorDetail;
 
+	FDroppedFrameStats													DroppedFrameStats;
+
 	AdaptiveStreamingPlayerConfig::FConfiguration						PlayerConfig;
 
 	static FAdaptiveStreamingPlayer* 									PointerToLatestPlayer;
diff --git a/Engine/Plugins/Media/ElectraPlayer/Source/ElectraPlayerRuntime/Private/Runtime/Player/AdaptiveStreamingPlayerMetrics.h b/Engine/Plugins/Media/ElectraPlayer/Source/ElectraPlayerRuntime/Private/Runtime/Player/AdaptiveStreamingPlayerMetrics.h
--- a/Engine/Plugins/Media/ElectraPlayer/Source/ElectraPlayerRuntime/Private/Runtime/Player/AdaptiveStreamingPlayerMetrics.h
+++ b/Engine/Plugins/Media/ElectraPlayer/Source/ElectraPlayerRuntime/Private/Runtime/Player/AdaptiveStreamingPlayerMetrics.h
@@ -249,6 +249,19 @@ namespace Metrics
 		EAvailability	Availability;
 	};
 
+	struct FFrameStats
+	{
+		uint64 NumTotal = 0;
+		uint64 NumDropped = 0;
+		uint32 NumNewDrops = 0;
+		uint32 AvgDropsPerInterval = 0;
+		uint32 MinDropsPerInterval = 0;
+		uint32 MaxDropsPerInterval = 0;
+		uint32 DropIntervalMs = 0;
+		uint32 NumIntervals = 0;
+		uint32 NumIntervalsWithDrops = 0;
+	};
+
 } // namespace Metrics
 
 
@@ -289,11 +302,6 @@ public:
 	 */
 	virtual void ReportPlaylistDownload(const Metrics::FPlaylistDownloadStats& PlaylistDownloadStats) = 0;
 
-	/**
-	 * Called when the player starts over from a clean state.
-	 */
-	virtual void ReportCleanStart() = 0;
-
 	/**
 	 * Called when buffering of data begins.
 	 */
@@ -418,8 +426,8 @@ public:
 	// Methods called from the renderers.
 	//
 
-	virtual void ReportDroppedVideoFrame() = 0;
-	virtual void ReportDroppedAudioFrame() = 0;
+	virtual void ReportVideoFrameStats(const Metrics::FFrameStats& InFrameStats) = 0;
+	virtual void ReportAudioFrameStats(const Metrics::FFrameStats& InFrameStats) = 0;
 };
 
 
diff --git a/Engine/Plugins/Media/ElectraPlayer/Source/ElectraPlayerRuntime/Private/Runtime/Player/AdaptiveStreamingPlayerSubtitles.h b/Engine/Plugins/Media/ElectraPlayer/Source/ElectraPlayerRuntime/Private/Runtime/Player/AdaptiveStreamingPlayerSubtitles.h
deleted file mode 100644
diff --git a/Engine/Plugins/Media/ElectraPlayer/Source/ElectraPlayerRuntime/Private/Runtime/Player/CMCDHandler.cpp b/Engine/Plugins/Media/ElectraPlayer/Source/ElectraPlayerRuntime/Private/Runtime/Player/CMCDHandler.cpp
--- a/Engine/Plugins/Media/ElectraPlayer/Source/ElectraPlayerRuntime/Private/Runtime/Player/CMCDHandler.cpp
+++ b/Engine/Plugins/Media/ElectraPlayer/Source/ElectraPlayerRuntime/Private/Runtime/Player/CMCDHandler.cpp
@@ -885,7 +885,6 @@ void FCMCDHandler::ReportReceivedMainPlaylist(const FString& InEffectiveURL) { }
 void FCMCDHandler::ReportReceivedPlaylists() { }
 void FCMCDHandler::ReportTracksChanged() { }
 void FCMCDHandler::ReportPlaylistDownload(const Metrics::FPlaylistDownloadStats& InPlaylistDownloadStats) { }
-void FCMCDHandler::ReportCleanStart() { }
 void FCMCDHandler::ReportBufferingStart(Metrics::EBufferingReason InBufferingReason)
 {
 	if (InBufferingReason == Metrics::EBufferingReason::Seeking)
@@ -930,7 +929,7 @@ void FCMCDHandler::ReportSeekCompleted() { }
 void FCMCDHandler::ReportMediaMetadataChanged(TSharedPtrTS<UtilsMP4::FMetadataParser> InMetadata) { }
 void FCMCDHandler::ReportError(const FString& InErrorReason) { }
 void FCMCDHandler::ReportLogMessage(IInfoLog::ELevel InLogLevel, const FString& InLogMessage, int64 InPlayerWallclockMilliseconds) { }
-void FCMCDHandler::ReportDroppedVideoFrame() { }
-void FCMCDHandler::ReportDroppedAudioFrame() { }
+void FCMCDHandler::ReportVideoFrameStats(const Metrics::FFrameStats& InFrameStats) { }
+void FCMCDHandler::ReportAudioFrameStats(const Metrics::FFrameStats& InFrameStats) { }
 
 } // namespace Electra
diff --git a/Engine/Plugins/Media/ElectraPlayer/Source/ElectraPlayerRuntime/Private/Runtime/Player/CMCDHandler.h b/Engine/Plugins/Media/ElectraPlayer/Source/ElectraPlayerRuntime/Private/Runtime/Player/CMCDHandler.h
--- a/Engine/Plugins/Media/ElectraPlayer/Source/ElectraPlayerRuntime/Private/Runtime/Player/CMCDHandler.h
+++ b/Engine/Plugins/Media/ElectraPlayer/Source/ElectraPlayerRuntime/Private/Runtime/Player/CMCDHandler.h
@@ -254,7 +254,6 @@ private:
 	void ReportReceivedPlaylists() override;
 	void ReportTracksChanged() override;
 	void ReportPlaylistDownload(const Metrics::FPlaylistDownloadStats& InPlaylistDownloadStats) override;
-	void ReportCleanStart() override;
 	void ReportBufferingStart(Metrics::EBufferingReason InBufferingReason) override;
 	void ReportBufferingEnd(Metrics::EBufferingReason InBufferingReason) override;
 	void ReportBandwidth(int64 InEffectiveBps, int64 InThroughputBps, double InLatencyInSeconds) override;
@@ -277,8 +276,8 @@ private:
 	void ReportMediaMetadataChanged(TSharedPtrTS<UtilsMP4::FMetadataParser> InMetadata) override;
 	void ReportError(const FString& InErrorReason) override;
 	void ReportLogMessage(IInfoLog::ELevel InLogLevel, const FString& InLogMessage, int64 InPlayerWallclockMilliseconds) override;
-	void ReportDroppedVideoFrame() override;
-	void ReportDroppedAudioFrame() override;
+	void ReportVideoFrameStats(const Metrics::FFrameStats& InFrameStats) override;
+	void ReportAudioFrameStats(const Metrics::FFrameStats& InFrameStats) override;
 
 	enum class EV1Keys : uint8
 	{
diff --git a/Engine/Plugins/Media/ElectraUtil/Source/ElectraSamples/Public/IElectraSubtitleSample.h b/Engine/Plugins/Media/ElectraUtil/Source/ElectraSamples/Public/IElectraSubtitleSample.h
--- a/Engine/Plugins/Media/ElectraUtil/Source/ElectraSamples/Public/IElectraSubtitleSample.h
+++ b/Engine/Plugins/Media/ElectraUtil/Source/ElectraSamples/Public/IElectraSubtitleSample.h
@@ -10,24 +10,26 @@
  * Extension of the media overlay sample to carry Electra
  * specific values.
  */
-class IElectraSubtitleSample : public IMediaOverlaySample
+class FElectraSubtitleSample : public IMediaOverlaySample
 {
 public:
-	static FGuid GetSampleTypeGUID()
-	{
-		static FGuid SampleTypeGUID(0xC94C9B6F, 0x07DF4B88, 0xACE0B08A, 0x9EF6AA39);
-		return SampleTypeGUID;
-	}
+	virtual ~FElectraSubtitleSample() = default;
 
-	virtual ~IElectraSubtitleSample() = default;
-
-	virtual FTimespan GetDuration() const = 0;
-	virtual TOptional<FVector2D> GetPosition() const = 0;
-	virtual FText GetText() const = 0;
-	virtual FMediaTimeStamp GetTime() const = 0;
-	//virtual TOptional<FTimecode> GetTimecode() const { return TOptional<FTimecode>(); }
-	virtual EMediaOverlaySampleType GetType() const = 0;
-
-	virtual FGuid GetGUID() const = 0;
+	FGuid GetGUID() const override						
+	{ static FGuid SampleTypeGUID(0xC94C9B6F, 0x07DF4B88, 0xACE0B08A, 0x9EF6AA39); return SampleTypeGUID; }
+	FMediaTimeStamp GetTime() const override			
+	{ return Timestamp; }
+	FTimespan GetDuration() const override				
+	{ return Duration; }
+	TOptional<FVector2D> GetPosition() const override
+	{ return TOptional<FVector2D>(); }
+	EMediaOverlaySampleType GetType() const override	
+	{ return EMediaOverlaySampleType::Subtitle; }
+	FText GetText() const override
+	{ return FText::FromString(Text); }
 
+	FString Text;
+	FMediaTimeStamp Timestamp;
+	FTimespan Duration;
 };
+
