ElectraPlayer: fixed event mode when unbinding

#rnx

diff --git a/Engine/Plugins/Media/ElectraPlayer/Source/ElectraPlayerRuntime/Private/ElectraPlayer.cpp b/Engine/Plugins/Media/ElectraPlayer/Source/ElectraPlayerRuntime/Private/ElectraPlayer.cpp
--- a/Engine/Plugins/Media/ElectraPlayer/Source/ElectraPlayerRuntime/Private/ElectraPlayer.cpp
+++ b/Engine/Plugins/Media/ElectraPlayer/Source/ElectraPlayerRuntime/Private/ElectraPlayer.cpp
@@ -740,7 +740,7 @@ void FElectraPlayer::CloseInternal()
 		if (MediaPlayerEventReceiver.IsValid())
 		{
 			MediaPlayerEventReceiver->GetEventReceivedDelegate().Unbind();
-			Player->AdaptivePlayer->RemoveAEMSReceiver(MediaPlayerEventReceiver, TEXT("*"), TEXT(""), IAdaptiveStreamingPlayerAEMSReceiver::EDispatchMode::OnStart);
+			Player->AdaptivePlayer->RemoveAEMSReceiver(MediaPlayerEventReceiver, TEXT("*"), TEXT(""), IAdaptiveStreamingPlayerAEMSReceiver::EDispatchMode::OnReceive);
 			MediaPlayerEventReceiver.Reset();
 		}
 		Player->AdaptivePlayer->SetStaticResourceProviderCallback(nullptr);
