ElectraPlayer: Releasing video and audio pools when closing the player instead of only in the destructor. Player instances are not reused all that often so it is better to release all pool memory right away instead of waiting for garbage collection. This can be controlled via a cvar if the old behavior is required.

#rnx

diff --git a/Engine/Plugins/Media/ElectraPlayer/Source/ElectraPlayerRuntime/Private/ElectraPlayer.cpp b/Engine/Plugins/Media/ElectraPlayer/Source/ElectraPlayerRuntime/Private/ElectraPlayer.cpp
--- a/Engine/Plugins/Media/ElectraPlayer/Source/ElectraPlayerRuntime/Private/ElectraPlayer.cpp
+++ b/Engine/Plugins/Media/ElectraPlayer/Source/ElectraPlayerRuntime/Private/ElectraPlayer.cpp
@@ -35,6 +35,17 @@ DECLARE_CYCLE_STAT(TEXT("FElectraPlayer::QueryOptions"), STAT_ElectraPlayer_Elec
 
 //-----------------------------------------------------------------------------
 
+#include "HAL/IConsoleManager.h"
+static TAutoConsoleVariable<bool> CVarElectraRetainPools(
+	TEXT("Electra.RetainPools"),
+	false,
+	TEXT("Retains video and audio sample pools in player for reuse.\n")
+	TEXT(" false: don't retain, create new pools for each instance; true: retain pools"),
+	ECVF_Default);
+
+//-----------------------------------------------------------------------------
+
+
 // Prefix to use in querying for a custom analytic value through QueryOptions()
 #define CUSTOM_ANALYTIC_METRIC_QUERYOPTION_KEY TEXT("ElectraCustomAnalytic")
 // Prefix to use in the metric event to set the custom value.
@@ -346,11 +357,15 @@ FElectraPlayer::FElectraPlayer(IMediaEventSink& InEventSink,
 {
 	CSV_EVENT(ElectraPlayer, TEXT("Player Creation"));
 
-	OutputAudioSamplePool = MakeShareable(new FElectraAudioSamplePool);
-	OutputTexturePool = MakeShareable(new FElectraTextureSamplePool);
-#if PLATFORM_ANDROID
-	OutputTexturePoolDecoderID = FElectraDecodersPlatformResources::RegisterOutputTexturePool(OutputTexturePool);
-#endif
+	bRetainPoolsForReuse = CVarElectraRetainPools.GetValueOnAnyThread();
+	if (bRetainPoolsForReuse)
+	{
+		OutputAudioSamplePool = MakeShareable(new FElectraAudioSamplePool);
+		OutputTexturePool = MakeShareable(new FElectraTextureSamplePool);
+	#if PLATFORM_ANDROID
+		OutputTexturePoolDecoderID = FElectraDecodersPlatformResources::RegisterOutputTexturePool(OutputTexturePool);
+	#endif
+	}
 
 	EmptyMediaSamples.Reset(new FMediaSamples);
 	CurrentPlaybackRange = TRange<FTimespan>::Empty();
@@ -389,9 +404,12 @@ FElectraPlayer::~FElectraPlayer()
 
 	CloseInternal();
 
-#if PLATFORM_ANDROID
-	FElectraDecodersPlatformResources::UnregisterOutputTexturePool(OutputTexturePoolDecoderID);
-#endif
+	if (bRetainPoolsForReuse)
+	{
+	#if PLATFORM_ANDROID
+		FElectraDecodersPlatformResources::UnregisterOutputTexturePool(OutputTexturePoolDecoderID);
+	#endif
+	}
 
 	SendAnalyticMetricsDelegate.RemoveAll(this);
 	SendAnalyticMetricsPerMinuteDelegate.RemoveAll(this);
@@ -476,7 +494,17 @@ bool FElectraPlayer::OpenInternal(const FString& InUrl, const IMediaOptions* InO
 		}
 		if (InPlayerOptions->SeekTimeType != EMediaPlayerOptionSeekTimeType::Ignored)
 		{
-			PlaystartOptions.TimeOffset = InPlayerOptions->SeekTime;
+			// If an explicit time is set we only use it if it is not zero. Zero is the default value and most of the time means nothing in particular.
+			// Playback will start at the beginning for VoD anyway and if this is an event using zero causes problems
+			if (InPlayerOptions->SeekTime > FTimespan::Zero())
+			{
+				PlaystartOptions.TimeOffset = InPlayerOptions->SeekTime;
+				UE_LOG(LogElectraPlayer, Log, TEXT("[%u] IMediaPlayer::Open: Setting initial seek time to %.3f seconds"), PlayerUniqueID, InPlayerOptions->SeekTime.GetTotalSeconds());
+			}
+			else
+			{
+				UE_LOG(LogElectraPlayer, Log, TEXT("[%u] IMediaPlayer::Open: Setting initial seek time of zero will be ignored as it implies the beginning."), PlayerUniqueID);
+			}
 		}
 		if (InPlayerOptions->TrackSelection == EMediaPlayerOptionTrackSelectMode::UseLanguageCodes)
 		{
@@ -514,7 +542,6 @@ bool FElectraPlayer::OpenInternal(const FString& InUrl, const IMediaOptions* InO
 
 	// Set up options to initialize the internal player with.
 	Electra::FParamDict PlayerOptions;
-	PlayerOptions.Set(Electra::OptionKeyOutputTexturePoolID, FVariantValue(OutputTexturePoolDecoderID));
 
 	// Set the session ID to be the same we will use for analytics here (see below)
 	PlayerOptions.Set(Electra::OptionKeySessionID, Electra::FVariantValue(SessionID.ToString(EGuidFormats::DigitsWithHyphensLower)));
@@ -652,8 +679,20 @@ bool FElectraPlayer::OpenInternal(const FString& InUrl, const IMediaOptions* InO
 
 	// Create the internal player and register ourselves as metrics receiver and static resource provider.
 	IAdaptiveStreamingPlayer::FCreateParam CreateParams;
-	CreateParams.OutputTexturePool = OutputTexturePool;
-	CreateParams.OutputAudioSamplePool = OutputAudioSamplePool;
+	if (!bRetainPoolsForReuse)
+	{
+		CreateParams.OutputTexturePool = MakeShareable(new FElectraTextureSamplePool);
+		CreateParams.OutputAudioSamplePool = MakeShareable(new FElectraAudioSamplePool);
+	#if PLATFORM_ANDROID
+		OutputTexturePoolDecoderID = FElectraDecodersPlatformResources::RegisterOutputTexturePool(CreateParams.OutputTexturePool);
+	#endif
+	}
+	else
+	{
+		CreateParams.OutputTexturePool = OutputTexturePool;
+		CreateParams.OutputAudioSamplePool = OutputAudioSamplePool;
+	}
+	PlayerOptions.Set(Electra::OptionKeyOutputTexturePoolID, FVariantValue(OutputTexturePoolDecoderID));
 	CreateParams.ExternalPlayerGUID = PlayerGuid;
 	FString WorkerThreadOption = PlayerOptions.GetValue(Electra::OptionKeyWorkerThreads).SafeGetFString(TEXT("shared"));
 	CreateParams.WorkerThreads = WorkerThreadOption.Equals(TEXT("worker"), ESearchCase::IgnoreCase) ? IAdaptiveStreamingPlayer::FCreateParam::EWorkerThreads::DedicatedWorker :
@@ -760,16 +799,24 @@ void FElectraPlayer::CloseInternal()
 	CurrentPlayer.Reset();
 
 	// Kick off asynchronous closing now.
-	FInternalPlayerImpl::DoCloseAsync(MoveTemp(Player), PlayerUniqueID, AsyncResourceReleaseNotification);
+	FInternalPlayerImpl::DoCloseAsync(MoveTemp(Player), PlayerUniqueID, AsyncResourceReleaseNotification, bRetainPoolsForReuse==false, OutputTexturePoolDecoderID);
 }
 
-void FElectraPlayer::FInternalPlayerImpl::DoCloseAsync(TSharedPtr<FInternalPlayerImpl, ESPMode::ThreadSafe>&& InPlayer, uint32 InPlayerID, TSharedPtr<IAsyncResourceReleaseNotification, ESPMode::ThreadSafe> InAsyncResourceReleaseNotification)
+void FElectraPlayer::FInternalPlayerImpl::DoCloseAsync(TSharedPtr<FInternalPlayerImpl, ESPMode::ThreadSafe>&& InPlayer, uint32 InPlayerID, TSharedPtr<IAsyncResourceReleaseNotification, ESPMode::ThreadSafe> InAsyncResourceReleaseNotification, bool bReleasePools, int64 TexturePoolDecoderID)
 {
-	TFunction<void()> CloseTask = [InPlayer, InPlayerID, InAsyncResourceReleaseNotification]()
+	TFunction<void()> CloseTask = [InPlayer, InPlayerID, InAsyncResourceReleaseNotification, bReleasePools, TexturePoolDecoderID]()
 	{
 		double TimeCloseBegan = FPlatformTime::Seconds();
 		InPlayer->AdaptivePlayer->Stop();
 		InPlayer->AdaptivePlayer.Reset();
+
+		if (bReleasePools)
+		{
+		#if PLATFORM_ANDROID
+			FElectraDecodersPlatformResources::UnregisterOutputTexturePool(TexturePoolDecoderID);
+		#endif
+		}
+
 		double TimeCloseEnded = FPlatformTime::Seconds();
 		UE_LOG(LogElectraPlayer, Verbose, TEXT("[%u] DoCloseAsync() finished after %.3f msec!"), InPlayerID, (TimeCloseEnded - TimeCloseBegan) * 1000.0);
 		if (InAsyncResourceReleaseNotification.IsValid())
diff --git a/Engine/Plugins/Media/ElectraPlayer/Source/ElectraPlayerRuntime/Private/ElectraPlayer.h b/Engine/Plugins/Media/ElectraPlayer/Source/ElectraPlayerRuntime/Private/ElectraPlayer.h
--- a/Engine/Plugins/Media/ElectraPlayer/Source/ElectraPlayerRuntime/Private/ElectraPlayer.h
+++ b/Engine/Plugins/Media/ElectraPlayer/Source/ElectraPlayerRuntime/Private/ElectraPlayer.h
@@ -269,7 +269,7 @@ private:
 		/** The media player itself **/
 		TSharedPtr<IAdaptiveStreamingPlayer, ESPMode::ThreadSafe> AdaptivePlayer;
 
-		static void DoCloseAsync(TSharedPtr<FInternalPlayerImpl, ESPMode::ThreadSafe>&& InPlayer, uint32 InPlayerID, TSharedPtr<IAsyncResourceReleaseNotification, ESPMode::ThreadSafe> InAsyncResourceReleaseNotification);
+		static void DoCloseAsync(TSharedPtr<FInternalPlayerImpl, ESPMode::ThreadSafe>&& InPlayer, uint32 InPlayerID, TSharedPtr<IAsyncResourceReleaseNotification, ESPMode::ThreadSafe> InAsyncResourceReleaseNotification, bool bReleasePools, int64 TexturePoolDecoderID);
 	};
 
 	struct FPlaystartOptions
@@ -722,6 +722,7 @@ private:
 	TSharedPtr<FElectraTextureSamplePool, ESPMode::ThreadSafe> OutputTexturePool;
 	int64 OutputTexturePoolDecoderID = 0;
 	TSharedPtr<FElectraAudioSamplePool, ESPMode::ThreadSafe> OutputAudioSamplePool;
+	bool bRetainPoolsForReuse = false;
 
 	/** Application system event handler */
 	TSharedPtr<Electra::FApplicationTerminationHandler, ESPMode::ThreadSafe> AppTerminationHandler;
