ElectraPlayer: Read byte range of DASH init segment if given instead of looking for it as if this were a self initializing media segment.

#rnx

diff --git a/Engine/Plugins/Media/ElectraPlayer/Source/ElectraPlayerRuntime/Private/Runtime/Player/DASH/ManifestDASH.cpp b/Engine/Plugins/Media/ElectraPlayer/Source/ElectraPlayerRuntime/Private/Runtime/Player/DASH/ManifestDASH.cpp
--- a/Engine/Plugins/Media/ElectraPlayer/Source/ElectraPlayerRuntime/Private/Runtime/Player/DASH/ManifestDASH.cpp
+++ b/Engine/Plugins/Media/ElectraPlayer/Source/ElectraPlayerRuntime/Private/Runtime/Player/DASH/ManifestDASH.cpp
@@ -2495,7 +2495,7 @@ bool FManifestDASHInternal::FRepresentation::PrepareDownloadURLs(IPlayerSessionS
 		InOutSegmentInfo.InitializationURL.Range = InitializationSegmentURLRange;
 		InOutSegmentInfo.InitializationURL.Url.URL = DASHUrlHelpers::ApplyAnnexEByteRange(InOutSegmentInfo.InitializationURL.Url.URL, InOutSegmentInfo.InitializationURL.Range, ChosenInit.AdditionalParams.GetValue(DASHUrlHelpers::SteerOption_ByteRange).SafeGetFString());
 	}
-	else if (GetStreamContainerType() == EStreamContainerType::Matroska)
+	else
 	{
 		// The init segment is part of the media segment, so we can just copy this across.
 		InOutSegmentInfo.InitializationURL = InOutSegmentInfo.MediaURL;
diff --git a/Engine/Plugins/Media/ElectraPlayer/Source/ElectraPlayerRuntime/Private/Runtime/Player/StreamSegmentReaderCommon.cpp b/Engine/Plugins/Media/ElectraPlayer/Source/ElectraPlayerRuntime/Private/Runtime/Player/StreamSegmentReaderCommon.cpp
--- a/Engine/Plugins/Media/ElectraPlayer/Source/ElectraPlayerRuntime/Private/Runtime/Player/StreamSegmentReaderCommon.cpp
+++ b/Engine/Plugins/Media/ElectraPlayer/Source/ElectraPlayerRuntime/Private/Runtime/Player/StreamSegmentReaderCommon.cpp
@@ -891,22 +891,24 @@ bool FStreamSegmentReaderCommon::FStreamHandler::FetchInitSegment(TSharedPtrTS<F
 	// The presence of an init segment is signalled differently for HLS and DASH.
 	if (InRequest->StreamingProtocol == FStreamSegmentRequestCommon::EStreamingProtocol::DASH)
 	{
-		// With DASH, if the init segment is not specified and the container is fmp4 it is
-		// said that the media stream is self-initializing. This means that the init segment
-		// is the 'moov' box located somewhere in the actual media stream.
-		// Thankfully with DASH we also know at this point what the container format is
-		// (or at least should be).
-		check(InRequest->ExpectedContainerFormat != FStreamSegmentRequestCommon::EContainerFormat::Undefined);
+		// With DASH, if the init segment is not specified it is said that the media stream is self-initializing.
+		// This means that the init segment is the 'moov' box located somewhere in the actual media stream (if this is an fmp4).
 		if (InRequest->ExpectedContainerFormat == FStreamSegmentRequestCommon::EContainerFormat::ISO_14496_12)
 		{
-			if (InitURLToUse.Url.URL.IsEmpty())
+			if (InitURLToUse.Url.URL.IsEmpty() || (InitURLToUse.Range.IsEmpty() && InitURLToUse.Url.URL.Equals(segInfo.MediaURL.Url.URL)))
 			{
-				InitURLToUse = segInfo.MediaURL;
-				// Clear out the media range as we need to scan for the 'moov' box through the entire file.
-				InitURLToUse.Range.Empty();
-				bFetchFromSelfInitializing = true;
-				DownloadStats.URL = InitURLToUse.Url;
-				DownloadStats.Range = InitURLToUse.Range;
+				FString InitRange(InitURLToUse.Range);
+				if (InitURLToUse.Url.URL.IsEmpty())
+				{
+					InitURLToUse = segInfo.MediaURL;
+				}
+				InitURLToUse.Range = MoveTemp(InitRange);
+				if (InitURLToUse.Range.IsEmpty())
+				{
+					bFetchFromSelfInitializing = true;
+					DownloadStats.URL = InitURLToUse.Url;
+					DownloadStats.Range = InitURLToUse.Range;
+				}
 			}
 		}
 	}
