ElectraPlayer: Added handling of AV1 codec OTI

#rnx

diff --git a/Engine/Plugins/Media/ElectraPlayer/Source/ElectraPlayerRuntime/Private/Runtime/StreamTypes.cpp b/Engine/Plugins/Media/ElectraPlayer/Source/ElectraPlayerRuntime/Private/Runtime/StreamTypes.cpp
--- a/Engine/Plugins/Media/ElectraPlayer/Source/ElectraPlayerRuntime/Private/Runtime/StreamTypes.cpp
+++ b/Engine/Plugins/Media/ElectraPlayer/Source/ElectraPlayerRuntime/Private/Runtime/StreamTypes.cpp
@@ -344,6 +344,28 @@ namespace Electra
 			SetCodecSpecifierRFC6381(FString::Printf(TEXT("vp09.%02d.%02d.%02d"), 0, 0, 8));
 			return true;
 		}
+		else if (CodecOTI.StartsWith(TEXT("av01")))
+		{
+			ElectraDecodersUtil::FMimeTypeVideoCodecInfo ci;
+			if (ElectraDecodersUtil::ParseCodecAV1(ci, CodecOTI, Extras.GetValue(StreamCodecInformationOptions::AV1cBox).SafeGetArray(), Extras.GetValue(StreamCodecInformationOptions::ColrBox).SafeGetArray()))
+			{
+				StreamType = EStreamType::Video;
+				CodecSpecifier = CodecOTI;
+				Codec = ECodec::Video4CC;
+				Codec4CC = Make4CC('a','v','0','1');
+				SetProfile(ci.Profile);
+				SetProfileLevel(ci.Level);
+				FString rfc(FString::Printf(TEXT("av01.%d.%02d%c.%02d.%d.%d%d%d.%02d.%02d.%02d.%d"), ci.Profile, ci.Level, ci.Tier?TCHAR('H'):TCHAR('M'), ci.NumBitsLuma, ci.Extras[4], ci.Extras[5], ci.Extras[6], ci.Extras[7], ci.Extras[8], ci.Extras[9], ci.Extras[10], ci.Extras[11]));
+				// Defaults after the bitdepth are to be omitted.
+				if (rfc.EndsWith(TEXT(".0.110.01.01.01.0")))
+				{
+					rfc.LeftChopInline(17);
+				}
+				SetCodecSpecifierRFC6381(rfc);
+				return true;
+			}
+			return false;
+		}
 		else if (CodecOTI.Equals(TEXT("wvtt")))
 		{
 			StreamType = EStreamType::Subtitle;
diff --git a/Engine/Plugins/Media/ElectraPlayer/Source/ElectraPlayerRuntime/Private/Runtime/StreamTypes.h b/Engine/Plugins/Media/ElectraPlayer/Source/ElectraPlayerRuntime/Private/Runtime/StreamTypes.h
--- a/Engine/Plugins/Media/ElectraPlayer/Source/ElectraPlayerRuntime/Private/Runtime/StreamTypes.h
+++ b/Engine/Plugins/Media/ElectraPlayer/Source/ElectraPlayerRuntime/Private/Runtime/StreamTypes.h
@@ -709,6 +709,8 @@ namespace Electra
 	{
 		static const FName DecoderConfigurationRecord(TEXT("dcr"));
 		static const FName VPccBox(TEXT("vpcC_box"));
+		static const FName AV1cBox(TEXT("av1C_box"));
+		static const FName ColrBox(TEXT("colr_box"));
 		static const FName DOpsBox(TEXT("dOps_box"));
 		static const FName ChanBox(TEXT("chan_box"));
 		static const FName SamplesPerBlock(TEXT("samples_per_block"));
