Electra: Some codec related utility work

#rnx

#Synced-CL 48204442

diff --git a/Engine/Plugins/Media/ElectraCodecs/Source/ElectraDecoders/Private/ElectraDecodersUtils.cpp b/Engine/Plugins/Media/ElectraCodecs/Source/ElectraDecoders/Private/ElectraDecodersUtils.cpp
--- a/Engine/Plugins/Media/ElectraCodecs/Source/ElectraDecoders/Private/ElectraDecodersUtils.cpp
+++ b/Engine/Plugins/Media/ElectraCodecs/Source/ElectraDecoders/Private/ElectraDecodersUtils.cpp
@@ -40,6 +40,7 @@ bool ParseCodecMP4A(FMimeTypeAudioCodecInfo& OutInfo, const FString& InCodecForm
 		return false;
 	}
 	OutInfo.Codec = TEXT("mp4a");
+	OutInfo.Codec4CC = Make4CC('m','p','4','a');
 	// Object and profile follow?
 	if (InCodecFormat.Len() > 6 && InCodecFormat[4] == TCHAR('.'))
 	{
@@ -69,6 +70,7 @@ bool ParseCodecH264(FMimeTypeVideoCodecInfo& OutInfo, const FString& InCodecForm
 				return false;
 			}
 			OutInfo.Codec = TEXT("avc");
+			OutInfo.Codec4CC = InCodecFormat[3] == TCHAR('1') ? Make4CC('a','v','c','1') : Make4CC('a','v','c','3');
 
 			// Profile and level follow?
 			if (InCodecFormat.Len() > 5 && InCodecFormat[4] == TCHAR('.'))
@@ -104,8 +106,8 @@ bool ParseCodecH264(FMimeTypeVideoCodecInfo& OutInfo, const FString& InCodecForm
 					return false;
 				}
 			}
+			return true;
 		}
-		return true;
 	}
 	return false;
 }
@@ -118,6 +120,7 @@ bool ParseCodecH265(FMimeTypeVideoCodecInfo& OutInfo, const FString& InCodecForm
 		if (InCodecFormat.Len() > 4)
 		{
 			OutInfo.Codec = TEXT("hevc");
+			OutInfo.Codec4CC = InCodecFormat.StartsWith("hvc1") ? Make4CC('h','v','c','1') : Make4CC('h','e','v','1');
 
 			FString oti = InCodecFormat;
 			FString Temp;
@@ -195,9 +198,11 @@ bool ParseCodecH265(FMimeTypeVideoCodecInfo& OutInfo, const FString& InCodecForm
 
 bool ParseCodecVP8(FMimeTypeVideoCodecInfo& OutInfo, const FString& InCodecFormat, const TArray<uint8>& InvpcCBox)
 {
-	if (InCodecFormat.StartsWith("vp08"))
+	if (InCodecFormat.StartsWith("vp08") || InCodecFormat.StartsWith("vp8"))
 	{
 		FString oti = InCodecFormat;
+		OutInfo.Codec = TEXT("vp8");
+		OutInfo.Codec4CC = Make4CC('v','p','0','8');
 		if (!InvpcCBox.IsEmpty())
 		{
 			// Enough data to represent the `vpcC` box, and is it of version 1?
@@ -213,7 +218,6 @@ bool ParseCodecVP8(FMimeTypeVideoCodecInfo& OutInfo, const FString& InCodecForma
 		}
 		if (oti.Len() > 4)
 		{
-			OutInfo.Codec = TEXT("vp08");
 			FString Temp;
 			int32 DotPos;
 			if (oti.FindChar(TCHAR('.'), DotPos))
@@ -246,9 +250,11 @@ bool ParseCodecVP8(FMimeTypeVideoCodecInfo& OutInfo, const FString& InCodecForma
 
 bool ParseCodecVP9(FMimeTypeVideoCodecInfo& OutInfo, const FString& InCodecFormat, const TArray<uint8>& InvpcCBox)
 {
-	if (InCodecFormat.StartsWith("vp09"))
+	if (InCodecFormat.StartsWith("vp09") || InCodecFormat.StartsWith("vp9"))
 	{
 		FString oti = InCodecFormat;
+		OutInfo.Codec = TEXT("vp9");
+		OutInfo.Codec4CC = Make4CC('v','p','0','9');
 		if (!InvpcCBox.IsEmpty())
 		{
 			// Enough data to represent the `vpcC` box, and is it of version 1?
@@ -269,7 +275,6 @@ bool ParseCodecVP9(FMimeTypeVideoCodecInfo& OutInfo, const FString& InCodecForma
 		}
 		if (oti.Len() > 4)
 		{
-			OutInfo.Codec = TEXT("vp09");
 			FString Temp;
 			int32 DotPos;
 			if (oti.FindChar(TCHAR('.'), DotPos))
@@ -310,6 +315,8 @@ bool ParseCodecAV1(FMimeTypeVideoCodecInfo& OutInfo, const FString& InCodecForma
 	if (InCodecFormat.StartsWith("av01"))
 	{
 		FString oti = InCodecFormat;
+		OutInfo.Codec = TEXT("av1");
+		OutInfo.Codec4CC = Make4CC('a','v','0','1');
 		if (!Inav1CBox.IsEmpty())
 		{
 			// Enough data to represent the `av1C` box, and is it of version 1?
@@ -320,7 +327,6 @@ bool ParseCodecAV1(FMimeTypeVideoCodecInfo& OutInfo, const FString& InCodecForma
 			auto high_bitdepth = [&]() -> bool { return (Inav1CBox[2] & 0x40) == 0x40; };
 			auto twelve_bit = [&]() -> bool { return (Inav1CBox[2] & 0x20) == 0x20; };
 
-			OutInfo.Codec = TEXT("av01");
 			OutInfo.Extras[0] = OutInfo.Profile = Inav1CBox[1] >> 5;	// profile
 			OutInfo.Extras[1] = OutInfo.Level = Inav1CBox[1] & 31;		// level
 			OutInfo.Extras[2] = OutInfo.Tier = Inav1CBox[2] >> 7;		// tier
@@ -363,7 +369,6 @@ bool ParseCodecAV1(FMimeTypeVideoCodecInfo& OutInfo, const FString& InCodecForma
 		}
 		if (oti.Len() > 4)
 		{
-			OutInfo.Codec = TEXT("av01");
 			TArray<FString> Parts;
 			InCodecFormat.ParseIntoArray(Parts, TEXT("."), false);
 			// Profile, level, tier and bitdepth are required.
diff --git a/Engine/Plugins/Media/ElectraCodecs/Source/ElectraDecoders/Private/Utils/Google/ElectraBitstreamProcessorVPx.cpp b/Engine/Plugins/Media/ElectraCodecs/Source/ElectraDecoders/Private/Utils/Google/ElectraBitstreamProcessorVPx.cpp
--- a/Engine/Plugins/Media/ElectraCodecs/Source/ElectraDecoders/Private/Utils/Google/ElectraBitstreamProcessorVPx.cpp
+++ b/Engine/Plugins/Media/ElectraCodecs/Source/ElectraDecoders/Private/Utils/Google/ElectraBitstreamProcessorVPx.cpp
@@ -6,6 +6,7 @@
 #include "Utilities/UtilitiesMP4.h"
 #include "IElectraDecoderFeaturesAndOptions.h"
 #include "ElectraDecodersUtils.h"
+#include "Misc/Optional.h"
 
 
 
@@ -27,6 +28,7 @@ public:
 		bSentColorimetry = false;
 		bSentMDCV = false;
 		bSentCLLI = false;
+		LastVP9Header.Reset();
 		// Do not clear out any error message!
 	}
 
@@ -35,8 +37,9 @@ public:
 		return LastErrorMessage;
 	}
 
-	bool ProcessInputForDecoding(FElectraDecoderInputAccessUnit& InOutAccessUnit, const TMap<FString, FVariant>& InAccessUnitSidebandData)
+	bool ProcessInputForDecoding(bool& OutCSDChanged, FElectraDecoderInputAccessUnit& InOutAccessUnit, const TMap<FString, FVariant>& InAccessUnitSidebandData)
 	{
+		OutCSDChanged = false;
 		InOutAccessUnit.Flags &= ~EElectraDecoderFlags::IsDiscardable;
 		if (Version == EVersion::VP8)
 		{
@@ -66,9 +69,27 @@ public:
 				LastErrorMessage = FString::Printf(TEXT("Failed to parse VP9 header"));
 				return false;
 			}
+
 			if (Header.IsKeyframe())
 			{
 				InOutAccessUnit.Flags |= EElectraDecoderFlags::IsSyncSample;
+
+				if (LastVP9Header.IsSet())
+				{
+					const ElectraDecodersUtil::VPxVideo::FVP9UncompressedHeader& prv(LastVP9Header.GetValue());
+					OutCSDChanged = Header.frame_width != prv.frame_width ||
+									Header.frame_height != prv.frame_height ||
+									Header.render_width != prv.render_width ||
+									Header.render_height != prv.render_height ||
+									Header.profile_low_bit != prv.profile_low_bit ||
+									Header.profile_high_bit != prv.profile_high_bit ||
+									Header.ten_or_twelve_bit != prv.ten_or_twelve_bit ||
+									Header.color_space != prv.color_space ||
+									Header.color_range != prv.color_range ||
+									Header.subsampling_x != prv.subsampling_x ||
+									Header.subsampling_y != prv.subsampling_y;
+				}
+				LastVP9Header = Header;
 			}
 			else
 			{
@@ -199,6 +220,7 @@ private:
 	TArray<uint8> COLRBox;
 	TArray<uint8> MDCVBox;
 	TArray<uint8> CLLIBox;
+	TOptional<ElectraDecodersUtil::VPxVideo::FVP9UncompressedHeader> LastVP9Header;
 	FString LastErrorMessage;
 	bool bSentColorimetry = false;
 	bool bSentMDCV = false;
@@ -248,12 +270,13 @@ IElectraDecoderBitstreamProcessor::EProcessResult FElectraDecoderBitstreamProces
 	}
 
 	InOutAccessUnit.Flags |= EElectraDecoderFlags::InputIsProcessed;
-	if (!Impl->ProcessInputForDecoding(InOutAccessUnit, InAccessUnitSidebandData))
+	bool bCSDChanged = false;
+	if (!Impl->ProcessInputForDecoding(bCSDChanged, InOutAccessUnit, InAccessUnitSidebandData))
 	{
 		return EProcessResult::Error;
 	}
 
-	return EProcessResult::Ok;
+	return bCSDChanged ? EProcessResult::CSDChanged : EProcessResult::Ok;
 }
 
 void FElectraDecoderBitstreamProcessorVPx::SetPropertiesOnOutput(TMap<FString, FVariant>& InOutProperties, TSharedPtr<IElectraDecoderBitstreamInfo, ESPMode::ThreadSafe> InBSI)
diff --git a/Engine/Plugins/Media/ElectraCodecs/Source/ElectraDecoders/Public/ElectraDecodersUtils.h b/Engine/Plugins/Media/ElectraCodecs/Source/ElectraDecoders/Public/ElectraDecodersUtils.h
--- a/Engine/Plugins/Media/ElectraCodecs/Source/ElectraDecoders/Public/ElectraDecodersUtils.h
+++ b/Engine/Plugins/Media/ElectraCodecs/Source/ElectraDecoders/Public/ElectraDecodersUtils.h
@@ -56,6 +56,7 @@ namespace ElectraDecodersUtil
 	struct FMimeTypeCodecInfo
 	{
 		FString Codec;
+		uint32 Codec4CC = 0;
 	};
 
 	struct FMimeTypeVideoCodecInfo : public FMimeTypeCodecInfo
diff --git a/Engine/Plugins/Media/ElectraPlayer/Source/ElectraPlayerRuntime/Private/Runtime/Decoder/VideoDecoder.cpp b/Engine/Plugins/Media/ElectraPlayer/Source/ElectraPlayerRuntime/Private/Runtime/Decoder/VideoDecoder.cpp
--- a/Engine/Plugins/Media/ElectraPlayer/Source/ElectraPlayerRuntime/Private/Runtime/Decoder/VideoDecoder.cpp
+++ b/Engine/Plugins/Media/ElectraPlayer/Source/ElectraPlayerRuntime/Private/Runtime/Decoder/VideoDecoder.cpp
@@ -300,6 +300,10 @@ TSharedPtr<IElectraCodecFactory, ESPMode::ThreadSafe> FVideoDecoderImpl::GetDeco
 	{
 		OutAddtlCfg.Add(TEXT("dcr"), FVariant(InCodecData->RawCSD));
 	}
+	else if (InCodecInfo.GetDecoderConfigRecord().Num())
+	{
+		OutAddtlCfg.Add(TEXT("dcr"), FVariant(InCodecInfo.GetDecoderConfigRecord()));
+	}
 	OutFormat = InCodecInfo.GetCodecSpecifierRFC6381();
 	if (OutFormat.Len() == 0)
 	{
@@ -1287,7 +1291,6 @@ FVideoDecoderImpl::ENextDecodingState FVideoDecoderImpl::HandleReplaying()
 bool FVideoDecoderImpl::HandleDummyDecoding()
 {
 	check(CurrentAccessUnit.IsValid());
-	check(bIsDecoderClean);
 
 	// Get output unless flushing or terminating
 	while(!TerminateThreadSignal.IsSignaled() && !FlushDecoderSignal.IsSignaled())
diff --git a/Engine/Plugins/Media/ElectraPlayer/Source/ElectraPlayerRuntime/Private/Runtime/Demuxer/ParserISO14496-12.cpp b/Engine/Plugins/Media/ElectraPlayer/Source/ElectraPlayerRuntime/Private/Runtime/Demuxer/ParserISO14496-12.cpp
--- a/Engine/Plugins/Media/ElectraPlayer/Source/ElectraPlayerRuntime/Private/Runtime/Demuxer/ParserISO14496-12.cpp
+++ b/Engine/Plugins/Media/ElectraPlayer/Source/ElectraPlayerRuntime/Private/Runtime/Demuxer/ParserISO14496-12.cpp
@@ -7452,7 +7452,10 @@ namespace Electra
 							Track->CodecSpecificDataAVC = AVCCBox->GetDecoderConfigurationRecord().GetValue();
 							Track->CodecInformation.SetStreamType(EStreamType::Video);
 							Track->CodecInformation.SetCodec(FStreamCodecInformation::ECodec::H264);
+							Track->CodecInformation.SetCodec4CC(InAvcType == 1 ? Utils::Make4CC('a','v','c','1') : Utils::Make4CC('a','v','c','3'));
 							Track->CodecInformation.SetStreamLanguageTag(Track->GetLanguageTag());
+							Track->CodecInformation.SetDecoderConfigRecord(Track->CodecSpecificDataAVC.GetRawData());
+							Track->CodecInformation.SetCodecSpecificData(Track->CodecSpecificDataAVC.GetCodecSpecificData());
 							if (Track->CodecSpecificDataAVC.GetSequenceParameterSets().IsEmpty())
 							{
 								if (InAvcType == 3)
@@ -7573,6 +7576,7 @@ namespace Electra
 							Track->CodecSpecificDataHEVC = HVCCBox->GetDecoderConfigurationRecord().GetValue();
 							Track->CodecInformation.SetStreamType(EStreamType::Video);
 							Track->CodecInformation.SetCodec(FStreamCodecInformation::ECodec::H265);
+							Track->CodecInformation.SetCodec4CC(Utils::Make4CC('h','e','v','1'));
 							Track->CodecInformation.SetStreamLanguageTag(Track->GetLanguageTag());
 							Track->CodecInformation.SetProfileSpace(Track->CodecSpecificDataHEVC.GetGeneralProfileSpace());
 							Track->CodecInformation.SetProfileTier(Track->CodecSpecificDataHEVC.GetGeneralTierFlag());
@@ -7580,7 +7584,8 @@ namespace Electra
 							Track->CodecInformation.SetProfileLevel(Track->CodecSpecificDataHEVC.GetGeneralLevelIDC());
 							Track->CodecInformation.SetProfileConstraints(Track->CodecSpecificDataHEVC.GetGeneralConstraintIndicatorFlags());
 							Track->CodecInformation.SetProfileCompatibilityFlags(Track->CodecSpecificDataHEVC.GetGeneralProfileCompatibilityFlags());
-
+							Track->CodecInformation.SetDecoderConfigRecord(Track->CodecSpecificDataHEVC.GetRawData());
+							Track->CodecInformation.SetCodecSpecificData(Track->CodecSpecificDataHEVC.GetCodecSpecificData());
 							if (Track->CodecSpecificDataHEVC.GetSequenceParameterSets().IsEmpty())
 							{
 								if (Track->TKHDBox)
@@ -7610,6 +7615,7 @@ namespace Electra
 								Track->CodecInformation.SetAspectRatio(ar);
 								Track->CodecInformation.SetFrameRate(sps.GetTiming().Denom ? FTimeFraction(sps.GetTiming().Num, sps.GetTiming().Denom) : FTimeFraction());
 								Track->CodecInformation.SetCodecSpecifierRFC6381(sps.GetRFC6381(TEXT("hvc1")));
+								Track->CodecInformation.SetCodec4CC(Utils::Make4CC('h','v','c','1'));
 								bGotVideoFormat = true;
 							}
 							else
@@ -7669,6 +7675,7 @@ namespace Electra
 				{
 					Track->CodecInformation.SetStreamType(EStreamType::Audio);
 					Track->CodecInformation.SetCodec(FStreamCodecInformation::ECodec::AAC);
+					Track->CodecInformation.SetCodec4CC(Utils::Make4CC('m','p','4','a'));
 					Track->CodecInformation.SetStreamLanguageTag(Track->GetLanguageTag());
 
 					ElectraDecodersUtil::MPEG::FAACDecoderConfigurationRecord ConfigRecord;
@@ -7821,6 +7828,7 @@ namespace Electra
 			Track->CodecSpecificDataRAW = SubtitleSampleEntry->GetRawBoxData();
 			Track->CodecInformation.SetStreamType(EStreamType::Subtitle);
 			Track->CodecInformation.SetCodec(FStreamCodecInformation::ECodec::TX3G);
+			Track->CodecInformation.SetCodec4CC(Utils::Make4CC('t','x','3','g'));
 			Track->CodecInformation.SetCodecSpecificData(Track->CodecSpecificDataRAW);
 			Track->CodecInformation.SetStreamLanguageTag(Track->GetLanguageTag());
 			Track->CodecInformation.SetCodecSpecifierRFC6381(TEXT("tx3g"));
@@ -7842,6 +7850,7 @@ namespace Electra
 			Track->CodecSpecificDataRAW = SubtitleSampleEntry->GetRawBoxData();
 			Track->CodecInformation.SetStreamType(EStreamType::Subtitle);
 			Track->CodecInformation.SetCodec(FStreamCodecInformation::ECodec::WebVTT);
+			Track->CodecInformation.SetCodec4CC(Utils::Make4CC('w','v','t','t'));
 			Track->CodecInformation.SetCodecSpecificData(Track->CodecSpecificDataRAW);
 			Track->CodecInformation.SetStreamLanguageTag(Track->GetLanguageTag());
 			Track->CodecInformation.SetCodecSpecifierRFC6381(TEXT("wvtt"));
@@ -7863,6 +7872,7 @@ namespace Electra
 			Track->CodecSpecificDataRAW = SubtitleSampleEntry->GetRawBoxData();
 			Track->CodecInformation.SetStreamType(EStreamType::Subtitle);
 			Track->CodecInformation.SetCodec(FStreamCodecInformation::ECodec::TTML);
+			Track->CodecInformation.SetCodec4CC(Utils::Make4CC('s','t','p','p'));
 			Track->CodecInformation.SetCodecSpecificData(Track->CodecSpecificDataRAW);
 			Track->CodecInformation.SetStreamLanguageTag(Track->GetLanguageTag());
 			Track->CodecInformation.SetCodecSpecifierRFC6381(TEXT("stpp"));
diff --git a/Engine/Plugins/Media/ElectraPlayer/Source/ElectraPlayerRuntime/Private/Runtime/StreamTypes.cpp b/Engine/Plugins/Media/ElectraPlayer/Source/ElectraPlayerRuntime/Private/Runtime/StreamTypes.cpp
--- a/Engine/Plugins/Media/ElectraPlayer/Source/ElectraPlayerRuntime/Private/Runtime/StreamTypes.cpp
+++ b/Engine/Plugins/Media/ElectraPlayer/Source/ElectraPlayerRuntime/Private/Runtime/StreamTypes.cpp
@@ -90,6 +90,7 @@ namespace Electra
 			StreamType = EStreamType::Video;
 			CodecSpecifier = CodecOTI;
 			Codec = ECodec::H264;
+			Codec4CC = Make4CC('a','v','c','1');
 			if (CodecOTI.Len() > 3)
 			{
 				// avc 1 or 3 only.
@@ -97,6 +98,8 @@ namespace Electra
 				{
 					return false;
 				}
+				Codec4CC = CodecOTI[3] == TCHAR('1') ? Make4CC('a','v','c','1') : Make4CC('a','v','c','3');
+
 				// Profile and level follow?
 				if (CodecOTI.Len() > 5 && CodecOTI[4] == TCHAR('.'))
 				{
@@ -145,7 +148,7 @@ namespace Electra
 			StreamType = EStreamType::Video;
 			CodecSpecifier = CodecOTI;
 			Codec = ECodec::H265;
-
+			Codec4CC = CodecOTI.StartsWith(TEXT("hvc")) ? Make4CC('h','v','c','1') : Make4CC('h','e','v','1');
 			int32 DotPos;
 			if (oti.FindChar(TCHAR('.'), DotPos))
 			{
@@ -228,6 +231,7 @@ namespace Electra
 			StreamType = EStreamType::Audio;
 			CodecSpecifier = CodecOTI;
 			Codec = ECodec::AAC;
+			Codec4CC = Make4CC('m','p','4','a');
 			// Object and profile follow?
 			if (CodecOTI.Len() > 6 && CodecOTI[4] == TCHAR('.'))
 			{
@@ -371,6 +375,7 @@ namespace Electra
 			StreamType = EStreamType::Subtitle;
 			CodecSpecifier = CodecOTI;
 			Codec = ECodec::WebVTT;
+			Codec4CC = Make4CC('w','v','t','t');
 			return true;
 		}
 		// This is indicating one of the many TTML variants (eg. IMSC1, SMPTE-TT, EBU-TT) and profiles (eg. stpp.ttml.im1t)
@@ -380,6 +385,7 @@ namespace Electra
 			StreamType = EStreamType::Subtitle;
 			CodecSpecifier = CodecOTI;
 			Codec = ECodec::TTML;
+			Codec4CC = Make4CC('s','t','p','p');
 			return true;
 		}
 		else if (CodecOTI.Equals(TEXT("tx3g")))
@@ -387,6 +393,7 @@ namespace Electra
 			StreamType = EStreamType::Subtitle;
 			CodecSpecifier = CodecOTI;
 			Codec = ECodec::TX3G;
+			Codec4CC = Make4CC('t','x','3','g');
 			return true;
 		}
 		else
@@ -492,6 +499,9 @@ namespace Electra
 					case Make4CC('v','p','0','9'):
 						HumanReadableCodecName = TEXT("VP9");
 						return true;
+					case Make4CC('a','v','0','1'):
+						HumanReadableCodecName = TEXT("AV1");
+						return true;
 
 					case Make4CC('a','p','c','h'):
 						HumanReadableCodecName = TEXT("Apple ProRes 422 High Quality");
diff --git a/Engine/Plugins/Media/VPxDecoderElectra/Source/VPxDecoderElectra/Private/VPxDecoder/ElectraMediaVPxDecoder.cpp b/Engine/Plugins/Media/VPxDecoderElectra/Source/VPxDecoderElectra/Private/VPxDecoder/ElectraMediaVPxDecoder.cpp
--- a/Engine/Plugins/Media/VPxDecoderElectra/Source/VPxDecoderElectra/Private/VPxDecoder/ElectraMediaVPxDecoder.cpp
+++ b/Engine/Plugins/Media/VPxDecoderElectra/Source/VPxDecoderElectra/Private/VPxDecoder/ElectraMediaVPxDecoder.cpp
@@ -242,7 +242,14 @@ private:
 	{
 		~FDecoderInput()
 		{
-			FMemory::Free(InputDataCopy);
+			if (InputDataCopy)
+			{
+				FMemory::Free(InputDataCopy);
+			}
+			if (InputAlphaDataCopy)
+			{
+				FMemory::Free(InputAlphaDataCopy);
+			}
 		}
 		FInputAccessUnit AccessUnit;
 		TMap<FString, FVariant> AdditionalOptions;
@@ -603,6 +610,10 @@ IElectraDecoder::EDecoderError FVideoDecoderVPxElectra::DecodeNextPending()
 
 	// Add to the list of inputs passed to the decoder.
 	InDecoderInput.Emplace(MoveTemp(In));
+	InDecoderInput.StableSort([](const TSharedPtr<FDecoderInput, ESPMode::NotThreadSafe>& a, const TSharedPtr<FDecoderInput, ESPMode::NotThreadSafe>& b)
+	{
+		return a->AccessUnit.PTS < b->AccessUnit.PTS;
+	});
 
 	// Did we produce a new frame?
 	Decoders[EDecoderType::Color].OutputIterator = nullptr;
@@ -910,6 +921,12 @@ FVideoDecoderVPxElectra::EConvertResult FVideoDecoderVPxElectra::ConvertDecoderO
 		PostError(0, TEXT("There is no matching decoder input for the decoded output!"), ERRCODE_INTERNAL_FAILED_TO_CONVERT_OUTPUT_SAMPLE);
 		return EConvertResult::Failure;
 	}
+	// Remove outdated entries.
+	InDecoderInput.RemoveAll([FramePTS=In->AccessUnit.PTS](const TSharedPtr<FDecoderInput, ESPMode::NotThreadSafe>& e)
+	{
+		return e->AccessUnit.PTS < FramePTS;
+	});
+
 
 	TSharedPtr<FVideoDecoderOutputVPxElectra, ESPMode::ThreadSafe> NewOutput = MakeShared<FVideoDecoderOutputVPxElectra>();
 	NewOutput->PTS = In->AccessUnit.PTS;
