ElectraPlayer: changed time assignments

#rnx

#Synced-CL 49644370

diff --git a/Engine/Plugins/Media/ElectraPlayer/Source/ElectraPlayerRuntime/Private/Runtime/Player/mkv/ManifestMKV.cpp b/Engine/Plugins/Media/ElectraPlayer/Source/ElectraPlayerRuntime/Private/Runtime/Player/mkv/ManifestMKV.cpp
--- a/Engine/Plugins/Media/ElectraPlayer/Source/ElectraPlayerRuntime/Private/Runtime/Player/mkv/ManifestMKV.cpp
+++ b/Engine/Plugins/Media/ElectraPlayer/Source/ElectraPlayerRuntime/Private/Runtime/Player/mkv/ManifestMKV.cpp
@@ -633,7 +633,7 @@ FErrorDetail FManifestMKVInternal::FTimelineAssetMKV::Build(IPlayerSessionServic
 	PlayerSessionServices = InPlayerSessionServices;
 	MediaURL = URL;
 
-	MovieDuration.SetFromTimeFraction(InMKVParser->GetDuration());
+	MovieDuration = InMKVParser->GetDuration();
 	// Go over the supported tracks and create an internal manifest-like structure for the player to work with.
 	for(int32 nTrack=0,nMaxTrack=InMKVParser->GetNumberOfTracks(); nTrack < nMaxTrack; ++nTrack)
 	{
