Electra: Moved mp4 related utilities into dedicated plugin

#rnx

#Synced-CL 49721941

diff --git a/Engine/Plugins/Media/APVDecoderElectra/Source/APVDecoderElectra/APVDecoderElectra.Build.cs b/Engine/Plugins/Media/APVDecoderElectra/Source/APVDecoderElectra/APVDecoderElectra.Build.cs
--- a/Engine/Plugins/Media/APVDecoderElectra/Source/APVDecoderElectra/APVDecoderElectra.Build.cs
+++ b/Engine/Plugins/Media/APVDecoderElectra/Source/APVDecoderElectra/APVDecoderElectra.Build.cs
@@ -15,7 +15,8 @@ namespace UnrealBuildTool.Rules
                     "ElectraBase",
                     "ElectraSamples",
 					"ElectraCodecFactory",
-					"ElectraDecoders"
+					"ElectraDecoders",
+                    "MP4Utilities"
                 });
 
             if (Target.Platform == UnrealTargetPlatform.Win64)// || Target.Platform == UnrealTargetPlatform.Mac || Target.Platform == UnrealTargetPlatform.Linux)
diff --git a/Engine/Plugins/Media/APVDecoderElectra/Source/APVDecoderElectra/Private/APVDecoder/ElectraMediaAPVDecoder.cpp b/Engine/Plugins/Media/APVDecoderElectra/Source/APVDecoderElectra/Private/APVDecoder/ElectraMediaAPVDecoder.cpp
--- a/Engine/Plugins/Media/APVDecoderElectra/Source/APVDecoderElectra/Private/APVDecoder/ElectraMediaAPVDecoder.cpp
+++ b/Engine/Plugins/Media/APVDecoderElectra/Source/APVDecoderElectra/Private/APVDecoder/ElectraMediaAPVDecoder.cpp
@@ -15,7 +15,7 @@
 #include "ElectraDecodersUtils.h"
 #include "Utils/OpenAPV/ElectraBitstreamProcessor_APV.h"
 #include "Utils/OpenAPV/ElectraUtilsAPVVideo.h"
-#include "Utilities/UtilitiesMP4.h"
+#include "MP4Utilities.h"
 
 #include "ElectraDecodersPlatformResources.h"
 
@@ -257,7 +257,6 @@ public:
 	{
 		TMap<FString, FVariant> DecoderFeatures;
 		GetFeatures(DecoderFeatures);
-		DecoderFeatures.Emplace(IElectraDecoderFeature::ModifiesInputAUInPlace, FVariant(true));
 		return FElectraDecoderBitstreamProcessorAPV::Create(DecoderFeatures, InitialCreationOptions);
 	}
 
@@ -542,21 +541,41 @@ IElectraDecoder::EDecoderError FVideoDecoderAPVElectra::DecodeAccessUnit(const F
 		oapv_bitb_t InputBS;
 		FMemory::Memzero(InputBS);
 		// The decoder checks for the `aPv1` magic which may not be present if we decode one sub frame from
-		// the access unit. We insert this magic, thereby modifying the bitstream which is why we are signalling
-		// this when creating the bitstream processor.
-		if (FramePBUs[0].PBUOffset < 4)
+		// the access unit. In this case we have to create a temporary buffer to prepend the magic.
+		uint8* TempAU = nullptr;
+		bool bHaveMagic = FramePBUs[0].PBUOffset >= 4 && *reinterpret_cast<const uint32*>(reinterpret_cast<const uint8*>(InInputAccessUnit.Data) + FramePBUs[0].PBUOffset - 4) == MP4Utilities::EndianSwap((uint32)0x61507631U);
+		if (!bHaveMagic)
 		{
-			PostError(0, TEXT("APV decoding failed, bad frame PBU offset"), ERRCODE_INTERNAL_FAILED_TO_DECODE_INPUT);
+			TempAU = reinterpret_cast<uint8*>(FMemory::Malloc(FramePBUs[0].PBUSize + 4));
+			if (!TempAU)
+			{
+				PostError(0, TEXT("APV decoding failed, out of memory"), ERRCODE_INTERNAL_FAILED_TO_DECODE_INPUT);
+				return IElectraDecoder::EDecoderError::Error;
+			}
+			FMemory::Memcpy(TempAU + 4, reinterpret_cast<const uint8*>(InInputAccessUnit.Data) + FramePBUs[0].PBUOffset, FramePBUs[0].PBUSize);
+			*reinterpret_cast<uint32*>(TempAU) = MP4Utilities::EndianSwap((uint32)0x61507631U);
+			InputBS.addr = TempAU;
+			InputBS.ssize = FramePBUs[0].PBUSize + 4;
+		}
+		else if (FramePBUs[0].PBUOffset >= 4)
+		{
+			InputBS.addr = (void*)(reinterpret_cast<const uint8*>(InInputAccessUnit.Data) + FramePBUs[0].PBUOffset - 4);
+			InputBS.ssize = FramePBUs[0].PBUSize + 4;
+		}
+		else
+		{
+			PostError(0, TEXT("Failed to decode APV. Malformed PBU?"), ERRCODE_INTERNAL_FAILED_TO_DECODE_INPUT);
 			return IElectraDecoder::EDecoderError::Error;
 		}
-		InputBS.addr = (void*)(reinterpret_cast<const uint8*>(InInputAccessUnit.Data) + FramePBUs[0].PBUOffset - 4);
-		InputBS.ssize = FramePBUs[0].PBUSize + 4;
-		*reinterpret_cast<uint32*>(InputBS.addr) = Electra::EndianSwap((uint32)0x61507631U);
 
 		oapvd_stat_t Stats;
 		FMemory::Memzero(Stats);
-
         int Result = oapvd_decode(DecoderInstance, &InputBS, &OutFrames, DecoderMetaInstance, &Stats);
+		if (TempAU)
+		{
+			FMemory::Free(TempAU);
+			TempAU = nullptr;
+		}
 
 		if (OAPV_FAILED(Result))
 		{
diff --git a/Engine/Plugins/Media/AppleProResDecoderElectra/Source/AppleProResDecoderElectra/AppleProResDecoderElectra.Build.cs b/Engine/Plugins/Media/AppleProResDecoderElectra/Source/AppleProResDecoderElectra/AppleProResDecoderElectra.Build.cs
--- a/Engine/Plugins/Media/AppleProResDecoderElectra/Source/AppleProResDecoderElectra/AppleProResDecoderElectra.Build.cs
+++ b/Engine/Plugins/Media/AppleProResDecoderElectra/Source/AppleProResDecoderElectra/AppleProResDecoderElectra.Build.cs
@@ -15,7 +15,8 @@ namespace UnrealBuildTool.Rules
 					"ElectraBase",
                     "ElectraSamples",
 					"ElectraCodecFactory",
-					"ElectraDecoders"
+					"ElectraDecoders",
+                    "MP4Utilities"
                 });
 
             if (Target.Platform == UnrealTargetPlatform.Win64)
diff --git a/Engine/Plugins/Media/AppleProResDecoderElectra/Source/AppleProResDecoderElectra/Private/ProResDecoder/Mac/ElectraMediaProResDecoder.cpp b/Engine/Plugins/Media/AppleProResDecoderElectra/Source/AppleProResDecoderElectra/Private/ProResDecoder/Mac/ElectraMediaProResDecoder.cpp
--- a/Engine/Plugins/Media/AppleProResDecoderElectra/Source/AppleProResDecoderElectra/Private/ProResDecoder/Mac/ElectraMediaProResDecoder.cpp
+++ b/Engine/Plugins/Media/AppleProResDecoderElectra/Source/AppleProResDecoderElectra/Private/ProResDecoder/Mac/ElectraMediaProResDecoder.cpp
@@ -13,7 +13,7 @@
 #include "IElectraDecoder.h"
 #include "IElectraDecoderOutputVideo.h"
 #include "ElectraDecodersUtils.h"
-#include "Utilities/UtilitiesMP4.h"	// to read the ProRes header, which is given in the form of an MP4 Atom
+#include "MP4Utilities.h"	// to read the ProRes header, which is given in the form of an MP4 Atom
 #include "Utils/ElectraBitstreamProcessorGenericVideo.h"
 #include "AppleProResDecoderElectraModule.h"
 #include "Apple/AppleElectraDecoderPlatformOutputFormatTypes.h"
@@ -725,10 +725,10 @@ bool FVideoDecoderProResElectra::ParseHeader(const void* InData, int32 InNumData
 		return false;
 	}
 	CurrentHeader.Reset();
-	Electra::UtilitiesMP4::FMP4AtomReader Reader(MakeConstArrayView<const uint8>(reinterpret_cast<const uint8*>(InData), InNumData));
+	MP4Utilities::FMP4AtomReader Reader(MakeConstArrayView<const uint8>(reinterpret_cast<const uint8*>(InData), InNumData));
 	Reader.Read(CurrentHeader.FrameSizeInBytes);
 	Reader.Read(CurrentHeader.FrameType);
-	if (CurrentHeader.FrameType != Electra::UtilitiesMP4::MakeBoxAtom('i','c','p','f'))
+	if (CurrentHeader.FrameType != MP4Utilities::MakeBoxAtom('i','c','p','f'))
 	{
 		return false;
 	}
diff --git a/Engine/Plugins/Media/AppleProResDecoderElectra/Source/AppleProResDecoderElectra/Private/ProResDecoder/Windows/ElectraMediaProResDecoder.cpp b/Engine/Plugins/Media/AppleProResDecoderElectra/Source/AppleProResDecoderElectra/Private/ProResDecoder/Windows/ElectraMediaProResDecoder.cpp
--- a/Engine/Plugins/Media/AppleProResDecoderElectra/Source/AppleProResDecoderElectra/Private/ProResDecoder/Windows/ElectraMediaProResDecoder.cpp
+++ b/Engine/Plugins/Media/AppleProResDecoderElectra/Source/AppleProResDecoderElectra/Private/ProResDecoder/Windows/ElectraMediaProResDecoder.cpp
@@ -13,7 +13,7 @@
 #include "IElectraDecoder.h"
 #include "IElectraDecoderOutputVideo.h"
 #include "ElectraDecodersUtils.h"
-#include "Utilities/UtilitiesMP4.h"	// to read the ProRes header, which is given in the form of an MP4 Atom
+#include "MP4Utilities.h"	// to read the ProRes header, which is given in the form of an MP4 Atom
 #include "Utils/ElectraBitstreamProcessorGenericVideo.h"
 
 #include "ElectraDecodersPlatformResources.h"
@@ -585,10 +585,10 @@ bool FVideoDecoderProResElectra::ParseHeader(const void* InData, int32 InNumData
 		return false;
 	}
 	CurrentHeader.Reset();
-	Electra::UtilitiesMP4::FMP4AtomReader Reader(MakeConstArrayView<const uint8>(reinterpret_cast<const uint8*>(InData), InNumData));
+	MP4Utilities::FMP4AtomReader Reader(MakeConstArrayView<const uint8>(reinterpret_cast<const uint8*>(InData), InNumData));
 	Reader.Read(CurrentHeader.FrameSizeInBytes);
 	Reader.Read(CurrentHeader.FrameType);
-	if (CurrentHeader.FrameType != Electra::UtilitiesMP4::MakeBoxAtom('i','c','p','f'))
+	if (CurrentHeader.FrameType != MP4Utilities::MakeBoxAtom('i','c','p','f'))
 	{
 		return false;
 	}
diff --git a/Engine/Plugins/Media/ElectraCodecs/ElectraCodecs.uplugin b/Engine/Plugins/Media/ElectraCodecs/ElectraCodecs.uplugin
--- a/Engine/Plugins/Media/ElectraCodecs/ElectraCodecs.uplugin
+++ b/Engine/Plugins/Media/ElectraCodecs/ElectraCodecs.uplugin
@@ -41,6 +41,10 @@
 		{
 			"Name": "ElectraUtil",
 			"Enabled": true
+		},
+		{
+			"Name": "MP4Utilities",
+			"Enabled": true
 		}
 	]
 }
diff --git a/Engine/Plugins/Media/ElectraCodecs/Source/ElectraDecoders/ElectraDecoders.Build.cs b/Engine/Plugins/Media/ElectraCodecs/Source/ElectraDecoders/ElectraDecoders.Build.cs
--- a/Engine/Plugins/Media/ElectraCodecs/Source/ElectraDecoders/ElectraDecoders.Build.cs
+++ b/Engine/Plugins/Media/ElectraCodecs/Source/ElectraDecoders/ElectraDecoders.Build.cs
@@ -33,6 +33,8 @@ namespace UnrealBuildTool.Rules
 					"MediaUtils",
                     "ElectraSamples",
 					"ElectraBase",
+                    "MP4Utilities",
+                    "MP4Boxes",
                     // Common codecs
 					"dr_libs"					// MPEG audio layer 3
 				});
diff --git a/Engine/Plugins/Media/ElectraCodecs/Source/ElectraDecoders/Private/Utils/ElectraCodecFormatUtils.cpp b/Engine/Plugins/Media/ElectraCodecs/Source/ElectraDecoders/Private/Utils/ElectraCodecFormatUtils.cpp
--- a/Engine/Plugins/Media/ElectraCodecs/Source/ElectraDecoders/Private/Utils/ElectraCodecFormatUtils.cpp
+++ b/Engine/Plugins/Media/ElectraCodecs/Source/ElectraDecoders/Private/Utils/ElectraCodecFormatUtils.cpp
@@ -10,8 +10,8 @@
 #include "Utils/Google/ElectraUtilsVPxVideo.h"
 #include "Utils/AOMedia/ElectraUtilsAV1Video.h"
 
-#include "Utilities/MP4Boxes/MP4Boxes.h"
-#include "Utilities/UtilitiesMP4.h"
+#include "MP4Utilities.h"
+#include "MP4Boxes.h"
 
 DECLARE_LOG_CATEGORY_EXTERN(LogElectraFormatUtils, Log, All);
 DEFINE_LOG_CATEGORY(LogElectraFormatUtils);
@@ -1328,10 +1328,10 @@ static int32 DetermineCodecTypeFromDCR_mp4a(Electra::FCodecTypeFormat& InOut)
 	{
 		if (const TArray<uint8>* wave_Box = InOut.ExtraBoxes.Find(Make4CC('w','a','v','e')))
 		{
-			TArray<TSharedPtr<Electra::UtilitiesMP4::FMP4BoxData, ESPMode::ThreadSafe>> ChildBoxes;
-			if (Electra::UtilitiesMP4::FMP4BoxLocatorReader::LocateAndReadRootBoxesFromBuffer(ChildBoxes, *wave_Box))
+			TArray<TSharedPtr<MP4Utilities::FMP4BoxData, ESPMode::ThreadSafe>> ChildBoxes;
+			if (MP4Utilities::FMP4BoxLocator::LocateAndReadRootBoxesFromBuffer(ChildBoxes, *wave_Box))
 			{
-				if (auto ESDS = ChildBoxes.FindByPredicate([](const TSharedPtr<Electra::UtilitiesMP4::FMP4BoxData, ESPMode::ThreadSafe>& e) { return e->Type == Electra::UtilitiesMP4::MakeBoxAtom('e','s','d','s');} ))
+				if (auto ESDS = ChildBoxes.FindByPredicate([](const TSharedPtr<MP4Utilities::FMP4BoxData, ESPMode::ThreadSafe>& e) { return e->Type == MP4Utilities::MakeBoxAtom('e','s','d','s');} ))
 				{
 					NestedESDSBoxData = (*ESDS)->DataBuffer;
 					esds_Box = &NestedESDSBoxData;
@@ -1555,7 +1555,7 @@ static bool ProcessExtraBoxes(Electra::FCodecTypeFormat& InOut)
 {
 	if (const TArray<uint8>* btrt_Box = InOut.ExtraBoxes.Find(Make4CC('b','t','r','t')))
 	{
-		auto Btrt = StaticCastSharedPtr<Electra::UtilitiesMP4::FMP4BoxBTRT>(Electra::UtilitiesMP4::FMP4BoxBTRT::Create(nullptr, { *btrt_Box }));
+		auto Btrt = StaticCastSharedPtr<MP4Boxes::FMP4BoxBTRT>(MP4Boxes::FMP4BoxBTRT::Create(nullptr, { *btrt_Box }));
 		InOut.Bitrate = Btrt->GetMaxBitrate() ? Btrt->GetMaxBitrate() : InOut.Bitrate;
 		InOut.AverageBitrate = Btrt->GetAverageBitrate() ? Btrt->GetAverageBitrate() : InOut.AverageBitrate;
 	}
diff --git a/Engine/Plugins/Media/ElectraCodecs/Source/ElectraDecoders/Private/Utils/ElectraDecodersUtils.cpp b/Engine/Plugins/Media/ElectraCodecs/Source/ElectraDecoders/Private/Utils/ElectraDecodersUtils.cpp
--- a/Engine/Plugins/Media/ElectraCodecs/Source/ElectraDecoders/Private/Utils/ElectraDecodersUtils.cpp
+++ b/Engine/Plugins/Media/ElectraCodecs/Source/ElectraDecoders/Private/Utils/ElectraDecodersUtils.cpp
@@ -1,7 +1,7 @@
 // Copyright Epic Games, Inc. All Rights Reserved.
 
 #include "ElectraDecodersUtils.h"
-#include "Utilities/UtilitiesMP4.h"
+#include "MP4Utilities.h"
 #include "Utils/Google/ElectraUtilsVPxVideo.h"
 
 
@@ -375,9 +375,9 @@ bool ParseCodecAV1(FMimeTypeVideoCodecInfo& OutInfo, const FString& InCodecForma
 			OutInfo.NumBitsLuma = dcr.GetBitDepth();
 			OutInfo.Extras.Emplace<AV1Video::FAV1CodecConfigurationRecord::FColorInfo>(dcr.GetColorInfo());
 			// Is a 'colr' box provided? And does it contain the mandated 'nclx' box?
-			Electra::UtilitiesMP4::FMP4AtomReader boxReader(IncolrBox);
+			MP4Utilities::FMP4AtomReader boxReader(IncolrBox);
 			uint32 Type;
-			if (boxReader.Read(Type) && Type == Electra::UtilitiesMP4::MakeBoxAtom('n','c','l','x'))
+			if (boxReader.Read(Type) && Type == MP4Utilities::MakeBoxAtom('n','c','l','x'))
 			{
 				uint16 colour_primaries, transfer_characteristics, matrix_coeffs;
 				uint8 video_full_range_flag = 0;
diff --git a/Engine/Plugins/Media/ElectraCodecs/Source/ElectraDecoders/Private/Utils/Google/ElectraBitstreamProcessorVPx.cpp b/Engine/Plugins/Media/ElectraCodecs/Source/ElectraDecoders/Private/Utils/Google/ElectraBitstreamProcessorVPx.cpp
--- a/Engine/Plugins/Media/ElectraCodecs/Source/ElectraDecoders/Private/Utils/Google/ElectraBitstreamProcessorVPx.cpp
+++ b/Engine/Plugins/Media/ElectraCodecs/Source/ElectraDecoders/Private/Utils/Google/ElectraBitstreamProcessorVPx.cpp
@@ -3,7 +3,6 @@
 #include "Utils/Google/ElectraBitstreamProcessorVPx.h"
 #include "Utils/Google/ElectraUtilsVPxVideo.h"
 #include "Utils/MPEG/ElectraUtilsMPEGVideo.h"	// because information may be carried in mp4 specific boxes
-#include "Utilities/UtilitiesMP4.h"
 #include "IElectraDecoderFeaturesAndOptions.h"
 #include "ElectraDecodersUtils.h"
 #include "Misc/Optional.h"
diff --git a/Engine/Plugins/Media/ElectraCodecs/Source/ElectraDecoders/Private/Utils/MPEG/ElectraBitstreamProcessor_H264.cpp b/Engine/Plugins/Media/ElectraCodecs/Source/ElectraDecoders/Private/Utils/MPEG/ElectraBitstreamProcessor_H264.cpp
--- a/Engine/Plugins/Media/ElectraCodecs/Source/ElectraDecoders/Private/Utils/MPEG/ElectraBitstreamProcessor_H264.cpp
+++ b/Engine/Plugins/Media/ElectraCodecs/Source/ElectraDecoders/Private/Utils/MPEG/ElectraBitstreamProcessor_H264.cpp
@@ -3,9 +3,9 @@
 #include "Utils/MPEG/ElectraBitstreamProcessor_H264.h"
 #include "Utils/MPEG/ElectraUtilsMPEGVideo_H264.h"
 #include "Utils/MPEG/ElectraUtilsMPEGVideo.h"
-#include "Utilities/UtilitiesMP4.h"
 #include "IElectraDecoderFeaturesAndOptions.h"
 #include "ElectraDecodersUtils.h"
+#include "MP4Utilities.h"
 
 
 class FElectraDecoderBitstreamProcessorH264::FImpl
@@ -275,7 +275,7 @@ IElectraDecoderBitstreamProcessor::EProcessResult FElectraDecoderBitstreamProces
 	uint32* LastNALU = (uint32*)ElectraDecodersUtil::AdvancePointer(NALU, InOutAccessUnit.DataSize);
 	while(NALU < LastNALU)
 	{
-		uint32 naluLen = Electra::GetFromBigEndian(*NALU);
+		uint32 naluLen = MP4Utilities::GetFromBigEndian(*NALU);
 
 		// Check the nal_ref_idc in the NAL unit for dependencies.
 		uint8 nal = *(const uint8*)(NALU + 1);
@@ -315,7 +315,7 @@ IElectraDecoderBitstreamProcessor::EProcessResult FElectraDecoderBitstreamProces
 
 		if (Impl->bReplaceLengthWithStartcode)
 		{
-			*NALU = Electra::GetFromBigEndian(0x00000001U);
+			*NALU = MP4Utilities::GetFromBigEndian(0x00000001U);
 		}
 		NALU = ElectraDecodersUtil::AdvancePointer(NALU, naluLen + 4);
 	}
diff --git a/Engine/Plugins/Media/ElectraCodecs/Source/ElectraDecoders/Private/Utils/MPEG/ElectraBitstreamProcessor_H265.cpp b/Engine/Plugins/Media/ElectraCodecs/Source/ElectraDecoders/Private/Utils/MPEG/ElectraBitstreamProcessor_H265.cpp
--- a/Engine/Plugins/Media/ElectraCodecs/Source/ElectraDecoders/Private/Utils/MPEG/ElectraBitstreamProcessor_H265.cpp
+++ b/Engine/Plugins/Media/ElectraCodecs/Source/ElectraDecoders/Private/Utils/MPEG/ElectraBitstreamProcessor_H265.cpp
@@ -3,7 +3,7 @@
 #include "Utils/MPEG/ElectraBitstreamProcessor_H265.h"
 #include "Utils/MPEG/ElectraUtilsMPEGVideo_H265.h"
 #include "Utils/MPEG/ElectraUtilsMPEGVideo.h"
-#include "Utilities/UtilitiesMP4.h"
+#include "MP4Utilities.h"
 #include "IElectraDecoderFeaturesAndOptions.h"
 #include "ElectraDecodersUtils.h"
 
@@ -278,7 +278,7 @@ IElectraDecoderBitstreamProcessor::EProcessResult FElectraDecoderBitstreamProces
 	uint32* LastNALU = (uint32*)ElectraDecodersUtil::AdvancePointer(NALU, InOutAccessUnit.DataSize);
 	while(NALU < LastNALU)
 	{
-		uint32 naluLen = Electra::GetFromBigEndian(*NALU);
+		uint32 naluLen = MP4Utilities::GetFromBigEndian(*NALU);
 
 		// Check the nal_ref_idc in the NAL unit for dependencies.
 		uint8 nut = *(const uint8*)(NALU + 1);
@@ -309,7 +309,7 @@ IElectraDecoderBitstreamProcessor::EProcessResult FElectraDecoderBitstreamProces
 
 		if (Impl->bReplaceLengthWithStartcode)
 		{
-			*NALU = Electra::GetFromBigEndian(0x00000001U);
+			*NALU = MP4Utilities::GetFromBigEndian(0x00000001U);
 		}
 		NALU = ElectraDecodersUtil::AdvancePointer(NALU, naluLen + 4);
 	}
diff --git a/Engine/Plugins/Media/ElectraCodecs/Source/ElectraDecoders/Private/Utils/MPEG/ElectraUtilsMPEGVideo.cpp b/Engine/Plugins/Media/ElectraCodecs/Source/ElectraDecoders/Private/Utils/MPEG/ElectraUtilsMPEGVideo.cpp
--- a/Engine/Plugins/Media/ElectraCodecs/Source/ElectraDecoders/Private/Utils/MPEG/ElectraUtilsMPEGVideo.cpp
+++ b/Engine/Plugins/Media/ElectraCodecs/Source/ElectraDecoders/Private/Utils/MPEG/ElectraUtilsMPEGVideo.cpp
@@ -2,7 +2,7 @@
 
 #include "Utils/MPEG/ElectraUtilsMPEGVideo.h"
 #include "Utilities/ElectraBitstream.h"
-#include "Utilities/UtilitiesMP4.h"
+#include "MP4Utilities.h"
 #include "Utils/MPEG/ElectraUtilsMPEGVideo_H264.h"
 #include "Utils/MPEG/ElectraUtilsMPEGVideo_H265.h"
 #include "ElectraDecodersUtils.h"
@@ -174,7 +174,7 @@ namespace ElectraDecodersUtil
 			{
 				return false;
 			}
-			Electra::UtilitiesMP4::FMP4AtomReader boxReader(InMDCVBox);
+			MP4Utilities::FMP4AtomReader boxReader(InMDCVBox);
 			for(int32 i=0; i<3; ++i)
 			{
 				if (!boxReader.Read(OutMDCV.display_primaries_x[i]) || !boxReader.Read(OutMDCV.display_primaries_y[i]))
@@ -244,7 +244,7 @@ namespace ElectraDecodersUtil
 			{
 				return false;
 			}
-			Electra::UtilitiesMP4::FMP4AtomReader boxReader(InCLLIBox);
+			MP4Utilities::FMP4AtomReader boxReader(InCLLIBox);
 			OutCLLI.max_content_light_level = 0;			// MaxCLL
 			OutCLLI.max_pic_average_light_level = 0;		// MaxFALL
 			if (!boxReader.Read(OutCLLI.max_content_light_level) || !boxReader.Read(OutCLLI.max_pic_average_light_level))
@@ -281,17 +281,17 @@ namespace ElectraDecodersUtil
 		{
 			if (InCOLRBox.Num())
 			{
-				Electra::UtilitiesMP4::FMP4AtomReader boxReader(InCOLRBox);
+				MP4Utilities::FMP4AtomReader boxReader(InCOLRBox);
 				uint32 Type;
 				if (boxReader.Read(Type))
 				{
-					if (Type == Electra::UtilitiesMP4::MakeBoxAtom('n','c','l','x') || Type == Electra::UtilitiesMP4::MakeBoxAtom('n', 'c', 'l', 'c'))
+					if (Type == MP4Utilities::MakeBoxAtom('n','c','l','x') || Type == MP4Utilities::MakeBoxAtom('n', 'c', 'l', 'c'))
 					{
 						uint16 colour_primaries, transfer_characteristics, matrix_coeffs;
 						uint8 video_full_range_flag = 0;
 						if (boxReader.Read(colour_primaries) && boxReader.Read(transfer_characteristics) && boxReader.Read(matrix_coeffs))
 						{
-							if (Type == Electra::UtilitiesMP4::MakeBoxAtom('n','c','l','x') && !boxReader.Read(video_full_range_flag))
+							if (Type == MP4Utilities::MakeBoxAtom('n','c','l','x') && !boxReader.Read(video_full_range_flag))
 							{
 								return false;
 							}
@@ -311,7 +311,7 @@ namespace ElectraDecodersUtil
 		TArray<uint8> ELECTRADECODERS_API BuildCOLRBox(const FCommonColorimetry& InFromColorimetry)
 		{
 			Electra::FBitstreamWriter bw;
-			bw.PutBits(Electra::UtilitiesMP4::MakeBoxAtom('n','c','l','x'), 32);
+			bw.PutBits(MP4Utilities::MakeBoxAtom('n','c','l','x'), 32);
 			bw.PutBits((uint16)InFromColorimetry.colour_primaries, 16);
 			bw.PutBits((uint16)InFromColorimetry.transfer_characteristics, 16);
 			bw.PutBits((uint16)InFromColorimetry.matrix_coeffs, 16);
diff --git a/Engine/Plugins/Media/ElectraCodecs/Source/ElectraDecoders/Private/Utils/OpenAPV/ElectraBitstreamProcessor_APV.cpp b/Engine/Plugins/Media/ElectraCodecs/Source/ElectraDecoders/Private/Utils/OpenAPV/ElectraBitstreamProcessor_APV.cpp
--- a/Engine/Plugins/Media/ElectraCodecs/Source/ElectraDecoders/Private/Utils/OpenAPV/ElectraBitstreamProcessor_APV.cpp
+++ b/Engine/Plugins/Media/ElectraCodecs/Source/ElectraDecoders/Private/Utils/OpenAPV/ElectraBitstreamProcessor_APV.cpp
@@ -1,9 +1,7 @@
 // Copyright Epic Games, Inc. All Rights Reserved.
 
 #include "Utils/OpenAPV/ElectraBitstreamProcessor_APV.h"
-//#include "Utils/OpenAPV/ElectraUtilsAPVVideo.h"
 #include "Utils/MPEG/ElectraUtilsMPEGVideo.h"
-#include "Utilities/UtilitiesMP4.h"
 #include "IElectraDecoderFeaturesAndOptions.h"
 #include "ElectraDecodersUtils.h"
 
diff --git a/Engine/Plugins/Media/ElectraCodecs/Source/ElectraDecoders/Private/Utils/OpenAPV/ElectraUtilsAPVVideo.cpp b/Engine/Plugins/Media/ElectraCodecs/Source/ElectraDecoders/Private/Utils/OpenAPV/ElectraUtilsAPVVideo.cpp
--- a/Engine/Plugins/Media/ElectraCodecs/Source/ElectraDecoders/Private/Utils/OpenAPV/ElectraUtilsAPVVideo.cpp
+++ b/Engine/Plugins/Media/ElectraCodecs/Source/ElectraDecoders/Private/Utils/OpenAPV/ElectraUtilsAPVVideo.cpp
@@ -2,7 +2,7 @@
 
 #include "Utils/OpenAPV/ElectraUtilsAPVVideo.h"
 #include "Utilities/ElectraBitstream.h"
-#include "Utilities/UtilitiesMP4.h"
+#include "MP4Utilities.h"
 
 
 namespace ElectraDecodersUtil
@@ -144,9 +144,9 @@ namespace ElectraDecodersUtil
 					fi.ProfileIDC = *InapvCBox++;
 					fi.LevelIDC = *InapvCBox++;
 					fi.BandIDC = *InapvCBox++;
-					fi.FrameWidth = Electra::GetFromBigEndian(*reinterpret_cast<const uint32*>(InapvCBox));
+					fi.FrameWidth = MP4Utilities::GetFromBigEndian(*reinterpret_cast<const uint32*>(InapvCBox));
 					InapvCBox += 4;
-					fi.FrameHeight = Electra::GetFromBigEndian(*reinterpret_cast<const uint32*>(InapvCBox));
+					fi.FrameHeight = MP4Utilities::GetFromBigEndian(*reinterpret_cast<const uint32*>(InapvCBox));
 					InapvCBox += 4;
 					fi.ChromaFormatIDC = *InapvCBox >> 4;
 					fi.BitDepthMinus8 = *InapvCBox++ & 15;
@@ -221,14 +221,14 @@ namespace ElectraDecodersUtil
 			const uint8* EndOfPBUs = CurrentPBU + InAccessUnit.Num();
 			bool bHavePrimary = false;
 			// Skip over size and signature if present.
-			if (InAccessUnit.Num() >= 8 && Electra::GetFromBigEndian(reinterpret_cast<const uint32*>(CurrentPBU)[1]) == 0x61507631)
+			if (InAccessUnit.Num() >= 8 && MP4Utilities::GetFromBigEndian(reinterpret_cast<const uint32*>(CurrentPBU)[1]) == 0x61507631)
 			{
 				CurrentPBU += 8;
 			}
 			TMap<uint8, TArray<FPBU>> PBUs;
 			while(CurrentPBU < EndOfPBUs)
 			{
-				const uint32 PBUsize = Electra::GetFromBigEndian(*reinterpret_cast<const uint32*>(CurrentPBU));
+				const uint32 PBUsize = MP4Utilities::GetFromBigEndian(*reinterpret_cast<const uint32*>(CurrentPBU));
 				FPBU pbu;
 				pbu.Type = CurrentPBU[4];
 				pbu.Data = MakeConstArrayView(CurrentPBU + 4, PBUsize);
@@ -267,9 +267,10 @@ namespace ElectraDecodersUtil
 							return false;
 						}
 					}
-					if (cfgE && InType == 1)
+					// Record convenience values for primary frame only, or any type of frame if there is only one.
+					if (cfgE && (InType == 1 || PBUs.Num() == 1))
 					{
-						bHavePrimary = true;
+						bHavePrimary = InType == 1;
 						// Set the convenience values from the primary frame
 						Width = cfgE->FrameInfos[0].FrameWidth;
 						Height = cfgE->FrameInfos[0].FrameHeight;
@@ -279,7 +280,7 @@ namespace ElectraDecodersUtil
 						Band = cfgE->FrameInfos[0].BandIDC;
 						BitDepth = 8 + cfgE->FrameInfos[0].BitDepthMinus8;
 						// There can be at most one primary frame.
-						bIsNonConformant = bIsNonConformant | (cfgE->FrameInfos.Num() > 1);
+						bIsNonConformant = bIsNonConformant | (InType == 1 && cfgE->FrameInfos.Num() > 1);
 					}
 				}
 				return true;
@@ -339,18 +340,18 @@ namespace ElectraDecodersUtil
 			const uint8* CurrentPBU = AUBase;
 			const uint8* EndOfPBUs = CurrentPBU + InAccessUnit.Num();
 			// Skip over size and signature if present.
-			if (InAccessUnit.Num() >= 8 && Electra::GetFromBigEndian(reinterpret_cast<const uint32*>(CurrentPBU)[1]) == 0x61507631)
+			if (InAccessUnit.Num() >= 8 && MP4Utilities::GetFromBigEndian(reinterpret_cast<const uint32*>(CurrentPBU)[1]) == 0x61507631)
 			{
 				CurrentPBU += 8;
 			}
 			// If there is no size, skip over signature if present.
-			else if (InAccessUnit.Num() >= 4 && Electra::GetFromBigEndian(reinterpret_cast<const uint32*>(CurrentPBU)[0]) == 0x61507631)
+			else if (InAccessUnit.Num() >= 4 && MP4Utilities::GetFromBigEndian(reinterpret_cast<const uint32*>(CurrentPBU)[0]) == 0x61507631)
 			{
 				CurrentPBU += 4;
 			}
 			while(CurrentPBU < EndOfPBUs)
 			{
-				const uint32 PBUsize = Electra::GetFromBigEndian(*reinterpret_cast<const uint32*>(CurrentPBU));
+				const uint32 PBUsize = MP4Utilities::GetFromBigEndian(*reinterpret_cast<const uint32*>(CurrentPBU));
 				const uint8 Type = CurrentPBU[4];
 				// Any frame type PBU?
 				if (Type == 1 || Type == 2 || Type == 25 || Type == 26 || Type == 27)
diff --git a/Engine/Plugins/Media/ElectraCodecs/Source/ElectraDecoders/Private/Utils/VideoDecoderHelpers.cpp b/Engine/Plugins/Media/ElectraCodecs/Source/ElectraDecoders/Private/Utils/VideoDecoderHelpers.cpp
--- a/Engine/Plugins/Media/ElectraCodecs/Source/ElectraDecoders/Private/Utils/VideoDecoderHelpers.cpp
+++ b/Engine/Plugins/Media/ElectraCodecs/Source/ElectraDecoders/Private/Utils/VideoDecoderHelpers.cpp
@@ -1,7 +1,7 @@
 // Copyright Epic Games, Inc. All Rights Reserved.
 
 #include "Utils/VideoDecoderHelpers.h"
-#include "Utilities/UtilitiesMP4.h"
+#include "MP4Utilities.h"
 #include "MediaDecoderOutput.h"
 
 namespace Electra
@@ -31,13 +31,13 @@ void FColorimetryHelper::Update(uint8 colour_primaries, uint8 transfer_character
 
 void FColorimetryHelper::Update(const TArray<uint8>& InFromCOLRBox)
 {
-	Electra::UtilitiesMP4::FMP4AtomReader boxReader(InFromCOLRBox);
+	MP4Utilities::FMP4AtomReader boxReader(InFromCOLRBox);
 	uint32 Type;
 	if (!boxReader.Read(Type))
 	{
 		return;
 	}
-	if (Type != Electra::UtilitiesMP4::MakeBoxAtom('n','c','l','x') && Type != Electra::UtilitiesMP4::MakeBoxAtom('n', 'c', 'l', 'c'))
+	if (Type != MP4Utilities::MakeBoxAtom('n','c','l','x') && Type != MP4Utilities::MakeBoxAtom('n', 'c', 'l', 'c'))
 	{
 		return;
 	}
@@ -47,7 +47,7 @@ void FColorimetryHelper::Update(const TArray<uint8>& InFromCOLRBox)
 	{
 		return;
 	}
-	if (Type == Electra::UtilitiesMP4::MakeBoxAtom('n', 'c', 'l', 'x') && !boxReader.Read(video_full_range_flag))
+	if (Type == MP4Utilities::MakeBoxAtom('n', 'c', 'l', 'x') && !boxReader.Read(video_full_range_flag))
 	{
 		return;
 	}
@@ -258,7 +258,7 @@ void FHDRHelper::UpdateFromMPEGBoxes(int32 BitDepth, const FColorimetryHelper& I
 
 	if (InMDCVBox.Num())
 	{
-		Electra::UtilitiesMP4::FMP4AtomReader boxReader(InMDCVBox);
+		MP4Utilities::FMP4AtomReader boxReader(InMDCVBox);
 
 		struct mastering_display_colour_volume
 		{
@@ -304,7 +304,7 @@ void FHDRHelper::UpdateFromMPEGBoxes(int32 BitDepth, const FColorimetryHelper& I
 
 	if (InCLLIBox.Num())
 	{
-		Electra::UtilitiesMP4::FMP4AtomReader boxReader(InCLLIBox);
+		MP4Utilities::FMP4AtomReader boxReader(InCLLIBox);
 
 		uint16 max_content_light_level = 0;			// MaxCLL
 		uint16 max_pic_average_light_level = 0;		// MaxFALL
diff --git a/Engine/Plugins/Media/ElectraPlayer/ElectraPlayer.uplugin b/Engine/Plugins/Media/ElectraPlayer/ElectraPlayer.uplugin
--- a/Engine/Plugins/Media/ElectraPlayer/ElectraPlayer.uplugin
+++ b/Engine/Plugins/Media/ElectraPlayer/ElectraPlayer.uplugin
@@ -89,6 +89,10 @@
       "Name": "ElectraSubtitles",
       "Enabled": true
     },
+    {
+      "Name": "MP4Utilities",
+      "Enabled": true
+    },
     {
       "Name": "OpusDecoderElectra",
       "Enabled": true
diff --git a/Engine/Plugins/Media/ElectraPlayer/Source/ElectraPlayerRuntime/ElectraPlayerRuntime.Build.cs b/Engine/Plugins/Media/ElectraPlayer/Source/ElectraPlayerRuntime/ElectraPlayerRuntime.Build.cs
--- a/Engine/Plugins/Media/ElectraPlayer/Source/ElectraPlayerRuntime/ElectraPlayerRuntime.Build.cs
+++ b/Engine/Plugins/Media/ElectraPlayer/Source/ElectraPlayerRuntime/ElectraPlayerRuntime.Build.cs
@@ -30,6 +30,7 @@ namespace UnrealBuildTool.Rules
 					"ElectraHTTPStream",
 					"ElectraCDM",
 					"ElectraSubtitles",
+                    "MP4Utilities",
 					"XmlParser",
 					"SoundTouchZ",
                     "HTTP"
diff --git a/Engine/Plugins/Media/ElectraPlayer/Source/ElectraPlayerRuntime/Private/Runtime/Demuxer/ParserISO13818-1.cpp b/Engine/Plugins/Media/ElectraPlayer/Source/ElectraPlayerRuntime/Private/Runtime/Demuxer/ParserISO13818-1.cpp
--- a/Engine/Plugins/Media/ElectraPlayer/Source/ElectraPlayerRuntime/Private/Runtime/Demuxer/ParserISO13818-1.cpp
+++ b/Engine/Plugins/Media/ElectraPlayer/Source/ElectraPlayerRuntime/Private/Runtime/Demuxer/ParserISO13818-1.cpp
@@ -2,7 +2,6 @@
 
 #include "Demuxer/ParserISO13818-1.h"
 #include "Utilities/ElectraBitstream.h"
-#include "Utilities/UtilitiesMP4.h"
 #include "Utils/MPEG/ElectraUtilsMPEGAudio.h"
 #include "Utils/MPEG/ElectraUtilsMPEGVideo_H264.h"
 #include "Utils/MPEG/ElectraUtilsMPEGVideo_H265.h"
diff --git a/Engine/Plugins/Media/ElectraPlayer/Source/ElectraPlayerRuntime/Private/Runtime/Player/StreamSegmentReaderCommon.cpp b/Engine/Plugins/Media/ElectraPlayer/Source/ElectraPlayerRuntime/Private/Runtime/Player/StreamSegmentReaderCommon.cpp
--- a/Engine/Plugins/Media/ElectraPlayer/Source/ElectraPlayerRuntime/Private/Runtime/Player/StreamSegmentReaderCommon.cpp
+++ b/Engine/Plugins/Media/ElectraPlayer/Source/ElectraPlayerRuntime/Private/Runtime/Player/StreamSegmentReaderCommon.cpp
@@ -3,6 +3,7 @@
 #include "StreamSegmentReaderCommon.h"
 #include "PlayerCore.h"
 #include "HTTP/HTTPManager.h"
+#include "MP4Utilities.h"
 #include "Demuxer/ParserISO14496-12.h"
 #include "Demuxer/ParserISO14496-12_Utils.h"
 #include "Demuxer/ParserISO13818-1.h"
@@ -778,10 +779,10 @@ void FStreamSegmentReaderCommon::FStreamHandler::HandleRequest()
 			};
 			// ISO/IEC 14496-12 ?
 			const uint32 boxName = GetUINT32BE(ProbeBytes + 4);
-			if (boxName == Electra::UtilitiesMP4::MakeBoxAtom('s','t','y','p') || boxName == Electra::UtilitiesMP4::MakeBoxAtom('m','o','o','f') || boxName == Electra::UtilitiesMP4::MakeBoxAtom('s','i','d','x')
-				|| boxName == Electra::UtilitiesMP4::MakeBoxAtom('f','t','y','p') || boxName == Electra::UtilitiesMP4::MakeBoxAtom('m','o','o','v') || boxName == Electra::UtilitiesMP4::MakeBoxAtom('e','m','s','g')
-				|| boxName == Electra::UtilitiesMP4::MakeBoxAtom('f','r','e','e') || boxName == Electra::UtilitiesMP4::MakeBoxAtom('s','k','i','p') || boxName == Electra::UtilitiesMP4::MakeBoxAtom('s','s','i','x')
-				|| boxName == Electra::UtilitiesMP4::MakeBoxAtom('p','r','f','t') || boxName == Electra::UtilitiesMP4::MakeBoxAtom('u','d','t','a'))
+			if (boxName == MP4Utilities::MakeBoxAtom('s','t','y','p') || boxName == MP4Utilities::MakeBoxAtom('m','o','o','f') || boxName == MP4Utilities::MakeBoxAtom('s','i','d','x')
+				|| boxName == MP4Utilities::MakeBoxAtom('f','t','y','p') || boxName == MP4Utilities::MakeBoxAtom('m','o','o','v') || boxName == MP4Utilities::MakeBoxAtom('e','m','s','g')
+				|| boxName == MP4Utilities::MakeBoxAtom('f','r','e','e') || boxName == MP4Utilities::MakeBoxAtom('s','k','i','p') || boxName == MP4Utilities::MakeBoxAtom('s','s','i','x')
+				|| boxName == MP4Utilities::MakeBoxAtom('p','r','f','t') || boxName == MP4Utilities::MakeBoxAtom('u','d','t','a'))
 			{
 				ContainerFormat = FStreamSegmentRequestCommon::EContainerFormat::ISO_14496_12;
 			}
@@ -1095,8 +1096,8 @@ bool FStreamSegmentReaderCommon::FStreamHandler::FetchInitSegment(TSharedPtrTS<F
 		};
 		// ISO/IEC 14496-12 ?
 		const uint32 boxName = InitDataSize >= 8 ? GetUINT32BE(InitDataPtr + 4) : 0;
-		if (boxName == Electra::UtilitiesMP4::MakeBoxAtom('f','t','y','p') || boxName == Electra::UtilitiesMP4::MakeBoxAtom('f','r','e','e') || boxName == Electra::UtilitiesMP4::MakeBoxAtom('s','k','i','p')
-			|| boxName == Electra::UtilitiesMP4::MakeBoxAtom('p','d','i','n') || boxName == Electra::UtilitiesMP4::MakeBoxAtom('u','d','t','a'))
+		if (boxName == MP4Utilities::MakeBoxAtom('f','t','y','p') || boxName == MP4Utilities::MakeBoxAtom('f','r','e','e') || boxName == MP4Utilities::MakeBoxAtom('s','k','i','p')
+			|| boxName == MP4Utilities::MakeBoxAtom('p','d','i','n') || boxName == MP4Utilities::MakeBoxAtom('u','d','t','a'))
 		{
 			StaticSegmentDataReader.SetSourceBuffer(ResponseBuffer);
 			ContainerFormat = FStreamSegmentRequestCommon::EContainerFormat::ISO_14496_12;
@@ -1132,8 +1133,8 @@ bool FStreamSegmentReaderCommon::FStreamHandler::FetchInitSegment(TSharedPtrTS<F
 		// For self-initializing media we need to locate the 'moov' box.
 		UtilsMP4::FMP4RootBoxLocator BoxLocator;
 		TArray<UtilsMP4::FMP4RootBoxLocator::FBoxInfo> Boxes;
-		const TArray<uint32> FirstBox { Electra::UtilitiesMP4::MakeBoxAtom('f','t','y','p') };
-		const TArray<uint32> MoovBox { Electra::UtilitiesMP4::MakeBoxAtom('m','o','o','v') };
+		const TArray<uint32> FirstBox { MP4Utilities::MakeBoxAtom('f','t','y','p') };
+		const TArray<uint32> MoovBox { MP4Utilities::MakeBoxAtom('m','o','o','v') };
 		FString InitSegmentUrl(InitURLToUse.Url.URL);
 
 		// Set up information for CMCD.
@@ -1168,7 +1169,7 @@ bool FStreamSegmentReaderCommon::FStreamHandler::FetchInitSegment(TSharedPtrTS<F
 			SetError(FString::Printf(TEXT("Self-initializing media 'moov' scan error: %s"),	*ErrMsg), INTERNAL_SEG_ERROR_INIT_SEGMENT_DOWNLOAD_ERROR);
 			return false;
 		}
-		const UtilsMP4::FMP4RootBoxLocator::FBoxInfo* Moov = Boxes.FindByPredicate([InType=Electra::UtilitiesMP4::MakeBoxAtom('m','o','o','v')](const UtilsMP4::FMP4RootBoxLocator::FBoxInfo& InBox){return InBox.Type == InType;});
+		const UtilsMP4::FMP4RootBoxLocator::FBoxInfo* Moov = Boxes.FindByPredicate([InType=MP4Utilities::MakeBoxAtom('m','o','o','v')](const UtilsMP4::FMP4RootBoxLocator::FBoxInfo& InBox){return InBox.Type == InType;});
 		if (Moov)
 		{
 			StaticSegmentDataReader.SetSourceBuffer(Moov->DataBuffer);
diff --git a/Engine/Plugins/Media/ElectraPlayer/Source/ElectraPlayerRuntime/Private/Runtime/Player/mp4/PlaylistReaderMP4.cpp b/Engine/Plugins/Media/ElectraPlayer/Source/ElectraPlayerRuntime/Private/Runtime/Player/mp4/PlaylistReaderMP4.cpp
--- a/Engine/Plugins/Media/ElectraPlayer/Source/ElectraPlayerRuntime/Private/Runtime/Player/mp4/PlaylistReaderMP4.cpp
+++ b/Engine/Plugins/Media/ElectraPlayer/Source/ElectraPlayerRuntime/Private/Runtime/Player/mp4/PlaylistReaderMP4.cpp
@@ -4,6 +4,7 @@
 #include "Player/Manifest.h"
 #include "HTTP/HTTPManager.h"
 #include "HTTP/HTTPResponseCache.h"
+#include "MP4Utilities.h"
 #include "Demuxer/ParserISO14496-12.h"
 #include "Stats/Stats.h"
 #include "SynchronizedClock.h"
@@ -228,9 +229,9 @@ void FPlaylistReaderMP4::WorkerThread()
 	SCOPE_CYCLE_COUNTER(STAT_ElectraPlayer_MP4_PlaylistWorker);
 	CSV_SCOPED_TIMING_STAT(ElectraPlayer, MP4_PlaylistWorker);
 
-	const TArray<uint32> StopAfterBoxes { Electra::UtilitiesMP4::MakeBoxAtom('m','o','o','v') };
-	const TArray<uint32> ReadBoxes { Electra::UtilitiesMP4::MakeBoxAtom('f','t','y','p'), Electra::UtilitiesMP4::MakeBoxAtom('m','o','o','v') };
-	const TArray<uint32> FirstBoxes { Electra::UtilitiesMP4::MakeBoxAtom('f','t','y','p'), Electra::UtilitiesMP4::MakeBoxAtom('s','t','y','p'), Electra::UtilitiesMP4::MakeBoxAtom('s','i','d','x'), Electra::UtilitiesMP4::MakeBoxAtom('f','r','e','e'), Electra::UtilitiesMP4::MakeBoxAtom('s','k','i','p') };
+	const TArray<uint32> StopAfterBoxes { MP4Utilities::MakeBoxAtom('m','o','o','v') };
+	const TArray<uint32> ReadBoxes { MP4Utilities::MakeBoxAtom('f','t','y','p'), MP4Utilities::MakeBoxAtom('m','o','o','v') };
+	const TArray<uint32> FirstBoxes { MP4Utilities::MakeBoxAtom('f','t','y','p'), MP4Utilities::MakeBoxAtom('s','t','y','p'), MP4Utilities::MakeBoxAtom('s','i','d','x'), MP4Utilities::MakeBoxAtom('f','r','e','e'), MP4Utilities::MakeBoxAtom('s','k','i','p') };
 	UtilsMP4::FMP4RootBoxLocator BoxLocator;
 	TArray<UtilsMP4::FMP4RootBoxLocator::FBoxInfo> BoxInfos;
 	TArray<HTTP::FHTTPHeader> ReqHeaders;
@@ -263,10 +264,10 @@ void FPlaylistReaderMP4::WorkerThread()
 	else if (bGotBoxes && !bAbort && !bHasErrored)
 	{
 		// Do we have the `ftyp` and `moov` boxes?
-		if (BoxInfos.ContainsByPredicate([InType=Electra::UtilitiesMP4::MakeBoxAtom('f','t','y','p')](const UtilsMP4::FMP4RootBoxLocator::FBoxInfo& InBox){return InBox.Type == InType;}) &&
-			BoxInfos.ContainsByPredicate([InType=Electra::UtilitiesMP4::MakeBoxAtom('m','o','o','v')](const UtilsMP4::FMP4RootBoxLocator::FBoxInfo& InBox){return InBox.Type == InType;}))
+		if (BoxInfos.ContainsByPredicate([InType=MP4Utilities::MakeBoxAtom('f','t','y','p')](const UtilsMP4::FMP4RootBoxLocator::FBoxInfo& InBox){return InBox.Type == InType;}) &&
+			BoxInfos.ContainsByPredicate([InType=MP4Utilities::MakeBoxAtom('m','o','o','v')](const UtilsMP4::FMP4RootBoxLocator::FBoxInfo& InBox){return InBox.Type == InType;}))
 		{
-			if (BoxInfos.ContainsByPredicate([InType=Electra::UtilitiesMP4::MakeBoxAtom('m','d','a','t')](const UtilsMP4::FMP4RootBoxLocator::FBoxInfo& InBox){return InBox.Type == InType;}))
+			if (BoxInfos.ContainsByPredicate([InType=MP4Utilities::MakeBoxAtom('m','d','a','t')](const UtilsMP4::FMP4RootBoxLocator::FBoxInfo& InBox){return InBox.Type == InType;}))
 			{
 				LogMessage(IInfoLog::ELevel::Info, FString::Printf(TEXT("The mp4 at \"%s\" is not fast-startable. Consider moving the 'moov' box in front of the 'mdat' for faster startup times."), *ConnectionInfo.EffectiveURL));
 			}
@@ -274,14 +275,14 @@ void FPlaylistReaderMP4::WorkerThread()
 			MP4Parser = IParserISO14496_12::CreateParser();
 
 			// First parse the ftyp on which some elements in the moov depend.
-			ParseBuffer = BoxInfos.FindByPredicate([InType=Electra::UtilitiesMP4::MakeBoxAtom('f','t','y','p')](const UtilsMP4::FMP4RootBoxLocator::FBoxInfo& InBox){return InBox.Type == InType;})->DataBuffer;
+			ParseBuffer = BoxInfos.FindByPredicate([InType=MP4Utilities::MakeBoxAtom('f','t','y','p')](const UtilsMP4::FMP4RootBoxLocator::FBoxInfo& InBox){return InBox.Type == InType;})->DataBuffer;
 			ParsePos = 0;
 			ParseBufferSize = ParseBuffer->Num();
 			UEMediaError parseError = MP4Parser->ParseHeader(this, this, PlayerSessionServices, nullptr);
 			// If successful parse the moov
 			if (parseError == UEMEDIA_ERROR_OK || parseError == UEMEDIA_ERROR_END_OF_STREAM)
 			{
-				ParseBuffer = BoxInfos.FindByPredicate([InType=Electra::UtilitiesMP4::MakeBoxAtom('m','o','o','v')](const UtilsMP4::FMP4RootBoxLocator::FBoxInfo& InBox){return InBox.Type == InType;})->DataBuffer;
+				ParseBuffer = BoxInfos.FindByPredicate([InType=MP4Utilities::MakeBoxAtom('m','o','o','v')](const UtilsMP4::FMP4RootBoxLocator::FBoxInfo& InBox){return InBox.Type == InType;})->DataBuffer;
 				ParsePos = 0;
 				ParseBufferSize = ParseBuffer->Num();
 				parseError = MP4Parser->ParseHeader(this, this, PlayerSessionServices, nullptr);
diff --git a/Engine/Plugins/Media/ElectraPlayer/Source/ElectraPlayerRuntime/Private/Runtime/Utilities/UtilsMP4.cpp b/Engine/Plugins/Media/ElectraPlayer/Source/ElectraPlayerRuntime/Private/Runtime/Utilities/UtilsMP4.cpp
--- a/Engine/Plugins/Media/ElectraPlayer/Source/ElectraPlayerRuntime/Private/Runtime/Utilities/UtilsMP4.cpp
+++ b/Engine/Plugins/Media/ElectraPlayer/Source/ElectraPlayerRuntime/Private/Runtime/Utilities/UtilsMP4.cpp
@@ -1,8 +1,9 @@
 // Copyright Epic Games, Inc. All Rights Reserved.
 
 #include "Utilities/UtilsMP4.h"
+#include "Misc/Base64.h"
+#include "MP4Utilities.h"
 #include "HTTP/HTTPResponseCache.h"
-#include <Misc/Base64.h>
 
 
 namespace Electra
@@ -24,14 +25,14 @@ UtilsMP4::FMetadataParser::FMetadataParser()
 {
 	NumTotalItems = 0;
 
-	WellKnownItems.Emplace(Electra::UtilitiesMP4::MakeBoxAtom(0xa9u,'n','a','m'), FMetadataTags::Title);
-	WellKnownItems.Emplace(Electra::UtilitiesMP4::MakeBoxAtom(0xa9u,'A','R','T'), FMetadataTags::Artist);
-	WellKnownItems.Emplace(Electra::UtilitiesMP4::MakeBoxAtom(0xa9u,'g','e','n'), FMetadataTags::Genre);
-	WellKnownItems.Emplace(Electra::UtilitiesMP4::MakeBoxAtom(0xa9u,'d','a','y'), FMetadataTags::Date);
-	WellKnownItems.Emplace(Electra::UtilitiesMP4::MakeBoxAtom('d','e','s','c'), FMetadataTags::Description);
-	WellKnownItems.Emplace(Electra::UtilitiesMP4::MakeBoxAtom('l','d','e','s'), FMetadataTags::LongDescription);
-	WellKnownItems.Emplace(Electra::UtilitiesMP4::MakeBoxAtom(0xa9u,'a','l','b'), FMetadataTags::Album);
-	WellKnownItems.Emplace(Electra::UtilitiesMP4::MakeBoxAtom(0xa9u,'t','o','o'), FMetadataTags::Encoder);
+	WellKnownItems.Emplace(MP4Utilities::MakeBoxAtom(0xa9u,'n','a','m'), FMetadataTags::Title);
+	WellKnownItems.Emplace(MP4Utilities::MakeBoxAtom(0xa9u,'A','R','T'), FMetadataTags::Artist);
+	WellKnownItems.Emplace(MP4Utilities::MakeBoxAtom(0xa9u,'g','e','n'), FMetadataTags::Genre);
+	WellKnownItems.Emplace(MP4Utilities::MakeBoxAtom(0xa9u,'d','a','y'), FMetadataTags::Date);
+	WellKnownItems.Emplace(MP4Utilities::MakeBoxAtom('d','e','s','c'), FMetadataTags::Description);
+	WellKnownItems.Emplace(MP4Utilities::MakeBoxAtom('l','d','e','s'), FMetadataTags::LongDescription);
+	WellKnownItems.Emplace(MP4Utilities::MakeBoxAtom(0xa9u,'a','l','b'), FMetadataTags::Album);
+	WellKnownItems.Emplace(MP4Utilities::MakeBoxAtom(0xa9u,'t','o','o'), FMetadataTags::Encoder);
 }
 
 FString UtilsMP4::FMetadataParser::PrintableBoxAtom(const uint32 InAtom)
@@ -90,7 +91,7 @@ void UtilsMP4::FMetadataParser::AddItem(const FString& InType, const FString& In
 UtilsMP4::FMetadataParser::EResult UtilsMP4::FMetadataParser::Parse(uint32 InHandler, uint32 InHandlerReserved0, const TArray<UtilsMP4::FMetadataParser::FBoxInfo>& InBoxes)
 {
 	// We only support the Apple iTunes metadata at the moment.
-	if (InHandler != Electra::UtilitiesMP4::MakeBoxAtom('m','d','i','r'))
+	if (InHandler != MP4Utilities::MakeBoxAtom('m','d','i','r'))
 	{
 		return UtilsMP4::FMetadataParser::EResult::NotSupported;
 	}
@@ -122,7 +123,7 @@ UtilsMP4::FMetadataParser::EResult UtilsMP4::FMetadataParser::Parse(uint32 InHan
 	bool bHaveILST = false;
 	for(int32 i=0; i<InBoxes.Num(); ++i)
 	{
-		if (InBoxes[i].Type == Electra::UtilitiesMP4::MakeBoxAtom('i','l','s','t'))
+		if (InBoxes[i].Type == MP4Utilities::MakeBoxAtom('i','l','s','t'))
 		{
 			// We will actually parse all 'ilst' boxes if for some reason the metadata has
 			// been scattered across several.
@@ -140,7 +141,7 @@ void UtilsMP4::FMetadataParser::Parse(const UtilsMP4::FMetadataParser::FBoxInfo&
 		return;
 	}
 
-	Electra::UtilitiesMP4::FMP4AtomReader dr(MakeConstArrayView<const uint8>(reinterpret_cast<const uint8*>(InBox.Data), InBox.Size));
+	MP4Utilities::FMP4AtomReader dr(MakeConstArrayView<const uint8>(reinterpret_cast<const uint8*>(InBox.Data), InBox.Size));
 	while(dr.GetNumBytesRemaining() > 8)
 	{
 		int32 BoxSize;
@@ -153,7 +154,7 @@ void UtilsMP4::FMetadataParser::Parse(const UtilsMP4::FMetadataParser::FBoxInfo&
 		}
 		else switch(BoxType)
 		{
-			case Electra::UtilitiesMP4::MakeBoxAtom('-','-','-','-'):
+			case MP4Utilities::MakeBoxAtom('-','-','-','-'):
 			{
 				ParseBoxDataiTunes(dr.GetCurrentDataPointer(), BoxSize-8);
 				break;
@@ -176,7 +177,7 @@ void UtilsMP4::FMetadataParser::ParseBoxDataList(const FString& AsCategory, cons
 		return;
 	}
 
-	Electra::UtilitiesMP4::FMP4AtomReader dr(MakeConstArrayView<const uint8>(reinterpret_cast<const uint8*>(InBoxData), InBoxSize));
+	MP4Utilities::FMP4AtomReader dr(MakeConstArrayView<const uint8>(reinterpret_cast<const uint8*>(InBoxData), InBoxSize));
 
 #define RETURN_IF_READERROR(expr)	\
 	if (!(expr))					\
@@ -190,7 +191,7 @@ void UtilsMP4::FMetadataParser::ParseBoxDataList(const FString& AsCategory, cons
 		RETURN_IF_READERROR(dr.Read(BoxSize));
 		RETURN_IF_READERROR(dr.Read(BoxType));
 		CurrentOffset = dr.GetCurrentOffset();
-		if (BoxType != Electra::UtilitiesMP4::MakeBoxAtom('d','a','t','a'))
+		if (BoxType != MP4Utilities::MakeBoxAtom('d','a','t','a'))
 		{
 			continue;
 		}
@@ -353,7 +354,7 @@ void UtilsMP4::FMetadataParser::ParseBoxDataiTunes(const uint8* InBoxData, uint3
 		return;
 	}
 
-	Electra::UtilitiesMP4::FMP4AtomReader dr(MakeConstArrayView<const uint8>(reinterpret_cast<const uint8*>(InBoxData), InBoxSize));
+	MP4Utilities::FMP4AtomReader dr(MakeConstArrayView<const uint8>(reinterpret_cast<const uint8*>(InBoxData), InBoxSize));
 
 #define RETURN_IF_READERROR(expr)	\
 	if (!(expr))					\
@@ -368,7 +369,7 @@ void UtilsMP4::FMetadataParser::ParseBoxDataiTunes(const uint8* InBoxData, uint3
 		RETURN_IF_READERROR(dr.Read(BoxSize));
 		RETURN_IF_READERROR(dr.Read(BoxType));
 		CurrentOffset = dr.GetCurrentOffset();
-		if (BoxType == Electra::UtilitiesMP4::MakeBoxAtom('m','e','a','n'))
+		if (BoxType == MP4Utilities::MakeBoxAtom('m','e','a','n'))
 		{
 			// 4 reserved bytes follow
 			RETURN_IF_READERROR(dr.ReadBytes(nullptr, 4));
@@ -380,14 +381,14 @@ void UtilsMP4::FMetadataParser::ParseBoxDataiTunes(const uint8* InBoxData, uint3
 				return;
 			}
 		}
-		else if (BoxType == Electra::UtilitiesMP4::MakeBoxAtom('n','a','m','e'))
+		else if (BoxType == MP4Utilities::MakeBoxAtom('n','a','m','e'))
 		{
 			// 4 reserved bytes follow
 			RETURN_IF_READERROR(dr.ReadBytes(nullptr, 4));
 			int32 NumRemainingDataBytes = BoxSize - 8 - (dr.GetCurrentOffset() - CurrentOffset);
 			RETURN_IF_READERROR(dr.ReadStringUTF8(Name, NumRemainingDataBytes));
 		}
-		else if (BoxType == Electra::UtilitiesMP4::MakeBoxAtom('d','a','t','a'))
+		else if (BoxType == MP4Utilities::MakeBoxAtom('d','a','t','a'))
 		{
 			ParseBoxDataList(FString(TEXT("iTunes."))+Name, dr.GetCurrentDataPointer()-8, BoxSize);
 		}
@@ -687,7 +688,7 @@ bool UtilsMP4::FMP4RootBoxLocator::LocateRootBoxes(TArray<FBoxInfo>& OutBoxInfos
 				BoxInternalOffset += 8;
 			}
 
-			if (bi.Type == UtilitiesMP4::MakeBoxAtom('u','u','i','d'))
+			if (bi.Type == MP4Utilities::MakeBoxAtom('u','u','i','d'))
 			{
 				if (DataBuffer->Num() < BoxInternalOffset + 16)
 				{
diff --git a/Engine/Plugins/Media/ElectraPlayer/Source/ElectraProtron/ElectraProtron.Build.cs b/Engine/Plugins/Media/ElectraPlayer/Source/ElectraProtron/ElectraProtron.Build.cs
--- a/Engine/Plugins/Media/ElectraPlayer/Source/ElectraProtron/ElectraProtron.Build.cs
+++ b/Engine/Plugins/Media/ElectraPlayer/Source/ElectraProtron/ElectraProtron.Build.cs
@@ -21,7 +21,9 @@ namespace UnrealBuildTool.Rules
 					"RHI",
 					"ElectraBase",
                     "ElectraDecoders",
-					"ElectraSamples"
+					"ElectraSamples",
+                    "MP4Utilities",
+                    "MP4Boxes"
 				});
 
 			PrivateIncludePathModuleNames.AddRange(
diff --git a/Engine/Plugins/Media/ElectraPlayer/Source/ElectraProtron/Private/Player/ElectraProtronPlayer.cpp b/Engine/Plugins/Media/ElectraPlayer/Source/ElectraProtron/Private/Player/ElectraProtronPlayer.cpp
--- a/Engine/Plugins/Media/ElectraPlayer/Source/ElectraProtron/Private/Player/ElectraProtronPlayer.cpp
+++ b/Engine/Plugins/Media/ElectraPlayer/Source/ElectraProtron/Private/Player/ElectraProtronPlayer.cpp
@@ -4,8 +4,7 @@
 #include "ElectraProtronPlayerImpl.h"
 #include "ElectraProtronPrivate.h"
 #include "MediaSamples.h"
-#include "Utilities/UtilitiesMP4.h"
-#include "Utilities/MP4Boxes/MP4Boxes.h"
+#include "MP4Boxes.h"
 #include "ElectraTextureSample.h"
 #include "IElectraAudioSample.h"
 
diff --git a/Engine/Plugins/Media/ElectraPlayer/Source/ElectraProtron/Private/Player/ElectraProtronPlayerDecoder.cpp b/Engine/Plugins/Media/ElectraPlayer/Source/ElectraProtron/Private/Player/ElectraProtronPlayerDecoder.cpp
--- a/Engine/Plugins/Media/ElectraPlayer/Source/ElectraProtron/Private/Player/ElectraProtronPlayerDecoder.cpp
+++ b/Engine/Plugins/Media/ElectraPlayer/Source/ElectraProtron/Private/Player/ElectraProtronPlayerDecoder.cpp
@@ -344,7 +344,7 @@ void FElectraProtronPlayer::FImpl::FDecoderThread::UpdateTrackSampleDurationMap(
 		FTimespan Duration;
 	};
 	TArray<FSampleTimes> SampleTimeDurations;
-	TSharedPtr<Electra::UtilitiesMP4::FMP4Track::FIterator, ESPMode::ThreadSafe> TkIt = TrackSampleBuffer->TrackAndCodecInfo->MP4Track->CreateIterator();
+	TSharedPtr<MP4Boxes::FMP4Track::FIterator, ESPMode::ThreadSafe> TkIt = TrackSampleBuffer->TrackAndCodecInfo->MP4Track->CreateIterator();
 	SampleTimeDurations.Reserve(TkIt->GetNumSamples());
 	auto It = TrackSampleBuffer->TrackAndCodecInfo->MP4Track->CreateIterator();
 	do
@@ -373,9 +373,7 @@ void FElectraProtronPlayer::FImpl::FDecoderThread::UpdateTrackIterator(const FTi
 		return;
 	}
 
-	TSharedPtr<Electra::UtilitiesMP4::FMP4Track::FIterator, ESPMode::ThreadSafe> TkIt =
-		TrackSampleBuffer->TrackAndCodecInfo->MP4Track->CreateIteratorAtKeyframe(Electra::FTimeValue().SetFromTimespan(InForTime),
-			Electra::FTimeValue().SetFromMilliseconds(TrackSampleBuffer->TrackAndCodecInfo->bIsKeyframeOnlyFormat ? 0 : Config.NextKeyframeThresholdMillis));
+	TSharedPtr<MP4Boxes::FMP4Track::FIterator, ESPMode::ThreadSafe> TkIt = TrackSampleBuffer->TrackAndCodecInfo->MP4Track->CreateIteratorAtKeyframe(InForTime, FTimespan::FromMilliseconds(TrackSampleBuffer->TrackAndCodecInfo->bIsKeyframeOnlyFormat ? 0 : Config.NextKeyframeThresholdMillis));
 	if (!TkIt.IsValid())
 	{
 		LastErrorMessage = TrackSampleBuffer->TrackAndCodecInfo->MP4Track->GetLastError();
@@ -449,7 +447,7 @@ void FElectraProtronPlayer::FImpl::FDecoderThread::HandlePlaybackRangeChanges()
 	if (!FirstRangeSampleIt.IsValid())
 	{
 		// Locate the first and last sample numbers for the range
-		FTrackIterator RangeIt = TrackSampleBuffer->TrackAndCodecInfo->MP4Track->CreateIteratorAtKeyframe(Electra::FTimeValue().SetFromTimespan(PlaybackRange.GetLowerBoundValue()), Electra::FTimeValue::GetZero());
+		FTrackIterator RangeIt = TrackSampleBuffer->TrackAndCodecInfo->MP4Track->CreateIteratorAtKeyframe(PlaybackRange.GetLowerBoundValue(), FTimespan(0));
 		if (!RangeIt.IsValid())
 		{
 			LastErrorMessage = TrackSampleBuffer->TrackAndCodecInfo->MP4Track->GetLastError();
@@ -1187,7 +1185,7 @@ bool FElectraProtronPlayer::FImpl::FDecoderThread::CreateDecoder()
 	// Add the extra boxes found in the sample description.
 	for(auto& boxIt : ci.ExtraBoxes)
 	{
-		FString BoxName = FString::Printf(TEXT("$%s_box"), *Electra::UtilitiesMP4::Printable4CC(boxIt.Key));
+		FString BoxName = FString::Printf(TEXT("$%s_box"), *MP4Utilities::Printable4CC(boxIt.Key));
 		CurrentCodecSpecificData.Add(BoxName, FVariant(boxIt.Value));
 	}
 	DecoderCfgParams.Append(CurrentCodecSpecificData);
diff --git a/Engine/Plugins/Media/ElectraPlayer/Source/ElectraProtron/Private/Player/ElectraProtronPlayerImpl.cpp b/Engine/Plugins/Media/ElectraPlayer/Source/ElectraProtron/Private/Player/ElectraProtronPlayerImpl.cpp
--- a/Engine/Plugins/Media/ElectraPlayer/Source/ElectraProtron/Private/Player/ElectraProtronPlayerImpl.cpp
+++ b/Engine/Plugins/Media/ElectraPlayer/Source/ElectraProtron/Private/Player/ElectraProtronPlayerImpl.cpp
@@ -229,7 +229,7 @@ void FElectraProtronPlayer::FImpl::InternalOpen(const FString& InFilename)
 	}
 
 	// Open the file.
-	TSharedPtr<Electra::IFileDataReader, ESPMode::ThreadSafe> Reader = Electra::IFileDataReader::Create();
+	TSharedPtr<MP4Utilities::FMP4FileDataReader, ESPMode::ThreadSafe> Reader = MP4Utilities::FMP4FileDataReader::Create();
 	if (!Reader->Open(InFilename))
 	{
 		LastErrorMessage = Reader->GetLastError();
@@ -237,19 +237,19 @@ void FElectraProtronPlayer::FImpl::InternalOpen(const FString& InFilename)
 	}
 
 	// Read the mp4 box structure.
-	Electra::UtilitiesMP4::FMP4BoxLocatorReader BoxLocator;
-	const TArray<uint32> FirstBoxes { Electra::UtilitiesMP4::MakeBoxAtom('f','t','y','p'), Electra::UtilitiesMP4::MakeBoxAtom('s','t','y','p'), Electra::UtilitiesMP4::MakeBoxAtom('s','i','d','x'), Electra::UtilitiesMP4::MakeBoxAtom('f','r','e','e'), Electra::UtilitiesMP4::MakeBoxAtom('s','k','i','p') };
+	MP4Utilities::FMP4BoxLocator BoxLocator;
+	const TArray<uint32> FirstBoxes { MP4Utilities::MakeBoxAtom('f','t','y','p'), MP4Utilities::MakeBoxAtom('s','t','y','p'), MP4Utilities::MakeBoxAtom('s','i','d','x'), MP4Utilities::MakeBoxAtom('f','r','e','e'), MP4Utilities::MakeBoxAtom('s','k','i','p') };
 	const TArray<uint32> ReadBoxes;			// Empty means to read in all boxes except `mdat`
 	const TArray<uint32> StopAfterBoxes;	// Empty means to read all boxes and not stop after a specific one.
-	TArray<TSharedPtr<Electra::UtilitiesMP4::FMP4BoxData, ESPMode::ThreadSafe>> RootBoxes;
-	if (!BoxLocator.LocateAndReadRootBoxes(RootBoxes, Reader, FirstBoxes, StopAfterBoxes, ReadBoxes, Electra::IBaseDataReader::FCancellationCheckDelegate::CreateLambda([&](){return bAbort;})))
+	TArray<TSharedPtr<MP4Utilities::FMP4BoxData, ESPMode::ThreadSafe>> RootBoxes;
+	if (!BoxLocator.LocateAndReadRootBoxes(RootBoxes, Reader, FirstBoxes, StopAfterBoxes, ReadBoxes, MP4Utilities::IMP4DataReaderBase::FCancellationCheckDelegate::CreateLambda([&](){return bAbort;})))
 	{
 		LastErrorMessage = BoxLocator.GetLastError();
 		return;
 	}
 
 	// In order to be usable the mp4 needs to have a `moov` box.
-	if (!RootBoxes.FindByPredicate([](const TSharedPtr<Electra::UtilitiesMP4::FMP4BoxData, ESPMode::ThreadSafe>& In){ return In->Type == Electra::UtilitiesMP4::MakeBoxAtom('m','o','o','v'); }))
+	if (!RootBoxes.FindByPredicate([](const TSharedPtr<MP4Utilities::FMP4BoxData, ESPMode::ThreadSafe>& In){ return In->Type == MP4Utilities::MakeBoxAtom('m','o','o','v'); }))
 	{
 		LastErrorMessage = FString::Printf(TEXT("No `moov` box found in this file. It does not appear to be an mp4/mov file."));
 		return;
@@ -259,7 +259,7 @@ void FElectraProtronPlayer::FImpl::InternalOpen(const FString& InFilename)
 	ParsedRootBoxes.SetNum(RootBoxes.Num());
 	for(int32 i=0; i<RootBoxes.Num(); ++i)
 	{
-		Electra::UtilitiesMP4::FMP4BoxTreeParser tp;
+		MP4Boxes::FMP4BoxTreeParser tp;
 		if (!tp.ParseBoxTree(RootBoxes[i]))
 		{
 			LastErrorMessage = FString::Printf(TEXT("Failed to parse the file's box structure."));
@@ -267,18 +267,18 @@ void FElectraProtronPlayer::FImpl::InternalOpen(const FString& InFilename)
 		}
 		ParsedRootBoxes[i] = tp.GetBoxTree();
 	}
-	auto FindRootBox = [&](uint32 InType) -> TSharedPtr<Electra::UtilitiesMP4::FMP4BoxBase>
+	auto FindRootBox = [&](uint32 InType) -> TSharedPtr<MP4Boxes::FMP4BoxBase>
 	{
-		auto rb = ParsedRootBoxes.FindByPredicate([Type=InType](const TSharedPtr<Electra::UtilitiesMP4::FMP4BoxBase>& In) { return In->GetType() == Type;});
+		auto rb = ParsedRootBoxes.FindByPredicate([Type=InType](const TSharedPtr<MP4Boxes::FMP4BoxBase>& In) { return In->GetType() == Type;});
 		return rb ? *rb : nullptr;
 	};
 
 	// We need the `moov` box now to determine what is inside this file.
 	// Yes, we looked at the presence of it above already as a quick reject before parsing the boxes.
 	// Now we need it for real and we know that it needs to be there, so here we go
-	TSharedPtr<Electra::UtilitiesMP4::FMP4BoxMOOV, ESPMode::ThreadSafe> MoovBox = StaticCastSharedPtr<Electra::UtilitiesMP4::FMP4BoxMOOV>(FindRootBox(Electra::UtilitiesMP4::MakeBoxAtom('m','o','o','v')));
+	TSharedPtr<MP4Boxes::FMP4BoxMOOV, ESPMode::ThreadSafe> MoovBox = StaticCastSharedPtr<MP4Boxes::FMP4BoxMOOV>(FindRootBox(MP4Utilities::MakeBoxAtom('m','o','o','v')));
 	// Get the `mvhd` box for the duration of the movie and the timescale other values are measured in.
-	auto MvhdBox = MoovBox->FindBoxRecursive<Electra::UtilitiesMP4::FMP4BoxMVHD>(Electra::UtilitiesMP4::MakeBoxAtom('m','v','h','d'), 0);
+	auto MvhdBox = MoovBox->FindBoxRecursive<MP4Boxes::FMP4BoxMVHD>(MP4Utilities::MakeBoxAtom('m','v','h','d'), 0);
 	if (!MvhdBox.IsValid())
 	{
 		LastErrorMessage = FString::Printf(TEXT("No `mvhd` box found. This is not a usable mp4/mov file."));
@@ -286,8 +286,8 @@ void FElectraProtronPlayer::FImpl::InternalOpen(const FString& InFilename)
 	}
 
 	// Check if this is a fragmented file. We do not look for `sidx` or `mfra` boxes since these are purely optional.
-	TSharedPtr<Electra::UtilitiesMP4::FMP4BoxMOOF, ESPMode::ThreadSafe> FirstMoofBox = StaticCastSharedPtr<Electra::UtilitiesMP4::FMP4BoxMOOF>(FindRootBox(Electra::UtilitiesMP4::MakeBoxAtom('m','o','o','f')));
-	auto MvexBox = MoovBox->FindBoxRecursive<Electra::UtilitiesMP4::FMP4BoxMVEX>(Electra::UtilitiesMP4::MakeBoxAtom('m','v','e','x'), 0);
+	TSharedPtr<MP4Boxes::FMP4BoxMOOF, ESPMode::ThreadSafe> FirstMoofBox = StaticCastSharedPtr<MP4Boxes::FMP4BoxMOOF>(FindRootBox(MP4Utilities::MakeBoxAtom('m','o','o','f')));
+	auto MvexBox = MoovBox->FindBoxRecursive<MP4Boxes::FMP4BoxMVEX>(MP4Utilities::MakeBoxAtom('m','v','e','x'), 0);
 	if (FirstMoofBox && !MvexBox)
 	{
 		LastErrorMessage = FString::Printf(TEXT("No `mvex` box present in a file containing `moof` fragment boxes."));
@@ -304,7 +304,7 @@ void FElectraProtronPlayer::FImpl::InternalOpen(const FString& InFilename)
 	// If there is a movie extends header box we take the movie duration from there.
 	if (bIsFragmented && MvexBox)
 	{
-		auto MehdBox = MvexBox->FindBoxRecursive<Electra::UtilitiesMP4::FMP4BoxMEHD>(Electra::UtilitiesMP4::MakeBoxAtom('m','e','h','d'), 0);
+		auto MehdBox = MvexBox->FindBoxRecursive<MP4Boxes::FMP4BoxMEHD>(MP4Utilities::MakeBoxAtom('m','e','h','d'), 0);
 		/*
 			From ISO/IEC 14496-12:2022 section 8.8.2.1
 				If movie fragments are present but there is no MediaExtendsHeaderBox
@@ -330,8 +330,8 @@ void FElectraProtronPlayer::FImpl::InternalOpen(const FString& InFilename)
 	}
 
 	// Get all the tracks.
-	TArray<TSharedPtr<Electra::UtilitiesMP4::FMP4BoxTRAK, ESPMode::ThreadSafe>> AllTracks;
-	MoovBox->GetAllBoxInstances(AllTracks, Electra::UtilitiesMP4::MakeBoxAtom('t','r','a','k'));
+	TArray<TSharedPtr<MP4Boxes::FMP4BoxTRAK, ESPMode::ThreadSafe>> AllTracks;
+	MoovBox->GetAllBoxInstances(AllTracks, MP4Utilities::MakeBoxAtom('t','r','a','k'));
 	// Prepare the internal track structure and check which tracks we can use and which ones we cannot.
 	Tracks.SetNum(AllTracks.Num());
 	int32 NumTimecodeTracks = 0;
@@ -340,13 +340,13 @@ void FElectraProtronPlayer::FImpl::InternalOpen(const FString& InFilename)
 	int64 ShortestSupportedTrackDuration = TNumericLimits<int64>::Max();
 	for(int32 i=0; i<AllTracks.Num(); ++i)
 	{
-		const TSharedPtr<Electra::UtilitiesMP4::FMP4BoxTRAK, ESPMode::ThreadSafe>& ThisTrakBox(AllTracks[i]);
+		const TSharedPtr<MP4Boxes::FMP4BoxTRAK, ESPMode::ThreadSafe>& ThisTrakBox(AllTracks[i]);
 
 		Tracks[i] = MakeShared<FTrackInfo>();
 		Tracks[i]->TrackBox = ThisTrakBox;
 
-		auto Tkhd = ThisTrakBox->FindBoxRecursive<Electra::UtilitiesMP4::FMP4BoxTKHD>(Electra::UtilitiesMP4::MakeBoxAtom('t','k','h','d'), 1);
-		auto Mdhd = ThisTrakBox->FindBoxRecursive<Electra::UtilitiesMP4::FMP4BoxMDHD>(Electra::UtilitiesMP4::MakeBoxAtom('m','d','h','d'), 3);
+		auto Tkhd = ThisTrakBox->FindBoxRecursive<MP4Boxes::FMP4BoxTKHD>(MP4Utilities::MakeBoxAtom('t','k','h','d'), 1);
+		auto Mdhd = ThisTrakBox->FindBoxRecursive<MP4Boxes::FMP4BoxMDHD>(MP4Utilities::MakeBoxAtom('m','d','h','d'), 3);
 		if (!Tkhd)
 		{
 			LastErrorMessage = FString::Printf(TEXT("No `tkhd` box found on track %d. This file cannot be used."), i);
@@ -382,7 +382,7 @@ void FElectraProtronPlayer::FImpl::InternalOpen(const FString& InFilename)
 		// If the file is fragmented we collect all the track fragment boxes for this track.
 		if (bIsFragmented)
 		{
-			Tracks[i]->FragmentInfo = MakeShared<Electra::UtilitiesMP4::FMP4Track::FFragmentInfo, ESPMode::ThreadSafe>();
+			Tracks[i]->FragmentInfo = MakeShared<MP4Boxes::FMP4Track::FFragmentInfo, ESPMode::ThreadSafe>();
 			if (!Tracks[i]->FragmentInfo->Prepare(ThisTrakBox, MvexBox, ParsedRootBoxes))
 			{
 				UE_LOG(LogElectraProtron, Warning, TEXT("Fragmented track #%u could not be prepared, ignoring this track."), TrackID);
@@ -453,7 +453,7 @@ void FElectraProtronPlayer::FImpl::InternalOpen(const FString& InFilename)
 			// Add children box data
 			for(auto& Ch : CodecInfo.ExtraBoxes)
 			{
-				FString BoxName = FString::Printf(TEXT("$%s_box"), *Electra::UtilitiesMP4::GetPrintableBoxAtom(Ch.Key));
+				FString BoxName = FString::Printf(TEXT("$%s_box"), *MP4Utilities::GetPrintableBoxAtom(Ch.Key));
 				Params.Add(BoxName, FVariant(Ch.Value));
 			}
 
@@ -492,7 +492,7 @@ void FElectraProtronPlayer::FImpl::InternalOpen(const FString& InFilename)
 		}
 	}
 
-	Duration = MovieDuration.GetAsTimebase(ETimespan::TicksPerSecond);
+	Duration = MovieDuration.GetAsTimespan();
 	// Check that the duration given in the `mvhd` box is not larger than the longest track is.
 	int64 MvhdDur = MovieDuration.GetNumerator();
 
@@ -500,17 +500,17 @@ void FElectraProtronPlayer::FImpl::InternalOpen(const FString& InFilename)
 	{
 		UE_LOG(LogElectraProtron, Warning, TEXT("Movie duration in `mvhd` box (%lld) is larger than the longest track (%lld) in the file, adjusting movie duration down."), (long long int)MvhdDur, (long long int)LongestTrackDuration);
 		MovieDuration.SetFromND(LongestTrackDuration, MovieDuration.GetDenominator());
-		Duration = MovieDuration.GetAsTimebase(ETimespan::TicksPerSecond);
+		Duration = MovieDuration.GetAsTimespan();
 		MvhdDur = LongestTrackDuration;
 	}
-	Electra::FTimeFraction EntireMovieDuration(MovieDuration);
+	MP4Utilities::FFractionalTime EntireMovieDuration(MovieDuration);
 
 	// Check that the movie duration is not larger than the shortest supported track.
 	if (ShortestSupportedTrackDuration < TNumericLimits<int64>::Max() && MvhdDur > ShortestSupportedTrackDuration)
 	{
 		UE_LOG(LogElectraProtron, Warning, TEXT("Movie duration in `mvhd` box (%lld) is larger than the shortest supported track (%lld) in the file, adjusting movie duration down."), (long long int)MvhdDur, (long long int)ShortestSupportedTrackDuration);
 		MovieDuration.SetFromND(ShortestSupportedTrackDuration, MovieDuration.GetDenominator());
-		Duration = MovieDuration.GetAsTimebase(ETimespan::TicksPerSecond);
+		Duration = MovieDuration.GetAsTimespan();
 	}
 
 	// If there are timecode tracks, find which tracks references them.
@@ -526,11 +526,11 @@ void FElectraProtronPlayer::FImpl::InternalOpen(const FString& InFilename)
 			// We do not check if a timecode track references another one as this would be silly.
 			if (Tracks[TrkNum]->CodecInfo.Type != Electra::FCodecTypeFormat::EType::Timecode)
 			{
-				auto Tref = Tracks[TrkNum]->TrackBox->FindBoxRecursive<Electra::UtilitiesMP4::FMP4BoxTREF>(Electra::UtilitiesMP4::MakeBoxAtom('t','r','e','f'), 1);
+				auto Tref = Tracks[TrkNum]->TrackBox->FindBoxRecursive<MP4Boxes::FMP4BoxTREF>(MP4Utilities::MakeBoxAtom('t','r','e','f'), 1);
 				if (Tref.IsValid())
 				{
 					// Get timecode references, if any.
-					auto References = Tref->GetEntriesOfType(Electra::UtilitiesMP4::MakeBoxAtom('t','m','c','d'));
+					auto References = Tref->GetEntriesOfType(MP4Utilities::MakeBoxAtom('t','m','c','d'));
 					if (References.Num())
 					{
 						if (References.Num() > 1)
@@ -602,17 +602,17 @@ void FElectraProtronPlayer::FImpl::InternalOpen(const FString& InFilename)
 	else
 	{
 		// See if there is an `udta` box in the `moov` that contains `(c)TIM` and `(C)TSC` boxes.
-		auto UdtaBox = MoovBox->FindBoxRecursive<Electra::UtilitiesMP4::FMP4BoxUDTA>(Electra::UtilitiesMP4::MakeBoxAtom('u','d','t','a'), 0);
+		auto UdtaBox = MoovBox->FindBoxRecursive<MP4Boxes::FMP4BoxUDTA>(MP4Utilities::MakeBoxAtom('u','d','t','a'), 0);
 		if (UdtaBox.IsValid())
 		{
-			auto CTIMBox = UdtaBox->FindBoxRecursive<Electra::UtilitiesMP4::FMP4BoxBase>(Electra::UtilitiesMP4::MakeBoxAtom(0xa9U,'T','I','M'), 0);
-			auto CTSCBox = UdtaBox->FindBoxRecursive<Electra::UtilitiesMP4::FMP4BoxBase>(Electra::UtilitiesMP4::MakeBoxAtom(0xa9U,'T','S','C'), 0);
-			//auto CTSZBox = UdtaBox->FindBoxRecursive<Electra::UtilitiesMP4::FMP4BoxBase>(Electra::UtilitiesMP4::MakeBoxAtom(0xa9U,'T','S','Z'), 0);
+			auto CTIMBox = UdtaBox->FindBoxRecursive<MP4Boxes::FMP4BoxBase>(MP4Utilities::MakeBoxAtom(0xa9U,'T','I','M'), 0);
+			auto CTSCBox = UdtaBox->FindBoxRecursive<MP4Boxes::FMP4BoxBase>(MP4Utilities::MakeBoxAtom(0xa9U,'T','S','C'), 0);
+			//auto CTSZBox = UdtaBox->FindBoxRecursive<MP4Boxes::FMP4BoxBase>(MP4Utilities::MakeBoxAtom(0xa9U,'T','S','Z'), 0);
 			if (CTIMBox.IsValid() && CTSCBox.IsValid())
 			{
-				auto GetValue = [](const TSharedPtr<Electra::UtilitiesMP4::FMP4BoxBase, ESPMode::ThreadSafe>& InBox) -> FString
+				auto GetValue = [](const TSharedPtr<MP4Boxes::FMP4BoxBase, ESPMode::ThreadSafe>& InBox) -> FString
 				{
-					Electra::UtilitiesMP4::FMP4AtomReader timReader(InBox->GetBoxData());
+					MP4Utilities::FMP4AtomReader timReader(InBox->GetBoxData());
 					FString String;
 					uint16 StringLength, UnknownValue;
 					if (timReader.Read(StringLength)&&  timReader.Read(UnknownValue))
@@ -672,7 +672,7 @@ void FElectraProtronPlayer::FImpl::InternalOpen(const FString& InFilename)
 		{
 			continue;
 		}
-		Tracks[i]->MP4Track = Electra::UtilitiesMP4::FMP4Track::Create(Tracks[i]->TrackBox, Tracks[i]->FragmentInfo);
+		Tracks[i]->MP4Track = MP4Boxes::FMP4Track::Create(Tracks[i]->TrackBox, Tracks[i]->FragmentInfo);
 		if (!Tracks[i]->MP4Track->Prepare(EntireMovieDuration, MovieDuration))
 		{
 			LastErrorMessage = Tracks[i]->MP4Track->GetLastError();
@@ -682,7 +682,7 @@ void FElectraProtronPlayer::FImpl::InternalOpen(const FString& InFilename)
 
 		// If this is a timecode track we may want to read in the first timecode.
 		if (Config.bReadFirstTimecode && Tracks[i]->CodecInfo.Type == Electra::FCodecTypeFormat::EType::Timecode &&
-			Tracks[i]->CodecInfo.FourCC == Electra::UtilitiesMP4::MakeBoxAtom('t','m','c','d'))
+			Tracks[i]->CodecInfo.FourCC == MP4Utilities::MakeBoxAtom('t','m','c','d'))
 		{
 			auto It = Tracks[i]->MP4Track->CreateIterator();
 			if (!It.IsValid())
@@ -702,7 +702,7 @@ void FElectraProtronPlayer::FImpl::InternalOpen(const FString& InFilename)
 			}
 			TArray<uint32> TimecodeBuffer;
 			TimecodeBuffer.SetNumUninitialized(Align(SampleSize, sizeof(uint32)));
-			int64 NumRead = Reader->ReadData(TimecodeBuffer.GetData(), SampleSize, SampleFileOffset, Electra::IBaseDataReader::FCancellationCheckDelegate::CreateLambda([&](){return bAbort;}));
+			int64 NumRead = Reader->ReadData(TimecodeBuffer.GetData(), SampleSize, SampleFileOffset, MP4Utilities::IMP4DataReaderBase::FCancellationCheckDelegate::CreateLambda([&](){return bAbort;}));
 			if (NumRead != SampleSize)
 			{
 				LastErrorMessage = FString::Printf(TEXT("Failed to read first timecode sample."));
@@ -712,7 +712,7 @@ void FElectraProtronPlayer::FImpl::InternalOpen(const FString& InFilename)
 			Electra::FCodecTypeFormat::FTMCDTimecode TimecodeInfo(Tracks[i]->CodecInfo.Properties.Get<Electra::FCodecTypeFormat::FTMCDTimecode>());
 			// Set the first timecode sample on the track.
 			FTrackInfo::FFirstSampleTimecode firstTC;
-			firstTC.TimecodeValue = Electra::GetFromBigEndian(TimecodeBuffer[0]);
+			firstTC.TimecodeValue = MP4Utilities::GetFromBigEndian(TimecodeBuffer[0]);
 			firstTC.Framerate = TimecodeInfo.GetFrameRate().ToPrettyText().ToString();
 			firstTC.Timecode = TimecodeInfo.ConvertToTimecode(firstTC.TimecodeValue).ToString();
 			Tracks[i]->FirstSampleTimecode.Emplace(MoveTemp(firstTC));
@@ -729,11 +729,11 @@ void FElectraProtronPlayer::FImpl::InternalOpen(const FString& InFilename)
 	}
 }
 
-void FElectraProtronPlayer::FImpl::GetTrackCodecInfo(Electra::FCodecTypeFormat& OutCodecInfo, const TSharedPtr<Electra::UtilitiesMP4::FMP4BoxTRAK, ESPMode::ThreadSafe>& InTrack, uint32 InTrackID)
+void FElectraProtronPlayer::FImpl::GetTrackCodecInfo(Electra::FCodecTypeFormat& OutCodecInfo, const TSharedPtr<MP4Boxes::FMP4BoxTRAK, ESPMode::ThreadSafe>& InTrack, uint32 InTrackID)
 {
 	// There needs to be an `stsd` box in this track. We do not try the expected path of `trak`->`mdia`->`minf`->`stbl`->`stsd` as
 	// this is not that much faster and if the file is somewhat ill-formed we may not find it if it's grouped under elsewhere.
-	auto Stsd = InTrack->FindBoxRecursive<Electra::UtilitiesMP4::FMP4BoxSTSD>(Electra::UtilitiesMP4::MakeBoxAtom('s','t','s','d'), 6);
+	auto Stsd = InTrack->FindBoxRecursive<MP4Boxes::FMP4BoxSTSD>(MP4Utilities::MakeBoxAtom('s','t','s','d'), 6);
 	if (!Stsd)
 	{
 		// If not found we just ignore the track. That's a warning but not an error.
@@ -741,7 +741,19 @@ void FElectraProtronPlayer::FImpl::GetTrackCodecInfo(Electra::FCodecTypeFormat&
 		return;
 	}
 
-	auto AddChildren = [&OutCodecInfo](const TSharedPtr<Electra::UtilitiesMP4::FMP4BoxBase, ESPMode::ThreadSafe>& InSample) -> void
+	// Get language code
+	auto Mdhd = InTrack->FindBoxRecursive<MP4Boxes::FMP4BoxMDHD>(MP4Utilities::MakeBoxAtom('m','d','h','d'), 2);
+	if (Mdhd)
+	{
+		Electra::BCP47::ParseRFC5646Tag(OutCodecInfo.LanguageTag, Mdhd->GetLanguageCode639_2T());
+	}
+	auto Elng = InTrack->FindBoxRecursive<MP4Boxes::FMP4BoxELNG>(MP4Utilities::MakeBoxAtom('e','l','n','g'), 6);
+	if (Elng)
+	{
+		Electra::BCP47::ParseRFC5646Tag(OutCodecInfo.LanguageTag, Elng->GetLanguageTag());
+	}
+
+	auto AddChildren = [&OutCodecInfo](const TSharedPtr<MP4Boxes::FMP4BoxBase, ESPMode::ThreadSafe>& InSample) -> void
 	{
 		for(auto& Ch : InSample->GetChildren())
 		{
@@ -750,9 +762,9 @@ void FElectraProtronPlayer::FImpl::GetTrackCodecInfo(Electra::FCodecTypeFormat&
 	};
 
 	// Get the sample type of this track. This call is necessary to trigger parsing of child nodes.
-	Electra::UtilitiesMP4::FMP4BoxSampleEntry::ESampleType SampleType = Stsd->GetSampleType();
+	MP4Boxes::FMP4BoxSampleEntry::ESampleType SampleType = Stsd->GetSampleType();
 	// If already knows to not be supported, skip it.
-	if (SampleType == Electra::UtilitiesMP4::FMP4BoxSampleEntry::ESampleType::Unsupported)
+	if (SampleType == MP4Boxes::FMP4BoxSampleEntry::ESampleType::Unsupported)
 	{
 		UE_LOG(LogElectraProtron, Warning, TEXT("Unsupported sample type, ignoring track #%u."), InTrackID);
 		return;
@@ -771,9 +783,9 @@ void FElectraProtronPlayer::FImpl::GetTrackCodecInfo(Electra::FCodecTypeFormat&
 
 	bool bIsEncrypted = false;
 	// Based on the sample type, get it and see if it is using codec we support.
-	if (SampleType == Electra::UtilitiesMP4::FMP4BoxSampleEntry::ESampleType::Video)
+	if (SampleType == MP4Boxes::FMP4BoxSampleEntry::ESampleType::Video)
 	{
-		TArray<TSharedPtr<Electra::UtilitiesMP4::FMP4BoxVisualSampleEntry, ESPMode::ThreadSafe>> Visuals;
+		TArray<TSharedPtr<MP4Boxes::FMP4BoxVisualSampleEntry, ESPMode::ThreadSafe>> Visuals;
 		Stsd->GetSampleDescriptions(Visuals);
 		if (Visuals.Num() != 1)
 		{
@@ -790,7 +802,7 @@ void FElectraProtronPlayer::FImpl::GetTrackCodecInfo(Electra::FCodecTypeFormat&
 
 		switch(Visuals[0]->GetType())
 		{
-			case Electra::UtilitiesMP4::MakeBoxAtom('e','n','c','v'):
+			case MP4Utilities::MakeBoxAtom('e','n','c','v'):
 			{
 				bIsEncrypted = true;
 				break;
@@ -805,8 +817,7 @@ void FElectraProtronPlayer::FImpl::GetTrackCodecInfo(Electra::FCodecTypeFormat&
 					Electra::FCodecTypeFormat::FVideo& cvid = OutCodecInfo.Properties.Get<Electra::FCodecTypeFormat::FVideo>();
 					cvid.Width = Visuals[0]->GetWidth();
 					cvid.Height = Visuals[0]->GetHeight();
-					auto Mdhd = InTrack->FindBoxRecursive<Electra::UtilitiesMP4::FMP4BoxMDHD>(Electra::UtilitiesMP4::MakeBoxAtom('m','d','h','d'), 2);
-					auto Stts = InTrack->FindBoxRecursive<Electra::UtilitiesMP4::FMP4BoxSTTS>(Electra::UtilitiesMP4::MakeBoxAtom('s','t','t','s'), 5);
+					auto Stts = InTrack->FindBoxRecursive<MP4Boxes::FMP4BoxSTTS>(MP4Utilities::MakeBoxAtom('s','t','t','s'), 5);
 					if (Mdhd.IsValid() && Stts.IsValid() && Stts->GetEntries().Num())
 					{
 						cvid.FrameRate = FFrameRate(Mdhd->GetTimescale(), Stts->GetEntries()[0].sample_delta);
@@ -829,9 +840,9 @@ void FElectraProtronPlayer::FImpl::GetTrackCodecInfo(Electra::FCodecTypeFormat&
 			}
 		}
 	}
-	else if (SampleType == Electra::UtilitiesMP4::FMP4BoxSampleEntry::ESampleType::Audio)
+	else if (SampleType == MP4Boxes::FMP4BoxSampleEntry::ESampleType::Audio)
 	{
-		TArray<TSharedPtr<Electra::UtilitiesMP4::FMP4BoxAudioSampleEntry, ESPMode::ThreadSafe>> Audios;
+		TArray<TSharedPtr<MP4Boxes::FMP4BoxAudioSampleEntry, ESPMode::ThreadSafe>> Audios;
 		Stsd->GetSampleDescriptions(Audios);
 		if (Audios.Num() != 1)
 		{
@@ -848,7 +859,7 @@ void FElectraProtronPlayer::FImpl::GetTrackCodecInfo(Electra::FCodecTypeFormat&
 
 		switch(Audios[0]->GetType())
 		{
-			case Electra::UtilitiesMP4::MakeBoxAtom('e','n','c','a'):
+			case MP4Utilities::MakeBoxAtom('e','n','c','a'):
 			{
 				bIsEncrypted = true;
 				break;
@@ -875,14 +886,14 @@ void FElectraProtronPlayer::FImpl::GetTrackCodecInfo(Electra::FCodecTypeFormat&
 		}
 
 	}
-	else if (SampleType == Electra::UtilitiesMP4::FMP4BoxSampleEntry::ESampleType::QTFFTimecode)
+	else if (SampleType == MP4Boxes::FMP4BoxSampleEntry::ESampleType::QTFFTimecode)
 	{
-		TArray<TSharedPtr<Electra::UtilitiesMP4::FMP4BoxQTFFTimecodeSampleEntry, ESPMode::ThreadSafe>> Timecodes;
+		TArray<TSharedPtr<MP4Boxes::FMP4BoxQTFFTimecodeSampleEntry, ESPMode::ThreadSafe>> Timecodes;
 		Stsd->GetSampleDescriptions(Timecodes);
 		check(Timecodes.Num() == 1);
 		switch(Timecodes[0]->GetType())
 		{
-			case Electra::UtilitiesMP4::MakeBoxAtom('t','m','c','d'):
+			case MP4Utilities::MakeBoxAtom('t','m','c','d'):
 			{
 				Electra::FCodecTypeFormat::FTMCDTimecode Timecode;
 
@@ -891,7 +902,7 @@ void FElectraProtronPlayer::FImpl::GetTrackCodecInfo(Electra::FCodecTypeFormat&
 				Timecode.FrameDuration = Timecodes[0]->GetFrameDuration();
 				Timecode.NumberOfFrames = Timecodes[0]->GetNumberOfFrames();
 				OutCodecInfo.Properties.Emplace<Electra::FCodecTypeFormat::FTMCDTimecode>(MoveTemp(Timecode));
-				OutCodecInfo.RFC6381 = Electra::UtilitiesMP4::GetPrintableBoxAtom(Timecodes[0]->GetType());
+				OutCodecInfo.RFC6381 = MP4Utilities::GetPrintableBoxAtom(Timecodes[0]->GetType());
 				OutCodecInfo.FourCC = Timecodes[0]->GetType();
 				OutCodecInfo.DCR = Timecodes[0]->GetBoxData();
 				OutCodecInfo.Type = Electra::FCodecTypeFormat::EType::Timecode;
@@ -906,10 +917,10 @@ void FElectraProtronPlayer::FImpl::GetTrackCodecInfo(Electra::FCodecTypeFormat&
 	{
 		if (!bIsEncrypted)
 		{
-			TArray<TSharedPtr<Electra::UtilitiesMP4::FMP4BoxSampleEntry, ESPMode::ThreadSafe>> Entries;
+			TArray<TSharedPtr<MP4Boxes::FMP4BoxSampleEntry, ESPMode::ThreadSafe>> Entries;
 			Stsd->GetSampleDescriptions(Entries);
 			check(Entries.Num() == 1);
-			UE_LOG(LogElectraProtron, Warning, TEXT("Track of sample type \"%s\" is not supported, ignoring track #%u."), *Electra::UtilitiesMP4::GetPrintableBoxAtom(Entries[0]->GetType()), InTrackID);
+			UE_LOG(LogElectraProtron, Warning, TEXT("Track of sample type \"%s\" is not supported, ignoring track #%u."), *MP4Utilities::GetPrintableBoxAtom(Entries[0]->GetType()), InTrackID);
 		}
 		else
 		{
@@ -1071,7 +1082,7 @@ FString FElectraProtronPlayer::FImpl::GetTrackLanguage(EMediaTrackType InTrackTy
 			{
 				if (InTrackIndex < UsableTrackArrayIndicesByType[CodecTypeIndex(Electra::FCodecTypeFormat::EType::Video)].Num())
 				{
-					return Tracks[UsableTrackArrayIndicesByType[CodecTypeIndex(Electra::FCodecTypeFormat::EType::Video)][InTrackIndex]]->MP4Track->GetCommonMetadata().LanguageTag.Get();
+					return Tracks[UsableTrackArrayIndicesByType[CodecTypeIndex(Electra::FCodecTypeFormat::EType::Video)][InTrackIndex]]->CodecInfo.LanguageTag.Get();
 				}
 				break;
 			}
@@ -1079,7 +1090,7 @@ FString FElectraProtronPlayer::FImpl::GetTrackLanguage(EMediaTrackType InTrackTy
 			{
 				if (InTrackIndex < UsableTrackArrayIndicesByType[CodecTypeIndex(Electra::FCodecTypeFormat::EType::Audio)].Num())
 				{
-					return Tracks[UsableTrackArrayIndicesByType[CodecTypeIndex(Electra::FCodecTypeFormat::EType::Audio)][InTrackIndex]]->MP4Track->GetCommonMetadata().LanguageTag.Get();
+					return Tracks[UsableTrackArrayIndicesByType[CodecTypeIndex(Electra::FCodecTypeFormat::EType::Audio)][InTrackIndex]]->CodecInfo.LanguageTag.Get();
 				}
 				break;
 			}
diff --git a/Engine/Plugins/Media/ElectraPlayer/Source/ElectraProtron/Private/Player/ElectraProtronPlayerImpl.h b/Engine/Plugins/Media/ElectraPlayer/Source/ElectraProtron/Private/Player/ElectraProtronPlayerImpl.h
--- a/Engine/Plugins/Media/ElectraPlayer/Source/ElectraProtron/Private/Player/ElectraProtronPlayerImpl.h
+++ b/Engine/Plugins/Media/ElectraPlayer/Source/ElectraProtron/Private/Player/ElectraProtronPlayerImpl.h
@@ -11,9 +11,10 @@
 #include "Misc/Optional.h"
 #include "Math/Range.h"
 #include "Math/RangeSet.h"
-#include "Utilities/UtilitiesMP4.h"
-#include "Utilities/MP4Boxes/MP4Boxes.h"
-#include "Utilities/MP4Boxes/MP4Track.h"
+#include "MP4Utilities.h"
+#include "MP4DataReader.h"
+#include "MP4Boxes.h"
+#include "MP4Track.h"
 #include "CodecTypeFormat.h"
 #include "IElectraDecoder.h"
 #include "IElectraCodecFormatUtils.h"
@@ -24,6 +25,7 @@
 #include "Utils/AudioChannelMapper.h"
 #include "Utils/VideoDecoderHelpers.h"
 #include "ElectraProtronPlayerCache.h"
+#include "Utilities/BCP47-Helpers.h"
 
 
 
@@ -350,9 +352,9 @@ private:
 			FString Vendor;
 		};
 		FDecoderProviderInfo DecoderProviderInfo;
-		TSharedPtr<Electra::UtilitiesMP4::FMP4BoxTRAK, ESPMode::ThreadSafe> TrackBox;
-		TSharedPtr<Electra::UtilitiesMP4::FMP4Track::FFragmentInfo, ESPMode::ThreadSafe> FragmentInfo;
-		TSharedPtr<Electra::UtilitiesMP4::FMP4Track, ESPMode::ThreadSafe> MP4Track;
+		TSharedPtr<MP4Boxes::FMP4BoxTRAK, ESPMode::ThreadSafe> TrackBox;
+		TSharedPtr<MP4Boxes::FMP4Track::FFragmentInfo, ESPMode::ThreadSafe> FragmentInfo;
+		TSharedPtr<MP4Boxes::FMP4Track, ESPMode::ThreadSafe> MP4Track;
 		TWeakPtr<FTrackInfo, ESPMode::ThreadSafe> ReferencedTimecodeTrack;
 		uint32 TrackID = 0;
 		bool bIsUsable = false;
@@ -391,7 +393,7 @@ private:
 	};
 	using FMP4SamplePtr = TSharedPtr<FMP4Sample, ESPMode::ThreadSafe>;
 
-	using FTrackIterator = TSharedPtr<Electra::UtilitiesMP4::FMP4Track::FIterator, ESPMode::ThreadSafe>;
+	using FTrackIterator = TSharedPtr<MP4Boxes::FMP4Track::FIterator, ESPMode::ThreadSafe>;
 
 	struct FMP4TrackSampleBuffer
 	{
@@ -413,7 +415,7 @@ private:
 	void SendWorkerThreadMessage(FWorkerThreadMessage&& InMessage);
 
 	void InternalOpen(const FString& InFilename);
-	void GetTrackCodecInfo(Electra::FCodecTypeFormat& OutCodecInfo, const TSharedPtr<Electra::UtilitiesMP4::FMP4BoxTRAK, ESPMode::ThreadSafe>& InTrack, uint32 InTrackID);
+	void GetTrackCodecInfo(Electra::FCodecTypeFormat& OutCodecInfo, const TSharedPtr<MP4Boxes::FMP4BoxTRAK, ESPMode::ThreadSafe>& InTrack, uint32 InTrackID);
 
 	void UpdateTrackLoader(int32 InCodecTypeIndex);
 
@@ -439,11 +441,11 @@ private:
 	TQueue<FWorkerThreadMessage, EQueueMode::Mpsc> WorkMessages;
 	FString LastErrorMessage;
 	bool bAbort = false;
-	TArray<TSharedPtr<Electra::UtilitiesMP4::FMP4BoxBase>> ParsedRootBoxes;
+	TArray<TSharedPtr<MP4Boxes::FMP4BoxBase>> ParsedRootBoxes;
 	TArray<TSharedPtr<FTrackInfo, ESPMode::ThreadSafe>> Tracks;
 	TArray<TArray<int32>> UsableTrackArrayIndicesByType;
 
-	Electra::FTimeFraction MovieDuration;
+	MP4Utilities::FFractionalTime MovieDuration;
 	FTimespan Duration;						// already converted `MovieDuration` for faster access
 	FTrackSelection TrackSelection;
 	TRangeSet<float> UnthinnedRates;
@@ -523,7 +525,7 @@ private:
 		volatile bool bTerminateThread = false;
 
 		TSharedPtr<FSharedPlayParams, ESPMode::ThreadSafe> SharedPlayParams;
-		TSharedPtr<Electra::IFileDataReader, ESPMode::ThreadSafe> Reader;
+		TSharedPtr<MP4Utilities::FMP4FileDataReader, ESPMode::ThreadSafe> Reader;
 		FString LastErrorMessage;
 
 		FCriticalSection LoadRequestLock;
diff --git a/Engine/Plugins/Media/ElectraPlayer/Source/ElectraProtron/Private/Player/ElectraProtronPlayerLoader.cpp b/Engine/Plugins/Media/ElectraPlayer/Source/ElectraProtron/Private/Player/ElectraProtronPlayerLoader.cpp
--- a/Engine/Plugins/Media/ElectraPlayer/Source/ElectraProtron/Private/Player/ElectraProtronPlayerLoader.cpp
+++ b/Engine/Plugins/Media/ElectraPlayer/Source/ElectraProtron/Private/Player/ElectraProtronPlayerLoader.cpp
@@ -45,7 +45,7 @@ uint32 FElectraProtronPlayer::FImpl::FLoaderThread::Run()
 		{
 			SharedPlayParams = MoveTemp(OpenRequest.SharedPlayParams);
 			check(SharedPlayParams.IsValid());
-			Reader = Electra::IFileDataReader::Create();
+			Reader = MP4Utilities::FMP4FileDataReader::Create();
 			FString Filename = MoveTemp(OpenRequest.Filename);
 			if (!Reader->Open(Filename))
 			{
@@ -266,7 +266,7 @@ void FElectraProtronPlayer::FImpl::FLoaderThread::RequestLoad(FMP4TrackSampleBuf
 	{
 		InTrackSampleBuffer->CurrentPlaybackRange = RangeNow;
 		// Locate the first and last sample numbers for the range
-		FTrackIterator RangeIt = InTrackSampleBuffer->TrackAndCodecInfo->MP4Track->CreateIteratorAtKeyframe(Electra::FTimeValue().SetFromTimespan(RangeNow.GetLowerBoundValue()), Electra::FTimeValue::GetZero());
+		FTrackIterator RangeIt = InTrackSampleBuffer->TrackAndCodecInfo->MP4Track->CreateIteratorAtKeyframe(RangeNow.GetLowerBoundValue(), FTimespan(0));
 		if (!RangeIt.IsValid())
 		{
 			LastErrorMessage = InTrackSampleBuffer->TrackAndCodecInfo->MP4Track->GetLastError();
@@ -291,9 +291,7 @@ void FElectraProtronPlayer::FImpl::FLoaderThread::RequestLoad(FMP4TrackSampleBuf
 		InTrackSampleBuffer->LastRangeSampleIt = MoveTemp(RangeIt);
 	}
 
-	FTrackIterator TkIt =
-		InTrackSampleBuffer->TrackAndCodecInfo->MP4Track->CreateIteratorAtKeyframe(Electra::FTimeValue().SetFromTimespan(InTime),
-			Electra::FTimeValue().SetFromMilliseconds(InTrackSampleBuffer->TrackAndCodecInfo->bIsKeyframeOnlyFormat ? 0 : Config.NextKeyframeThresholdMillis));
+	FTrackIterator TkIt = InTrackSampleBuffer->TrackAndCodecInfo->MP4Track->CreateIteratorAtKeyframe(InTime, FTimespan::FromMilliseconds(InTrackSampleBuffer->TrackAndCodecInfo->bIsKeyframeOnlyFormat ? 0 : Config.NextKeyframeThresholdMillis));
 	if (!TkIt.IsValid())
 	{
 		LastErrorMessage = InTrackSampleBuffer->TrackAndCodecInfo->MP4Track->GetLastError();
@@ -474,7 +472,7 @@ TSharedPtr<FElectraProtronPlayer::FImpl::FMP4Sample, ESPMode::ThreadSafe> FElect
 		return *CurrentSample;
 	}
 
-	auto CheckAbort = Electra::IBaseDataReader::FCancellationCheckDelegate::CreateLambda([&]()
+	auto CheckAbort = MP4Utilities::IMP4DataReaderBase::FCancellationCheckDelegate::CreateLambda([&]()
 	{
 		// We do NOT abort loading of a frame!
 		return false;
@@ -509,7 +507,7 @@ TSharedPtr<FElectraProtronPlayer::FImpl::FMP4Sample, ESPMode::ThreadSafe> FElect
 			NumRead = Reader->ReadData(TimecodeBuffer.GetData(), tcSampleSize, InOptionalTimecodeIt->GetSampleFileOffset(), CheckAbort);
 			if (NumRead == tcSampleSize)
 			{
-				Sample->AssociatedTimecode = InTimecodeInfo.ConvertToTimecode(Electra::GetFromBigEndian(TimecodeBuffer[0]));
+				Sample->AssociatedTimecode = InTimecodeInfo.ConvertToTimecode(MP4Utilities::GetFromBigEndian(TimecodeBuffer[0]));
 				Sample->AssociatedTimecodeFramerate = InTimecodeInfo.GetFrameRate();
 			}
 		}
diff --git a/Engine/Plugins/Media/ElectraUtil/Source/ElectraBase/Private/PlayerTime.cpp b/Engine/Plugins/Media/ElectraUtil/Source/ElectraBase/Private/PlayerTime.cpp
--- a/Engine/Plugins/Media/ElectraUtil/Source/ElectraBase/Private/PlayerTime.cpp
+++ b/Engine/Plugins/Media/ElectraUtil/Source/ElectraBase/Private/PlayerTime.cpp
@@ -1,64 +1,11 @@
 // Copyright Epic Games, Inc. All Rights Reserved.
 
 #include "PlayerTime.h"
-
 #include "Math/BigInt.h"
 
 
 namespace Electra
 {
-
-	namespace TimeStringHelpers
-	{
-		int32 FindFirstNotOf(const FString& InString, const FString& InNotOfChars, int32 FirstPos = 0)
-		{
-			for (int32 i = FirstPos; i < InString.Len(); ++i)
-			{
-				bool bFoundCharFromNotOfChars = false;
-				for (int32 j = 0; j < InNotOfChars.Len(); ++j)
-				{
-					if (InString[i] == InNotOfChars[j])
-					{
-						// We found a character from the "NOT of" list. Check next...
-						bFoundCharFromNotOfChars = true;
-						break;
-					}
-				}
-				if (!bFoundCharFromNotOfChars)
-				{
-					// We did not find any of the characters. This is what we are looking for and return the index of the first "not of" character
-					return i;
-				}
-			}
-			return INDEX_NONE;
-		}
-
-		int32 FindLastNotOf(const FString& InString, const FString& InNotOfChars, int32 InStartPos = MAX_int32)
-		{
-			InStartPos = FMath::Min(InStartPos, InString.Len() - 1);
-			for (int32 i = InStartPos; i >= 0; --i)
-			{
-				bool bFoundCharFromNotOfChars = false;
-				for (int32 j = 0; j < InNotOfChars.Len(); ++j)
-				{
-					if (InString[i] == InNotOfChars[j])
-					{
-						// We found a character from the "NOT of" list. Check next...
-						bFoundCharFromNotOfChars = true;
-						break;
-					}
-				}
-				if (!bFoundCharFromNotOfChars)
-				{
-					// We did not find any of the characters. This is what we are looking for and return the index of the first "not of" character
-					return i;
-				}
-			}
-			return INDEX_NONE;
-		}
-	}
-
-
 	FTimeValue& FTimeValue::SetFromTimeFraction(const FTimeFraction& TimeFraction, int64 InSequenceIndex)
 	{
 		if (TimeFraction.IsValid())
@@ -148,153 +95,5 @@ namespace Electra
 		int64 r = n.ToInt();
 		return bIsNeg ? -r : r;
 	}
-
-	// ---------------------------------------------------------------------------------------------------------------------
-
-	FTimeFraction& FTimeFraction::SetFromFloatString(const FString& InString)
-	{
-		// The string value must consist only of a sign, decimal digits and an optional period.
-		static const FString kTextDigitsEtc(TEXT("0123456789.-+"));
-		static const FString kTextZero(TEXT("0"));
-		if (TimeStringHelpers::FindFirstNotOf(InString, kTextDigitsEtc) == INDEX_NONE) 
-		{
-			Denominator = 1;
-			int32 DotIndex;
-			InString.FindChar(TCHAR('.'), DotIndex);
-			if (DotIndex == INDEX_NONE)
-			{
-				LexFromString(Numerator, *InString);
-				bIsValid = true;
-			}
-			else
-			{
-				LexFromString(Numerator, *(InString.Mid(0, DotIndex)));
-				FString frc = InString.Mid(DotIndex + 1);
-				bool bIsNeg = Numerator < 0;
-				if (bIsNeg)
-				{
-					Numerator = -Numerator;
-				}
-				// Remove all trailing zeros
-				int32 last0 = TimeStringHelpers::FindLastNotOf(frc, kTextZero);
-				if (last0 != INDEX_NONE)
-				{
-					frc.MidInline(0, last0 + 1);
-					// Convert at most 7 fractional digits (giving us hundreds of nanoseconds (HNS))
-					for(int32 i = 0; i < frc.Len() && i<7; ++i)
-					{
-						Numerator = Numerator * 10 + (frc[i] - TCHAR('0'));
-						Denominator *= 10;
-					}
-				}
-				if (bIsNeg)
-				{
-					Numerator = -Numerator;
-				}
-				bIsValid = true;
-			}
-		}
-		else
-		{
-			static const FString kInf0(TEXT("INF"));
-			static const FString kInf1(TEXT("+INF"));
-			static const FString kInf2(TEXT("-INF"));
-			if (InString.Equals(kInf0) || InString.Equals(kInf1))
-			{
-				Numerator = 1;
-				Denominator = 0;
-				bIsValid = true;
-			}
-			else if (InString.Equals(kInf2))
-			{
-				Numerator = -1;
-				Denominator = 0;
-				bIsValid = true;
-			}
-			else
-			{
-				bIsValid = false;
-			}
-		}
-		return *this;
-	}
-
-
-	int64 FTimeFraction::GetAsTimebase(uint32 CustomTimebase) const
-	{
-		// Some shortcuts
-		if (!bIsValid)
-		{
-			return 0;
-		}
-		else if (CustomTimebase == Denominator)
-		{
-			return Numerator;
-		}
-		else if (Numerator == 0)
-		{
-			return 0;
-		}
-		// Infinity?
-		else if (Denominator == 0 || CustomTimebase == 0)
-		{
-			return Numerator >= 0 ? 0x7fffffffffffffffLL : -0x7fffffffffffffffLL;
-		}
-
-		bool bIsNeg = Numerator < 0;
-		TBigInt<128> n(bIsNeg ? -Numerator : Numerator);
-		TBigInt<128> d(Denominator);
-		TBigInt<128> s(CustomTimebase);
-
-		n *= s;
-		n /= d;
-
-		int64 r = n.ToInt();
-		return bIsNeg ? -r : r;
-	}
-
-	FTimespan FTimeFraction::GetAsTimespan() const
-	{
-		int64 Ticks = 0;
-		if (bIsValid)
-		{
-			if (Denominator)
-			{
-				if (Numerator != 0)
-				{
-					if (Denominator == ETimespan::TicksPerSecond)
-					{
-						Ticks = Numerator;
-					}
-					else if (Numerator >= -922337203685LL && Numerator <= 922337203685LL)
-					{
-						static_assert(ETimespan::TicksPerSecond == 10000000); // otherwise the constants above are wrong
-						Ticks = Numerator * ETimespan::TicksPerSecond / Denominator;
-					}
-					else
-					{
-						bool bIsNeg = Numerator < 0;
-						TBigInt<128> n(bIsNeg ? -Numerator : Numerator);
-						TBigInt<128> d(Denominator);
-						TBigInt<128> s(ETimespan::TicksPerSecond);
-						n *= s;
-						n /= d;
-						Ticks = n.ToInt();
-						if (bIsNeg)
-						{
-							Ticks = -Ticks;
-						}
-					}
-				}
-			}
-			else
-			{
-				Ticks = Numerator < 0 ? ETimespan::MinTicks : ETimespan::MaxTicks;
-			}
-		}
-		return FTimespan(Ticks);
-	}
-
-
 }
 
diff --git a/Engine/Plugins/Media/ElectraUtil/Source/ElectraBase/Public/PlayerTime.h b/Engine/Plugins/Media/ElectraUtil/Source/ElectraBase/Public/PlayerTime.h
--- a/Engine/Plugins/Media/ElectraUtil/Source/ElectraBase/Public/PlayerTime.h
+++ b/Engine/Plugins/Media/ElectraUtil/Source/ElectraBase/Public/PlayerTime.h
@@ -2,13 +2,13 @@
 
 #pragma once
 
-#include "CoreMinimal.h"
+#include "CoreTypes.h"
 #include "Misc/Timespan.h"
+#include "TimeFraction.h"
 #include <limits>
 
 namespace Electra
 {
-class FTimeFraction;
 
 /**
  * Keeps a time value in hundred nanoseconds (HNS).
@@ -588,167 +588,5 @@ struct FTimeRange
 	FTimeValue	End;
 };
 
-
-
-/**
- * Keeps a time value as a fractional.
- */
-class FTimeFraction
-{
-public:
-	static const FTimeFraction& GetInvalid()
-	{
-		static FTimeFraction Invalid;
-		return Invalid;
-	}
-
-	static const FTimeFraction& GetZero()
-	{
-		static FTimeFraction Zero(0, 1);
-		return Zero;
-	}
-
-	static const FTimeFraction& GetPositiveInfinity()
-	{
-		static FTimeFraction Inf(0, 0);
-		return Inf;
-	}
-
-	FTimeFraction() : Numerator(0), Denominator(0), bIsValid(false)
-	{
-	}
-
-	FTimeFraction(int64 n, uint32 d) : Numerator(n), Denominator(d), bIsValid(true)
-	{
-	}
-
-	FTimeFraction(const FTimeFraction& rhs) : Numerator(rhs.Numerator), Denominator(rhs.Denominator), bIsValid(rhs.bIsValid)
-	{
-	}
-
-	FTimeFraction(const FTimeValue& tv)
-	{
-		SetFromTimeValue(tv);
-	}
-
-	FTimeFraction(const FTimespan& tv)
-	{
-		SetFromTimespan(tv);
-	}
-
-
-	FTimeFraction& operator=(const FTimeFraction& rhs)
-	{
-		Numerator = rhs.Numerator;
-		Denominator = rhs.Denominator;
-		bIsValid = rhs.bIsValid;
-		return *this;
-	}
-
-	bool IsValid() const
-	{
-		return bIsValid;
-	}
-
-	bool IsPositiveInfinity() const
-	{
-		return bIsValid && Denominator==0 && Numerator>=0;
-	}
-
-	int64 GetNumerator() const
-	{
-		return Numerator;
-	}
-
-	uint32 GetDenominator() const
-	{
-		return Denominator;
-	}
-
-	double GetAsDouble() const
-	{
-		return Numerator / (double)Denominator;
-	}
-
-	//! Returns this time value in a custom timebase. Requires internal bigint conversion and is therefor SLOW!
-	ELECTRABASE_API int64 GetAsTimebase(uint32 CustomTimebase) const;
-
-	ELECTRABASE_API FTimespan GetAsTimespan() const;
-
-	FTimeFraction& SetFromND(int64 InNumerator, uint32 InDenominator)
-	{
-		Numerator = InNumerator;
-		Denominator = InDenominator;
-		bIsValid = true;
-		return *this;
-	}
-
-	FTimeFraction& SetNumerator(int64 InNumerator)
-	{
-		Numerator = InNumerator;
-		bIsValid = true;
-		return *this;
-	}
-
-	FTimeFraction& SetDenominator(uint32 InDenominator)
-	{
-		Denominator = InDenominator;
-		bIsValid = true;
-		return *this;
-	}
-
-	FTimeFraction& SetToPositiveInfinity()
-	{
-		Numerator = 0;
-		Denominator = 0;
-		bIsValid = true;
-		return *this;
-	}
-
-	FTimeFraction& SetFromTimeValue(const FTimeValue& tv)
-	{
-		if (tv.IsValid())
-		{
-			Numerator = tv.GetAsHNS();
-			Denominator = tv.IsInfinity() ? 0 : 10000000;
-			bIsValid = true;
-		}
-		else
-		{
-			Numerator = 0;
-			Denominator = 0;
-			bIsValid = false;
-		}
-		return *this;
-	}
-
-	FTimeFraction& SetFromTimespan(const FTimespan& tv)
-	{
-		Numerator = tv.GetTicks();
-		Denominator = (Numerator == ETimespan::MinTicks || Numerator == ETimespan::MaxTicks) ? 0U : (uint32)ETimespan::TicksPerSecond;
-		bIsValid = true;
-		return *this;
-	}
-
-	ELECTRABASE_API FTimeFraction& SetFromFloatString(const FString& In);
-
-
-	bool operator == (const FTimeFraction& rhs) const
-	{
-		return bIsValid == rhs.bIsValid && Numerator == rhs.Numerator && Denominator == rhs.Denominator;
-	}
-
-	bool operator != (const FTimeFraction& rhs) const
-	{
-		return !(*this == rhs);
-	}
-
-private:
-	int64	Numerator;
-	uint32	Denominator;
-	bool	bIsValid;
-};
-
-
 }
 
diff --git a/Engine/Plugins/Media/FlacDecoderElectra/Source/FlacDecoderElectra/FlacDecoderElectra.Build.cs b/Engine/Plugins/Media/FlacDecoderElectra/Source/FlacDecoderElectra/FlacDecoderElectra.Build.cs
--- a/Engine/Plugins/Media/FlacDecoderElectra/Source/FlacDecoderElectra/FlacDecoderElectra.Build.cs
+++ b/Engine/Plugins/Media/FlacDecoderElectra/Source/FlacDecoderElectra/FlacDecoderElectra.Build.cs
@@ -12,6 +12,7 @@ namespace UnrealBuildTool.Rules
                     "ElectraBase",
 					"ElectraCodecFactory",
 					"ElectraDecoders",
+                    "MP4Utilities",
 					"libFLAC"
                 });
 		}
diff --git a/Engine/Plugins/Media/FlacDecoderElectra/Source/FlacDecoderElectra/Private/FlacDecoder/ElectraMediaFlacDecoder.cpp b/Engine/Plugins/Media/FlacDecoderElectra/Source/FlacDecoderElectra/Private/FlacDecoder/ElectraMediaFlacDecoder.cpp
--- a/Engine/Plugins/Media/FlacDecoderElectra/Source/FlacDecoderElectra/Private/FlacDecoder/ElectraMediaFlacDecoder.cpp
+++ b/Engine/Plugins/Media/FlacDecoderElectra/Source/FlacDecoderElectra/Private/FlacDecoder/ElectraMediaFlacDecoder.cpp
@@ -13,7 +13,7 @@
 #include "IElectraDecoder.h"
 #include "IElectraDecoderOutputAudio.h"
 #include "ElectraDecodersUtils.h"
-#include "Utilities/UtilitiesMP4.h"		// for parsing the 'dfLa' box
+#include "MP4Utilities.h"				// for parsing the 'dfLa' box
 #include "Utilities/ElectraBitstream.h"
 #include "Utils/ElectraBitstreamProcessorDefault.h"
 
@@ -323,7 +323,7 @@ bool FElectraFlacDecoder::Parse_dfLa(FdfLaConfig& OutConfig, const TArray<uint8>
 		return !bFailOnError ? false : PostError(0, TEXT("Incomplete 'dfLa' box"), ERRCODE_INTERNAL_FAILED_TO_PARSE_CSD);
 	}
 
-	Electra::UtilitiesMP4::FMP4AtomReader rd(IndfLaBox);
+	MP4Utilities::FMP4AtomReader rd(IndfLaBox);
 	uint8 Value8 = 0;
 	rd.Read(Value8);		// Version
 	if (Value8 != 0)
diff --git a/Engine/Plugins/Media/MP4Utilities/MP4Utilities.uplugin b/Engine/Plugins/Media/MP4Utilities/MP4Utilities.uplugin
--- a/Engine/Plugins/Media/MP4Utilities/MP4Utilities.uplugin
+++ b/Engine/Plugins/Media/MP4Utilities/MP4Utilities.uplugin
@@ -9,13 +9,22 @@
   "CreatedByURL": "https://epicgames.com",
   "EnabledByDefault": false,
   "CanContainContent": false,
-  "IsBetaVersion" : true,
   "Installed": false,
   "Modules": [
+    {
+      "Name": "MP4Utilities",
+      "Type": "Runtime",
+      "LoadingPhase": "PostEngineInit"
+    },
+    {
+      "Name": "MP4Boxes",
+      "Type": "Runtime",
+      "LoadingPhase": "PostEngineInit"
+    },
     {
       "Name": "MP4Muxer",
       "Type": "Runtime",
-      "LoadingPhase": "Default"
+      "LoadingPhase": "PostEngineInit"
     }
   ]
 }
\ No newline at end of file
diff --git a/Engine/Plugins/Media/MP4Utilities/Source/MP4Muxer/MP4Muxer.Build.cs b/Engine/Plugins/Media/MP4Utilities/Source/MP4Muxer/MP4Muxer.Build.cs
--- a/Engine/Plugins/Media/MP4Utilities/Source/MP4Muxer/MP4Muxer.Build.cs
+++ b/Engine/Plugins/Media/MP4Utilities/Source/MP4Muxer/MP4Muxer.Build.cs
@@ -10,7 +10,8 @@ namespace UnrealBuildTool.Rules
 		public MP4Muxer(ReadOnlyTargetRules Target) : base(Target)
 		{
 			PrivateDependencyModuleNames.AddRange(new string[] {
-				"Core"
+				"Core",
+                "MP4Utilities"
             });
 		}
 	}
diff --git a/Engine/Plugins/Media/MP4Utilities/Source/MP4Muxer/Private/MP4Muxer.cpp b/Engine/Plugins/Media/MP4Utilities/Source/MP4Muxer/Private/MP4Muxer.cpp
--- a/Engine/Plugins/Media/MP4Utilities/Source/MP4Muxer/Private/MP4Muxer.cpp
+++ b/Engine/Plugins/Media/MP4Utilities/Source/MP4Muxer/Private/MP4Muxer.cpp
@@ -17,42 +17,12 @@
 #include "Math/BigInt.h"
 
 #include "MP4MuxerModule.h"
+#include "MP4Utilities.h"
 #include "MP4MuxerBoxes.h"
 
 #include <numeric>
 
 
-namespace TimeConversion
-{
-static int64 ToTimescale(uint32 InTargetDenominator, int64 InSourceNumerator, uint32 InSourceDenominator)
-{
-	if (InTargetDenominator == InSourceDenominator)
-	{
-		return InSourceNumerator;
-	}
-	else if (InSourceNumerator == 0)
-	{
-		return 0;
-	}
-	else if (InSourceDenominator == 0 || InTargetDenominator == 0)
-	{
-		return InSourceNumerator >= 0 ? 0x7fffffffffffffffLL : -0x7fffffffffffffffLL;
-	}
-
-	bool bIsNeg = InSourceNumerator < 0;
-	TBigInt<128> n(bIsNeg ? -InSourceNumerator : InSourceNumerator);
-	TBigInt<128> d(InSourceDenominator);
-	TBigInt<128> s(InTargetDenominator);
-
-	n *= s;
-	n /= d;
-
-	int64 r = n.ToInt();
-	return bIsNeg ? -r : r;
-}
-}
-
-
 class FMP4RawMuxer : public IMP4RawMuxer, public FRunnable
 {
 public:
@@ -273,7 +243,7 @@ bool FMP4RawMuxer::Start(FRequestTrackDataDelegate InSampleRequestDelegate, FSta
 	for(int32 i=0; i<TrackChunkSamples.Num(); ++i)
 	{
 		TrackChunkSamples[i].Timescale = TrackSpecs[i].Timescale;
-		TrackChunkSamples[i].InterleaveDurationNeeded = TimeConversion::ToTimescale(TrackSpecs[i].Timescale, Configuration.InterleaveDuration.GetTicks(), (uint32)ETimespan::TicksPerSecond);
+		TrackChunkSamples[i].InterleaveDurationNeeded = MP4Utilities::ConvertToTimescale(TrackSpecs[i].Timescale, Configuration.InterleaveDuration.GetTicks(), (uint32)ETimespan::TicksPerSecond);
 		if (!TrackChunkSamples[i].InterleaveDurationNeeded)
 		{
 			return SetError(FString::Printf(TEXT("Interleave duration not representable by a timescale of %u"), TrackSpecs[i].Timescale));
@@ -598,7 +568,7 @@ bool FMP4RawMuxer::WriteFileEnd()
 	int64 LongestTrackInMovieUnits = 0;
 	for(auto &tcs : TrackChunkSamples)
 	{
-		int64 TrackDurationInMovieUnits = TimeConversion::ToTimescale(MovieTimescale, tcs.TotalTrackDuration, tcs.Timescale);
+		int64 TrackDurationInMovieUnits = MP4Utilities::ConvertToTimescale(MovieTimescale, tcs.TotalTrackDuration, tcs.Timescale);
 		if (TrackDurationInMovieUnits > LongestTrackInMovieUnits)
 		{
 			LongestTrackInMovieUnits = TrackDurationInMovieUnits;
diff --git a/Engine/Plugins/Media/OpusDecoderElectra/Source/OpusDecoderElectra/Private/OpusDecoder/ElectraMediaOpusDecoder.cpp b/Engine/Plugins/Media/OpusDecoderElectra/Source/OpusDecoderElectra/Private/OpusDecoder/ElectraMediaOpusDecoder.cpp
--- a/Engine/Plugins/Media/OpusDecoderElectra/Source/OpusDecoderElectra/Private/OpusDecoder/ElectraMediaOpusDecoder.cpp
+++ b/Engine/Plugins/Media/OpusDecoderElectra/Source/OpusDecoderElectra/Private/OpusDecoder/ElectraMediaOpusDecoder.cpp
@@ -14,7 +14,7 @@
 #include "IElectraDecoderOutputAudio.h"
 #include "ElectraDecodersUtils.h"
 #include "Utils/ElectraBitstreamProcessorDefault.h"
-#include "Utilities/UtilitiesMP4.h"		// for parsing the 'dOps' box
+#include "MP4Utilities.h"		// for parsing the 'dOps' box
 
 #include "opus_multistream.h"
 
@@ -336,7 +336,7 @@ bool FElectraOpusDecoder::Parse_dOps(FDOpsConfig& OutConfig, const TArray<uint8>
 		return !bFailOnError ? false : PostError(0, TEXT("Incomplete 'dOps' box"), ERRCODE_INTERNAL_FAILED_TO_PARSE_CSD);
 	}
 
-	Electra::UtilitiesMP4::FMP4AtomReader rd(IndOpsBox);
+	MP4Utilities::FMP4AtomReader rd(IndOpsBox);
 	uint32 Value32;
 	uint16 Value16;
 	uint8 Value8;
