ElectraPlayer: Stopped feeding decoders when a seek operation is pending

#rnx

#Synced-CL 0

diff --git a/Engine/Plugins/Media/ElectraPlayer/Source/ElectraPlayerRuntime/Private/Runtime/Player/AdaptiveStreamingPlayerDecoders.cpp b/Engine/Plugins/Media/ElectraPlayer/Source/ElectraPlayerRuntime/Private/Runtime/Player/AdaptiveStreamingPlayerDecoders.cpp
--- a/Engine/Plugins/Media/ElectraPlayer/Source/ElectraPlayerRuntime/Private/Runtime/Player/AdaptiveStreamingPlayerDecoders.cpp
+++ b/Engine/Plugins/Media/ElectraPlayer/Source/ElectraPlayerRuntime/Private/Runtime/Player/AdaptiveStreamingPlayerDecoders.cpp
@@ -567,6 +567,14 @@ void FAdaptiveStreamingPlayer::FeedDecoder(EStreamType Type, IAccessUnitBufferIn
 	{
 		return;
 	}
+	// If there is a seek pending we do not feed the decoders since their output will be discarded
+	// by the media sample queues anyway due to non-matching sequence indices.
+	// This has the side effect of quickly draining our buffers, possibly resulting in a buffer
+	// underrun that will not even apply due to the upcoming seek.
+	if (SeekVars.HasPendingSeek())
+	{
+		return;
+	}
 
 	FBufferStats* pStats = nullptr;
 	FStreamCodecInformation* CurrentCodecInfo = nullptr;
diff --git a/Engine/Plugins/Media/ElectraPlayer/Source/ElectraPlayerRuntime/Private/Runtime/Player/AdaptiveStreamingPlayerInternal.h b/Engine/Plugins/Media/ElectraPlayer/Source/ElectraPlayerRuntime/Private/Runtime/Player/AdaptiveStreamingPlayerInternal.h
--- a/Engine/Plugins/Media/ElectraPlayer/Source/ElectraPlayerRuntime/Private/Runtime/Player/AdaptiveStreamingPlayerInternal.h
+++ b/Engine/Plugins/Media/ElectraPlayer/Source/ElectraPlayerRuntime/Private/Runtime/Player/AdaptiveStreamingPlayerInternal.h
@@ -1985,6 +1985,12 @@ private:
 		{
 			ActiveRequest.Reset();
 		}
+		
+		bool HasPendingSeek()
+		{
+			FScopeLock lock(&Lock);
+			return PendingRequest.IsSet();
+		}
 
 		bool bPrerollDone = false;
 		bool bIsPlayStart = true;
