ElectraPlayer: moved mp4 metadata out into MP4Utilities plugin and refactored the player to use it from there

#rnx

#Synced-CL 50154202

diff --git a/Engine/Plugins/Media/ElectraPlayer/Source/ElectraPlayerRuntime/Private/ElectraPlayer.cpp b/Engine/Plugins/Media/ElectraPlayer/Source/ElectraPlayerRuntime/Private/ElectraPlayer.cpp
--- a/Engine/Plugins/Media/ElectraPlayer/Source/ElectraPlayerRuntime/Private/ElectraPlayer.cpp
+++ b/Engine/Plugins/Media/ElectraPlayer/Source/ElectraPlayerRuntime/Private/ElectraPlayer.cpp
@@ -19,6 +19,7 @@
 #include "MediaMetaDataDecoderOutput.h"
 #include "IElectraMetadataSample.h"
 #include "IMediaMetadataItem.h"
+#include "MediaStreamMetadata.h"
 
 #include "ElectraDecodersPlatformResources.h"
 
@@ -322,20 +323,6 @@ public:
 };
 
 
-class FStreamMetadataItem : public IMediaMetadataItem
-{
-public:
-	FStreamMetadataItem(const TSharedPtr<Electra::IMediaStreamMetadata::IItem, ESPMode::ThreadSafe>& InItem) : Item(InItem.ToSharedRef())
-	{ }
-	virtual ~FStreamMetadataItem() = default;
-	const FString& GetLanguageCode() const override	{ return Item->GetLanguageCode(); }
-	const FString& GetMimeType() const override		{ return Item->GetMimeType(); }
-	const FVariant& GetValue() const override		{ return Item->GetValue(); }
-private:
-	TSharedRef<Electra::IMediaStreamMetadata::IItem, ESPMode::ThreadSafe> Item;
-};
-
-
 //-----------------------------------------------------------------------------
 
 std::atomic<uint32> FElectraPlayer::NextPlayerUniqueID { 0 };
@@ -2223,26 +2210,14 @@ void FElectraPlayer::HandlePlayerEventSeekCompleted()
 	MediaStateOnSeekFinished();
 }
 
-void FElectraPlayer::HandlePlayerMediaMetadataChanged(const TSharedPtrTS<Electra::UtilsMP4::FMetadataParser>& InMetadata)
+void FElectraPlayer::HandlePlayerMediaMetadataChanged(const TSharedPtrTS<Electra::FMetadataParser>& InMetadata)
 {
 	if (InMetadata.IsValid())
 	{
-		TSharedPtr<TMap<FString, TArray<TSharedPtr<Electra::IMediaStreamMetadata::IItem, ESPMode::ThreadSafe>>>, ESPMode::ThreadSafe> PlayerMeta = InMetadata->GetMediaStreamMetadata();
+		TSharedPtr<TMap<FString, TArray<TUniquePtr<IMediaMetadataItem>>>, ESPMode::ThreadSafe> PlayerMeta = InMetadata->GetMediaStreamMetadata();
 		if (PlayerMeta.IsValid())
 		{
-			TSharedPtr<TMap<FString, TArray<TUniquePtr<IMediaMetadataItem>>>, ESPMode::ThreadSafe> NewMeta(new TMap<FString, TArray<TUniquePtr<IMediaMetadataItem>>>);
-			for(auto& PlayerMetaItem : *PlayerMeta)
-			{
-				TArray<TUniquePtr<IMediaMetadataItem>>& NewItemList = NewMeta->Emplace(PlayerMetaItem.Key);
-				for(auto& PlayerMetaListItem : PlayerMetaItem.Value)
-				{
-					if (PlayerMetaListItem.IsValid())
-					{
-						NewItemList.Emplace(MakeUnique<FStreamMetadataItem>(PlayerMetaListItem));
-					}
-				}
-			}
-			CurrentStreamMetadata = MoveTemp(NewMeta);
+			CurrentStreamMetadata = MoveTemp(PlayerMeta);
 			DeferredMediaEvents.Enqueue(EMediaEvent::MetadataChanged);
 		}
 	}
@@ -3838,7 +3813,7 @@ void FElectraPlayer::ReportSeekCompleted()
 	DeferredPlayerEvents.Enqueue(MakeSharedTS<FPlayerMetricEventBase>(FPlayerMetricEventBase::EType::SeekCompleted));
 }
 
-void FElectraPlayer::ReportMediaMetadataChanged(TSharedPtrTS<Electra::UtilsMP4::FMetadataParser> InMetadata)
+void FElectraPlayer::ReportMediaMetadataChanged(TSharedPtrTS<Electra::FMetadataParser> InMetadata)
 {
 	DeferredPlayerEvents.Enqueue(MakeSharedTS<FPlayerMetricEvent_MediaMetadataChange>(InMetadata));
 }
diff --git a/Engine/Plugins/Media/ElectraPlayer/Source/ElectraPlayerRuntime/Private/ElectraPlayer.h b/Engine/Plugins/Media/ElectraPlayer/Source/ElectraPlayerRuntime/Private/ElectraPlayer.h
--- a/Engine/Plugins/Media/ElectraPlayer/Source/ElectraPlayerRuntime/Private/ElectraPlayer.h
+++ b/Engine/Plugins/Media/ElectraPlayer/Source/ElectraPlayerRuntime/Private/ElectraPlayer.h
@@ -159,7 +159,7 @@ public:
 	void ReportJumpInPlayPosition(const FTimeValue& InToNewTime, const FTimeValue& InFromTime, Metrics::ETimeJumpReason InTimejumpReason) override;
 	void ReportPlaybackStopped() override;
 	void ReportSeekCompleted() override;
-	void ReportMediaMetadataChanged(TSharedPtrTS<Electra::UtilsMP4::FMetadataParser> InMetadata) override;
+	void ReportMediaMetadataChanged(TSharedPtrTS<Electra::FMetadataParser> InMetadata) override;
 	void ReportError(const FString& InErrorReason) override;
 	void ReportLogMessage(IInfoLog::ELevel InLogLevel, const FString& InLogMessage, int64 InPlayerWallclockMilliseconds) override;
 	void ReportVideoFrameStats(const Metrics::FFrameStats& InFrameStats) override;
@@ -220,7 +220,7 @@ private:
 	void HandlePlayerEventJumpInPlayPosition(const FTimeValue& InToNewTime, const FTimeValue& InFromTime, Metrics::ETimeJumpReason InTimejumpReason);
 	void HandlePlayerEventPlaybackStopped();
 	void HandlePlayerEventSeekCompleted();
-	void HandlePlayerMediaMetadataChanged(const TSharedPtrTS<Electra::UtilsMP4::FMetadataParser>& InMetadata);
+	void HandlePlayerMediaMetadataChanged(const TSharedPtrTS<Electra::FMetadataParser>& InMetadata);
 	void HandlePlayerEventError(const FString& InErrorReason);
 	void HandlePlayerEventLogMessage(IInfoLog::ELevel InLogLevel, const FString& InLogMessage, int64 InPlayerWallclockMilliseconds);
 	void HandlePlayerEventDroppedVideoFrame(const Metrics::FFrameStats& InFrameStats);
@@ -449,8 +449,8 @@ private:
 	};
 	struct FPlayerMetricEvent_MediaMetadataChange : public FPlayerMetricEventBase
 	{
-		FPlayerMetricEvent_MediaMetadataChange(const TSharedPtrTS<Electra::UtilsMP4::FMetadataParser>& InMetadata) : FPlayerMetricEventBase(EType::MediaMetadataChanged), NewMetadata(InMetadata) {}
-		TSharedPtrTS<Electra::UtilsMP4::FMetadataParser> NewMetadata;
+		FPlayerMetricEvent_MediaMetadataChange(const TSharedPtrTS<Electra::FMetadataParser>& InMetadata) : FPlayerMetricEventBase(EType::MediaMetadataChanged), NewMetadata(InMetadata) {}
+		TSharedPtrTS<Electra::FMetadataParser> NewMetadata;
 	};
 	struct FPlayerMetricEvent_Error : public FPlayerMetricEventBase
 	{
diff --git a/Engine/Plugins/Media/ElectraPlayer/Source/ElectraPlayerRuntime/Private/Runtime/MediaStreamMetadata.h b/Engine/Plugins/Media/ElectraPlayer/Source/ElectraPlayerRuntime/Private/Runtime/MediaStreamMetadata.h
new file mode 100644
--- /dev/null
+++ b/Engine/Plugins/Media/ElectraPlayer/Source/ElectraPlayerRuntime/Private/Runtime/MediaStreamMetadata.h
@@ -0,0 +1,13 @@
+// Copyright Epic Games, Inc. All Rights Reserved.
+
+#pragma once
+
+#include "MP4Metadata.h"
+
+namespace Electra
+{
+
+// We use the mp4 metadata parser as-is for now.
+using FMetadataParser = MP4Utilities::FMP4MetadataParser;
+
+} // namespace Electra
diff --git a/Engine/Plugins/Media/ElectraPlayer/Source/ElectraPlayerRuntime/Private/Runtime/Player/AdaptiveStreamingPlayerABR.cpp b/Engine/Plugins/Media/ElectraPlayer/Source/ElectraPlayerRuntime/Private/Runtime/Player/AdaptiveStreamingPlayerABR.cpp
--- a/Engine/Plugins/Media/ElectraPlayer/Source/ElectraPlayerRuntime/Private/Runtime/Player/AdaptiveStreamingPlayerABR.cpp
+++ b/Engine/Plugins/Media/ElectraPlayer/Source/ElectraPlayerRuntime/Private/Runtime/Player/AdaptiveStreamingPlayerABR.cpp
@@ -71,7 +71,7 @@ namespace Electra
 		void ReportJumpInPlayPosition(const FTimeValue& ToNewTime, const FTimeValue& FromTime, Metrics::ETimeJumpReason TimejumpReason) override {}
 		void ReportPlaybackStopped() override {}
 		void ReportSeekCompleted() override {}
-		void ReportMediaMetadataChanged(TSharedPtrTS<UtilsMP4::FMetadataParser> Metadata) override {}
+		void ReportMediaMetadataChanged(TSharedPtrTS<FMetadataParser> Metadata) override {}
 		void ReportError(const FString& ErrorReason) override {}
 		void ReportLogMessage(IInfoLog::ELevel InLogLevel, const FString& LogMessage, int64 PlayerWallclockMilliseconds) override {}
 		void ReportVideoFrameStats(const Metrics::FFrameStats& InFrameStats) override {}
@@ -256,7 +256,7 @@ namespace Electra
 							{
 								return a->Bitrate < b->Bitrate;
 							});
-							
+
 							// Set the lowest and highest bitrate with the CMCD handler
 							PlayerSessionServices->GetCMCDHandler()->SetPlayableBitrateRange(EStreamType::Video, StreamInformationVideo.Num() ? StreamInformationVideo[0]->Bitrate : 0, StreamInformationVideo.Num() ? StreamInformationVideo.Last()->Bitrate : 0);
 							break;
diff --git a/Engine/Plugins/Media/ElectraPlayer/Source/ElectraPlayerRuntime/Private/Runtime/Player/AdaptiveStreamingPlayerInternal.h b/Engine/Plugins/Media/ElectraPlayer/Source/ElectraPlayerRuntime/Private/Runtime/Player/AdaptiveStreamingPlayerInternal.h
--- a/Engine/Plugins/Media/ElectraPlayer/Source/ElectraPlayerRuntime/Private/Runtime/Player/AdaptiveStreamingPlayerInternal.h
+++ b/Engine/Plugins/Media/ElectraPlayer/Source/ElectraPlayerRuntime/Private/Runtime/Player/AdaptiveStreamingPlayerInternal.h
@@ -27,7 +27,7 @@
 #include "Player/CMCDHandler.h"
 #include "Player/ABRRules/ABRStatisticTypes.h"
 #include "OutputHandler.h"
-#include "Utilities/UtilsMP4.h"
+#include "MediaStreamMetadata.h"
 #include "ElectraCDM.h"
 #include "InfoLog.h"
 
@@ -714,7 +714,7 @@ struct FMetricEvent
 		};
 		struct FMediaMedataDataChanged
 		{
-			TSharedPtrTS<UtilsMP4::FMetadataParser> NewMetadata;
+			TSharedPtrTS<FMetadataParser> NewMetadata;
 		};
 		FString								URL;
 		Metrics::EBufferingReason			BufferingReason;
@@ -906,7 +906,7 @@ struct FMetricEvent
 		Evt->Param.SeekComplete.bWasAlreadyThere = bWasAlreadyThere;
 		return Evt;
 	}
-	static TSharedPtrTS<FMetricEvent> ReportMediaMetadataChanged(const TSharedPtrTS<UtilsMP4::FMetadataParser>& InNextMetadata)
+	static TSharedPtrTS<FMetricEvent> ReportMediaMetadataChanged(const TSharedPtrTS<FMetadataParser>& InNextMetadata)
 	{
 		TSharedPtrTS<FMetricEvent> Evt = MakeSharedTS<FMetricEvent>();
 		Evt->Type = EType::MediaMetadataChanged;
@@ -2095,7 +2095,7 @@ private:
 			// Do NOT reset the active metadata, but the time it became valid!
 			ActiveSince.SetToInvalid();
 		}
-		void AddEntry(const FTimeValue& InValidFrom, const TSharedPtrTS<UtilsMP4::FMetadataParser>& InMetadata, bool bInTriggerInternalRefresh)
+		void AddEntry(const FTimeValue& InValidFrom, const TSharedPtrTS<FMetadataParser>& InMetadata, bool bInTriggerInternalRefresh)
 		{
 			FEntry& e = NextEntries.Emplace_GetRef();
 			e.ValidFrom = InValidFrom.IsValid() ? InValidFrom : FTimeValue::GetZero();
@@ -2117,7 +2117,7 @@ private:
 			ChangedAndUpdate
 		};
 		EResult Handle(const FTimeValue& InAtTime);
-		TSharedPtrTS<UtilsMP4::FMetadataParser> GetActive() const
+		TSharedPtrTS<FMetadataParser> GetActive() const
 		{
 			return ActiveMetadata;
 		}
@@ -2125,11 +2125,11 @@ private:
 		struct FEntry
 		{
 			FTimeValue ValidFrom;
-			TSharedPtrTS<UtilsMP4::FMetadataParser> Metadata;
+			TSharedPtrTS<FMetadataParser> Metadata;
 			bool bTriggerInternalRefresh = false;
 		};
 		TArray<FEntry> NextEntries;
-		TSharedPtrTS<UtilsMP4::FMetadataParser> ActiveMetadata;
+		TSharedPtrTS<FMetadataParser> ActiveMetadata;
 		FTimeValue ActiveSince;
 	};
 
diff --git a/Engine/Plugins/Media/ElectraPlayer/Source/ElectraPlayerRuntime/Private/Runtime/Player/AdaptiveStreamingPlayerManifest.cpp b/Engine/Plugins/Media/ElectraPlayer/Source/ElectraPlayerRuntime/Private/Runtime/Player/AdaptiveStreamingPlayerManifest.cpp
--- a/Engine/Plugins/Media/ElectraPlayer/Source/ElectraPlayerRuntime/Private/Runtime/Player/AdaptiveStreamingPlayerManifest.cpp
+++ b/Engine/Plugins/Media/ElectraPlayer/Source/ElectraPlayerRuntime/Private/Runtime/Player/AdaptiveStreamingPlayerManifest.cpp
@@ -13,6 +13,7 @@
 #include "Utilities/StringHelpers.h"
 #include "Utilities/URLParser.h"
 #include "ParameterDictionary.h"
+#include "MediaStreamMetadata.h"
 
 
 namespace Electra
@@ -420,7 +421,7 @@ void FAdaptiveStreamingPlayer::UpdateManifest()
 
 FAdaptiveStreamingPlayer::FMediaMetadataUpdate::EResult FAdaptiveStreamingPlayer::FMediaMetadataUpdate::Handle(const FTimeValue& InAtTime)
 {
-	TSharedPtrTS<UtilsMP4::FMetadataParser> NextMetadata;
+	TSharedPtrTS<FMetadataParser> NextMetadata;
 	FTimeValue NextActiveTime;
 	bool bTriggerInternalRefresh = false;
 	while(NextEntries.Num())
diff --git a/Engine/Plugins/Media/ElectraPlayer/Source/ElectraPlayerRuntime/Private/Runtime/Player/AdaptiveStreamingPlayerMetrics.h b/Engine/Plugins/Media/ElectraPlayer/Source/ElectraPlayerRuntime/Private/Runtime/Player/AdaptiveStreamingPlayerMetrics.h
--- a/Engine/Plugins/Media/ElectraPlayer/Source/ElectraPlayerRuntime/Private/Runtime/Player/AdaptiveStreamingPlayerMetrics.h
+++ b/Engine/Plugins/Media/ElectraPlayer/Source/ElectraPlayerRuntime/Private/Runtime/Player/AdaptiveStreamingPlayerMetrics.h
@@ -9,7 +9,7 @@
 #include "Player/AdaptiveStreamingPlayerABR_State.h"
 #include "Player/Playlist.h"
 #include "ElectraHTTPStream.h"
-#include "Utilities/UtilsMP4.h"
+#include "MediaStreamMetadata.h"
 
 namespace Electra
 {
@@ -409,7 +409,7 @@ public:
 	/**
 	 * Called when the media metadata has changed.
 	 */
-	virtual void ReportMediaMetadataChanged(TSharedPtrTS<UtilsMP4::FMetadataParser> Metadata) = 0;
+	virtual void ReportMediaMetadataChanged(TSharedPtrTS<FMetadataParser> Metadata) = 0;
 
 	/**
 	 * Called when an error occurs. Errors always result in termination of playback.
diff --git a/Engine/Plugins/Media/ElectraPlayer/Source/ElectraPlayerRuntime/Private/Runtime/Player/CMCDHandler.cpp b/Engine/Plugins/Media/ElectraPlayer/Source/ElectraPlayerRuntime/Private/Runtime/Player/CMCDHandler.cpp
--- a/Engine/Plugins/Media/ElectraPlayer/Source/ElectraPlayerRuntime/Private/Runtime/Player/CMCDHandler.cpp
+++ b/Engine/Plugins/Media/ElectraPlayer/Source/ElectraPlayerRuntime/Private/Runtime/Player/CMCDHandler.cpp
@@ -926,7 +926,7 @@ void FCMCDHandler::ReportPlaybackEnded() { }
 void FCMCDHandler::ReportJumpInPlayPosition(const FTimeValue& InToNewTime, const FTimeValue& InFromTime, Metrics::ETimeJumpReason InTimejumpReason) { }
 void FCMCDHandler::ReportPlaybackStopped() { }
 void FCMCDHandler::ReportSeekCompleted() { }
-void FCMCDHandler::ReportMediaMetadataChanged(TSharedPtrTS<UtilsMP4::FMetadataParser> InMetadata) { }
+void FCMCDHandler::ReportMediaMetadataChanged(TSharedPtrTS<FMetadataParser> InMetadata) { }
 void FCMCDHandler::ReportError(const FString& InErrorReason) { }
 void FCMCDHandler::ReportLogMessage(IInfoLog::ELevel InLogLevel, const FString& InLogMessage, int64 InPlayerWallclockMilliseconds) { }
 void FCMCDHandler::ReportVideoFrameStats(const Metrics::FFrameStats& InFrameStats) { }
diff --git a/Engine/Plugins/Media/ElectraPlayer/Source/ElectraPlayerRuntime/Private/Runtime/Player/CMCDHandler.h b/Engine/Plugins/Media/ElectraPlayer/Source/ElectraPlayerRuntime/Private/Runtime/Player/CMCDHandler.h
--- a/Engine/Plugins/Media/ElectraPlayer/Source/ElectraPlayerRuntime/Private/Runtime/Player/CMCDHandler.h
+++ b/Engine/Plugins/Media/ElectraPlayer/Source/ElectraPlayerRuntime/Private/Runtime/Player/CMCDHandler.h
@@ -14,6 +14,7 @@
 #include "Utilities/URLParser.h"
 #include "StreamTypes.h"
 #include "Player/AdaptiveStreamingPlayerMetrics.h"
+#include "MediaStreamMetadata.h"
 
 
 namespace Electra
@@ -273,7 +274,7 @@ private:
 	void ReportJumpInPlayPosition(const FTimeValue& InToNewTime, const FTimeValue& InFromTime, Metrics::ETimeJumpReason InTimejumpReason) override;
 	void ReportPlaybackStopped() override;
 	void ReportSeekCompleted() override;
-	void ReportMediaMetadataChanged(TSharedPtrTS<UtilsMP4::FMetadataParser> InMetadata) override;
+	void ReportMediaMetadataChanged(TSharedPtrTS<FMetadataParser> InMetadata) override;
 	void ReportError(const FString& InErrorReason) override;
 	void ReportLogMessage(IInfoLog::ELevel InLogLevel, const FString& InLogMessage, int64 InPlayerWallclockMilliseconds) override;
 	void ReportVideoFrameStats(const Metrics::FFrameStats& InFrameStats) override;
diff --git a/Engine/Plugins/Media/ElectraPlayer/Source/ElectraPlayerRuntime/Private/Runtime/Player/DASH/ManifestDASH.cpp b/Engine/Plugins/Media/ElectraPlayer/Source/ElectraPlayerRuntime/Private/Runtime/Player/DASH/ManifestDASH.cpp
--- a/Engine/Plugins/Media/ElectraPlayer/Source/ElectraPlayerRuntime/Private/Runtime/Player/DASH/ManifestDASH.cpp
+++ b/Engine/Plugins/Media/ElectraPlayer/Source/ElectraPlayerRuntime/Private/Runtime/Player/DASH/ManifestDASH.cpp
@@ -509,7 +509,7 @@ void FManifestDASH::GetTrackMetadata(TArray<FTrackMetadata>& OutMetadata, EStrea
 		Manifest->GetPeriods()[0]->GetMetaData(OutMetadata, StreamType);
 	}
 }
-void FManifestDASH::UpdateRunningMetaData(TSharedPtrTS<UtilsMP4::FMetadataParser> InUpdatedMetaData)
+void FManifestDASH::UpdateRunningMetaData(TSharedPtrTS<FMetadataParser> InUpdatedMetaData)
 {
 	// No-op.
 }
diff --git a/Engine/Plugins/Media/ElectraPlayer/Source/ElectraPlayerRuntime/Private/Runtime/Player/DASH/ManifestDASH.h b/Engine/Plugins/Media/ElectraPlayer/Source/ElectraPlayerRuntime/Private/Runtime/Player/DASH/ManifestDASH.h
--- a/Engine/Plugins/Media/ElectraPlayer/Source/ElectraPlayerRuntime/Private/Runtime/Player/DASH/ManifestDASH.h
+++ b/Engine/Plugins/Media/ElectraPlayer/Source/ElectraPlayerRuntime/Private/Runtime/Player/DASH/ManifestDASH.h
@@ -37,7 +37,7 @@ public:
 	TRangeSet<double> GetPossiblePlaybackRates(EPlayRateType InForType) const override;
 	TSharedPtrTS<IProducerReferenceTimeInfo> GetProducerReferenceTimeInfo(int64 ID) const override;
 	void GetTrackMetadata(TArray<FTrackMetadata>& OutMetadata, EStreamType StreamType) const override;
-	void UpdateRunningMetaData(TSharedPtrTS<UtilsMP4::FMetadataParser> InUpdatedMetaData) override;
+	void UpdateRunningMetaData(TSharedPtrTS<FMetadataParser> InUpdatedMetaData) override;
 	void UpdateDynamicRefetchCounter() override;
 	void PrepareForLooping(int32 InNumLoopsToAdd) override;
 	void TriggerClockSync(IManifest::EClockSyncType InClockSyncType) override;
diff --git a/Engine/Plugins/Media/ElectraPlayer/Source/ElectraPlayerRuntime/Private/Runtime/Player/HLS/PlaylistHLS.cpp b/Engine/Plugins/Media/ElectraPlayer/Source/ElectraPlayerRuntime/Private/Runtime/Player/HLS/PlaylistHLS.cpp
--- a/Engine/Plugins/Media/ElectraPlayer/Source/ElectraPlayerRuntime/Private/Runtime/Player/HLS/PlaylistHLS.cpp
+++ b/Engine/Plugins/Media/ElectraPlayer/Source/ElectraPlayerRuntime/Private/Runtime/Player/HLS/PlaylistHLS.cpp
@@ -2525,7 +2525,7 @@ void FActiveHLSPlaylist::GetTrackMetadata(TArray<FTrackMetadata>& OutMetadata, E
 	}
 }
 
-void FActiveHLSPlaylist::UpdateRunningMetaData(TSharedPtrTS<UtilsMP4::FMetadataParser> InUpdatedMetaData)
+void FActiveHLSPlaylist::UpdateRunningMetaData(TSharedPtrTS<FMetadataParser> InUpdatedMetaData)
 {
 	// Not used with HLS
 }
diff --git a/Engine/Plugins/Media/ElectraPlayer/Source/ElectraPlayerRuntime/Private/Runtime/Player/HLS/PlaylistHLS.h b/Engine/Plugins/Media/ElectraPlayer/Source/ElectraPlayerRuntime/Private/Runtime/Player/HLS/PlaylistHLS.h
--- a/Engine/Plugins/Media/ElectraPlayer/Source/ElectraPlayerRuntime/Private/Runtime/Player/HLS/PlaylistHLS.h
+++ b/Engine/Plugins/Media/ElectraPlayer/Source/ElectraPlayerRuntime/Private/Runtime/Player/HLS/PlaylistHLS.h
@@ -14,6 +14,7 @@
 #include "Player/StreamSegmentReaderCommon.h"
 #include "Player/DRM/DRMManager.h"
 #include "PlayerFacility.h"
+#include "MediaStreamMetadata.h"
 #include "ErrorDetail.h"
 
 namespace Electra
@@ -547,7 +548,7 @@ public:
 	TRangeSet<double> GetPossiblePlaybackRates(EPlayRateType InForType) const override;
 	TSharedPtrTS<IProducerReferenceTimeInfo> GetProducerReferenceTimeInfo(int64 InID) const override;
 	void GetTrackMetadata(TArray<FTrackMetadata>& OutMetadata, EStreamType InStreamType) const override;
-	void UpdateRunningMetaData(TSharedPtrTS<UtilsMP4::FMetadataParser> InUpdatedMetaData) override;
+	void UpdateRunningMetaData(TSharedPtrTS<FMetadataParser> InUpdatedMetaData) override;
 	void UpdateDynamicRefetchCounter() override;
 	void PrepareForLooping(int32 InNumLoopsToAdd) override;
 	void TriggerClockSync(EClockSyncType InClockSyncType) override;
diff --git a/Engine/Plugins/Media/ElectraPlayer/Source/ElectraPlayerRuntime/Private/Runtime/Player/Manifest.h b/Engine/Plugins/Media/ElectraPlayer/Source/ElectraPlayerRuntime/Private/Runtime/Player/Manifest.h
--- a/Engine/Plugins/Media/ElectraPlayer/Source/ElectraPlayerRuntime/Private/Runtime/Player/Manifest.h
+++ b/Engine/Plugins/Media/ElectraPlayer/Source/ElectraPlayerRuntime/Private/Runtime/Player/Manifest.h
@@ -10,7 +10,7 @@
 #include "Player/AdaptiveStreamingPlayerMetrics.h"
 #include "Player/PlayerSessionServices.h"
 #include "ParameterDictionary.h"
-#include "Utilities/UtilsMP4.h"
+#include "MediaStreamMetadata.h"
 
 namespace Electra
 {
@@ -374,7 +374,7 @@ namespace Electra
 		virtual void GetTrackMetadata(TArray<FTrackMetadata>& OutMetadata, EStreamType StreamType) const = 0;
 
 		//! Updates long term playback metadata (mostly for audio casts)
-		virtual void UpdateRunningMetaData(TSharedPtrTS<UtilsMP4::FMetadataParser> InUpdatedMetaData) = 0;
+		virtual void UpdateRunningMetaData(TSharedPtrTS<FMetadataParser> InUpdatedMetaData) = 0;
 
 		//
 		virtual FTimeValue GetMinBufferTime() const = 0;
@@ -598,7 +598,7 @@ namespace Electra
 	class FPlaylistMetadataUpdateMessage : public IPlayerMessage
 	{
 	public:
-		static TSharedPtrTS<IPlayerMessage> Create(FTimeValue InValidFrom, TSharedPtrTS<UtilsMP4::FMetadataParser> InMetadata, bool bInTriggerInternalRefresh)
+		static TSharedPtrTS<IPlayerMessage> Create(FTimeValue InValidFrom, TSharedPtrTS<FMetadataParser> InMetadata, bool bInTriggerInternalRefresh)
 		{
 			TSharedPtrTS<FPlaylistMetadataUpdateMessage> p(new FPlaylistMetadataUpdateMessage(InValidFrom, InMetadata, bInTriggerInternalRefresh));
 			return p;
@@ -620,7 +620,7 @@ namespace Electra
 			return ValidFrom;
 		}
 
-		TSharedPtrTS<UtilsMP4::FMetadataParser> GetMetadata() const
+		TSharedPtrTS<FMetadataParser> GetMetadata() const
 		{
 			return Metadata;
 		}
@@ -631,13 +631,13 @@ namespace Electra
 		}
 
 	private:
-		FPlaylistMetadataUpdateMessage(FTimeValue InValidFrom, TSharedPtrTS<UtilsMP4::FMetadataParser> InMetadata, bool bInTriggerInternalRefresh)
+		FPlaylistMetadataUpdateMessage(FTimeValue InValidFrom, TSharedPtrTS<FMetadataParser> InMetadata, bool bInTriggerInternalRefresh)
 			: ValidFrom(InValidFrom)
 			, Metadata(MoveTemp(InMetadata))
 			, bTriggerInternalRefresh(bInTriggerInternalRefresh)
 		{ }
 		FTimeValue ValidFrom;
-		TSharedPtrTS<UtilsMP4::FMetadataParser> Metadata;
+		TSharedPtrTS<FMetadataParser> Metadata;
 		bool bTriggerInternalRefresh;
 	};
 
diff --git a/Engine/Plugins/Media/ElectraPlayer/Source/ElectraPlayerRuntime/Private/Runtime/Player/StreamSegmentReaderCommon.cpp b/Engine/Plugins/Media/ElectraPlayer/Source/ElectraPlayerRuntime/Private/Runtime/Player/StreamSegmentReaderCommon.cpp
--- a/Engine/Plugins/Media/ElectraPlayer/Source/ElectraPlayerRuntime/Private/Runtime/Player/StreamSegmentReaderCommon.cpp
+++ b/Engine/Plugins/Media/ElectraPlayer/Source/ElectraPlayerRuntime/Private/Runtime/Player/StreamSegmentReaderCommon.cpp
@@ -30,6 +30,7 @@
 #include "Utils/MPEG/ElectraUtilsMPEGVideo_H264.h"
 #include "Utils/MPEG/ElectraUtilsMPEGVideo_H265.h"
 #include "Utils/MPEG/ElectraUtilsMPEGAudio.h"
+#include "MediaStreamMetadata.h"
 #include "Misc/TVariant.h"
 #include "BufferedDataReader.h"
 
@@ -3957,13 +3958,13 @@ void FStreamSegmentReaderCommon::FStreamHandler::HandleMP4Metadata(const TShared
 		const IParserISO14496_12::IMetadata* md = MoofMetadata ? MoofMetadata : MoovMetadata;
 		uint32 hdlr = md->GetHandler();
 		uint32 res0 = md->GetReserved(0);
-		TArray<UtilsMP4::FMetadataParser::FBoxInfo> Boxes;
+		TArray<MP4Utilities::FMP4BoxInfo> Boxes;
 		for(int32 i=0, iMax=md->GetNumChildBoxes(); i<iMax; ++i)
 		{
-			Boxes.Emplace(UtilsMP4::FMetadataParser::FBoxInfo(md->GetChildBoxType(i), md->GetChildBoxData(i), md->GetChildBoxDataSize(i)));
+			Boxes.Emplace(MP4Utilities::FMP4BoxInfo({.Data = MakeConstArrayView<uint8>(reinterpret_cast<const uint8*>(md->GetChildBoxData(i)), md->GetChildBoxDataSize(i)), .Type=md->GetChildBoxType(i)}));
 		}
-		TSharedPtrTS<UtilsMP4::FMetadataParser> MediaMetadata = MakeSharedTS<UtilsMP4::FMetadataParser>();
-		if (MediaMetadata->Parse(hdlr, res0, Boxes) == UtilsMP4::FMetadataParser::EResult::Success)
+		TSharedPtr<MP4Utilities::FMP4MetadataParser, ESPMode::ThreadSafe> MediaMetadata = MakeShared<MP4Utilities::FMP4MetadataParser, ESPMode::ThreadSafe>();
+		if (MediaMetadata->Parse(hdlr, res0, Boxes) == FMetadataParser::EResult::Success)
 		{
 			FTimeValue StartTime = md == MoofMetadata ? InRequest->GetFirstPTS() : InBaseTime;
 			StartTime.SetSequenceIndex(InRequest->TimestampSequenceIndex);
diff --git a/Engine/Plugins/Media/ElectraPlayer/Source/ElectraPlayerRuntime/Private/Runtime/Player/mkv/ManifestMKV.cpp b/Engine/Plugins/Media/ElectraPlayer/Source/ElectraPlayerRuntime/Private/Runtime/Player/mkv/ManifestMKV.cpp
--- a/Engine/Plugins/Media/ElectraPlayer/Source/ElectraPlayerRuntime/Private/Runtime/Player/mkv/ManifestMKV.cpp
+++ b/Engine/Plugins/Media/ElectraPlayer/Source/ElectraPlayerRuntime/Private/Runtime/Player/mkv/ManifestMKV.cpp
@@ -95,7 +95,7 @@ void FManifestMKVInternal::GetTrackMetadata(TArray<FTrackMetadata>& OutMetadata,
 	}
 }
 
-void FManifestMKVInternal::UpdateRunningMetaData(TSharedPtrTS<UtilsMP4::FMetadataParser> InUpdatedMetaData)
+void FManifestMKVInternal::UpdateRunningMetaData(TSharedPtrTS<FMetadataParser> InUpdatedMetaData)
 {
 	// No-op.
 }
diff --git a/Engine/Plugins/Media/ElectraPlayer/Source/ElectraPlayerRuntime/Private/Runtime/Player/mkv/ManifestMKV.h b/Engine/Plugins/Media/ElectraPlayer/Source/ElectraPlayerRuntime/Private/Runtime/Player/mkv/ManifestMKV.h
--- a/Engine/Plugins/Media/ElectraPlayer/Source/ElectraPlayerRuntime/Private/Runtime/Player/mkv/ManifestMKV.h
+++ b/Engine/Plugins/Media/ElectraPlayer/Source/ElectraPlayerRuntime/Private/Runtime/Player/mkv/ManifestMKV.h
@@ -12,6 +12,7 @@
 #include "ElectraPlayerPrivate.h"
 #include "Player/PlayerStreamReader.h"
 #include "Utilities/URLParser.h"
+#include "MediaStreamMetadata.h"
 
 
 namespace Electra
@@ -57,7 +58,7 @@ public:
 	void ClearDefaultEndTime() override
 	{ DefaultEndTime.SetToInvalid(); }
 	void GetTrackMetadata(TArray<FTrackMetadata>& OutMetadata, EStreamType StreamType) const override;
-	void UpdateRunningMetaData(TSharedPtrTS<UtilsMP4::FMetadataParser> InUpdatedMetaData) override;
+	void UpdateRunningMetaData(TSharedPtrTS<FMetadataParser> InUpdatedMetaData) override;
 	FTimeValue GetMinBufferTime() const override;
 	FTimeValue GetDesiredLiveLatency() const override
 	{ return FTimeValue(); }
diff --git a/Engine/Plugins/Media/ElectraPlayer/Source/ElectraPlayerRuntime/Private/Runtime/Player/mp4/ManifestMP4.cpp b/Engine/Plugins/Media/ElectraPlayer/Source/ElectraPlayerRuntime/Private/Runtime/Player/mp4/ManifestMP4.cpp
--- a/Engine/Plugins/Media/ElectraPlayer/Source/ElectraPlayerRuntime/Private/Runtime/Player/mp4/ManifestMP4.cpp
+++ b/Engine/Plugins/Media/ElectraPlayer/Source/ElectraPlayerRuntime/Private/Runtime/Player/mp4/ManifestMP4.cpp
@@ -10,8 +10,8 @@
 #include "Utilities/HashFunctions.h"
 #include "Utilities/TimeUtilities.h"
 #include "HAL/LowLevelMemTracker.h"
-#include "Utilities/UtilsMP4.h"
 #include "Utilities/BCP47-Helpers.h"
+#include "MediaStreamMetadata.h"
 #include "ElectraPlayerPrivate.h"
 #include "Player/AdaptiveStreamingPlayerResourceRequest.h"
 #include "Player/AdaptivePlayerOptionKeynames.h"
@@ -116,7 +116,7 @@ void FManifestMP4Internal::GetTrackMetadata(TArray<FTrackMetadata>& OutMetadata,
 	}
 }
 
-void FManifestMP4Internal::UpdateRunningMetaData(TSharedPtrTS<UtilsMP4::FMetadataParser> InUpdatedMetaData)
+void FManifestMP4Internal::UpdateRunningMetaData(TSharedPtrTS<FMetadataParser> InUpdatedMetaData)
 {
 	// No-op.
 }
@@ -804,13 +804,13 @@ FErrorDetail FManifestMP4Internal::FTimelineAssetMP4::Build(IPlayerSessionServic
 	{
 		uint32 hdlr = md->GetHandler();
 		uint32 res0 = md->GetReserved(0);
-		TArray<UtilsMP4::FMetadataParser::FBoxInfo> Boxes;
+		TArray<MP4Utilities::FMP4BoxInfo> Boxes;
 		for(int32 i=0, iMax=md->GetNumChildBoxes(); i<iMax; ++i)
 		{
-			Boxes.Emplace(UtilsMP4::FMetadataParser::FBoxInfo(md->GetChildBoxType(i), md->GetChildBoxData(i), md->GetChildBoxDataSize(i)));
+			Boxes.Emplace(MP4Utilities::FMP4BoxInfo({.Data = MakeConstArrayView<uint8>(reinterpret_cast<const uint8*>(md->GetChildBoxData(i)), md->GetChildBoxDataSize(i)), .Type=md->GetChildBoxType(i)}));
 		}
-		MediaMetadata = MakeSharedTS<UtilsMP4::FMetadataParser>();
-		if (MediaMetadata->Parse(hdlr, res0, Boxes) == UtilsMP4::FMetadataParser::EResult::Success)
+		MediaMetadata = MakeSharedTS<FMetadataParser>();
+		if (MediaMetadata->Parse(hdlr, res0, Boxes) == FMetadataParser::EResult::Success)
 		{
 			PlayerSessionServices->SendMessageToPlayer(FPlaylistMetadataUpdateMessage::Create(FTimeValue(), MediaMetadata, false));
 		}
diff --git a/Engine/Plugins/Media/ElectraPlayer/Source/ElectraPlayerRuntime/Private/Runtime/Player/mp4/ManifestMP4.h b/Engine/Plugins/Media/ElectraPlayer/Source/ElectraPlayerRuntime/Private/Runtime/Player/mp4/ManifestMP4.h
--- a/Engine/Plugins/Media/ElectraPlayer/Source/ElectraPlayerRuntime/Private/Runtime/Player/mp4/ManifestMP4.h
+++ b/Engine/Plugins/Media/ElectraPlayer/Source/ElectraPlayerRuntime/Private/Runtime/Player/mp4/ManifestMP4.h
@@ -12,8 +12,8 @@
 #include "ElectraPlayerPrivate.h"
 #include "Player/PlayerStreamReader.h"
 #include "Utilities/URLParser.h"
-#include "Utilities/UtilsMP4.h"
 #include "ElectraCDMClient.h"
+#include "MediaStreamMetadata.h"
 
 namespace Electra
 {
@@ -58,7 +58,7 @@ public:
 	void ClearDefaultEndTime() override
 	{ DefaultEndTime.SetToInvalid(); }
 	void GetTrackMetadata(TArray<FTrackMetadata>& OutMetadata, EStreamType StreamType) const override;
-	void UpdateRunningMetaData(TSharedPtrTS<UtilsMP4::FMetadataParser> InUpdatedMetaData) override;
+	void UpdateRunningMetaData(TSharedPtrTS<FMetadataParser> InUpdatedMetaData) override;
 	FTimeValue GetMinBufferTime() const override;
 	FTimeValue GetDesiredLiveLatency() const override
 	{ return FTimeValue(); }
@@ -269,7 +269,7 @@ public:
 		TArray<TSharedPtrTS<FAdaptationSetMP4>>			SubtitleAdaptationSets;
 		FPlayRangeEndInfo								PlayRangeEndInfo;
 		// Metadata from the 'meta' box, if any.
-		TSharedPtrTS<UtilsMP4::FMetadataParser>			MediaMetadata;
+		TSharedPtrTS<FMetadataParser>			MediaMetadata;
 		TSharedPtrTS<ElectraCDM::IMediaCDMClient>		DrmClient;
 	};
 
diff --git a/Engine/Plugins/Media/ElectraPlayer/Source/ElectraPlayerRuntime/Private/Runtime/Player/mpegaudio/ManifestMPEGAudio.cpp b/Engine/Plugins/Media/ElectraPlayer/Source/ElectraPlayerRuntime/Private/Runtime/Player/mpegaudio/ManifestMPEGAudio.cpp
--- a/Engine/Plugins/Media/ElectraPlayer/Source/ElectraPlayerRuntime/Private/Runtime/Player/mpegaudio/ManifestMPEGAudio.cpp
+++ b/Engine/Plugins/Media/ElectraPlayer/Source/ElectraPlayerRuntime/Private/Runtime/Player/mpegaudio/ManifestMPEGAudio.cpp
@@ -128,9 +128,9 @@ void FManifestMPEGAudioInternal::GetTrackMetadata(TArray<FTrackMetadata>& OutMet
 }
 
 
-void FManifestMPEGAudioInternal::UpdateRunningMetaData(TSharedPtrTS<UtilsMP4::FMetadataParser> InUpdatedMetaData)
+void FManifestMPEGAudioInternal::UpdateRunningMetaData(TSharedPtrTS<FMetadataParser> InUpdatedMetaData)
 {
-	TSharedPtr<TMap<FString, TArray<TSharedPtr<IMediaStreamMetadata::IItem, ESPMode::ThreadSafe>>>, ESPMode::ThreadSafe> md = InUpdatedMetaData->GetMediaStreamMetadata();
+	TSharedPtr<TMap<FString, TArray<TUniquePtr<IMediaMetadataItem>>>, ESPMode::ThreadSafe> md = InUpdatedMetaData->GetMediaStreamMetadata();
 	if (md.IsValid() && MediaAsset.IsValid())
 	{
 		for(auto& it : *md)
@@ -717,7 +717,7 @@ FErrorDetail FManifestMPEGAudioInternal::FTimelineAssetMPEGAudio::Build(IPlayerS
 	// Convert metadata, if present.
 	if (MPEGInfoHeader.ID3v2.IsValid())
 	{
-		TSharedPtr<UtilsMP4::FMetadataParser, ESPMode::ThreadSafe> mp = MakeShared<UtilsMP4::FMetadataParser, ESPMode::ThreadSafe>();
+		TSharedPtr<FMetadataParser, ESPMode::ThreadSafe> mp = MakeShared<FMetadataParser, ESPMode::ThreadSafe>();
 
 		for(auto& it : MPEGInfoHeader.ID3v2->GetTags())
 		{
@@ -760,7 +760,7 @@ FErrorDetail FManifestMPEGAudioInternal::FTimelineAssetMPEGAudio::Build(IPlayerS
 	}
 	else if (MPEGInfoHeader.HTTPResponseHeaders.Contains(TEXT("icy-name")))
 	{
-		TSharedPtr<UtilsMP4::FMetadataParser, ESPMode::ThreadSafe> mp = MakeShared<UtilsMP4::FMetadataParser, ESPMode::ThreadSafe>();
+		TSharedPtr<FMetadataParser, ESPMode::ThreadSafe> mp = MakeShared<FMetadataParser, ESPMode::ThreadSafe>();
 		TArray<FString> HeaderValues;
 		MPEGInfoHeader.HTTPResponseHeaders.MultiFind(TEXT("icy-name"), HeaderValues);
 		mp->AddItem(TEXT("Album"), HeaderValues[0]);
diff --git a/Engine/Plugins/Media/ElectraPlayer/Source/ElectraPlayerRuntime/Private/Runtime/Player/mpegaudio/ManifestMPEGAudio.h b/Engine/Plugins/Media/ElectraPlayer/Source/ElectraPlayerRuntime/Private/Runtime/Player/mpegaudio/ManifestMPEGAudio.h
--- a/Engine/Plugins/Media/ElectraPlayer/Source/ElectraPlayerRuntime/Private/Runtime/Player/mpegaudio/ManifestMPEGAudio.h
+++ b/Engine/Plugins/Media/ElectraPlayer/Source/ElectraPlayerRuntime/Private/Runtime/Player/mpegaudio/ManifestMPEGAudio.h
@@ -13,6 +13,7 @@
 #include "Utilities/Utilities.h"
 #include "Utilities/URLParser.h"
 #include "Utilities/UtilsMPEG.h"
+#include "MediaStreamMetadata.h"
 
 namespace Electra
 {
@@ -101,7 +102,7 @@ public:
 	void ClearDefaultEndTime() override
 	{ DefaultEndTime.SetToInvalid(); }
 	void GetTrackMetadata(TArray<FTrackMetadata>& OutMetadata, EStreamType StreamType) const override;
-	void UpdateRunningMetaData(TSharedPtrTS<UtilsMP4::FMetadataParser> InUpdatedMetaData) override;
+	void UpdateRunningMetaData(TSharedPtrTS<FMetadataParser> InUpdatedMetaData) override;
 	FTimeValue GetMinBufferTime() const override;
 	FTimeValue GetDesiredLiveLatency() const override;
 	ELiveEdgePlayMode GetLiveEdgePlayMode() const override;
diff --git a/Engine/Plugins/Media/ElectraPlayer/Source/ElectraPlayerRuntime/Private/Runtime/Player/mpegaudio/StreamReaderMPEGAudio.cpp b/Engine/Plugins/Media/ElectraPlayer/Source/ElectraPlayerRuntime/Private/Runtime/Player/mpegaudio/StreamReaderMPEGAudio.cpp
--- a/Engine/Plugins/Media/ElectraPlayer/Source/ElectraPlayerRuntime/Private/Runtime/Player/mpegaudio/StreamReaderMPEGAudio.cpp
+++ b/Engine/Plugins/Media/ElectraPlayer/Source/ElectraPlayerRuntime/Private/Runtime/Player/mpegaudio/StreamReaderMPEGAudio.cpp
@@ -9,7 +9,6 @@
 #include "PlaylistReaderMPEGAudio.h"
 #include "Utilities/Utilities.h"
 #include "Utilities/TimeUtilities.h"
-#include "Utilities/UtilsMP4.h"
 #include "Utilities/BCP47-Helpers.h"
 #include "Utils/MPEG/ElectraUtilsMPEGAudio.h"
 #include "HAL/LowLevelMemTracker.h"
@@ -21,6 +20,7 @@
 #include "Player/mpegaudio/OptionKeynamesMPEGAudio.h"
 #include "Player/PlayerSessionServices.h"
 #include "Player/PlayerStreamFilter.h"
+#include "MediaStreamMetadata.h"
 
 
 DECLARE_CYCLE_STAT(TEXT("FStreamReaderMPEGAudio_HandleRequest"), STAT_ElectraPlayer_MPEGAudio_StreamReader, STATGROUP_ElectraPlayer);
@@ -709,7 +709,7 @@ void FStreamReaderMPEGAudio::HandleRequest()
 											Title.MidInline(1, Title.Len()-2);
 										}
 
-										TSharedPtr<UtilsMP4::FMetadataParser, ESPMode::ThreadSafe> mp = MakeShared<UtilsMP4::FMetadataParser, ESPMode::ThreadSafe>();
+										TSharedPtr<FMetadataParser, ESPMode::ThreadSafe> mp = MakeShared<FMetadataParser, ESPMode::ThreadSafe>();
 										mp->AddItem(TEXT("Title"), Title);
 										PlayerSessionServices->SendMessageToPlayer(FPlaylistMetadataUpdateMessage::Create(NextAUPTS, mp, true));
 									}
diff --git a/Engine/Plugins/Media/ElectraPlayer/Source/ElectraPlayerRuntime/Private/Runtime/Utilities/UtilsMP4.cpp b/Engine/Plugins/Media/ElectraPlayer/Source/ElectraPlayerRuntime/Private/Runtime/Utilities/UtilsMP4.cpp
--- a/Engine/Plugins/Media/ElectraPlayer/Source/ElectraPlayerRuntime/Private/Runtime/Utilities/UtilsMP4.cpp
+++ b/Engine/Plugins/Media/ElectraPlayer/Source/ElectraPlayerRuntime/Private/Runtime/Utilities/UtilsMP4.cpp
@@ -9,546 +9,6 @@
 namespace Electra
 {
 
-namespace FMetadataTags
-{
-static const TCHAR * const Title = TEXT("Title");
-static const TCHAR * const Artist = TEXT("Artist");
-static const TCHAR * const Genre = TEXT("Genre");
-static const TCHAR * const Date = TEXT("Date");						// Content creation date
-static const TCHAR * const Description = TEXT("Description");
-static const TCHAR * const LongDescription = TEXT("LongDescription");
-static const TCHAR * const Album = TEXT("Album");
-static const TCHAR * const Encoder = TEXT("Encoder");
-}
-
-UtilsMP4::FMetadataParser::FMetadataParser()
-{
-	NumTotalItems = 0;
-
-	WellKnownItems.Emplace(MP4Utilities::MakeBoxAtom(0xa9u,'n','a','m'), FMetadataTags::Title);
-	WellKnownItems.Emplace(MP4Utilities::MakeBoxAtom(0xa9u,'A','R','T'), FMetadataTags::Artist);
-	WellKnownItems.Emplace(MP4Utilities::MakeBoxAtom(0xa9u,'g','e','n'), FMetadataTags::Genre);
-	WellKnownItems.Emplace(MP4Utilities::MakeBoxAtom(0xa9u,'d','a','y'), FMetadataTags::Date);
-	WellKnownItems.Emplace(MP4Utilities::MakeBoxAtom('d','e','s','c'), FMetadataTags::Description);
-	WellKnownItems.Emplace(MP4Utilities::MakeBoxAtom('l','d','e','s'), FMetadataTags::LongDescription);
-	WellKnownItems.Emplace(MP4Utilities::MakeBoxAtom(0xa9u,'a','l','b'), FMetadataTags::Album);
-	WellKnownItems.Emplace(MP4Utilities::MakeBoxAtom(0xa9u,'t','o','o'), FMetadataTags::Encoder);
-}
-
-FString UtilsMP4::FMetadataParser::PrintableBoxAtom(const uint32 InAtom)
-{
-	FString Out;
-	// Not so much just printable as alphanumeric.
-	for(uint32 i=0, Atom=InAtom; i<4; ++i, Atom<<=8)
-	{
-		int32 v = Atom >> 24;
-		if ((v>='A' && v<='Z') || (v>='a' && v<='z') || (v>='0' && v<='9') || v=='_')
-		{
-			Out.AppendChar(v);
-		}
-		else
-		{
-			// Not alphanumeric, return it as a hex string.
-			return FString::Printf(TEXT("%08x"), InAtom);
-		}
-	}
-	return Out;
-}
-
-void UtilsMP4::FMetadataParser::AddItem(const FString& InType, const FString& InValue)
-{
-	// Add a string item.
-	TSharedPtr<FItem, ESPMode::ThreadSafe> Item = MakeShared<FItem, ESPMode::ThreadSafe>();
-	Item->Type = 1;
-	Item->Value = FVariant(InValue);
-	++NumTotalItems;
-	Items.FindOrAdd(InType).Emplace(MoveTemp(Item));
-}
-
-void UtilsMP4::FMetadataParser::AddItem(const FString& InType, const FString& InMimeType, const TArray<uint8>& InValue)
-{
-	if (InMimeType.Equals(TEXT("image/jpeg")))
-	{
-		TSharedPtr<FItem, ESPMode::ThreadSafe> Item = MakeShared<FItem, ESPMode::ThreadSafe>();
-		Item->MimeType = InMimeType;
-		Item->Type = 13;
-		Item->Value = FVariant(InValue);
-		++NumTotalItems;
-		Items.FindOrAdd(InType).Emplace(MoveTemp(Item));
-	}
-	else if (InMimeType.Equals(TEXT("image/png")))
-	{
-		TSharedPtr<FItem, ESPMode::ThreadSafe> Item = MakeShared<FItem, ESPMode::ThreadSafe>();
-		Item->MimeType = InMimeType;
-		Item->Type = 14;
-		Item->Value = FVariant(InValue);
-		++NumTotalItems;
-		Items.FindOrAdd(InType).Emplace(MoveTemp(Item));
-	}
-}
-
-
-UtilsMP4::FMetadataParser::EResult UtilsMP4::FMetadataParser::Parse(uint32 InHandler, uint32 InHandlerReserved0, const TArray<UtilsMP4::FMetadataParser::FBoxInfo>& InBoxes)
-{
-	// We only support the Apple iTunes metadata at the moment.
-	if (InHandler != MP4Utilities::MakeBoxAtom('m','d','i','r'))
-	{
-		return UtilsMP4::FMetadataParser::EResult::NotSupported;
-	}
-	/*
-		As per ISO/IEC 14496-12:2015 section 8.11.1.2 the 'meta' box must contain a handler ('hdlr'), followed by nothing but optional boxes.
-		Of these optional boxes the ones recognized and defined by the document are:
-		'pitm' - 8.11.4 Primary Item Box
-		'dinf' - 8.7.1 Data Information Box
-		'iloc' - 8.11.3 The Item Location Box
-		'ipro' - 8.11.5 Item Protection Box
-		'iinf' - 8.11.6 Item Information Box
-		 ....  - IPMPControlBox (8.12 Support for Protected Streams)
-		'iref' - 8.11.12 Item Reference Box
-		'idat' - 8.11.11 Item Data Box
-
-		followed by any number of other boxes.
-
-		Presently the boxes mentioned above are not processed since they are apparently not used by any of the tools
-		that allow insertion of metadata. It appears that the way iTunes handles metadata has become the de facto standard
-		in which there is
-		    'meta'
-			   'hdlr'
-			   'ilst'
-			     ....               // See: https://developer.apple.com/library/archive/documentation/QuickTime/QTFF/Metadata/Metadata.html#//apple_ref/doc/uid/TP40000939-CH1-SW26
-				   'data'			// See: https://developer.apple.com/library/archive/documentation/QuickTime/QTFF/Metadata/Metadata.html#//apple_ref/doc/uid/TP40000939-CH1-SW27
-	*/
-
-	// By some definition the items are contained inside an 'ilst' box, so we look for that one.
-	bool bHaveILST = false;
-	for(int32 i=0; i<InBoxes.Num(); ++i)
-	{
-		if (InBoxes[i].Type == MP4Utilities::MakeBoxAtom('i','l','s','t'))
-		{
-			// We will actually parse all 'ilst' boxes if for some reason the metadata has
-			// been scattered across several.
-			bHaveILST = true;
-			Parse(InBoxes[i]);
-		}
-	}
-	return bHaveILST ? UtilsMP4::FMetadataParser::EResult::Success : UtilsMP4::FMetadataParser::EResult::MissingBox;
-}
-
-void UtilsMP4::FMetadataParser::Parse(const UtilsMP4::FMetadataParser::FBoxInfo& InBox)
-{
-	if (!InBox.Data || !InBox.Size)
-	{
-		return;
-	}
-
-	MP4Utilities::FMP4AtomReader dr(MakeConstArrayView<const uint8>(reinterpret_cast<const uint8*>(InBox.Data), InBox.Size));
-	while(dr.GetNumBytesRemaining() > 8)
-	{
-		int32 BoxSize;
-		uint32 BoxType;
-		dr.Read(BoxSize);
-		dr.Read(BoxType);
-		if (WellKnownItems.Contains(BoxType))
-		{
-			ParseBoxDataList(WellKnownItems[BoxType], dr.GetCurrentDataPointer(), BoxSize-8);
-		}
-		else switch(BoxType)
-		{
-			case MP4Utilities::MakeBoxAtom('-','-','-','-'):
-			{
-				ParseBoxDataiTunes(dr.GetCurrentDataPointer(), BoxSize-8);
-				break;
-			}
-			default:
-			{
-				ParseBoxDataList(FString(TEXT("qt."))+PrintableBoxAtom(BoxType), dr.GetCurrentDataPointer(), BoxSize-8);
-				break;
-			}
-		}
-		dr.ReadBytes(nullptr, BoxSize - 8);
-	}
-}
-
-
-void UtilsMP4::FMetadataParser::ParseBoxDataList(const FString& AsCategory, const uint8* InBoxData, uint32 InBoxSize)
-{
-	if (!InBoxData || InBoxSize <= 8)
-	{
-		return;
-	}
-
-	MP4Utilities::FMP4AtomReader dr(MakeConstArrayView<const uint8>(reinterpret_cast<const uint8*>(InBoxData), InBoxSize));
-
-#define RETURN_IF_READERROR(expr)	\
-	if (!(expr))					\
-	{								\
-		return;						\
-	}								\
-
-	for(int32 CurrentOffset=dr.GetCurrentOffset(), BoxSize=0; dr.GetNumBytesRemaining() > 8; dr.SetCurrentOffset(CurrentOffset), dr.ReadBytes(nullptr, BoxSize - 8))
-	{
-		uint32 BoxType;
-		RETURN_IF_READERROR(dr.Read(BoxSize));
-		RETURN_IF_READERROR(dr.Read(BoxType));
-		CurrentOffset = dr.GetCurrentOffset();
-		if (BoxType != MP4Utilities::MakeBoxAtom('d','a','t','a'))
-		{
-			continue;
-		}
-		uint32 TypeIndicatorAndWellKnownType = 0;
-		RETURN_IF_READERROR(dr.Read(TypeIndicatorAndWellKnownType));
-		if ((TypeIndicatorAndWellKnownType >> 24) != 0)
-		{
-			continue;
-		}
-		uint16 CountryIndicator=0;
-		RETURN_IF_READERROR(dr.Read(CountryIndicator));
-		uint16 LanguageIndicator=0;
-		RETURN_IF_READERROR(dr.Read(LanguageIndicator));
-		int32 NumRemainingDataBytes = BoxSize - 8 - (dr.GetCurrentOffset() - CurrentOffset);
-
-		TSharedPtr<FItem, ESPMode::ThreadSafe> Item = MakeShared<FItem, ESPMode::ThreadSafe>();
-
-		// We do not handle the country indicator at the moment.
-		// Likewise, the language indicator must be a directly specified language code since we do not handle
-		// language tables.
-		if (LanguageIndicator >= 1 && LanguageIndicator <= 255)
-		{
-			// This would be an index into the language table (which is a 'lang' box somewhere under the 'meta' box).
-			continue;
-		}
-		if (LanguageIndicator)
-		{
-			char Lang[3];
-			Lang[0] = (char)(0x60 + ((LanguageIndicator & 0x7c00) >> 10));
-			Lang[1] = (char)(0x60 + ((LanguageIndicator & 0x03e0) >> 5));
-			Lang[2] = (char)(0x60 + (LanguageIndicator & 0x001f));
-			Item->Language = FString::ConstructFromPtrSize(Lang, 3);
-		}
-
-		bool bSet = true;
-		int32 WKT = (int32)(TypeIndicatorAndWellKnownType & 0x00ffffffU);
-		switch(WKT)
-		{
-			// UTF-8 string
-			case 1:
-			{
-				FString String;
-				if ((bSet = dr.ReadStringUTF8(String, NumRemainingDataBytes)) == true)
-				{
-					Item->Value = String;
-					Item->Type = 1;
-				}
-				break;
-			}
-/*
-			// UTF-16 string
-			case 2:
-			{
-				FString String;
-				if ((bSet = dr.ReadStringUTF16(String, NumRemainingDataBytes)) == true)
-				{
-					Item->Value = String;
-					Item->Type = 1;
-				}
-				break;
-			}
-*/
-			case 13:	// JPEG image
-			case 14:	// PNG image
-			case 27:	// BMP image
-			{
-				TArray<uint8> Image(dr.GetCurrentDataPointer(), NumRemainingDataBytes);
-				Item->Value = MoveTemp(Image);
-				Item->Type = WKT;
-				if (WKT == 13)
-				{
-					Item->MimeType = TEXT("image/jpeg");
-				}
-				else if (WKT == 14)
-				{
-					Item->MimeType = TEXT("image/png");
-				}
-				else if (WKT == 27)
-				{
-					Item->MimeType = TEXT("image/bmp");
-				}
-				break;
-			}
-
-			// Float
-			case 23:
-			{
-				float Flt;
-				if ((bSet = dr.ReadAsNumber(Flt)) == true)
-				{
-					// Set as double!
-					Item->Value = (double)Flt;
-					Item->Type = 24;
-				}
-				break;
-			}
-			// Double
-			case 24:
-			{
-				double Dbl;
-				if ((bSet = dr.ReadAsNumber(Dbl)) == true)
-				{
-					Item->Value = Dbl;
-					Item->Type = 24;
-				}
-				break;
-			}
-			// Signed integers
-			case 65:	// int8
-			case 66:	// int16
-			case 67:	// int32
-			case 74:	// int64
-			case 21:	// 1-4 byte signed integer
-			{
-				int64 Number;
-				if ((bSet = dr.ReadAsNumber(Number, WKT==65?1:WKT==66?2:WKT==67?4:WKT==74?8:NumRemainingDataBytes)) == true)
-				{
-					Item->Value = Number;
-					Item->Type = 74;
-				}
-				break;
-			}
-			// Unsigned integers
-			case 75:	// uint8
-			case 76:	// uint16
-			case 77:	// uint32
-			case 78:	// uint64
-			case 22:	// 1-4 byte unsigned integer
-			{
-				uint64 Number;
-				if ((bSet = dr.ReadAsNumber(Number, WKT==75?1:WKT==76?2:WKT==77?4:WKT==78?8:NumRemainingDataBytes)) == true)
-				{
-					Item->Value = Number;
-					Item->Type = 78;
-				}
-				break;
-			}
-
-			// Not handled here.
-			default:
-			{
-				bSet = false;
-				break;
-			}
-		}
-
-		if (bSet)
-		{
-			++NumTotalItems;
-			Items.FindOrAdd(AsCategory).Emplace(MoveTemp(Item));
-		}
-	}
-#undef RETURN_IF_READERROR
-}
-
-void UtilsMP4::FMetadataParser::ParseBoxDataiTunes(const uint8* InBoxData, uint32 InBoxSize)
-{
-	if (!InBoxData || InBoxSize <= 8)
-	{
-		return;
-	}
-
-	MP4Utilities::FMP4AtomReader dr(MakeConstArrayView<const uint8>(reinterpret_cast<const uint8*>(InBoxData), InBoxSize));
-
-#define RETURN_IF_READERROR(expr)	\
-	if (!(expr))					\
-	{								\
-		return;						\
-	}								\
-
-	FString Name;
-	for(int32 CurrentOffset=dr.GetCurrentOffset(), BoxSize=0; dr.GetNumBytesRemaining() > 8; dr.SetCurrentOffset(CurrentOffset), dr.ReadBytes(nullptr, BoxSize - 8))
-	{
-		uint32 BoxType;
-		RETURN_IF_READERROR(dr.Read(BoxSize));
-		RETURN_IF_READERROR(dr.Read(BoxType));
-		CurrentOffset = dr.GetCurrentOffset();
-		if (BoxType == MP4Utilities::MakeBoxAtom('m','e','a','n'))
-		{
-			// 4 reserved bytes follow
-			RETURN_IF_READERROR(dr.ReadBytes(nullptr, 4));
-			FString Meaning;
-			int32 NumRemainingDataBytes = BoxSize - 8 - (dr.GetCurrentOffset() - CurrentOffset);
-			RETURN_IF_READERROR(dr.ReadStringUTF8(Meaning, NumRemainingDataBytes));
-			if (!Meaning.Equals(TEXT("com.apple.iTunes")))
-			{
-				return;
-			}
-		}
-		else if (BoxType == MP4Utilities::MakeBoxAtom('n','a','m','e'))
-		{
-			// 4 reserved bytes follow
-			RETURN_IF_READERROR(dr.ReadBytes(nullptr, 4));
-			int32 NumRemainingDataBytes = BoxSize - 8 - (dr.GetCurrentOffset() - CurrentOffset);
-			RETURN_IF_READERROR(dr.ReadStringUTF8(Name, NumRemainingDataBytes));
-		}
-		else if (BoxType == MP4Utilities::MakeBoxAtom('d','a','t','a'))
-		{
-			ParseBoxDataList(FString(TEXT("iTunes."))+Name, dr.GetCurrentDataPointer()-8, BoxSize);
-		}
-	}
-#undef RETURN_IF_READERROR
-}
-
-bool UtilsMP4::FMetadataParser::IsDifferentFrom(const UtilsMP4::FMetadataParser& Other)
-{
-	// Different number of items or map entries means the metadata can't be identical.
-	if (NumTotalItems != Other.NumTotalItems || Items.Num() != Other.Items.Num())
-	{
-		return true;
-	}
-
-	// Check if the maps keys or items are different.
-	for(auto& Key : Items)
-	{
-		const TArray<TSharedPtr<FItem, ESPMode::ThreadSafe>>* OtherKeyItems = Other.Items.Find(Key.Key);
-		if (!OtherKeyItems || Key.Value.Num() != OtherKeyItems->Num())
-		{
-			return true;
-		}
-		// Expensive item by item test :-(
-		for(auto& Item : Key.Value)
-		{
-			if (!OtherKeyItems->ContainsByPredicate([Item](const TSharedPtr<FItem, ESPMode::ThreadSafe>& CompItem)
-			{
-				return *Item == *CompItem;
-			}))
-			{
-				return true;
-			}
-		}
-	}
-	return false;
-}
-
-
-TSharedPtr<TMap<FString, TArray<TSharedPtr<IMediaStreamMetadata::IItem, ESPMode::ThreadSafe>>>, ESPMode::ThreadSafe> UtilsMP4::FMetadataParser::GetMediaStreamMetadata() const
-{
-	TSharedPtr<TMap<FString, TArray<TSharedPtr<IMediaStreamMetadata::IItem, ESPMode::ThreadSafe>>>, ESPMode::ThreadSafe> NewMeta(new TMap<FString, TArray<TSharedPtr<IMediaStreamMetadata::IItem, ESPMode::ThreadSafe>>>);
-	for(auto& Item : Items)
-	{
-		const TArray<TSharedPtr<FItem, ESPMode::ThreadSafe>>& SrcList = Item.Value;
-		if (SrcList.Num())
-		{
-			TArray<TSharedPtr<IMediaStreamMetadata::IItem, ESPMode::ThreadSafe>> DstList;
-			for(int32 i=0; i<SrcList.Num(); ++i)
-			{
-				DstList.Emplace(SrcList[i]);
-			}
-			NewMeta->Emplace(Item.Key, MoveTemp(DstList));
-		}
-	}
-	return NewMeta;
-}
-
-
-FString UtilsMP4::FMetadataParser::GetAsJSON() const
-{
-	FString JSON(TEXT("{"));
-	bool bFirstItem = true;
-	for(auto& Item : Items)
-	{
-		bool bFirstItemValue = true;
-		const TArray<TSharedPtr<FItem, ESPMode::ThreadSafe>>& List=Item.Value;
-		FString ItemJSON;
-		for(int32 i=0; i<List.Num(); ++i)
-		{
-			const TSharedPtr<FItem, ESPMode::ThreadSafe>& it = List[i];
-			FString ItemValue = it->ToJSONValue();
-			if (ItemValue.Len())
-			{
-				if (!bFirstItemValue)
-				{
-					ItemJSON.AppendChar(TCHAR(','));
-				}
-				bFirstItemValue = false;
-				ItemJSON.Append(FString::Printf(TEXT("\"%s\":"), *(it->Language)));
-				ItemJSON.Append(MoveTemp(ItemValue));
-			}
-		}
-		if (ItemJSON.Len())
-		{
-			if (!bFirstItem)
-			{
-				JSON.AppendChar(TCHAR(','));
-			}
-			bFirstItem = false;
-
-			JSON.Append(FString::Printf(TEXT("\"%s\":{"), *Item.Key));
-			JSON.Append(MoveTemp(ItemJSON));
-			JSON.Append(TEXT("}"));
-		}
-	}
-	JSON.Append(TEXT("}"));
-	return JSON;
-}
-
-FString UtilsMP4::FMetadataParser::FItem::ToJSONValue() const
-{
-	switch(Type)
-	{
-		// String
-		case 1:
-		{
-			return FString::Printf(TEXT("\"%s\""), *Value.GetValue<FString>().ReplaceCharWithEscapedChar(&CharsToEscapeInJSON));
-		}
-		// Double
-		case 24:
-		{
-			return FString::Printf(TEXT("%f"), Value.GetValue<double>());
-		}
-		// Signed integer
-		case 74:
-		{
-			int64 v = Value.GetValue<int64>();
-			if (v >= -(1LL<<53) && v <= (1LL<<53)-1)
-			{
-				return FString::Printf(TEXT("%lld"), (long long int)v);
-			}
-		}
-		// Unsigned integer
-		case 78:
-		{
-			uint64 v = Value.GetValue<uint64>();
-			if (v <= (1LL<<53)-1)
-			{
-				return FString::Printf(TEXT("%lld"), (long long int)v);
-			}
-		}
-		// JPEG image
-		case 13:
-		{
-			FString js(TEXT("{\"mimetype\":\"image/jpeg\",\"data\":\""));
-			js.Append(FBase64::Encode(Value.GetValue<TArray<uint8>>()));
-			js.Append(TEXT("\"}"));
-			return js;
-		}
-		// PNG image
-		case 14:
-		{
-			FString js(TEXT("{\"mimetype\":\"image/png\",\"data\":\""));
-			js.Append(FBase64::Encode(Value.GetValue<TArray<uint8>>()));
-			js.Append(TEXT("\"}"));
-			return js;
-		}
-		default:
-		{
-			break;
-		}
-	}
-	return FString();
-}
-TArray<TCHAR> UtilsMP4::FMetadataParser::FItem::CharsToEscapeInJSON = { TCHAR('\n'), TCHAR('\r'), TCHAR('\t'), TCHAR('\\'), TCHAR('\"') };
-
-
-
-
 UtilsMP4::FMP4RootBoxLocator::~FMP4RootBoxLocator()
 {
 }
@@ -755,7 +215,7 @@ TSharedPtrTS<FWaitableBuffer> UtilsMP4::FMP4ChunkLoader::LoadChunk(const int64 I
 	volatile bool bAbort = false;
 	FileSize = -1;
 	bHasErrored = false;
-	
+
 	TSharedPtrTS<IElectraHttpManager::FProgressListener> ProgressListener = MakeSharedTS<IElectraHttpManager::FProgressListener>();
 	ProgressListener->ProgressDelegate = IElectraHttpManager::FProgressListener::FProgressDelegate::CreateLambda([&](const IElectraHttpManager::FRequest* InRequest)->int32
 	{
@@ -818,5 +278,3 @@ TSharedPtrTS<FWaitableBuffer> UtilsMP4::FMP4ChunkLoader::LoadChunk(const int64 I
 }
 
 }
-
-
diff --git a/Engine/Plugins/Media/ElectraPlayer/Source/ElectraPlayerRuntime/Private/Runtime/Utilities/UtilsMP4.h b/Engine/Plugins/Media/ElectraPlayer/Source/ElectraPlayerRuntime/Private/Runtime/Utilities/UtilsMP4.h
--- a/Engine/Plugins/Media/ElectraPlayer/Source/ElectraPlayerRuntime/Private/Runtime/Utilities/UtilsMP4.h
+++ b/Engine/Plugins/Media/ElectraPlayer/Source/ElectraPlayerRuntime/Private/Runtime/Utilities/UtilsMP4.h
@@ -3,9 +3,7 @@
 #pragma once
 
 #include <CoreMinimal.h>
-#include <Misc/Variant.h>
 #include "HTTP/HTTPManager.h"
-#include "MediaStreamMetadata.h"
 
 namespace Electra
 {
@@ -13,73 +11,6 @@ namespace Electra
 class UtilsMP4
 {
 public:
-	/**
-	 * This class parses metadata embedded in an mp4 / ISO14496-12 file.
-	 * Presently only the structure as used and defined by Apple iTunes is supported.
-	 */
-	class FMetadataParser
-	{
-	public:
-		FMetadataParser();
-		enum class EResult
-		{
-			Success,
-			NotSupported,
-			MissingBox
-		};
-		struct FBoxInfo
-		{
-			FBoxInfo(uint32 InType, const void* InData, uint32 InSize) : Type(InType), Data(InData), Size(InSize)
-			{}
-			uint32 Type = 0;
-			const void* Data = nullptr;
-			uint32 Size = 0;
-		};
-		EResult Parse(uint32 InHandler, uint32 InHandlerReserved0, const TArray<FBoxInfo>& InBoxes);
-		bool IsDifferentFrom(const FMetadataParser& Other);
-		FString GetAsJSON() const;
-		TSharedPtr<TMap<FString, TArray<TSharedPtr<IMediaStreamMetadata::IItem, ESPMode::ThreadSafe>>>, ESPMode::ThreadSafe> GetMediaStreamMetadata() const;
-
-		void AddItem(const FString& InType, const FString& InValue);
-		void AddItem(const FString& InType, const FString& InMimeType, const TArray<uint8>& InValue);
-	private:
-		FString PrintableBoxAtom(const uint32 InAtom);
-		void Parse(const FBoxInfo& InBox);
-		void ParseBoxDataList(const FString& AsCategory, const uint8* InBoxData, uint32 InBoxSize);
-		void ParseBoxDataiTunes(const uint8* InBoxData, uint32 InBoxSize);
-
-		class FItem : public IMediaStreamMetadata::IItem
-		{
-		public:
-			~FItem() = default;
-			const FString& GetLanguageCode() const override
-			{ return Language; }
-			const FString& GetMimeType() const override
-			{ return MimeType; }
-			const FVariant& GetValue() const override
-			{ return Value; }
-
-			FString Language;				// ISO 639-2; if not set (all zero) the default entry for all languages
-			FString MimeType;
-			int32 Type = 0;					// Well-known data type (see Quicktime reference)
-			FVariant Value;
-			FString ToJSONValue() const;
-			static TArray<TCHAR> CharsToEscapeInJSON;
-			bool operator != (const FItem& Other) const
-			{
-				return Type != Other.Type || Language != Other.Language || Value != Other.Value;
-			}
-			bool operator == (const FItem& Other) const
-			{
-				return !(*this != Other);
-			}
-		};
-
-		TMap<uint32, FString> WellKnownItems;
-		TMap<FString, TArray<TSharedPtr<FItem, ESPMode::ThreadSafe>>> Items;
-		uint32 NumTotalItems = 0;
-	};
-
 
 	class FMP4RootBoxLocator
 	{
diff --git a/Engine/Plugins/Media/ElectraUtil/Source/ElectraBase/Public/MediaStreamMetadata.h b/Engine/Plugins/Media/ElectraUtil/Source/ElectraBase/Public/MediaStreamMetadata.h
deleted file mode 100644
diff --git a/Engine/Plugins/Media/MP4Utilities/Source/MP4Utilities/MP4Utilities.Build.cs b/Engine/Plugins/Media/MP4Utilities/Source/MP4Utilities/MP4Utilities.Build.cs
--- a/Engine/Plugins/Media/MP4Utilities/Source/MP4Utilities/MP4Utilities.Build.cs
+++ b/Engine/Plugins/Media/MP4Utilities/Source/MP4Utilities/MP4Utilities.Build.cs
@@ -10,7 +10,8 @@ namespace UnrealBuildTool.Rules
 		public MP4Utilities(ReadOnlyTargetRules Target) : base(Target)
 		{
 			PrivateDependencyModuleNames.AddRange(new string[] {
-				"Core"
+				"Core",
+                "Media"
             });
 		}
 	}
diff --git a/Engine/Plugins/Media/MP4Utilities/Source/MP4Utilities/Private/MP4Metadata.cpp b/Engine/Plugins/Media/MP4Utilities/Source/MP4Utilities/Private/MP4Metadata.cpp
new file mode 100644
--- /dev/null
+++ b/Engine/Plugins/Media/MP4Utilities/Source/MP4Utilities/Private/MP4Metadata.cpp
@@ -0,0 +1,545 @@
+// Copyright Epic Games, Inc. All Rights Reserved.
+
+#include "MP4Metadata.h"
+#include "Misc/Base64.h"
+
+namespace MP4Utilities
+{
+
+namespace FMetadataTags
+{
+static const TCHAR * const Title = TEXT("Title");
+static const TCHAR * const Artist = TEXT("Artist");
+static const TCHAR * const Genre = TEXT("Genre");
+static const TCHAR * const Date = TEXT("Date");						// Content creation date
+static const TCHAR * const Description = TEXT("Description");
+static const TCHAR * const LongDescription = TEXT("LongDescription");
+static const TCHAR * const Album = TEXT("Album");
+static const TCHAR * const Encoder = TEXT("Encoder");
+
+static TMap<uint32, const TCHAR* const> WellKnownItems
+{
+	{ MakeBoxAtom(0xa9u,'n','a','m'), FMetadataTags::Title },
+	{ MakeBoxAtom(0xa9u,'A','R','T'), FMetadataTags::Artist },
+	{ MakeBoxAtom(0xa9u,'g','e','n'), FMetadataTags::Genre },
+	{ MakeBoxAtom(0xa9u,'d','a','y'), FMetadataTags::Date },
+	{ MakeBoxAtom('d','e','s','c'), FMetadataTags::Description },
+	{ MakeBoxAtom('l','d','e','s'), FMetadataTags::LongDescription },
+	{ MakeBoxAtom(0xa9u,'a','l','b'), FMetadataTags::Album },
+	{ MakeBoxAtom(0xa9u,'t','o','o'), FMetadataTags::Encoder }
+};
+}
+
+
+class FMP4MetadataParser::FItem : public IMediaMetadataItem
+{
+public:
+	~FItem() = default;
+	const FString& GetLanguageCode() const override
+	{ return Language; }
+	const FString& GetMimeType() const override
+	{ return MimeType; }
+	const FVariant& GetValue() const override
+	{ return Value; }
+
+	FString Language;				// ISO 639-2; if not set the default entry for all languages
+	FString MimeType;
+	FVariant Value;
+	int32 Type = 0;					// Well-known data type (see Quicktime reference)
+	FString ToJSONValue() const;
+	static TArray<TCHAR> CharsToEscapeInJSON;
+	bool operator != (const FItem& Other) const
+	{
+		return Type != Other.Type || Language != Other.Language || Value != Other.Value;
+	}
+	bool operator == (const FItem& Other) const
+	{
+		return !(*this != Other);
+	}
+};
+
+
+FMP4MetadataParser::FMP4MetadataParser()
+	: Items(MakeShared<TMap<FString, TArray<TUniquePtr<IMediaMetadataItem>>>, ESPMode::ThreadSafe>())
+	, NumTotalItems(0)
+{
+}
+
+void FMP4MetadataParser::AddItem(const FString& InType, const FString& InValue)
+{
+	// Add a string item.
+	FItem* Item = new FItem;
+	Item->Type = 1;
+	Item->Value = FVariant(InValue);
+	++NumTotalItems;
+	Items->FindOrAdd(InType).Emplace(TUniquePtr<IMediaMetadataItem>(Item));
+}
+
+void FMP4MetadataParser::AddItem(const FString& InType, const FString& InMimeType, const TArray<uint8>& InValue)
+{
+	if (InMimeType.Equals(TEXT("image/jpeg")))
+	{
+		FItem* Item = new FItem;
+		Item->MimeType = InMimeType;
+		Item->Type = 13;
+		Item->Value = FVariant(InValue);
+		++NumTotalItems;
+		Items->FindOrAdd(InType).Emplace(TUniquePtr<IMediaMetadataItem>(Item));
+	}
+	else if (InMimeType.Equals(TEXT("image/png")))
+	{
+		FItem* Item = new FItem;
+		Item->MimeType = InMimeType;
+		Item->Type = 14;
+		Item->Value = FVariant(InValue);
+		++NumTotalItems;
+		Items->FindOrAdd(InType).Emplace(TUniquePtr<IMediaMetadataItem>(Item));
+	}
+}
+
+
+FMP4MetadataParser::EResult FMP4MetadataParser::Parse(uint32 InHandler, uint32 InHandlerReserved0, const TArray<FMP4BoxInfo>& InBoxes)
+{
+	// We only support the Apple iTunes metadata at the moment.
+	if (InHandler != MakeBoxAtom('m','d','i','r'))
+	{
+		return FMP4MetadataParser::EResult::NotSupported;
+	}
+	/*
+		As per ISO/IEC 14496-12:2015 section 8.11.1.2 the 'meta' box must contain a handler ('hdlr'), followed by nothing but optional boxes.
+		Of these optional boxes the ones recognized and defined by the document are:
+		'pitm' - 8.11.4 Primary Item Box
+		'dinf' - 8.7.1 Data Information Box
+		'iloc' - 8.11.3 The Item Location Box
+		'ipro' - 8.11.5 Item Protection Box
+		'iinf' - 8.11.6 Item Information Box
+		'iref' - 8.11.12 Item Reference Box
+		'idat' - 8.11.11 Item Data Box
+
+		followed by any number of other boxes.
+
+		Presently the boxes mentioned above are not processed since they are apparently not used by any of the tools
+		that allow insertion of metadata. It appears that the way iTunes handles metadata has become the de facto standard
+		in which there is
+		    'meta'
+			   'hdlr'
+			   'ilst'
+			     ....		// See: https://developer.apple.com/library/archive/documentation/QuickTime/QTFF/Metadata/Metadata.html#//apple_ref/doc/uid/TP40000939-CH1-SW26
+				   'data'	// See: https://developer.apple.com/library/archive/documentation/QuickTime/QTFF/Metadata/Metadata.html#//apple_ref/doc/uid/TP40000939-CH1-SW27
+	*/
+
+	// By some definition the items are contained inside an 'ilst' box, so we look for that one.
+	bool bHaveILST = false;
+	for(int32 i=0; i<InBoxes.Num(); ++i)
+	{
+		if (InBoxes[i].Type == MakeBoxAtom('i','l','s','t'))
+		{
+			// We will actually parse all 'ilst' boxes if for some reason the metadata has
+			// been scattered across several.
+			bHaveILST = true;
+			Parse(InBoxes[i]);
+		}
+	}
+	return bHaveILST ? FMP4MetadataParser::EResult::Success : FMP4MetadataParser::EResult::MissingBox;
+}
+
+void FMP4MetadataParser::Parse(const FMP4BoxInfo& InBox)
+{
+	FMP4AtomReader dr(InBox.Data);
+	while(dr.GetNumBytesRemaining() > 8)
+	{
+		int32 BoxSize;
+		uint32 BoxType;
+		dr.Read(BoxSize);
+		dr.Read(BoxType);
+		if (FMetadataTags::WellKnownItems.Contains(BoxType))
+		{
+			ParseBoxDataList(FMetadataTags::WellKnownItems[BoxType], MakeConstArrayView(dr.GetCurrentDataPointer(), BoxSize-8));
+		}
+		else
+		{
+			switch(BoxType)
+			{
+				case MakeBoxAtom('-','-','-','-'):
+				{
+					ParseBoxDataiTunes(MakeConstArrayView(dr.GetCurrentDataPointer(), BoxSize-8));
+					break;
+				}
+				default:
+				{
+					ParseBoxDataList(FString(TEXT("qt."))+Printable4CC(BoxType), MakeConstArrayView(dr.GetCurrentDataPointer(), BoxSize-8));
+					break;
+				}
+			}
+		}
+		dr.ReadBytes(nullptr, BoxSize - 8);
+	}
+}
+
+
+void FMP4MetadataParser::ParseBoxDataList(const FString& AsCategory, TConstArrayView<uint8> InBoxData)
+{
+	if (InBoxData.Num() <= 8)
+	{
+		return;
+	}
+
+	FMP4AtomReader dr(InBoxData);
+
+#define RETURN_IF_READERROR(expr)	\
+	if (!(expr))					\
+	{								\
+		return;						\
+	}								\
+
+	for(int32 CurrentOffset=dr.GetCurrentOffset(), BoxSize=0; dr.GetNumBytesRemaining() > 8; dr.SetCurrentOffset(CurrentOffset), dr.ReadBytes(nullptr, BoxSize - 8))
+	{
+		uint32 BoxType;
+		RETURN_IF_READERROR(dr.Read(BoxSize));
+		RETURN_IF_READERROR(dr.Read(BoxType));
+		CurrentOffset = dr.GetCurrentOffset();
+		if (BoxType != MakeBoxAtom('d','a','t','a'))
+		{
+			continue;
+		}
+		uint32 TypeIndicatorAndWellKnownType = 0;
+		RETURN_IF_READERROR(dr.Read(TypeIndicatorAndWellKnownType));
+		if ((TypeIndicatorAndWellKnownType >> 24) != 0)
+		{
+			continue;
+		}
+		uint16 CountryIndicator=0;
+		RETURN_IF_READERROR(dr.Read(CountryIndicator));
+		uint16 LanguageIndicator=0;
+		RETURN_IF_READERROR(dr.Read(LanguageIndicator));
+		int32 NumRemainingDataBytes = BoxSize - 8 - (dr.GetCurrentOffset() - CurrentOffset);
+
+		TUniquePtr<FItem> Item(new FItem);
+
+		// We do not handle the country indicator at the moment.
+		// Likewise, the language indicator must be a directly specified language code since we do not handle
+		// language tables.
+		if (LanguageIndicator >= 1 && LanguageIndicator <= 255)
+		{
+			// This would be an index into the language table (which is a 'lang' box somewhere under the 'meta' box).
+			continue;
+		}
+		if (LanguageIndicator)
+		{
+			char Lang[3];
+			Lang[0] = (char)(0x60 + ((LanguageIndicator & 0x7c00) >> 10));
+			Lang[1] = (char)(0x60 + ((LanguageIndicator & 0x03e0) >> 5));
+			Lang[2] = (char)(0x60 + (LanguageIndicator & 0x001f));
+			Item->Language = FString::ConstructFromPtrSize(Lang, 3);
+		}
+
+		bool bSet = true;
+		int32 WKT = (int32)(TypeIndicatorAndWellKnownType & 0x00ffffffU);
+		switch(WKT)
+		{
+			// UTF-8 string
+			case 1:
+			{
+				FString String;
+				if ((bSet = dr.ReadStringUTF8(String, NumRemainingDataBytes)) == true)
+				{
+					Item->Value = String;
+					Item->Type = 1;
+				}
+				break;
+			}
+/*
+			// UTF-16 string
+			case 2:
+			{
+				FString String;
+				if ((bSet = dr.ReadStringUTF16(String, NumRemainingDataBytes)) == true)
+				{
+					Item->Value = String;
+					Item->Type = 1;
+				}
+				break;
+			}
+*/
+			case 13:	// JPEG image
+			case 14:	// PNG image
+			case 27:	// BMP image
+			{
+				TArray<uint8> Image(dr.GetCurrentDataPointer(), NumRemainingDataBytes);
+				Item->Value = MoveTemp(Image);
+				Item->Type = WKT;
+				if (WKT == 13)
+				{
+					Item->MimeType = TEXT("image/jpeg");
+				}
+				else if (WKT == 14)
+				{
+					Item->MimeType = TEXT("image/png");
+				}
+				else if (WKT == 27)
+				{
+					Item->MimeType = TEXT("image/bmp");
+				}
+				break;
+			}
+
+			// Float
+			case 23:
+			{
+				float Flt;
+				if ((bSet = dr.ReadAsNumber(Flt)) == true)
+				{
+					// Set as double!
+					Item->Value = (double)Flt;
+					Item->Type = 24;
+				}
+				break;
+			}
+			// Double
+			case 24:
+			{
+				double Dbl;
+				if ((bSet = dr.ReadAsNumber(Dbl)) == true)
+				{
+					Item->Value = Dbl;
+					Item->Type = 24;
+				}
+				break;
+			}
+			// Signed integers
+			case 65:	// int8
+			case 66:	// int16
+			case 67:	// int32
+			case 74:	// int64
+			case 21:	// 1-4 byte signed integer
+			{
+				int64 Number;
+				if ((bSet = dr.ReadAsNumber(Number, WKT==65?1:WKT==66?2:WKT==67?4:WKT==74?8:NumRemainingDataBytes)) == true)
+				{
+					Item->Value = Number;
+					Item->Type = 74;
+				}
+				break;
+			}
+			// Unsigned integers
+			case 75:	// uint8
+			case 76:	// uint16
+			case 77:	// uint32
+			case 78:	// uint64
+			case 22:	// 1-4 byte unsigned integer
+			{
+				uint64 Number;
+				if ((bSet = dr.ReadAsNumber(Number, WKT==75?1:WKT==76?2:WKT==77?4:WKT==78?8:NumRemainingDataBytes)) == true)
+				{
+					Item->Value = Number;
+					Item->Type = 78;
+				}
+				break;
+			}
+
+			// Not handled here.
+			default:
+			{
+				bSet = false;
+				break;
+			}
+		}
+
+		if (bSet)
+		{
+			TUniquePtr<IMediaMetadataItem> mi(Item.Release());
+			Items->FindOrAdd(AsCategory).Emplace(MoveTemp(mi));
+			++NumTotalItems;
+		}
+	}
+#undef RETURN_IF_READERROR
+}
+
+void FMP4MetadataParser::ParseBoxDataiTunes(TConstArrayView<uint8> InBoxData)
+{
+	if (InBoxData.Num() <= 8)
+	{
+		return;
+	}
+
+	FMP4AtomReader dr(InBoxData);
+
+#define RETURN_IF_READERROR(expr)	\
+	if (!(expr))					\
+	{								\
+		return;						\
+	}								\
+
+	FString Name;
+	for(int32 CurrentOffset=dr.GetCurrentOffset(), BoxSize=0; dr.GetNumBytesRemaining() > 8; dr.SetCurrentOffset(CurrentOffset), dr.ReadBytes(nullptr, BoxSize - 8))
+	{
+		uint32 BoxType;
+		RETURN_IF_READERROR(dr.Read(BoxSize));
+		RETURN_IF_READERROR(dr.Read(BoxType));
+		CurrentOffset = dr.GetCurrentOffset();
+		if (BoxType == MakeBoxAtom('m','e','a','n'))
+		{
+			// 4 reserved bytes follow
+			RETURN_IF_READERROR(dr.ReadBytes(nullptr, 4));
+			FString Meaning;
+			int32 NumRemainingDataBytes = BoxSize - 8 - (dr.GetCurrentOffset() - CurrentOffset);
+			RETURN_IF_READERROR(dr.ReadStringUTF8(Meaning, NumRemainingDataBytes));
+			if (!Meaning.Equals(TEXT("com.apple.iTunes")))
+			{
+				return;
+			}
+		}
+		else if (BoxType == MakeBoxAtom('n','a','m','e'))
+		{
+			// 4 reserved bytes follow
+			RETURN_IF_READERROR(dr.ReadBytes(nullptr, 4));
+			int32 NumRemainingDataBytes = BoxSize - 8 - (dr.GetCurrentOffset() - CurrentOffset);
+			RETURN_IF_READERROR(dr.ReadStringUTF8(Name, NumRemainingDataBytes));
+		}
+		else if (BoxType == MakeBoxAtom('d','a','t','a'))
+		{
+			ParseBoxDataList(FString(TEXT("iTunes."))+Name, MakeConstArrayView(dr.GetCurrentDataPointer()-8, BoxSize));
+		}
+	}
+#undef RETURN_IF_READERROR
+}
+
+bool FMP4MetadataParser::IsDifferentFrom(const FMP4MetadataParser& Other)
+{
+	// Different number of items or map entries means the metadata can't be identical.
+	if (NumTotalItems != Other.NumTotalItems || Items->Num() != Other.Items->Num())
+	{
+		return true;
+	}
+
+	// Check if the maps keys or items are different.
+	for(auto& ItemMapIterator : *Items)
+	{
+		const TArray<TUniquePtr<IMediaMetadataItem>>* OtherKeyItems = Other.Items->Find(ItemMapIterator.Key);
+		if (!OtherKeyItems || ItemMapIterator.Value.Num() != OtherKeyItems->Num())
+		{
+			return true;
+		}
+		// Expensive item by item test
+		for(auto& ItemIterator : ItemMapIterator.Value)
+		{
+			if (!OtherKeyItems->ContainsByPredicate([&ItemIterator](const TUniquePtr<IMediaMetadataItem>& CompItem)
+			{
+				return *static_cast<FItem*>(ItemIterator.Get()) == *static_cast<FItem*>(CompItem.Get());
+			}))
+			{
+				return true;
+			}
+		}
+	}
+	return false;
+}
+
+
+TSharedPtr<TMap<FString, TArray<TUniquePtr<IMediaMetadataItem>>>, ESPMode::ThreadSafe> FMP4MetadataParser::GetMediaStreamMetadata() const
+{
+	return Items;
+}
+
+
+FString FMP4MetadataParser::GetAsJSON() const
+{
+	FString JSON(TEXT("{"));
+	bool bFirstItem = true;
+	for(auto& Item : *Items)
+	{
+		bool bFirstItemValue = true;
+		const TArray<TUniquePtr<IMediaMetadataItem>>& List=Item.Value;
+		FString ItemJSON;
+		for(int32 i=0; i<List.Num(); ++i)
+		{
+			const FItem* it = static_cast<const FItem*>(List[i].Get());
+			FString ItemValue = it->ToJSONValue();
+			if (ItemValue.Len())
+			{
+				if (!bFirstItemValue)
+				{
+					ItemJSON.AppendChar(TCHAR(','));
+				}
+				bFirstItemValue = false;
+				ItemJSON.Append(FString::Printf(TEXT("\"%s\":"), *(it->Language)));
+				ItemJSON.Append(MoveTemp(ItemValue));
+			}
+		}
+		if (ItemJSON.Len())
+		{
+			if (!bFirstItem)
+			{
+				JSON.AppendChar(TCHAR(','));
+			}
+			bFirstItem = false;
+
+			JSON.Append(FString::Printf(TEXT("\"%s\":{"), *Item.Key));
+			JSON.Append(MoveTemp(ItemJSON));
+			JSON.Append(TEXT("}"));
+		}
+	}
+	JSON.Append(TEXT("}"));
+	return JSON;
+}
+
+FString FMP4MetadataParser::FItem::ToJSONValue() const
+{
+	switch(Type)
+	{
+		// String
+		case 1:
+		{
+			return FString::Printf(TEXT("\"%s\""), *Value.GetValue<FString>().ReplaceCharWithEscapedChar(&CharsToEscapeInJSON));
+		}
+		// Double
+		case 24:
+		{
+			return FString::Printf(TEXT("%f"), Value.GetValue<double>());
+		}
+		// Signed integer
+		case 74:
+		{
+			int64 v = Value.GetValue<int64>();
+			if (v >= -(1LL<<53) && v <= (1LL<<53)-1)
+			{
+				return FString::Printf(TEXT("%lld"), (long long int)v);
+			}
+		}
+		// Unsigned integer
+		case 78:
+		{
+			uint64 v = Value.GetValue<uint64>();
+			if (v <= (1LL<<53)-1)
+			{
+				return FString::Printf(TEXT("%lld"), (long long int)v);
+			}
+		}
+		// JPEG image
+		case 13:
+		{
+			FString js(TEXT("{\"mimetype\":\"image/jpeg\",\"data\":\""));
+			js.Append(FBase64::Encode(Value.GetValue<TArray<uint8>>()));
+			js.Append(TEXT("\"}"));
+			return js;
+		}
+		// PNG image
+		case 14:
+		{
+			FString js(TEXT("{\"mimetype\":\"image/png\",\"data\":\""));
+			js.Append(FBase64::Encode(Value.GetValue<TArray<uint8>>()));
+			js.Append(TEXT("\"}"));
+			return js;
+		}
+		default:
+		{
+			break;
+		}
+	}
+	return FString();
+}
+TArray<TCHAR> FMP4MetadataParser::FItem::CharsToEscapeInJSON = { TCHAR('\n'), TCHAR('\r'), TCHAR('\t'), TCHAR('\\'), TCHAR('\"') };
+
+
+
+} // namespace MP4Utilities
diff --git a/Engine/Plugins/Media/MP4Utilities/Source/MP4Utilities/Public/MP4Metadata.h b/Engine/Plugins/Media/MP4Utilities/Source/MP4Utilities/Public/MP4Metadata.h
new file mode 100644
--- /dev/null
+++ b/Engine/Plugins/Media/MP4Utilities/Source/MP4Utilities/Public/MP4Metadata.h
@@ -0,0 +1,55 @@
+// Copyright Epic Games, Inc. All Rights Reserved.
+
+#pragma once
+
+#include "CoreTypes.h"
+#include "Containers/Array.h"
+#include "Containers/Map.h"
+#include "Containers/UnrealString.h"
+#include "Templates/SharedPointer.h"
+#include "Templates/UniquePtr.h"
+
+#include "IMediaMetadataItem.h"
+
+#include "MP4Utilities.h"
+
+#define UE_API MP4UTILITIES_API
+
+namespace MP4Utilities
+{
+	/**
+	 * This class parses metadata embedded in an ISO/IEC 14496-12 file.
+	 * Presently only the structure as used and defined by Apple iTunes is supported.
+	 */
+	class FMP4MetadataParser
+	{
+	public:
+		UE_API FMP4MetadataParser();
+		enum class EResult
+		{
+			Success,
+			NotSupported,
+			MissingBox
+		};
+		UE_API EResult Parse(uint32 InHandler, uint32 InHandlerReserved0, const TArray<FMP4BoxInfo>& InBoxes);
+		UE_API bool IsDifferentFrom(const FMP4MetadataParser& Other);
+		UE_API FString GetAsJSON() const;
+		UE_API TSharedPtr<TMap<FString, TArray<TUniquePtr<IMediaMetadataItem>>>, ESPMode::ThreadSafe> GetMediaStreamMetadata() const;
+
+		UE_API void AddItem(const FString& InType, const FString& InValue);
+		UE_API void AddItem(const FString& InType, const FString& InMimeType, const TArray<uint8>& InValue);
+	private:
+		class FItem;
+
+		FString PrintableBoxAtom(const uint32 InAtom);
+		void Parse(const FMP4BoxInfo& InBox);
+		void ParseBoxDataList(const FString& AsCategory, TConstArrayView<uint8> InBoxData);
+		void ParseBoxDataiTunes(TConstArrayView<uint8> InBoxData);
+
+		TSharedPtr<TMap<FString, TArray<TUniquePtr<IMediaMetadataItem>>>, ESPMode::ThreadSafe> Items;
+		uint32 NumTotalItems = 0;
+	};
+
+} // namespace MP4Utilities
+
+#undef UE_API
