Electra: moved hardwired properties from player into manifest and reader for better flexibility and new formats

#rnx

#Synced-CL 50154202

diff --git a/Engine/Plugins/Media/ElectraPlayer/Source/ElectraPlayerRuntime/Private/Runtime/Player/AdaptiveStreamingPlayer.cpp b/Engine/Plugins/Media/ElectraPlayer/Source/ElectraPlayerRuntime/Private/Runtime/Player/AdaptiveStreamingPlayer.cpp
--- a/Engine/Plugins/Media/ElectraPlayer/Source/ElectraPlayerRuntime/Private/Runtime/Player/AdaptiveStreamingPlayer.cpp
+++ b/Engine/Plugins/Media/ElectraPlayer/Source/ElectraPlayerRuntime/Private/Runtime/Player/AdaptiveStreamingPlayer.cpp
@@ -60,7 +60,6 @@ FAdaptiveStreamingPlayer::FAdaptiveStreamingPlayer(const IAdaptiveStreamingPlaye
 
 	ExternalPlayerGUID = InCreateParameters.ExternalPlayerGUID;
 	UseSharedWorkerThreads = InCreateParameters.WorkerThreads;
-	ManifestType = EMediaFormatType::Unknown;
 	CurrentState = EPlayerState::eState_Idle;
 	PipelineState = EPipelineState::ePipeline_Stopped;
 	DecoderState = EDecoderState::eDecoder_Paused;
@@ -1910,7 +1909,7 @@ void FAdaptiveStreamingPlayer::HandleSeeking()
 void FAdaptiveStreamingPlayer::HandleMetadataChanges()
 {
 	// The timeline can change dynamically. Refresh it on occasion.
-	if (Manifest.IsValid() && (ManifestType == EMediaFormatType::DASH || ManifestType == EMediaFormatType::HLS))
+	if (Manifest.IsValid() && Manifest->ReloadsPlaylistPeriodically())
 	{
 		PlaybackState.SetSeekableRange(Manifest->GetSeekableTimeRange());
 		PlaybackState.SetTimelineRange(Manifest->GetTotalTimeRange());
@@ -2246,8 +2245,7 @@ bool FAdaptiveStreamingPlayer::HaveEnoughBufferedDataToStartPlayback()
 
 	// When we are dealing with a single multiplexed stream and one buffer is blocked then essentially all buffers
 	// must be considered blocked since demuxing cannot continue.
-	// FIXME: do this more elegant somehow
-	if (ManifestType == EMediaFormatType::ISOBMFF || ManifestType == EMediaFormatType::MKV)
+	if (StreamReaderHandler && StreamReaderHandler->LoadsMultiplexedTracksSequentially())
 	{
 		if (VideoBufferStats.StreamBuffer.bLastPushWasBlocked ||
 			AudioBufferStats.StreamBuffer.bLastPushWasBlocked ||
@@ -2795,7 +2793,7 @@ void FAdaptiveStreamingPlayer::InternalHandlePendingStartRequest(const FTimeValu
 						if (BufferSourceInfoVid.IsValid())
 						{
 							SelectedStreamAttributesVid.UpdateWith(BufferSourceInfoVid->Kind, BufferSourceInfoVid->LanguageTag, BufferSourceInfoVid->Codec, BufferSourceInfoVid->HardIndex);
-							if (ManifestType != EMediaFormatType::ISOBMFF && ManifestType != EMediaFormatType::MKV)
+							if (StreamReaderHandler && !StreamReaderHandler->LoadsMultiplexedTracksSequentially())
 							{
 								StreamSelectionAttributesVid.UpdateIfOverrideSet(BufferSourceInfoVid->Kind, BufferSourceInfoVid->LanguageTag, BufferSourceInfoVid->Codec);
 							}
@@ -2803,7 +2801,7 @@ void FAdaptiveStreamingPlayer::InternalHandlePendingStartRequest(const FTimeValu
 						if (BufferSourceInfoAud.IsValid())
 						{
 							SelectedStreamAttributesAud.UpdateWith(BufferSourceInfoAud->Kind, BufferSourceInfoAud->LanguageTag, BufferSourceInfoAud->Codec, BufferSourceInfoAud->HardIndex);
-							if (ManifestType != EMediaFormatType::ISOBMFF && ManifestType != EMediaFormatType::MKV)
+							if (StreamReaderHandler && !StreamReaderHandler->LoadsMultiplexedTracksSequentially())
 							{
 								StreamSelectionAttributesAud.UpdateIfOverrideSet(BufferSourceInfoAud->Kind, BufferSourceInfoAud->LanguageTag, BufferSourceInfoAud->Codec);
 							}
@@ -2811,7 +2809,7 @@ void FAdaptiveStreamingPlayer::InternalHandlePendingStartRequest(const FTimeValu
 						if (BufferSourceInfoTxt.IsValid())
 						{
 							SelectedStreamAttributesTxt.UpdateWith(BufferSourceInfoTxt->Kind, BufferSourceInfoTxt->LanguageTag, BufferSourceInfoTxt->Codec, BufferSourceInfoTxt->HardIndex);
-							if (ManifestType != EMediaFormatType::ISOBMFF && ManifestType != EMediaFormatType::MKV)
+							if (StreamReaderHandler && !StreamReaderHandler->LoadsMultiplexedTracksSequentially())
 							{
 								StreamSelectionAttributesTxt.UpdateIfOverrideSet(BufferSourceInfoTxt->Kind, BufferSourceInfoTxt->LanguageTag, BufferSourceInfoTxt->Codec);
 							}
@@ -2875,7 +2873,7 @@ void FAdaptiveStreamingPlayer::InternalHandlePendingStartRequest(const FTimeValu
 									Manifest->ClearDefaultStartTime();
 
 									// For DASH Live check if it would be better to start on the next segment.
-									if (ManifestType == EMediaFormatType::DASH &&
+									if (Manifest->GetMediaFormatType() == EMediaFormatType::DASH &&
 										Manifest->GetPresentationType() != IManifest::EType::OnDemand)
 									{
 										struct FSegTime { FTimeValue Start, Into, Duration; };
@@ -3000,9 +2998,9 @@ void FAdaptiveStreamingPlayer::InternalHandlePendingFirstSegmentRequest(const FT
 			NewReceiveBuffers->AudBuffer = MakeSharedTS<FMultiTrackAccessUnitBuffer>(EStreamType::Audio);
 			NewReceiveBuffers->TxtBuffer = MakeSharedTS<FMultiTrackAccessUnitBuffer>(EStreamType::Subtitle);
 
-			if (ManifestType == EMediaFormatType::ISOBMFF || ManifestType == EMediaFormatType::MKV)
+			// Some formats have multiple tracks in a single stream that provide data all at the same time.
+			if (StreamReaderHandler && StreamReaderHandler->LoadsMultiplexedTracksSequentially())
 			{
-				// mp4 and mkv have multiple tracks in a single stream that provide data all at the same time.
 				NewReceiveBuffers->VidBuffer->SetParallelTrackMode();
 				NewReceiveBuffers->AudBuffer->SetParallelTrackMode();
 				NewReceiveBuffers->TxtBuffer->SetParallelTrackMode();
@@ -3424,7 +3422,7 @@ void FAdaptiveStreamingPlayer::HandlePendingMediaSegmentRequests()
 								case EStreamType::Video:
 								{
 									SelectedStreamAttributesVid.UpdateWith(BufferSourceInfo->Kind, BufferSourceInfo->LanguageTag, BufferSourceInfo->Codec, BufferSourceInfo->HardIndex);
-									if (ManifestType == EMediaFormatType::DASH)
+									if (Manifest->OrderOfTracksCanChangeDynamically())
 									{
 										StreamSelectionAttributesVid.ClearOverrideIndex();
 									}
@@ -3433,7 +3431,7 @@ void FAdaptiveStreamingPlayer::HandlePendingMediaSegmentRequests()
 								case EStreamType::Audio:
 								{
 									SelectedStreamAttributesAud.UpdateWith(BufferSourceInfo->Kind, BufferSourceInfo->LanguageTag, BufferSourceInfo->Codec, BufferSourceInfo->HardIndex);
-									if (ManifestType == EMediaFormatType::DASH)
+									if (Manifest->OrderOfTracksCanChangeDynamically())
 									{
 										StreamSelectionAttributesAud.ClearOverrideIndex();
 									}
@@ -3442,7 +3440,7 @@ void FAdaptiveStreamingPlayer::HandlePendingMediaSegmentRequests()
 								case EStreamType::Subtitle:
 								{
 									SelectedStreamAttributesTxt.UpdateWith(BufferSourceInfo->Kind, BufferSourceInfo->LanguageTag, BufferSourceInfo->Codec, BufferSourceInfo->HardIndex);
-									if (ManifestType == EMediaFormatType::DASH)
+									if (Manifest->OrderOfTracksCanChangeDynamically())
 									{
 										StreamSelectionAttributesTxt.ClearOverrideIndex();
 									}
@@ -3707,85 +3705,88 @@ void FAdaptiveStreamingPlayer::InternalHandleSegmentTrackChanges(const FTimeValu
 	bool bChangeMade = false;
 	bool bForceStartover = false;
 
-	if (ManifestType == EMediaFormatType::ISOBMFF || ManifestType == EMediaFormatType::MKV)
+	if (Manifest.IsValid() && Manifest->SupportsMultipleTracksPerKind() && StreamReaderHandler)
 	{
-		auto HandleMultiplexedTrackSelection = [&bChangeMade, this](FStreamSelectionAttributes& OutSelectionAttributes, FInternalStreamSelectionAttributes& InOutSelectedAttributes, EStreamType InType,
-																TSharedPtrTS<FStreamSelectionAttributes>& InPendingSelection, TSharedPtrTS<IManifest::IPlayPeriod> InCurrentPeriod) -> void
+		if (StreamReaderHandler->LoadsMultiplexedTracksSequentially())
 		{
-			if (InPendingSelection.IsValid() && InCurrentPeriod.IsValid())
+			auto HandleMultiplexedTrackSelection = [&bChangeMade, this](FStreamSelectionAttributes& OutSelectionAttributes, FInternalStreamSelectionAttributes& InOutSelectedAttributes, EStreamType InType,
+																	TSharedPtrTS<FStreamSelectionAttributes>& InPendingSelection, TSharedPtrTS<IManifest::IPlayPeriod> InCurrentPeriod) -> void
 			{
-				IManifest::IPlayPeriod::ETrackChangeResult Result = InCurrentPeriod->ChangeTrackStreamPreference(InType, *InPendingSelection);
-				TSharedPtrTS<FBufferSourceInfo> BufferSourceInfo = InCurrentPeriod->GetSelectedStreamBufferSourceInfo(InType);
-				if (Result == IManifest::IPlayPeriod::ETrackChangeResult::Changed)
+				if (InPendingSelection.IsValid() && InCurrentPeriod.IsValid())
 				{
-					TSharedPtrTS<FMultiTrackAccessUnitBuffer> StreamRcvBuf = GetCurrentReceiveStreamBuffer(InType);
-					if (StreamRcvBuf.IsValid())
+					IManifest::IPlayPeriod::ETrackChangeResult Result = InCurrentPeriod->ChangeTrackStreamPreference(InType, *InPendingSelection);
+					TSharedPtrTS<FBufferSourceInfo> BufferSourceInfo = InCurrentPeriod->GetSelectedStreamBufferSourceInfo(InType);
+					if (Result == IManifest::IPlayPeriod::ETrackChangeResult::Changed)
 					{
-						StreamRcvBuf->SelectTrackWhenAvailable(CurrentPlaybackSequenceID[StreamTypeToArrayIndex(InType)], BufferSourceInfo);
+						TSharedPtrTS<FMultiTrackAccessUnitBuffer> StreamRcvBuf = GetCurrentReceiveStreamBuffer(InType);
+						if (StreamRcvBuf.IsValid())
+						{
+							StreamRcvBuf->SelectTrackWhenAvailable(CurrentPlaybackSequenceID[StreamTypeToArrayIndex(InType)], BufferSourceInfo);
+						}
+						OutSelectionAttributes = *InPendingSelection;
+						bChangeMade = true;
 					}
-					OutSelectionAttributes = *InPendingSelection;
-					bChangeMade = true;
+					InOutSelectedAttributes.UpdateWith(BufferSourceInfo->Kind, BufferSourceInfo->LanguageTag, BufferSourceInfo->Codec, BufferSourceInfo->HardIndex);
+					InOutSelectedAttributes.Select();
+					InPendingSelection.Reset();
 				}
-				InOutSelectedAttributes.UpdateWith(BufferSourceInfo->Kind, BufferSourceInfo->LanguageTag, BufferSourceInfo->Codec, BufferSourceInfo->HardIndex);
-				InOutSelectedAttributes.Select();
-				InPendingSelection.Reset();
-			}
-		};
+			};
 
-		HandleMultiplexedTrackSelection(StreamSelectionAttributesVid, SelectedStreamAttributesVid, EStreamType::Video, PendingTrackSelectionVid, CurrentPlayPeriodVideo);
-		HandleMultiplexedTrackSelection(StreamSelectionAttributesAud, SelectedStreamAttributesAud, EStreamType::Audio, PendingTrackSelectionAud, CurrentPlayPeriodAudio);
-		HandleMultiplexedTrackSelection(StreamSelectionAttributesTxt, SelectedStreamAttributesTxt, EStreamType::Subtitle, PendingTrackSelectionTxt, CurrentPlayPeriodText);
-	}
-	else if (ManifestType == EMediaFormatType::DASH || ManifestType == EMediaFormatType::HLS)
-	{
-		if (PendingStartRequest.IsValid() || PendingFirstSegmentRequest.IsValid())
-		{
-			return;
+			HandleMultiplexedTrackSelection(StreamSelectionAttributesVid, SelectedStreamAttributesVid, EStreamType::Video, PendingTrackSelectionVid, CurrentPlayPeriodVideo);
+			HandleMultiplexedTrackSelection(StreamSelectionAttributesAud, SelectedStreamAttributesAud, EStreamType::Audio, PendingTrackSelectionAud, CurrentPlayPeriodAudio);
+			HandleMultiplexedTrackSelection(StreamSelectionAttributesTxt, SelectedStreamAttributesTxt, EStreamType::Subtitle, PendingTrackSelectionTxt, CurrentPlayPeriodText);
 		}
-
-		auto HandleElementaryStreamTrackSelection = [&bChangeMade, &bForceStartover, this](FStreamSelectionAttributes& OutSelectionAttributes, FInternalStreamSelectionAttributes& InOutSelectedAttributes, EStreamType InType,
-															 TSharedPtrTS<FStreamSelectionAttributes>& InPendingSelection, TSharedPtrTS<IManifest::IPlayPeriod> InCurrentPeriod) -> void
+		else
 		{
-			if (InPendingSelection.IsValid() && InCurrentPeriod.IsValid())
+			if (PendingStartRequest.IsValid() || PendingFirstSegmentRequest.IsValid())
 			{
-				if (InOutSelectedAttributes.IsDeselected() || !InOutSelectedAttributes.IsCompatibleWith(*InPendingSelection))
-				{
-					IManifest::IPlayPeriod::ETrackChangeResult TrackChangeResult = InCurrentPeriod->ChangeTrackStreamPreference(InType, *InPendingSelection);
-					if (TrackChangeResult == IManifest::IPlayPeriod::ETrackChangeResult::NewPeriodNeeded)
-					{
-						// Make sure any finishing current requests will be ignored by increasing the sequence ID.
-						++CurrentPlaybackSequenceID[StreamTypeToArrayIndex(InType)];
-						// Cancel any ongoing segment download now.
-						CancelPendingMediaSegmentRequests(InType);
-
-						OutSelectionAttributes = *InPendingSelection;
-						InOutSelectedAttributes.Select();
-
-						FPendingSegmentRequest NextReq;
-						NextReq.bStartOver = true;
-						NextReq.StartoverPosition.Time = PlaybackState.GetPlayPosition();
-						SetPlaystartOptions(NextReq.StartoverPosition.Options);
-						NextReq.StreamType = InType;
-						NextPendingSegmentRequests->Enqueue(MoveTemp(NextReq));
+				return;
+			}
 
-						bChangeMade = true;
-					}
-					else if (TrackChangeResult == IManifest::IPlayPeriod::ETrackChangeResult::StartOver)
+			auto HandleElementaryStreamTrackSelection = [&bChangeMade, &bForceStartover, this](FStreamSelectionAttributes& OutSelectionAttributes, FInternalStreamSelectionAttributes& InOutSelectedAttributes, EStreamType InType,
+																 TSharedPtrTS<FStreamSelectionAttributes>& InPendingSelection, TSharedPtrTS<IManifest::IPlayPeriod> InCurrentPeriod) -> void
+			{
+				if (InPendingSelection.IsValid() && InCurrentPeriod.IsValid())
+				{
+					if (InOutSelectedAttributes.IsDeselected() || !InOutSelectedAttributes.IsCompatibleWith(*InPendingSelection))
 					{
-						OutSelectionAttributes = *InPendingSelection;
-						InOutSelectedAttributes.Select();
-						bChangeMade = true;
-						bForceStartover = true;
+						IManifest::IPlayPeriod::ETrackChangeResult TrackChangeResult = InCurrentPeriod->ChangeTrackStreamPreference(InType, *InPendingSelection);
+						if (TrackChangeResult == IManifest::IPlayPeriod::ETrackChangeResult::NewPeriodNeeded)
+						{
+							// Make sure any finishing current requests will be ignored by increasing the sequence ID.
+							++CurrentPlaybackSequenceID[StreamTypeToArrayIndex(InType)];
+							// Cancel any ongoing segment download now.
+							CancelPendingMediaSegmentRequests(InType);
+
+							OutSelectionAttributes = *InPendingSelection;
+							InOutSelectedAttributes.Select();
+
+							FPendingSegmentRequest NextReq;
+							NextReq.bStartOver = true;
+							NextReq.StartoverPosition.Time = PlaybackState.GetPlayPosition();
+							SetPlaystartOptions(NextReq.StartoverPosition.Options);
+							NextReq.StreamType = InType;
+							NextPendingSegmentRequests->Enqueue(MoveTemp(NextReq));
+
+							bChangeMade = true;
+						}
+						else if (TrackChangeResult == IManifest::IPlayPeriod::ETrackChangeResult::StartOver)
+						{
+							OutSelectionAttributes = *InPendingSelection;
+							InOutSelectedAttributes.Select();
+							bChangeMade = true;
+							bForceStartover = true;
+						}
 					}
+					InOutSelectedAttributes.Select();
+					InPendingSelection.Reset();
 				}
-				InOutSelectedAttributes.Select();
-				InPendingSelection.Reset();
-			}
-		};
+			};
 
-		HandleElementaryStreamTrackSelection(StreamSelectionAttributesVid, SelectedStreamAttributesVid, EStreamType::Video, PendingTrackSelectionVid, CurrentPlayPeriodVideo);
-		HandleElementaryStreamTrackSelection(StreamSelectionAttributesAud, SelectedStreamAttributesAud, EStreamType::Audio, PendingTrackSelectionAud, CurrentPlayPeriodAudio);
-		HandleElementaryStreamTrackSelection(StreamSelectionAttributesTxt, SelectedStreamAttributesTxt, EStreamType::Subtitle, PendingTrackSelectionTxt, CurrentPlayPeriodText);
+			HandleElementaryStreamTrackSelection(StreamSelectionAttributesVid, SelectedStreamAttributesVid, EStreamType::Video, PendingTrackSelectionVid, CurrentPlayPeriodVideo);
+			HandleElementaryStreamTrackSelection(StreamSelectionAttributesAud, SelectedStreamAttributesAud, EStreamType::Audio, PendingTrackSelectionAud, CurrentPlayPeriodAudio);
+			HandleElementaryStreamTrackSelection(StreamSelectionAttributesTxt, SelectedStreamAttributesTxt, EStreamType::Subtitle, PendingTrackSelectionTxt, CurrentPlayPeriodText);
+		}
 	}
 	else
 	{
diff --git a/Engine/Plugins/Media/ElectraPlayer/Source/ElectraPlayerRuntime/Private/Runtime/Player/AdaptiveStreamingPlayerInternal.h b/Engine/Plugins/Media/ElectraPlayer/Source/ElectraPlayerRuntime/Private/Runtime/Player/AdaptiveStreamingPlayerInternal.h
--- a/Engine/Plugins/Media/ElectraPlayer/Source/ElectraPlayerRuntime/Private/Runtime/Player/AdaptiveStreamingPlayerInternal.h
+++ b/Engine/Plugins/Media/ElectraPlayer/Source/ElectraPlayerRuntime/Private/Runtime/Player/AdaptiveStreamingPlayerInternal.h
@@ -2320,7 +2320,6 @@ private:
 	TMediaQueueDynamic<TSharedPtrTS<FErrorDetail>>						ErrorQueue;
 
 	FString																ManifestURL;
-	EMediaFormatType													ManifestType;
 	TSharedPtrTS<IManifest>												Manifest;
 	TSharedPtrTS<IPlaylistReader>										ManifestReader;
 	TSharedPtrTS<FHTTPResourceRequest>									ManifestMimeTypeRequest;
diff --git a/Engine/Plugins/Media/ElectraPlayer/Source/ElectraPlayerRuntime/Private/Runtime/Player/AdaptiveStreamingPlayerManifest.cpp b/Engine/Plugins/Media/ElectraPlayer/Source/ElectraPlayerRuntime/Private/Runtime/Player/AdaptiveStreamingPlayerManifest.cpp
--- a/Engine/Plugins/Media/ElectraPlayer/Source/ElectraPlayerRuntime/Private/Runtime/Player/AdaptiveStreamingPlayerManifest.cpp
+++ b/Engine/Plugins/Media/ElectraPlayer/Source/ElectraPlayerRuntime/Private/Runtime/Player/AdaptiveStreamingPlayerManifest.cpp
@@ -197,7 +197,6 @@ void FAdaptiveStreamingPlayer::InternalLoadManifest(const FString& InURL, const
 	// Remember the original request URL since we may lose the fragment part in requests.
 	ManifestURL = URL;
 	ManifestMimeTypeRequest.Reset();
-	ManifestType = EMediaFormatType::Unknown;
 	if (CurrentState == EPlayerState::eState_Idle)
 	{
 		FString mimeType = MimeType;
@@ -228,29 +227,24 @@ void FAdaptiveStreamingPlayer::InternalLoadManifest(const FString& InURL, const
 			if (mimeType == Playlist::MIMETypeHLS)
 			{
 				ManifestReader = IPlaylistHandlerHLS::Create(this);
-				ManifestType = EMediaFormatType::HLS;
 				CMCDHandler->SetStreamingFormat(FCMCDHandler::EStreamingFormat::HLS);
 			}
 			else if (mimeType == Playlist::MIMETypeMP4 || mimeType == Playlist::MIMETypeMP4A || mimeType == Playlist::MIMETypeQuickTime)
 			{
 				ManifestReader = IPlaylistReaderMP4::Create(this);
-				ManifestType = EMediaFormatType::ISOBMFF;
 			}
 			else if (mimeType == Playlist::MIMETypeDASH)
 			{
 				ManifestReader = IPlaylistReaderDASH::Create(this);
-				ManifestType = EMediaFormatType::DASH;
 				CMCDHandler->SetStreamingFormat(FCMCDHandler::EStreamingFormat::DASH);
 			}
 			else if (mimeType == Playlist::MIMETypeMKV || mimeType == Playlist::MIMETypeMKA)
 			{
 				ManifestReader = IPlaylistReaderMKV::Create(this);
-				ManifestType = EMediaFormatType::MKV;
 			}
 			else if (mimeType == Playlist::MIMETypeMPEGAudio || mimeType == Playlist::MIMETypeMPEGAudio4)
 			{
 				ManifestReader = IPlaylistReaderMPEGAudio::Create(this);
-				ManifestType = EMediaFormatType::MPEGAudio;
 			}
 			else
 			{
@@ -307,105 +301,100 @@ bool FAdaptiveStreamingPlayer::SelectManifest()
 	if (ManifestReader.IsValid())
 	{
 		check(Manifest == nullptr);
-		if (ManifestType != EMediaFormatType::Unknown)
+		TSharedPtrTS<IManifest> NewPresentation = ManifestReader->GetManifest();
+		check(NewPresentation.IsValid());
+
+		PlaybackState.SetSeekableRange(NewPresentation->GetSeekableTimeRange());
+		PlaybackState.SetTimelineRange(NewPresentation->GetTotalTimeRange());
+		PlaybackState.SetDuration(NewPresentation->GetDuration());
+
+		// Check for playback range restriction. This is currently assumed to come from URL fragment parameters
+		// like example.mp4#r=10.8,18.4
+		// These do not override any user defined range values!
+		FPlaybackRange PlaybackRange;
+		FTimeRange RestrictedPlaybackRange;
+		RestrictedPlaybackRange = NewPresentation->GetPlaybackRange(IManifest::EPlaybackRangeType::LockedPlaybackRange);
+		PlaybackState.GetPlayRange(PlaybackRange);
+		if (RestrictedPlaybackRange.Start.IsValid() && !PlaybackRange.Start.Get(FTimeValue()).IsValid())
 		{
-			TSharedPtrTS<IManifest> NewPresentation = ManifestReader->GetManifest();
-			check(NewPresentation.IsValid());
-
-			PlaybackState.SetSeekableRange(NewPresentation->GetSeekableTimeRange());
-			PlaybackState.SetTimelineRange(NewPresentation->GetTotalTimeRange());
-			PlaybackState.SetDuration(NewPresentation->GetDuration());
-
-			// Check for playback range restriction. This is currently assumed to come from URL fragment parameters
-			// like example.mp4#r=10.8,18.4
-			// These do not override any user defined range values!
-			FPlaybackRange PlaybackRange;
-			FTimeRange RestrictedPlaybackRange;
-			RestrictedPlaybackRange = NewPresentation->GetPlaybackRange(IManifest::EPlaybackRangeType::LockedPlaybackRange);
-			PlaybackState.GetPlayRange(PlaybackRange);
-			if (RestrictedPlaybackRange.Start.IsValid() && !PlaybackRange.Start.Get(FTimeValue()).IsValid())
-			{
-				PlaybackRange.Start = RestrictedPlaybackRange.Start;
-			}
-			if (RestrictedPlaybackRange.End.IsValid() && !PlaybackRange.End.Get(FTimeValue()).IsValid())
-			{
-				PlaybackRange.End = RestrictedPlaybackRange.End;
-			}
-			// Set the new range, but say it did not change. This is the first start and we do not need to issue a seek.
-			PlaybackState.SetPlayRange(PlaybackRange);
-			PlaybackState.ActivateNewPlayRange(nullptr);
-
-			TArray<FTrackMetadata> VideoTrackMetadata;
-			TArray<FTrackMetadata> AudioTrackMetadata;
-			TArray<FTrackMetadata> SubtitleTrackMetadata;
-			NewPresentation->GetTrackMetadata(VideoTrackMetadata, EStreamType::Video);
-			NewPresentation->GetTrackMetadata(AudioTrackMetadata, EStreamType::Audio);
-			NewPresentation->GetTrackMetadata(SubtitleTrackMetadata, EStreamType::Subtitle);
-			PlaybackState.SetTrackMetadata(VideoTrackMetadata, AudioTrackMetadata, SubtitleTrackMetadata);
-			PlaybackState.SetHaveMetadata(true);
-			// Get the supported playback rates for thinned and unthinned playback.
-			PlaybackState.SetPlaybackRates(EPlaybackRateType::Unthinned, NewPresentation->GetPossiblePlaybackRates(IManifest::EPlayRateType::UnthinnedRate));
-			PlaybackState.SetPlaybackRates(EPlaybackRateType::Thinned, NewPresentation->GetPossiblePlaybackRates(IManifest::EPlayRateType::ThinnedRate));
-
-			Manifest = NewPresentation;
-
-			CurrentState = EPlayerState::eState_Ready;
-
-			double minBufTimeMPD = Manifest->GetMinBufferTime().GetAsSeconds();
-			PlayerConfig.InitialBufferMinTimeAvailBeforePlayback = Utils::Min(minBufTimeMPD, PlayerConfig.InitialBufferMinTimeAvailBeforePlayback);
-			PlayerConfig.SeekBufferMinTimeAvailBeforePlayback = Utils::Min(minBufTimeMPD, PlayerConfig.SeekBufferMinTimeAvailBeforePlayback);
-			PlayerConfig.RebufferMinTimeAvailBeforePlayback = Utils::Min(minBufTimeMPD, PlayerConfig.RebufferMinTimeAvailBeforePlayback);
-
-			// For an mp4 or mkv stream we can now get rid of the manifest reader. It is no longer needed and we don't need to have it linger.
-			if (ManifestType == EMediaFormatType::ISOBMFF || ManifestType == EMediaFormatType::MKV || ManifestType == EMediaFormatType::MPEGAudio)
+			PlaybackRange.Start = RestrictedPlaybackRange.Start;
+		}
+		if (RestrictedPlaybackRange.End.IsValid() && !PlaybackRange.End.Get(FTimeValue()).IsValid())
+		{
+			PlaybackRange.End = RestrictedPlaybackRange.End;
+		}
+		// Set the new range, but say it did not change. This is the first start and we do not need to issue a seek.
+		PlaybackState.SetPlayRange(PlaybackRange);
+		PlaybackState.ActivateNewPlayRange(nullptr);
+
+		TArray<FTrackMetadata> VideoTrackMetadata;
+		TArray<FTrackMetadata> AudioTrackMetadata;
+		TArray<FTrackMetadata> SubtitleTrackMetadata;
+		NewPresentation->GetTrackMetadata(VideoTrackMetadata, EStreamType::Video);
+		NewPresentation->GetTrackMetadata(AudioTrackMetadata, EStreamType::Audio);
+		NewPresentation->GetTrackMetadata(SubtitleTrackMetadata, EStreamType::Subtitle);
+		PlaybackState.SetTrackMetadata(VideoTrackMetadata, AudioTrackMetadata, SubtitleTrackMetadata);
+		PlaybackState.SetHaveMetadata(true);
+		// Get the supported playback rates for thinned and unthinned playback.
+		PlaybackState.SetPlaybackRates(EPlaybackRateType::Unthinned, NewPresentation->GetPossiblePlaybackRates(IManifest::EPlayRateType::UnthinnedRate));
+		PlaybackState.SetPlaybackRates(EPlaybackRateType::Thinned, NewPresentation->GetPossiblePlaybackRates(IManifest::EPlayRateType::ThinnedRate));
+
+		Manifest = NewPresentation;
+
+		CurrentState = EPlayerState::eState_Ready;
+
+		double minBufTimeMPD = Manifest->GetMinBufferTime().GetAsSeconds();
+		PlayerConfig.InitialBufferMinTimeAvailBeforePlayback = Utils::Min(minBufTimeMPD, PlayerConfig.InitialBufferMinTimeAvailBeforePlayback);
+		PlayerConfig.SeekBufferMinTimeAvailBeforePlayback = Utils::Min(minBufTimeMPD, PlayerConfig.SeekBufferMinTimeAvailBeforePlayback);
+		PlayerConfig.RebufferMinTimeAvailBeforePlayback = Utils::Min(minBufTimeMPD, PlayerConfig.RebufferMinTimeAvailBeforePlayback);
+
+		// Let the ABR know the format as well.
+		StreamSelector->SetFormatType(ManifestReader->GetMediaFormatType());
+
+		// Live streams must not use any cache!
+		if (NewPresentation->GetPresentationType() == IManifest::EType::Live && HttpResponseCache.IsValid())
+		{
+			HttpResponseCache->Disable();
+		}
+
+		// Update the CMCD handler
+		switch(ManifestReader->GetMediaFormatType())
+		{
+			case EMediaFormatType::HLS:
 			{
-				InternalCloseManifestReader();
+				CMCDHandler->SetStreamingFormat(FCMCDHandler::EStreamingFormat::HLS);
+				CMCDHandler->SetStreamType(NewPresentation->GetPresentationType() == IManifest::EType::OnDemand ? FCMCDHandler::EStreamType::VOD : FCMCDHandler::EStreamType::Live);
+				break;
 			}
-
-			// Let the ABR know the format as well.
-			StreamSelector->SetFormatType(ManifestType);
-
-			// Live streams must not use any cache!
-			if (NewPresentation->GetPresentationType() == IManifest::EType::Live && HttpResponseCache.IsValid())
+			case EMediaFormatType::DASH:
 			{
-				HttpResponseCache->Disable();
+				CMCDHandler->SetStreamingFormat(NewPresentation->GetLowLatencyDescriptor().IsValid() ? FCMCDHandler::EStreamingFormat::LowLatencyDASH : FCMCDHandler::EStreamingFormat::DASH);
+				CMCDHandler->SetStreamType(NewPresentation->GetPresentationType() == IManifest::EType::OnDemand ? FCMCDHandler::EStreamType::VOD : FCMCDHandler::EStreamType::Live);
+				break;
 			}
-
-			// Update the CMCD handler
-			switch(ManifestType)
+			default:
 			{
-				case EMediaFormatType::HLS:
-				{
-					CMCDHandler->SetStreamingFormat(FCMCDHandler::EStreamingFormat::HLS);
-					CMCDHandler->SetStreamType(NewPresentation->GetPresentationType() == IManifest::EType::OnDemand ? FCMCDHandler::EStreamType::VOD : FCMCDHandler::EStreamType::Live);
-					break;
-				}
-				case EMediaFormatType::DASH:
-				{
-					CMCDHandler->SetStreamingFormat(NewPresentation->GetLowLatencyDescriptor().IsValid() ? FCMCDHandler::EStreamingFormat::LowLatencyDASH : FCMCDHandler::EStreamingFormat::DASH);
-					CMCDHandler->SetStreamType(NewPresentation->GetPresentationType() == IManifest::EType::OnDemand ? FCMCDHandler::EStreamType::VOD : FCMCDHandler::EStreamType::Live);
-					break;
-				}
-				default:
-				{
-					CMCDHandler->SetStreamingFormat(FCMCDHandler::EStreamingFormat::Undefined);
-					break;
-				}
+				CMCDHandler->SetStreamingFormat(FCMCDHandler::EStreamingFormat::Undefined);
+				break;
 			}
-
-			return true;
 		}
-		else
+
+		// For formats that do not reload the playlist every so often we can get rid of the manifest reader.
+		if (!Manifest->ReloadsPlaylistPeriodically())
 		{
-			// Handle other types of playlist here.
-			FErrorDetail err;
-			err.SetFacility(Facility::EFacility::Player);
-			err.SetMessage(TEXT("Unsupported playlist/manifest type"));
-			err.SetCode(INTERR_UNSUPPORTED_FORMAT);
-			PostError(err);
+			InternalCloseManifestReader();
 		}
+		return true;
+	}
+	else
+	{
+		// Handle other types of playlist here.
+		FErrorDetail err;
+		err.SetFacility(Facility::EFacility::Player);
+		err.SetMessage(TEXT("Unsupported playlist/manifest type"));
+		err.SetCode(INTERR_UNSUPPORTED_FORMAT);
+		PostError(err);
 	}
-
 	return false;
 }
 
diff --git a/Engine/Plugins/Media/ElectraPlayer/Source/ElectraPlayerRuntime/Private/Runtime/Player/DASH/ManifestDASH.cpp b/Engine/Plugins/Media/ElectraPlayer/Source/ElectraPlayerRuntime/Private/Runtime/Player/DASH/ManifestDASH.cpp
--- a/Engine/Plugins/Media/ElectraPlayer/Source/ElectraPlayerRuntime/Private/Runtime/Player/DASH/ManifestDASH.cpp
+++ b/Engine/Plugins/Media/ElectraPlayer/Source/ElectraPlayerRuntime/Private/Runtime/Player/DASH/ManifestDASH.cpp
@@ -338,6 +338,26 @@ void FManifestDASH::UpdateInternalManifest(TSharedPtrTS<FManifestDASHInternal> U
 	CurrentManifest = UpdatedManifest;
 }
 
+EMediaFormatType FManifestDASH::GetMediaFormatType() const
+{
+	return EMediaFormatType::DASH;
+}
+
+bool FManifestDASH::ReloadsPlaylistPeriodically() const
+{
+	return true;
+}
+
+bool FManifestDASH::SupportsMultipleTracksPerKind() const
+{
+	return true;
+}
+
+bool FManifestDASH::OrderOfTracksCanChangeDynamically() const
+{
+	return true;
+}
+
 IManifest::EType FManifestDASH::GetPresentationType() const
 {
 	TSharedPtrTS<FManifestDASHInternal> Manifest(CurrentManifest);
diff --git a/Engine/Plugins/Media/ElectraPlayer/Source/ElectraPlayerRuntime/Private/Runtime/Player/DASH/ManifestDASH.h b/Engine/Plugins/Media/ElectraPlayer/Source/ElectraPlayerRuntime/Private/Runtime/Player/DASH/ManifestDASH.h
--- a/Engine/Plugins/Media/ElectraPlayer/Source/ElectraPlayerRuntime/Private/Runtime/Player/DASH/ManifestDASH.h
+++ b/Engine/Plugins/Media/ElectraPlayer/Source/ElectraPlayerRuntime/Private/Runtime/Player/DASH/ManifestDASH.h
@@ -18,6 +18,10 @@ public:
 	void UpdateInternalManifest(TSharedPtrTS<FManifestDASHInternal> UpdatedManifest);
 
 	virtual ~FManifestDASH();
+	EMediaFormatType GetMediaFormatType() const override;
+	bool ReloadsPlaylistPeriodically() const override;
+	bool SupportsMultipleTracksPerKind() const override;
+	bool OrderOfTracksCanChangeDynamically() const override;
 	EType GetPresentationType() const override;
 	EReplayEventType GetReplayType() const override;
 	TSharedPtrTS<const FLowLatencyDescriptor> GetLowLatencyDescriptor() const override;
diff --git a/Engine/Plugins/Media/ElectraPlayer/Source/ElectraPlayerRuntime/Private/Runtime/Player/DASH/PlaylistReaderDASH.cpp b/Engine/Plugins/Media/ElectraPlayer/Source/ElectraPlayerRuntime/Private/Runtime/Player/DASH/PlaylistReaderDASH.cpp
--- a/Engine/Plugins/Media/ElectraPlayer/Source/ElectraPlayerRuntime/Private/Runtime/Player/DASH/PlaylistReaderDASH.cpp
+++ b/Engine/Plugins/Media/ElectraPlayer/Source/ElectraPlayerRuntime/Private/Runtime/Player/DASH/PlaylistReaderDASH.cpp
@@ -51,6 +51,8 @@ public:
 	void Initialize(IPlayerSessionServices* PlayerSessionServices);
 
 	virtual ~FPlaylistReaderDASH();
+	EMediaFormatType GetMediaFormatType() const override
+	{ return EMediaFormatType::DASH; }
 	void Close() override;
 	void HandleOnce() override;
 	const FString& GetPlaylistType() const override
diff --git a/Engine/Plugins/Media/ElectraPlayer/Source/ElectraPlayerRuntime/Private/Runtime/Player/HLS/PlaylistHLS.cpp b/Engine/Plugins/Media/ElectraPlayer/Source/ElectraPlayerRuntime/Private/Runtime/Player/HLS/PlaylistHLS.cpp
--- a/Engine/Plugins/Media/ElectraPlayer/Source/ElectraPlayerRuntime/Private/Runtime/Player/HLS/PlaylistHLS.cpp
+++ b/Engine/Plugins/Media/ElectraPlayer/Source/ElectraPlayerRuntime/Private/Runtime/Player/HLS/PlaylistHLS.cpp
@@ -2380,6 +2380,27 @@ void FActiveHLSPlaylist::FTimelineMediaAsset::UpdateRunningMetaData(const FStrin
 }
 
 
+EMediaFormatType FActiveHLSPlaylist::GetMediaFormatType() const
+{
+	return EMediaFormatType::HLS;
+}
+
+bool FActiveHLSPlaylist::ReloadsPlaylistPeriodically() const
+{
+	// This is true even for VoD since each variant has its own playlist that gets loaded just in time.
+	return true;
+}
+
+bool FActiveHLSPlaylist::SupportsMultipleTracksPerKind() const
+{
+	return true;
+}
+
+bool FActiveHLSPlaylist::OrderOfTracksCanChangeDynamically() const
+{
+	return false;
+}
+
 
 IManifest::EType FActiveHLSPlaylist::GetPresentationType() const
 {
diff --git a/Engine/Plugins/Media/ElectraPlayer/Source/ElectraPlayerRuntime/Private/Runtime/Player/HLS/PlaylistHLS.h b/Engine/Plugins/Media/ElectraPlayer/Source/ElectraPlayerRuntime/Private/Runtime/Player/HLS/PlaylistHLS.h
--- a/Engine/Plugins/Media/ElectraPlayer/Source/ElectraPlayerRuntime/Private/Runtime/Player/HLS/PlaylistHLS.h
+++ b/Engine/Plugins/Media/ElectraPlayer/Source/ElectraPlayerRuntime/Private/Runtime/Player/HLS/PlaylistHLS.h
@@ -529,6 +529,10 @@ public:
 	void GetAllMediaPlaylistLoadRequests(TArray<TSharedPtr<FLoadRequestHLSPlaylist, ESPMode::ThreadSafe>>& OutPlaylistLoadRequests, EStreamType InForType);
 	void CheckForPathwaySwitch();
 
+	EMediaFormatType GetMediaFormatType() const override;
+	bool ReloadsPlaylistPeriodically() const override;
+	bool SupportsMultipleTracksPerKind() const override;
+	bool OrderOfTracksCanChangeDynamically() const override;
 	EType GetPresentationType() const override;
 	EReplayEventType GetReplayType() const override;
 	TSharedPtrTS<const FLowLatencyDescriptor> GetLowLatencyDescriptor() const override;
diff --git a/Engine/Plugins/Media/ElectraPlayer/Source/ElectraPlayerRuntime/Private/Runtime/Player/HLS/PlaylistHandlerHLS.cpp b/Engine/Plugins/Media/ElectraPlayer/Source/ElectraPlayerRuntime/Private/Runtime/Player/HLS/PlaylistHandlerHLS.cpp
--- a/Engine/Plugins/Media/ElectraPlayer/Source/ElectraPlayerRuntime/Private/Runtime/Player/HLS/PlaylistHandlerHLS.cpp
+++ b/Engine/Plugins/Media/ElectraPlayer/Source/ElectraPlayerRuntime/Private/Runtime/Player/HLS/PlaylistHandlerHLS.cpp
@@ -25,6 +25,8 @@ public:
 	FPlaylistHandlerHLS(IPlayerSessionServices* InPlayerSessionServices);
 	virtual ~FPlaylistHandlerHLS();
 
+	EMediaFormatType GetMediaFormatType() const override
+	{ return EMediaFormatType::HLS; }
 	void Close() override;
 	void HandleOnce() override;
 	const FString& GetPlaylistType() const override
diff --git a/Engine/Plugins/Media/ElectraPlayer/Source/ElectraPlayerRuntime/Private/Runtime/Player/Manifest.h b/Engine/Plugins/Media/ElectraPlayer/Source/ElectraPlayerRuntime/Private/Runtime/Player/Manifest.h
--- a/Engine/Plugins/Media/ElectraPlayer/Source/ElectraPlayerRuntime/Private/Runtime/Player/Manifest.h
+++ b/Engine/Plugins/Media/ElectraPlayer/Source/ElectraPlayerRuntime/Private/Runtime/Player/Manifest.h
@@ -271,6 +271,16 @@ namespace Electra
 		//-------------------------------------------------------------------------
 		// Presentation related functions
 		//
+
+		//! Returns the media format of this presentation.
+		virtual EMediaFormatType GetMediaFormatType() const = 0;
+		//! Returns true if the media playlist will be updated periodically, false if not.
+		virtual bool ReloadsPlaylistPeriodically() const = 0;
+		//! Returns true if there could be more than one track per kind to switch between.
+		virtual bool SupportsMultipleTracksPerKind() const = 0;
+		//! Returns true if during playback the order of tracks can change (like in DASH period transitions).
+		virtual bool OrderOfTracksCanChangeDynamically() const = 0;
+
 		//! Returns the type of this presentation, either on-demand or live.
 		virtual EType GetPresentationType() const = 0;
 
diff --git a/Engine/Plugins/Media/ElectraPlayer/Source/ElectraPlayerRuntime/Private/Runtime/Player/PlayerStreamReader.h b/Engine/Plugins/Media/ElectraPlayer/Source/ElectraPlayerRuntime/Private/Runtime/Player/PlayerStreamReader.h
--- a/Engine/Plugins/Media/ElectraPlayer/Source/ElectraPlayerRuntime/Private/Runtime/Player/PlayerStreamReader.h
+++ b/Engine/Plugins/Media/ElectraPlayer/Source/ElectraPlayerRuntime/Private/Runtime/Player/PlayerStreamReader.h
@@ -43,6 +43,9 @@ namespace Electra
 		//! Creates the instance internally based on the creation parameters.
 		virtual UEMediaError Create(IPlayerSessionServices* PlayerSessionService, const CreateParam& createParam) = 0;
 
+		//! Returns true if data from all multiplexed tracks is read at the same time in a single sequential read.
+		virtual bool LoadsMultiplexedTracksSequentially() = 0;
+
 		//! Cancels all pending requests and stops processing. The instance can no longer be used after that!
 		virtual void Close() = 0;
 
diff --git a/Engine/Plugins/Media/ElectraPlayer/Source/ElectraPlayerRuntime/Private/Runtime/Player/PlaylistReader.h b/Engine/Plugins/Media/ElectraPlayer/Source/ElectraPlayerRuntime/Private/Runtime/Player/PlaylistReader.h
--- a/Engine/Plugins/Media/ElectraPlayer/Source/ElectraPlayerRuntime/Private/Runtime/Player/PlaylistReader.h
+++ b/Engine/Plugins/Media/ElectraPlayer/Source/ElectraPlayerRuntime/Private/Runtime/Player/PlaylistReader.h
@@ -6,6 +6,7 @@
 #include "Player/PlayerSessionServices.h"
 #include "StreamTypes.h"
 #include "Player/Playlist.h"
+#include "Player/Manifest.h"
 #include "ParameterDictionary.h"
 #include "HTTP/HTTPManager.h"
 
@@ -21,6 +22,11 @@ namespace Electra
 	public:
 		virtual ~IPlaylistReader() = default;
 
+		/**
+		 * Returns the media format this reader handles.
+		 */
+		virtual EMediaFormatType GetMediaFormatType() const = 0;
+
 		/**
 		 * Must call Close() before dropping any TShared..<> of this to allow for any internally
 		 * used TWeak..<> to still be valid if used by AsShared().
diff --git a/Engine/Plugins/Media/ElectraPlayer/Source/ElectraPlayerRuntime/Private/Runtime/Player/StreamSegmentReaderCommon.cpp b/Engine/Plugins/Media/ElectraPlayer/Source/ElectraPlayerRuntime/Private/Runtime/Player/StreamSegmentReaderCommon.cpp
--- a/Engine/Plugins/Media/ElectraPlayer/Source/ElectraPlayerRuntime/Private/Runtime/Player/StreamSegmentReaderCommon.cpp
+++ b/Engine/Plugins/Media/ElectraPlayer/Source/ElectraPlayerRuntime/Private/Runtime/Player/StreamSegmentReaderCommon.cpp
@@ -413,6 +413,11 @@ UEMediaError FStreamSegmentReaderCommon::Create(IPlayerSessionServices* InPlayer
 	return UEMEDIA_ERROR_OK;
 }
 
+bool FStreamSegmentReaderCommon::LoadsMultiplexedTracksSequentially()
+{
+	return false;
+}
+
 void FStreamSegmentReaderCommon::Close()
 {
 	if (bIsStarted)
diff --git a/Engine/Plugins/Media/ElectraPlayer/Source/ElectraPlayerRuntime/Private/Runtime/Player/StreamSegmentReaderCommon.h b/Engine/Plugins/Media/ElectraPlayer/Source/ElectraPlayerRuntime/Private/Runtime/Player/StreamSegmentReaderCommon.h
--- a/Engine/Plugins/Media/ElectraPlayer/Source/ElectraPlayerRuntime/Private/Runtime/Player/StreamSegmentReaderCommon.h
+++ b/Engine/Plugins/Media/ElectraPlayer/Source/ElectraPlayerRuntime/Private/Runtime/Player/StreamSegmentReaderCommon.h
@@ -325,17 +325,18 @@ public:
 	FStreamSegmentReaderCommon();
 	virtual ~FStreamSegmentReaderCommon();
 
-	virtual UEMediaError Create(IPlayerSessionServices* PlayerSessionService, const CreateParam& InCreateParam) override;
-	virtual void Close() override;
+	UEMediaError Create(IPlayerSessionServices* PlayerSessionService, const CreateParam& InCreateParam) override;
+	bool LoadsMultiplexedTracksSequentially() override;
+	void Close() override;
 
 	//! Adds a request to read from a stream
-	virtual EAddResult AddRequest(uint32 CurrentPlaybackSequenceID, TSharedPtrTS<IStreamSegment> Request) override;
+	EAddResult AddRequest(uint32 CurrentPlaybackSequenceID, TSharedPtrTS<IStreamSegment> Request) override;
 
 	//! Cancels any ongoing requests of the given stream type. Silent cancellation will not notify OnFragmentClose() or OnFragmentReachedEOS().
-	virtual void CancelRequest(EStreamType StreamType, bool bSilent) override;
+	void CancelRequest(EStreamType StreamType, bool bSilent) override;
 
 	//! Cancels all pending requests.
-	virtual void CancelRequests() override;
+	void CancelRequests() override;
 
 private:
 	ELECTRA_IMPL_DEFAULT_ERROR_METHODS(CommonSegmentReader);
diff --git a/Engine/Plugins/Media/ElectraPlayer/Source/ElectraPlayerRuntime/Private/Runtime/Player/mkv/ManifestMKV.cpp b/Engine/Plugins/Media/ElectraPlayer/Source/ElectraPlayerRuntime/Private/Runtime/Player/mkv/ManifestMKV.cpp
--- a/Engine/Plugins/Media/ElectraPlayer/Source/ElectraPlayerRuntime/Private/Runtime/Player/mkv/ManifestMKV.cpp
+++ b/Engine/Plugins/Media/ElectraPlayer/Source/ElectraPlayerRuntime/Private/Runtime/Player/mkv/ManifestMKV.cpp
@@ -67,6 +67,27 @@ FErrorDetail FManifestMKVInternal::Build(TSharedPtrTS<IParserMKV> MKVParser, con
 }
 
 
+EMediaFormatType FManifestMKVInternal::GetMediaFormatType() const
+{
+	return EMediaFormatType::MKV;
+}
+
+bool FManifestMKVInternal::ReloadsPlaylistPeriodically() const
+{
+	return false;
+}
+
+bool FManifestMKVInternal::SupportsMultipleTracksPerKind() const
+{
+	return true;
+}
+
+bool FManifestMKVInternal::OrderOfTracksCanChangeDynamically() const
+{
+	return false;
+}
+
+
 //-----------------------------------------------------------------------------
 /**
  * Returns the type of presentation.
diff --git a/Engine/Plugins/Media/ElectraPlayer/Source/ElectraPlayerRuntime/Private/Runtime/Player/mkv/ManifestMKV.h b/Engine/Plugins/Media/ElectraPlayer/Source/ElectraPlayerRuntime/Private/Runtime/Player/mkv/ManifestMKV.h
--- a/Engine/Plugins/Media/ElectraPlayer/Source/ElectraPlayerRuntime/Private/Runtime/Player/mkv/ManifestMKV.h
+++ b/Engine/Plugins/Media/ElectraPlayer/Source/ElectraPlayerRuntime/Private/Runtime/Player/mkv/ManifestMKV.h
@@ -30,6 +30,10 @@ public:
 
 	FErrorDetail Build(TSharedPtrTS<IParserMKV> MP4Parser, const FString& URL, const HTTP::FConnectionInfo& InConnectionInfo);
 
+	EMediaFormatType GetMediaFormatType() const override;
+	bool ReloadsPlaylistPeriodically() const override;
+	bool SupportsMultipleTracksPerKind() const override;
+	bool OrderOfTracksCanChangeDynamically() const override;
 	EType GetPresentationType() const override;
 	EReplayEventType GetReplayType() const override
 	{ return IManifest::EReplayEventType::NoReplay; }
diff --git a/Engine/Plugins/Media/ElectraPlayer/Source/ElectraPlayerRuntime/Private/Runtime/Player/mkv/PlaylistReaderMKV.cpp b/Engine/Plugins/Media/ElectraPlayer/Source/ElectraPlayerRuntime/Private/Runtime/Player/mkv/PlaylistReaderMKV.cpp
--- a/Engine/Plugins/Media/ElectraPlayer/Source/ElectraPlayerRuntime/Private/Runtime/Player/mkv/PlaylistReaderMKV.cpp
+++ b/Engine/Plugins/Media/ElectraPlayer/Source/ElectraPlayerRuntime/Private/Runtime/Player/mkv/PlaylistReaderMKV.cpp
@@ -36,9 +36,11 @@ public:
 	void Initialize(IPlayerSessionServices* PlayerSessionServices);
 
 	virtual ~FPlaylistReaderMKV();
-
-	virtual void Close() override;
-	virtual void HandleOnce() override;
+	
+	EMediaFormatType GetMediaFormatType() const override
+	{ return EMediaFormatType::MKV; }
+	void Close() override;
+	void HandleOnce() override;
 
 	/**
 	 * Returns the type of playlist format.
@@ -47,7 +49,7 @@ public:
 	 *
 	 * @return "mkv" to indicate this is an mkv file.
 	 */
-	virtual const FString& GetPlaylistType() const override
+	const FString& GetPlaylistType() const override
 	{
 		static FString Type("mkv");
 		return Type;
@@ -59,21 +61,21 @@ public:
 	 * @param URL URL of the playlist to load
 	 * @param InReplayEventParams URL fragments pertaining to replay event which have been removed from the source URL
 	 */
-	virtual void LoadAndParse(const FString& InURL, const Playlist::FReplayEventParams& InReplayEventParams) override;
+	void LoadAndParse(const FString& InURL, const Playlist::FReplayEventParams& InReplayEventParams) override;
 
 	/**
 	 * Returns the URL from which the playlist was loaded (or supposed to be loaded).
 	 *
 	 * @return The playlist URL
 	 */
-	virtual FString GetURL() const override;
+	FString GetURL() const override;
 
 	/**
 	 * Returns an interface to the manifest created from the loaded mkv playlists.
 	 *
 	 * @return A shared manifest interface pointer.
 	 */
-	virtual TSharedPtrTS<IManifest> GetManifest() override;
+	TSharedPtrTS<IManifest> GetManifest() override;
 
 private:
 	// Methods from IParserMKV::IReader
diff --git a/Engine/Plugins/Media/ElectraPlayer/Source/ElectraPlayerRuntime/Private/Runtime/Player/mkv/StreamReaderMKV.cpp b/Engine/Plugins/Media/ElectraPlayer/Source/ElectraPlayerRuntime/Private/Runtime/Player/mkv/StreamReaderMKV.cpp
--- a/Engine/Plugins/Media/ElectraPlayer/Source/ElectraPlayerRuntime/Private/Runtime/Player/mkv/StreamReaderMKV.cpp
+++ b/Engine/Plugins/Media/ElectraPlayer/Source/ElectraPlayerRuntime/Private/Runtime/Player/mkv/StreamReaderMKV.cpp
@@ -155,6 +155,11 @@ UEMediaError FStreamReaderMKV::Create(IPlayerSessionServices* InPlayerSessionSer
 	return UEMEDIA_ERROR_OK;
 }
 
+bool FStreamReaderMKV::LoadsMultiplexedTracksSequentially()
+{
+	return true;
+}
+
 void FStreamReaderMKV::Close()
 {
 	if (bIsStarted)
diff --git a/Engine/Plugins/Media/ElectraPlayer/Source/ElectraPlayerRuntime/Private/Runtime/Player/mkv/StreamReaderMKV.h b/Engine/Plugins/Media/ElectraPlayer/Source/ElectraPlayerRuntime/Private/Runtime/Player/mkv/StreamReaderMKV.h
--- a/Engine/Plugins/Media/ElectraPlayer/Source/ElectraPlayerRuntime/Private/Runtime/Player/mkv/StreamReaderMKV.h
+++ b/Engine/Plugins/Media/ElectraPlayer/Source/ElectraPlayerRuntime/Private/Runtime/Player/mkv/StreamReaderMKV.h
@@ -107,6 +107,7 @@ public:
 	FStreamReaderMKV() = default;
 	virtual ~FStreamReaderMKV();
 	UEMediaError Create(IPlayerSessionServices* PlayerSessionService, const CreateParam &createParam) override;
+	bool LoadsMultiplexedTracksSequentially() override;
 	void Close() override;
 	EAddResult AddRequest(uint32 CurrentPlaybackSequenceID, TSharedPtrTS<IStreamSegment> Request) override;
 	void CancelRequest(EStreamType StreamType, bool bSilent) override;
diff --git a/Engine/Plugins/Media/ElectraPlayer/Source/ElectraPlayerRuntime/Private/Runtime/Player/mp4/ManifestMP4.cpp b/Engine/Plugins/Media/ElectraPlayer/Source/ElectraPlayerRuntime/Private/Runtime/Player/mp4/ManifestMP4.cpp
--- a/Engine/Plugins/Media/ElectraPlayer/Source/ElectraPlayerRuntime/Private/Runtime/Player/mp4/ManifestMP4.cpp
+++ b/Engine/Plugins/Media/ElectraPlayer/Source/ElectraPlayerRuntime/Private/Runtime/Player/mp4/ManifestMP4.cpp
@@ -88,6 +88,27 @@ void FManifestMP4Internal::LogMessage(IInfoLog::ELevel Level, const FString& Mes
 }
 
 
+
+EMediaFormatType FManifestMP4Internal::GetMediaFormatType() const
+{
+	return EMediaFormatType::ISOBMFF;
+}
+
+bool FManifestMP4Internal::ReloadsPlaylistPeriodically() const
+{
+	return false;
+}
+
+bool FManifestMP4Internal::SupportsMultipleTracksPerKind() const
+{
+	return true;
+}
+
+bool FManifestMP4Internal::OrderOfTracksCanChangeDynamically() const
+{
+	return false;
+}
+
 //-----------------------------------------------------------------------------
 /**
  * Returns the type of presentation.
diff --git a/Engine/Plugins/Media/ElectraPlayer/Source/ElectraPlayerRuntime/Private/Runtime/Player/mp4/ManifestMP4.h b/Engine/Plugins/Media/ElectraPlayer/Source/ElectraPlayerRuntime/Private/Runtime/Player/mp4/ManifestMP4.h
--- a/Engine/Plugins/Media/ElectraPlayer/Source/ElectraPlayerRuntime/Private/Runtime/Player/mp4/ManifestMP4.h
+++ b/Engine/Plugins/Media/ElectraPlayer/Source/ElectraPlayerRuntime/Private/Runtime/Player/mp4/ManifestMP4.h
@@ -30,6 +30,10 @@ public:
 
 	FErrorDetail Build(TSharedPtrTS<IParserISO14496_12> MP4Parser, const FString& URL, const HTTP::FConnectionInfo& InConnectionInfo);
 
+	EMediaFormatType GetMediaFormatType() const override;
+	bool ReloadsPlaylistPeriodically() const override;
+	bool SupportsMultipleTracksPerKind() const override;
+	bool OrderOfTracksCanChangeDynamically() const override;
 	EType GetPresentationType() const override;
 	EReplayEventType GetReplayType() const override
 	{ return IManifest::EReplayEventType::NoReplay; }
diff --git a/Engine/Plugins/Media/ElectraPlayer/Source/ElectraPlayerRuntime/Private/Runtime/Player/mp4/PlaylistReaderMP4.cpp b/Engine/Plugins/Media/ElectraPlayer/Source/ElectraPlayerRuntime/Private/Runtime/Player/mp4/PlaylistReaderMP4.cpp
--- a/Engine/Plugins/Media/ElectraPlayer/Source/ElectraPlayerRuntime/Private/Runtime/Player/mp4/PlaylistReaderMP4.cpp
+++ b/Engine/Plugins/Media/ElectraPlayer/Source/ElectraPlayerRuntime/Private/Runtime/Player/mp4/PlaylistReaderMP4.cpp
@@ -40,8 +40,10 @@ public:
 
 	virtual ~FPlaylistReaderMP4();
 
-	virtual void Close() override;
-	virtual void HandleOnce() override;
+	EMediaFormatType GetMediaFormatType() const override
+	{ return EMediaFormatType::ISOBMFF; }
+	void Close() override;
+	void HandleOnce() override;
 
 	/**
 	 * Returns the type of playlist format.
@@ -49,7 +51,7 @@ public:
 	 *
 	 * @return "mp4" to indicate this is an mp4 file.
 	 */
-	virtual const FString& GetPlaylistType() const override
+	const FString& GetPlaylistType() const override
 	{
 		static FString Type("mp4");
 		return Type;
@@ -61,21 +63,21 @@ public:
 	 * @param URL URL of the playlist to load
 	 * @param InReplayEventParams URL fragments pertaining to replay event which have been removed from the source URL
 	 */
-	virtual void LoadAndParse(const FString& InURL, const Playlist::FReplayEventParams& InReplayEventParams) override;
+	void LoadAndParse(const FString& InURL, const Playlist::FReplayEventParams& InReplayEventParams) override;
 
 	/**
 	 * Returns the URL from which the playlist was loaded (or supposed to be loaded).
 	 *
 	 * @return The playlist URL
 	 */
-	virtual FString GetURL() const override;
+	FString GetURL() const override;
 
 	/**
 	 * Returns an interface to the manifest created from the loaded mp4 playlists.
 	 *
 	 * @return A shared manifest interface pointer.
 	 */
-	virtual TSharedPtrTS<IManifest> GetManifest() override;
+	TSharedPtrTS<IManifest> GetManifest() override;
 
 private:
 	// Methods from IGenericDataReader
diff --git a/Engine/Plugins/Media/ElectraPlayer/Source/ElectraPlayerRuntime/Private/Runtime/Player/mp4/StreamReaderMP4.cpp b/Engine/Plugins/Media/ElectraPlayer/Source/ElectraPlayerRuntime/Private/Runtime/Player/mp4/StreamReaderMP4.cpp
--- a/Engine/Plugins/Media/ElectraPlayer/Source/ElectraPlayerRuntime/Private/Runtime/Player/mp4/StreamReaderMP4.cpp
+++ b/Engine/Plugins/Media/ElectraPlayer/Source/ElectraPlayerRuntime/Private/Runtime/Player/mp4/StreamReaderMP4.cpp
@@ -179,6 +179,11 @@ UEMediaError FStreamReaderMP4::Create(IPlayerSessionServices* InPlayerSessionSer
 	return UEMEDIA_ERROR_OK;
 }
 
+bool FStreamReaderMP4::LoadsMultiplexedTracksSequentially()
+{
+	return true;
+}
+
 void FStreamReaderMP4::Close()
 {
 	if (bIsStarted)
diff --git a/Engine/Plugins/Media/ElectraPlayer/Source/ElectraPlayerRuntime/Private/Runtime/Player/mp4/StreamReaderMP4.h b/Engine/Plugins/Media/ElectraPlayer/Source/ElectraPlayerRuntime/Private/Runtime/Player/mp4/StreamReaderMP4.h
--- a/Engine/Plugins/Media/ElectraPlayer/Source/ElectraPlayerRuntime/Private/Runtime/Player/mp4/StreamReaderMP4.h
+++ b/Engine/Plugins/Media/ElectraPlayer/Source/ElectraPlayerRuntime/Private/Runtime/Player/mp4/StreamReaderMP4.h
@@ -96,11 +96,12 @@ class FStreamReaderMP4 : public IStreamReader, public FMediaThread
 public:
 	FStreamReaderMP4();
 	virtual ~FStreamReaderMP4();
-	virtual UEMediaError Create(IPlayerSessionServices* PlayerSessionService, const CreateParam &createParam) override;
-	virtual void Close() override;
-	virtual EAddResult AddRequest(uint32 CurrentPlaybackSequenceID, TSharedPtrTS<IStreamSegment> Request) override;
-	virtual void CancelRequest(EStreamType StreamType, bool bSilent) override;
-	virtual void CancelRequests() override;
+	UEMediaError Create(IPlayerSessionServices* PlayerSessionService, const CreateParam &createParam) override;
+	bool LoadsMultiplexedTracksSequentially() override;
+	void Close() override;
+	EAddResult AddRequest(uint32 CurrentPlaybackSequenceID, TSharedPtrTS<IStreamSegment> Request) override;
+	void CancelRequest(EStreamType StreamType, bool bSilent) override;
+	void CancelRequests() override;
 
 private:
 	void WorkerThread();
diff --git a/Engine/Plugins/Media/ElectraPlayer/Source/ElectraPlayerRuntime/Private/Runtime/Player/mpegaudio/ManifestMPEGAudio.cpp b/Engine/Plugins/Media/ElectraPlayer/Source/ElectraPlayerRuntime/Private/Runtime/Player/mpegaudio/ManifestMPEGAudio.cpp
--- a/Engine/Plugins/Media/ElectraPlayer/Source/ElectraPlayerRuntime/Private/Runtime/Player/mpegaudio/ManifestMPEGAudio.cpp
+++ b/Engine/Plugins/Media/ElectraPlayer/Source/ElectraPlayerRuntime/Private/Runtime/Player/mpegaudio/ManifestMPEGAudio.cpp
@@ -100,6 +100,28 @@ void FManifestMPEGAudioInternal::LogMessage(IInfoLog::ELevel Level, const FStrin
 }
 
 
+
+EMediaFormatType FManifestMPEGAudioInternal::GetMediaFormatType() const
+{
+	return EMediaFormatType::MPEGAudio;
+}
+
+bool FManifestMPEGAudioInternal::ReloadsPlaylistPeriodically() const
+{
+	return false;
+}
+
+bool FManifestMPEGAudioInternal::SupportsMultipleTracksPerKind() const
+{
+	return false;
+}
+
+bool FManifestMPEGAudioInternal::OrderOfTracksCanChangeDynamically() const
+{
+	return false;
+}
+
+
 //-----------------------------------------------------------------------------
 /**
  * Returns the type of presentation.
diff --git a/Engine/Plugins/Media/ElectraPlayer/Source/ElectraPlayerRuntime/Private/Runtime/Player/mpegaudio/ManifestMPEGAudio.h b/Engine/Plugins/Media/ElectraPlayer/Source/ElectraPlayerRuntime/Private/Runtime/Player/mpegaudio/ManifestMPEGAudio.h
--- a/Engine/Plugins/Media/ElectraPlayer/Source/ElectraPlayerRuntime/Private/Runtime/Player/mpegaudio/ManifestMPEGAudio.h
+++ b/Engine/Plugins/Media/ElectraPlayer/Source/ElectraPlayerRuntime/Private/Runtime/Player/mpegaudio/ManifestMPEGAudio.h
@@ -79,6 +79,10 @@ public:
 
 	FErrorDetail Build(const FMPEGAudioInfoHeader& InHeader, const FString& URL);
 
+	EMediaFormatType GetMediaFormatType() const override;
+	bool ReloadsPlaylistPeriodically() const override;
+	bool SupportsMultipleTracksPerKind() const override;
+	bool OrderOfTracksCanChangeDynamically() const override;
 	EType GetPresentationType() const override;
 	EReplayEventType GetReplayType() const override
 	{ return IManifest::EReplayEventType::NoReplay; }
diff --git a/Engine/Plugins/Media/ElectraPlayer/Source/ElectraPlayerRuntime/Private/Runtime/Player/mpegaudio/PlaylistReaderMPEGAudio.cpp b/Engine/Plugins/Media/ElectraPlayer/Source/ElectraPlayerRuntime/Private/Runtime/Player/mpegaudio/PlaylistReaderMPEGAudio.cpp
--- a/Engine/Plugins/Media/ElectraPlayer/Source/ElectraPlayerRuntime/Private/Runtime/Player/mpegaudio/PlaylistReaderMPEGAudio.cpp
+++ b/Engine/Plugins/Media/ElectraPlayer/Source/ElectraPlayerRuntime/Private/Runtime/Player/mpegaudio/PlaylistReaderMPEGAudio.cpp
@@ -40,6 +40,8 @@ public:
 
 	virtual ~FPlaylistReaderMPEGAudio();
 
+	EMediaFormatType GetMediaFormatType() const override
+	{ return EMediaFormatType::MPEGAudio; }
 	void Close() override;
 	void HandleOnce() override;
 
diff --git a/Engine/Plugins/Media/ElectraPlayer/Source/ElectraPlayerRuntime/Private/Runtime/Player/mpegaudio/StreamReaderMPEGAudio.cpp b/Engine/Plugins/Media/ElectraPlayer/Source/ElectraPlayerRuntime/Private/Runtime/Player/mpegaudio/StreamReaderMPEGAudio.cpp
--- a/Engine/Plugins/Media/ElectraPlayer/Source/ElectraPlayerRuntime/Private/Runtime/Player/mpegaudio/StreamReaderMPEGAudio.cpp
+++ b/Engine/Plugins/Media/ElectraPlayer/Source/ElectraPlayerRuntime/Private/Runtime/Player/mpegaudio/StreamReaderMPEGAudio.cpp
@@ -119,6 +119,11 @@ UEMediaError FStreamReaderMPEGAudio::Create(IPlayerSessionServices* InPlayerSess
 	return UEMEDIA_ERROR_OK;
 }
 
+bool FStreamReaderMPEGAudio::LoadsMultiplexedTracksSequentially()
+{
+	return false;
+}
+
 void FStreamReaderMPEGAudio::Close()
 {
 	if (bIsStarted)
diff --git a/Engine/Plugins/Media/ElectraPlayer/Source/ElectraPlayerRuntime/Private/Runtime/Player/mpegaudio/StreamReaderMPEGAudio.h b/Engine/Plugins/Media/ElectraPlayer/Source/ElectraPlayerRuntime/Private/Runtime/Player/mpegaudio/StreamReaderMPEGAudio.h
--- a/Engine/Plugins/Media/ElectraPlayer/Source/ElectraPlayerRuntime/Private/Runtime/Player/mpegaudio/StreamReaderMPEGAudio.h
+++ b/Engine/Plugins/Media/ElectraPlayer/Source/ElectraPlayerRuntime/Private/Runtime/Player/mpegaudio/StreamReaderMPEGAudio.h
@@ -98,11 +98,12 @@ class FStreamReaderMPEGAudio : public IStreamReader, public FMediaThread
 public:
 	FStreamReaderMPEGAudio();
 	virtual ~FStreamReaderMPEGAudio();
-	virtual UEMediaError Create(IPlayerSessionServices* PlayerSessionService, const CreateParam &createParam) override;
-	virtual void Close() override;
-	virtual EAddResult AddRequest(uint32 CurrentPlaybackSequenceID, TSharedPtrTS<IStreamSegment> Request) override;
-	virtual void CancelRequest(EStreamType StreamType, bool bSilent) override;
-	virtual void CancelRequests() override;
+	UEMediaError Create(IPlayerSessionServices* PlayerSessionService, const CreateParam &createParam) override;
+	bool LoadsMultiplexedTracksSequentially() override;
+	void Close() override;
+	EAddResult AddRequest(uint32 CurrentPlaybackSequenceID, TSharedPtrTS<IStreamSegment> Request) override;
+	void CancelRequest(EStreamType StreamType, bool bSilent) override;
+	void CancelRequests() override;
 
 private:
 	void WorkerThread();
