[IOS] Setting the requires platform sdk flag in multiple modules that depends on them.

This feature is overall disabled, but updating the modules allows for easier testing until we can turn it on by default.

This CL is basically a no op until the feature is enabled

#Synced-CL 50224750

diff --git a/Engine/Plugins/Animation/ControlRig/Source/ControlRig/ControlRig.Build.cs b/Engine/Plugins/Animation/ControlRig/Source/ControlRig/ControlRig.Build.cs
--- a/Engine/Plugins/Animation/ControlRig/Source/ControlRig/ControlRig.Build.cs
+++ b/Engine/Plugins/Animation/ControlRig/Source/ControlRig/ControlRig.Build.cs
@@ -8,6 +8,8 @@ namespace UnrealBuildTool.Rules
         {
 			NumIncludedBytesPerUnityCPPOverride = 688128; // best unity size found from using UBT ProfileUnitySizes mode
 
+			bRequiresPlatformSDK = true;
+
             PrivateDependencyModuleNames.AddRange(
                 new string[]
                 {
diff --git a/Engine/Plugins/Media/ElectraCodecs/Source/ElectraDecoders/ElectraDecoders.Build.cs b/Engine/Plugins/Media/ElectraCodecs/Source/ElectraDecoders/ElectraDecoders.Build.cs
--- a/Engine/Plugins/Media/ElectraCodecs/Source/ElectraDecoders/ElectraDecoders.Build.cs
+++ b/Engine/Plugins/Media/ElectraCodecs/Source/ElectraDecoders/ElectraDecoders.Build.cs
@@ -25,6 +25,8 @@ namespace UnrealBuildTool.Rules
 			//
 			bLegalToDistributeObjectCode = true;
 
+			bRequiresPlatformSDK = true;
+
 			PrivateDependencyModuleNames.AddRange(
 				new string[] {
 					"Core",
diff --git a/Engine/Plugins/Media/ElectraPlayer/Source/ElectraPlayerRuntime/ElectraPlayerRuntime.Build.cs b/Engine/Plugins/Media/ElectraPlayer/Source/ElectraPlayerRuntime/ElectraPlayerRuntime.Build.cs
--- a/Engine/Plugins/Media/ElectraPlayer/Source/ElectraPlayerRuntime/ElectraPlayerRuntime.Build.cs
+++ b/Engine/Plugins/Media/ElectraPlayer/Source/ElectraPlayerRuntime/ElectraPlayerRuntime.Build.cs
@@ -15,6 +15,8 @@ namespace UnrealBuildTool.Rules
 			bLegalToDistributeObjectCode = true;
 			IWYUSupport = IWYUSupport.None;
 
+			bRequiresPlatformSDK = true;
+
 			PrivateDependencyModuleNames.AddRange(
 				new string[] {
 					"Core",
diff --git a/Engine/Plugins/Media/ElectraUtil/Source/ElectraSamples/ElectraSamples.Build.cs b/Engine/Plugins/Media/ElectraUtil/Source/ElectraSamples/ElectraSamples.Build.cs
--- a/Engine/Plugins/Media/ElectraUtil/Source/ElectraSamples/ElectraSamples.Build.cs
+++ b/Engine/Plugins/Media/ElectraUtil/Source/ElectraSamples/ElectraSamples.Build.cs
@@ -9,6 +9,7 @@ namespace UnrealBuildTool.Rules
 	{
 		public ElectraSamples(ReadOnlyTargetRules Target) : base(Target)
 		{
+			bRequiresPlatformSDK = true;
 			bLegalToDistributeObjectCode = true;
 
 			DynamicallyLoadedModuleNames.AddRange(
diff --git a/Engine/Plugins/Online/EOSShared/Source/EOSShared/EOSShared.Build.cs b/Engine/Plugins/Online/EOSShared/Source/EOSShared/EOSShared.Build.cs
--- a/Engine/Plugins/Online/EOSShared/Source/EOSShared/EOSShared.Build.cs
+++ b/Engine/Plugins/Online/EOSShared/Source/EOSShared/EOSShared.Build.cs
@@ -25,6 +25,8 @@ public class EOSShared : ModuleRules
 		
 		Type = ModuleType.CPlusPlus;
 
+		bRequiresPlatformSDK = true;
+
 		PublicDefinitions.Add("UE_WITH_EOS_SDK_APIVERSION_WARNINGS=" + (EnableApiVersionWarnings ? "1" : "0"));
 
 		PublicDependencyModuleNames.AddRange(
diff --git a/Engine/Plugins/Online/OnlineServicesEOSGS/Source/OnlineServicesEOSGS.Build.cs b/Engine/Plugins/Online/OnlineServicesEOSGS/Source/OnlineServicesEOSGS.Build.cs
--- a/Engine/Plugins/Online/OnlineServicesEOSGS/Source/OnlineServicesEOSGS.Build.cs
+++ b/Engine/Plugins/Online/OnlineServicesEOSGS/Source/OnlineServicesEOSGS.Build.cs
@@ -7,6 +7,7 @@ public class OnlineServicesEOSGS : ModuleRules
 {
 	public OnlineServicesEOSGS(ReadOnlyTargetRules Target) : base(Target)
 	{
+		bRequiresPlatformSDK = true;
 		PublicDependencyModuleNames.AddRange(
 			new string[] {
 				"Core",
diff --git a/Engine/Plugins/Online/WebAuth/Source/WebAuth.Build.cs b/Engine/Plugins/Online/WebAuth/Source/WebAuth.Build.cs
--- a/Engine/Plugins/Online/WebAuth/Source/WebAuth.Build.cs
+++ b/Engine/Plugins/Online/WebAuth/Source/WebAuth.Build.cs
@@ -11,6 +11,7 @@ public class WebAuth : ModuleRules
 		
 		PCHUsage = PCHUsageMode.UseExplicitOrSharedPCHs;
 
+		bRequiresPlatformSDK = true;
 		PrivateDependencyModuleNames.AddRange(
 			new string[]
 			{
diff --git a/Engine/Plugins/Runtime/Firebase/Source/Firebase.Build.cs b/Engine/Plugins/Runtime/Firebase/Source/Firebase.Build.cs
--- a/Engine/Plugins/Runtime/Firebase/Source/Firebase.Build.cs
+++ b/Engine/Plugins/Runtime/Firebase/Source/Firebase.Build.cs
@@ -8,6 +8,7 @@ namespace UnrealBuildTool.Rules
 	{
 		public Firebase(ReadOnlyTargetRules Target) : base(Target)
 		{
+			bRequiresPlatformSDK = true;
 			PrivateDependencyModuleNames.AddRange(new string[] {
                     "Core",
                     "CoreUObject",
