Electra: Removed dead code and added upstream error reporting when video decoder implementation returns an error on obtaining output

#rnx

#Synced-CL 50433613

diff --git a/Engine/Plugins/Media/ElectraPlayer/Source/ElectraPlayerRuntime/Private/Runtime/Decoder/VideoDecoder.cpp b/Engine/Plugins/Media/ElectraPlayer/Source/ElectraPlayerRuntime/Private/Runtime/Decoder/VideoDecoder.cpp
--- a/Engine/Plugins/Media/ElectraPlayer/Source/ElectraPlayerRuntime/Private/Runtime/Decoder/VideoDecoder.cpp
+++ b/Engine/Plugins/Media/ElectraPlayer/Source/ElectraPlayerRuntime/Private/Runtime/Decoder/VideoDecoder.cpp
@@ -69,11 +69,6 @@ public:
 	void SetReadyBufferListener(IDecoderOutputBufferListener* Listener) override;
 
 private:
-	enum
-	{
-		CfgMakeAUCopy = 0
-	};
-
 	struct FDecoderInput
 	{
 		~FDecoderInput()
@@ -83,7 +78,6 @@ private:
 
 		IElectraDecoder::FInputAccessUnit DecAU;
 		TMap<FString, FVariant> CSDOptions;
-		TArray<uint8> DataCopy;
 		TSharedPtr<IElectraDecoderBitstreamInfo, ESPMode::ThreadSafe> BSI;
 		FTimeValue AdjustedPTS;
 		FTimeValue AdjustedDuration;
@@ -735,15 +729,8 @@ FVideoDecoderImpl::EAUChangeFlags FVideoDecoderImpl::PrepareAU(TSharedPtrTS<FDec
 		AU->DecAU.Flags |= AU->AccessUnit->bIsSyncSample ? EElectraDecoderFlags::IsSyncSample : EElectraDecoderFlags::None;
 
 		// Set the bitstream data and size in the decoder input description.
-		// If the bitstream processor will modify the data in place we need to make a copy.
-		// NOTE: In-place modification only means changing values in place, but not removing or inserting new data.
 		AU->DecAU.Data = AU->AccessUnit->AUData;
 		AU->DecAU.DataSize = AU->AccessUnit->AUSize;
-		if (DecoderBitstreamProcessor.IsValid() && DecoderBitstreamProcessor->WillModifyBitstreamInPlace() && CfgMakeAUCopy)
-		{
-			AU->DataCopy = MakeConstArrayView<const uint8>(reinterpret_cast<const uint8*>(AU->AccessUnit->AUData), AU->AccessUnit->AUSize);
-			AU->DecAU.Data = AU->DataCopy.GetData();
-		}
 		// We do not need to set up more than the bitstream data and the codec specific info to call the bitstream processor.
 		IElectraDecoderBitstreamProcessor::EProcessResult BSResult = DecoderBitstreamProcessor->ProcessInputForDecoding(AU->BSI, AU->DecAU, AU->CSDOptions);
 		if (BSResult == IElectraDecoderBitstreamProcessor::EProcessResult::Error)
@@ -1098,7 +1085,11 @@ IElectraDecoder::EOutputStatus FVideoDecoderImpl::HandleOutput()
 	{
 		OutputStatus = IElectraDecoder::EOutputStatus::EndOfData;
 	}
-
+	// Did the decoder return an error?
+	if (OutputStatus == IElectraDecoder::EOutputStatus::Error)
+	{
+		PostError(DecoderInstance->GetError());
+	}
 	return OutputStatus;
 }
 
diff --git a/Engine/Plugins/Media/NVDECElectra/Source/NVDECElectra/Private/NVDEC/ElectraMediaNVDEC.cpp b/Engine/Plugins/Media/NVDECElectra/Source/NVDECElectra/Private/NVDEC/ElectraMediaNVDEC.cpp
--- a/Engine/Plugins/Media/NVDECElectra/Source/NVDECElectra/Private/NVDEC/ElectraMediaNVDEC.cpp
+++ b/Engine/Plugins/Media/NVDECElectra/Source/NVDECElectra/Private/NVDEC/ElectraMediaNVDEC.cpp
@@ -1530,11 +1530,15 @@ FVideoDecoderNVDECElectra::EConvertResult FVideoDecoderNVDECElectra::ConvertDeco
 			return EConvertResult::Failure;
 		}
 	}
-	// Remove outdated entries.
-	InDecoderInput.RemoveAll([&](const TSharedPtr<FDecoderInput, ESPMode::NotThreadSafe>& e)
+	// Remove entries that are severely outdated and have not been returned for some reason.
+	if (DecodeState != EDecodeState::Draining)
 	{
-		return e->AccessUnit.PTS < FramePTS;
-	});
+		FramePTS -= FTimespan::FromSeconds(1.0);
+		InDecoderInput.RemoveAll([&](const TSharedPtr<FDecoderInput, ESPMode::NotThreadSafe>& e)
+		{
+			return e->AccessUnit.PTS < FramePTS;
+		});
+	}
 
 	// Output to be dropped?
 	if (In->bDropOutput)
